MovieScene: Moved tools module into Editor; reorganized files; removed Classes folder.

diff --git a/Engine/Source/Developer/MovieSceneTools/Classes/K2Node_PlayMovieScene.h b/Engine/Source/Developer/MovieSceneTools/Classes/K2Node_PlayMovieScene.h
deleted file mode 100644
diff --git a/Engine/Source/Developer/MovieSceneTools/Classes/MovieSceneFactory.h b/Engine/Source/Developer/MovieSceneTools/Classes/MovieSceneFactory.h
deleted file mode 100644
diff --git a/Engine/Source/Developer/MovieSceneTools/MovieSceneTools.Build.cs b/Engine/Source/Developer/MovieSceneTools/MovieSceneTools.Build.cs
deleted file mode 100644
diff --git a/Engine/Source/Developer/MovieSceneTools/Private/AnimationTrackEditor.cpp b/Engine/Source/Developer/MovieSceneTools/Private/AnimationTrackEditor.cpp
deleted file mode 100644
diff --git a/Engine/Source/Developer/MovieSceneTools/Private/AnimationTrackEditor.h b/Engine/Source/Developer/MovieSceneTools/Private/AnimationTrackEditor.h
deleted file mode 100644
diff --git a/Engine/Source/Developer/MovieSceneTools/Private/AssetTypeActions_MovieScene.cpp b/Engine/Source/Developer/MovieSceneTools/Private/AssetTypeActions_MovieScene.cpp
deleted file mode 100644
diff --git a/Engine/Source/Developer/MovieSceneTools/Private/AssetTypeActions_MovieScene.h b/Engine/Source/Developer/MovieSceneTools/Private/AssetTypeActions_MovieScene.h
deleted file mode 100644
diff --git a/Engine/Source/Developer/MovieSceneTools/Private/AudioTrackEditor.cpp b/Engine/Source/Developer/MovieSceneTools/Private/AudioTrackEditor.cpp
deleted file mode 100644
diff --git a/Engine/Source/Developer/MovieSceneTools/Private/AudioTrackEditor.h b/Engine/Source/Developer/MovieSceneTools/Private/AudioTrackEditor.h
deleted file mode 100644
diff --git a/Engine/Source/Developer/MovieSceneTools/Private/DirectorTrackEditor.cpp b/Engine/Source/Developer/MovieSceneTools/Private/DirectorTrackEditor.cpp
deleted file mode 100644
diff --git a/Engine/Source/Developer/MovieSceneTools/Private/DirectorTrackEditor.h b/Engine/Source/Developer/MovieSceneTools/Private/DirectorTrackEditor.h
deleted file mode 100644
diff --git a/Engine/Source/Developer/MovieSceneTools/Private/K2Node_PlayMovieScene.cpp b/Engine/Source/Developer/MovieSceneTools/Private/K2Node_PlayMovieScene.cpp
deleted file mode 100644
diff --git a/Engine/Source/Developer/MovieSceneTools/Private/MovieSceneFactory.cpp b/Engine/Source/Developer/MovieSceneTools/Private/MovieSceneFactory.cpp
deleted file mode 100644
diff --git a/Engine/Source/Developer/MovieSceneTools/Private/MovieSceneToolHelpers.cpp b/Engine/Source/Developer/MovieSceneTools/Private/MovieSceneToolHelpers.cpp
deleted file mode 100644
diff --git a/Engine/Source/Developer/MovieSceneTools/Private/MovieSceneToolsModule.cpp b/Engine/Source/Developer/MovieSceneTools/Private/MovieSceneToolsModule.cpp
deleted file mode 100644
diff --git a/Engine/Source/Developer/MovieSceneTools/Private/MovieSceneToolsPrivatePCH.h b/Engine/Source/Developer/MovieSceneTools/Private/MovieSceneToolsPrivatePCH.h
deleted file mode 100644
diff --git a/Engine/Source/Developer/MovieSceneTools/Private/ParticleTrackEditor.cpp b/Engine/Source/Developer/MovieSceneTools/Private/ParticleTrackEditor.cpp
deleted file mode 100644
diff --git a/Engine/Source/Developer/MovieSceneTools/Private/ParticleTrackEditor.h b/Engine/Source/Developer/MovieSceneTools/Private/ParticleTrackEditor.h
deleted file mode 100644
diff --git a/Engine/Source/Developer/MovieSceneTools/Private/PropertyTrackEditor.cpp b/Engine/Source/Developer/MovieSceneTools/Private/PropertyTrackEditor.cpp
deleted file mode 100644
diff --git a/Engine/Source/Developer/MovieSceneTools/Private/PropertyTrackEditor.h b/Engine/Source/Developer/MovieSceneTools/Private/PropertyTrackEditor.h
deleted file mode 100644
diff --git a/Engine/Source/Developer/MovieSceneTools/Private/SEnumCurveKeyEditor.cpp b/Engine/Source/Developer/MovieSceneTools/Private/SEnumCurveKeyEditor.cpp
deleted file mode 100644
diff --git a/Engine/Source/Developer/MovieSceneTools/Private/SEnumCurveKeyEditor.h b/Engine/Source/Developer/MovieSceneTools/Private/SEnumCurveKeyEditor.h
deleted file mode 100644
diff --git a/Engine/Source/Developer/MovieSceneTools/Private/SFloatCurveKeyEditor.cpp b/Engine/Source/Developer/MovieSceneTools/Private/SFloatCurveKeyEditor.cpp
deleted file mode 100644
diff --git a/Engine/Source/Developer/MovieSceneTools/Private/SFloatCurveKeyEditor.h b/Engine/Source/Developer/MovieSceneTools/Private/SFloatCurveKeyEditor.h
deleted file mode 100644
diff --git a/Engine/Source/Developer/MovieSceneTools/Private/SIntegralCurveKeyEditor.cpp b/Engine/Source/Developer/MovieSceneTools/Private/SIntegralCurveKeyEditor.cpp
deleted file mode 100644
diff --git a/Engine/Source/Developer/MovieSceneTools/Private/SIntegralCurveKeyEditor.h b/Engine/Source/Developer/MovieSceneTools/Private/SIntegralCurveKeyEditor.h
deleted file mode 100644
diff --git a/Engine/Source/Developer/MovieSceneTools/Private/SubMovieSceneTrackEditor.cpp b/Engine/Source/Developer/MovieSceneTools/Private/SubMovieSceneTrackEditor.cpp
deleted file mode 100644
diff --git a/Engine/Source/Developer/MovieSceneTools/Private/SubMovieSceneTrackEditor.h b/Engine/Source/Developer/MovieSceneTools/Private/SubMovieSceneTrackEditor.h
deleted file mode 100644
diff --git a/Engine/Source/Developer/MovieSceneTools/Private/TransformTrackEditor.cpp b/Engine/Source/Developer/MovieSceneTools/Private/TransformTrackEditor.cpp
deleted file mode 100644
diff --git a/Engine/Source/Developer/MovieSceneTools/Private/TransformTrackEditor.h b/Engine/Source/Developer/MovieSceneTools/Private/TransformTrackEditor.h
deleted file mode 100644
diff --git a/Engine/Source/Developer/MovieSceneTools/Public/CommonMovieSceneTools.h b/Engine/Source/Developer/MovieSceneTools/Public/CommonMovieSceneTools.h
deleted file mode 100644
diff --git a/Engine/Source/Developer/MovieSceneTools/Public/IMovieSceneTools.h b/Engine/Source/Developer/MovieSceneTools/Public/IMovieSceneTools.h
deleted file mode 100644
diff --git a/Engine/Source/Developer/MovieSceneTools/Public/MovieSceneToolHelpers.h b/Engine/Source/Developer/MovieSceneTools/Public/MovieSceneToolHelpers.h
deleted file mode 100644
diff --git a/Engine/Source/Developer/MovieSceneTools/Public/PropertySection.h b/Engine/Source/Developer/MovieSceneTools/Public/PropertySection.h
deleted file mode 100644
diff --git a/Engine/Source/Editor/MovieSceneTools/MovieSceneTools.Build.cs b/Engine/Source/Editor/MovieSceneTools/MovieSceneTools.Build.cs
new file mode 100644
--- /dev/null
+++ b/Engine/Source/Editor/MovieSceneTools/MovieSceneTools.Build.cs
@@ -0,0 +1,57 @@
+// Copyright 1998-2015 Epic Games, Inc. All Rights Reserved.
+
+using UnrealBuildTool;
+
+public class MovieSceneTools : ModuleRules
+{
+	public MovieSceneTools(TargetInfo Target)
+	{
+		PrivateIncludePaths.AddRange(
+            new string[] {
+                "Editor/MovieSceneTools/Private",
+                "Editor/MovieSceneTools/Private/CurveKeyEditors",
+                "Editor/MovieSceneTools/Private/TrackEditors",
+            }
+        );
+
+		PublicDependencyModuleNames.AddRange(
+			new string[] {
+				"Core",
+				"CoreUObject",
+				"Engine",
+                "InputCore",
+				"UnrealEd",
+				"Sequencer"
+			}
+		);
+
+		PrivateDependencyModuleNames.AddRange(
+			new string[] {
+				"MovieSceneCore",
+				"MovieSceneCoreTypes",
+				"BlueprintGraph",
+				"Slate",
+				"SlateCore",
+                "EditorStyle",
+				"RenderCore",
+				"RHI"
+			}
+		);
+
+		PrivateIncludePathModuleNames.AddRange(
+			new string[] {
+				"AssetTools",
+				"Sequencer",
+				"PropertyEditor"
+			}
+		);
+
+		DynamicallyLoadedModuleNames.AddRange(
+			new string[] {
+				"AssetTools",
+				"Sequencer",
+				"PropertyEditor"
+			}
+		);
+	}
+}
diff --git a/Engine/Source/Editor/MovieSceneTools/Private/AssetTypeActions_MovieScene.cpp b/Engine/Source/Editor/MovieSceneTools/Private/AssetTypeActions_MovieScene.cpp
new file mode 100644
--- /dev/null
+++ b/Engine/Source/Editor/MovieSceneTools/Private/AssetTypeActions_MovieScene.cpp
@@ -0,0 +1,37 @@
+// Copyright 1998-2015 Epic Games, Inc. All Rights Reserved.
+
+#include "MovieSceneToolsPrivatePCH.h"
+#include "IAssetTools.h"
+#include "AssetTypeActions_MovieScene.h"
+#include "Toolkits/IToolkit.h"
+#include "ISequencerModule.h"
+#include "MovieScene.h"
+
+#define LOCTEXT_NAMESPACE "AssetTypeActions"
+
+
+UClass* FAssetTypeActions_MovieScene::GetSupportedClass() const
+{
+	return UMovieScene::StaticClass(); 
+}
+
+void FAssetTypeActions_MovieScene::OpenAssetEditor( const TArray<UObject*>& InObjects, TSharedPtr<IToolkitHost> EditWithinLevelEditor )
+{
+	EToolkitMode::Type Mode = EditWithinLevelEditor.IsValid() ? EToolkitMode::WorldCentric : EToolkitMode::Standalone;
+
+	for (auto ObjIt = InObjects.CreateConstIterator(); ObjIt; ++ObjIt)
+	{
+		UMovieScene* MovieScene = Cast<UMovieScene>(*ObjIt);
+		if (MovieScene != NULL)
+		{
+			// @todo sequencer: Only allow users to create new MovieScenes if that feature is turned on globally.
+			if( FParse::Param( FCommandLine::Get(), TEXT( "Sequencer" ) ) )
+			{
+				FSequencerViewParams ViewParams;
+				FModuleManager::LoadModuleChecked< ISequencerModule >( "Sequencer" ).CreateSequencerAssetEditor( Mode, ViewParams, EditWithinLevelEditor, MovieScene, true );
+			}
+		}
+	}
+}
+
+#undef LOCTEXT_NAMESPACE
\ No newline at end of file
diff --git a/Engine/Source/Editor/MovieSceneTools/Private/AssetTypeActions_MovieScene.h b/Engine/Source/Editor/MovieSceneTools/Private/AssetTypeActions_MovieScene.h
new file mode 100644
--- /dev/null
+++ b/Engine/Source/Editor/MovieSceneTools/Private/AssetTypeActions_MovieScene.h
@@ -0,0 +1,22 @@
+// Copyright 1998-2015 Epic Games, Inc. All Rights Reserved.
+
+#pragma once
+
+#include "AssetTypeActions_Base.h"
+
+class FAssetTypeActions_MovieScene : public FAssetTypeActions_Base
+{
+public:
+	// IAssetTypeActions Implementation
+	virtual FText GetName() const override { return NSLOCTEXT("AssetTypeActions", "AssetTypeActions_MovieScene", "Movie Scene"); }
+	virtual FColor GetTypeColor() const override { return FColor(200, 80, 80); }
+	virtual UClass* GetSupportedClass() const override;
+	virtual void OpenAssetEditor( const TArray<UObject*>& InObjects, TSharedPtr<class IToolkitHost> EditWithinLevelEditor = TSharedPtr<IToolkitHost>() ) override;
+	virtual uint32 GetCategories() override { return EAssetTypeCategories::Animation; }
+
+	virtual bool ShouldForceWorldCentric() override
+	{
+		// @todo sequencer: Hack to force world-centric mode for Sequencer
+		return true;
+	}
+};
\ No newline at end of file
diff --git a/Engine/Source/Editor/MovieSceneTools/Private/CurveKeyEditors/SEnumCurveKeyEditor.cpp b/Engine/Source/Editor/MovieSceneTools/Private/CurveKeyEditors/SEnumCurveKeyEditor.cpp
new file mode 100644
--- /dev/null
+++ b/Engine/Source/Editor/MovieSceneTools/Private/CurveKeyEditors/SEnumCurveKeyEditor.cpp
@@ -0,0 +1,79 @@
+// Copyright 1998-2015 Epic Games, Inc. All Rights Reserved.
+
+#include "MovieSceneToolsPrivatePCH.h"
+#include "SEnumCurveKeyEditor.h"
+
+#define LOCTEXT_NAMESPACE "EnumCurveKeyEditor"
+
+void SEnumCurveKeyEditor::Construct(const FArguments& InArgs)
+{
+	Sequencer = InArgs._Sequencer;
+	OwningSection = InArgs._OwningSection;
+	Curve = InArgs._Curve;
+	for (int32 i = 0; i < InArgs._Enum->NumEnums(); i++)
+	{
+		FString EnumDisplayValue = InArgs._Enum->GetDisplayNameText(i).ToString();
+		if (EnumDisplayValue.Len() > 0)
+		{
+			EnumDisplayValue = InArgs._Enum->GetEnumName(i);
+		}
+		EnumValues.Add(MakeShareable(new FString(EnumDisplayValue)));
+	}
+	float CurrentTime = Sequencer->GetCurrentLocalTime(*Sequencer->GetFocusedMovieScene());
+	ChildSlot
+	[
+		SNew(SComboBox<TSharedPtr<FString>>)
+		.ButtonStyle(FEditorStyle::Get(), "FlatButton.Light")
+		.OptionsSource(&EnumValues)
+		.InitiallySelectedItem(EnumValues[Curve->Evaluate(CurrentTime)])
+		.OnGenerateWidget(this, &SEnumCurveKeyEditor::OnGenerateWidget)
+		.OnSelectionChanged(this, &SEnumCurveKeyEditor::OnComboSelectionChanged)
+		.ContentPadding(FMargin(2, 0))
+		[
+			SNew(STextBlock)
+			.Font(FEditorStyle::GetFontStyle("Sequencer.AnimationOutliner.RegularFont"))
+			.Text(this, &SEnumCurveKeyEditor::GetCurrentValue)
+		]
+	];
+}
+
+FText SEnumCurveKeyEditor::GetCurrentValue() const
+{
+	float CurrentTime = Sequencer->GetCurrentLocalTime(*Sequencer->GetFocusedMovieScene());
+	return FText::FromString(*EnumValues[Curve->Evaluate(CurrentTime)]);
+}
+
+TSharedRef<SWidget> SEnumCurveKeyEditor::OnGenerateWidget(TSharedPtr<FString> InItem)
+{
+	return SNew(STextBlock)
+		.Text(FText::FromString(*InItem));
+}
+
+void SEnumCurveKeyEditor::OnComboSelectionChanged(TSharedPtr<FString> InSelectedItem, ESelectInfo::Type SelectInfo)
+{
+	FScopedTransaction Transaction(LOCTEXT("SetEnumKey", "Set enum key value"));
+	OwningSection->SetFlags(RF_Transactional);
+	OwningSection->Modify();
+
+	float CurrentTime = Sequencer->GetCurrentLocalTime(*Sequencer->GetFocusedMovieScene());
+
+	bool bKeyWillBeAdded = Curve->IsKeyHandleValid(Curve->FindKey(CurrentTime)) == false;
+	if (bKeyWillBeAdded)
+	{
+		if (OwningSection->GetStartTime() > CurrentTime)
+		{
+			OwningSection->SetStartTime(CurrentTime);
+		}
+		if (OwningSection->GetEndTime() < CurrentTime)
+		{
+			OwningSection->SetEndTime(CurrentTime);
+		}
+	}
+
+	int32 SelectedIndex;
+	EnumValues.Find(InSelectedItem, SelectedIndex);
+	Curve->UpdateOrAddKey(CurrentTime, SelectedIndex);
+	Sequencer->UpdateRuntimeInstances();
+}
+
+#undef LOCTEXT_NAMESPACE
\ No newline at end of file
diff --git a/Engine/Source/Editor/MovieSceneTools/Private/CurveKeyEditors/SEnumCurveKeyEditor.h b/Engine/Source/Editor/MovieSceneTools/Private/CurveKeyEditors/SEnumCurveKeyEditor.h
new file mode 100644
--- /dev/null
+++ b/Engine/Source/Editor/MovieSceneTools/Private/CurveKeyEditors/SEnumCurveKeyEditor.h
@@ -0,0 +1,31 @@
+// Copyright 1998-2015 Epic Games, Inc. All Rights Reserved.
+
+#pragma once
+
+/**
+ * A widget for editing a curve representing enum keys.
+ */
+class SEnumCurveKeyEditor : public SCompoundWidget
+{
+public:
+	SLATE_BEGIN_ARGS(SEnumCurveKeyEditor) {}
+		SLATE_ARGUMENT(ISequencer*, Sequencer)
+		SLATE_ARGUMENT(UMovieSceneSection*, OwningSection)
+		SLATE_ARGUMENT(FIntegralCurve*, Curve)
+		SLATE_ARGUMENT(UEnum*, Enum)
+	SLATE_END_ARGS();
+
+	void Construct(const FArguments& InArgs);
+
+private:
+	FText GetCurrentValue() const;
+
+	TSharedRef<SWidget> OnGenerateWidget(TSharedPtr<FString> InItem);
+
+	void OnComboSelectionChanged(TSharedPtr<FString> InSelectedItem, ESelectInfo::Type SelectInfo);
+
+	ISequencer* Sequencer;
+	UMovieSceneSection* OwningSection;
+	FIntegralCurve* Curve;
+	TArray<TSharedPtr<FString>> EnumValues;
+};
\ No newline at end of file
diff --git a/Engine/Source/Editor/MovieSceneTools/Private/CurveKeyEditors/SFloatCurveKeyEditor.cpp b/Engine/Source/Editor/MovieSceneTools/Private/CurveKeyEditors/SFloatCurveKeyEditor.cpp
new file mode 100644
--- /dev/null
+++ b/Engine/Source/Editor/MovieSceneTools/Private/CurveKeyEditors/SFloatCurveKeyEditor.cpp
@@ -0,0 +1,77 @@
+// Copyright 1998-2015 Epic Games, Inc. All Rights Reserved.
+
+#include "MovieSceneToolsPrivatePCH.h"
+#include "SFloatCurveKeyEditor.h"
+
+#define LOCTEXT_NAMESPACE "FloatCurveKeyEditor"
+
+void SFloatCurveKeyEditor::Construct(const FArguments& InArgs)
+{
+	Sequencer = InArgs._Sequencer;
+	OwningSection = InArgs._OwningSection;
+	Curve = InArgs._Curve;
+	ChildSlot
+	[
+		SNew(SSpinBox<float>)
+		.Style(&FEditorStyle::GetWidgetStyle<FSpinBoxStyle>("Sequencer.AnimationOutliner.KeyEditorSpinBoxStyle"))
+		.Font(FEditorStyle::GetFontStyle("Sequencer.AnimationOutliner.RegularFont"))
+		.MinValue(TOptional<float>())
+		.MaxValue(TOptional<float>())
+		.MaxSliderValue(TOptional<float>())
+		.MinSliderValue(TOptional<float>())
+		.Delta(0.001f)
+		.MinDesiredWidth(60.0f)
+		.Value(this, &SFloatCurveKeyEditor::OnGetKeyValue)
+		.OnValueChanged(this, &SFloatCurveKeyEditor::OnValueChanged)
+		.OnBeginSliderMovement(this, &SFloatCurveKeyEditor::OnBeginSliderMovement)
+		.OnEndSliderMovement(this, &SFloatCurveKeyEditor::OnEndSliderMovement)
+	];
+}
+
+void SFloatCurveKeyEditor::OnBeginSliderMovement()
+{
+	GEditor->BeginTransaction(LOCTEXT("SetFloatKey", "Set float key value"));
+	OwningSection->SetFlags(RF_Transactional);
+	OwningSection->Modify();
+}
+
+void SFloatCurveKeyEditor::OnEndSliderMovement(float Value)
+{
+	if (GEditor->IsTransactionActive())
+	{
+		GEditor->EndTransaction();
+	}
+}
+
+float SFloatCurveKeyEditor::OnGetKeyValue() const
+{
+	float CurrentTime = Sequencer->GetCurrentLocalTime(*Sequencer->GetFocusedMovieScene());
+	return Curve->Eval(CurrentTime);
+}
+
+void SFloatCurveKeyEditor::OnValueChanged(float Value)
+{
+	OwningSection->Modify();
+
+	float CurrentTime = Sequencer->GetCurrentLocalTime(*Sequencer->GetFocusedMovieScene());
+	FKeyHandle CurrentKeyHandle = Curve->FindKey(CurrentTime);
+	if (Curve->IsKeyHandleValid(CurrentKeyHandle))
+	{
+		Curve->SetKeyValue(CurrentKeyHandle, Value);
+	}
+	else
+	{
+		Curve->AddKey(CurrentTime, Value, false, CurrentKeyHandle);
+		if (OwningSection->GetStartTime() > CurrentTime)
+		{
+			OwningSection->SetStartTime(CurrentTime);
+		}
+		if (OwningSection->GetEndTime() < CurrentTime)
+		{
+			OwningSection->SetEndTime(CurrentTime);
+		}
+	}
+	Sequencer->UpdateRuntimeInstances();
+}
+
+#undef LOCTEXT_NAMESPACE
\ No newline at end of file
diff --git a/Engine/Source/Editor/MovieSceneTools/Private/CurveKeyEditors/SFloatCurveKeyEditor.h b/Engine/Source/Editor/MovieSceneTools/Private/CurveKeyEditors/SFloatCurveKeyEditor.h
new file mode 100644
--- /dev/null
+++ b/Engine/Source/Editor/MovieSceneTools/Private/CurveKeyEditors/SFloatCurveKeyEditor.h
@@ -0,0 +1,31 @@
+// Copyright 1998-2015 Epic Games, Inc. All Rights Reserved.
+
+#pragma once
+
+/**
+ * A widget for editing a curve representing float keys.
+ */
+class SFloatCurveKeyEditor : public SCompoundWidget
+{
+public:
+	SLATE_BEGIN_ARGS(SFloatCurveKeyEditor) {}
+		SLATE_ARGUMENT(ISequencer*, Sequencer)
+		SLATE_ARGUMENT(UMovieSceneSection*, OwningSection)
+		SLATE_ARGUMENT(FRichCurve*, Curve)
+	SLATE_END_ARGS();
+
+	void Construct(const FArguments& InArgs);
+
+private:
+	void OnBeginSliderMovement();
+
+	void OnEndSliderMovement(float Value);
+
+	float OnGetKeyValue() const;
+
+	void OnValueChanged(float Value);
+
+	ISequencer* Sequencer;
+	UMovieSceneSection* OwningSection;
+	FRichCurve* Curve;
+};
\ No newline at end of file
diff --git a/Engine/Source/Editor/MovieSceneTools/Private/CurveKeyEditors/SIntegralCurveKeyEditor.cpp b/Engine/Source/Editor/MovieSceneTools/Private/CurveKeyEditors/SIntegralCurveKeyEditor.cpp
new file mode 100644
--- /dev/null
+++ b/Engine/Source/Editor/MovieSceneTools/Private/CurveKeyEditors/SIntegralCurveKeyEditor.cpp
@@ -0,0 +1,75 @@
+// Copyright 1998-2015 Epic Games, Inc. All Rights Reserved.
+
+#include "MovieSceneToolsPrivatePCH.h"
+#include "SIntegralCurveKeyEditor.h"
+
+#define LOCTEXT_NAMESPACE "IntegralCurveKeyEditor"
+
+void SIntegralCurveKeyEditor::Construct(const FArguments& InArgs)
+{
+	Sequencer = InArgs._Sequencer;
+	OwningSection = InArgs._OwningSection;
+	Curve = InArgs._Curve;
+	ChildSlot
+	[
+		SNew(SSpinBox<int32>)
+		.Style(&FEditorStyle::GetWidgetStyle<FSpinBoxStyle>("Sequencer.AnimationOutliner.KeyEditorSpinBoxStyle"))
+		.Font(FEditorStyle::GetFontStyle("Sequencer.AnimationOutliner.RegularFont"))
+		.MinValue(TOptional<int32>())
+		.MaxValue(TOptional<int32>())
+		.MaxSliderValue(TOptional<int32>())
+		.MinSliderValue(TOptional<int32>())
+		.Delta(1)
+		.MinDesiredWidth(60.0f)
+		.Value(this, &SIntegralCurveKeyEditor::OnGetKeyValue)
+		.OnValueChanged(this, &SIntegralCurveKeyEditor::OnValueChanged)
+		.OnBeginSliderMovement(this, &SIntegralCurveKeyEditor::OnBeginSliderMovement)
+		.OnEndSliderMovement(this, &SIntegralCurveKeyEditor::OnEndSliderMovement)
+	];
+}
+
+void SIntegralCurveKeyEditor::OnBeginSliderMovement()
+{
+	GEditor->BeginTransaction(LOCTEXT("SetIntegralKey", "Set integral key value"));
+	OwningSection->SetFlags(RF_Transactional);
+	OwningSection->Modify();
+}
+
+void SIntegralCurveKeyEditor::OnEndSliderMovement(int32 Value)
+{
+	if (GEditor->IsTransactionActive())
+	{
+		GEditor->EndTransaction();
+	}
+}
+
+int32 SIntegralCurveKeyEditor::OnGetKeyValue() const
+{
+	float CurrentTime = Sequencer->GetCurrentLocalTime(*Sequencer->GetFocusedMovieScene());
+	return Curve->Evaluate(CurrentTime);
+}
+
+void SIntegralCurveKeyEditor::OnValueChanged(int32 Value)
+{
+	OwningSection->Modify();
+
+	float CurrentTime = Sequencer->GetCurrentLocalTime(*Sequencer->GetFocusedMovieScene());
+
+	bool bKeyWillBeAdded = Curve->IsKeyHandleValid(Curve->FindKey(CurrentTime)) == false;
+	if (bKeyWillBeAdded)
+	{
+		if (OwningSection->GetStartTime() > CurrentTime)
+		{
+			OwningSection->SetStartTime(CurrentTime);
+		}
+		if (OwningSection->GetEndTime() < CurrentTime)
+		{
+			OwningSection->SetEndTime(CurrentTime);
+		}
+	}
+
+	Curve->UpdateOrAddKey(CurrentTime, Value);
+	Sequencer->UpdateRuntimeInstances();
+}
+
+#undef LOCTEXT_NAMESPACE
\ No newline at end of file
diff --git a/Engine/Source/Editor/MovieSceneTools/Private/CurveKeyEditors/SIntegralCurveKeyEditor.h b/Engine/Source/Editor/MovieSceneTools/Private/CurveKeyEditors/SIntegralCurveKeyEditor.h
new file mode 100644
--- /dev/null
+++ b/Engine/Source/Editor/MovieSceneTools/Private/CurveKeyEditors/SIntegralCurveKeyEditor.h
@@ -0,0 +1,31 @@
+// Copyright 1998-2015 Epic Games, Inc. All Rights Reserved.
+
+#pragma once
+
+/**
+ * A widget for editing a curve representing integer keys.
+ */
+class SIntegralCurveKeyEditor : public SCompoundWidget
+{
+public:
+	SLATE_BEGIN_ARGS(SIntegralCurveKeyEditor) {}
+		SLATE_ARGUMENT(ISequencer*, Sequencer)
+		SLATE_ARGUMENT(UMovieSceneSection*, OwningSection)
+		SLATE_ARGUMENT(FIntegralCurve*, Curve)
+	SLATE_END_ARGS();
+
+	void Construct(const FArguments& InArgs);
+
+private:
+	void OnBeginSliderMovement();
+
+	void OnEndSliderMovement(int32 Value);
+
+	int32 OnGetKeyValue() const;
+
+	void OnValueChanged(int32 Value);
+
+	ISequencer* Sequencer;
+	UMovieSceneSection* OwningSection;
+	FIntegralCurve* Curve;
+};
\ No newline at end of file
diff --git a/Engine/Source/Editor/MovieSceneTools/Private/K2Node_PlayMovieScene.cpp b/Engine/Source/Editor/MovieSceneTools/Private/K2Node_PlayMovieScene.cpp
new file mode 100644
--- /dev/null
+++ b/Engine/Source/Editor/MovieSceneTools/Private/K2Node_PlayMovieScene.cpp
@@ -0,0 +1,349 @@
+// Copyright 1998-2015 Epic Games, Inc. All Rights Reserved.
+
+#include "MovieSceneToolsPrivatePCH.h"
+#include "ScopedTransaction.h"
+#include "K2Node_PlayMovieScene.h"
+#include "BlueprintUtilities.h"
+#include "MovieScene.h"
+#include "Kismet2/BlueprintEditorUtils.h"
+#include "MovieSceneBindings.h"
+
+
+UEdGraphNode* FEdGraphSchemaAction_K2AddPlayMovieScene::PerformAction(class UEdGraph* ParentGraph, UEdGraphPin* FromPin, const FVector2D Location, bool bSelectNewNode/* = true*/)
+{
+	const FScopedTransaction Transaction( NSLOCTEXT("PlayMovieSceneNode", "UndoAddingNewNode", "Add PlayMovieScene Node") );
+	UBlueprint* Blueprint = FBlueprintEditorUtils::FindBlueprintForGraphChecked(ParentGraph);
+	UEdGraphNode* NewNode = FEdGraphSchemaAction_K2NewNode::PerformAction(ParentGraph, FromPin, Location, bSelectNewNode);
+
+	if(NewNode != NULL && Blueprint != NULL)
+	{
+		UK2Node_PlayMovieScene* PlayMovieSceneNode = CastChecked<UK2Node_PlayMovieScene>(NewNode);
+
+		// @todo sequencer: Perform any setup work on the new PlayMovieScene node here!
+	}
+
+	return NewNode;
+}
+
+
+UK2Node_PlayMovieScene::UK2Node_PlayMovieScene(const FObjectInitializer& ObjectInitializer)
+	: Super(ObjectInitializer)
+{
+}
+
+
+namespace PlayMovieScenePinNames
+{
+	static FString Play( TEXT( "Play" ) );
+	static FString Pause( TEXT( "Pause" ) );
+}
+
+void UK2Node_PlayMovieScene::AllocateDefaultPins()
+{
+	const UEdGraphSchema_K2* K2Schema = GetDefault<UEdGraphSchema_K2>();
+
+	// Do not allow users to rename the node
+	// @todo sequencer: If we want to support renaming, this we'll need to update FNameValidatorFactory::MakeValidator() and perhaps other locations
+	bCanRenameNode = 0;
+
+	// "Play" starts or resumes playback from the current position
+	CreatePin( EGPD_Input, K2Schema->PC_Exec, TEXT(""), NULL, false, false, PlayMovieScenePinNames::Play );
+
+	// "Pause" stops playback, leaving the time cursor at its current position
+	CreatePin( EGPD_Input, K2Schema->PC_Exec, TEXT(""), NULL, false, false, PlayMovieScenePinNames::Pause );
+
+	// @todo sequencer: Add PlayFromStart?
+	// @todo sequencer: Add PlayReverse pin?  PlayReverseFromEnd?
+	// @todo sequencer: Add "set time" input
+	// @todo sequencer: Needs output pin for "finished playing"
+
+	if( MovieSceneBindings != NULL )
+	{
+		for( auto BoundObjectIter( MovieSceneBindings->GetBoundObjects().CreateIterator() ); BoundObjectIter; ++BoundObjectIter )
+		{
+			auto& BoundObject = *BoundObjectIter;
+
+			CreatePinForBoundObject( BoundObject );
+		}
+	}
+	
+	Super::AllocateDefaultPins();
+}
+
+FLinearColor UK2Node_PlayMovieScene::GetNodeTitleColor() const
+{
+	return FLinearColor(0.8f, 0.05f, 0.05f);
+}
+
+FText UK2Node_PlayMovieScene::GetNodeTitle(ENodeTitleType::Type TitleType) const
+{
+	UMovieScene* MovieScene = (MovieSceneBindings != nullptr) ? MovieSceneBindings->GetRootMovieScene() : nullptr;
+	if (MovieScene == nullptr)
+	{
+		return NSLOCTEXT("PlayMovieSceneNode", "NodeTitleWithNoMovieScene", "Play Movie Scene (No Asset)");
+	}
+	// @TODO: don't know enough about this node type to comfortably assert that
+	//        the MovieScene won't change after the node has spawned... until
+	//        then, we'll leave this optimization off
+	else //if (CachedNodeTitle.IsOutOfDate(this))
+	{
+		FFormatNamedArguments Args;
+		Args.Add(TEXT("SceneName"), FText::FromString(MovieScene->GetName()));
+		// FText::Format() is slow, so we cache this to save on performance
+		CachedNodeTitle.SetCachedText(FText::Format(NSLOCTEXT("PlayMovieSceneNode", "NodeTitle", "Play Movie Scene: {SceneName}"), Args), this);
+	}
+	return CachedNodeTitle;
+}
+
+UEdGraphPin* UK2Node_PlayMovieScene::GetPlayPin() const
+{
+	UEdGraphPin* Pin = FindPinChecked( PlayMovieScenePinNames::Play );
+	check( Pin->Direction == EGPD_Input );
+	return Pin;
+}
+
+
+UEdGraphPin* UK2Node_PlayMovieScene::GetPausePin() const
+{
+	UEdGraphPin* Pin = FindPinChecked( PlayMovieScenePinNames::Pause );
+	check( Pin->Direction == EGPD_Input );
+	return Pin;
+}
+
+
+void UK2Node_PlayMovieScene::PinConnectionListChanged( UEdGraphPin* InPin )
+{
+	// Call parent implementation
+	Super::PostReconstructNode();
+
+	// @todo sequencer: What about the "default" value for the pin?  Right now we don't do anything with that.  Might feel buggy.
+	//			--> Actually, we can't use that yet because Actor references won't be fixed up for PIE except for literals
+
+	CreateBindingsIfNeeded();
+
+	// Update the MovieScene bindings for any changes that were made to the node
+	// @todo sequencer: Only a single incoming Pin is changing, but we're refreshing all pin states here.  Could be simplified
+	bool bAnyBindingsChanged = false;
+	{
+		auto& BoundObjects = MovieSceneBindings->GetBoundObjects();
+		for( auto BoundObjectIter( BoundObjects.CreateIterator() ); BoundObjectIter; ++BoundObjectIter )
+		{
+			auto& BoundObject = *BoundObjectIter;
+
+			TArray< UObject* > OtherObjects;
+
+			// Find the pin that goes with this object
+			auto* Pin = FindPinChecked( BoundObject.GetPossessableGuid().ToString( EGuidFormats::DigitsWithHyphens ) );
+
+			// Is the pin even linked to anything?
+			if( Pin->LinkedTo.Num() > 0 )
+			{
+				// @todo sequencer major: Add support for multiple connections to a single input pin to be treated as an "unordered array"
+				for( auto OtherPinIter( Pin->LinkedTo.CreateIterator() ); OtherPinIter; ++OtherPinIter )
+				{
+					auto* OtherPin = *OtherPinIter;
+					if( OtherPin != NULL )
+					{
+						// Look for an object bound to a literal.  We can use these for scrub preview in the editor!
+						UK2Node_Literal* OtherLiteralPin = Cast< UK2Node_Literal >( OtherPin->GetOwningNode() );
+						if( OtherLiteralPin != NULL )
+						{
+							// @todo sequencer: Because we only recognize object literals, any dynamically bound actor won't be scrubable
+							//   in the editor.  Maybe we should support a "preview actor" that can be hooked up just for scrubbing and puppeting?
+							//  ==> Maybe use the pin's default value as the preview actor, even when overridden with a dynamically spawned actor?
+							UObject* OtherObject = OtherLiteralPin->GetObjectRef();
+							if ( OtherObject != NULL )
+							{
+								OtherObjects.Add( OtherObject );
+							}							
+						}
+					}
+				}
+			}
+
+			// Update our bindings to match the state of the node
+			if( BoundObject.GetObjects() != OtherObjects )	// @todo sequencer: Kind of weird to compare entire array (order change shouldn't cause us to invalidate).  Change to a set compare?
+			{
+				// @todo sequencer: No type checking is happening here.  Should we? (interactive during pin drag?)
+				Modify();
+				BoundObject.SetObjects( OtherObjects );
+
+				bAnyBindingsChanged = true;
+			}
+		}
+	}
+
+	if( bAnyBindingsChanged )
+	{
+		OnBindingsChangedEvent.Broadcast();
+	}
+}
+
+
+void UK2Node_PlayMovieScene::CreateBindingsIfNeeded()
+{
+	// NOTE: This function can be called on the transient copy of the UK2Node_PlayMovieScene before we've called
+	// SpawnNodeFromTemplate to create the actual node in the graph.  Just something to be aware of if you go looking at
+	// Outers and find yourself in the transient package.
+	if( MovieSceneBindings == NULL )
+	{
+		Modify();
+
+		// The new MovieSceneBindings object will be stored as an "inner" of this object
+		UObject* MovieSceneBindingsOuter = this;
+		UMovieSceneBindings* NewMovieSceneBindings = NewObject<UMovieSceneBindings>(MovieSceneBindingsOuter);
+		check( NewMovieSceneBindings != NULL );
+
+		MovieSceneBindings = NewMovieSceneBindings;
+		CachedNodeTitle.MarkDirty();
+	}
+}
+
+
+UMovieScene* UK2Node_PlayMovieScene::GetMovieScene()
+{
+	return MovieSceneBindings != NULL ? MovieSceneBindings->GetRootMovieScene() : NULL;
+}
+
+
+UMovieSceneBindings* UK2Node_PlayMovieScene::GetMovieSceneBindings()
+{
+	return MovieSceneBindings;
+}
+
+
+void UK2Node_PlayMovieScene::SetMovieScene( class UMovieScene* NewMovieScene )
+{
+	CreateBindingsIfNeeded();
+
+	MovieSceneBindings->SetRootMovieScene( NewMovieScene );
+}
+
+
+void UK2Node_PlayMovieScene::BindPossessableToObjects( const FGuid& PossessableGuid, const TArray< UObject* >& Objects )
+{
+	CreateBindingsIfNeeded();
+
+	// Bind the possessable from the MovieScene asset to the supplied object
+	FMovieSceneBoundObject& BoundObject = MovieSceneBindings->AddBinding( PossessableGuid, Objects );
+
+	// Create a pin for the bound object
+	Modify();
+	CreatePinForBoundObject( BoundObject );
+	FBlueprintEditorUtils::MarkBlueprintAsStructurallyModified( this->GetBlueprint() );
+}
+
+
+TArray< UObject* > UK2Node_PlayMovieScene::FindBoundObjects( const FGuid& Guid )
+{
+	if( MovieSceneBindings != NULL )
+	{
+		return MovieSceneBindings->FindBoundObjects( Guid );
+	}
+
+	return TArray< UObject* >();
+}
+
+FGuid UK2Node_PlayMovieScene::FindGuidForObject( UObject* Object ) const
+{
+	if( MovieSceneBindings != NULL )
+	{
+		// Search all bindings for the object 
+		TArray<FMovieSceneBoundObject>& ObjectBindings = MovieSceneBindings->GetBoundObjects();
+		for( int32 BoundObjIndex = 0; BoundObjIndex < ObjectBindings.Num(); ++BoundObjIndex )
+		{
+			// A binding slots can have multiple objects.  Ask the slot if the object is bound
+			FMovieSceneBoundObject& BindingSlot = ObjectBindings[BoundObjIndex];
+			if( BindingSlot.ContainsObject( Object ) )
+			{
+				return BindingSlot.GetPossessableGuid();
+			}
+		}
+	}
+
+	return FGuid();
+}
+
+
+void UK2Node_PlayMovieScene::CreatePinForBoundObject( FMovieSceneBoundObject& BoundObject )
+{
+	const UEdGraphSchema_K2* K2Schema = GetDefault<UEdGraphSchema_K2>();
+
+	// We use the GUID as the pin name as this uniquely identifies it
+	const FString PinName = BoundObject.GetPossessableGuid().ToString( EGuidFormats::DigitsWithHyphens );
+
+	// For the friendly name, we use the possessable name from the MovieScene asset that is associated with this node
+	FText PinFriendlyName = FText::FromString(PinName);
+	{
+		UMovieScene* MovieScene = GetMovieScene();
+		if( MovieScene != NULL )		// @todo sequencer: Need to refresh the PinFriendlyName if the MovieScene asset changes, or if the possessable slot is renamed within Sequencer
+		{
+			for( auto PossessableIndex = 0; PossessableIndex < MovieScene->GetPossessableCount(); ++PossessableIndex )
+			{
+				auto& Possessable = MovieScene->GetPossessable( PossessableIndex );
+				if( Possessable.GetGuid() == BoundObject.GetPossessableGuid() )
+				{
+					// Found a name for this possessable
+					PinFriendlyName = FText::FromString(Possessable.GetName());
+					break;
+				}
+			}
+		}
+	}
+
+	const FString PinSubCategory(TEXT(""));
+	UObject* PinSubCategoryObject = AActor::StaticClass();
+	const bool bIsArray = false;
+	const bool bIsReference = false;
+	UEdGraphPin* NewPin = CreatePin( EGPD_Input, K2Schema->PC_Object, PinSubCategory, PinSubCategoryObject, bIsArray, bIsReference, PinName );
+	check( NewPin != NULL );
+
+	// Set the friendly name for this pin
+	NewPin->PinFriendlyName = PinFriendlyName;
+
+	// Place a literal for the bound object and hook it up to the pin
+	// @todo sequencer: Should we instead set the default on the pin to the object instead?
+	const TArray< UObject* >& Objects = BoundObject.GetObjects();
+	if( Objects.Num() > 0 )
+	{
+		for( auto ObjectIter( Objects.CreateConstIterator() ); ObjectIter; ++ObjectIter )
+		{
+			auto* Object = *ObjectIter;
+			if( ensure( Object != NULL ) )
+			{
+				// Check to see if we have a literal for this object already
+				UK2Node_Literal* LiteralNode = NULL;
+				{
+					TArray< UK2Node_Literal* > LiteralNodes;
+					GetGraph()->GetNodesOfClass( LiteralNodes );
+					for( auto NodeIt( LiteralNodes.CreateConstIterator() ); NodeIt; ++NodeIt )
+					{
+						const auto CurLiteralNode = *NodeIt;
+
+						if( CurLiteralNode->GetObjectRef() == Object )
+						{
+							// Found one!
+							LiteralNode = CurLiteralNode;
+							break;
+						}
+					}
+				}
+
+				if( LiteralNode == NULL )
+				{
+					// No literal for this object yet, so we'll make one now.
+					UK2Node_Literal* LiteralNodeTemplate = NewObject<UK2Node_Literal>();
+					LiteralNodeTemplate->SetObjectRef( Object );
+
+					// Figure out a decent place to stick the node
+					// @todo sequencer: The placement of these is really unacceptable.  We should setup new logic specific for moviescene pin objects.
+					const FVector2D NewNodePos = GetGraph()->GetGoodPlaceForNewNode();
+					LiteralNode = FEdGraphSchemaAction_K2NewNode::SpawnNodeFromTemplate<UK2Node_Literal>(GetGraph(), LiteralNodeTemplate, NewNodePos);
+				}
+
+				// Hook up the object reference literal to our pin
+				LiteralNode->GetValuePin()->MakeLinkTo( NewPin );
+			}
+		}
+	}
+}
diff --git a/Engine/Source/Editor/MovieSceneTools/Private/MovieSceneFactory.cpp b/Engine/Source/Editor/MovieSceneTools/Private/MovieSceneFactory.cpp
new file mode 100644
--- /dev/null
+++ b/Engine/Source/Editor/MovieSceneTools/Private/MovieSceneFactory.cpp
@@ -0,0 +1,30 @@
+// Copyright 1998-2015 Epic Games, Inc. All Rights Reserved.
+
+#include "MovieSceneToolsPrivatePCH.h"
+#include "MovieSceneFactory.h"
+#include "MovieScene.h"
+
+#define LOCTEXT_NAMESPACE "MovieSceneFactory"
+
+UMovieSceneFactory::UMovieSceneFactory( const FObjectInitializer& ObjectInitializer )
+	: Super(ObjectInitializer)
+{
+
+	// @todo sequencer: Only allow users to create new MovieScenes if that feature is turned on globally.
+	if( FParse::Param( FCommandLine::Get(), TEXT( "Sequencer" ) ) )
+	{
+		bCreateNew = true;
+	}
+
+	bEditAfterNew = true;
+
+	SupportedClass = UMovieScene::StaticClass();
+}
+
+
+UObject* UMovieSceneFactory::FactoryCreateNew(UClass* Class, UObject* InParent, FName Name, EObjectFlags Flags, UObject* Context, FFeedbackContext* Warn)
+{
+	return NewObject<UMovieScene>(InParent, Name, Flags);
+}
+
+#undef LOCTEXT_NAMESPACE
\ No newline at end of file
diff --git a/Engine/Source/Editor/MovieSceneTools/Private/MovieSceneToolHelpers.cpp b/Engine/Source/Editor/MovieSceneTools/Private/MovieSceneToolHelpers.cpp
new file mode 100644
--- /dev/null
+++ b/Engine/Source/Editor/MovieSceneTools/Private/MovieSceneToolHelpers.cpp
@@ -0,0 +1,76 @@
+// Copyright 1998-2015 Epic Games, Inc. All Rights Reserved.
+
+#include "MovieSceneToolsPrivatePCH.h"
+#include "PropertyEditing.h"
+#include "MovieScene.h"
+#include "MovieSceneTrack.h"
+#include "ScopedTransaction.h"
+#include "MovieSceneSection.h"
+#include "ISequencerObjectChangeListener.h"
+#include "ISequencerSection.h"
+#include "ISectionLayoutBuilder.h"
+#include "IKeyArea.h"
+#include "MovieSceneToolHelpers.h"
+#include "SFloatCurveKeyEditor.h"
+#include "SIntegralCurveKeyEditor.h"
+#include "SEnumCurveKeyEditor.h"
+
+void FFloatCurveKeyArea::AddKeyUnique(float Time)
+{
+	FKeyHandle CurrentKeyHandle = Curve->FindKey(Time);
+	if (Curve->IsKeyHandleValid(CurrentKeyHandle) == false)
+	{
+		if (OwningSection->GetStartTime() > Time)
+		{
+			OwningSection->SetStartTime(Time);
+		}
+		if (OwningSection->GetEndTime() < Time)
+		{
+			OwningSection->SetEndTime(Time);
+		}
+		Curve->AddKey(Time, Curve->Eval(Time), false, CurrentKeyHandle);
+	}
+}
+
+TSharedRef<SWidget> FFloatCurveKeyArea::CreateKeyEditor(ISequencer* Sequencer)
+{
+	return SNew(SFloatCurveKeyEditor)
+		.Sequencer(Sequencer)
+		.OwningSection(OwningSection)
+		.Curve(Curve);
+};
+
+void FIntegralKeyArea::AddKeyUnique(float Time)
+{
+	FKeyHandle CurrentKey = Curve.FindKey(Time);
+	if (Curve.IsKeyHandleValid(CurrentKey) == false)
+	{
+		if (OwningSection->GetStartTime() > Time)
+		{
+			OwningSection->SetStartTime(Time);
+		}
+		if (OwningSection->GetEndTime() < Time)
+		{
+			OwningSection->SetEndTime(Time);
+		}
+		Curve.AddKey(Time, Curve.Evaluate(Time), CurrentKey);
+	}
+}
+
+TSharedRef<SWidget> FIntegralKeyArea::CreateKeyEditor(ISequencer* Sequencer)
+{
+	return SNew(SIntegralCurveKeyEditor)
+		.Sequencer(Sequencer)
+		.OwningSection(OwningSection)
+		.Curve(&Curve);
+};
+
+TSharedRef<SWidget> FEnumKeyArea::CreateKeyEditor(ISequencer* Sequencer)
+{
+	return SNew(SEnumCurveKeyEditor)
+		.Sequencer(Sequencer)
+		.OwningSection(OwningSection)
+		.Curve(&Curve)
+		.Enum(Enum);
+};
+
diff --git a/Engine/Source/Editor/MovieSceneTools/Private/MovieSceneToolsModule.cpp b/Engine/Source/Editor/MovieSceneTools/Private/MovieSceneToolsModule.cpp
new file mode 100644
--- /dev/null
+++ b/Engine/Source/Editor/MovieSceneTools/Private/MovieSceneToolsModule.cpp
@@ -0,0 +1,94 @@
+// Copyright 1998-2015 Epic Games, Inc. All Rights Reserved.
+
+#include "MovieSceneToolsPrivatePCH.h"
+#include "ModuleManager.h"
+
+#include "MovieSceneFactory.h"
+#include "K2Node_PlayMovieScene.h"
+
+#include "ISequencerModule.h"
+#include "Kismet2/BlueprintEditorUtils.h"
+
+#include "MovieSceneSection.h"
+#include "ISequencerSection.h"
+#include "AssetToolsModule.h"
+#include "AssetTypeActions_MovieScene.h"
+#include "ScopedTransaction.h"
+#include "MovieScene.h"
+#include "MovieSceneTrackEditor.h"
+#include "PropertyTrackEditor.h"
+#include "TransformTrackEditor.h"
+#include "DirectorTrackEditor.h"
+#include "SubMovieSceneTrackEditor.h"
+#include "AudioTrackEditor.h"
+#include "AnimationTrackEditor.h"
+#include "ParticleTrackEditor.h"
+
+/**
+ * MovieSceneTools module implementation (private)
+ */
+class FMovieSceneToolsModule : public IMovieSceneTools
+{
+
+	virtual void StartupModule() override
+	{
+		MovieSceneAssetTypeActions = MakeShareable( new FAssetTypeActions_MovieScene );
+		FModuleManager::LoadModuleChecked< FAssetToolsModule >( "AssetTools" ).Get().RegisterAssetTypeActions( MovieSceneAssetTypeActions.ToSharedRef() );
+
+		// Register with the sequencer module that we provide auto-key handlers.
+		ISequencerModule& SequencerModule = FModuleManager::Get().LoadModuleChecked<ISequencerModule>( "Sequencer" );
+		PropertyTrackEditorCreateTrackEditorDelegateHandle      = SequencerModule.RegisterTrackEditor_Handle( FOnCreateTrackEditor::CreateStatic( &FPropertyTrackEditor::CreateTrackEditor ) );
+		TransformTrackEditorCreateTrackEditorDelegateHandle     = SequencerModule.RegisterTrackEditor_Handle( FOnCreateTrackEditor::CreateStatic( &F3DTransformTrackEditor::CreateTrackEditor ) );
+		DirectorTrackEditorCreateTrackEditorDelegateHandle      = SequencerModule.RegisterTrackEditor_Handle( FOnCreateTrackEditor::CreateStatic( &FDirectorTrackEditor::CreateTrackEditor ) );
+		SubMovieSceneTrackEditorCreateTrackEditorDelegateHandle = SequencerModule.RegisterTrackEditor_Handle( FOnCreateTrackEditor::CreateStatic( &FSubMovieSceneTrackEditor::CreateTrackEditor ) );
+		AudioTrackEditorCreateTrackEditorDelegateHandle         = SequencerModule.RegisterTrackEditor_Handle( FOnCreateTrackEditor::CreateStatic( &FAudioTrackEditor::CreateTrackEditor ) );
+		AnimationTrackEditorCreateTrackEditorDelegateHandle     = SequencerModule.RegisterTrackEditor_Handle( FOnCreateTrackEditor::CreateStatic( &FAnimationTrackEditor::CreateTrackEditor ) );
+		ParticleTrackEditorCreateTrackEditorDelegateHandle      = SequencerModule.RegisterTrackEditor_Handle( FOnCreateTrackEditor::CreateStatic( &FParticleTrackEditor::CreateTrackEditor ) );
+	}
+
+
+	virtual void ShutdownModule() override
+	{
+		// Only unregister if the asset tools module is loaded.  We don't want to forcibly load it during shutdown phase.
+		check( MovieSceneAssetTypeActions.IsValid() );
+		if( FModuleManager::Get().IsModuleLoaded( "AssetTools" ) )
+		{
+			FModuleManager::GetModuleChecked< FAssetToolsModule >( "AssetTools" ).Get().UnregisterAssetTypeActions( MovieSceneAssetTypeActions.ToSharedRef() );
+		}
+		MovieSceneAssetTypeActions.Reset();
+
+
+		if( FModuleManager::Get().IsModuleLoaded( "Sequencer" ) )
+		{
+			// Unregister auto key handlers
+			ISequencerModule& SequencerModule = FModuleManager::Get().GetModuleChecked<ISequencerModule>( "Sequencer" );
+			SequencerModule.UnRegisterTrackEditor_Handle( PropertyTrackEditorCreateTrackEditorDelegateHandle );
+			SequencerModule.UnRegisterTrackEditor_Handle( TransformTrackEditorCreateTrackEditorDelegateHandle );
+			SequencerModule.UnRegisterTrackEditor_Handle( DirectorTrackEditorCreateTrackEditorDelegateHandle );
+			SequencerModule.UnRegisterTrackEditor_Handle( SubMovieSceneTrackEditorCreateTrackEditorDelegateHandle );
+			SequencerModule.UnRegisterTrackEditor_Handle( AudioTrackEditorCreateTrackEditorDelegateHandle );
+			SequencerModule.UnRegisterTrackEditor_Handle( AnimationTrackEditorCreateTrackEditorDelegateHandle );
+			SequencerModule.UnRegisterTrackEditor_Handle( ParticleTrackEditorCreateTrackEditorDelegateHandle );
+		}
+		
+	}
+
+
+private:
+
+	/** Asset type actions for MovieScene assets.  Cached here so that we can unregister it during shutdown. */
+	TSharedPtr< FAssetTypeActions_MovieScene > MovieSceneAssetTypeActions;
+
+	/** Registered delegate handles */
+	FDelegateHandle PropertyTrackEditorCreateTrackEditorDelegateHandle;
+	FDelegateHandle TransformTrackEditorCreateTrackEditorDelegateHandle;
+	FDelegateHandle DirectorTrackEditorCreateTrackEditorDelegateHandle;
+	FDelegateHandle SubMovieSceneTrackEditorCreateTrackEditorDelegateHandle;
+	FDelegateHandle AudioTrackEditorCreateTrackEditorDelegateHandle;
+	FDelegateHandle AnimationTrackEditorCreateTrackEditorDelegateHandle;
+	FDelegateHandle ParticleTrackEditorCreateTrackEditorDelegateHandle;
+};
+
+
+IMPLEMENT_MODULE( FMovieSceneToolsModule, MovieSceneTools );
+
diff --git a/Engine/Source/Editor/MovieSceneTools/Private/MovieSceneToolsPrivatePCH.h b/Engine/Source/Editor/MovieSceneTools/Private/MovieSceneToolsPrivatePCH.h
new file mode 100644
--- /dev/null
+++ b/Engine/Source/Editor/MovieSceneTools/Private/MovieSceneToolsPrivatePCH.h
@@ -0,0 +1,15 @@
+// Copyright 1998-2015 Epic Games, Inc. All Rights Reserved.
+
+#ifndef __MovieSceneToolsPrivatePCH_h__
+#define __MovieSceneToolsPrivatePCH_h__
+
+#include "UnrealEd.h"
+#include "MovieScene.h"
+#include "ScopedTransaction.h"
+#include "ISequencerSection.h"
+#include "MovieSceneInstance.h"
+#include "MovieSceneTrackEditor.h"
+#include "IMovieSceneTools.h"
+#include "MovieSceneSection.h"
+
+#endif // __MovieSceneToolsPrivatePCH_h__
diff --git a/Engine/Source/Editor/MovieSceneTools/Private/TrackEditors/AnimationTrackEditor.cpp b/Engine/Source/Editor/MovieSceneTools/Private/TrackEditors/AnimationTrackEditor.cpp
new file mode 100644
--- /dev/null
+++ b/Engine/Source/Editor/MovieSceneTools/Private/TrackEditors/AnimationTrackEditor.cpp
@@ -0,0 +1,278 @@
+// Copyright 1998-2015 Epic Games, Inc. All Rights Reserved.
+
+#include "MovieSceneToolsPrivatePCH.h"
+#include "MovieScene.h"
+#include "MovieSceneSection.h"
+#include "ISequencerSection.h"
+#include "PropertyEditorModule.h"
+#include "PropertyHandle.h"
+#include "MovieSceneTrack.h"
+#include "MovieSceneAnimationTrack.h"
+#include "ScopedTransaction.h"
+#include "ISequencerObjectChangeListener.h"
+#include "ISectionLayoutBuilder.h"
+#include "IKeyArea.h"
+#include "MovieSceneToolHelpers.h"
+#include "MovieSceneTrackEditor.h"
+#include "AnimationTrackEditor.h"
+#include "MovieSceneAnimationSection.h"
+#include "CommonMovieSceneTools.h"
+#include "AssetRegistryModule.h"
+#include "Animation/SkeletalMeshActor.h"
+
+
+
+namespace AnimatableAnimationEditorConstants
+{
+	// @todo Sequencer Allow this to be customizable
+	const uint32 AnimationTrackHeight = 20;
+}
+
+
+
+FAnimationSection::FAnimationSection( UMovieSceneSection& InSection )
+	: Section( InSection )
+{
+}
+
+FAnimationSection::~FAnimationSection()
+{
+}
+
+UMovieSceneSection* FAnimationSection::GetSectionObject()
+{ 
+	return &Section;
+}
+
+FText FAnimationSection::GetDisplayName() const
+{
+	return NSLOCTEXT("FAnimationSection", "AnimationSection", "Animation");
+}
+
+FText FAnimationSection::GetSectionTitle() const
+{
+	return FText::FromString( Cast<UMovieSceneAnimationSection>(&Section)->GetAnimSequence()->GetName() );
+}
+
+float FAnimationSection::GetSectionHeight() const
+{
+	return (float)AnimatableAnimationEditorConstants::AnimationTrackHeight;
+ }
+
+int32 FAnimationSection::OnPaintSection( const FGeometry& AllottedGeometry, const FSlateRect& SectionClippingRect, FSlateWindowElementList& OutDrawElements, int32 LayerId, bool bParentEnabled ) const
+{
+	UMovieSceneAnimationSection* AnimSection = Cast<UMovieSceneAnimationSection>(&Section);
+	
+	FTimeToPixel TimeToPixelConverter( AllottedGeometry, TRange<float>( Section.GetStartTime(), Section.GetEndTime() ) );
+
+	// Add a box for the section
+	FSlateDrawElement::MakeBox(
+		OutDrawElements,
+		LayerId,
+		AllottedGeometry.ToPaintGeometry(),
+		FEditorStyle::GetBrush("Sequencer.GenericSection.Background"),
+		SectionClippingRect,
+		ESlateDrawEffect::None,
+		FLinearColor(0.7f, 0.4f, 0.7f, 1.f)
+	);
+
+	// Darken the part that doesn't have animation
+	if (AnimSection->GetAnimationStartTime() > AnimSection->GetStartTime())
+	{
+		float StartDarkening = AnimSection->GetStartTime();
+		float EndDarkening = FMath::Min(AnimSection->GetAnimationStartTime(), AnimSection->GetEndTime());
+		
+		float StartPixels = TimeToPixelConverter.TimeToPixel(StartDarkening);
+		float EndPixels = TimeToPixelConverter.TimeToPixel(EndDarkening);
+
+		FSlateDrawElement::MakeBox(
+			OutDrawElements,
+			LayerId + 1,
+			AllottedGeometry.ToPaintGeometry(FVector2D(StartPixels, 0), FVector2D(EndPixels - StartPixels, AllottedGeometry.Size.Y)),
+			FEditorStyle::GetBrush("WhiteTexture"),
+			SectionClippingRect,
+			ESlateDrawEffect::None,
+			FLinearColor(0.f, 0.f, 0.f, 0.3f)
+		);
+	}
+
+	// Add lines where the animation starts and ends/loops
+	float CurrentTime = AnimSection->GetAnimationStartTime();
+	while (CurrentTime < AnimSection->GetEndTime())
+	{
+		if (CurrentTime > AnimSection->GetStartTime())
+		{
+			float CurrentPixels = TimeToPixelConverter.TimeToPixel(CurrentTime);
+
+			TArray<FVector2D> Points;
+			Points.Add(FVector2D(CurrentPixels, 0));
+			Points.Add(FVector2D(CurrentPixels, AllottedGeometry.Size.Y));
+
+			FSlateDrawElement::MakeLines(
+				OutDrawElements,
+				LayerId + 2,
+				AllottedGeometry.ToPaintGeometry(),
+				Points,
+				SectionClippingRect
+			);
+		}
+		CurrentTime += AnimSection->GetAnimationDuration();
+	}
+
+	return LayerId+3;
+}
+
+FAnimationTrackEditor::FAnimationTrackEditor( TSharedRef<ISequencer> InSequencer )
+	: FMovieSceneTrackEditor( InSequencer ) 
+{
+}
+
+FAnimationTrackEditor::~FAnimationTrackEditor()
+{
+}
+
+TSharedRef<FMovieSceneTrackEditor> FAnimationTrackEditor::CreateTrackEditor( TSharedRef<ISequencer> InSequencer )
+{
+	return MakeShareable( new FAnimationTrackEditor( InSequencer ) );
+}
+
+bool FAnimationTrackEditor::SupportsType( TSubclassOf<UMovieSceneTrack> Type ) const
+{
+	return Type == UMovieSceneAnimationTrack::StaticClass();
+}
+
+TSharedRef<ISequencerSection> FAnimationTrackEditor::MakeSectionInterface( UMovieSceneSection& SectionObject, UMovieSceneTrack* Track )
+{
+	check( SupportsType( SectionObject.GetOuter()->GetClass() ) );
+	
+	TSharedRef<ISequencerSection> NewSection( new FAnimationSection(SectionObject) );
+
+	return NewSection;
+}
+
+void FAnimationTrackEditor::AddKey(const FGuid& ObjectGuid, UObject* AdditionalAsset)
+{
+	UAnimSequence* AnimSequence = Cast<UAnimSequence>(AdditionalAsset);
+
+	if (AnimSequence)
+	{
+		TArray<UObject*> OutObjects;
+		GetSequencer()->GetRuntimeObjects( GetSequencer()->GetFocusedMovieSceneInstance(), ObjectGuid, OutObjects);
+
+		AnimatablePropertyChanged( UMovieSceneAnimationTrack::StaticClass(), false,
+			FOnKeyProperty::CreateRaw( this, &FAnimationTrackEditor::AddKeyInternal, OutObjects, AnimSequence) );
+	}
+}
+
+bool FAnimationTrackEditor::HandleAssetAdded(UObject* Asset, const FGuid& TargetObjectGuid)
+{
+	if (Asset->IsA<UAnimSequence>())
+	{
+		UAnimSequence* AnimSequence = Cast<UAnimSequence>(Asset);
+		
+		if (TargetObjectGuid.IsValid())
+		{
+			USkeleton* Skeleton = AcquireSkeletonFromObjectGuid(TargetObjectGuid);
+
+			if (Skeleton && Skeleton == AnimSequence->GetSkeleton())
+			{
+				TArray<UObject*> OutObjects;
+				GetSequencer()->GetRuntimeObjects(GetSequencer()->GetFocusedMovieSceneInstance(), TargetObjectGuid, OutObjects);
+
+				AnimatablePropertyChanged(UMovieSceneAnimationTrack::StaticClass(), false,
+					FOnKeyProperty::CreateRaw(this, &FAnimationTrackEditor::AddKeyInternal, OutObjects, AnimSequence));
+
+				return true;
+			}
+		}
+	}
+	return false;
+}
+
+void FAnimationTrackEditor::BuildObjectBindingContextMenu(FMenuBuilder& MenuBuilder, const FGuid& ObjectBinding, const UClass* ObjectClass)
+{
+	if (ObjectClass->IsChildOf(ASkeletalMeshActor::StaticClass()))
+	{
+		const TSharedPtr<ISequencer> ParentSequencer = GetSequencer();
+
+		USkeleton* Skeleton = AcquireSkeletonFromObjectGuid(ObjectBinding);
+
+		if (Skeleton)
+		{
+			// Load the asset registry module
+			FAssetRegistryModule& AssetRegistryModule = FModuleManager::LoadModuleChecked<FAssetRegistryModule>(TEXT("AssetRegistry"));
+
+			// Collect a full list of assets with the specified class
+			TArray<FAssetData> AssetDataList;
+			AssetRegistryModule.Get().GetAssetsByClass(UAnimSequence::StaticClass()->GetFName(), AssetDataList);
+
+			for( int32 AssetIndex = 0; AssetIndex < AssetDataList.Num(); ++AssetIndex )
+			{
+				const FAssetData& PossibleAnimSequence = AssetDataList[AssetIndex];
+				if( FAssetData(Skeleton).GetExportTextName() == PossibleAnimSequence.TagsAndValues.FindRef("Skeleton") )
+				{
+					FFormatNamedArguments Args;
+					Args.Add( TEXT("AnimationName"), FText::FromName( PossibleAnimSequence.AssetName ) );
+
+					UAnimSequence* AnimSequence = CastChecked<UAnimSequence>(AssetDataList[AssetIndex].GetAsset());
+					MenuBuilder.AddMenuEntry(
+						FText::Format( NSLOCTEXT("Sequencer", "AddAnimSequence", "Add {AnimationName}"), Args ),
+						FText::Format( NSLOCTEXT("Sequencer", "AddAnimSequenceTooltip", "Adds a {AnimationName} animation to this skeletal mesh."), Args ),
+						FSlateIcon(),
+						FUIAction(FExecuteAction::CreateSP(ParentSequencer.Get(), &ISequencer::AddAnimation, ObjectBinding, AnimSequence))
+						);
+				}
+			}
+		}
+	}
+}
+
+void FAnimationTrackEditor::AddKeyInternal( float KeyTime, const TArray<UObject*> Objects, class UAnimSequence* AnimSequence )
+{
+	for( int32 ObjectIndex = 0; ObjectIndex < Objects.Num(); ++ObjectIndex )
+	{
+		UObject* Object = Objects[ObjectIndex];
+
+		FGuid ObjectHandle = FindOrCreateHandleToObject( Object );
+		if (ObjectHandle.IsValid())
+		{
+			UMovieSceneTrack* Track = GetTrackForObject( ObjectHandle, UMovieSceneAnimationTrack::StaticClass(), FName("Animation"));
+
+			if (ensure(Track))
+			{
+				Cast<UMovieSceneAnimationTrack>(Track)->AddNewAnimation( KeyTime, AnimSequence );
+			}
+		}
+	}
+}
+
+USkeleton* FAnimationTrackEditor::AcquireSkeletonFromObjectGuid(const FGuid& Guid)
+{
+	TArray<UObject*> OutObjects;
+	GetSequencer()->GetRuntimeObjects(GetSequencer()->GetFocusedMovieSceneInstance(), Guid, OutObjects);
+
+	USkeleton* Skeleton = NULL;
+	for (int32 i = 0; i < OutObjects.Num(); ++i)
+	{
+		AActor* Actor = Cast<AActor>(OutObjects[i]);
+
+		if (Actor != NULL)
+		{
+			TInlineComponentArray<USkeletalMeshComponent*> SkeletalMeshComponents;
+			Actor->GetComponents(SkeletalMeshComponents);
+
+			for (int32 j = 0; j <SkeletalMeshComponents.Num(); ++j)
+			{
+				USkeletalMeshComponent* SkeletalMeshComp = SkeletalMeshComponents[j];
+				if (SkeletalMeshComp->SkeletalMesh && SkeletalMeshComp->SkeletalMesh->Skeleton)
+				{
+					// @todo Multiple actors, multiple components
+					check(!Skeleton);
+					Skeleton = SkeletalMeshComp->SkeletalMesh->Skeleton;
+				}
+			}
+		}
+	}
+
+	return Skeleton;
+}
diff --git a/Engine/Source/Editor/MovieSceneTools/Private/TrackEditors/AnimationTrackEditor.h b/Engine/Source/Editor/MovieSceneTools/Private/TrackEditors/AnimationTrackEditor.h
new file mode 100644
--- /dev/null
+++ b/Engine/Source/Editor/MovieSceneTools/Private/TrackEditors/AnimationTrackEditor.h
@@ -0,0 +1,63 @@
+// Copyright 1998-2015 Epic Games, Inc. All Rights Reserved.
+
+#pragma once
+
+
+/**
+ * Tools for animation tracks
+ */
+class FAnimationTrackEditor : public FMovieSceneTrackEditor
+{
+public:
+	/**
+	 * Constructor
+	 *
+	 * @param InSequencer	The sequencer instance to be used by this tool
+	 */
+	FAnimationTrackEditor( TSharedRef<ISequencer> InSequencer );
+	~FAnimationTrackEditor();
+
+	/**
+	 * Creates an instance of this class.  Called by a sequencer 
+	 *
+	 * @param OwningSequencer The sequencer instance to be used by this tool
+	 * @return The new instance of this class
+	 */
+	static TSharedRef<FMovieSceneTrackEditor> CreateTrackEditor( TSharedRef<ISequencer> OwningSequencer );
+
+	/** FMovieSceneTrackEditor Interface */
+	virtual bool SupportsType( TSubclassOf<UMovieSceneTrack> Type ) const override;
+	virtual TSharedRef<ISequencerSection> MakeSectionInterface( UMovieSceneSection& SectionObject, UMovieSceneTrack* Track ) override;
+	virtual void AddKey(const FGuid& ObjectGuid, UObject* AdditionalAsset) override;
+	virtual bool HandleAssetAdded(UObject* Asset, const FGuid& TargetObjectGuid) override;
+	virtual void BuildObjectBindingContextMenu(FMenuBuilder& MenuBuilder, const FGuid& ObjectBinding, const UClass* ObjectClass) override;
+
+private:
+	/** Delegate for AnimatablePropertyChanged in AddKey */
+	void AddKeyInternal(float KeyTime, const TArray<UObject*> Objects, class UAnimSequence* AnimSequence);
+
+	/** Gets a skeleton from an object guid in the movie scene */
+	class USkeleton* AcquireSkeletonFromObjectGuid(const FGuid& Guid);
+};
+
+
+
+/** Class for animation sections, handles drawing of all waveform previews */
+class FAnimationSection : public ISequencerSection, public TSharedFromThis<FAnimationSection>
+{
+public:
+	FAnimationSection( UMovieSceneSection& InSection );
+	~FAnimationSection();
+
+	/** ISequencerSection interface */
+	virtual UMovieSceneSection* GetSectionObject() override;
+	virtual FText GetDisplayName() const override;
+	virtual FText GetSectionTitle() const override;
+	virtual float GetSectionHeight() const override;
+	virtual void GenerateSectionLayout( class ISectionLayoutBuilder& LayoutBuilder ) const override {}
+	virtual int32 OnPaintSection( const FGeometry& AllottedGeometry, const FSlateRect& SectionClippingRect, FSlateWindowElementList& OutDrawElements, int32 LayerId, bool bParentEnabled ) const override;
+
+private:
+	/** The section we are visualizing */
+	UMovieSceneSection& Section;
+};
diff --git a/Engine/Source/Editor/MovieSceneTools/Private/TrackEditors/AudioTrackEditor.cpp b/Engine/Source/Editor/MovieSceneTools/Private/TrackEditors/AudioTrackEditor.cpp
new file mode 100644
--- /dev/null
+++ b/Engine/Source/Editor/MovieSceneTools/Private/TrackEditors/AudioTrackEditor.cpp
@@ -0,0 +1,519 @@
+// Copyright 1998-2015 Epic Games, Inc. All Rights Reserved.
+
+#include "MovieSceneToolsPrivatePCH.h"
+#include "MovieScene.h"
+#include "MovieSceneSection.h"
+#include "ISequencerSection.h"
+#include "PropertyEditorModule.h"
+#include "PropertyHandle.h"
+#include "MovieSceneTrack.h"
+#include "MovieSceneAudioTrack.h"
+#include "ScopedTransaction.h"
+#include "ISequencerObjectChangeListener.h"
+#include "ISectionLayoutBuilder.h"
+#include "IKeyArea.h"
+#include "MovieSceneToolHelpers.h"
+#include "MovieSceneTrackEditor.h"
+#include "AudioTrackEditor.h"
+#include "MovieSceneAudioSection.h"
+#include "CommonMovieSceneTools.h"
+#include "SoundDefinitions.h"
+#include "Sound/SoundNodeWavePlayer.h"
+#include "Runtime/Engine/Public/Slate/SlateTextures.h"
+#include "ObjectTools.h"
+#include "Runtime/Engine/Public/Slate/SceneViewport.h"
+#include "Runtime/MovieSceneCoreTypes/Classes/MovieSceneShotSection.h"
+#include "Runtime/Engine/Public/AudioDecompress.h"
+#include "Sound/SoundBase.h"
+#include "Sound/SoundWave.h"
+#include "Sound/SoundCue.h"
+
+
+namespace AnimatableAudioEditorConstants
+{
+	// @todo Sequencer Allow this to be customizable
+	const uint32 AudioTrackHeight = 50;
+
+	// Optimization - maximum samples per pixel this sound allows
+	const uint32 MaxSamplesPerPixel = 60;
+}
+
+
+
+/** These utility functions should go away once we start handling sound cues properly */
+USoundWave* DeriveSoundWave(UMovieSceneAudioSection* AudioSection)
+{
+	USoundWave* SoundWave = NULL;
+	
+	USoundBase* Sound = AudioSection->GetSound();
+
+	if (Sound->IsA<USoundWave>())
+	{
+		SoundWave = Cast<USoundWave>(Sound);
+	}
+	else if (Sound->IsA<USoundCue>())
+	{
+		USoundCue* SoundCue = Cast<USoundCue>(Sound);
+
+		// @todo Sequencer - Right now for sound cues, we just use the first sound wave in the cue
+		// In the future, it would be better to properly generate the sound cue's data after forcing determinism
+		const TArray<USoundNode*>& AllNodes = SoundCue->AllNodes;
+		for (int32 i = 0; i < AllNodes.Num() && SoundWave == nullptr; ++i)
+		{
+			if (AllNodes[i]->IsA<USoundNodeWavePlayer>())
+			{
+				SoundWave = Cast<USoundNodeWavePlayer>(AllNodes[i])->GetSoundWave();
+			}
+		}
+	}
+
+	return SoundWave;
+}
+
+float DeriveUnloopedDuration(UMovieSceneAudioSection* AudioSection)
+{
+	USoundWave* SoundWave = DeriveSoundWave(AudioSection);
+
+	const float Duration = (SoundWave ? SoundWave->GetDuration() : 0.f);
+	return Duration == INDEFINITELY_LOOPING_DURATION ? SoundWave->Duration : Duration;
+}
+
+
+
+/**
+ * The audio thumnail, which holds a texture which it can pass back to a viewport to render
+ */
+class FAudioThumbnail : public ISlateViewport, public TSharedFromThis<FAudioThumbnail>
+{
+public:
+	FAudioThumbnail(UMovieSceneSection& InSection, TRange<float> DrawRange, int InTextureSize);
+	~FAudioThumbnail();
+
+	/* ISlateViewport interface */
+	virtual FIntPoint GetSize() const override;
+	virtual class FSlateShaderResource* GetViewportRenderTargetTexture() const override;
+	virtual bool RequiresVsync() const override;
+
+	/** Returns whether this thumbnail has a texture to render */
+	virtual bool ShouldRender() const {return TextureSize > 0;}
+	
+private:
+	/** Generates the waveform preview and dumps it out to the OutBuffer */
+	void GenerateWaveformPreview(TArray<uint8>& OutBuffer, TRange<float> DrawRange);
+
+	/** Given a buffer of data, a PCM lookup, draws a line of the PCM waveform at the specified X of the data buffer */
+	void DrawWaveformLine(TArray<uint8>& OutData, int32 NumChannels, int16* LookupData, int32 LookupStartIndex, int32 LookupEndIndex, int32 LookupSize, int32 X, int32 XAxisHeight, int32 MaxAmplitude);
+
+private:
+	/** The section we are visualizing */
+	UMovieSceneSection& Section;
+
+	/** The Texture RHI that holds the thumbnail */
+	class FSlateTexture2DRHIRef* Texture;
+
+	/** Size of the texture */
+	int32 TextureSize;
+};
+
+FAudioThumbnail::FAudioThumbnail(UMovieSceneSection& InSection, TRange<float> DrawRange, int32 InTextureSize)
+	: Section(InSection)
+	, Texture(NULL)
+	, TextureSize(InTextureSize)
+{
+	UMovieSceneAudioSection* AudioSection = Cast<UMovieSceneAudioSection>(&Section);
+	
+	if (ShouldRender())
+	{
+		uint32 Size = GetSize().X * GetSize().Y * GPixelFormats[PF_B8G8R8A8].BlockBytes;
+		TArray<uint8> RawData;
+		RawData.AddZeroed(Size);
+
+		GenerateWaveformPreview(RawData, DrawRange);
+
+		FSlateTextureDataPtr BulkData(new FSlateTextureData(GetSize().X, GetSize().Y, GPixelFormats[PF_B8G8R8A8].BlockBytes, RawData));
+
+		Texture = new FSlateTexture2DRHIRef(GetSize().X, GetSize().Y, PF_B8G8R8A8, BulkData, TexCreate_Dynamic);
+
+		BeginInitResource( Texture );
+	}
+}
+
+FAudioThumbnail::~FAudioThumbnail()
+{
+	if (ShouldRender())
+	{
+		BeginReleaseResource( Texture );
+
+		FlushRenderingCommands();
+	}
+
+	if (Texture) 
+	{
+		delete Texture;
+	}
+}
+
+FIntPoint FAudioThumbnail::GetSize() const {return FIntPoint(TextureSize, AnimatableAudioEditorConstants::AudioTrackHeight);}
+FSlateShaderResource* FAudioThumbnail::GetViewportRenderTargetTexture() const {return Texture;}
+bool FAudioThumbnail::RequiresVsync() const {return false;}
+
+void FAudioThumbnail::GenerateWaveformPreview(TArray<uint8>& OutData, TRange<float> DrawRange)
+{
+	UMovieSceneAudioSection* AudioSection = Cast<UMovieSceneAudioSection>(&Section);
+
+	USoundWave* SoundWave = DeriveSoundWave(AudioSection);
+	check(SoundWave);
+	
+	check(SoundWave->NumChannels == 1 || SoundWave->NumChannels == 2);
+
+	// @todo Sequencer - This may be fixed when we switch to double precision, but the waveform flickers a bit on resize
+	// If double precision don't solve the issue, this system may have to be re-engineered
+
+	// decompress PCM data if necessary
+	if (SoundWave->RawPCMData == NULL)
+	{
+		// @todo Sequencer optimize - We might want to generate the data when we generate the texture
+		// and then discard the data afterwards, though that might be a perf hit traded for better memory usage
+		FAudioDevice* AudioDevice = GEngine->GetMainAudioDevice();
+		if (AudioDevice)
+		{
+			SoundWave->InitAudioResource(AudioDevice->GetRuntimeFormat(SoundWave));
+			FAsyncAudioDecompress TempDecompress(SoundWave);
+			TempDecompress.StartSynchronousTask();
+		}
+	}
+
+	int16* LookupData = (int16*)SoundWave->RawPCMData;
+	int32 LookupDataSize = SoundWave->RawPCMDataSize;
+	int32 LookupSize = LookupDataSize * sizeof(uint8) / sizeof(int16);
+
+	// @todo Sequencer This fixes looping drawing by pretending we are only dealing with a SoundWave
+	TRange<float> AudioTrueRange = TRange<float>(
+		AudioSection->GetAudioStartTime(),
+		AudioSection->GetAudioStartTime() + DeriveUnloopedDuration(AudioSection) * AudioSection->GetAudioDilationFactor());
+	float TrueRangeSize = AudioTrueRange.Size<float>();
+
+	float DrawRangeSize = DrawRange.Size<float>();
+
+	for (int32 X = 0; X < GetSize().X; ++X)
+	{
+		float LookupTime = ((float)(X - 0.5f) / (float)GetSize().X) * DrawRangeSize + DrawRange.GetLowerBoundValue();
+		float LookupFraction = (LookupTime - AudioTrueRange.GetLowerBoundValue()) / TrueRangeSize;
+		int32 LookupIndex = FMath::TruncToInt(LookupFraction * LookupSize);
+		
+		float NextLookupTime = ((float)(X + 0.5f) / (float)GetSize().X) * DrawRangeSize + DrawRange.GetLowerBoundValue();
+		float NextLookupFraction = (NextLookupTime - AudioTrueRange.GetLowerBoundValue()) / TrueRangeSize;
+		int32 NextLookupIndex = FMath::TruncToInt(NextLookupFraction * LookupSize);
+
+		DrawWaveformLine(OutData, SoundWave->NumChannels, LookupData, LookupIndex, NextLookupIndex, LookupSize, X, GetSize().Y / 2, GetSize().Y / 2);
+	}
+}
+
+void FAudioThumbnail::DrawWaveformLine(TArray<uint8>& OutData, int32 NumChannels, int16* LookupData, int32 LookupStartIndex, int32 LookupEndIndex, int32 LookupSize, int32 X, int32 XAxisHeight, int32 MaxAmplitude)
+{
+	int32 ModifiedLookupStartIndex = NumChannels == 2 ? (LookupStartIndex % 2 == 0 ? LookupStartIndex : LookupStartIndex - 1) : LookupStartIndex;
+
+	int32 ClampedStartIndex = FMath::Max(0, ModifiedLookupStartIndex);
+	int32 ClampedEndIndex = FMath::Max(FMath::Min(LookupSize-1, LookupEndIndex), ClampedStartIndex + 1);
+	
+	int32 StepSize = NumChannels;
+
+	// optimization - don't take more than a maximum number of samples per pixel
+	int32 Range = ClampedEndIndex - ClampedStartIndex;
+	int32 SampleCount = Range / StepSize;
+	int32 MaxSampleCount = AnimatableAudioEditorConstants::MaxSamplesPerPixel;
+	int32 ModifiedStepSize = (SampleCount > MaxSampleCount ? SampleCount / MaxSampleCount : 1) * StepSize;
+
+	// Two techniques for rendering smooth waveforms: Summed Density Histograms and Root Mean Square
+#if 1 // Summed Density Histogram
+	struct FChannelData
+	{
+		FChannelData(int32 InMaxAmplitude)
+			: HistogramSum(0)
+		{
+			Histogram.AddZeroed(InMaxAmplitude);
+		}
+
+		TArray<int32> Histogram;
+		int32 HistogramSum;
+	};
+	TArray<FChannelData> ChannelData;
+	for (int32 i = 0; i < NumChannels; ++i)
+	{
+		ChannelData.Add(FChannelData(MaxAmplitude));
+	}
+
+	// build up the histogram(s)
+	for (int32 i = ClampedStartIndex; i < ClampedEndIndex; i += ModifiedStepSize)
+	{
+		for (int32 ChannelIndex = 0; ChannelIndex < NumChannels; ++ChannelIndex)
+		{
+			int32 DataPoint = LookupData[i + ChannelIndex];
+			int32 HistogramData = FMath::Clamp(FMath::TruncToInt(FMath::Abs(DataPoint) / 32768.f * MaxAmplitude), 0, MaxAmplitude - 1);
+
+			++ChannelData[ChannelIndex].Histogram[HistogramData];
+			++ChannelData[ChannelIndex].HistogramSum;
+		}
+	}
+
+	// output the data to the texture
+	const int32 MaxSupportedChannels = 2;
+	for (int32 ChannelIndex = 0; ChannelIndex < MaxSupportedChannels; ++ChannelIndex)
+	{
+		FChannelData& ChannelDataUsed = ChannelData[ChannelIndex % NumChannels];
+		int32 IntensitySummation = 0;
+		for (int32 PixelIndex = MaxAmplitude - 1; PixelIndex >= 0; --PixelIndex)
+		{
+			IntensitySummation += ChannelDataUsed.Histogram[PixelIndex];
+			if (IntensitySummation > 0)
+			{
+				int32 AlphaIntensity = FMath::Clamp(IntensitySummation * 255 / ChannelDataUsed.HistogramSum, 0, 255);
+
+				int32 Y = ChannelIndex == 0 ? XAxisHeight - PixelIndex : XAxisHeight + PixelIndex;
+				int32 Index = (Y * GetSize().X + X) * GPixelFormats[PF_B8G8R8A8].BlockBytes;
+				OutData[Index+0] = 255;
+				OutData[Index+1] = 255;
+				OutData[Index+2] = 255;
+				OutData[Index+3] = AlphaIntensity;
+			}
+		}
+	}
+#else // Root Mean Square
+	int32 Samples = 0;
+	float SummedData = 0;
+	int32 DataAmplitude = 0;
+	for (int32 i = ClampedStartIndex; i < ClampedEndIndex; i += StepSize)
+	{
+		int32 DataPoint = LookupData[i];
+		SummedData += DataPoint * DataPoint;
+		++Samples;
+		DataAmplitude = FMath::Max(DataAmplitude, FMath::Abs(DataPoint));
+	}
+	float RootMeanSquare = FMath::Sqrt(SummedData / (float)Samples);
+	
+	int32 RMSHeight = FMath::Clamp(FMath::TruncToInt(RootMeanSquare / 32768.f * MaxAmplitude), 1, MaxAmplitude);
+	int32 MaxHeight = FMath::Clamp(FMath::TruncToInt(DataAmplitude / 32768.f * MaxAmplitude), 1, MaxAmplitude);
+	float HeightRange = FMath::Max(1, MaxHeight - RMSHeight);
+	
+	for (int32 i = 0; i < MaxHeight; ++i)
+	{
+		float Intensity = (MaxHeight - i) / HeightRange;
+		int32 IntegralIntensity = FMath::Clamp(FMath::TruncToInt(Intensity * 255), 0, 255);
+
+		int32 Y = XAxisHeight - i;
+		int32 Index = (Y * GetSize().X + X) * GPixelFormats[PF_B8G8R8A8].BlockBytes;
+		OutData[Index+0] = 255;
+		OutData[Index+1] = 255;
+		OutData[Index+2] = 255;
+		OutData[Index+3] = IntegralIntensity;
+
+		Y = XAxisHeight + i;
+		Index = (Y * GetSize().X + X) * GPixelFormats[PF_B8G8R8A8].BlockBytes;
+		OutData[Index+0] = 255;
+		OutData[Index+1] = 255;
+		OutData[Index+2] = 255;
+		OutData[Index+3] = IntegralIntensity;
+	}
+#endif
+}
+
+
+
+FAudioSection::FAudioSection( UMovieSceneSection& InSection, bool bOnAMasterTrack )
+	: Section( InSection )
+	, StoredDrawRange(TRange<float>::Empty())
+	, bIsOnAMasterTrack(bOnAMasterTrack)
+{
+}
+
+FAudioSection::~FAudioSection()
+{
+}
+
+UMovieSceneSection* FAudioSection::GetSectionObject()
+{ 
+	return &Section;
+}
+
+FText FAudioSection::GetDisplayName() const
+{
+	return bIsOnAMasterTrack ? NSLOCTEXT("FAudioSection", "MasterAudioDisplayName", "Master Audio") :  NSLOCTEXT("FAudioSection", "AudioDisplayName", "Audio");
+}
+
+FText FAudioSection::GetSectionTitle() const
+{
+	return FText::FromString( Cast<UMovieSceneAudioSection>(&Section)->GetSound()->GetName() );
+}
+
+float FAudioSection::GetSectionHeight() const
+{
+	return (float)AnimatableAudioEditorConstants::AudioTrackHeight;
+ }
+
+int32 FAudioSection::OnPaintSection( const FGeometry& AllottedGeometry, const FSlateRect& SectionClippingRect, FSlateWindowElementList& OutDrawElements, int32 LayerId, bool bParentEnabled ) const
+{
+	// Add a box for the section
+	FSlateDrawElement::MakeBox(
+		OutDrawElements,
+		LayerId,
+		AllottedGeometry.ToPaintGeometry(),
+		FEditorStyle::GetBrush("Sequencer.GenericSection.Background"),
+		SectionClippingRect,
+		ESlateDrawEffect::None,
+		FLinearColor(0.4f, 0.8f, 0.4f, 1.f)
+	);
+
+	if (WaveformThumbnail.IsValid() && WaveformThumbnail->ShouldRender())
+	{
+		// @todo Sequencer draw multiple times if looping possibly - requires some thought about SoundCues
+		FSlateDrawElement::MakeViewport(
+			OutDrawElements,
+			LayerId+1,
+			AllottedGeometry.ToPaintGeometry(FVector2D(StoredXOffset, 0), FVector2D(StoredXSize, AllottedGeometry.Size.Y)),
+			WaveformThumbnail,
+			SectionClippingRect,
+			false,
+			false,
+			ESlateDrawEffect::None,
+			FLinearColor(1.f, 1.f, 1.f, 0.9f)
+		);
+	}
+
+	return LayerId+1;
+}
+
+void FAudioSection::Tick( const FGeometry& AllottedGeometry, const FGeometry& ParentGeometry, const double InCurrentTime, const float InDeltaTime )
+{
+	UMovieSceneAudioSection* AudioSection = Cast<UMovieSceneAudioSection>(&Section);
+
+	USoundWave* SoundWave = DeriveSoundWave(AudioSection);
+	if (SoundWave->NumChannels == 1 || SoundWave->NumChannels == 2)
+	{
+		FTimeToPixel TimeToPixelConverter( AllottedGeometry, TRange<float>( AudioSection->GetStartTime(), AudioSection->GetEndTime() ) );
+
+		TRange<float> VisibleRange = TRange<float>(
+			TimeToPixelConverter.PixelToTime(-AllottedGeometry.Position.X),
+			TimeToPixelConverter.PixelToTime(-AllottedGeometry.Position.X + ParentGeometry.Size.X));
+		TRange<float> SectionRange = AudioSection->GetRange();
+		TRange<float> AudioRange = TRange<float>(
+			AudioSection->GetAudioStartTime(),
+			AudioSection->GetAudioStartTime() + DeriveUnloopedDuration(AudioSection) * AudioSection->GetAudioDilationFactor());
+
+		TRange<float> DrawRange = TRange<float>::Intersection(TRange<float>::Intersection(SectionRange, AudioRange), VisibleRange);
+	
+		TRange<int32> PixelRange = TRange<int32>(
+			TimeToPixelConverter.TimeToPixel(DrawRange.GetLowerBoundValue()),
+			TimeToPixelConverter.TimeToPixel(DrawRange.GetUpperBoundValue()));
+		
+		// generate texture x offset and x size
+		int32 XOffset = PixelRange.GetLowerBoundValue() - TimeToPixelConverter.TimeToPixel(SectionRange.GetLowerBoundValue());
+		int32 XSize = PixelRange.Size<int32>();
+
+		if (!FMath::IsNearlyEqual(DrawRange.GetLowerBoundValue(), StoredDrawRange.GetLowerBoundValue()) ||
+			!FMath::IsNearlyEqual(DrawRange.GetUpperBoundValue(), StoredDrawRange.GetUpperBoundValue()) ||
+			XOffset != StoredXOffset || XSize != StoredXSize)
+		{
+			RegenerateWaveforms(DrawRange, XOffset, XSize);
+		}
+	}
+}
+
+void FAudioSection::RegenerateWaveforms(TRange<float> DrawRange, int32 XOffset, int32 XSize)
+{
+	StoredDrawRange = DrawRange;
+	StoredXOffset = XOffset;
+	StoredXSize = XSize;
+	
+	if (DrawRange.IsDegenerate() || DrawRange.IsEmpty() || Cast<UMovieSceneAudioSection>(&Section)->GetSound() == NULL)
+	{
+		WaveformThumbnail.Reset();
+	}
+	else
+	{
+		WaveformThumbnail = MakeShareable(new FAudioThumbnail(Section, DrawRange, XSize));
+	}
+}
+
+
+
+FAudioTrackEditor::FAudioTrackEditor( TSharedRef<ISequencer> InSequencer )
+	: FMovieSceneTrackEditor( InSequencer ) 
+{
+}
+
+FAudioTrackEditor::~FAudioTrackEditor()
+{
+}
+
+TSharedRef<FMovieSceneTrackEditor> FAudioTrackEditor::CreateTrackEditor( TSharedRef<ISequencer> InSequencer )
+{
+	return MakeShareable( new FAudioTrackEditor( InSequencer ) );
+}
+
+bool FAudioTrackEditor::SupportsType( TSubclassOf<UMovieSceneTrack> Type ) const
+{
+	return Type == UMovieSceneAudioTrack::StaticClass();
+}
+
+TSharedRef<ISequencerSection> FAudioTrackEditor::MakeSectionInterface( UMovieSceneSection& SectionObject, UMovieSceneTrack* Track )
+{
+	check( SupportsType( SectionObject.GetOuter()->GetClass() ) );
+	
+	bool bIsAMasterTrack = Cast<UMovieScene>(Track->GetOuter())->IsAMasterTrack(Track);
+	TSharedRef<ISequencerSection> NewSection( new FAudioSection(SectionObject, bIsAMasterTrack) );
+
+	return NewSection;
+}
+
+bool FAudioTrackEditor::HandleAssetAdded(UObject* Asset, const FGuid& TargetObjectGuid)
+{
+	if (Asset->IsA<USoundBase>())
+	{
+		auto Sound = Cast<USoundBase>(Asset);
+		
+		if (TargetObjectGuid.IsValid())
+		{
+			TArray<UObject*> OutObjects;
+			GetSequencer()->GetRuntimeObjects(GetSequencer()->GetFocusedMovieSceneInstance(), TargetObjectGuid, OutObjects);
+
+			AnimatablePropertyChanged(UMovieSceneAudioTrack::StaticClass(), false,
+				FOnKeyProperty::CreateRaw(this, &FAudioTrackEditor::AddNewAttachedSound, Sound, OutObjects));
+		}
+		else
+		{
+			AnimatablePropertyChanged(UMovieSceneAudioTrack::StaticClass(), false,
+				FOnKeyProperty::CreateRaw(this, &FAudioTrackEditor::AddNewMasterSound, Sound));
+		}
+
+		return true;
+	}
+	return false;
+}
+
+void FAudioTrackEditor::AddNewMasterSound( float KeyTime, USoundBase* Sound )
+{
+	UMovieSceneTrack* Track = GetMasterTrack( UMovieSceneAudioTrack::StaticClass() );
+
+	auto AudioTrack = Cast<UMovieSceneAudioTrack>(Track);
+	AudioTrack->AddNewSound( Sound, KeyTime );
+}
+
+void FAudioTrackEditor::AddNewAttachedSound( float KeyTime, USoundBase* Sound, TArray<UObject*> ObjectsToAttachTo )
+{
+	for( int32 ObjectIndex = 0; ObjectIndex < ObjectsToAttachTo.Num(); ++ObjectIndex )
+	{
+		UObject* Object = ObjectsToAttachTo[ObjectIndex];
+
+		FGuid ObjectHandle = FindOrCreateHandleToObject( Object );
+		if (ObjectHandle.IsValid())
+		{
+			
+			UMovieSceneTrack* Track = GetTrackForObject( ObjectHandle, UMovieSceneAudioTrack::StaticClass(), AudioTrackConstants::UniqueTrackName );
+
+			if (ensure(Track))
+			{
+				Cast<UMovieSceneAudioTrack>(Track)->AddNewSound( Sound, KeyTime );
+			}
+		}
+	}
+}
diff --git a/Engine/Source/Editor/MovieSceneTools/Private/TrackEditors/AudioTrackEditor.h b/Engine/Source/Editor/MovieSceneTools/Private/TrackEditors/AudioTrackEditor.h
new file mode 100644
--- /dev/null
+++ b/Engine/Source/Editor/MovieSceneTools/Private/TrackEditors/AudioTrackEditor.h
@@ -0,0 +1,76 @@
+// Copyright 1998-2015 Epic Games, Inc. All Rights Reserved.
+
+#pragma once
+
+
+/**
+ * Tools for audio tracks
+ */
+class FAudioTrackEditor : public FMovieSceneTrackEditor
+{
+public:
+	/**
+	 * Constructor
+	 *
+	 * @param InSequencer	The sequencer instance to be used by this tool
+	 */
+	FAudioTrackEditor( TSharedRef<ISequencer> InSequencer );
+	~FAudioTrackEditor();
+
+	/**
+	 * Creates an instance of this class.  Called by a sequencer 
+	 *
+	 * @param OwningSequencer The sequencer instance to be used by this tool
+	 * @return The new instance of this class
+	 */
+	static TSharedRef<FMovieSceneTrackEditor> CreateTrackEditor( TSharedRef<ISequencer> OwningSequencer );
+
+	/** FMovieSceneTrackEditor Interface */
+	virtual bool SupportsType( TSubclassOf<UMovieSceneTrack> Type ) const override;
+	virtual TSharedRef<ISequencerSection> MakeSectionInterface( UMovieSceneSection& SectionObject, UMovieSceneTrack* Track ) override;
+	virtual bool HandleAssetAdded(UObject* Asset, const FGuid& TargetObjectGuid) override;
+	
+private:
+	/** Delegate for AnimatablePropertyChanged in HandleAssetAdded for master sounds */
+	void AddNewMasterSound( float KeyTime, class USoundBase* Sound );
+	/** Delegate for AnimatablePropertyChanged in HandleAssetAdded for attached sounds */
+	void AddNewAttachedSound( float KeyTime, class USoundBase* Sound, TArray<UObject*> ObjectsToAttachTo );
+};
+
+
+
+/** Class for audio sections, handles drawing of all waveform previews */
+class FAudioSection : public ISequencerSection, public TSharedFromThis<FAudioSection>
+{
+public:
+	FAudioSection( UMovieSceneSection& InSection, bool bOnAMasterTrack );
+	~FAudioSection();
+
+	/** ISequencerSection interface */
+	virtual UMovieSceneSection* GetSectionObject() override;
+	virtual FText GetDisplayName() const override;
+	virtual FText GetSectionTitle() const override;
+	virtual float GetSectionHeight() const override;
+	virtual void GenerateSectionLayout( class ISectionLayoutBuilder& LayoutBuilder ) const override {}
+	virtual int32 OnPaintSection( const FGeometry& AllottedGeometry, const FSlateRect& SectionClippingRect, FSlateWindowElementList& OutDrawElements, int32 LayerId, bool bParentEnabled ) const override;
+	virtual void Tick( const FGeometry& AllottedGeometry, const FGeometry& ParentGeometry, const double InCurrentTime, const float InDeltaTime ) override;
+	
+private:
+	/* Re-creates the texture used to preview the waveform */
+	void RegenerateWaveforms(TRange<float> DrawRange, int32 XOffset, int32 XSize);
+
+private:
+	/** The section we are visualizing */
+	UMovieSceneSection& Section;
+	
+	/** The waveform thumbnail render object */
+	TSharedPtr<class FAudioThumbnail> WaveformThumbnail;
+
+	/** Stored data about the waveform to determine when it is invalidated */
+	TRange<float> StoredDrawRange;
+	int32 StoredXOffset;
+	int32 StoredXSize;
+
+	/** Whether this section is on a master audio track or an attached audio track */
+	bool bIsOnAMasterTrack;
+};
diff --git a/Engine/Source/Editor/MovieSceneTools/Private/TrackEditors/DirectorTrackEditor.cpp b/Engine/Source/Editor/MovieSceneTools/Private/TrackEditors/DirectorTrackEditor.cpp
new file mode 100644
--- /dev/null
+++ b/Engine/Source/Editor/MovieSceneTools/Private/TrackEditors/DirectorTrackEditor.cpp
@@ -0,0 +1,464 @@
+// Copyright 1998-2015 Epic Games, Inc. All Rights Reserved.
+
+#include "MovieSceneToolsPrivatePCH.h"
+#include "MovieScene.h"
+#include "MovieSceneSection.h"
+#include "ISequencerSection.h"
+#include "PropertyEditorModule.h"
+#include "PropertyHandle.h"
+#include "MovieSceneTrack.h"
+#include "MovieSceneDirectorTrack.h"
+#include "ScopedTransaction.h"
+#include "ISequencerObjectChangeListener.h"
+#include "ISectionLayoutBuilder.h"
+#include "Runtime/Engine/Public/Slate/SlateTextures.h"
+#include "ObjectTools.h"
+#include "Runtime/Engine/Public/Slate/SceneViewport.h"
+#include "Runtime/MovieSceneCoreTypes/Classes/MovieSceneShotSection.h"
+#include "IKeyArea.h"
+#include "MovieSceneToolHelpers.h"
+#include "MovieSceneTrackEditor.h"
+#include "DirectorTrackEditor.h"
+#include "CommonMovieSceneTools.h"
+#include "Camera/CameraActor.h"
+
+namespace AnimatableShotToolConstants
+{
+	// @todo Sequencer Perhaps allow this to be customizable
+	const uint32 DirectorTrackHeight = 100;
+	const double ThumbnailFadeInDuration = 0.25f;
+
+}
+
+/////////////////////////////////////////////////
+// FShotThumbnailPool
+
+FShotThumbnailPool::FShotThumbnailPool(TSharedPtr<ISequencer> InSequencer, uint32 InMaxThumbnailsToDrawAtATime)
+	: Sequencer(InSequencer)
+	, ThumbnailsNeedingDraw()
+	, MaxThumbnailsToDrawAtATime(InMaxThumbnailsToDrawAtATime)
+{
+}
+
+void FShotThumbnailPool::AddThumbnailsNeedingRedraw(const TArray< TSharedPtr<FShotThumbnail> >& InThumbnails)
+{
+	ThumbnailsNeedingDraw.Append(InThumbnails);
+}
+
+void FShotThumbnailPool::RemoveThumbnailsNeedingRedraw(const TArray< TSharedPtr<FShotThumbnail> >& InThumbnails)
+{
+	for (int32 i = 0; i < InThumbnails.Num(); ++i)
+	{
+		ThumbnailsNeedingDraw.Remove(InThumbnails[i]);
+	}
+}
+
+void FShotThumbnailPool::DrawThumbnails()
+{
+	if (ThumbnailsNeedingDraw.Num() > 0)
+	{
+		// save the global time
+		Sequencer.Pin()->SetPerspectiveViewportPossessionEnabled(false);
+		float SavedTime = Sequencer.Pin()->GetGlobalTime();
+
+		uint32 ThumbnailsDrawn = 0;
+		for (int32 ThumbnailsIndex = 0; ThumbnailsDrawn < MaxThumbnailsToDrawAtATime && ThumbnailsIndex < ThumbnailsNeedingDraw.Num(); ++ThumbnailsIndex)
+		{
+			TSharedPtr<FShotThumbnail> Thumbnail = ThumbnailsNeedingDraw[ThumbnailsIndex];
+			
+			if (Thumbnail->IsVisible())
+			{
+				Thumbnail->DrawThumbnail();
+				++ThumbnailsDrawn;
+
+				ThumbnailsNeedingDraw.Remove(Thumbnail);
+			}
+			else if (!Thumbnail->IsValid())
+			{
+				ensure(0);
+
+				ThumbnailsNeedingDraw.Remove(Thumbnail);
+			}
+		}
+
+		// restore the global time
+		Sequencer.Pin()->SetGlobalTime(SavedTime);
+		Sequencer.Pin()->SetPerspectiveViewportPossessionEnabled(true);
+	}
+}
+
+
+/////////////////////////////////////////////////
+// FShotThumbnail
+
+FShotThumbnail::FShotThumbnail(TSharedPtr<FShotSection> InSection, TRange<float> InTimeRange)
+	: OwningSection(InSection)
+	, Texture(NULL)
+	, TimeRange(InTimeRange)
+	, FadeInStartTime(0.0)
+{
+	Texture = new FSlateTexture2DRHIRef(GetSize().X, GetSize().Y, PF_B8G8R8A8, NULL, TexCreate_Dynamic, true);
+
+	BeginInitResource( Texture );
+
+}
+
+FShotThumbnail::~FShotThumbnail()
+{
+	BeginReleaseResource( Texture );
+
+	FlushRenderingCommands();
+
+	if (Texture) 
+	{
+		delete Texture;
+	}
+}
+
+FIntPoint FShotThumbnail::GetSize() const {return FIntPoint(OwningSection.Pin()->GetThumbnailWidth(), AnimatableShotToolConstants::DirectorTrackHeight);}
+bool FShotThumbnail::RequiresVsync() const {return false;}
+
+FSlateShaderResource* FShotThumbnail::GetViewportRenderTargetTexture() const
+{
+	return Texture;
+}
+
+void FShotThumbnail::DrawThumbnail()
+{
+	OwningSection.Pin()->DrawViewportThumbnail(SharedThis(this));
+	FadeInStartTime = FSlateApplication::Get().GetCurrentTime();
+}
+
+float FShotThumbnail::GetTime() const {return TimeRange.GetLowerBoundValue();}
+
+void FShotThumbnail::CopyTextureIn(FSlateRenderTargetRHI* InTexture)
+{
+	ENQUEUE_UNIQUE_RENDER_COMMAND_TWOPARAMETER( ReadTexture,
+		FSlateRenderTargetRHI*, RenderTarget, InTexture,
+		FSlateTexture2DRHIRef*, TargetTexture, Texture,
+	{
+		RHICmdList.CopyToResolveTarget(RenderTarget->GetRHIRef(), TargetTexture->GetTypedResource(), false, FResolveParams());
+	});
+}
+
+float FShotThumbnail::GetFadeInCurve() const 
+{
+	return (float)FMath::Clamp(FadeInStartTime / AnimatableShotToolConstants::ThumbnailFadeInDuration, 0.0, 1.0);
+}
+
+bool FShotThumbnail::IsVisible() const
+{
+	return OwningSection.IsValid() ? TimeRange.Overlaps(OwningSection.Pin()->GetVisibleTimeRange()) : false;
+}
+
+bool FShotThumbnail::IsValid() const
+{
+	return OwningSection.IsValid();
+}
+
+/////////////////////////////////////////////////
+// FShotSection
+
+FShotSection::FShotSection( TSharedPtr<ISequencer> InSequencer, TSharedPtr<FShotThumbnailPool> InThumbnailPool, UMovieSceneSection& InSection, UObject* InTargetObject )
+	: Section( &InSection )
+	, Sequencer( InSequencer )
+	, Camera(Cast<ACameraActor>(InTargetObject))
+	, ThumbnailPool( InThumbnailPool )
+	, Thumbnails()
+	, ThumbnailWidth(0)
+	, StoredSize(ForceInit)
+	, StoredStartTime(0.f)
+	, VisibleTimeRange(TRange<float>::Empty())
+	, InternalViewportScene(nullptr)
+	, InternalViewportClient(nullptr)
+{
+	if (Camera.IsValid())
+	{
+		// @todo Sequencer optimize This code shouldn't even be called if this is offscreen
+		// the following code could be generated on demand when the section is scrolled onscreen
+		InternalViewportClient = MakeShareable(new FLevelEditorViewportClient(nullptr));
+		InternalViewportClient->ViewportType = LVT_Perspective;
+		InternalViewportClient->bDisableInput = true;
+		InternalViewportClient->bDrawAxes = false;
+		InternalViewportClient->EngineShowFlags = FEngineShowFlags(ESFIM_Game);
+		InternalViewportClient->EngineShowFlags.DisableAdvancedFeatures();
+		InternalViewportClient->SetActorLock(Camera.Get());
+
+		InternalViewportScene = MakeShareable(new FSceneViewport(InternalViewportClient.Get(), nullptr));
+		InternalViewportClient->Viewport = InternalViewportScene.Get();
+
+		CalculateThumbnailWidthAndResize();
+	}
+}
+
+FShotSection::~FShotSection()
+{
+	ThumbnailPool.Pin()->RemoveThumbnailsNeedingRedraw(Thumbnails);
+
+	if (InternalViewportClient.IsValid())
+	{
+		InternalViewportClient->Viewport = nullptr;
+	}
+}
+
+UMovieSceneSection* FShotSection::GetSectionObject() 
+{
+	return Section;
+}
+
+FText FShotSection::GetSectionTitle() const
+{
+	return Cast<UMovieSceneShotSection>(Section)->GetTitle();
+}
+
+float FShotSection::GetSectionHeight() const
+{
+	return AnimatableShotToolConstants::DirectorTrackHeight;
+}
+
+FReply FShotSection::OnSectionDoubleClicked( const FGeometry& SectionGeometry, const FPointerEvent& MouseEvent )
+{
+	if( MouseEvent.GetEffectingButton() == EKeys::LeftMouseButton )
+	{
+		check( Section->IsA<UMovieSceneSection>() );
+
+		Sequencer.Pin()->FilterToSelectedShotSections();
+	}
+
+	return FReply::Handled();
+}
+
+int32 FShotSection::OnPaintSection( const FGeometry& AllottedGeometry, const FSlateRect& SectionClippingRect, FSlateWindowElementList& OutDrawElements, int32 LayerId, bool bParentEnabled ) const
+{
+	if (Camera.IsValid())
+	{
+		for (int32 i = 0; i < Thumbnails.Num(); ++i)
+		{
+			FGeometry TruncatedGeometry = AllottedGeometry.MakeChild(
+				FVector2D(ThumbnailWidth, AllottedGeometry.Size.Y),
+				FSlateLayoutTransform(AllottedGeometry.Scale, FVector2D(i * ThumbnailWidth, 0.0f))
+				);
+			ensure(TruncatedGeometry.Size.Y == AnimatableShotToolConstants::DirectorTrackHeight);
+
+			FSlateDrawElement::MakeViewport(
+				OutDrawElements,
+				LayerId,
+				TruncatedGeometry.ToPaintGeometry(),
+				Thumbnails[i],
+				SectionClippingRect,
+				false,
+				false,
+				ESlateDrawEffect::None,
+				FLinearColor::White
+			);
+
+			FSlateDrawElement::MakeBox( 
+				OutDrawElements,
+				LayerId + 1,
+				TruncatedGeometry.ToPaintGeometry(),
+				FEditorStyle::GetBrush("WhiteBrush"),
+				SectionClippingRect,
+				ESlateDrawEffect::None,
+				FColor(255, 255, 255, Thumbnails[i]->GetFadeInCurve() * 255)
+			);
+		}
+	}
+	else
+	{
+		FSlateDrawElement::MakeBox( 
+			OutDrawElements,
+			LayerId,
+			AllottedGeometry.ToPaintGeometry(),
+			FEditorStyle::GetBrush("Sequencer.GenericSection.Background"),
+			SectionClippingRect
+		); 
+	}
+
+	return LayerId + 1;
+}
+
+void FShotSection::Tick( const FGeometry& AllottedGeometry, const FGeometry& ParentGeometry, const double InCurrentTime, const float InDeltaTime )
+{
+	if (!Camera.IsValid())
+	{
+		return;
+	}
+
+	FTimeToPixel TimeToPixelConverter( AllottedGeometry, TRange<float>( Section->GetStartTime(), Section->GetEndTime() ) );
+
+	// get the visible time range
+	VisibleTimeRange = TRange<float>(
+		TimeToPixelConverter.PixelToTime(-AllottedGeometry.Position.X),
+		TimeToPixelConverter.PixelToTime(-AllottedGeometry.Position.X + ParentGeometry.Size.X));
+
+	FIntPoint AllocatedSize = AllottedGeometry.Size.IntPoint();
+	float StartTime = Section->GetStartTime();
+
+	if (AllocatedSize.X != StoredSize.X || !FMath::IsNearlyEqual(StartTime, StoredStartTime))
+	{
+		RegenerateViewportThumbnails(AllocatedSize);
+	}
+}
+
+uint32 FShotSection::GetThumbnailWidth() const
+{
+	return ThumbnailWidth;
+}
+
+void FShotSection::RegenerateViewportThumbnails(const FIntPoint& Size)
+{
+	StoredSize = Size;
+	StoredStartTime = Section->GetStartTime();
+	
+	ThumbnailPool.Pin()->RemoveThumbnailsNeedingRedraw(Thumbnails);
+	Thumbnails.Empty();
+
+	if (Size.X == 0 || Size.Y == 0)
+	{
+		return;
+	}
+
+	CalculateThumbnailWidthAndResize();
+
+	int32 ThumbnailCount = FMath::DivideAndRoundUp(StoredSize.X, (int32)ThumbnailWidth);
+	
+	float StartTime = Section->GetStartTime();
+	float EndTime = Section->GetEndTime();
+	float DeltaTime = EndTime - StartTime;
+
+	int32 TotalX = StoredSize.X;
+	
+	// @todo sequencer optimize this to use a single viewport and a single thumbnail
+	// instead paste the textures into the thumbnail at specific offsets
+	for (int32 i = 0; i < ThumbnailCount; ++i)
+	{
+		int32 StartX = i * ThumbnailWidth;
+		float FractionThrough = (float)StartX / (float)TotalX;
+		float Time = StartTime + DeltaTime * FractionThrough;
+
+		int32 NextStartX = (i+1) * ThumbnailWidth;
+		float NextFractionThrough = (float)NextStartX / (float)TotalX;
+		float NextTime = FMath::Min(StartTime + DeltaTime * NextFractionThrough, EndTime);
+	
+		check(FractionThrough >= 0.f && FractionThrough <= 1.f && NextFractionThrough >= 0.f);
+		TSharedPtr<FShotThumbnail> NewThumbnail = MakeShareable(new FShotThumbnail(SharedThis(this), TRange<float>(Time, NextTime)));
+
+		Thumbnails.Add(NewThumbnail);
+	}
+
+	// @todo Sequencer Optimize Only say a thumbnail needs redraw if it is onscreen
+	ThumbnailPool.Pin()->AddThumbnailsNeedingRedraw(Thumbnails);
+}
+
+void FShotSection::DrawViewportThumbnail(TSharedPtr<FShotThumbnail> ShotThumbnail)
+{
+	Sequencer.Pin()->SetGlobalTime(ShotThumbnail->GetTime());
+	InternalViewportClient->UpdateViewForLockedActor();
+			
+	GWorld->SendAllEndOfFrameUpdates();
+	InternalViewportScene->Draw(false);
+	ShotThumbnail->CopyTextureIn((FSlateRenderTargetRHI*)InternalViewportScene->GetViewportRenderTargetTexture());
+
+	FlushRenderingCommands();
+}
+
+void FShotSection::CalculateThumbnailWidthAndResize()
+{
+	// @todo Sequencer We also should detect when the property is changed or keyframed to update
+
+	// get the aspect ratio at the first frame
+	// if it changes over time, we ignore this
+	Sequencer.Pin()->SetPerspectiveViewportPossessionEnabled(false);
+	float SavedTime = Sequencer.Pin()->GetGlobalTime();
+	
+	Sequencer.Pin()->SetGlobalTime(Section->GetStartTime());
+	uint32 OutThumbnailWidth = AnimatableShotToolConstants::DirectorTrackHeight * Camera->GetCameraComponent()->AspectRatio;
+
+	// restore the global time
+	Sequencer.Pin()->SetGlobalTime(SavedTime);
+	Sequencer.Pin()->SetPerspectiveViewportPossessionEnabled(true);
+
+	if (OutThumbnailWidth != ThumbnailWidth)
+	{
+		ThumbnailWidth = OutThumbnailWidth;
+		InternalViewportScene->UpdateViewportRHI( false, ThumbnailWidth, AnimatableShotToolConstants::DirectorTrackHeight, EWindowMode::Windowed );
+	}
+}
+
+/////////////////////////////////////////////////
+// FDirectorTrackEditor
+
+FDirectorTrackEditor::FDirectorTrackEditor( TSharedRef<ISequencer> InSequencer )
+	: FMovieSceneTrackEditor( InSequencer ) 
+{
+	ThumbnailPool = MakeShareable(new FShotThumbnailPool(InSequencer));
+}
+
+FDirectorTrackEditor::~FDirectorTrackEditor()
+{
+}
+
+TSharedRef<FMovieSceneTrackEditor> FDirectorTrackEditor::CreateTrackEditor( TSharedRef<ISequencer> InSequencer )
+{
+	return MakeShareable( new FDirectorTrackEditor( InSequencer ) );
+}
+
+bool FDirectorTrackEditor::SupportsType( TSubclassOf<UMovieSceneTrack> Type ) const
+{
+	return Type == UMovieSceneDirectorTrack::StaticClass();
+}
+
+TSharedRef<ISequencerSection> FDirectorTrackEditor::MakeSectionInterface( UMovieSceneSection& SectionObject, UMovieSceneTrack* Track )
+{
+	check( SupportsType( SectionObject.GetOuter()->GetClass() ) );
+	
+	TArray<UObject*> OutObjects;
+	// @todo Sequencer - Sub-MovieScenes The director track may be able to get cameras from sub-movie scenes
+	GetSequencer()->GetRuntimeObjects( GetSequencer()->GetRootMovieSceneInstance(), Cast<const UMovieSceneShotSection>(&SectionObject)->GetCameraGuid(), OutObjects);
+	check(OutObjects.Num() <= 1);
+	TSharedRef<ISequencerSection> NewSection( new FShotSection( GetSequencer(), ThumbnailPool, SectionObject, OutObjects.Num() ? OutObjects[0] : NULL ) );
+
+	return NewSection;
+}
+
+void FDirectorTrackEditor::AddKey(const FGuid& ObjectGuid, UObject* AdditionalAsset)
+{
+	TArray<UObject*> OutObjects;
+	// @todo Sequencer - Sub-MovieScenes The director track may be able to get cameras from sub-movie scenes
+	GetSequencer()->GetRuntimeObjects(  GetSequencer()->GetRootMovieSceneInstance(), ObjectGuid, OutObjects);
+	bool bValidKey = OutObjects.Num() == 1 && OutObjects[0]->IsA<ACameraActor>();
+	if (bValidKey)
+	{
+		AnimatablePropertyChanged( UMovieSceneDirectorTrack::StaticClass(), false, FOnKeyProperty::CreateRaw( this, &FDirectorTrackEditor::AddKeyInternal, ObjectGuid ) );
+	}
+}
+
+void FDirectorTrackEditor::Tick(float DeltaTime)
+{
+	if (FSlateThrottleManager::Get().IsAllowingExpensiveTasks())
+	{
+		ThumbnailPool->DrawThumbnails();
+	}
+}
+
+void FDirectorTrackEditor::BuildObjectBindingContextMenu(FMenuBuilder& MenuBuilder, const FGuid& ObjectBinding, const UClass* ObjectClass)
+{
+	if (ObjectClass->IsChildOf(ACameraActor::StaticClass()))
+	{
+		const TSharedPtr<ISequencer> ParentSequencer = GetSequencer();
+
+		MenuBuilder.AddMenuEntry(
+			NSLOCTEXT("Sequencer", "AddShot", "Add New Shot"),
+			NSLOCTEXT("Sequencer", "AddShotTooltip", "Adds a new shot using this camera at the scrubber location."),
+			FSlateIcon(),
+			FUIAction(FExecuteAction::CreateSP(ParentSequencer.Get(), &ISequencer::AddNewShot, ObjectBinding))
+			);
+	}
+}
+
+void FDirectorTrackEditor::AddKeyInternal( float KeyTime, const FGuid ObjectGuid )
+{
+	UMovieSceneTrack* Track = GetMasterTrack( UMovieSceneDirectorTrack::StaticClass() );
+
+	Cast<UMovieSceneDirectorTrack>(Track)->AddNewShot(ObjectGuid, KeyTime);
+}
diff --git a/Engine/Source/Editor/MovieSceneTools/Private/TrackEditors/DirectorTrackEditor.h b/Engine/Source/Editor/MovieSceneTools/Private/TrackEditors/DirectorTrackEditor.h
new file mode 100644
--- /dev/null
+++ b/Engine/Source/Editor/MovieSceneTools/Private/TrackEditors/DirectorTrackEditor.h
@@ -0,0 +1,182 @@
+// Copyright 1998-2015 Epic Games, Inc. All Rights Reserved.
+
+#pragma once
+
+
+/**
+ * Tools for director tracks
+ */
+class FDirectorTrackEditor : public FMovieSceneTrackEditor
+{
+public:
+	/**
+	 * Constructor
+	 *
+	 * @param InSequencer	The sequencer instance to be used by this tool
+	 */
+	FDirectorTrackEditor( TSharedRef<ISequencer> InSequencer );
+	~FDirectorTrackEditor();
+
+	/**
+	 * Creates an instance of this class.  Called by a sequencer 
+	 *
+	 * @param OwningSequencer The sequencer instance to be used by this tool
+	 * @return The new instance of this class
+	 */
+	static TSharedRef<FMovieSceneTrackEditor> CreateTrackEditor( TSharedRef<ISequencer> OwningSequencer );
+
+	/** FMovieSceneTrackEditor Interface */
+	virtual bool SupportsType( TSubclassOf<UMovieSceneTrack> Type ) const override;
+	virtual TSharedRef<ISequencerSection> MakeSectionInterface( UMovieSceneSection& SectionObject, UMovieSceneTrack* Track ) override;
+	virtual void AddKey(const FGuid& ObjectGuid, UObject* AdditionalAsset = NULL) override;
+	virtual void Tick(float DeltaTime) override;
+	virtual void BuildObjectBindingContextMenu(FMenuBuilder& MenuBuilder, const FGuid& ObjectBinding, const UClass* ObjectClass) override;
+
+private:
+	/** Delegate for AnimatablePropertyChanged in AddKey */
+	void AddKeyInternal(float AutoKeyTime, const FGuid ObjectGuid);
+
+private:
+	/** The Thumbnail pool which draws all the viewport thumbnails for the director track */
+	TSharedPtr<class FShotThumbnailPool> ThumbnailPool;
+};
+
+
+/**
+ * Shot Thumbnail pool, which keeps a list of thumbnails that need to be drawn
+ * and draws them incrementally
+ */
+class FShotThumbnailPool
+{
+public:
+	FShotThumbnailPool(TSharedPtr<ISequencer> InSequencer, uint32 InMaxThumbnailsToDrawAtATime = 1);
+
+	/** Requests that the passed in thumbnails need to be drawn */
+	void AddThumbnailsNeedingRedraw(const TArray< TSharedPtr<class FShotThumbnail> >& InThumbnails);
+
+	/** Informs the pool that the thumbnails passed in no longer need to be drawn */
+	void RemoveThumbnailsNeedingRedraw(const TArray< TSharedPtr<class FShotThumbnail> >& InThumbnails);
+
+	/** Draws a small number of thumbnails that are enqueued for drawing */
+	void DrawThumbnails();
+
+private:
+	/** Parent sequencer we're drawing thumbnails for */
+	TWeakPtr<ISequencer> Sequencer;
+
+	/** Thumbnails enqueued for drawing */
+	TArray< TSharedPtr<class FShotThumbnail> > ThumbnailsNeedingDraw;
+
+	/** How many thumbnails we are allowed to draw in a single DrawThumbnails call */
+	uint32 MaxThumbnailsToDrawAtATime;
+};
+
+
+/**
+ * Shot Thumbnail, which keeps a Texture to be displayed by a viewport
+ */
+class FShotThumbnail : public ISlateViewport, public TSharedFromThis<FShotThumbnail>
+{
+public:
+	FShotThumbnail(TSharedPtr<class FShotSection> InSection, TRange<float> InTimeRange);
+	~FShotThumbnail();
+
+	/* ISlateViewport interface */
+	virtual FIntPoint GetSize() const override;
+	virtual class FSlateShaderResource* GetViewportRenderTargetTexture() const override;
+	virtual bool RequiresVsync() const override;
+
+	/** Gets the time that this thumbnail is a rendering of */
+	float GetTime() const;
+
+	/** Renders the thumbnail to the texture */
+	void DrawThumbnail();
+
+	/** Copies the incoming render target to this thumbnails texture */
+	void CopyTextureIn(class FSlateRenderTargetRHI* InTexture);
+
+	/** Gets the curve for fading in the thumbnail */
+	float GetFadeInCurve() const;
+
+	/** Returns whether this thumbnail is visible based on the shot section geometry visibility */
+	bool IsVisible() const;
+
+	bool IsValid() const;
+
+private:
+	/** Parent shot section we are a thumbnail of */
+	TWeakPtr<class FShotSection> OwningSection;
+	
+	/** The Texture RHI that holds the thumbnail */
+	class FSlateTexture2DRHIRef* Texture;
+
+	/** Where in time this thumbnail is a rendering of */
+	TRange<float> TimeRange;
+
+	/** The time when the thumbnail started to fade in. */
+	double FadeInStartTime;
+};
+
+
+/**
+ * Shot section, which paints and ticks the appropriate section
+ */
+class FShotSection : public ISequencerSection, public TSharedFromThis<FShotSection>
+{
+public:
+	FShotSection( TSharedPtr<ISequencer> InSequencer, TSharedPtr<FShotThumbnailPool> InThumbnailPool, UMovieSceneSection& InSection, UObject* InTargetObject );
+	~FShotSection();
+
+	/** ISequencerSection interface */
+	virtual UMovieSceneSection* GetSectionObject() override;
+	virtual int32 OnPaintSection( const FGeometry& AllottedGeometry, const FSlateRect& SectionClippingRect, FSlateWindowElementList& OutDrawElements, int32 LayerId, bool bParentEnabled ) const override;
+	virtual void Tick( const FGeometry& AllottedGeometry, const FGeometry& ParentGeometry, const double InCurrentTime, const float InDeltaTime ) override;
+	virtual FText GetDisplayName() const override { return NSLOCTEXT("FShotSection", "DirectorTrack", "Director Track"); }
+	virtual FText GetSectionTitle() const override;
+	virtual float GetSectionHeight() const override;
+	virtual void GenerateSectionLayout( class ISectionLayoutBuilder& LayoutBuilder ) const override {}
+	virtual FReply OnSectionDoubleClicked( const FGeometry& SectionGeometry, const FPointerEvent& MouseEvent ) override;
+
+	/** Gets the thumbnail width */
+	uint32 GetThumbnailWidth() const;
+
+	/** Regenerates all viewports and thumbnails at the new size */
+	void RegenerateViewportThumbnails(const FIntPoint& Size);
+
+	/** Draws the passed in viewport thumbnail and copies it to the thumbnail's texture */
+	void DrawViewportThumbnail(TSharedPtr<FShotThumbnail> ShotThumbnail);
+
+	/** Calculates and sets the thumbnail width, and resizes if it is different than before */
+	void CalculateThumbnailWidthAndResize();
+
+	/** Gets the time range of what in the sequencer is visible */
+	TRange<float> GetVisibleTimeRange() const {return VisibleTimeRange;}
+
+private:
+	/** The section we are visualizing */
+	UMovieSceneSection* Section;
+	/** The parent sequencer we are a part of */
+	TWeakPtr<ISequencer> Sequencer;
+	/** The actual camera actor we are possessing */
+	TWeakObjectPtr<ACameraActor> Camera;
+
+	/** The thumbnail pool that we are sending all of our thumbnails to */
+	TWeakPtr<FShotThumbnailPool> ThumbnailPool;
+
+	/** A list of all thumbnails this shot section has */
+	TArray< TSharedPtr<FShotThumbnail> > Thumbnails;
+	/** The width of our thumbnails */
+	uint32 ThumbnailWidth;
+	/** The stored size of this section in the Slate geometry */
+	FIntPoint StoredSize;
+	/** The stored start time, to query for invalidations */
+	float StoredStartTime;
+
+	/** Cached Time Range of the visible parent section area */
+	TRange<float> VisibleTimeRange;
+	
+	/** An internal viewport scene we use to render the thumbnails with */
+	TSharedPtr<FSceneViewport> InternalViewportScene;
+	/** An internal editor viewport client to render the thumbnails with */
+	TSharedPtr<FLevelEditorViewportClient> InternalViewportClient;
+};
diff --git a/Engine/Source/Editor/MovieSceneTools/Private/TrackEditors/ParticleTrackEditor.cpp b/Engine/Source/Editor/MovieSceneTools/Private/TrackEditors/ParticleTrackEditor.cpp
new file mode 100644
--- /dev/null
+++ b/Engine/Source/Editor/MovieSceneTools/Private/TrackEditors/ParticleTrackEditor.cpp
@@ -0,0 +1,183 @@
+// Copyright 1998-2015 Epic Games, Inc. All Rights Reserved.
+
+#include "MovieSceneToolsPrivatePCH.h"
+#include "MovieScene.h"
+#include "MovieSceneSection.h"
+#include "ISequencerSection.h"
+#include "PropertyEditorModule.h"
+#include "PropertyHandle.h"
+#include "MovieSceneTrack.h"
+#include "MovieSceneParticleTrack.h"
+#include "ScopedTransaction.h"
+#include "ISequencerObjectChangeListener.h"
+#include "ISectionLayoutBuilder.h"
+#include "IKeyArea.h"
+#include "MovieSceneToolHelpers.h"
+#include "MovieSceneTrackEditor.h"
+#include "ParticleTrackEditor.h"
+#include "MovieSceneParticleSection.h"
+#include "CommonMovieSceneTools.h"
+#include "AssetRegistryModule.h"
+#include "Particles/Emitter.h"
+
+
+namespace AnimatableParticleEditorConstants
+{
+	// @todo Sequencer Allow this to be customizable
+	const uint32 ParticleTrackHeight = 20;
+}
+
+
+
+FParticleSection::FParticleSection( UMovieSceneSection& InSection )
+	: Section( InSection )
+{
+}
+
+FParticleSection::~FParticleSection()
+{
+}
+
+UMovieSceneSection* FParticleSection::GetSectionObject()
+{ 
+	return &Section;
+}
+
+FText FParticleSection::GetDisplayName() const
+{
+	return NSLOCTEXT("FParticleSection", "Emitter", "Emitter");
+}
+
+FText FParticleSection::GetSectionTitle() const
+{
+	EParticleKey::Type KeyType = Cast<UMovieSceneParticleSection>(&Section)->GetKeyType();
+	return KeyType == EParticleKey::Toggle ? NSLOCTEXT("FParticleSection", "Toggle", "Toggle") :
+		KeyType == EParticleKey::Trigger ? NSLOCTEXT("FParticleSection", "Trigger", "Trigger") : NSLOCTEXT("FParticleSection", "None", "None");
+}
+
+float FParticleSection::GetSectionHeight() const
+{
+	return (float)AnimatableParticleEditorConstants::ParticleTrackHeight;
+}
+
+int32 FParticleSection::OnPaintSection( const FGeometry& AllottedGeometry, const FSlateRect& SectionClippingRect, FSlateWindowElementList& OutDrawElements, int32 LayerId, bool bParentEnabled ) const
+{
+	UMovieSceneParticleSection* AnimSection = Cast<UMovieSceneParticleSection>(&Section);
+	
+	FTimeToPixel TimeToPixelConverter( AllottedGeometry, TRange<float>( Section.GetStartTime(), Section.GetEndTime() ) );
+	
+	EParticleKey::Type KeyType = Cast<UMovieSceneParticleSection>(&Section)->GetKeyType();
+
+	// @todo Sequencer Particle Triggers should be sections with 0 duration, but a custom graphic which
+	// allows them to be dragged around
+	if (KeyType == EParticleKey::Trigger)
+	{
+		FSlateDrawElement::MakeBox(
+			OutDrawElements,
+			LayerId,
+			AllottedGeometry.ToPaintGeometry(),
+			FEditorStyle::GetBrush("Sequencer.GenericSection.Background"),
+			SectionClippingRect,
+			ESlateDrawEffect::None,
+			FLinearColor(0.6f, 0.4f, 0.3f, 1.f)
+		);
+	}
+	else
+	{
+		FSlateDrawElement::MakeBox(
+			OutDrawElements,
+			LayerId,
+			AllottedGeometry.ToPaintGeometry(),
+			FEditorStyle::GetBrush("Sequencer.GenericSection.Background"),
+			SectionClippingRect,
+			ESlateDrawEffect::None,
+			FLinearColor(0.8f, 0.4f, 0.3f, 1.f)
+		);
+	}
+
+	return LayerId+1;
+}
+
+bool FParticleSection::SectionIsResizable() const
+{
+	return Cast<UMovieSceneParticleSection>(&Section)->GetKeyType() != EParticleKey::Trigger;
+}
+
+
+
+FParticleTrackEditor::FParticleTrackEditor( TSharedRef<ISequencer> InSequencer )
+	: FMovieSceneTrackEditor( InSequencer ) 
+{
+}
+
+FParticleTrackEditor::~FParticleTrackEditor()
+{
+}
+
+TSharedRef<FMovieSceneTrackEditor> FParticleTrackEditor::CreateTrackEditor( TSharedRef<ISequencer> InSequencer )
+{
+	return MakeShareable( new FParticleTrackEditor( InSequencer ) );
+}
+
+bool FParticleTrackEditor::SupportsType( TSubclassOf<UMovieSceneTrack> Type ) const
+{
+	return Type == UMovieSceneParticleTrack::StaticClass();
+}
+
+TSharedRef<ISequencerSection> FParticleTrackEditor::MakeSectionInterface( UMovieSceneSection& SectionObject, UMovieSceneTrack* Track )
+{
+	check( SupportsType( SectionObject.GetOuter()->GetClass() ) );
+	
+	TSharedRef<ISequencerSection> NewSection( new FParticleSection(SectionObject) );
+
+	return NewSection;
+}
+
+void FParticleTrackEditor::BuildObjectBindingContextMenu(FMenuBuilder& MenuBuilder, const FGuid& ObjectBinding, const UClass* ObjectClass)
+{
+	if (ObjectClass->IsChildOf(AEmitter::StaticClass()))
+	{
+		const TSharedPtr<ISequencer> ParentSequencer = GetSequencer();
+
+		MenuBuilder.AddMenuEntry(
+			NSLOCTEXT("Sequencer", "TriggerParticles", "Trigger"),
+			NSLOCTEXT("Sequencer", "TriggerParticlesTooltip", "Adds a trigger event for this particle emitter at the current time."),
+			FSlateIcon(),
+			FUIAction(FExecuteAction::CreateSP(this, &FParticleTrackEditor::AddParticleKey, ObjectBinding, true))
+			);
+		MenuBuilder.AddMenuEntry(
+			NSLOCTEXT("Sequencer", "ToggleParticles", "Toggle"),
+			NSLOCTEXT("Sequencer", "ToggleParticlesTooltip", "Adds a toggle section for this particle emitter at the current time."),
+			FSlateIcon(),
+			FUIAction(FExecuteAction::CreateSP(this, &FParticleTrackEditor::AddParticleKey, ObjectBinding, false))
+			);
+	}
+}
+
+void FParticleTrackEditor::AddParticleKey(const FGuid ObjectGuid, bool bTrigger)
+{
+	TArray<UObject*> OutObjects;
+	GetSequencer()->GetRuntimeObjects( GetSequencer()->GetFocusedMovieSceneInstance(), ObjectGuid, OutObjects );
+
+	AnimatablePropertyChanged( UMovieSceneParticleTrack::StaticClass(), false,
+		FOnKeyProperty::CreateRaw( this, &FParticleTrackEditor::AddKeyInternal, OutObjects, bTrigger) );
+}
+
+void FParticleTrackEditor::AddKeyInternal( float KeyTime, const TArray<UObject*> Objects, bool bTrigger)
+{
+	for( int32 ObjectIndex = 0; ObjectIndex < Objects.Num(); ++ObjectIndex )
+	{
+		UObject* Object = Objects[ObjectIndex];
+
+		FGuid ObjectHandle = FindOrCreateHandleToObject( Object );
+		if (ObjectHandle.IsValid())
+		{
+			UMovieSceneTrack* Track = GetTrackForObject( ObjectHandle, UMovieSceneParticleTrack::StaticClass(), FName("ParticleSystem"));
+
+			if (ensure(Track))
+			{
+				Cast<UMovieSceneParticleTrack>(Track)->AddNewParticleSystem( KeyTime, bTrigger );
+			}
+		}
+	}
+}
diff --git a/Engine/Source/Editor/MovieSceneTools/Private/TrackEditors/ParticleTrackEditor.h b/Engine/Source/Editor/MovieSceneTools/Private/TrackEditors/ParticleTrackEditor.h
new file mode 100644
--- /dev/null
+++ b/Engine/Source/Editor/MovieSceneTools/Private/TrackEditors/ParticleTrackEditor.h
@@ -0,0 +1,61 @@
+// Copyright 1998-2015 Epic Games, Inc. All Rights Reserved.
+
+#pragma once
+
+
+/**
+ * Tools for particle tracks
+ */
+class FParticleTrackEditor : public FMovieSceneTrackEditor, public TSharedFromThis<FParticleTrackEditor>
+{
+public:
+	/**
+	 * Constructor
+	 *
+	 * @param InSequencer	The sequencer instance to be used by this tool
+	 */
+	FParticleTrackEditor( TSharedRef<ISequencer> InSequencer );
+	~FParticleTrackEditor();
+
+	/**
+	 * Creates an instance of this class.  Called by a sequencer 
+	 *
+	 * @param OwningSequencer The sequencer instance to be used by this tool
+	 * @return The new instance of this class
+	 */
+	static TSharedRef<FMovieSceneTrackEditor> CreateTrackEditor( TSharedRef<ISequencer> OwningSequencer );
+
+	/** FMovieSceneTrackEditor Interface */
+	virtual bool SupportsType( TSubclassOf<UMovieSceneTrack> Type ) const override;
+	virtual TSharedRef<ISequencerSection> MakeSectionInterface( UMovieSceneSection& SectionObject, UMovieSceneTrack* Track ) override;
+	virtual void BuildObjectBindingContextMenu(FMenuBuilder& MenuBuilder, const FGuid& ObjectBinding, const UClass* ObjectClass) override;
+	
+	void AddParticleKey(const FGuid ObjectGuid, bool bTrigger);
+
+private:
+	/** Delegate for AnimatablePropertyChanged in AddKey */
+	virtual void AddKeyInternal( float KeyTime, const TArray<UObject*> Objects, bool bTrigger);
+};
+
+
+
+/** Class for particle sections */
+class FParticleSection : public ISequencerSection, public TSharedFromThis<FParticleSection>
+{
+public:
+	FParticleSection( UMovieSceneSection& InSection );
+	~FParticleSection();
+
+	/** ISequencerSection interface */
+	virtual UMovieSceneSection* GetSectionObject() override;
+	virtual FText GetDisplayName() const override;
+	virtual FText GetSectionTitle() const override;
+	virtual float GetSectionHeight() const override;
+	virtual void GenerateSectionLayout( class ISectionLayoutBuilder& LayoutBuilder ) const override {}
+	virtual int32 OnPaintSection( const FGeometry& AllottedGeometry, const FSlateRect& SectionClippingRect, FSlateWindowElementList& OutDrawElements, int32 LayerId, bool bParentEnabled ) const override;
+	virtual bool SectionIsResizable() const override;
+
+private:
+	/** The section we are visualizing */
+	UMovieSceneSection& Section;
+};
diff --git a/Engine/Source/Editor/MovieSceneTools/Private/TrackEditors/PropertyTrackEditor.cpp b/Engine/Source/Editor/MovieSceneTools/Private/TrackEditors/PropertyTrackEditor.cpp
new file mode 100644
--- /dev/null
+++ b/Engine/Source/Editor/MovieSceneTools/Private/TrackEditors/PropertyTrackEditor.cpp
@@ -0,0 +1,577 @@
+// Copyright 1998-2015 Epic Games, Inc. All Rights Reserved.
+
+#include "MovieSceneToolsPrivatePCH.h"
+#include "PropertyEditing.h"
+#include "MovieScene.h"
+#include "MovieSceneTrack.h"
+#include "MovieSceneFloatTrack.h"
+#include "MovieSceneVectorTrack.h"
+#include "MovieSceneBoolTrack.h"
+#include "MovieSceneByteTrack.h"
+#include "MovieSceneColorTrack.h"
+#include "ScopedTransaction.h"
+#include "MovieSceneSection.h"
+#include "ISequencerObjectBindingManager.h"
+#include "ISequencerObjectChangeListener.h"
+#include "ISequencerSection.h"
+#include "ISectionLayoutBuilder.h"
+#include "IKeyArea.h"
+#include "MovieSceneFloatSection.h"
+#include "MovieSceneVectorSection.h"
+#include "MovieSceneBoolSection.h"
+#include "MovieSceneByteSection.h"
+#include "MovieSceneColorSection.h"
+#include "MovieSceneToolHelpers.h"
+#include "PropertyTrackEditor.h"
+#include "CommonMovieSceneTools.h"
+#include "PropertySection.h"
+
+/**
+ * A color section implementation
+ */
+class FColorPropertySection : public FPropertySection
+{
+public:
+	FColorPropertySection( UMovieSceneSection& InSectionObject, FName SectionName )
+		: FPropertySection(InSectionObject, SectionName) {}
+
+	virtual void GenerateSectionLayout( class ISectionLayoutBuilder& LayoutBuilder ) const override
+	{
+		UMovieSceneColorSection* ColorSection = Cast<UMovieSceneColorSection>(&SectionObject);
+
+		LayoutBuilder.AddKeyArea("R", NSLOCTEXT("FColorPropertySection", "RedArea", "Red"), MakeShareable(new FFloatCurveKeyArea(&ColorSection->GetRedCurve(), ColorSection)));
+		LayoutBuilder.AddKeyArea("G", NSLOCTEXT("FColorPropertySection", "GreenArea", "Green"), MakeShareable(new FFloatCurveKeyArea(&ColorSection->GetGreenCurve(), ColorSection)));
+		LayoutBuilder.AddKeyArea("B", NSLOCTEXT("FColorPropertySection", "BlueArea", "Blue"), MakeShareable(new FFloatCurveKeyArea(&ColorSection->GetBlueCurve(), ColorSection)));
+		LayoutBuilder.AddKeyArea("A", NSLOCTEXT("FColorPropertySection", "OpacityArea", "Opacity"), MakeShareable(new FFloatCurveKeyArea(&ColorSection->GetAlphaCurve(), ColorSection)));
+	}
+
+	virtual int32 OnPaintSection( const FGeometry& AllottedGeometry, const FSlateRect& SectionClippingRect, FSlateWindowElementList& OutDrawElements, int32 LayerId, bool bParentEnabled ) const override 
+	{
+		const UMovieSceneColorSection* ColorSection = Cast<const UMovieSceneColorSection>(&SectionObject);
+
+		float StartTime = ColorSection->GetStartTime();
+		float EndTime = ColorSection->GetEndTime();
+		float SectionDuration = EndTime - StartTime;
+
+		if (!FMath::IsNearlyZero(SectionDuration))
+		{
+			LayerId = FPropertySection::OnPaintSection( AllottedGeometry, SectionClippingRect, OutDrawElements, LayerId, bParentEnabled );
+
+			FVector2D GradientSize = FVector2D(AllottedGeometry.Size.X, (AllottedGeometry.Size.Y / 4)-3.0f);
+
+			FPaintGeometry PaintGeometry = AllottedGeometry.ToPaintGeometry(FVector2D(0,0), GradientSize );
+
+			// If we are showing a background pattern and the colors is transparent, draw a checker pattern
+			const FSlateBrush* CheckerBrush = FEditorStyle::GetBrush( "Checker" );
+			FSlateDrawElement::MakeBox(OutDrawElements, LayerId, PaintGeometry, CheckerBrush, SectionClippingRect);
+	
+			TArray<FSlateGradientStop> GradientStops;
+
+			TArray< TKeyValuePair<float, FLinearColor> > ColorKeys;
+			ConsolidateColorCurves(ColorKeys, ColorSection);
+
+			for (int32 i = 0; i < ColorKeys.Num(); ++i)
+			{
+				float Time = ColorKeys[i].Key;
+				FLinearColor Color = ColorKeys[i].Value;
+				float TimeFraction = (Time - StartTime) / SectionDuration;
+
+				GradientStops.Add(FSlateGradientStop(FVector2D(TimeFraction * AllottedGeometry.Size.X, 0),
+					Color ));
+			}
+
+			if (GradientStops.Num() > 0)
+			{
+				FSlateDrawElement::MakeGradient(
+					OutDrawElements,
+					LayerId + 1,
+					PaintGeometry,
+					GradientStops,
+					Orient_Vertical,
+					SectionClippingRect
+				);
+			}
+		}
+
+		return LayerId + 1;
+	}
+
+private:
+	void ConsolidateColorCurves(TArray< TKeyValuePair<float, FLinearColor> >& OutColorKeys, const UMovieSceneColorSection* Section) const
+	{
+		// @todo Sequencer Optimize - This could all get cached, instead of recalculating everything every OnPaint
+
+		const FRichCurve* Curves[4] = {
+			&Section->GetRedCurve(),
+			&Section->GetGreenCurve(),
+			&Section->GetBlueCurve(),
+			&Section->GetAlphaCurve()
+		};
+
+		// @todo Sequencer Optimize - This is a O(n^2) loop!
+		// Our times are floats, which means we can't use a map and
+		// do a quick lookup to see if the keys already exist
+		// because the keys are ordered, we could take advantage of that, however
+		TArray<float> TimesWithKeys;
+		for (int32 i = 0; i < 4; ++i)
+		{
+			const FRichCurve* Curve = Curves[i];
+			for (auto It(Curve->GetKeyIterator()); It; ++It)
+			{
+				float KeyTime = It->Time;
+
+				bool bShouldAddKey = true;
+
+				int32 InsertKeyIndex = INDEX_NONE;
+				for (int32 k = 0; k < TimesWithKeys.Num(); ++k)
+				{
+					if (FMath::IsNearlyEqual(TimesWithKeys[k], KeyTime))
+					{
+						bShouldAddKey = false;
+						break;
+					}
+					else if (TimesWithKeys[k] > KeyTime)
+					{
+						InsertKeyIndex = k;
+						break;
+					}
+				}
+
+				if (InsertKeyIndex == INDEX_NONE && bShouldAddKey)
+				{
+					InsertKeyIndex = TimesWithKeys.Num();
+				}
+
+				if (bShouldAddKey)
+				{
+					TimesWithKeys.Insert(KeyTime, InsertKeyIndex);
+				}
+			}
+		}
+
+		// @todo Sequencer Optimize - This another O(n^2) loop, since Eval is O(n)!
+		for (int32 i = 0; i < TimesWithKeys.Num(); ++i)
+		{
+			OutColorKeys.Add( TKeyValuePair<float, FLinearColor>(TimesWithKeys[i], Section->Eval(TimesWithKeys[i], FLinearColor(0,0,0,0) ) ) );
+		}
+	}
+};
+
+/**
+ * An implementation of float property sections
+ */
+class FFloatPropertySection : public FPropertySection
+{
+public:
+	FFloatPropertySection( UMovieSceneSection& InSectionObject, FName SectionName )
+		: FPropertySection(InSectionObject, SectionName) {}
+
+	virtual void GenerateSectionLayout( class ISectionLayoutBuilder& LayoutBuilder ) const override
+	{
+		UMovieSceneFloatSection* FloatSection = Cast<UMovieSceneFloatSection>( &SectionObject );
+		LayoutBuilder.SetSectionAsKeyArea( MakeShareable( new FFloatCurveKeyArea( &FloatSection->GetFloatCurve(), FloatSection ) ) );
+	}
+};
+
+/**
+ * An implementation of vector property sections
+ */
+class FVectorPropertySection : public FPropertySection
+{
+public:
+	FVectorPropertySection( UMovieSceneSection& InSectionObject, FName SectionName )
+		: FPropertySection(InSectionObject, SectionName) {}
+
+	virtual void GenerateSectionLayout( class ISectionLayoutBuilder& LayoutBuilder ) const override
+	{
+		UMovieSceneVectorSection* VectorSection = Cast<UMovieSceneVectorSection>( &SectionObject );
+		int32 ChannelsUsed = VectorSection->GetChannelsUsed();
+		check(ChannelsUsed >= 2 && ChannelsUsed <= 4);
+
+		LayoutBuilder.AddKeyArea("Vector.X", NSLOCTEXT("FVectorPropertySection", "XArea", "X"), MakeShareable(new FFloatCurveKeyArea(&VectorSection->GetCurve(0), VectorSection)));
+		LayoutBuilder.AddKeyArea("Vector.Y", NSLOCTEXT("FVectorPropertySection", "YArea", "Y"), MakeShareable(new FFloatCurveKeyArea(&VectorSection->GetCurve(1), VectorSection)));
+		if (ChannelsUsed >= 3)
+		{
+			LayoutBuilder.AddKeyArea("Vector.Z", NSLOCTEXT("FVectorPropertySection", "ZArea", "Z"), MakeShareable(new FFloatCurveKeyArea(&VectorSection->GetCurve(2), VectorSection)));
+		}
+		if (ChannelsUsed >= 4)
+		{
+			LayoutBuilder.AddKeyArea("Vector.W", NSLOCTEXT("FVectorPropertySection", "WArea", "W"), MakeShareable(new FFloatCurveKeyArea(&VectorSection->GetCurve(3), VectorSection)));
+		}
+	}
+};
+
+/**
+ * An implementation of bool property sections
+ */
+class FBoolPropertySection : public FPropertySection
+{
+public:
+	FBoolPropertySection( UMovieSceneSection& InSectionObject, FName SectionName )
+		: FPropertySection(InSectionObject, SectionName) {}
+
+	virtual void GenerateSectionLayout( class ISectionLayoutBuilder& LayoutBuilder ) const override
+	{
+		UMovieSceneBoolSection* BoolSection = Cast<UMovieSceneBoolSection>( &SectionObject );
+		LayoutBuilder.SetSectionAsKeyArea( MakeShareable( new FIntegralKeyArea( BoolSection->GetCurve(), BoolSection ) ) );
+	}
+
+	virtual int32 OnPaintSection( const FGeometry& AllottedGeometry, const FSlateRect& SectionClippingRect, FSlateWindowElementList& OutDrawElements, int32 LayerId, bool bParentEnabled ) const override 
+	{
+		FTimeToPixel TimeToPixelConverter( AllottedGeometry, TRange<float>( SectionObject.GetStartTime(), SectionObject.GetEndTime() ) );
+
+		// custom drawing for bool curves
+		UMovieSceneBoolSection* BoolSection = Cast<UMovieSceneBoolSection>( &SectionObject );
+
+		TArray<float> SectionSwitchTimes;
+		SectionSwitchTimes.Add(SectionObject.GetStartTime());
+		FIntegralCurve& BoolCurve = BoolSection->GetCurve();
+		for (TArray<FIntegralKey>::TConstIterator It(BoolCurve.GetKeyIterator()); It; ++It)
+		{
+			float Time = It->Time;
+			if (Time > SectionObject.GetStartTime() && Time < SectionObject.GetEndTime())
+			{
+				SectionSwitchTimes.Add(Time);
+			}
+		}
+		SectionSwitchTimes.Add(SectionObject.GetEndTime());
+
+		for (int32 i = 0; i < SectionSwitchTimes.Num() - 1; ++i)
+		{
+			float FirstTime = SectionSwitchTimes[i];
+			float NextTime = SectionSwitchTimes[i+1];
+			float FirstTimeInPixels = TimeToPixelConverter.TimeToPixel(FirstTime);
+			float NextTimeInPixels = TimeToPixelConverter.TimeToPixel(NextTime);
+			if (BoolSection->Eval(FirstTime))
+			{
+				FSlateDrawElement::MakeBox( 
+					OutDrawElements,
+					LayerId,
+					AllottedGeometry.ToPaintGeometry(FVector2D(FirstTimeInPixels, 0), FVector2D(NextTimeInPixels - FirstTimeInPixels, AllottedGeometry.Size.Y)),
+					FEditorStyle::GetBrush("Sequencer.GenericSection.Background"),
+					SectionClippingRect,
+					ESlateDrawEffect::None,
+					FLinearColor(0.f, 1.f, 0.f, 1.f)
+				);
+			}
+			else
+			{
+				FSlateDrawElement::MakeBox( 
+					OutDrawElements,
+					LayerId,
+					AllottedGeometry.ToPaintGeometry(FVector2D(FirstTimeInPixels, 0), FVector2D(NextTimeInPixels - FirstTimeInPixels, AllottedGeometry.Size.Y)),
+					FEditorStyle::GetBrush("Sequencer.GenericSection.Background"),
+					SectionClippingRect,
+					ESlateDrawEffect::None,
+					FLinearColor(1.f, 0.f, 0.f, 1.f)
+				);
+			}
+		}
+
+		return LayerId+1;
+	}
+};
+
+/**
+* An implementation of byte property sections
+*/
+class FBytePropertySection : public FPropertySection
+{
+public:
+	FBytePropertySection(UMovieSceneSection& InSectionObject, FName InSectionName, UEnum* InEnum)
+		: FPropertySection(InSectionObject, InSectionName)
+		, Enum(InEnum) {}
+
+	virtual void GenerateSectionLayout(class ISectionLayoutBuilder& LayoutBuilder) const override
+	{
+		UMovieSceneByteSection* ByteSection = Cast<UMovieSceneByteSection>(&SectionObject);
+		if (Enum != nullptr)
+		{
+			LayoutBuilder.SetSectionAsKeyArea(MakeShareable(new FEnumKeyArea(ByteSection->GetCurve(), ByteSection, Enum)));
+		}
+		else
+		{
+			LayoutBuilder.SetSectionAsKeyArea(MakeShareable(new FIntegralKeyArea(ByteSection->GetCurve(), ByteSection)));
+		}
+	}
+
+private:
+	UEnum* Enum;
+};
+
+FPropertyTrackEditor::FPropertyTrackEditor( TSharedRef<ISequencer> InSequencer )
+	: FMovieSceneTrackEditor( InSequencer ) 
+{
+	// Get the object change listener for the sequencer and register a delegates for when properties change that we care about
+	ISequencerObjectChangeListener& ObjectChangeListener = InSequencer->GetObjectChangeListener();
+	ObjectChangeListener.GetOnAnimatablePropertyChanged( NAME_FloatProperty ).AddRaw( this, &FPropertyTrackEditor::OnAnimatedPropertyChanged<float, UMovieSceneFloatTrack> );
+	ObjectChangeListener.GetOnAnimatablePropertyChanged( NAME_BoolProperty ).AddRaw( this, &FPropertyTrackEditor::OnAnimatedBoolPropertyChanged<bool, UMovieSceneBoolTrack> );
+	ObjectChangeListener.GetOnAnimatablePropertyChanged( NAME_ByteProperty ).AddRaw( this, &FPropertyTrackEditor::OnAnimatedPropertyChanged<uint8, UMovieSceneByteTrack> );
+	ObjectChangeListener.GetOnAnimatablePropertyChanged( NAME_Vector ).AddRaw( this, &FPropertyTrackEditor::OnAnimatedVectorPropertyChanged );
+	ObjectChangeListener.GetOnAnimatablePropertyChanged( NAME_Vector4 ).AddRaw( this, &FPropertyTrackEditor::OnAnimatedVectorPropertyChanged );
+	ObjectChangeListener.GetOnAnimatablePropertyChanged( NAME_Vector2D ).AddRaw( this, &FPropertyTrackEditor::OnAnimatedVectorPropertyChanged );
+	ObjectChangeListener.GetOnAnimatablePropertyChanged( NAME_Color ).AddRaw( this, &FPropertyTrackEditor::OnAnimatedColorPropertyChanged );
+	ObjectChangeListener.GetOnAnimatablePropertyChanged( NAME_LinearColor ).AddRaw( this, &FPropertyTrackEditor::OnAnimatedColorPropertyChanged );
+	ObjectChangeListener.GetOnAnimatablePropertyChanged( "SlateColor").AddRaw( this, &FPropertyTrackEditor::OnAnimatedColorPropertyChanged );
+
+}
+
+FPropertyTrackEditor::~FPropertyTrackEditor()
+{
+	TSharedPtr<ISequencer> Sequencer = GetSequencer();
+	if( Sequencer.IsValid() )
+	{
+		ISequencerObjectChangeListener& ObjectChangeListener = Sequencer->GetObjectChangeListener();
+		ObjectChangeListener.GetOnAnimatablePropertyChanged(NAME_FloatProperty).RemoveAll(this);
+		ObjectChangeListener.GetOnAnimatablePropertyChanged(NAME_BoolProperty).RemoveAll(this);
+		ObjectChangeListener.GetOnAnimatablePropertyChanged(NAME_ByteProperty).RemoveAll(this);
+		ObjectChangeListener.GetOnAnimatablePropertyChanged(NAME_Vector).RemoveAll(this);
+		ObjectChangeListener.GetOnAnimatablePropertyChanged(NAME_Vector4).RemoveAll(this);
+		ObjectChangeListener.GetOnAnimatablePropertyChanged(NAME_Vector2D).RemoveAll(this);
+		ObjectChangeListener.GetOnAnimatablePropertyChanged(NAME_Color).RemoveAll(this);
+		ObjectChangeListener.GetOnAnimatablePropertyChanged(NAME_LinearColor).RemoveAll(this);
+		ObjectChangeListener.GetOnAnimatablePropertyChanged("SlateColor").RemoveAll(this);
+	}
+}
+
+
+
+TSharedRef<FMovieSceneTrackEditor> FPropertyTrackEditor::CreateTrackEditor( TSharedRef<ISequencer> InSequencer )
+{
+	return MakeShareable( new FPropertyTrackEditor( InSequencer ) );
+}
+
+bool FPropertyTrackEditor::SupportsType( TSubclassOf<UMovieSceneTrack> Type ) const
+{
+	// We support animatable floats and bools and colors
+	return Type == UMovieSceneFloatTrack::StaticClass() ||
+		Type == UMovieSceneBoolTrack::StaticClass() ||
+		Type == UMovieSceneByteTrack::StaticClass() ||
+		Type == UMovieSceneVectorTrack::StaticClass() ||
+		Type == UMovieSceneColorTrack::StaticClass();
+}
+
+TSharedRef<ISequencerSection> FPropertyTrackEditor::MakeSectionInterface( UMovieSceneSection& SectionObject, UMovieSceneTrack* Track )
+{
+	check( SupportsType( SectionObject.GetOuter()->GetClass() ) );
+
+	UClass* SectionClass = SectionObject.GetOuter()->GetClass();
+
+	TSharedRef<ISequencerSection> NewSection = TSharedRef<ISequencerSection>(
+		SectionClass == UMovieSceneColorTrack::StaticClass() ? new FColorPropertySection( SectionObject, Track->GetTrackName() ) :
+		SectionClass == UMovieSceneBoolTrack::StaticClass() ? new FBoolPropertySection( SectionObject, Track->GetTrackName() ) :
+		SectionClass == UMovieSceneByteTrack::StaticClass() ? new FBytePropertySection(SectionObject, Track->GetTrackName(), Cast<UMovieSceneByteTrack>(SectionObject.GetOuter())->GetEnum()) :
+		SectionClass == UMovieSceneVectorTrack::StaticClass() ? new FVectorPropertySection( SectionObject, Track->GetTrackName() ) :
+		SectionClass == UMovieSceneFloatTrack::StaticClass() ? new FFloatPropertySection( SectionObject, Track->GetTrackName() ) :
+		new FPropertySection( SectionObject, Track->GetTrackName() )
+		);
+
+	return NewSection;
+}
+
+void FPropertyTrackEditor::AddKey(const FGuid& ObjectGuid, UObject* AdditionalAsset)
+{
+	ISequencerObjectChangeListener& ObjectChangeListener = GetSequencer()->GetObjectChangeListener();
+
+	TArray<UObject*> OutObjects;
+	GetSequencer()->GetRuntimeObjects( GetSequencer()->GetFocusedMovieSceneInstance(), ObjectGuid, OutObjects);
+	for (int32 i = 0; i < OutObjects.Num(); ++i)
+	{
+		ObjectChangeListener.TriggerAllPropertiesChanged(OutObjects[i]);
+	}
+}
+
+UEnum* GetEnumForByteTrack(TSharedPtr<ISequencer> Sequencer, const FGuid& OwnerObjectHandle, FName PropertyName, UMovieSceneByteTrack* ByteTrack)
+{
+	TArray<UObject*> RuntimeObjects;
+	Sequencer->GetObjectBindingManager()->GetRuntimeObjects(Sequencer->GetFocusedMovieSceneInstance(), OwnerObjectHandle, RuntimeObjects);
+
+	TSet<UEnum*> PropertyEnums;
+	for (UObject* RuntimeObject : RuntimeObjects)
+	{
+		UProperty* Property = RuntimeObject->GetClass()->FindPropertyByName(PropertyName);
+		if (Property != nullptr)
+		{
+			UByteProperty* ByteProperty = Cast<UByteProperty>(Property);
+			if (ByteProperty != nullptr && ByteProperty->Enum != nullptr)
+			{
+				PropertyEnums.Add(ByteProperty->Enum);
+			}
+		}
+	}
+
+	UEnum* TrackEnum;
+	if (PropertyEnums.Num() == 1)
+	{
+		TrackEnum = PropertyEnums.Array()[0];
+	}
+	else
+	{
+		TrackEnum = nullptr;
+	}
+	return TrackEnum;
+}
+
+UMovieSceneTrack* FPropertyTrackEditor::AddTrack(UMovieScene* FocusedMovieScene, const FGuid& ObjectHandle, TSubclassOf<class UMovieSceneTrack> TrackClass, FName UniqueTypeName)
+{
+	UMovieSceneTrack* NewTrack = FMovieSceneTrackEditor::AddTrack(FocusedMovieScene, ObjectHandle, TrackClass, UniqueTypeName);
+	if (NewTrack->GetClass() == UMovieSceneByteTrack::StaticClass())
+	{
+		UMovieSceneByteTrack* ByteTrack = Cast<UMovieSceneByteTrack>(NewTrack);
+		UEnum* TrackEnum = GetEnumForByteTrack(GetSequencer(), ObjectHandle, UniqueTypeName, ByteTrack);
+		if (TrackEnum != nullptr)
+		{
+			ByteTrack->SetEnum(TrackEnum);
+		}
+	}
+	return NewTrack;
+}
+
+template <typename Type, typename TrackType>
+void FPropertyTrackEditor::OnAnimatedPropertyChanged(const FPropertyChangedParams& PropertyChangedParams )
+{
+	// Get the value from the property 
+	Type Value = *PropertyChangedParams.GetPropertyValue<Type>();
+
+	AnimatablePropertyChanged(TrackType::StaticClass(), PropertyChangedParams.bRequireAutoKey,
+		FOnKeyProperty::CreateRaw(this, &FPropertyTrackEditor::OnKeyProperty<Type, TrackType>, PropertyChangedParams, Value) );
+}
+
+template <typename Type, typename TrackType>
+void FPropertyTrackEditor::OnAnimatedBoolPropertyChanged(const FPropertyChangedParams& PropertyChangedParams)
+{
+	// Get the value from the property 
+	const UBoolProperty* BoolProperty = Cast<const UBoolProperty>(PropertyChangedParams.PropertyPath.Last());
+	if (BoolProperty)
+	{
+		bool BoolValue = BoolProperty->GetPropertyValue(BoolProperty->ContainerPtrToValuePtr<void>(PropertyChangedParams.ObjectsThatChanged.Last()));
+
+		AnimatablePropertyChanged(TrackType::StaticClass(), PropertyChangedParams.bRequireAutoKey,
+			FOnKeyProperty::CreateRaw(this, &FPropertyTrackEditor::OnKeyProperty<Type, TrackType>, PropertyChangedParams, BoolValue));
+	}
+}
+
+void FPropertyTrackEditor::OnAnimatedVectorPropertyChanged(const FPropertyChangedParams& PropertyChangedParams)
+{
+	bool bIsVector2D = false, bIsVector = false, bIsVector4 = false;
+	const UStructProperty* StructProp = Cast<const UStructProperty>(PropertyChangedParams.PropertyPath.Last());
+	if (StructProp && StructProp->Struct)
+	{
+		FName StructName = StructProp->Struct->GetFName();
+
+		bIsVector2D = StructName == NAME_Vector2D;
+		bIsVector = StructName == NAME_Vector;
+		bIsVector4 = StructName == NAME_Vector4;
+	}
+	if (!bIsVector2D && !bIsVector && !bIsVector4) 
+	{
+		return;
+	}
+
+	// Get the vector value from the property
+	if (bIsVector2D)
+	{
+		FVectorKey<FVector2D> NewKey;
+		NewKey.CurveName = PropertyChangedParams.StructPropertyNameToKey;
+		NewKey.bAddKeyEvenIfUnchanged = !PropertyChangedParams.bRequireAutoKey;
+		NewKey.Value = *PropertyChangedParams.GetPropertyValue<FVector2D>();
+
+		FOnKeyProperty OnKeyPropertyDelegate = FOnKeyProperty::CreateRaw( this, &FPropertyTrackEditor::OnKeyProperty<FVectorKey<FVector2D>, UMovieSceneVectorTrack>, PropertyChangedParams, NewKey );
+		AnimatablePropertyChanged(UMovieSceneVectorTrack::StaticClass(), PropertyChangedParams.bRequireAutoKey, OnKeyPropertyDelegate);
+	}
+	else if (bIsVector)
+	{
+		FVectorKey<FVector> NewKey;
+		NewKey.CurveName = PropertyChangedParams.StructPropertyNameToKey;
+		NewKey.bAddKeyEvenIfUnchanged = !PropertyChangedParams.bRequireAutoKey;
+		NewKey.Value = *PropertyChangedParams.GetPropertyValue<FVector>();
+
+		FOnKeyProperty OnKeyPropertyDelegate = FOnKeyProperty::CreateRaw(this, &FPropertyTrackEditor::OnKeyProperty<FVectorKey<FVector>, UMovieSceneVectorTrack>, PropertyChangedParams, NewKey );
+		AnimatablePropertyChanged(UMovieSceneVectorTrack::StaticClass(), PropertyChangedParams.bRequireAutoKey, OnKeyPropertyDelegate);
+	}
+	else if (bIsVector4)
+	{
+		FVectorKey<FVector4> NewKey;
+		NewKey.CurveName = PropertyChangedParams.StructPropertyNameToKey;
+		NewKey.bAddKeyEvenIfUnchanged = !PropertyChangedParams.bRequireAutoKey;
+		NewKey.Value = *PropertyChangedParams.GetPropertyValue<FVector4>();
+
+		FOnKeyProperty OnKeyPropertyDelegate = FOnKeyProperty::CreateRaw(this, &FPropertyTrackEditor::OnKeyProperty<FVectorKey<FVector4>, UMovieSceneVectorTrack>, PropertyChangedParams, NewKey );
+		AnimatablePropertyChanged(UMovieSceneVectorTrack::StaticClass(), PropertyChangedParams.bRequireAutoKey, OnKeyPropertyDelegate );
+	}
+}
+
+void FPropertyTrackEditor::OnAnimatedColorPropertyChanged(const FPropertyChangedParams& PropertyChangedParams )
+{
+	bool bIsFColor = false, bIsFLinearColor = false, bIsSlateColor = false;
+	const UStructProperty* StructProp = Cast<const UStructProperty>(PropertyChangedParams.PropertyPath.Last());
+	if (StructProp && StructProp->Struct)
+	{
+		FName StructName = StructProp->Struct->GetFName();
+
+		bIsFColor = StructName == NAME_Color;
+		bIsFLinearColor = StructName == NAME_LinearColor;
+		bIsSlateColor = StructName == FName("SlateColor");
+	}
+
+	if (!bIsFColor && !bIsFLinearColor && !bIsSlateColor) 
+	{
+		return;
+	}
+
+	FName PropertyName = PropertyChangedParams.PropertyPath.Last()->GetFName();
+
+	FLinearColor ColorValue;
+	if (bIsFColor)
+	{
+		ColorValue = PropertyChangedParams.GetPropertyValue<FColor>()->ReinterpretAsLinear();
+	}
+	else
+	{
+		ColorValue = *PropertyChangedParams.GetPropertyValue<FLinearColor>();
+	}
+
+	if( StructProp->HasMetaData("HideAlphaChannel") )
+	{
+		ColorValue.A = 1;
+	}
+
+
+	FColorKey Key;
+	Key.Value = ColorValue;
+	Key.CurveName = PropertyChangedParams.StructPropertyNameToKey;
+	Key.bAddKeyEvenIfUnchanged = !PropertyChangedParams.bRequireAutoKey;
+	Key.bIsSlateColor = bIsSlateColor;
+
+	AnimatablePropertyChanged(UMovieSceneColorTrack::StaticClass(), PropertyChangedParams.bRequireAutoKey, FOnKeyProperty::CreateRaw(this, &FPropertyTrackEditor::OnKeyProperty<FColorKey, UMovieSceneColorTrack>, PropertyChangedParams, Key) );
+}
+
+
+template <typename Type, typename TrackType>
+void FPropertyTrackEditor::OnKeyProperty( float KeyTime, FPropertyChangedParams PropertyChangedParams, Type Value )
+{
+	for( UObject* Object : PropertyChangedParams.ObjectsThatChanged )
+	{
+		FGuid ObjectHandle = FindOrCreateHandleToObject( Object );
+		if (ObjectHandle.IsValid())
+		{
+			FName PropertyName = PropertyChangedParams.PropertyPath.Last()->GetFName();
+
+			UMovieSceneTrack* Track = GetTrackForObject( ObjectHandle, TrackType::StaticClass(), PropertyName );
+			if( ensure( Track ) )
+			{
+				TrackType* TypedTrack = CastChecked<TrackType>(Track);
+
+				TypedTrack->SetPropertyNameAndPath( PropertyName, PropertyChangedParams.GetPropertyPathString() );
+				// Find or add a new section at the auto-key time and changing the property same property
+				// AddKeyToSection is not actually a virtual, it's redefined in each class with a different type
+				bool bSuccessfulAdd = TypedTrack->AddKeyToSection( KeyTime, Value );
+				if (bSuccessfulAdd)
+				{
+					TypedTrack->SetAsShowable();
+				}
+			}
+		}
+	}
+}
+
diff --git a/Engine/Source/Editor/MovieSceneTools/Private/TrackEditors/PropertyTrackEditor.h b/Engine/Source/Editor/MovieSceneTools/Private/TrackEditors/PropertyTrackEditor.h
new file mode 100644
--- /dev/null
+++ b/Engine/Source/Editor/MovieSceneTools/Private/TrackEditors/PropertyTrackEditor.h
@@ -0,0 +1,77 @@
+// Copyright 1998-2015 Epic Games, Inc. All Rights Reserved.
+
+#pragma once
+
+class IPropertyHandle;
+
+class FPropertyChangedParams;
+
+/**
+ * Tools for animatable property types such as floats ands vectors
+ */
+class FPropertyTrackEditor : public FMovieSceneTrackEditor
+{
+public:
+	/**
+	 * Constructor
+	 *
+	 * @param InSequencer	The sequencer instance to be used by this tool
+	 */
+	FPropertyTrackEditor( TSharedRef<ISequencer> InSequencer );
+	~FPropertyTrackEditor();
+
+	/**
+	 * Creates an instance of this class.  Called by a sequencer 
+	 *
+	 * @param OwningSequencer The sequencer instance to be used by this tool
+	 * @return The new instance of this class
+	 */
+	static TSharedRef<FMovieSceneTrackEditor> CreateTrackEditor( TSharedRef<ISequencer> OwningSequencer );
+
+	/** FMovieSceneTrackEditor Interface */
+	virtual bool SupportsType( TSubclassOf<UMovieSceneTrack> Type ) const override;
+	virtual TSharedRef<ISequencerSection> MakeSectionInterface( UMovieSceneSection& SectionObject, UMovieSceneTrack* Track ) override;
+	virtual void AddKey(const FGuid& ObjectGuid, UObject* AdditionalAsset = NULL) override;
+	virtual UMovieSceneTrack* AddTrack(UMovieScene* FocusedMovieScene, const FGuid& ObjectHandle, TSubclassOf<class UMovieSceneTrack> TrackClass, FName UniqueTypeName) override;
+
+private:
+	/**
+	 * Called by the details panel when an animatable property changes
+	 *
+	 * @param InObjectsThatChanged	List of objects that changed
+	 * @param KeyPropertyParams		Parameters for the property change.
+	 */
+	template <typename Type, typename TrackType>
+	void OnAnimatedPropertyChanged( const FPropertyChangedParams& KeyPropertyParams );
+
+	/**
+	 * Called by the details panel when an animatable bool property changes
+	 *
+	 * @param InObjectsThatChanged	List of objects that changed
+	 * @param KeyPropertyParams		Parameters for the property change.
+	 */
+	template <typename Type, typename TrackType>
+	void OnAnimatedBoolPropertyChanged( const FPropertyChangedParams& KeyPropertyParams );
+
+	/**
+	 * Called by the details panel when an animatable vector property changes
+	 *
+	 * @param InObjectsThatChanged	List of objects that changed
+	 * @param KeyPropertyParams		Parameters for the property change.
+	 */
+	void OnAnimatedVectorPropertyChanged( const FPropertyChangedParams& KeyPropertyParams );
+	
+	/**
+	 * Called by the details panel when an animatable color property changes
+	 *
+	 * @param InObjectsThatChanged	List of objects that changed
+	 * @param KeyPropertyParams		Parameters for the property change.
+	 */
+	void OnAnimatedColorPropertyChanged( const FPropertyChangedParams& KeyPropertyParams );
+	
+	/** Delegates for animatable property changed functions */
+	template <typename Type, typename TrackType>
+	void OnKeyProperty( float KeyTime, FPropertyChangedParams KeyPropertyParams, Type Value );
+};
+
+
diff --git a/Engine/Source/Editor/MovieSceneTools/Private/TrackEditors/SubMovieSceneTrackEditor.cpp b/Engine/Source/Editor/MovieSceneTools/Private/TrackEditors/SubMovieSceneTrackEditor.cpp
new file mode 100644
--- /dev/null
+++ b/Engine/Source/Editor/MovieSceneTools/Private/TrackEditors/SubMovieSceneTrackEditor.cpp
@@ -0,0 +1,111 @@
+// Copyright 1998-2015 Epic Games, Inc. All Rights Reserved.
+
+#include "MovieSceneToolsPrivatePCH.h"
+#include "SubMovieSceneTrackEditor.h"
+#include "SubMovieSceneTrack.h"
+#include "SubMovieSceneSection.h"
+
+/**
+ * A generic implementation for displaying simple property sections
+ */
+class FSubMovieSceneSection : public ISequencerSection
+{
+public:
+	FSubMovieSceneSection( TSharedPtr<ISequencer> InSequencer, UMovieSceneSection& InSectionObject, FName SectionName )
+		: DisplayName( NSLOCTEXT("SubMovieSceneSection", "DisplayName", "Scenes") )
+		, SectionObject( *CastChecked<USubMovieSceneSection>( &InSectionObject ) )
+		, Sequencer( InSequencer )
+	{
+		MovieSceneInstance = InSequencer->GetInstanceForSubMovieSceneSection( InSectionObject );
+	}
+
+	/** ISequencerSection interface */
+	virtual UMovieSceneSection* GetSectionObject() override { return &SectionObject; }
+
+	virtual FText GetDisplayName() const override
+	{
+		return DisplayName;
+	}
+
+	virtual FText GetSectionTitle() const override
+	{
+		return FText::FromString( SectionObject.GetMovieScene()->GetName() );
+	}
+
+	virtual void GenerateSectionLayout( class ISectionLayoutBuilder& LayoutBuilder ) const override
+	{
+	}
+
+	virtual FReply OnSectionDoubleClicked( const FGeometry& SectionGeometry, const FPointerEvent& MouseEvent ) override
+	{
+		if( MouseEvent.GetEffectingButton() == EKeys::LeftMouseButton )
+		{
+			Sequencer.Pin()->FocusSubMovieScene( MovieSceneInstance.Pin().ToSharedRef() );
+		}
+		return FReply::Handled();
+	}
+
+	virtual float GetSectionHeight() const override { return 30.0f; }
+
+	virtual int32 OnPaintSection( const FGeometry& AllottedGeometry, const FSlateRect& SectionClippingRect, FSlateWindowElementList& OutDrawElements, int32 LayerId, bool bParentEnabled ) const override 
+	{
+		// Add a box for the section
+		FSlateDrawElement::MakeBox(
+			OutDrawElements,
+			LayerId,
+			AllottedGeometry.ToPaintGeometry(),
+			FEditorStyle::GetBrush("Sequencer.GenericSection.Background"),
+			SectionClippingRect,
+			ESlateDrawEffect::None,
+			FColor( 220, 120, 120 )
+		);
+
+		return LayerId;
+	}
+	
+private:
+	/** Display name of the section */
+	FText DisplayName;
+	/** The section we are visualizing */
+	USubMovieSceneSection& SectionObject;
+	/** The instance that this section is part of */
+	TWeakPtr<FMovieSceneInstance> MovieSceneInstance;
+	/** Sequencer interface */
+	TWeakPtr<ISequencer> Sequencer;
+};
+
+FSubMovieSceneTrackEditor::FSubMovieSceneTrackEditor( TSharedRef<ISequencer> InSequencer )
+	: FMovieSceneTrackEditor( InSequencer ) 
+{
+
+}
+
+
+TSharedRef<FMovieSceneTrackEditor> FSubMovieSceneTrackEditor::CreateTrackEditor( TSharedRef<ISequencer> InSequencer )
+{
+	return MakeShareable( new FSubMovieSceneTrackEditor( InSequencer ) );
+}
+
+bool FSubMovieSceneTrackEditor::SupportsType( TSubclassOf<UMovieSceneTrack> Type ) const
+{
+	// We support sub movie scenes
+	return Type == USubMovieSceneTrack::StaticClass();
+}
+
+TSharedRef<ISequencerSection> FSubMovieSceneTrackEditor::MakeSectionInterface( UMovieSceneSection& SectionObject, UMovieSceneTrack* Track )
+{
+	return MakeShareable( new FSubMovieSceneSection( GetSequencer(), SectionObject, Track->GetTrackName() ) );
+}
+
+bool FSubMovieSceneTrackEditor::HandleAssetAdded(UObject* Asset, const FGuid& TargetObjectGuid)
+{
+	if (Asset->IsA<UMovieScene>())
+	{
+		GetSequencer()->AddSubMovieScene( CastChecked<UMovieScene>( Asset ) );
+
+		return true;
+	}
+	return false;
+}
+
+
diff --git a/Engine/Source/Editor/MovieSceneTools/Private/TrackEditors/SubMovieSceneTrackEditor.h b/Engine/Source/Editor/MovieSceneTools/Private/TrackEditors/SubMovieSceneTrackEditor.h
new file mode 100644
--- /dev/null
+++ b/Engine/Source/Editor/MovieSceneTools/Private/TrackEditors/SubMovieSceneTrackEditor.h
@@ -0,0 +1,33 @@
+// Copyright 1998-2015 Epic Games, Inc. All Rights Reserved.
+
+#pragma once
+
+
+/**
+ * Tools for animatable property types such as floats ands vectors
+ */
+class FSubMovieSceneTrackEditor : public FMovieSceneTrackEditor
+{
+public:
+	/**
+	 * Constructor
+	 *
+	 * @param InSequencer	The sequencer instance to be used by this tool
+	 */
+	FSubMovieSceneTrackEditor( TSharedRef<ISequencer> InSequencer );
+
+	/**
+	 * Creates an instance of this class.  Called by a sequencer 
+	 *
+	 * @param OwningSequencer The sequencer instance to be used by this tool
+	 * @return The new instance of this class
+	 */
+	static TSharedRef<FMovieSceneTrackEditor> CreateTrackEditor( TSharedRef<ISequencer> OwningSequencer );
+
+	/** FMovieSceneTrackEditor Interface */
+	virtual bool SupportsType( TSubclassOf<UMovieSceneTrack> Type ) const override;
+	virtual TSharedRef<ISequencerSection> MakeSectionInterface( UMovieSceneSection& SectionObject, UMovieSceneTrack* Track ) override;
+	virtual bool HandleAssetAdded(UObject* Asset, const FGuid& TargetObjectGuid) override;
+};
+
+
diff --git a/Engine/Source/Editor/MovieSceneTools/Private/TrackEditors/TransformTrackEditor.cpp b/Engine/Source/Editor/MovieSceneTools/Private/TrackEditors/TransformTrackEditor.cpp
new file mode 100644
--- /dev/null
+++ b/Engine/Source/Editor/MovieSceneTools/Private/TrackEditors/TransformTrackEditor.cpp
@@ -0,0 +1,283 @@
+// Copyright 1998-2015 Epic Games, Inc. All Rights Reserved.
+
+#include "MovieSceneToolsPrivatePCH.h"
+#include "ScopedTransaction.h"
+#include "MovieScene.h"
+#include "MovieSceneSection.h"
+#include "MovieScene3DTransformTrack.h"
+#include "MovieScene3DTransformSection.h"
+#include "ISequencerObjectChangeListener.h"
+#include "ISequencerSection.h"
+#include "ISectionLayoutBuilder.h"
+#include "IKeyArea.h"
+#include "MovieSceneToolHelpers.h"
+#include "MovieSceneTrackEditor.h"
+#include "TransformTrackEditor.h"
+
+/**
+ * Class that draws a transform section in the sequencer
+ */
+class F3DTransformSection : public ISequencerSection
+{
+public:
+	F3DTransformSection( UMovieSceneSection& InSection )
+		: Section( InSection )
+	{
+	}
+
+	/** ISequencerSection interface */
+	virtual UMovieSceneSection* GetSectionObject() override
+	{ 
+		return &Section;
+	}
+
+	virtual FText GetDisplayName() const override
+	{ 
+		return NSLOCTEXT("FTransformSection", "DisplayName", "Transform");
+	}
+	
+	virtual FText GetSectionTitle() const override { return FText::GetEmpty(); }
+
+	virtual void GenerateSectionLayout( class ISectionLayoutBuilder& LayoutBuilder ) const override
+	{
+		UMovieScene3DTransformSection* TransformSection = Cast<UMovieScene3DTransformSection>( &Section );
+
+		// This generates the tree structure for the transform section
+		LayoutBuilder.PushCategory( "Location", NSLOCTEXT("FTransformSection", "LocationArea", "Location") );
+			LayoutBuilder.AddKeyArea("Location.X", NSLOCTEXT("FTransformSection", "LocXArea", "X"), MakeShareable( new FFloatCurveKeyArea ( &TransformSection->GetTranslationCurve( EAxis::X ), TransformSection ) ) );
+			LayoutBuilder.AddKeyArea("Location.Y", NSLOCTEXT("FTransformSection", "LocYArea", "Y"), MakeShareable( new FFloatCurveKeyArea ( &TransformSection->GetTranslationCurve( EAxis::Y ), TransformSection ) ) );
+			LayoutBuilder.AddKeyArea("Location.Z", NSLOCTEXT("FTransformSection", "LocZArea", "Z"), MakeShareable( new FFloatCurveKeyArea ( &TransformSection->GetTranslationCurve( EAxis::Z ), TransformSection ) ) );
+		LayoutBuilder.PopCategory();
+
+		LayoutBuilder.PushCategory( "Rotation", NSLOCTEXT("FTransformSection", "RotationArea", "Rotation") );
+			LayoutBuilder.AddKeyArea("Rotation.X", NSLOCTEXT("FTransformSection", "RotXArea", "X"), MakeShareable( new FFloatCurveKeyArea ( &TransformSection->GetRotationCurve( EAxis::X ), TransformSection ) ) );
+			LayoutBuilder.AddKeyArea("Rotation.Y", NSLOCTEXT("FTransformSection", "RotYArea", "Y"), MakeShareable( new FFloatCurveKeyArea ( &TransformSection->GetRotationCurve( EAxis::Y ), TransformSection ) ) );
+			LayoutBuilder.AddKeyArea("Rotation.Z", NSLOCTEXT("FTransformSection", "RotZArea", "Z"), MakeShareable( new FFloatCurveKeyArea ( &TransformSection->GetRotationCurve( EAxis::Z ), TransformSection ) ) );
+		LayoutBuilder.PopCategory();
+
+		LayoutBuilder.PushCategory( "Scale", NSLOCTEXT("FTransformSection", "ScaleArea", "Scale") );
+			LayoutBuilder.AddKeyArea("Scale.X", NSLOCTEXT("FTransformSection", "ScaleXArea", "X"), MakeShareable( new FFloatCurveKeyArea ( &TransformSection->GetScaleCurve( EAxis::X ), TransformSection ) ) );
+			LayoutBuilder.AddKeyArea("Scale.Y", NSLOCTEXT("FTransformSection", "ScaleYArea", "Y"), MakeShareable( new FFloatCurveKeyArea ( &TransformSection->GetScaleCurve( EAxis::Y ), TransformSection ) ) );
+			LayoutBuilder.AddKeyArea("Scale.Z", NSLOCTEXT("FTransformSection", "ScaleZArea", "Z"), MakeShareable( new FFloatCurveKeyArea ( &TransformSection->GetScaleCurve( EAxis::Z ), TransformSection ) ) );
+		LayoutBuilder.PopCategory();
+	}
+
+	virtual int32 OnPaintSection( const FGeometry& AllottedGeometry, const FSlateRect& SectionClippingRect, FSlateWindowElementList& OutDrawElements, int32 LayerId, bool bParentEnabled ) const override 
+	{
+		// Add a box for the section
+		FSlateDrawElement::MakeBox( 
+			OutDrawElements,
+			LayerId,
+			AllottedGeometry.ToPaintGeometry(),
+			FEditorStyle::GetBrush("Sequencer.GenericSection.Background"),
+			SectionClippingRect
+		); 
+
+		return LayerId;
+	}
+
+private:
+	/** The section we are visualizing */
+	UMovieSceneSection& Section;
+};
+
+F3DTransformTrackEditor::F3DTransformTrackEditor( TSharedRef<ISequencer> InSequencer )
+	: FMovieSceneTrackEditor( InSequencer ) 
+{
+	// Listen for actor/component movement
+	GEditor->OnBeginObjectMovement().AddRaw( this, &F3DTransformTrackEditor::OnPreTransformChanged );
+	GEditor->OnEndObjectMovement().AddRaw( this, &F3DTransformTrackEditor::OnTransformChanged );
+
+	// Listen for the viewport's viewed through camera starts and stops movement
+	GEditor->OnBeginCameraMovement().AddRaw( this, &F3DTransformTrackEditor::OnPreTransformChanged );
+	GEditor->OnEndCameraMovement().AddRaw( this, &F3DTransformTrackEditor::OnTransformChanged );
+}
+
+F3DTransformTrackEditor::~F3DTransformTrackEditor()
+{
+	GEditor->OnBeginObjectMovement().RemoveAll( this );
+	GEditor->OnEndObjectMovement().RemoveAll( this );
+
+	GEditor->OnBeginCameraMovement().RemoveAll( this );
+	GEditor->OnEndCameraMovement().RemoveAll( this );
+}
+
+TSharedRef<FMovieSceneTrackEditor> F3DTransformTrackEditor::CreateTrackEditor( TSharedRef<ISequencer> InSequencer )
+{
+	return MakeShareable( new F3DTransformTrackEditor( InSequencer ) );
+}
+
+bool F3DTransformTrackEditor::SupportsType( TSubclassOf<UMovieSceneTrack> Type ) const
+{
+	// We support animatable transforms
+	return Type == UMovieScene3DTransformTrack::StaticClass();
+}
+
+TSharedRef<ISequencerSection> F3DTransformTrackEditor::MakeSectionInterface( UMovieSceneSection& SectionObject, UMovieSceneTrack* Track )
+{
+	check( SupportsType( SectionObject.GetOuter()->GetClass() ) );
+
+	TSharedRef<ISequencerSection> NewSection( new F3DTransformSection( SectionObject ) );
+
+	return NewSection;
+}
+
+
+void F3DTransformTrackEditor::OnPreTransformChanged( UObject& InObject )
+{
+	UMovieScene* MovieScene = GetMovieScene();
+	float AutoKeyTime = GetTimeForKey( MovieScene );
+
+	if( IsAllowedToAutoKey() )
+	{
+		USceneComponent* SceneComponent = NULL;
+		AActor* Actor = Cast<AActor>( &InObject );
+		if( Actor && Actor->GetRootComponent() )
+		{
+			SceneComponent = Actor->GetRootComponent();
+		}
+		else
+		{
+			// If the object wasn't an actor attempt to get it directly as a scene component 
+			SceneComponent = Cast<USceneComponent>( &InObject );
+		}
+
+		if( SceneComponent )
+		{
+			// Cache off the existing transform so we can detect which components have changed
+			// and keys only when something has changed
+			FTransformData Transform( SceneComponent );
+
+			ObjectToExistingTransform.Add( &InObject, Transform );
+		}
+	}
+}
+
+/**
+ * Temp struct used because delegates only accept 4 or less payloads
+ * FTransformKey is immutable and would require heavy re-architecting to fit here
+ */
+struct FTransformDataPair
+{
+	FTransformDataPair(FTransformData InTransformData, FTransformData InLastTransformData)
+		: TransformData(InTransformData)
+		, LastTransformData(InLastTransformData) {}
+
+	FTransformData TransformData;
+	FTransformData LastTransformData;
+};
+
+void F3DTransformTrackEditor::OnTransformChanged( UObject& InObject )
+{
+	const TSharedPtr<ISequencer> Sequencer = GetSequencer();
+
+	USceneComponent* SceneComponentThatChanged = NULL;
+
+	// The runtime binding
+	FGuid ObjectHandle;
+
+	AActor* Actor = Cast<AActor>( &InObject );
+	if( Actor && Actor->GetRootComponent() )
+	{
+		// Get a handle bound to the key/section we are adding so we know what objects to change during playback
+		ObjectHandle = Sequencer->GetHandleToObject( Actor );
+		SceneComponentThatChanged = Actor->GetRootComponent();
+	}
+
+	else
+	{
+		// If the object wasn't an actor attempt to get it directly as a scene component 
+		SceneComponentThatChanged = Cast<USceneComponent>( &InObject );
+		if( SceneComponentThatChanged )
+		{
+			ObjectHandle = Sequencer->GetHandleToObject( SceneComponentThatChanged );
+		}
+
+	}
+
+	if( SceneComponentThatChanged && ObjectHandle.IsValid() )
+	{
+		// Find an existing transform if possible.  If one exists we will compare against the new one to decide what components of the transform need keys
+		FTransformData ExistingTransform = ObjectToExistingTransform.FindRef( &InObject );
+
+		// Remove it from the list of cached transforms. 
+		// @todo sequencer livecapture: This can be made much for efficient by not removing cached state during live capture situation
+		ObjectToExistingTransform.Remove( &InObject );
+
+		// Build new transform data
+		FTransformData NewTransformData( SceneComponentThatChanged );
+
+		FTransformDataPair TransformPair(NewTransformData, ExistingTransform);
+
+		AnimatablePropertyChanged(UMovieScene3DTransformTrack::StaticClass(), true,
+			FOnKeyProperty::CreateRaw(this, &F3DTransformTrackEditor::OnTransformChangedInternals, &InObject, ObjectHandle, TransformPair, true));
+	}
+}
+
+void F3DTransformTrackEditor::AddKey(const FGuid& ObjectGuid, UObject* AdditionalAsset)
+{
+	TArray<UObject*> OutObjects;
+	GetSequencer()->GetRuntimeObjects( GetSequencer()->GetFocusedMovieSceneInstance(), ObjectGuid, OutObjects);
+
+	for ( UObject* Object : OutObjects )
+	{
+		FGuid ObjectHandle;
+		USceneComponent* SceneComponent = NULL;
+		AActor* Actor = Cast<AActor>( Object );
+		if( Actor && Actor->GetRootComponent() )
+		{
+			ObjectHandle = GetSequencer()->GetHandleToObject( Actor );
+			SceneComponent = Actor->GetRootComponent();
+		}
+		else
+		{
+			// If the object wasn't an actor attempt to get it directly as a scene component 
+			SceneComponent = Cast<USceneComponent>( Object );
+			if( SceneComponent )
+			{
+				ObjectHandle = GetSequencer()->GetHandleToObject( SceneComponent );
+			}
+		}
+
+		if( SceneComponent )
+		{
+			// Cache off the existing transform so we can detect which components have changed
+			// and keys only when something has changed
+			FTransformData CurrentTransform( SceneComponent );
+
+			FTransformDataPair TransformPair(CurrentTransform, FTransformData());
+
+			AnimatablePropertyChanged(UMovieScene3DTransformTrack::StaticClass(), false,
+				FOnKeyProperty::CreateRaw(this, &F3DTransformTrackEditor::OnTransformChangedInternals, Object, ObjectHandle, TransformPair, false));
+		}
+	}
+}
+
+void F3DTransformTrackEditor::OnTransformChangedInternals(float KeyTime, UObject* InObject, FGuid ObjectHandle, FTransformDataPair TransformPair, bool bAutoKeying)
+{
+	// Only unwind rotation if we're generating keys while recording (scene is actively playing back)
+	const bool bUnwindRotation = GetSequencer()->IsRecordingLive();
+
+	FName Transform("Transform");
+	if (ObjectHandle.IsValid())
+	{
+		UMovieSceneTrack* Track = GetTrackForObject( ObjectHandle, UMovieScene3DTransformTrack::StaticClass(), Transform );
+		UMovieScene3DTransformTrack* TransformTrack = CastChecked<UMovieScene3DTransformTrack>( Track );
+		// Transform name and path are the same
+		TransformTrack->SetPropertyNameAndPath( Transform, Transform.ToString() );
+	
+		if (!TransformPair.LastTransformData.IsValid())
+		{
+			bool bHasTranslationKeys = false, bHasRotationKeys = false, bHasScaleKeys = false;
+			TransformPair.LastTransformData.bValid = TransformTrack->Eval(KeyTime, KeyTime, TransformPair.LastTransformData.Translation, TransformPair.LastTransformData.Rotation, TransformPair.LastTransformData.Scale, bHasTranslationKeys, bHasRotationKeys, bHasScaleKeys);
+		}
+
+		FTransformKey TransformKey = FTransformKey(KeyTime, TransformPair.TransformData, TransformPair.LastTransformData);
+
+		bool bSuccessfulAdd = TransformTrack->AddKeyToSection( ObjectHandle, TransformKey, bUnwindRotation );
+		if (bSuccessfulAdd)
+		{
+			TransformTrack->SetAsShowable();
+		}
+	}
+}
diff --git a/Engine/Source/Editor/MovieSceneTools/Private/TrackEditors/TransformTrackEditor.h b/Engine/Source/Editor/MovieSceneTools/Private/TrackEditors/TransformTrackEditor.h
new file mode 100644
--- /dev/null
+++ b/Engine/Source/Editor/MovieSceneTools/Private/TrackEditors/TransformTrackEditor.h
@@ -0,0 +1,56 @@
+// Copyright 1998-2015 Epic Games, Inc. All Rights Reserved.
+
+#pragma once
+
+#include "MovieScene3DTransformTrack.h"
+
+/**
+ * Tools for animatable transforms
+ */
+class F3DTransformTrackEditor : public FMovieSceneTrackEditor
+{
+public:
+	/**
+	 * Constructor
+	 *
+	 * @param InSequencer	The sequencer instance to be used by this tool
+	 */
+	F3DTransformTrackEditor( TSharedRef<ISequencer> InSequencer );
+	~F3DTransformTrackEditor();
+
+	/**
+	 * Creates an instance of this class.  Called by a sequencer 
+	 *
+	 * @param OwningSequencer The sequencer instance to be used by this tool
+	 * @return The new instance of this class
+	 */
+	static TSharedRef<FMovieSceneTrackEditor> CreateTrackEditor( TSharedRef<ISequencer> OwningSequencer );
+
+	/** FMovieSceneTrackEditor Interface */
+	virtual bool SupportsType( TSubclassOf<UMovieSceneTrack> Type ) const override;
+	virtual TSharedRef<ISequencerSection> MakeSectionInterface( UMovieSceneSection& SectionObject, UMovieSceneTrack* Track ) override;
+	virtual void AddKey( const FGuid& ObjectGuid, UObject* AdditionalAsset = NULL ) override;
+
+private:
+	/**
+	 * Called before an actor or component transform changes
+	 *
+	 * @param Object	The object whose transform is about to change
+	 */
+	void OnPreTransformChanged( UObject& InObject );
+
+	/**
+	 * Called when an actor or component transform changes
+	 *
+	 * @param Object	The object whose transform has changed
+	 */
+	void OnTransformChanged( UObject& InObject );
+
+	/** Delegate for animatable property changed in OnTransformChanged */
+	void OnTransformChangedInternals(float KeyTime, UObject* InObject, FGuid ObjectHandle, struct FTransformDataPair TransformPair, bool bAutoKeying);
+private:
+	/** Mapping of objects to their existing transform data (for comparing against new transform data) */
+	TMap< TWeakObjectPtr<UObject>, FTransformData > ObjectToExistingTransform;
+};
+
+
diff --git a/Engine/Source/Editor/MovieSceneTools/Public/CommonMovieSceneTools.h b/Engine/Source/Editor/MovieSceneTools/Public/CommonMovieSceneTools.h
new file mode 100644
--- /dev/null
+++ b/Engine/Source/Editor/MovieSceneTools/Public/CommonMovieSceneTools.h
@@ -0,0 +1,55 @@
+// Copyright 1998-2015 Epic Games, Inc. All Rights Reserved.
+
+#pragma once
+
+/**
+ * Utility for converting time units to slate pixel units and vice versa
+ */
+struct FTimeToPixel
+{
+public:
+	FTimeToPixel( const FGeometry& AllottedGeometry, TRange<float> InLocalViewRange )
+		: MaxPixelsPerInput(1000.0f)
+		, LocalViewRange( InLocalViewRange )
+	{
+		const float ViewRange = LocalViewRange.Size<float>();
+		PixelsPerInput = ViewRange > 0 ? AllottedGeometry.Size.X / ViewRange : MaxPixelsPerInput;
+	}
+
+	/**
+	 * Converts a time to a pixel point relative to the geometry of a widget (passed into the constructor)
+	 *
+	 * @param Time	The time to convert
+	 * @return The pixel equivalent of the time
+	 */
+	float TimeToPixel( float Time ) const
+	{
+		return (Time - LocalViewRange.GetLowerBoundValue()) * PixelsPerInput;
+	}
+
+	/**
+	 * Converts a pixel value to time 
+	 *
+	 * @param PixelX The x value of a pixel coordinate relative to the geometry that was passed into the constructor.
+	 * @return The time where the pixel is located
+	 */
+	float PixelToTime( float PixelX ) const
+	{
+		return (PixelX/PixelsPerInput) + LocalViewRange.GetLowerBoundValue();
+	}
+
+	/**
+	 * @return The pixels per input
+	 */
+	float GetPixelsPerInput() const 
+	{
+		return PixelsPerInput;
+	}
+private:
+	const float MaxPixelsPerInput;
+
+	/** time range of the sequencer */
+	TRange<float> LocalViewRange;
+	/** The number of pixels in the view range */
+	float PixelsPerInput;
+};
diff --git a/Engine/Source/Editor/MovieSceneTools/Public/IMovieSceneTools.h b/Engine/Source/Editor/MovieSceneTools/Public/IMovieSceneTools.h
new file mode 100644
--- /dev/null
+++ b/Engine/Source/Editor/MovieSceneTools/Public/IMovieSceneTools.h
@@ -0,0 +1,16 @@
+// Copyright 1998-2015 Epic Games, Inc. All Rights Reserved.
+
+#pragma once
+
+#include "ModuleInterface.h"
+
+/**
+ * The public interface of the MovieSceneTools module
+ */
+class IMovieSceneTools : public IModuleInterface
+{
+
+public:
+
+};
+
diff --git a/Engine/Source/Editor/MovieSceneTools/Public/K2Node_PlayMovieScene.h b/Engine/Source/Editor/MovieSceneTools/Public/K2Node_PlayMovieScene.h
new file mode 100644
--- /dev/null
+++ b/Engine/Source/Editor/MovieSceneTools/Public/K2Node_PlayMovieScene.h
@@ -0,0 +1,145 @@
+// Copyright 1998-2015 Epic Games, Inc. All Rights Reserved.
+
+#pragma once
+
+#include "BlueprintGraphDefinitions.h"		// @todo sequencer uobject: workaround for header generator dependencies
+#include "EdGraph/EdGraphNodeUtils.h" // for FNodeTextCache
+#include "K2Node_PlayMovieScene.generated.h"
+
+#if WITH_EDITOR
+DECLARE_EVENT( UK2Node_PlayMovieScene, FOnMovieSceneBindingsChanged );
+#endif
+
+/** Action to add a 'Play MovieScene' node to the graph */
+USTRUCT()
+struct FEdGraphSchemaAction_K2AddPlayMovieScene : public FEdGraphSchemaAction_K2NewNode
+{
+	GENERATED_USTRUCT_BODY()
+
+	FEdGraphSchemaAction_K2AddPlayMovieScene()
+		: FEdGraphSchemaAction_K2NewNode()
+	{
+	}
+	
+	FEdGraphSchemaAction_K2AddPlayMovieScene(const FText& InNodeCategory, const FText& InMenuDesc, const FString& InToolTip, const int32 InGrouping)
+		: FEdGraphSchemaAction_K2NewNode(InNodeCategory, InMenuDesc, InToolTip, InGrouping)
+	{
+	}
+	
+	// FEdGraphSchemaAction interface
+	virtual UEdGraphNode* PerformAction(class UEdGraph* ParentGraph, UEdGraphPin* FromPin, const FVector2D Location, bool bSelectNewNode = true) override;
+	// End of FEdGraphSchemaAction interface
+
+};
+
+
+UCLASS( MinimalAPI )
+class UK2Node_PlayMovieScene : public UK2Node
+{
+	GENERATED_UCLASS_BODY()
+
+#if WITH_EDITOR
+	// UEdGraphNode interface.
+	virtual void AllocateDefaultPins() override;
+	virtual FLinearColor GetNodeTitleColor() const override;
+	virtual FText GetNodeTitle(ENodeTitleType::Type TitleType) const override;
+
+	// UK2Node interface
+	virtual bool NodeCausesStructuralBlueprintChange() const override
+	{ 
+		// @todo sequencer: Not sure if we really need to return true for this
+		return true; 
+	}
+	virtual void PinConnectionListChanged( UEdGraphPin* Pin ) override;
+
+
+	/**
+	 * Associates a MovieScene asset with this node
+	 *
+	 * @param	NewMovieScene	The MovieScene asset to use (can be NULL to clear it)
+	 */
+	virtual void SetMovieScene( class UMovieScene* NewMovieScene );
+
+	/** @return	Returns the MovieScene associated with this node (or NULL, if no MovieScene is associated yet) */
+	virtual class UMovieScene* GetMovieScene();
+
+	/** @return	Returns MovieSceneBindings for this node */
+	virtual class UMovieSceneBindings* GetMovieSceneBindings();
+
+	/**
+	 * Binds a possessable in the MovieScene asset (based on the guid of that possessable) to the specified objects
+	 *
+	 * @param	PossessableGuid		The guid of the posessable in the MovieScene to bind to
+	 * @param	Objects				The object(s) to bind to
+	 */
+	virtual void BindPossessableToObjects( const FGuid& PossessableGuid, const TArray< UObject* >& Objects );
+
+	/**
+	 * Given a guid for a possessable, attempts to find the object bound to that guid
+	 *
+	 * @param	Guid	The possessable guid to search for
+	 *
+	 * @return	The bound objects, if found, otherwise NULL
+	 */
+	virtual TArray< UObject* > FindBoundObjects( const FGuid& Guid );
+
+	/**
+	 * Finds a guid for a bound posessable object if it exists 
+	 *
+	 * @param Object	The object to search for
+	 * @return A valid guid if the object was found.  Otherwise the guid is invalid
+	 */
+	virtual FGuid FindGuidForObject( UObject* Object ) const;
+
+	/**
+	 * Returns access to an event you can bind to to find out when the object bindings change
+	 *
+	 * @return	OnBindingsChanged event object
+	 */
+	virtual FOnMovieSceneBindingsChanged& OnBindingsChanged()
+	{
+		return OnBindingsChangedEvent;
+	}
+
+#endif // #if WITH_EDITOR
+
+	/** @return	Returns the graph pin for the "Play" input */
+	MOVIESCENETOOLS_API UEdGraphPin* GetPlayPin() const;
+
+	/** @return	Returns the graph pin for the "Pause" input */
+	MOVIESCENETOOLS_API UEdGraphPin* GetPausePin() const;
+
+protected:
+
+#if WITH_EDITOR
+
+	/**
+	 * Creates a pin on this node for the specified bound object
+	 *
+	 * @param	BoundObject		The object to create a pin for
+	 */
+	virtual void CreatePinForBoundObject( struct FMovieSceneBoundObject& BoundObject );
+
+	/**
+	 * Creates a MovieSceneBindings object if we don't have one yet
+	 */
+	void CreateBindingsIfNeeded();
+#endif
+
+protected:
+
+	/** Binding information that connects the MovieScene to objects that can be possessed.  At runtime, the bindings
+	   object is actually owned by the LevelScriptActor in the level this node is associated with. */
+	UPROPERTY()
+	class UMovieSceneBindings* MovieSceneBindings;
+
+#if WITH_EDITOR
+	/** Event that is broadcast when the bindings are changed */
+	FOnMovieSceneBindingsChanged OnBindingsChangedEvent;
+
+private:
+	/** Constructing FText strings can be costly, so we cache the node's title */
+	FNodeTextCache CachedNodeTitle;
+#endif
+};
+
diff --git a/Engine/Source/Editor/MovieSceneTools/Public/MovieSceneFactory.h b/Engine/Source/Editor/MovieSceneTools/Public/MovieSceneFactory.h
new file mode 100644
--- /dev/null
+++ b/Engine/Source/Editor/MovieSceneTools/Public/MovieSceneFactory.h
@@ -0,0 +1,18 @@
+// Copyright 1998-2015 Epic Games, Inc. All Rights Reserved.
+
+#pragma once
+
+#include "MovieSceneFactory.generated.h"
+
+UCLASS()
+class UMovieSceneFactory : public UFactory
+{
+	GENERATED_UCLASS_BODY()
+
+	// Begin UFactory Interface
+	virtual UObject* FactoryCreateNew(UClass* Class,UObject* InParent,FName Name,EObjectFlags Flags,UObject* Context,FFeedbackContext* Warn) override;
+	// Begin UFactory Interface
+};
+
+
+
diff --git a/Engine/Source/Editor/MovieSceneTools/Public/MovieSceneToolHelpers.h b/Engine/Source/Editor/MovieSceneTools/Public/MovieSceneToolHelpers.h
new file mode 100644
--- /dev/null
+++ b/Engine/Source/Editor/MovieSceneTools/Public/MovieSceneToolHelpers.h
@@ -0,0 +1,147 @@
+// Copyright 1998-2015 Epic Games, Inc. All Rights Reserved.
+
+#pragma once
+
+#include "ISequencer.h"
+#include "ISequencerSection.h"
+#include "IKeyArea.h"
+
+/**
+ * A key area for float keys
+ */
+class MOVIESCENETOOLS_API FFloatCurveKeyArea : public IKeyArea
+{
+public:
+	FFloatCurveKeyArea( FRichCurve* InCurve, UMovieSceneSection* InOwningSection )
+		: Curve( InCurve )
+		, OwningSection( InOwningSection )
+	{}
+
+	/** IKeyArea interface */
+	virtual TArray<FKeyHandle> GetUnsortedKeyHandles() const override
+	{
+		TArray<FKeyHandle> OutKeyHandles;
+		for (auto It(Curve->GetKeyHandleIterator()); It; ++It)
+		{
+			OutKeyHandles.Add(It.Key());
+		}
+		return OutKeyHandles;
+	}
+
+	virtual float GetKeyTime( FKeyHandle KeyHandle ) const override
+	{
+		return Curve->GetKeyTime( KeyHandle );
+	}
+
+	virtual FKeyHandle MoveKey( FKeyHandle KeyHandle, float DeltaPosition ) override
+	{
+		return Curve->SetKeyTime( KeyHandle, Curve->GetKeyTime( KeyHandle ) + DeltaPosition );
+	}
+
+	virtual void DeleteKey(FKeyHandle KeyHandle) override
+	{
+		if ( Curve->IsKeyHandleValid( KeyHandle ) )
+		{
+			Curve->DeleteKey( KeyHandle );
+		}
+	}
+
+	virtual void AddKeyUnique(float Time) override;
+
+	virtual FRichCurve* GetRichCurve() override { return Curve; }
+
+	virtual UMovieSceneSection* GetOwningSection() override { return OwningSection; }
+
+	virtual bool CanCreateKeyEditor() override
+	{
+		return true;
+	}
+
+	virtual TSharedRef<SWidget> CreateKeyEditor(ISequencer* Sequencer) override;
+
+private:
+	/** The curve which provides the keys for this key area. */
+	FRichCurve* Curve;
+	/** The section that owns this key area. */
+	UMovieSceneSection* OwningSection;
+};
+
+
+
+/** A key area for integral keys */
+class MOVIESCENETOOLS_API FIntegralKeyArea : public IKeyArea
+{
+public:
+	FIntegralKeyArea(FIntegralCurve& InCurve, UMovieSceneSection* InOwningSection)
+		: Curve(InCurve)
+		, OwningSection(InOwningSection)
+	{}
+
+	/** IKeyArea interface */
+	virtual TArray<FKeyHandle> GetUnsortedKeyHandles() const override
+	{
+		TArray<FKeyHandle> OutKeyHandles;
+		for (auto It(Curve.GetKeyHandleIterator()); It; ++It)
+		{
+			OutKeyHandles.Add(It.Key());
+		}
+		return OutKeyHandles;
+	}
+
+	virtual float GetKeyTime( FKeyHandle KeyHandle ) const override
+	{
+		return Curve.GetKeyTime( KeyHandle );
+	}
+
+	virtual FKeyHandle MoveKey( FKeyHandle KeyHandle, float DeltaPosition ) override
+	{
+		return Curve.SetKeyTime( KeyHandle, Curve.GetKeyTime( KeyHandle ) + DeltaPosition );
+	}
+
+	virtual void DeleteKey(FKeyHandle KeyHandle) override
+	{
+		Curve.DeleteKey(KeyHandle);
+	}
+
+	virtual void AddKeyUnique(float Time) override;
+
+	virtual FRichCurve* GetRichCurve() override { return nullptr; };
+
+	virtual UMovieSceneSection* GetOwningSection() override { return OwningSection; }
+
+	virtual bool CanCreateKeyEditor() override
+	{
+		return true;
+	}
+
+	virtual TSharedRef<SWidget> CreateKeyEditor(ISequencer* Sequencer) override;
+
+protected:
+	/** Curve with keys in this area */
+	FIntegralCurve& Curve;
+	/** The section that owns this key area. */
+	UMovieSceneSection* OwningSection;
+};
+
+/** A key area for displaying and editing intragral curves representing enums. */
+class FEnumKeyArea : public FIntegralKeyArea
+{
+public:
+	FEnumKeyArea(FIntegralCurve& InCurve, UMovieSceneSection* InOwningSection, UEnum* InEnum)
+		: FIntegralKeyArea(InCurve, InOwningSection)
+		, Enum(InEnum)
+	{}
+
+	virtual bool CanCreateKeyEditor() override 
+	{
+		return true;
+	}
+
+	virtual TSharedRef<SWidget> CreateKeyEditor(ISequencer* Sequencer) override;
+
+private:
+	/** The enum which provides available integral values for this key area. */
+	UEnum* Enum;
+};
+
+
diff --git a/Engine/Source/Editor/MovieSceneTools/Public/PropertySection.h b/Engine/Source/Editor/MovieSceneTools/Public/PropertySection.h
new file mode 100644
--- /dev/null
+++ b/Engine/Source/Editor/MovieSceneTools/Public/PropertySection.h
@@ -0,0 +1,51 @@
+// Copyright 1998-2015 Epic Games, Inc. All Rights Reserved.
+
+#pragma once
+
+#include "ISequencerSection.h"
+
+/**
+ * A generic implementation for displaying simple property sections
+ */
+class FPropertySection : public ISequencerSection
+{
+public:
+	FPropertySection( UMovieSceneSection& InSectionObject, FName SectionName )
+		: SectionObject( InSectionObject )
+	{
+		DisplayName = FText::FromString( FName::NameToDisplayString( SectionName.ToString(), false ) );
+	}
+
+	/** ISequencerSection interface */
+	virtual UMovieSceneSection* GetSectionObject() override { return &SectionObject; }
+
+	virtual FText GetDisplayName() const override
+	{
+		return DisplayName;
+	}
+	
+	virtual FText GetSectionTitle() const override { return FText::GetEmpty(); }
+
+	virtual void GenerateSectionLayout( class ISectionLayoutBuilder& LayoutBuilder ) const override {}
+
+	virtual int32 OnPaintSection( const FGeometry& AllottedGeometry, const FSlateRect& SectionClippingRect, FSlateWindowElementList& OutDrawElements, int32 LayerId, bool bParentEnabled ) const override 
+	{
+		// Add a box for the section
+		FSlateDrawElement::MakeBox( 
+			OutDrawElements,
+			LayerId,
+			AllottedGeometry.ToPaintGeometry(),
+			FEditorStyle::GetBrush("Sequencer.GenericSection.Background"),
+			SectionClippingRect
+		);
+
+		return LayerId;
+	}
+	
+protected:
+	/** Display name of the section */
+	FText DisplayName;
+	/** The section we are visualizing */
+	UMovieSceneSection& SectionObject;
+};
+
diff --git a/Engine/Source/Editor/UMGEditor/Private/Animation/MarginTrackEditor.cpp b/Engine/Source/Editor/UMGEditor/Private/Animation/MarginTrackEditor.cpp
--- a/Engine/Source/Editor/UMGEditor/Private/Animation/MarginTrackEditor.cpp
+++ b/Engine/Source/Editor/UMGEditor/Private/Animation/MarginTrackEditor.cpp
@@ -2,8 +2,8 @@
 
 #include "UMGEditorPrivatePCH.h"
 #include "MarginTrackEditor.h"
-#include "Developer/MovieSceneTools/Public/PropertySection.h"
-#include "Developer/MovieSceneTools/Public/MovieSceneToolHelpers.h"
+#include "Editor/MovieSceneTools/Public/PropertySection.h"
+#include "Editor/MovieSceneTools/Public/MovieSceneToolHelpers.h"
 #include "Runtime/UMG/Public/Animation/MovieSceneMarginSection.h"
 #include "Runtime/UMG/Public/Animation/MovieSceneMarginTrack.h"
 #include "Editor/Sequencer/Public/ISectionLayoutBuilder.h"
diff --git a/Engine/Source/Editor/UMGEditor/Private/Animation/Sequencer2DTransformTrackEditor.cpp b/Engine/Source/Editor/UMGEditor/Private/Animation/Sequencer2DTransformTrackEditor.cpp
--- a/Engine/Source/Editor/UMGEditor/Private/Animation/Sequencer2DTransformTrackEditor.cpp
+++ b/Engine/Source/Editor/UMGEditor/Private/Animation/Sequencer2DTransformTrackEditor.cpp
@@ -2,8 +2,8 @@
 
 #include "UMGEditorPrivatePCH.h"
 #include "Sequencer2DTransformTrackEditor.h"
-#include "Developer/MovieSceneTools/Public/PropertySection.h"
-#include "Developer/MovieSceneTools/Public/MovieSceneToolHelpers.h"
+#include "Editor/MovieSceneTools/Public/PropertySection.h"
+#include "Editor/MovieSceneTools/Public/MovieSceneToolHelpers.h"
 #include "Runtime/UMG/Public/Animation/MovieScene2DTransformSection.h"
 #include "Runtime/UMG/Public/Animation/MovieScene2DTransformTrack.h"
 #include "Editor/Sequencer/Public/ISectionLayoutBuilder.h"
