Sequencer: Add support for custom property tracks property.

diff --git a/Engine/Source/Editor/MovieSceneTools/Private/TrackEditors/PropertyTrackEditor.cpp b/Engine/Source/Editor/MovieSceneTools/Private/TrackEditors/PropertyTrackEditor.cpp
--- a/Engine/Source/Editor/MovieSceneTools/Private/TrackEditors/PropertyTrackEditor.cpp
+++ b/Engine/Source/Editor/MovieSceneTools/Private/TrackEditors/PropertyTrackEditor.cpp
@@ -9,6 +9,7 @@
 #include "MovieSceneBoolTrack.h"
 #include "MovieSceneByteTrack.h"
 #include "MovieSceneColorTrack.h"
+#include "MovieSceneVisibilityTrack.h"
 #include "ScopedTransaction.h"
 #include "MovieSceneSection.h"
 #include "ISequencerObjectBindingManager.h"
@@ -21,6 +22,7 @@
 #include "MovieSceneBoolSection.h"
 #include "MovieSceneByteSection.h"
 #include "MovieSceneColorSection.h"
+#include "MovieSceneVisibilitySection.h"
 #include "MovieSceneToolHelpers.h"
 #include "PropertyTrackEditor.h"
 #include "CommonMovieSceneTools.h"
@@ -348,7 +350,8 @@ bool FPropertyTrackEditor::SupportsType( TSubclassOf<UMovieSceneTrack> Type ) co
 		Type == UMovieSceneBoolTrack::StaticClass() ||
 		Type == UMovieSceneByteTrack::StaticClass() ||
 		Type == UMovieSceneVectorTrack::StaticClass() ||
-		Type == UMovieSceneColorTrack::StaticClass();
+		Type == UMovieSceneColorTrack::StaticClass() ||
+		Type == UMovieSceneVisibilityTrack::StaticClass();
 }
 
 TSharedRef<ISequencerSection> FPropertyTrackEditor::MakeSectionInterface( UMovieSceneSection& SectionObject, UMovieSceneTrack* Track )
@@ -363,6 +366,7 @@ TSharedRef<ISequencerSection> FPropertyTrackEditor::MakeSectionInterface( UMovie
 		SectionClass == UMovieSceneByteTrack::StaticClass() ? new FBytePropertySection(SectionObject, Track->GetTrackName(), Cast<UMovieSceneByteTrack>(SectionObject.GetOuter())->GetEnum()) :
 		SectionClass == UMovieSceneVectorTrack::StaticClass() ? new FVectorPropertySection( SectionObject, Track->GetTrackName() ) :
 		SectionClass == UMovieSceneFloatTrack::StaticClass() ? new FFloatPropertySection( SectionObject, Track->GetTrackName() ) :
+		SectionClass == UMovieSceneVisibilityTrack::StaticClass() ? new FBoolPropertySection( SectionObject, Track->GetTrackName() ) :
 		new FPropertySection( SectionObject, Track->GetTrackName() )
 		);
 
@@ -557,7 +561,21 @@ void FPropertyTrackEditor::OnKeyProperty( float KeyTime, FPropertyChangedParams
 		{
 			FName PropertyName = PropertyChangedParams.PropertyPath.Last()->GetFName();
 
-			UMovieSceneTrack* Track = GetTrackForObject( ObjectHandle, TrackType::StaticClass(), PropertyName );
+			// Look for a customized track class for this property on the meta data
+			const FString& MetaSequencerTrackClass = PropertyChangedParams.PropertyPath.Last()->GetMetaData(TEXT("SequencerTrackClass"));			
+			TSubclassOf<UMovieSceneTrack> SequencerTrackClass = TrackType::StaticClass();
+			if (!MetaSequencerTrackClass.IsEmpty())
+			{
+				UClass* MetaClass = FindObject<UClass>(ANY_PACKAGE, *MetaSequencerTrackClass);
+				if (!MetaClass)
+				{
+					MetaClass = LoadObject<UClass>(nullptr, *MetaSequencerTrackClass);
+				}
+				if (MetaClass != NULL)
+					SequencerTrackClass = MetaClass;
+			}
+
+			UMovieSceneTrack* Track = GetTrackForObject( ObjectHandle, SequencerTrackClass, PropertyName );
 			if( ensure( Track ) )
 			{
 				TrackType* TypedTrack = CastChecked<TrackType>(Track);
diff --git a/Engine/Source/Runtime/Engine/Classes/GameFramework/Actor.h b/Engine/Source/Runtime/Engine/Classes/GameFramework/Actor.h
--- a/Engine/Source/Runtime/Engine/Classes/GameFramework/Actor.h
+++ b/Engine/Source/Runtime/Engine/Classes/GameFramework/Actor.h
@@ -105,7 +105,7 @@ public:
 	 * Allows us to only see this Actor in the Editor, and not in the actual game.
 	 * @see SetActorHiddenInGame()
 	 */
-	UPROPERTY(EditAnywhere, Category=Rendering, BlueprintReadOnly, replicated, meta=(DisplayName = "Actor Hidden In Game"))
+	UPROPERTY(EditAnywhere, Category=Rendering, BlueprintReadOnly, replicated, meta=(DisplayName = "Actor Hidden In Game", SequencerTrackClass = "MovieSceneVisibilityTrack"))
 	uint32 bHidden:1;
 
 	/** If true, when the actor is spawned it will be sent to the client but receive no further replication updates from the server afterwards. */
