Sequencer: Fix crash on keying bool properties.

diff --git a/Engine/Source/Editor/MovieSceneTools/Private/TrackEditors/PropertyTrackEditor.cpp b/Engine/Source/Editor/MovieSceneTools/Private/TrackEditors/PropertyTrackEditor.cpp
--- a/Engine/Source/Editor/MovieSceneTools/Private/TrackEditors/PropertyTrackEditor.cpp
+++ b/Engine/Source/Editor/MovieSceneTools/Private/TrackEditors/PropertyTrackEditor.cpp
@@ -448,7 +448,7 @@ void FPropertyTrackEditor::OnAnimatedBoolPropertyChanged(const FPropertyChangedP
 	const UBoolProperty* BoolProperty = Cast<const UBoolProperty>(PropertyChangedParams.PropertyPath.Last());
 	if (BoolProperty)
 	{
-		bool BoolValue = BoolProperty->GetPropertyValue(BoolProperty->ContainerPtrToValuePtr<void>(PropertyChangedParams.ObjectsThatChanged.Last()));
+		bool BoolValue = BoolProperty->GetPropertyValue(BoolProperty->ContainerPtrToValuePtr<void>(static_cast<const void*>(PropertyChangedParams.ObjectsThatChanged.Last())));
 
 		AnimatablePropertyChanged(TrackType::StaticClass(), PropertyChangedParams.bRequireAutoKey,
 			FOnKeyProperty::CreateRaw(this, &FPropertyTrackEditor::OnKeyProperty<Type, TrackType>, PropertyChangedParams, BoolValue));
