Renamed the DirectorTrack and all related classes to ShotTrack

diff --git a/Engine/Source/Editor/MovieSceneTools/Private/MovieSceneToolsModule.cpp b/Engine/Source/Editor/MovieSceneTools/Private/MovieSceneToolsModule.cpp
--- a/Engine/Source/Editor/MovieSceneTools/Private/MovieSceneToolsModule.cpp
+++ b/Engine/Source/Editor/MovieSceneTools/Private/MovieSceneToolsModule.cpp
@@ -18,7 +18,7 @@
 #include "MovieSceneTrackEditor.h"
 #include "PropertyTrackEditor.h"
 #include "TransformTrackEditor.h"
-#include "DirectorTrackEditor.h"
+#include "ShotTrackEditor.h"
 #include "SubMovieSceneTrackEditor.h"
 #include "AudioTrackEditor.h"
 #include "AnimationTrackEditor.h"
diff --git a/Engine/Source/Editor/MovieSceneTools/Private/TrackEditors/DirectorTrackEditor.cpp b/Engine/Source/Editor/MovieSceneTools/Private/TrackEditors/DirectorTrackEditor.cpp
deleted file mode 100644
diff --git a/Engine/Source/Editor/MovieSceneTools/Private/TrackEditors/DirectorTrackEditor.h b/Engine/Source/Editor/MovieSceneTools/Private/TrackEditors/DirectorTrackEditor.h
deleted file mode 100644
diff --git a/Engine/Source/Editor/MovieSceneTools/Private/TrackEditors/ShotTrackEditor.cpp b/Engine/Source/Editor/MovieSceneTools/Private/TrackEditors/ShotTrackEditor.cpp
new file mode 100644
--- /dev/null
+++ b/Engine/Source/Editor/MovieSceneTools/Private/TrackEditors/ShotTrackEditor.cpp
@@ -0,0 +1,464 @@
+// Copyright 1998-2015 Epic Games, Inc. All Rights Reserved.
+
+#include "MovieSceneToolsPrivatePCH.h"
+#include "MovieScene.h"
+#include "MovieSceneSection.h"
+#include "ISequencerSection.h"
+#include "PropertyEditorModule.h"
+#include "PropertyHandle.h"
+#include "MovieSceneTrack.h"
+#include "MovieSceneShotTrack.h"
+#include "ScopedTransaction.h"
+#include "ISequencerObjectChangeListener.h"
+#include "ISectionLayoutBuilder.h"
+#include "Runtime/Engine/Public/Slate/SlateTextures.h"
+#include "ObjectTools.h"
+#include "Runtime/Engine/Public/Slate/SceneViewport.h"
+#include "Runtime/MovieSceneTracks/Public/Sections/MovieSceneShotSection.h"
+#include "IKeyArea.h"
+#include "MovieSceneToolHelpers.h"
+#include "MovieSceneTrackEditor.h"
+#include "ShotTrackEditor.h"
+#include "CommonMovieSceneTools.h"
+#include "Camera/CameraActor.h"
+
+namespace AnimatableShotToolConstants
+{
+	// @todo Sequencer Perhaps allow this to be customizable
+	const uint32 DirectorTrackHeight = 100;
+	const double ThumbnailFadeInDuration = 0.25f;
+
+}
+
+/////////////////////////////////////////////////
+// FShotThumbnailPool
+
+FShotThumbnailPool::FShotThumbnailPool(TSharedPtr<ISequencer> InSequencer, uint32 InMaxThumbnailsToDrawAtATime)
+	: Sequencer(InSequencer)
+	, ThumbnailsNeedingDraw()
+	, MaxThumbnailsToDrawAtATime(InMaxThumbnailsToDrawAtATime)
+{
+}
+
+void FShotThumbnailPool::AddThumbnailsNeedingRedraw(const TArray< TSharedPtr<FShotThumbnail> >& InThumbnails)
+{
+	ThumbnailsNeedingDraw.Append(InThumbnails);
+}
+
+void FShotThumbnailPool::RemoveThumbnailsNeedingRedraw(const TArray< TSharedPtr<FShotThumbnail> >& InThumbnails)
+{
+	for (int32 i = 0; i < InThumbnails.Num(); ++i)
+	{
+		ThumbnailsNeedingDraw.Remove(InThumbnails[i]);
+	}
+}
+
+void FShotThumbnailPool::DrawThumbnails()
+{
+	if (ThumbnailsNeedingDraw.Num() > 0)
+	{
+		// save the global time
+		Sequencer.Pin()->SetPerspectiveViewportPossessionEnabled(false);
+		float SavedTime = Sequencer.Pin()->GetGlobalTime();
+
+		uint32 ThumbnailsDrawn = 0;
+		for (int32 ThumbnailsIndex = 0; ThumbnailsDrawn < MaxThumbnailsToDrawAtATime && ThumbnailsIndex < ThumbnailsNeedingDraw.Num(); ++ThumbnailsIndex)
+		{
+			TSharedPtr<FShotThumbnail> Thumbnail = ThumbnailsNeedingDraw[ThumbnailsIndex];
+			
+			if (Thumbnail->IsVisible())
+			{
+				Thumbnail->DrawThumbnail();
+				++ThumbnailsDrawn;
+
+				ThumbnailsNeedingDraw.Remove(Thumbnail);
+			}
+			else if (!Thumbnail->IsValid())
+			{
+				ensure(0);
+
+				ThumbnailsNeedingDraw.Remove(Thumbnail);
+			}
+		}
+
+		// restore the global time
+		Sequencer.Pin()->SetGlobalTime(SavedTime);
+		Sequencer.Pin()->SetPerspectiveViewportPossessionEnabled(true);
+	}
+}
+
+
+/////////////////////////////////////////////////
+// FShotThumbnail
+
+FShotThumbnail::FShotThumbnail(TSharedPtr<FShotSection> InSection, TRange<float> InTimeRange)
+	: OwningSection(InSection)
+	, Texture(NULL)
+	, TimeRange(InTimeRange)
+	, FadeInStartTime(0.0)
+{
+	Texture = new FSlateTexture2DRHIRef(GetSize().X, GetSize().Y, PF_B8G8R8A8, NULL, TexCreate_Dynamic, true);
+
+	BeginInitResource( Texture );
+
+}
+
+FShotThumbnail::~FShotThumbnail()
+{
+	BeginReleaseResource( Texture );
+
+	FlushRenderingCommands();
+
+	if (Texture) 
+	{
+		delete Texture;
+	}
+}
+
+FIntPoint FShotThumbnail::GetSize() const {return FIntPoint(OwningSection.Pin()->GetThumbnailWidth(), AnimatableShotToolConstants::DirectorTrackHeight);}
+bool FShotThumbnail::RequiresVsync() const {return false;}
+
+FSlateShaderResource* FShotThumbnail::GetViewportRenderTargetTexture() const
+{
+	return Texture;
+}
+
+void FShotThumbnail::DrawThumbnail()
+{
+	OwningSection.Pin()->DrawViewportThumbnail(SharedThis(this));
+	FadeInStartTime = FSlateApplication::Get().GetCurrentTime();
+}
+
+float FShotThumbnail::GetTime() const {return TimeRange.GetLowerBoundValue();}
+
+void FShotThumbnail::CopyTextureIn(FSlateRenderTargetRHI* InTexture)
+{
+	ENQUEUE_UNIQUE_RENDER_COMMAND_TWOPARAMETER( ReadTexture,
+		FSlateRenderTargetRHI*, RenderTarget, InTexture,
+		FSlateTexture2DRHIRef*, TargetTexture, Texture,
+	{
+		RHICmdList.CopyToResolveTarget(RenderTarget->GetRHIRef(), TargetTexture->GetTypedResource(), false, FResolveParams());
+	});
+}
+
+float FShotThumbnail::GetFadeInCurve() const 
+{
+	return (float)FMath::Clamp(FadeInStartTime / AnimatableShotToolConstants::ThumbnailFadeInDuration, 0.0, 1.0);
+}
+
+bool FShotThumbnail::IsVisible() const
+{
+	return OwningSection.IsValid() ? TimeRange.Overlaps(OwningSection.Pin()->GetVisibleTimeRange()) : false;
+}
+
+bool FShotThumbnail::IsValid() const
+{
+	return OwningSection.IsValid();
+}
+
+/////////////////////////////////////////////////
+// FShotSection
+
+FShotSection::FShotSection( TSharedPtr<ISequencer> InSequencer, TSharedPtr<FShotThumbnailPool> InThumbnailPool, UMovieSceneSection& InSection, UObject* InTargetObject )
+	: Section( &InSection )
+	, Sequencer( InSequencer )
+	, Camera(Cast<ACameraActor>(InTargetObject))
+	, ThumbnailPool( InThumbnailPool )
+	, Thumbnails()
+	, ThumbnailWidth(0)
+	, StoredSize(ForceInit)
+	, StoredStartTime(0.f)
+	, VisibleTimeRange(TRange<float>::Empty())
+	, InternalViewportScene(nullptr)
+	, InternalViewportClient(nullptr)
+{
+	if (Camera.IsValid())
+	{
+		// @todo Sequencer optimize This code shouldn't even be called if this is offscreen
+		// the following code could be generated on demand when the section is scrolled onscreen
+		InternalViewportClient = MakeShareable(new FLevelEditorViewportClient(nullptr));
+		InternalViewportClient->ViewportType = LVT_Perspective;
+		InternalViewportClient->bDisableInput = true;
+		InternalViewportClient->bDrawAxes = false;
+		InternalViewportClient->EngineShowFlags = FEngineShowFlags(ESFIM_Game);
+		InternalViewportClient->EngineShowFlags.DisableAdvancedFeatures();
+		InternalViewportClient->SetActorLock(Camera.Get());
+
+		InternalViewportScene = MakeShareable(new FSceneViewport(InternalViewportClient.Get(), nullptr));
+		InternalViewportClient->Viewport = InternalViewportScene.Get();
+
+		CalculateThumbnailWidthAndResize();
+	}
+}
+
+FShotSection::~FShotSection()
+{
+	ThumbnailPool.Pin()->RemoveThumbnailsNeedingRedraw(Thumbnails);
+
+	if (InternalViewportClient.IsValid())
+	{
+		InternalViewportClient->Viewport = nullptr;
+	}
+}
+
+UMovieSceneSection* FShotSection::GetSectionObject() 
+{
+	return Section;
+}
+
+FText FShotSection::GetSectionTitle() const
+{
+	return Cast<UMovieSceneShotSection>(Section)->GetTitle();
+}
+
+float FShotSection::GetSectionHeight() const
+{
+	return AnimatableShotToolConstants::DirectorTrackHeight;
+}
+
+FReply FShotSection::OnSectionDoubleClicked( const FGeometry& SectionGeometry, const FPointerEvent& MouseEvent )
+{
+	if( MouseEvent.GetEffectingButton() == EKeys::LeftMouseButton )
+	{
+		check( Section->IsA<UMovieSceneSection>() );
+
+		Sequencer.Pin()->FilterToSelectedShotSections();
+	}
+
+	return FReply::Handled();
+}
+
+int32 FShotSection::OnPaintSection( const FGeometry& AllottedGeometry, const FSlateRect& SectionClippingRect, FSlateWindowElementList& OutDrawElements, int32 LayerId, bool bParentEnabled ) const
+{
+	if (Camera.IsValid())
+	{
+		for (int32 i = 0; i < Thumbnails.Num(); ++i)
+		{
+			FGeometry TruncatedGeometry = AllottedGeometry.MakeChild(
+				FVector2D(ThumbnailWidth, AllottedGeometry.Size.Y),
+				FSlateLayoutTransform(AllottedGeometry.Scale, FVector2D(i * ThumbnailWidth, 0.0f))
+				);
+			ensure(TruncatedGeometry.Size.Y == AnimatableShotToolConstants::DirectorTrackHeight);
+
+			FSlateDrawElement::MakeViewport(
+				OutDrawElements,
+				LayerId,
+				TruncatedGeometry.ToPaintGeometry(),
+				Thumbnails[i],
+				SectionClippingRect,
+				false,
+				false,
+				ESlateDrawEffect::None,
+				FLinearColor::White
+			);
+
+			FSlateDrawElement::MakeBox( 
+				OutDrawElements,
+				LayerId + 1,
+				TruncatedGeometry.ToPaintGeometry(),
+				FEditorStyle::GetBrush("WhiteBrush"),
+				SectionClippingRect,
+				ESlateDrawEffect::None,
+				FColor(255, 255, 255, Thumbnails[i]->GetFadeInCurve() * 255)
+			);
+		}
+	}
+	else
+	{
+		FSlateDrawElement::MakeBox( 
+			OutDrawElements,
+			LayerId,
+			AllottedGeometry.ToPaintGeometry(),
+			FEditorStyle::GetBrush("Sequencer.GenericSection.Background"),
+			SectionClippingRect
+		); 
+	}
+
+	return LayerId + 1;
+}
+
+void FShotSection::Tick( const FGeometry& AllottedGeometry, const FGeometry& ParentGeometry, const double InCurrentTime, const float InDeltaTime )
+{
+	if (!Camera.IsValid())
+	{
+		return;
+	}
+
+	FTimeToPixel TimeToPixelConverter( AllottedGeometry, TRange<float>( Section->GetStartTime(), Section->GetEndTime() ) );
+
+	// get the visible time range
+	VisibleTimeRange = TRange<float>(
+		TimeToPixelConverter.PixelToTime(-AllottedGeometry.Position.X),
+		TimeToPixelConverter.PixelToTime(-AllottedGeometry.Position.X + ParentGeometry.Size.X));
+
+	FIntPoint AllocatedSize = AllottedGeometry.Size.IntPoint();
+	float StartTime = Section->GetStartTime();
+
+	if (AllocatedSize.X != StoredSize.X || !FMath::IsNearlyEqual(StartTime, StoredStartTime))
+	{
+		RegenerateViewportThumbnails(AllocatedSize);
+	}
+}
+
+uint32 FShotSection::GetThumbnailWidth() const
+{
+	return ThumbnailWidth;
+}
+
+void FShotSection::RegenerateViewportThumbnails(const FIntPoint& Size)
+{
+	StoredSize = Size;
+	StoredStartTime = Section->GetStartTime();
+	
+	ThumbnailPool.Pin()->RemoveThumbnailsNeedingRedraw(Thumbnails);
+	Thumbnails.Empty();
+
+	if (Size.X == 0 || Size.Y == 0)
+	{
+		return;
+	}
+
+	CalculateThumbnailWidthAndResize();
+
+	int32 ThumbnailCount = FMath::DivideAndRoundUp(StoredSize.X, (int32)ThumbnailWidth);
+	
+	float StartTime = Section->GetStartTime();
+	float EndTime = Section->GetEndTime();
+	float DeltaTime = EndTime - StartTime;
+
+	int32 TotalX = StoredSize.X;
+	
+	// @todo sequencer optimize this to use a single viewport and a single thumbnail
+	// instead paste the textures into the thumbnail at specific offsets
+	for (int32 i = 0; i < ThumbnailCount; ++i)
+	{
+		int32 StartX = i * ThumbnailWidth;
+		float FractionThrough = (float)StartX / (float)TotalX;
+		float Time = StartTime + DeltaTime * FractionThrough;
+
+		int32 NextStartX = (i+1) * ThumbnailWidth;
+		float NextFractionThrough = (float)NextStartX / (float)TotalX;
+		float NextTime = FMath::Min(StartTime + DeltaTime * NextFractionThrough, EndTime);
+	
+		check(FractionThrough >= 0.f && FractionThrough <= 1.f && NextFractionThrough >= 0.f);
+		TSharedPtr<FShotThumbnail> NewThumbnail = MakeShareable(new FShotThumbnail(SharedThis(this), TRange<float>(Time, NextTime)));
+
+		Thumbnails.Add(NewThumbnail);
+	}
+
+	// @todo Sequencer Optimize Only say a thumbnail needs redraw if it is onscreen
+	ThumbnailPool.Pin()->AddThumbnailsNeedingRedraw(Thumbnails);
+}
+
+void FShotSection::DrawViewportThumbnail(TSharedPtr<FShotThumbnail> ShotThumbnail)
+{
+	Sequencer.Pin()->SetGlobalTime(ShotThumbnail->GetTime());
+	InternalViewportClient->UpdateViewForLockedActor();
+			
+	GWorld->SendAllEndOfFrameUpdates();
+	InternalViewportScene->Draw(false);
+	ShotThumbnail->CopyTextureIn((FSlateRenderTargetRHI*)InternalViewportScene->GetViewportRenderTargetTexture());
+
+	FlushRenderingCommands();
+}
+
+void FShotSection::CalculateThumbnailWidthAndResize()
+{
+	// @todo Sequencer We also should detect when the property is changed or keyframed to update
+
+	// get the aspect ratio at the first frame
+	// if it changes over time, we ignore this
+	Sequencer.Pin()->SetPerspectiveViewportPossessionEnabled(false);
+	float SavedTime = Sequencer.Pin()->GetGlobalTime();
+	
+	Sequencer.Pin()->SetGlobalTime(Section->GetStartTime());
+	uint32 OutThumbnailWidth = AnimatableShotToolConstants::DirectorTrackHeight * Camera->GetCameraComponent()->AspectRatio;
+
+	// restore the global time
+	Sequencer.Pin()->SetGlobalTime(SavedTime);
+	Sequencer.Pin()->SetPerspectiveViewportPossessionEnabled(true);
+
+	if (OutThumbnailWidth != ThumbnailWidth)
+	{
+		ThumbnailWidth = OutThumbnailWidth;
+		InternalViewportScene->UpdateViewportRHI( false, ThumbnailWidth, AnimatableShotToolConstants::DirectorTrackHeight, EWindowMode::Windowed );
+	}
+}
+
+/////////////////////////////////////////////////
+// FDirectorTrackEditor
+
+FDirectorTrackEditor::FDirectorTrackEditor( TSharedRef<ISequencer> InSequencer )
+	: FMovieSceneTrackEditor( InSequencer ) 
+{
+	ThumbnailPool = MakeShareable(new FShotThumbnailPool(InSequencer));
+}
+
+FDirectorTrackEditor::~FDirectorTrackEditor()
+{
+}
+
+TSharedRef<FMovieSceneTrackEditor> FDirectorTrackEditor::CreateTrackEditor( TSharedRef<ISequencer> InSequencer )
+{
+	return MakeShareable( new FDirectorTrackEditor( InSequencer ) );
+}
+
+bool FDirectorTrackEditor::SupportsType( TSubclassOf<UMovieSceneTrack> Type ) const
+{
+	return Type == UMovieSceneShotTrack::StaticClass();
+}
+
+TSharedRef<ISequencerSection> FDirectorTrackEditor::MakeSectionInterface( UMovieSceneSection& SectionObject, UMovieSceneTrack* Track )
+{
+	check( SupportsType( SectionObject.GetOuter()->GetClass() ) );
+	
+	TArray<UObject*> OutObjects;
+	// @todo Sequencer - Sub-MovieScenes The director track may be able to get cameras from sub-movie scenes
+	GetSequencer()->GetRuntimeObjects( GetSequencer()->GetRootMovieSceneInstance(), Cast<const UMovieSceneShotSection>(&SectionObject)->GetCameraGuid(), OutObjects);
+	check(OutObjects.Num() <= 1);
+	TSharedRef<ISequencerSection> NewSection( new FShotSection( GetSequencer(), ThumbnailPool, SectionObject, OutObjects.Num() ? OutObjects[0] : NULL ) );
+
+	return NewSection;
+}
+
+void FDirectorTrackEditor::AddKey(const FGuid& ObjectGuid, UObject* AdditionalAsset)
+{
+	TArray<UObject*> OutObjects;
+	// @todo Sequencer - Sub-MovieScenes The director track may be able to get cameras from sub-movie scenes
+	GetSequencer()->GetRuntimeObjects(  GetSequencer()->GetRootMovieSceneInstance(), ObjectGuid, OutObjects);
+	bool bValidKey = OutObjects.Num() == 1 && OutObjects[0]->IsA<ACameraActor>();
+	if (bValidKey)
+	{
+		AnimatablePropertyChanged( UMovieSceneShotTrack::StaticClass(), false, FOnKeyProperty::CreateRaw( this, &FDirectorTrackEditor::AddKeyInternal, ObjectGuid ) );
+	}
+}
+
+void FDirectorTrackEditor::Tick(float DeltaTime)
+{
+	if (FSlateThrottleManager::Get().IsAllowingExpensiveTasks())
+	{
+		ThumbnailPool->DrawThumbnails();
+	}
+}
+
+void FDirectorTrackEditor::BuildObjectBindingContextMenu(FMenuBuilder& MenuBuilder, const FGuid& ObjectBinding, const UClass* ObjectClass)
+{
+	if (ObjectClass->IsChildOf(ACameraActor::StaticClass()))
+	{
+		const TSharedPtr<ISequencer> ParentSequencer = GetSequencer();
+
+		MenuBuilder.AddMenuEntry(
+			NSLOCTEXT("Sequencer", "AddShot", "Add New Shot"),
+			NSLOCTEXT("Sequencer", "AddShotTooltip", "Adds a new shot using this camera at the scrubber location."),
+			FSlateIcon(),
+			FUIAction(FExecuteAction::CreateSP(ParentSequencer.Get(), &ISequencer::AddNewShot, ObjectBinding))
+			);
+	}
+}
+
+void FDirectorTrackEditor::AddKeyInternal( float KeyTime, const FGuid ObjectGuid )
+{
+	UMovieSceneTrack* Track = GetMasterTrack( UMovieSceneShotTrack::StaticClass() );
+
+	Cast<UMovieSceneShotTrack>(Track)->AddNewShot(ObjectGuid, KeyTime);
+}
diff --git a/Engine/Source/Editor/MovieSceneTools/Private/TrackEditors/ShotTrackEditor.h b/Engine/Source/Editor/MovieSceneTools/Private/TrackEditors/ShotTrackEditor.h
new file mode 100644
--- /dev/null
+++ b/Engine/Source/Editor/MovieSceneTools/Private/TrackEditors/ShotTrackEditor.h
@@ -0,0 +1,182 @@
+// Copyright 1998-2015 Epic Games, Inc. All Rights Reserved.
+
+#pragma once
+
+
+/**
+ * Tools for director tracks
+ */
+class FDirectorTrackEditor : public FMovieSceneTrackEditor
+{
+public:
+	/**
+	 * Constructor
+	 *
+	 * @param InSequencer	The sequencer instance to be used by this tool
+	 */
+	FDirectorTrackEditor( TSharedRef<ISequencer> InSequencer );
+	~FDirectorTrackEditor();
+
+	/**
+	 * Creates an instance of this class.  Called by a sequencer 
+	 *
+	 * @param OwningSequencer The sequencer instance to be used by this tool
+	 * @return The new instance of this class
+	 */
+	static TSharedRef<FMovieSceneTrackEditor> CreateTrackEditor( TSharedRef<ISequencer> OwningSequencer );
+
+	/** FMovieSceneTrackEditor Interface */
+	virtual bool SupportsType( TSubclassOf<UMovieSceneTrack> Type ) const override;
+	virtual TSharedRef<ISequencerSection> MakeSectionInterface( UMovieSceneSection& SectionObject, UMovieSceneTrack* Track ) override;
+	virtual void AddKey(const FGuid& ObjectGuid, UObject* AdditionalAsset = NULL) override;
+	virtual void Tick(float DeltaTime) override;
+	virtual void BuildObjectBindingContextMenu(FMenuBuilder& MenuBuilder, const FGuid& ObjectBinding, const UClass* ObjectClass) override;
+
+private:
+	/** Delegate for AnimatablePropertyChanged in AddKey */
+	void AddKeyInternal(float AutoKeyTime, const FGuid ObjectGuid);
+
+private:
+	/** The Thumbnail pool which draws all the viewport thumbnails for the director track */
+	TSharedPtr<class FShotThumbnailPool> ThumbnailPool;
+};
+
+
+/**
+ * Shot Thumbnail pool, which keeps a list of thumbnails that need to be drawn
+ * and draws them incrementally
+ */
+class FShotThumbnailPool
+{
+public:
+	FShotThumbnailPool(TSharedPtr<ISequencer> InSequencer, uint32 InMaxThumbnailsToDrawAtATime = 1);
+
+	/** Requests that the passed in thumbnails need to be drawn */
+	void AddThumbnailsNeedingRedraw(const TArray< TSharedPtr<class FShotThumbnail> >& InThumbnails);
+
+	/** Informs the pool that the thumbnails passed in no longer need to be drawn */
+	void RemoveThumbnailsNeedingRedraw(const TArray< TSharedPtr<class FShotThumbnail> >& InThumbnails);
+
+	/** Draws a small number of thumbnails that are enqueued for drawing */
+	void DrawThumbnails();
+
+private:
+	/** Parent sequencer we're drawing thumbnails for */
+	TWeakPtr<ISequencer> Sequencer;
+
+	/** Thumbnails enqueued for drawing */
+	TArray< TSharedPtr<class FShotThumbnail> > ThumbnailsNeedingDraw;
+
+	/** How many thumbnails we are allowed to draw in a single DrawThumbnails call */
+	uint32 MaxThumbnailsToDrawAtATime;
+};
+
+
+/**
+ * Shot Thumbnail, which keeps a Texture to be displayed by a viewport
+ */
+class FShotThumbnail : public ISlateViewport, public TSharedFromThis<FShotThumbnail>
+{
+public:
+	FShotThumbnail(TSharedPtr<class FShotSection> InSection, TRange<float> InTimeRange);
+	~FShotThumbnail();
+
+	/* ISlateViewport interface */
+	virtual FIntPoint GetSize() const override;
+	virtual class FSlateShaderResource* GetViewportRenderTargetTexture() const override;
+	virtual bool RequiresVsync() const override;
+
+	/** Gets the time that this thumbnail is a rendering of */
+	float GetTime() const;
+
+	/** Renders the thumbnail to the texture */
+	void DrawThumbnail();
+
+	/** Copies the incoming render target to this thumbnails texture */
+	void CopyTextureIn(class FSlateRenderTargetRHI* InTexture);
+
+	/** Gets the curve for fading in the thumbnail */
+	float GetFadeInCurve() const;
+
+	/** Returns whether this thumbnail is visible based on the shot section geometry visibility */
+	bool IsVisible() const;
+
+	bool IsValid() const;
+
+private:
+	/** Parent shot section we are a thumbnail of */
+	TWeakPtr<class FShotSection> OwningSection;
+	
+	/** The Texture RHI that holds the thumbnail */
+	class FSlateTexture2DRHIRef* Texture;
+
+	/** Where in time this thumbnail is a rendering of */
+	TRange<float> TimeRange;
+
+	/** The time when the thumbnail started to fade in. */
+	double FadeInStartTime;
+};
+
+
+/**
+ * Shot section, which paints and ticks the appropriate section
+ */
+class FShotSection : public ISequencerSection, public TSharedFromThis<FShotSection>
+{
+public:
+	FShotSection( TSharedPtr<ISequencer> InSequencer, TSharedPtr<FShotThumbnailPool> InThumbnailPool, UMovieSceneSection& InSection, UObject* InTargetObject );
+	~FShotSection();
+
+	/** ISequencerSection interface */
+	virtual UMovieSceneSection* GetSectionObject() override;
+	virtual int32 OnPaintSection( const FGeometry& AllottedGeometry, const FSlateRect& SectionClippingRect, FSlateWindowElementList& OutDrawElements, int32 LayerId, bool bParentEnabled ) const override;
+	virtual void Tick( const FGeometry& AllottedGeometry, const FGeometry& ParentGeometry, const double InCurrentTime, const float InDeltaTime ) override;
+	virtual FText GetDisplayName() const override { return NSLOCTEXT("FShotSection", "DirectorTrack", "Director Track"); }
+	virtual FText GetSectionTitle() const override;
+	virtual float GetSectionHeight() const override;
+	virtual void GenerateSectionLayout( class ISectionLayoutBuilder& LayoutBuilder ) const override {}
+	virtual FReply OnSectionDoubleClicked( const FGeometry& SectionGeometry, const FPointerEvent& MouseEvent ) override;
+
+	/** Gets the thumbnail width */
+	uint32 GetThumbnailWidth() const;
+
+	/** Regenerates all viewports and thumbnails at the new size */
+	void RegenerateViewportThumbnails(const FIntPoint& Size);
+
+	/** Draws the passed in viewport thumbnail and copies it to the thumbnail's texture */
+	void DrawViewportThumbnail(TSharedPtr<FShotThumbnail> ShotThumbnail);
+
+	/** Calculates and sets the thumbnail width, and resizes if it is different than before */
+	void CalculateThumbnailWidthAndResize();
+
+	/** Gets the time range of what in the sequencer is visible */
+	TRange<float> GetVisibleTimeRange() const {return VisibleTimeRange;}
+
+private:
+	/** The section we are visualizing */
+	UMovieSceneSection* Section;
+	/** The parent sequencer we are a part of */
+	TWeakPtr<ISequencer> Sequencer;
+	/** The actual camera actor we are possessing */
+	TWeakObjectPtr<ACameraActor> Camera;
+
+	/** The thumbnail pool that we are sending all of our thumbnails to */
+	TWeakPtr<FShotThumbnailPool> ThumbnailPool;
+
+	/** A list of all thumbnails this shot section has */
+	TArray< TSharedPtr<FShotThumbnail> > Thumbnails;
+	/** The width of our thumbnails */
+	uint32 ThumbnailWidth;
+	/** The stored size of this section in the Slate geometry */
+	FIntPoint StoredSize;
+	/** The stored start time, to query for invalidations */
+	float StoredStartTime;
+
+	/** Cached Time Range of the visible parent section area */
+	TRange<float> VisibleTimeRange;
+	
+	/** An internal viewport scene we use to render the thumbnails with */
+	TSharedPtr<FSceneViewport> InternalViewportScene;
+	/** An internal editor viewport client to render the thumbnails with */
+	TSharedPtr<FLevelEditorViewportClient> InternalViewportClient;
+};
diff --git a/Engine/Source/Editor/Sequencer/Private/SectionDragOperations.cpp b/Engine/Source/Editor/Sequencer/Private/SectionDragOperations.cpp
--- a/Engine/Source/Editor/Sequencer/Private/SectionDragOperations.cpp
+++ b/Engine/Source/Editor/Sequencer/Private/SectionDragOperations.cpp
@@ -9,7 +9,7 @@
 #include "CommonMovieSceneTools.h"
 #include "MovieScene.h"
 #include "MovieSceneTrack.h"
-#include "MovieSceneDirectorTrack.h"
+#include "MovieSceneShotTrack.h"
 
 /** How many pixels near the mouse has to be before snapping occurs */
 const float PixelSnapWidth = 10.f;
@@ -78,10 +78,10 @@ void FSequencerDragOperation::GetSectionSnapTimes(TArray<float>& OutSnapTimes, U
 	// snap to director track if it exists, and we are not the director track
 	UMovieSceneTrack* OuterTrack = Cast<UMovieSceneTrack>(Section->GetOuter());
 	UMovieScene* MovieScene = Cast<UMovieScene>(OuterTrack->GetOuter());
-	UMovieSceneTrack* DirectorTrack = MovieScene->FindMasterTrack(UMovieSceneDirectorTrack::StaticClass());
-	if (DirectorTrack && OuterTrack != DirectorTrack)
+	UMovieSceneTrack* ShotTrack = MovieScene->FindMasterTrack(UMovieSceneShotTrack::StaticClass());
+	if (ShotTrack && OuterTrack != ShotTrack)
 	{
-		TArray<UMovieSceneSection*> ShotSections = DirectorTrack->GetAllSections();
+		TArray<UMovieSceneSection*> ShotSections = ShotTrack->GetAllSections();
 		for (int32 SectionIndex = 0; SectionIndex < ShotSections.Num(); ++SectionIndex)
 		{
 			auto Shot = ShotSections[SectionIndex];
diff --git a/Engine/Source/Editor/Sequencer/Private/Sequencer.cpp b/Engine/Source/Editor/Sequencer/Private/Sequencer.cpp
--- a/Engine/Source/Editor/Sequencer/Private/Sequencer.cpp
+++ b/Engine/Source/Editor/Sequencer/Private/Sequencer.cpp
@@ -17,7 +17,7 @@
 #include "Kismet2/BlueprintEditorUtils.h"
 #include "BlueprintEditorModule.h"
 #include "MovieSceneTrack.h"
-#include "MovieSceneDirectorTrack.h"
+#include "MovieSceneShotTrack.h"
 #include "MovieSceneAudioTrack.h"
 #include "MovieSceneAnimationTrack.h"
 #include "MovieSceneTrackEditor.h"
@@ -118,9 +118,9 @@ void FSequencer::InitSequencer( const FSequencerInitParams& InitParams, const TA
 			TSharedRef<FMovieSceneTrackEditor> TrackEditor = TrackEditorDelegates[DelegateIndex].Execute( SharedThis( this ) );
 
 			// Keep track of certain editors
-			if ( TrackEditor->SupportsType( UMovieSceneDirectorTrack::StaticClass() ) )
+			if ( TrackEditor->SupportsType( UMovieSceneShotTrack::StaticClass() ) )
 			{
-				DirectorTrackEditor = TrackEditor;
+				ShotTrackEditor = TrackEditor;
 			}
 			else if ( TrackEditor->SupportsType( UMovieSceneAnimationTrack::StaticClass() ) )
 			{
@@ -357,9 +357,9 @@ void FSequencer::PopToMovieScene( TSharedRef<FMovieSceneInstance> SubMovieSceneI
 
 void FSequencer::AddNewShot(FGuid CameraGuid)
 {
-	if (DirectorTrackEditor.IsValid())
+	if (ShotTrackEditor.IsValid())
 	{
-		DirectorTrackEditor.Pin()->AddKey(CameraGuid);
+		ShotTrackEditor.Pin()->AddKey(CameraGuid);
 	}
 }
 
@@ -925,7 +925,7 @@ TRange<float> FSequencer::GetTimeBounds() const
 	}
 
 	const UMovieScene* MovieScene = GetFocusedMovieScene();
-	const UMovieSceneTrack* AnimatableShot = MovieScene->FindMasterTrack( UMovieSceneDirectorTrack::StaticClass() );
+	const UMovieSceneTrack* AnimatableShot = MovieScene->FindMasterTrack( UMovieSceneShotTrack::StaticClass() );
 	if (AnimatableShot)
 	{
 		// try getting filtered shot boundaries
@@ -1532,8 +1532,8 @@ void FSequencer::SetKey()
 		FGuid ObjectGuid = GetHandleToObject(*It);
 		for ( auto& TrackEditor : TrackEditors )
 		{
-			// @todo Handle this director track business better
-			if (TrackEditor != DirectorTrackEditor.Pin())
+			// @todo Handle this shot track business better
+			if (TrackEditor != ShotTrackEditor.Pin())
 			{
 				TrackEditor->AddKey(ObjectGuid);
 			}
diff --git a/Engine/Source/Editor/Sequencer/Private/Sequencer.h b/Engine/Source/Editor/Sequencer/Private/Sequencer.h
--- a/Engine/Source/Editor/Sequencer/Private/Sequencer.h
+++ b/Engine/Source/Editor/Sequencer/Private/Sequencer.h
@@ -379,7 +379,7 @@ private:
 	TArray< TSharedPtr<FMovieSceneTrackEditor> > TrackEditors;
 
 	/** The editors we keep track of for special behaviors */
-	TWeakPtr<FMovieSceneTrackEditor> DirectorTrackEditor;
+	TWeakPtr<FMovieSceneTrackEditor> ShotTrackEditor;
 	TWeakPtr<FMovieSceneTrackEditor> AnimationTrackEditor;
 
 	/** Manager for handling runtime object bindings */
diff --git a/Engine/Source/Runtime/MovieSceneTracks/Private/TrackInstances/MovieSceneDirectorTrackInstance.cpp b/Engine/Source/Runtime/MovieSceneTracks/Private/TrackInstances/MovieSceneDirectorTrackInstance.cpp
deleted file mode 100644
diff --git a/Engine/Source/Runtime/MovieSceneTracks/Private/TrackInstances/MovieSceneShotTrackInstance.cpp b/Engine/Source/Runtime/MovieSceneTracks/Private/TrackInstances/MovieSceneShotTrackInstance.cpp
new file mode 100644
--- /dev/null
+++ b/Engine/Source/Runtime/MovieSceneTracks/Private/TrackInstances/MovieSceneShotTrackInstance.cpp
@@ -0,0 +1,47 @@
+// Copyright 1998-2015 Epic Games, Inc. All Rights Reserved.
+
+#include "MovieSceneTracksPrivatePCH.h"
+#include "MovieSceneShotTrack.h"
+#include "MovieSceneShotTrackInstance.h"
+#include "MovieSceneShotSection.h"
+#include "IMovieScenePlayer.h"
+
+
+FMovieSceneShotTrackInstance::FMovieSceneShotTrackInstance( UMovieSceneShotTrack& InDirectorTrack )
+{
+	DirectorTrack = &InDirectorTrack;
+}
+
+
+void FMovieSceneShotTrackInstance::Update( float Position, float LastPosition, const TArray<UObject*>& RuntimeObjects, class IMovieScenePlayer& Player ) 
+{
+	// possess the 'first' shot section in order of track index
+	ACameraActor* PossessedCamera = NULL;
+	UMovieSceneShotSection* FirstShotSection = NULL;
+
+	const TArray<UMovieSceneSection*>& ShotSections = DirectorTrack->GetAllSections();
+
+	for (int32 ShotIndex = 0; ShotIndex < ShotSections.Num(); ++ShotIndex)
+	{
+		UMovieSceneShotSection* ShotSection = CastChecked<UMovieSceneShotSection>(ShotSections[ShotIndex]);
+
+		if (ShotSection->IsTimeWithinSection(Position) &&
+			(FirstShotSection == NULL || FirstShotSection->GetRowIndex() > ShotSection->GetRowIndex()))
+		{
+			TArray<UObject*> CameraObjects;
+			// @todo Sequencer - Sub-moviescenes: Get the cameras from the root movie scene instance.  We should support adding cameras for sub-moviescenes as shots
+			Player.GetRuntimeObjects( Player.GetRootMovieSceneInstance(), ShotSection->GetCameraGuid(), CameraObjects);
+
+			if (CameraObjects.Num() > 0)
+			{
+				check(CameraObjects.Num() == 1 && CameraObjects[0]->IsA(ACameraActor::StaticClass()));
+
+				PossessedCamera = Cast<ACameraActor>(CameraObjects[0]);
+
+				FirstShotSection = ShotSection;
+			}
+		}
+	}
+
+	Player.UpdatePreviewViewports(PossessedCamera);
+}
diff --git a/Engine/Source/Runtime/MovieSceneTracks/Private/Tracks/MovieSceneAudioTrack.cpp b/Engine/Source/Runtime/MovieSceneTracks/Private/Tracks/MovieSceneAudioTrack.cpp
--- a/Engine/Source/Runtime/MovieSceneTracks/Private/Tracks/MovieSceneAudioTrack.cpp
+++ b/Engine/Source/Runtime/MovieSceneTracks/Private/Tracks/MovieSceneAudioTrack.cpp
@@ -9,7 +9,7 @@
 #include "Runtime/Engine/Public/AudioDecompress.h"
 #include "MovieSceneAudioTrackInstance.h"
 
-#define LOCTEXT_NAMESPACE "MovieSceneDirectorTrack"
+#define LOCTEXT_NAMESPACE "MovieSceneAudioTrack"
 
 UMovieSceneAudioTrack::UMovieSceneAudioTrack( const FObjectInitializer& ObjectInitializer )
 	: Super( ObjectInitializer )
diff --git a/Engine/Source/Runtime/MovieSceneTracks/Private/Tracks/MovieSceneDirectorTrack.cpp b/Engine/Source/Runtime/MovieSceneTracks/Private/Tracks/MovieSceneDirectorTrack.cpp
deleted file mode 100644
diff --git a/Engine/Source/Runtime/MovieSceneTracks/Private/Tracks/MovieSceneShotTrack.cpp b/Engine/Source/Runtime/MovieSceneTracks/Private/Tracks/MovieSceneShotTrack.cpp
new file mode 100644
--- /dev/null
+++ b/Engine/Source/Runtime/MovieSceneTracks/Private/Tracks/MovieSceneShotTrack.cpp
@@ -0,0 +1,40 @@
+// Copyright 1998-2015 Epic Games, Inc. All Rights Reserved.
+
+#include "MovieSceneTracksPrivatePCH.h"
+#include "MovieSceneShotSection.h"
+#include "MovieSceneShotTrack.h"
+#include "IMovieScenePlayer.h"
+#include "MovieSceneShotTrackInstance.h"
+
+#define LOCTEXT_NAMESPACE "MovieSceneShotTrack"
+
+UMovieSceneShotTrack::UMovieSceneShotTrack( const FObjectInitializer& ObjectInitializer )
+	: Super( ObjectInitializer )
+{
+}
+
+
+FName UMovieSceneShotTrack::GetTrackName() const
+{
+	static const FName UniqueName("Shots");
+
+	return UniqueName;
+}
+
+TSharedPtr<IMovieSceneTrackInstance> UMovieSceneShotTrack::CreateInstance()
+{
+	return MakeShareable( new FMovieSceneShotTrackInstance( *this ) ); 
+}
+
+void UMovieSceneShotTrack::AddNewShot(FGuid CameraHandle, float Time)
+{
+	UMovieSceneShotSection* NewSection = NewObject<UMovieSceneShotSection>(this);
+	NewSection->InitialPlacement(SubMovieSceneSections, Time, Time + 1.f, SupportsMultipleRows()); //@todo find a better default end time
+	NewSection->SetCameraGuid(CameraHandle);
+	NewSection->SetTitle(LOCTEXT("AddNewShot", "New Shot"));
+
+	SubMovieSceneSections.Add(NewSection);
+}
+
+
+#undef LOCTEXT_NAMESPACE
\ No newline at end of file
diff --git a/Engine/Source/Runtime/MovieSceneTracks/Private/Tracks/SubMovieSceneTrack.cpp b/Engine/Source/Runtime/MovieSceneTracks/Private/Tracks/SubMovieSceneTrack.cpp
--- a/Engine/Source/Runtime/MovieSceneTracks/Private/Tracks/SubMovieSceneTrack.cpp
+++ b/Engine/Source/Runtime/MovieSceneTracks/Private/Tracks/SubMovieSceneTrack.cpp
@@ -23,7 +23,7 @@ FName USubMovieSceneTrack::GetTrackName() const
 	FName Name = DefaultName;
 	if( SubMovieSceneSections.Num() > 0 )
 	{
-		Name = SubMovieSceneSections[0]->GetMovieScene()->GetFName();
+		Name = CastChecked<USubMovieSceneSection>(SubMovieSceneSections[0])->GetMovieScene()->GetFName();
 	}
 
 	return Name;
@@ -54,6 +54,11 @@ void USubMovieSceneTrack::RemoveSection( UMovieSceneSection* Section )
 	}
 }
 
+void USubMovieSceneTrack::RemoveAllAnimationData()
+{
+	SubMovieSceneSections.Empty();
+}
+
 bool USubMovieSceneTrack::IsEmpty() const
 {
 	return SubMovieSceneSections.Num() == 0;
diff --git a/Engine/Source/Runtime/MovieSceneTracks/Public/Sections/MovieSceneShotSection.h b/Engine/Source/Runtime/MovieSceneTracks/Public/Sections/MovieSceneShotSection.h
--- a/Engine/Source/Runtime/MovieSceneTracks/Public/Sections/MovieSceneShotSection.h
+++ b/Engine/Source/Runtime/MovieSceneTracks/Public/Sections/MovieSceneShotSection.h
@@ -3,13 +3,14 @@
 #pragma once
 
 #include "MovieSceneSection.h"
+#include "SubMovieSceneSection.h"
 #include "MovieSceneShotSection.generated.h"
 
 /**
- * Movie shots are sections on the director's track, that show what the viewer "sees"
+ * Movie shots are sections on the shots track, that show what the viewer "sees"
  */
 UCLASS( MinimalAPI )
-class UMovieSceneShotSection : public UMovieSceneSection
+class UMovieSceneShotSection : public USubMovieSceneSection
 {
 	GENERATED_UCLASS_BODY()
 public:
diff --git a/Engine/Source/Runtime/MovieSceneTracks/Public/Sections/SubMovieSceneSection.h b/Engine/Source/Runtime/MovieSceneTracks/Public/Sections/SubMovieSceneSection.h
--- a/Engine/Source/Runtime/MovieSceneTracks/Public/Sections/SubMovieSceneSection.h
+++ b/Engine/Source/Runtime/MovieSceneTracks/Public/Sections/SubMovieSceneSection.h
@@ -25,12 +25,4 @@ private:
 	/** Movie scene being played by this section */
 	UPROPERTY()
 	UMovieScene* MovieScene;
-
-	/** The start time where the movie scene should begin playing */
-	UPROPERTY()
-	float MovieSceneStartTime;
-
-	/** The end time where the movie scene should begin playing */
-	UPROPERTY()
-	float MovieSceneTimeEndTime;
 };
diff --git a/Engine/Source/Runtime/MovieSceneTracks/Public/TrackInstances/MovieSceneDirectorTrackInstance.h b/Engine/Source/Runtime/MovieSceneTracks/Public/TrackInstances/MovieSceneDirectorTrackInstance.h
deleted file mode 100644
diff --git a/Engine/Source/Runtime/MovieSceneTracks/Public/TrackInstances/MovieSceneShotTrackInstance.h b/Engine/Source/Runtime/MovieSceneTracks/Public/TrackInstances/MovieSceneShotTrackInstance.h
new file mode 100644
--- /dev/null
+++ b/Engine/Source/Runtime/MovieSceneTracks/Public/TrackInstances/MovieSceneShotTrackInstance.h
@@ -0,0 +1,29 @@
+// Copyright 1998-2015 Epic Games, Inc. All Rights Reserved.
+
+#pragma once
+
+#include "IMovieSceneTrackInstance.h"
+
+
+class UMovieSceneShotTrack;
+
+
+/**
+ * Instance of a UMovieSceneShotTrack
+ */
+class FMovieSceneShotTrackInstance
+	: public IMovieSceneTrackInstance
+{
+public:
+
+	FMovieSceneShotTrackInstance( UMovieSceneShotTrack& InDirectorTrack );
+
+	/** IMovieSceneTrackInstance interface */
+	virtual void Update( float Position, float LastPosition, const TArray<UObject*>& RuntimeObjects, IMovieScenePlayer& Player ) override;
+	virtual void RefreshInstance( const TArray<UObject*>& RuntimeObjects, IMovieScenePlayer& Player  ) override {}
+
+private:
+
+	/** Track that is being instanced */
+	UMovieSceneShotTrack* DirectorTrack;
+};
diff --git a/Engine/Source/Runtime/MovieSceneTracks/Public/Tracks/MovieSceneDirectorTrack.h b/Engine/Source/Runtime/MovieSceneTracks/Public/Tracks/MovieSceneDirectorTrack.h
deleted file mode 100644
diff --git a/Engine/Source/Runtime/MovieSceneTracks/Public/Tracks/MovieSceneShotTrack.h b/Engine/Source/Runtime/MovieSceneTracks/Public/Tracks/MovieSceneShotTrack.h
new file mode 100644
--- /dev/null
+++ b/Engine/Source/Runtime/MovieSceneTracks/Public/Tracks/MovieSceneShotTrack.h
@@ -0,0 +1,25 @@
+// Copyright 1998-2015 Epic Games, Inc. All Rights Reserved.
+
+#pragma once
+
+#include "MovieScene.h"
+#include "MovieSceneTrack.h"
+#include "SubMovieSceneTrack.h"
+#include "MovieSceneShotTrack.generated.h"
+
+/**
+ * Handles manipulation of shot properties in a movie scene
+ */
+UCLASS( MinimalAPI )
+class UMovieSceneShotTrack : public USubMovieSceneTrack
+{
+	GENERATED_UCLASS_BODY()
+
+public:
+	/** UMovieSceneTrack interface */
+	virtual FName GetTrackName() const override;
+	virtual TSharedPtr<IMovieSceneTrackInstance> CreateInstance() override;
+
+	/** Adds a new shot at the specified time */
+	MOVIESCENETRACKS_API void AddNewShot(FGuid CameraHandle, float Time);
+};
diff --git a/Engine/Source/Runtime/MovieSceneTracks/Public/Tracks/SubMovieSceneTrack.h b/Engine/Source/Runtime/MovieSceneTracks/Public/Tracks/SubMovieSceneTrack.h
--- a/Engine/Source/Runtime/MovieSceneTracks/Public/Tracks/SubMovieSceneTrack.h
+++ b/Engine/Source/Runtime/MovieSceneTracks/Public/Tracks/SubMovieSceneTrack.h
@@ -19,6 +19,7 @@ public:
 	virtual TSharedPtr<IMovieSceneTrackInstance> CreateInstance() override;
 	virtual TArray<UMovieSceneSection*> GetAllSections() const override;
 	virtual void RemoveSection( UMovieSceneSection* Section ) override;
+	virtual void RemoveAllAnimationData() override;
 	virtual bool IsEmpty() const override;
 	virtual TRange<float> GetSectionBoundaries() const override;
 	virtual bool HasSection( UMovieSceneSection* Section ) const override;
@@ -30,8 +31,8 @@ public:
 	 * @param Time	The time to add the section at
 	 */
 	virtual void AddMovieSceneSection( UMovieScene* SubMovieScene, float Time );
-private:
+protected:
 	/** All movie scene sections.  */
 	UPROPERTY()
-	TArray<class USubMovieSceneSection*> SubMovieSceneSections;
+	TArray<class UMovieSceneSection*> SubMovieSceneSections;
 };
