Fix shot thumbnail rendering not working properly.

diff --git a/Engine/Source/Editor/MovieSceneTools/Private/TrackEditors/ShotTrackEditor.cpp b/Engine/Source/Editor/MovieSceneTools/Private/TrackEditors/ShotTrackEditor.cpp
--- a/Engine/Source/Editor/MovieSceneTools/Private/TrackEditors/ShotTrackEditor.cpp
+++ b/Engine/Source/Editor/MovieSceneTools/Private/TrackEditors/ShotTrackEditor.cpp
@@ -25,7 +25,8 @@
 namespace AnimatableShotToolConstants
 {
 	// @todo Sequencer Perhaps allow this to be customizable
-	const uint32 DirectorTrackHeight = 100;
+	const uint32 TrackHeight = 90;
+	const uint32 TrackWidth = 90;
 	const double ThumbnailFadeInDuration = 0.25f;
 
 }
@@ -95,7 +96,7 @@ FShotThumbnail::FShotThumbnail(TSharedPtr<FShotSection> InSection, TRange<float>
 	: OwningSection(InSection)
 	, Texture(NULL)
 	, TimeRange(InTimeRange)
-	, FadeInStartTime(0.0)
+	, FadeInCurve( 0.0f, AnimatableShotToolConstants::ThumbnailFadeInDuration )
 {
 	Texture = new FSlateTexture2DRHIRef(GetSize().X, GetSize().Y, PF_B8G8R8A8, NULL, TexCreate_Dynamic, true);
 
@@ -115,7 +116,7 @@ FShotThumbnail::~FShotThumbnail()
 	}
 }
 
-FIntPoint FShotThumbnail::GetSize() const {return FIntPoint(OwningSection.Pin()->GetThumbnailWidth(), AnimatableShotToolConstants::DirectorTrackHeight);}
+FIntPoint FShotThumbnail::GetSize() const {return FIntPoint(OwningSection.Pin()->GetThumbnailWidth(), AnimatableShotToolConstants::TrackHeight);}
 bool FShotThumbnail::RequiresVsync() const {return false;}
 
 FSlateShaderResource* FShotThumbnail::GetViewportRenderTargetTexture() const
@@ -126,7 +127,7 @@ FSlateShaderResource* FShotThumbnail::GetViewportRenderTargetTexture() const
 void FShotThumbnail::DrawThumbnail()
 {
 	OwningSection.Pin()->DrawViewportThumbnail(SharedThis(this));
-	FadeInStartTime = FSlateApplication::Get().GetCurrentTime();
+	FadeInCurve.PlayReverse( OwningSection.Pin()->GetSequencerWidget() );
 }
 
 float FShotThumbnail::GetTime() const {return TimeRange.GetLowerBoundValue();}
@@ -143,7 +144,7 @@ void FShotThumbnail::CopyTextureIn(FSlateRenderTargetRHI* InTexture)
 
 float FShotThumbnail::GetFadeInCurve() const 
 {
-	return (float)FMath::Clamp(FadeInStartTime / AnimatableShotToolConstants::ThumbnailFadeInDuration, 0.0, 1.0);
+	return FadeInCurve.GetLerp();
 }
 
 bool FShotThumbnail::IsVisible() const
@@ -189,6 +190,8 @@ FShotSection::FShotSection( TSharedPtr<ISequencer> InSequencer, TSharedPtr<FShot
 
 		CalculateThumbnailWidthAndResize();
 	}
+
+	WhiteBrush = FEditorStyle::GetBrush("WhiteBrush");
 }
 
 FShotSection::~FShotSection()
@@ -213,7 +216,7 @@ FText FShotSection::GetSectionTitle() const
 
 float FShotSection::GetSectionHeight() const
 {
-	return AnimatableShotToolConstants::DirectorTrackHeight;
+	return AnimatableShotToolConstants::TrackHeight+10;
 }
 
 FReply FShotSection::OnSectionDoubleClicked( const FGeometry& SectionGeometry, const FPointerEvent& MouseEvent )
@@ -232,19 +235,29 @@ int32 FShotSection::OnPaintSection( const FGeometry& AllottedGeometry, const FSl
 {
 	if (Camera.IsValid())
 	{
-		for (int32 i = 0; i < Thumbnails.Num(); ++i)
+		FSlateDrawElement::MakeBox(
+				OutDrawElements,
+				LayerId,
+				AllottedGeometry.ToPaintGeometry(),
+				WhiteBrush,
+				SectionClippingRect,
+				ESlateDrawEffect::None,
+				FLinearColor::Black
+			);
+
+		int32 ThumbnailIndex = 0;
+		for( const TSharedPtr<FShotThumbnail>& Thumbnail : Thumbnails )
 		{
 			FGeometry TruncatedGeometry = AllottedGeometry.MakeChild(
-				FVector2D(ThumbnailWidth, AllottedGeometry.Size.Y),
-				FSlateLayoutTransform(AllottedGeometry.Scale, FVector2D(i * ThumbnailWidth, 0.0f))
+				Thumbnail->GetSize(),
+				FSlateLayoutTransform(AllottedGeometry.Scale, FVector2D(ThumbnailIndex * ThumbnailWidth, 5.0f))
 				);
-			ensure(TruncatedGeometry.Size.Y == AnimatableShotToolConstants::DirectorTrackHeight);
-
+			
 			FSlateDrawElement::MakeViewport(
 				OutDrawElements,
 				LayerId,
 				TruncatedGeometry.ToPaintGeometry(),
-				Thumbnails[i],
+				Thumbnail,
 				SectionClippingRect,
 				false,
 				false,
@@ -252,15 +265,20 @@ int32 FShotSection::OnPaintSection( const FGeometry& AllottedGeometry, const FSl
 				FLinearColor::White
 			);
 
-			FSlateDrawElement::MakeBox( 
-				OutDrawElements,
-				LayerId + 1,
-				TruncatedGeometry.ToPaintGeometry(),
-				FEditorStyle::GetBrush("WhiteBrush"),
-				SectionClippingRect,
-				ESlateDrawEffect::None,
-				FColor(255, 255, 255, Thumbnails[i]->GetFadeInCurve() * 255)
-			);
+			if(Thumbnail->GetFadeInCurve() > 0.0f )
+			{
+				FSlateDrawElement::MakeBox(
+					OutDrawElements,
+					LayerId + 1,
+					TruncatedGeometry.ToPaintGeometry(),
+					WhiteBrush,
+					SectionClippingRect,
+					ESlateDrawEffect::None,
+					FColor(255, 255, 255, Thumbnail->GetFadeInCurve() * 255)
+					);
+			}
+
+			++ThumbnailIndex;
 		}
 	}
 	else
@@ -294,7 +312,7 @@ void FShotSection::Tick( const FGeometry& AllottedGeometry, const FGeometry& Par
 	FIntPoint AllocatedSize = AllottedGeometry.Size.IntPoint();
 	float StartTime = Section->GetStartTime();
 
-	if (AllocatedSize.X != StoredSize.X || !FMath::IsNearlyEqual(StartTime, StoredStartTime))
+	if (!FSlateApplication::Get().HasAnyMouseCaptor() && ( AllocatedSize.X != StoredSize.X || !FMath::IsNearlyEqual(StartTime, StoredStartTime) ) )
 	{
 		RegenerateViewportThumbnails(AllocatedSize);
 	}
@@ -372,7 +390,7 @@ void FShotSection::CalculateThumbnailWidthAndResize()
 	float SavedTime = Sequencer.Pin()->GetGlobalTime();
 	
 	Sequencer.Pin()->SetGlobalTime(Section->GetStartTime());
-	uint32 OutThumbnailWidth = AnimatableShotToolConstants::DirectorTrackHeight * Camera->GetCameraComponent()->AspectRatio;
+	uint32 OutThumbnailWidth = AnimatableShotToolConstants::TrackWidth * Camera->GetCameraComponent()->AspectRatio;
 
 	// restore the global time
 	Sequencer.Pin()->SetGlobalTime(SavedTime);
@@ -381,34 +399,34 @@ void FShotSection::CalculateThumbnailWidthAndResize()
 	if (OutThumbnailWidth != ThumbnailWidth)
 	{
 		ThumbnailWidth = OutThumbnailWidth;
-		InternalViewportScene->UpdateViewportRHI( false, ThumbnailWidth, AnimatableShotToolConstants::DirectorTrackHeight, EWindowMode::Windowed );
+		InternalViewportScene->UpdateViewportRHI( false, ThumbnailWidth, AnimatableShotToolConstants::TrackHeight, EWindowMode::Windowed );
 	}
 }
 
 /////////////////////////////////////////////////
-// FDirectorTrackEditor
+// FShotTrackEditor
 
-FDirectorTrackEditor::FDirectorTrackEditor( TSharedRef<ISequencer> InSequencer )
+FShotTrackEditor::FShotTrackEditor( TSharedRef<ISequencer> InSequencer )
 	: FMovieSceneTrackEditor( InSequencer ) 
 {
 	ThumbnailPool = MakeShareable(new FShotThumbnailPool(InSequencer));
 }
 
-FDirectorTrackEditor::~FDirectorTrackEditor()
+FShotTrackEditor::~FShotTrackEditor()
 {
 }
 
-TSharedRef<FMovieSceneTrackEditor> FDirectorTrackEditor::CreateTrackEditor( TSharedRef<ISequencer> InSequencer )
+TSharedRef<FMovieSceneTrackEditor> FShotTrackEditor::CreateTrackEditor( TSharedRef<ISequencer> InSequencer )
 {
-	return MakeShareable( new FDirectorTrackEditor( InSequencer ) );
+	return MakeShareable( new FShotTrackEditor( InSequencer ) );
 }
 
-bool FDirectorTrackEditor::SupportsType( TSubclassOf<UMovieSceneTrack> Type ) const
+bool FShotTrackEditor::SupportsType( TSubclassOf<UMovieSceneTrack> Type ) const
 {
 	return Type == UMovieSceneShotTrack::StaticClass();
 }
 
-TSharedRef<ISequencerSection> FDirectorTrackEditor::MakeSectionInterface( UMovieSceneSection& SectionObject, UMovieSceneTrack* Track )
+TSharedRef<ISequencerSection> FShotTrackEditor::MakeSectionInterface( UMovieSceneSection& SectionObject, UMovieSceneTrack* Track )
 {
 	check( SupportsType( SectionObject.GetOuter()->GetClass() ) );
 	
@@ -421,7 +439,7 @@ TSharedRef<ISequencerSection> FDirectorTrackEditor::MakeSectionInterface( UMovie
 	return NewSection;
 }
 
-void FDirectorTrackEditor::AddKey(const FGuid& ObjectGuid, UObject* AdditionalAsset)
+void FShotTrackEditor::AddKey(const FGuid& ObjectGuid, UObject* AdditionalAsset)
 {
 	TArray<UObject*> OutObjects;
 	// @todo Sequencer - Sub-MovieScenes The director track may be able to get cameras from sub-movie scenes
@@ -429,11 +447,11 @@ void FDirectorTrackEditor::AddKey(const FGuid& ObjectGuid, UObject* AdditionalAs
 	bool bValidKey = OutObjects.Num() == 1 && OutObjects[0]->IsA<ACameraActor>();
 	if (bValidKey)
 	{
-		AnimatablePropertyChanged( UMovieSceneShotTrack::StaticClass(), false, FOnKeyProperty::CreateRaw( this, &FDirectorTrackEditor::AddKeyInternal, ObjectGuid ) );
+		AnimatablePropertyChanged( UMovieSceneShotTrack::StaticClass(), false, FOnKeyProperty::CreateRaw( this, &FShotTrackEditor::AddKeyInternal, ObjectGuid ) );
 	}
 }
 
-void FDirectorTrackEditor::Tick(float DeltaTime)
+void FShotTrackEditor::Tick(float DeltaTime)
 {
 	if (FSlateThrottleManager::Get().IsAllowingExpensiveTasks())
 	{
@@ -441,7 +459,7 @@ void FDirectorTrackEditor::Tick(float DeltaTime)
 	}
 }
 
-void FDirectorTrackEditor::BuildObjectBindingContextMenu(FMenuBuilder& MenuBuilder, const FGuid& ObjectBinding, const UClass* ObjectClass)
+void FShotTrackEditor::BuildObjectBindingContextMenu(FMenuBuilder& MenuBuilder, const FGuid& ObjectBinding, const UClass* ObjectClass)
 {
 	if (ObjectClass->IsChildOf(ACameraActor::StaticClass()))
 	{
@@ -456,9 +474,16 @@ void FDirectorTrackEditor::BuildObjectBindingContextMenu(FMenuBuilder& MenuBuild
 	}
 }
 
-void FDirectorTrackEditor::AddKeyInternal( float KeyTime, const FGuid ObjectGuid )
+void FShotTrackEditor::AddKeyInternal( float KeyTime, const FGuid ObjectGuid )
 {
-	UMovieSceneTrack* Track = GetMasterTrack( UMovieSceneShotTrack::StaticClass() );
+	UMovieScene* MovieScene = GetSequencer()->GetFocusedMovieScene();
+
+	UMovieSceneTrack* ShotTrack = MovieScene->GetShotTrack();
+
+	if( !ShotTrack )
+	{
+		ShotTrack = MovieScene->AddShotTrack( UMovieSceneShotTrack::StaticClass() );
+	}
 
-	Cast<UMovieSceneShotTrack>(Track)->AddNewShot(ObjectGuid, KeyTime);
+	Cast<UMovieSceneShotTrack>(ShotTrack)->AddNewShot(ObjectGuid, KeyTime);
 }
diff --git a/Engine/Source/Editor/MovieSceneTools/Private/TrackEditors/ShotTrackEditor.h b/Engine/Source/Editor/MovieSceneTools/Private/TrackEditors/ShotTrackEditor.h
--- a/Engine/Source/Editor/MovieSceneTools/Private/TrackEditors/ShotTrackEditor.h
+++ b/Engine/Source/Editor/MovieSceneTools/Private/TrackEditors/ShotTrackEditor.h
@@ -6,7 +6,7 @@
 /**
  * Tools for director tracks
  */
-class FDirectorTrackEditor : public FMovieSceneTrackEditor
+class FShotTrackEditor : public FMovieSceneTrackEditor
 {
 public:
 	/**
@@ -14,8 +14,8 @@ public:
 	 *
 	 * @param InSequencer	The sequencer instance to be used by this tool
 	 */
-	FDirectorTrackEditor( TSharedRef<ISequencer> InSequencer );
-	~FDirectorTrackEditor();
+	FShotTrackEditor( TSharedRef<ISequencer> InSequencer );
+	~FShotTrackEditor();
 
 	/**
 	 * Creates an instance of this class.  Called by a sequencer 
@@ -113,8 +113,8 @@ private:
 	/** Where in time this thumbnail is a rendering of */
 	TRange<float> TimeRange;
 
-	/** The time when the thumbnail started to fade in. */
-	double FadeInStartTime;
+	/** Fade curve to display while the thumbnail is redrawing */
+	FCurveSequence FadeInCurve;
 };
 
 
@@ -131,7 +131,7 @@ public:
 	virtual UMovieSceneSection* GetSectionObject() override;
 	virtual int32 OnPaintSection( const FGeometry& AllottedGeometry, const FSlateRect& SectionClippingRect, FSlateWindowElementList& OutDrawElements, int32 LayerId, bool bParentEnabled ) const override;
 	virtual void Tick( const FGeometry& AllottedGeometry, const FGeometry& ParentGeometry, const double InCurrentTime, const float InDeltaTime ) override;
-	virtual FText GetDisplayName() const override { return NSLOCTEXT("FShotSection", "DirectorTrack", "Director Track"); }
+	virtual FText GetDisplayName() const override { return NSLOCTEXT("FShotSection", "", "Shots"); }
 	virtual FText GetSectionTitle() const override;
 	virtual float GetSectionHeight() const override;
 	virtual void GenerateSectionLayout( class ISectionLayoutBuilder& LayoutBuilder ) const override {}
@@ -152,6 +152,8 @@ public:
 	/** Gets the time range of what in the sequencer is visible */
 	TRange<float> GetVisibleTimeRange() const {return VisibleTimeRange;}
 
+	/** @return The sequencer widget owning the shot section */
+	TSharedRef<SWidget> GetSequencerWidget() { return Sequencer.Pin()->GetSequencerWidget(); }
 private:
 	/** The section we are visualizing */
 	UMovieSceneSection* Section;
@@ -179,4 +181,7 @@ private:
 	TSharedPtr<FSceneViewport> InternalViewportScene;
 	/** An internal editor viewport client to render the thumbnails with */
 	TSharedPtr<FLevelEditorViewportClient> InternalViewportClient;
+	
+	/** Fade brush */
+	const FSlateBrush* WhiteBrush;
 };
diff --git a/Engine/Source/Editor/Sequencer/Private/SSection.cpp b/Engine/Source/Editor/Sequencer/Private/SSection.cpp
--- a/Engine/Source/Editor/Sequencer/Private/SSection.cpp
+++ b/Engine/Source/Editor/Sequencer/Private/SSection.cpp
@@ -212,9 +212,10 @@ int32 SSection::OnPaint( const FPaintArgs& Args, const FGeometry& AllottedGeomet
 
 	FGeometry SectionGeometry = AllottedGeometry.MakeChild( FVector2D( SequencerSectionConstants::SectionGripSize, 0 ), AllottedGeometry.GetDrawSize() - FVector2D( SequencerSectionConstants::SectionGripSize*2, 0.0f ) );
 
+	FSlateRect SectionClipRect = SectionGeometry.GetClippingRect().IntersectionWith( MyClippingRect );
 
 	// Ask the interface to draw the section
-	int32 PostSectionLayer = SectionInterface->OnPaintSection( SectionGeometry, MyClippingRect, OutDrawElements, LayerId, bParentEnabled );
+	int32 PostSectionLayer = SectionInterface->OnPaintSection( SectionGeometry, SectionClipRect, OutDrawElements, LayerId, bParentEnabled );
 
 	// @todo Sequencer - Temp indicators of the area of a section which can be dragged.
 
