Missing file

diff --git a/Engine/Source/Editor/MovieSceneTools/Private/MovieSceneToolsModule.cpp b/Engine/Source/Editor/MovieSceneTools/Private/MovieSceneToolsModule.cpp
--- a/Engine/Source/Editor/MovieSceneTools/Private/MovieSceneToolsModule.cpp
+++ b/Engine/Source/Editor/MovieSceneTools/Private/MovieSceneToolsModule.cpp
@@ -39,7 +39,7 @@ class FMovieSceneToolsModule : public IMovieSceneTools
 		ISequencerModule& SequencerModule = FModuleManager::Get().LoadModuleChecked<ISequencerModule>( "Sequencer" );
 		PropertyTrackEditorCreateTrackEditorDelegateHandle      = SequencerModule.RegisterTrackEditor_Handle( FOnCreateTrackEditor::CreateStatic( &FPropertyTrackEditor::CreateTrackEditor ) );
 		TransformTrackEditorCreateTrackEditorDelegateHandle     = SequencerModule.RegisterTrackEditor_Handle( FOnCreateTrackEditor::CreateStatic( &F3DTransformTrackEditor::CreateTrackEditor ) );
-		DirectorTrackEditorCreateTrackEditorDelegateHandle      = SequencerModule.RegisterTrackEditor_Handle( FOnCreateTrackEditor::CreateStatic( &FDirectorTrackEditor::CreateTrackEditor ) );
+		EditorCreateTrackEditorDelegateHandle					= SequencerModule.RegisterTrackEditor_Handle( FOnCreateTrackEditor::CreateStatic( &FShotTrackEditor::CreateTrackEditor ) );
 		SubMovieSceneTrackEditorCreateTrackEditorDelegateHandle = SequencerModule.RegisterTrackEditor_Handle( FOnCreateTrackEditor::CreateStatic( &FSubMovieSceneTrackEditor::CreateTrackEditor ) );
 		AudioTrackEditorCreateTrackEditorDelegateHandle         = SequencerModule.RegisterTrackEditor_Handle( FOnCreateTrackEditor::CreateStatic( &FAudioTrackEditor::CreateTrackEditor ) );
 		AnimationTrackEditorCreateTrackEditorDelegateHandle     = SequencerModule.RegisterTrackEditor_Handle( FOnCreateTrackEditor::CreateStatic( &FAnimationTrackEditor::CreateTrackEditor ) );
@@ -64,7 +64,7 @@ class FMovieSceneToolsModule : public IMovieSceneTools
 			ISequencerModule& SequencerModule = FModuleManager::Get().GetModuleChecked<ISequencerModule>( "Sequencer" );
 			SequencerModule.UnRegisterTrackEditor_Handle( PropertyTrackEditorCreateTrackEditorDelegateHandle );
 			SequencerModule.UnRegisterTrackEditor_Handle( TransformTrackEditorCreateTrackEditorDelegateHandle );
-			SequencerModule.UnRegisterTrackEditor_Handle( DirectorTrackEditorCreateTrackEditorDelegateHandle );
+			SequencerModule.UnRegisterTrackEditor_Handle( EditorCreateTrackEditorDelegateHandle );
 			SequencerModule.UnRegisterTrackEditor_Handle( SubMovieSceneTrackEditorCreateTrackEditorDelegateHandle );
 			SequencerModule.UnRegisterTrackEditor_Handle( AudioTrackEditorCreateTrackEditorDelegateHandle );
 			SequencerModule.UnRegisterTrackEditor_Handle( AnimationTrackEditorCreateTrackEditorDelegateHandle );
@@ -82,7 +82,7 @@ private:
 	/** Registered delegate handles */
 	FDelegateHandle PropertyTrackEditorCreateTrackEditorDelegateHandle;
 	FDelegateHandle TransformTrackEditorCreateTrackEditorDelegateHandle;
-	FDelegateHandle DirectorTrackEditorCreateTrackEditorDelegateHandle;
+	FDelegateHandle EditorCreateTrackEditorDelegateHandle;
 	FDelegateHandle SubMovieSceneTrackEditorCreateTrackEditorDelegateHandle;
 	FDelegateHandle AudioTrackEditorCreateTrackEditorDelegateHandle;
 	FDelegateHandle AnimationTrackEditorCreateTrackEditorDelegateHandle;
