Sequencer: Add lock viewport to camera button.

diff --git a/Engine/Source/Editor/MovieSceneTools/Private/TrackEditors/TransformTrackEditor.cpp b/Engine/Source/Editor/MovieSceneTools/Private/TrackEditors/TransformTrackEditor.cpp
--- a/Engine/Source/Editor/MovieSceneTools/Private/TrackEditors/TransformTrackEditor.cpp
+++ b/Engine/Source/Editor/MovieSceneTools/Private/TrackEditors/TransformTrackEditor.cpp
@@ -14,6 +14,8 @@
 #include "MovieSceneTrackEditor.h"
 #include "TransformTrackEditor.h"
 
+#define LOCTEXT_NAMESPACE "MovieScene_TransformTrack"
+
 /**
  * Class that draws a transform section in the sequencer
  */
@@ -253,6 +255,48 @@ void F3DTransformTrackEditor::AddKey(const FGuid& ObjectGuid, UObject* Additiona
 	}
 }
 
+void F3DTransformTrackEditor::BuildObjectBindingEditButtons(TSharedPtr<SHorizontalBox> EditBox, const FGuid& ObjectGuid, const UClass* ObjectClass)
+{
+	TArray<UObject*> OutObjects;
+	GetSequencer()->GetRuntimeObjects( GetSequencer()->GetFocusedMovieSceneInstance(), ObjectGuid, OutObjects);
+
+	TWeakObjectPtr<ACameraActor> CameraActor;
+
+	for ( UObject* Object : OutObjects )
+	{
+		ACameraActor* Actor = Cast<ACameraActor>( Object );
+		if (Actor)
+		{
+			CameraActor = Actor;
+			break;
+		}
+	}
+
+	if (!CameraActor.IsValid())
+	{
+		return;
+	}
+
+	// If this is a camera track, add a button to lock the viewport to the camera
+	EditBox.Get()->AddSlot()
+	.VAlign(VAlign_Center)
+	.HAlign(HAlign_Right)
+	.AutoWidth()
+	.Padding(4, 0, 0, 0)
+	[
+		SNew(SCheckBox)
+		.IsChecked(this, &F3DTransformTrackEditor::IsCameraLocked, CameraActor)
+		.OnCheckStateChanged(this, &F3DTransformTrackEditor::OnLockCameraClicked, CameraActor)
+		.ToolTipText(this, &F3DTransformTrackEditor::GetLockCameraToolTip, CameraActor)
+		.CheckedImage(FEditorStyle::GetBrush("Sequencer.LockCamera"))
+		.CheckedHoveredImage(FEditorStyle::GetBrush("Sequencer.LockCamera"))
+		.CheckedPressedImage(FEditorStyle::GetBrush("Sequencer.LockCamera"))
+		.UncheckedImage(FEditorStyle::GetBrush("Sequencer.UnlockCamera"))
+		.UncheckedHoveredImage(FEditorStyle::GetBrush("Sequencer.UnlockCamera"))
+		.UncheckedPressedImage(FEditorStyle::GetBrush("Sequencer.UnlockCamera"))
+	];
+}
+
 void F3DTransformTrackEditor::OnTransformChangedInternals(float KeyTime, UObject* InObject, FGuid ObjectHandle, FTransformDataPair TransformPair, bool bAutoKeying)
 {
 	// Only unwind rotation if we're generating keys while recording (scene is actively playing back)
@@ -281,3 +325,73 @@ void F3DTransformTrackEditor::OnTransformChangedInternals(float KeyTime, UObject
 		}
 	}
 }
+
+ECheckBoxState F3DTransformTrackEditor::IsCameraLocked(TWeakObjectPtr<ACameraActor> CameraActor) const
+{
+	for (int32 i = 0; i < GEditor->LevelViewportClients.Num(); ++i)
+	{		
+		FLevelEditorViewportClient* LevelVC = GEditor->LevelViewportClients[i];
+		if (LevelVC && LevelVC->IsPerspective() && LevelVC->GetViewMode() != VMI_Unknown && LevelVC->IsActorLocked(CameraActor.Get()))
+		{
+			return ECheckBoxState::Checked;
+		}
+	}
+
+	return ECheckBoxState::Unchecked;
+}
+
+void F3DTransformTrackEditor::OnLockCameraClicked(ECheckBoxState CheckBoxState, TWeakObjectPtr<ACameraActor> CameraActor)
+{
+	// If toggle is on, lock the active viewport to the camera
+	if (CheckBoxState == ECheckBoxState::Checked)
+	{
+		// Set the active viewport or any viewport if there is no active viewport
+		FViewport* ActiveViewport = GEditor->GetActiveViewport();
+
+		FLevelEditorViewportClient* LevelVC = NULL;
+
+		for (int32 i = 0; i < GEditor->LevelViewportClients.Num(); ++i)
+		{		
+			FLevelEditorViewportClient* Viewport = GEditor->LevelViewportClients[i];
+			if (Viewport && Viewport->IsPerspective() && Viewport->GetViewMode() != VMI_Unknown)
+			{
+				LevelVC = Viewport;
+
+				if (LevelVC->Viewport == ActiveViewport)
+				{
+					break;
+				}
+			}
+		}
+
+		if (LevelVC != NULL)
+		{
+			LevelVC->SetActorLock(CameraActor.Get());
+			LevelVC->bLockedCameraView = true;
+			LevelVC->UpdateViewForLockedActor();
+			LevelVC->Invalidate();
+		}
+	}
+	// Otherwise, clear all locks on the camera
+	else
+	{
+		for (int32 i = 0; i < GEditor->LevelViewportClients.Num(); ++i)
+		{		
+			FLevelEditorViewportClient* LevelVC = GEditor->LevelViewportClients[i];
+			if (LevelVC && LevelVC->IsPerspective() && LevelVC->GetViewMode() != VMI_Unknown && LevelVC->IsActorLocked(CameraActor.Get()))
+			{
+				LevelVC->SetActorLock(nullptr);
+				LevelVC->bLockedCameraView = false;
+				LevelVC->UpdateViewForLockedActor();
+				LevelVC->Invalidate();
+			}
+		}
+	}
+}
+
+FText F3DTransformTrackEditor::GetLockCameraToolTip(TWeakObjectPtr<ACameraActor> CameraActor) const
+{
+	return IsCameraLocked(CameraActor) == ECheckBoxState::Checked ?
+		FText::Format(LOCTEXT("UnlockCamera", "Unlock {0} from Viewport"), FText::FromString(CameraActor.Get()->GetActorLabel())) :
+		FText::Format(LOCTEXT("LockCamera", "Lock {0} to Selected Viewport"), FText::FromString(CameraActor.Get()->GetActorLabel()));
+}
diff --git a/Engine/Source/Editor/MovieSceneTools/Private/TrackEditors/TransformTrackEditor.h b/Engine/Source/Editor/MovieSceneTools/Private/TrackEditors/TransformTrackEditor.h
--- a/Engine/Source/Editor/MovieSceneTools/Private/TrackEditors/TransformTrackEditor.h
+++ b/Engine/Source/Editor/MovieSceneTools/Private/TrackEditors/TransformTrackEditor.h
@@ -30,6 +30,7 @@ public:
 	virtual bool SupportsType( TSubclassOf<UMovieSceneTrack> Type ) const override;
 	virtual TSharedRef<ISequencerSection> MakeSectionInterface( UMovieSceneSection& SectionObject, UMovieSceneTrack* Track ) override;
 	virtual void AddKey( const FGuid& ObjectGuid, UObject* AdditionalAsset = NULL ) override;
+	virtual void BuildObjectBindingEditButtons(TSharedPtr<SHorizontalBox> EditBox, const FGuid& ObjectBinding, const UClass* ObjectClass) override;
 
 private:
 	/**
@@ -48,6 +49,16 @@ private:
 
 	/** Delegate for animatable property changed in OnTransformChanged */
 	void OnTransformChangedInternals(float KeyTime, UObject* InObject, FGuid ObjectHandle, struct FTransformDataPair TransformPair, bool bAutoKeying);
+
+	/** Delegate for camera button lock state */
+	ECheckBoxState IsCameraLocked(TWeakObjectPtr<ACameraActor> CameraActor) const; 
+
+	/** Delegate for locked camera button */
+	void OnLockCameraClicked(ECheckBoxState CheckBoxState, TWeakObjectPtr<ACameraActor> CameraActor);
+
+	/** Delegate for camera button lock tooltip */
+	FText GetLockCameraToolTip(TWeakObjectPtr<ACameraActor> CameraActor) const; 
+
 private:
 	/** Mapping of objects to their existing transform data (for comparing against new transform data) */
 	TMap< TWeakObjectPtr<UObject>, FTransformData > ObjectToExistingTransform;
