Fixed build errors

diff --git a/Engine/Source/Editor/MovieSceneTools/Private/TrackEditors/ShotTrackEditor.cpp b/Engine/Source/Editor/MovieSceneTools/Private/TrackEditors/ShotTrackEditor.cpp
--- a/Engine/Source/Editor/MovieSceneTools/Private/TrackEditors/ShotTrackEditor.cpp
+++ b/Engine/Source/Editor/MovieSceneTools/Private/TrackEditors/ShotTrackEditor.cpp
@@ -476,14 +476,7 @@ void FShotTrackEditor::BuildObjectBindingContextMenu(FMenuBuilder& MenuBuilder,
 
 void FShotTrackEditor::AddKeyInternal( float KeyTime, const FGuid ObjectGuid )
 {
-	UMovieScene* MovieScene = GetSequencer()->GetFocusedMovieScene();
+	UMovieSceneTrack* Track = GetMasterTrack( UMovieSceneShotTrack::StaticClass() );
 
-	UMovieSceneTrack* ShotTrack = MovieScene->GetShotTrack();
-
-	if( !ShotTrack )
-	{
-		ShotTrack = MovieScene->AddShotTrack( UMovieSceneShotTrack::StaticClass() );
-	}
-
-	Cast<UMovieSceneShotTrack>(ShotTrack)->AddNewShot(ObjectGuid, KeyTime);
+	Cast<UMovieSceneShotTrack>(Track)->AddNewShot(ObjectGuid, KeyTime);
 }
