Sequencer: Add support for custom menus for sections.

diff --git a/Engine/Source/Editor/MovieSceneTools/Private/TrackEditors/ShotTrackEditor.cpp b/Engine/Source/Editor/MovieSceneTools/Private/TrackEditors/ShotTrackEditor.cpp
--- a/Engine/Source/Editor/MovieSceneTools/Private/TrackEditors/ShotTrackEditor.cpp
+++ b/Engine/Source/Editor/MovieSceneTools/Private/TrackEditors/ShotTrackEditor.cpp
@@ -21,6 +21,7 @@
 #include "ShotTrackEditor.h"
 #include "CommonMovieSceneTools.h"
 #include "Camera/CameraActor.h"
+#include "STextEntryPopup.h"
 
 namespace AnimatableShotToolConstants
 {
@@ -231,6 +232,62 @@ FReply FShotSection::OnSectionDoubleClicked( const FGeometry& SectionGeometry, c
 	return FReply::Handled();
 }
 
+void FShotSection::BuildSectionContextMenu(FMenuBuilder& MenuBuilder)
+{
+	MenuBuilder.AddMenuEntry(
+		NSLOCTEXT("Sequencer", "RenameShot", "Rename"),
+		NSLOCTEXT("Sequencer", "RenameShotToolTip", "Renames this shot."),
+		FSlateIcon(),
+		FUIAction(FExecuteAction::CreateSP(this, &FShotSection::RenameShot))
+		);
+
+	MenuBuilder.AddMenuEntry(
+		NSLOCTEXT("Sequencer", "FilterToShots", "Filter To Shots"),
+		NSLOCTEXT("Sequencer", "FilterToShotsToolTip", "Filters to the selected shot sections"),
+		FSlateIcon(),
+		FUIAction(FExecuteAction::CreateSP(this, &FShotSection::FilterToSelectedShotSections, true))
+		);
+}
+
+void FShotSection::RenameShot()
+{
+	auto ActualShotSection = CastChecked<UMovieSceneShotSection>(Section);
+
+	TSharedRef<STextEntryPopup> TextEntry = 
+		SNew(STextEntryPopup)
+		.Label(NSLOCTEXT("Sequencer", "RenameShotHeader", "Name"))
+		.DefaultText( ActualShotSection->GetTitle() )
+		.OnTextCommitted(this, &FShotSection::RenameShotCommitted, Section)
+		.ClearKeyboardFocusOnCommit( false );
+	
+	NameEntryPopupMenu = FSlateApplication::Get().PushMenu(
+		Sequencer.Pin()->GetSequencerWidget(),
+		FWidgetPath(),
+		TextEntry,
+		FSlateApplication::Get().GetCursorPos(),
+		FPopupTransitionEffect( FPopupTransitionEffect::TypeInPopup )
+		);
+}
+
+void FShotSection::FilterToSelectedShotSections(bool bZoomToShotBounds)
+{
+	Sequencer.Pin()->FilterToSelectedShotSections(bZoomToShotBounds);
+}
+
+void FShotSection::RenameShotCommitted(const FText& RenameText, ETextCommit::Type CommitInfo, UMovieSceneSection* Section)
+{
+	if (CommitInfo == ETextCommit::OnEnter)
+	{
+		auto ShotSection = CastChecked<UMovieSceneShotSection>(Section);
+		ShotSection->SetTitle(RenameText);
+	}
+
+	if (NameEntryPopupMenu.IsValid())
+	{
+		NameEntryPopupMenu.Pin()->Dismiss();
+	}
+}
+
 int32 FShotSection::OnPaintSection( const FGeometry& AllottedGeometry, const FSlateRect& SectionClippingRect, FSlateWindowElementList& OutDrawElements, int32 LayerId, bool bParentEnabled ) const
 {
 	if (Camera.IsValid())
diff --git a/Engine/Source/Editor/MovieSceneTools/Private/TrackEditors/ShotTrackEditor.h b/Engine/Source/Editor/MovieSceneTools/Private/TrackEditors/ShotTrackEditor.h
--- a/Engine/Source/Editor/MovieSceneTools/Private/TrackEditors/ShotTrackEditor.h
+++ b/Engine/Source/Editor/MovieSceneTools/Private/TrackEditors/ShotTrackEditor.h
@@ -136,6 +136,16 @@ public:
 	virtual float GetSectionHeight() const override;
 	virtual void GenerateSectionLayout( class ISectionLayoutBuilder& LayoutBuilder ) const override {}
 	virtual FReply OnSectionDoubleClicked( const FGeometry& SectionGeometry, const FPointerEvent& MouseEvent ) override;
+	virtual void BuildSectionContextMenu(FMenuBuilder& MenuBuilder) override;
+
+	/** Opens a renaming dialog for the passed in shot section */
+	void RenameShot();
+	 
+	/** Called when committing a rename shot text entry popup */
+	void RenameShotCommitted(const FText& RenameText, ETextCommit::Type CommitInfo, UMovieSceneSection* Section);
+
+	/** Filter to selected shot sections */
+	void FilterToSelectedShotSections(bool bZoomToShotBounds);
 
 	/** Gets the thumbnail width */
 	uint32 GetThumbnailWidth() const;
@@ -184,4 +194,7 @@ private:
 	
 	/** Fade brush */
 	const FSlateBrush* WhiteBrush;
+
+	/** Reference to owner of the current popup */
+	TWeakPtr<class IMenu> NameEntryPopupMenu;
 };
diff --git a/Engine/Source/Editor/Sequencer/Private/SSection.cpp b/Engine/Source/Editor/Sequencer/Private/SSection.cpp
--- a/Engine/Source/Editor/Sequencer/Private/SSection.cpp
+++ b/Engine/Source/Editor/Sequencer/Private/SSection.cpp
@@ -446,24 +446,7 @@ TSharedPtr<SWidget> SSection::OnSummonContextMenu( const FGeometry& MyGeometry,
 	}
 	else
 	{
-		bool bIsAShotSection = SceneSection->IsA<UMovieSceneShotSection>();
-
-		if (bIsAShotSection)
-		{
-			MenuBuilder.AddMenuEntry(
-				NSLOCTEXT("Sequencer", "RenameShot", "Rename"),
-				NSLOCTEXT("Sequencer", "RenameShotToolTip", "Renames this shot."),
-				FSlateIcon(),
-				FUIAction(FExecuteAction::CreateSP(&Sequencer, &FSequencer::RenameShot, SceneSection))
-				);
-
-			MenuBuilder.AddMenuEntry(
-				NSLOCTEXT("Sequencer", "FilterToShots", "Filter To Shots"),
-				NSLOCTEXT("Sequencer", "FilterToShotsToolTip", "Filters to the selected shot sections"),
-				FSlateIcon(),
-				FUIAction(FExecuteAction::CreateSP(&Sequencer, &FSequencer::FilterToSelectedShotSections, true))
-				);
-		}
+		SectionInterface->BuildSectionContextMenu(MenuBuilder);
 
 		// @todo Sequencer this should delete all selected sections
 		// delete/selection needs to be rethought in general
diff --git a/Engine/Source/Editor/Sequencer/Private/Sequencer.cpp b/Engine/Source/Editor/Sequencer/Private/Sequencer.cpp
--- a/Engine/Source/Editor/Sequencer/Private/Sequencer.cpp
+++ b/Engine/Source/Editor/Sequencer/Private/Sequencer.cpp
@@ -32,7 +32,6 @@
 #include "MovieSceneInstance.h"
 #include "IKeyArea.h"
 #include "SnappingUtils.h"
-#include "STextEntryPopup.h"
 #include "GenericCommands.h"
 #include "Engine/BlueprintGeneratedClass.h"
 #include "Engine/Selection.h"
@@ -371,26 +370,6 @@ void FSequencer::AddAnimation(FGuid ObjectGuid, class UAnimSequence* AnimSequenc
 	}
 }
 
-void FSequencer::RenameShot(UMovieSceneSection* ShotSection)
-{
-	auto ActualShotSection = CastChecked<UMovieSceneShotSection>(ShotSection);
-
-	TSharedRef<STextEntryPopup> TextEntry = 
-		SNew(STextEntryPopup)
-		.Label(NSLOCTEXT("Sequencer", "RenameShotHeader", "Name"))
-		.DefaultText( ActualShotSection->GetTitle() )
-		.OnTextCommitted(this, &FSequencer::RenameShotCommitted, ShotSection)
-		.ClearKeyboardFocusOnCommit( false );
-	
-	NameEntryPopupMenu = FSlateApplication::Get().PushMenu(
-		SequencerWidget.ToSharedRef(),
-		FWidgetPath(),
-		TextEntry,
-		FSlateApplication::Get().GetCursorPos(),
-		FPopupTransitionEffect( FPopupTransitionEffect::TypeInPopup )
-		);
-}
-
 void FSequencer::DeleteSection(class UMovieSceneSection* Section)
 {
 	UMovieScene* MovieScene = GetFocusedMovieScene();
@@ -441,20 +420,6 @@ void FSequencer::DeleteSelectedKeys()
 	}
 }
 
-void FSequencer::RenameShotCommitted(const FText& RenameText, ETextCommit::Type CommitInfo, UMovieSceneSection* Section)
-{
-	if (CommitInfo == ETextCommit::OnEnter)
-	{
-		auto ShotSection = CastChecked<UMovieSceneShotSection>(Section);
-		ShotSection->SetTitle(RenameText);
-	}
-
-	if (NameEntryPopupMenu.IsValid())
-	{
-		NameEntryPopupMenu.Pin()->Dismiss();
-	}
-}
-
 void FSequencer::SpawnOrDestroyPuppetObjects( TSharedRef<FMovieSceneInstance> MovieSceneInstance )
 {
 	if( ObjectBindingManager->AllowsSpawnableObjects() )
diff --git a/Engine/Source/Editor/Sequencer/Private/Sequencer.h b/Engine/Source/Editor/Sequencer/Private/Sequencer.h
--- a/Engine/Source/Editor/Sequencer/Private/Sequencer.h
+++ b/Engine/Source/Editor/Sequencer/Private/Sequencer.h
@@ -118,11 +118,6 @@ public:
 	 */
 	void SpawnOrDestroyPuppetObjects( TSharedRef<FMovieSceneInstance> MovieSceneInstance );
 
-	/** 
-	 * Opens a renaming dialog for the passed in shot section
-	 */
-	void RenameShot(class UMovieSceneSection* ShotSection);
-	 
 	/** 
 	 * Deletes the passed in section
 	 */
@@ -341,9 +336,6 @@ protected:
 	/** Called via UEditorEngine::GetActorRecordingStateEvent to check to see whether we need to record actor state */
 	void GetActorRecordingState( bool& bIsRecording /* In+Out */ ) const;
 	
-	/* Called when committing a rename shot text entry popup */
-	void RenameShotCommitted(const FText& RenameText, ETextCommit::Type CommitInfo, UMovieSceneSection* Section);
-
 	/** Called when a user executes the delete command to delete sections or keys */
 	void DeleteSelectedItems();
 	
@@ -403,9 +395,6 @@ private:
 	/** Main sequencer widget */
 	TSharedPtr< class SSequencer > SequencerWidget;
 	
-	/** Reference to owner of the current popup */
-	TWeakPtr<class IMenu> NameEntryPopupMenu;
-
 	/** The asset editor that created this Sequencer if any */
 	TWeakPtr<IToolkitHost> ToolkitHost;
 
diff --git a/Engine/Source/Editor/Sequencer/Public/ISequencerSection.h b/Engine/Source/Editor/Sequencer/Public/ISequencerSection.h
--- a/Engine/Source/Editor/Sequencer/Public/ISequencerSection.h
+++ b/Engine/Source/Editor/Sequencer/Public/ISequencerSection.h
@@ -73,4 +73,11 @@ public:
 	 * @param  InDeltaTime  Real time passed since last tick
 	 */
 	virtual void Tick( const FGeometry& AllottedGeometry, const FGeometry& ClippedGeometry, const double InCurrentTime, const float InDeltaTime ) {}
+
+	/**
+	 * Builds up the section context menu for the outliner
+	 *
+	 * @param MenuBuilder	The menu builder to change
+	 */
+	virtual void BuildSectionContextMenu(FMenuBuilder& MenuBuilder) {}
 };
