Backend improvements to the Shot workflow in Sequencer.  UI still needs improvement

- Creating a new shot will now create a new sub-movie scene with a shot name and number
- Refactored shot playback in the editor so that it properly sends the correct data to the level viewport being previewed
- Shot tracks are no longer considered a master track.  Each movie scene has a dedicated shot track that is created if needed.  Only the root movie scene shot track is updated at runtime however

diff --git a/Engine/Source/Editor/MovieSceneTools/Private/TrackEditors/ShotTrackEditor.cpp b/Engine/Source/Editor/MovieSceneTools/Private/TrackEditors/ShotTrackEditor.cpp
--- a/Engine/Source/Editor/MovieSceneTools/Private/TrackEditors/ShotTrackEditor.cpp
+++ b/Engine/Source/Editor/MovieSceneTools/Private/TrackEditors/ShotTrackEditor.cpp
@@ -4,15 +4,12 @@
 #include "MovieScene.h"
 #include "MovieSceneSection.h"
 #include "ISequencerSection.h"
-#include "PropertyEditorModule.h"
-#include "PropertyHandle.h"
 #include "MovieSceneTrack.h"
 #include "MovieSceneShotTrack.h"
 #include "ScopedTransaction.h"
 #include "ISequencerObjectChangeListener.h"
 #include "ISectionLayoutBuilder.h"
 #include "Runtime/Engine/Public/Slate/SlateTextures.h"
-#include "ObjectTools.h"
 #include "Runtime/Engine/Public/Slate/SceneViewport.h"
 #include "Runtime/MovieSceneTracks/Public/Sections/MovieSceneShotSection.h"
 #include "IKeyArea.h"
@@ -22,6 +19,7 @@
 #include "CommonMovieSceneTools.h"
 #include "Camera/CameraActor.h"
 #include "STextEntryPopup.h"
+#include "AssetToolsModule.h"
 
 namespace AnimatableShotToolConstants
 {
@@ -212,7 +210,7 @@ UMovieSceneSection* FShotSection::GetSectionObject()
 
 FText FShotSection::GetSectionTitle() const
 {
-	return Cast<UMovieSceneShotSection>(Section)->GetTitle();
+	return Cast<UMovieSceneShotSection>(Section)->GetShotDisplayName();
 }
 
 float FShotSection::GetSectionHeight() const
@@ -234,12 +232,14 @@ FReply FShotSection::OnSectionDoubleClicked( const FGeometry& SectionGeometry, c
 
 void FShotSection::BuildSectionContextMenu(FMenuBuilder& MenuBuilder)
 {
+/*
 	MenuBuilder.AddMenuEntry(
 		NSLOCTEXT("Sequencer", "RenameShot", "Rename"),
 		NSLOCTEXT("Sequencer", "RenameShotToolTip", "Renames this shot."),
 		FSlateIcon(),
 		FUIAction(FExecuteAction::CreateSP(this, &FShotSection::RenameShot))
 		);
+*/
 
 	MenuBuilder.AddMenuEntry(
 		NSLOCTEXT("Sequencer", "FilterToShots", "Filter To Shots"),
@@ -253,7 +253,7 @@ void FShotSection::RenameShot()
 {
 	auto ActualShotSection = CastChecked<UMovieSceneShotSection>(Section);
 
-	TSharedRef<STextEntryPopup> TextEntry = 
+	/*TSharedRef<STextEntryPopup> TextEntry = 
 		SNew(STextEntryPopup)
 		.Label(NSLOCTEXT("Sequencer", "RenameShotHeader", "Name"))
 		.DefaultText( ActualShotSection->GetTitle() )
@@ -266,7 +266,7 @@ void FShotSection::RenameShot()
 		TextEntry,
 		FSlateApplication::Get().GetCursorPos(),
 		FPopupTransitionEffect( FPopupTransitionEffect::TypeInPopup )
-		);
+		);*/
 }
 
 void FShotSection::FilterToSelectedShotSections(bool bZoomToShotBounds)
@@ -276,6 +276,7 @@ void FShotSection::FilterToSelectedShotSections(bool bZoomToShotBounds)
 
 void FShotSection::RenameShotCommitted(const FText& RenameText, ETextCommit::Type CommitInfo, UMovieSceneSection* Section)
 {
+/*
 	if (CommitInfo == ETextCommit::OnEnter)
 	{
 		auto ShotSection = CastChecked<UMovieSceneShotSection>(Section);
@@ -285,7 +286,7 @@ void FShotSection::RenameShotCommitted(const FText& RenameText, ETextCommit::Typ
 	if (NameEntryPopupMenu.IsValid())
 	{
 		NameEntryPopupMenu.Pin()->Dismiss();
-	}
+	}*/
 }
 
 int32 FShotSection::OnPaintSection( const FGeometry& AllottedGeometry, const FSlateRect& SectionClippingRect, FSlateWindowElementList& OutDrawElements, int32 LayerId, bool bParentEnabled ) const
@@ -533,7 +534,130 @@ void FShotTrackEditor::BuildObjectBindingContextMenu(FMenuBuilder& MenuBuilder,
 
 void FShotTrackEditor::AddKeyInternal( float KeyTime, const FGuid ObjectGuid )
 {
-	UMovieSceneTrack* Track = GetMasterTrack( UMovieSceneShotTrack::StaticClass() );
+	UMovieScene* MovieScene = GetSequencer()->GetFocusedMovieScene();
+
+	UMovieSceneTrack* Track = MovieScene->GetShotTrack();
+	if( !Track )
+	{
+		Track = MovieScene->AddShotTrack( UMovieSceneShotTrack::StaticClass() );
+	}
+
+	UMovieSceneShotTrack* ShotTrack = CastChecked<UMovieSceneShotTrack>( Track );
+
+	FAssetToolsModule& AssetToolsModule = FModuleManager::GetModuleChecked<FAssetToolsModule>("AssetTools");
+
+	// Get the current path of the movie scene the shot is in
+	const FString LongPackagePath = FPackageName::GetLongPackagePath(ShotTrack->GetOutermost()->GetPathName());
+
+	const TArray<UMovieSceneSection*>& AllSections = ShotTrack->GetAllSections();
+
+	// Find where the section will be inserted
+	int32 SectionIndex = FindIndexForNewShot( AllSections, KeyTime );
+
+	const int32 ShotNumber = GenerateShotNumber( *ShotTrack, SectionIndex );
+
+	FString ShotName = FString::Printf( TEXT("Shot_%04d"), ShotNumber );
+
+	const FString NewShotPackagePath = LongPackagePath + TEXT("/") + MovieScene->GetName() + TEXT("_") + ShotName;
+
+	const FString EmptySuffix;
+
+	// Create a package name and asset name for the new movie scene
+	FString PackageName;
+	FString AssetName;
+	AssetToolsModule.Get().CreateUniqueAssetName(NewShotPackagePath, EmptySuffix, PackageName, AssetName);
+	const FString PackagePath = FPackageName::GetLongPackagePath(PackageName);
+
+	UObject* NewAsset = AssetToolsModule.Get().CreateAsset(AssetName, PackagePath, UMovieScene::StaticClass(), nullptr );
+	
+	float EndTime = KeyTime;
+	if( AllSections.IsValidIndex( SectionIndex + 1 ) )
+	{
+		// The new section ends where the next section begins
+		EndTime = AllSections[SectionIndex+1]->GetStartTime();
+	}
+	else
+	{
+		// There is no section after this one so the new section should be the length of the movie scene
+		EndTime = MovieScene->GetTimeRange().GetUpperBoundValue();
+	}
+
+
+	if (AllSections.IsValidIndex(SectionIndex - 1))
+	{
+		// Shift the previous shot's end time so that it ends when the new shot starts
+		AllSections[SectionIndex - 1]->SetEndTime(KeyTime);
+	}
+
+	if (AllSections.IsValidIndex(SectionIndex + 1))
+	{
+		// Shift the next shot's start time so that it starts when the new shot ends
+		AllSections[SectionIndex + 1]->SetStartTime(EndTime);
+	}
+
+	TRange<float> TimeRange( KeyTime, EndTime );
+	ShotTrack->AddNewShot( ObjectGuid, *CastChecked<UMovieScene>( NewAsset ), TimeRange, FText::FromString( ShotName ), ShotNumber );
+}
+
+int32 FShotTrackEditor::GenerateShotNumber( UMovieSceneShotTrack& ShotTrack, int32 SectionIndex ) const
+{
+	const TArray<UMovieSceneSection*>& AllSections = ShotTrack.GetAllSections();
+
+	const int32 Interval = 10;
+	int32 ShotNum = Interval;
+	int32 LastKeyIndex = AllSections.Num() - 1;
+
+	int32 PrevShotNum = 0;
+	//get the preceding shot number if any
+	if (SectionIndex > 0)
+	{
+		PrevShotNum = CastChecked<UMovieSceneShotSection>( AllSections[SectionIndex - 1] )->GetShotNumber();
+	}
 
-	Cast<UMovieSceneShotTrack>(Track)->AddNewShot(ObjectGuid, KeyTime);
+	if (SectionIndex < LastKeyIndex)
+	{
+		//we're inserting before something before the first frame
+		int32 NextShotNum = CastChecked<UMovieSceneShotSection>( AllSections[SectionIndex + 1] )->GetShotNumber();
+		if (NextShotNum == 0)
+		{
+			NextShotNum = PrevShotNum + (Interval * 2);
+		}
+
+		if (NextShotNum > PrevShotNum)
+		{
+			//find a midpoint if we're in order
+
+			//try to stick to the nearest interval if possible
+			int32 NearestInterval = PrevShotNum - (PrevShotNum % Interval) + Interval;
+			if (NearestInterval > PrevShotNum && NearestInterval < NextShotNum)
+			{
+				ShotNum = NearestInterval;
+			}
+			//else find the exact mid point
+			else
+			{
+				ShotNum = ((NextShotNum - PrevShotNum) / 2) + PrevShotNum;
+			}
+		}
+		else
+		{
+			//Just use the previous shot number + 1 with we're out of order
+			ShotNum = PrevShotNum + 1;
+		}
+	}
+	else
+	{
+		//we're adding to the end of the track
+		ShotNum = PrevShotNum + Interval;
+	}
+
+	return ShotNum;
 }
+
+int32 FShotTrackEditor::FindIndexForNewShot( const TArray<UMovieSceneSection*>& ShotSections, float NewShotTime ) const
+{
+	int32 SectionIndex = 0;
+	for (SectionIndex = 0; SectionIndex < ShotSections.Num() && ShotSections[SectionIndex]->GetStartTime() < NewShotTime; ++SectionIndex);
+
+	return SectionIndex;
+}
\ No newline at end of file
diff --git a/Engine/Source/Editor/MovieSceneTools/Private/TrackEditors/ShotTrackEditor.h b/Engine/Source/Editor/MovieSceneTools/Private/TrackEditors/ShotTrackEditor.h
--- a/Engine/Source/Editor/MovieSceneTools/Private/TrackEditors/ShotTrackEditor.h
+++ b/Engine/Source/Editor/MovieSceneTools/Private/TrackEditors/ShotTrackEditor.h
@@ -36,6 +36,11 @@ private:
 	/** Delegate for AnimatablePropertyChanged in AddKey */
 	void AddKeyInternal(float AutoKeyTime, const FGuid ObjectGuid);
 
+	/** Generates a shot number for a new section by looking at the existing shots in relation to the new shot time */
+	int32 GenerateShotNumber( class UMovieSceneShotTrack& ShotTrack, int32 SectionIndex ) const;
+
+	/** Finds the index in the ShotSections array where a new shot should be inserted */
+	int32 FindIndexForNewShot( const TArray<UMovieSceneSection*>& ShotSections, float NewShotTime ) const;
 private:
 	/** The Thumbnail pool which draws all the viewport thumbnails for the director track */
 	TSharedPtr<class FShotThumbnailPool> ThumbnailPool;
diff --git a/Engine/Source/Editor/NiagaraEditor/Private/NiagaraSequencer.h b/Engine/Source/Editor/NiagaraEditor/Private/NiagaraSequencer.h
--- a/Engine/Source/Editor/NiagaraEditor/Private/NiagaraSequencer.h
+++ b/Engine/Source/Editor/NiagaraEditor/Private/NiagaraSequencer.h
@@ -166,7 +166,7 @@ public:
 	{
 		return MakeShareable(new INiagaraTrackInstance(this));
 	}
-	virtual TArray<UMovieSceneSection*> GetAllSections() const
+	virtual const TArray<UMovieSceneSection*>& GetAllSections() const override
 	{
 		return Sections;
 	}
@@ -229,6 +229,7 @@ public:
 	virtual bool AllowsSpawnableObjects() const override { return false; }
 	virtual FGuid FindGuidForObject(const UMovieScene& MovieScene, UObject& Object) const override { return FGuid(); }
 	virtual void SpawnOrDestroyObjectsForInstance(TSharedRef<FMovieSceneInstance> MovieSceneInstance, bool bDestroyAll) override { }
+	virtual void RemoveMovieSceneInstance( TSharedRef<FMovieSceneInstance> MovieSceneInstance ) override {}
 	virtual void DestroyAllSpawnedObjects() override { }
 	virtual bool CanPossessObject(UObject& Object) const override 	{ return false; }
 	virtual void BindPossessableObject(const FGuid& PossessableGuid, UObject& PossessedObject) override	{};
diff --git a/Engine/Source/Editor/Sequencer/Private/SSection.cpp b/Engine/Source/Editor/Sequencer/Private/SSection.cpp
--- a/Engine/Source/Editor/Sequencer/Private/SSection.cpp
+++ b/Engine/Source/Editor/Sequencer/Private/SSection.cpp
@@ -208,8 +208,6 @@ int32 SSection::OnPaint( const FPaintArgs& Args, const FGeometry& AllottedGeomet
 {
 	int32 StartLayer = SCompoundWidget::OnPaint( Args, AllottedGeometry, MyClippingRect, OutDrawElements, LayerId, InWidgetStyle, bParentEnabled );
 
-	DrawSectionBorders( AllottedGeometry, MyClippingRect, OutDrawElements, StartLayer );
-
 	FGeometry SectionGeometry = AllottedGeometry.MakeChild( FVector2D( SequencerSectionConstants::SectionGripSize, 0 ), AllottedGeometry.GetDrawSize() - FVector2D( SequencerSectionConstants::SectionGripSize*2, 0.0f ) );
 
 	FSlateRect SectionClipRect = SectionGeometry.GetClippingRect().IntersectionWith( MyClippingRect );
@@ -218,10 +216,10 @@ int32 SSection::OnPaint( const FPaintArgs& Args, const FGeometry& AllottedGeomet
 	int32 PostSectionLayer = SectionInterface->OnPaintSection( SectionGeometry, SectionClipRect, OutDrawElements, LayerId, bParentEnabled );
 
 	// @todo Sequencer - Temp indicators of the area of a section which can be dragged.
+	DrawSectionBorders(AllottedGeometry, MyClippingRect, OutDrawElements, PostSectionLayer );
 
 	PaintKeys( SectionGeometry, MyClippingRect, OutDrawElements, PostSectionLayer, InWidgetStyle );
 
-	
 	// Section name with drop shadow
 	FText SectionTitle = SectionInterface->GetSectionTitle();
 	if (!SectionTitle.IsEmpty())
diff --git a/Engine/Source/Editor/Sequencer/Private/SSequencer.cpp b/Engine/Source/Editor/Sequencer/Private/SSequencer.cpp
--- a/Engine/Source/Editor/Sequencer/Private/SSequencer.cpp
+++ b/Engine/Source/Editor/Sequencer/Private/SSequencer.cpp
@@ -1131,7 +1131,7 @@ FText SSequencer::GetRootMovieSceneName() const
 
 FText SSequencer::GetShotSectionTitle(UMovieSceneSection* ShotSection) const
 {
-	return Cast<UMovieSceneShotSection>(ShotSection)->GetTitle();
+	return Cast<UMovieSceneShotSection>(ShotSection)->GetShotDisplayName();
 }
 
 void SSequencer::DeleteSelectedNodes()
diff --git a/Engine/Source/Editor/Sequencer/Private/SectionDragOperations.cpp b/Engine/Source/Editor/Sequencer/Private/SectionDragOperations.cpp
--- a/Engine/Source/Editor/Sequencer/Private/SectionDragOperations.cpp
+++ b/Engine/Source/Editor/Sequencer/Private/SectionDragOperations.cpp
@@ -81,7 +81,7 @@ void FSequencerDragOperation::GetSectionSnapTimes(TArray<float>& OutSnapTimes, U
 	UMovieSceneTrack* ShotTrack = MovieScene->FindMasterTrack(UMovieSceneShotTrack::StaticClass());
 	if (ShotTrack && OuterTrack != ShotTrack)
 	{
-		TArray<UMovieSceneSection*> ShotSections = ShotTrack->GetAllSections();
+		const TArray<UMovieSceneSection*>& ShotSections = ShotTrack->GetAllSections();
 		for (int32 SectionIndex = 0; SectionIndex < ShotSections.Num(); ++SectionIndex)
 		{
 			auto Shot = ShotSections[SectionIndex];
@@ -257,16 +257,25 @@ void FMoveSection::OnBeginDrag(const FVector2D& LocalMousePos, TSharedPtr<FTrack
 
 void FMoveSection::OnEndDrag(TSharedPtr<FTrackNode> SequencerNode)
 {
-	EndTransaction();
-
 	DraggedKeyHandles.Empty();
 
 	if (Section.IsValid())
 	{
 		SequencerNode->FixRowIndices();
+
+		UMovieSceneTrack* OuterTrack = Cast<UMovieSceneTrack>(Section->GetOuter());
+
+		if (OuterTrack)
+		{
+			OuterTrack->Modify();
+			OuterTrack->OnSectionMoved(*Section);
+		}
 	}
+
+	EndTransaction();
 }
 
+
 void FMoveSection::OnDrag( const FPointerEvent& MouseEvent, const FVector2D& LocalMousePos, const FTimeToPixel& TimeToPixelConverter, TSharedPtr<FTrackNode> SequencerNode )
 {
 	if( Section.IsValid() )
@@ -288,22 +297,22 @@ void FMoveSection::OnDrag( const FPointerEvent& MouseEvent, const FVector2D& Loc
 		{
 			bool bSnappedToSection = false;
 			if ( Settings->GetSnapSectionTimesToSections() )
-		{
-			TArray<float> TimesToSnapTo;
-			GetSectionSnapTimes(TimesToSnapTo, Section.Get(), SequencerNode, true);
+			{
+				TArray<float> TimesToSnapTo;
+				GetSectionSnapTimes(TimesToSnapTo, Section.Get(), SequencerNode, true);
 
-			TArray<float> TimesToSnap;
-			TimesToSnap.Add(DistanceMoved + Section->GetStartTime());
-			TimesToSnap.Add(DistanceMoved + Section->GetEndTime());
+				TArray<float> TimesToSnap;
+				TimesToSnap.Add(DistanceMoved + Section->GetStartTime());
+				TimesToSnap.Add(DistanceMoved + Section->GetEndTime());
 
-			float OutSnappedTime = 0.f;
-			float OutNewTime = 0.f;
-				if ( SnapToTimes( TimesToSnap, TimesToSnapTo, TimeToPixelConverter, OutSnappedTime, OutNewTime ) )
-			{
-				DeltaTime = OutNewTime - (OutSnappedTime - DistanceMoved);
+				float OutSnappedTime = 0.f;
+				float OutNewTime = 0.f;
+				if (SnapToTimes(TimesToSnap, TimesToSnapTo, TimeToPixelConverter, OutSnappedTime, OutNewTime))
+				{
+					DeltaTime = OutNewTime - (OutSnappedTime - DistanceMoved);
 					bSnappedToSection = true;
+				}
 			}
-		}
 
 			if ( bSnappedToSection == false && Settings->GetSnapSectionTimesToInterval() )
 			{
diff --git a/Engine/Source/Editor/Sequencer/Private/Sequencer.cpp b/Engine/Source/Editor/Sequencer/Private/Sequencer.cpp
--- a/Engine/Source/Editor/Sequencer/Private/Sequencer.cpp
+++ b/Engine/Source/Editor/Sequencer/Private/Sequencer.cpp
@@ -695,7 +695,7 @@ void FSequencer::GetRuntimeObjects( TSharedRef<FMovieSceneInstance> MovieSceneIn
 	}*/
 }
 
-void FSequencer::UpdatePreviewViewports(UObject* ObjectToViewThrough) const
+void FSequencer::UpdateCameraCut(UObject* ObjectToViewThrough, bool bNewCameraCut) const
 {
 	if(!IsPerspectiveViewportPosessionEnabled())
 	{
@@ -706,7 +706,7 @@ void FSequencer::UpdatePreviewViewports(UObject* ObjectToViewThrough) const
 	{
 		if(LevelVC && LevelVC->IsPerspective() && LevelVC->AllowsCinematicPreview())
 		{
-			LevelVC->SetMatineeActorLock(Cast<AActor>(ObjectToViewThrough));
+			UpdatePreviewLevelViewportClientFromCameraCut( *LevelVC, ObjectToViewThrough, bNewCameraCut );
 		}
 	}
 }
@@ -727,8 +727,7 @@ void FSequencer::RemoveMovieSceneInstance( UMovieSceneSection& MovieSceneSection
 {
 	if( ObjectBindingManager->AllowsSpawnableObjects() )
 	{
-		const bool bDestroyAll = true;
-		ObjectBindingManager->SpawnOrDestroyObjectsForInstance( InstanceToRemove, bDestroyAll );
+		ObjectBindingManager->RemoveMovieSceneInstance( InstanceToRemove );
 	}
 	
 
@@ -1256,6 +1255,63 @@ void FSequencer::OnPostSaveWorld(uint32 SaveFlags, class UWorld* World, bool bSu
 	SetGlobalTime(GetGlobalTime());
 }
 
+void FSequencer::UpdatePreviewLevelViewportClientFromCameraCut( FLevelEditorViewportClient& InViewportClient, UObject* InCameraObject, bool bNewCameraCut ) const
+{
+	ACameraActor* Cam = Cast<ACameraActor>(InCameraObject);
+	if (Cam)
+	{
+		InViewportClient.SetViewLocation(Cam->GetActorLocation());
+		InViewportClient.SetViewRotation(Cam->GetActorRotation());
+
+		// @todo Sequencer previewing
+		/*InViewportClient.FadeAmount = InFadeAmount;
+		InViewportClient.bEnableFading = true;*/
+
+		/*InViewportClient.bEnableColorScaling = bInEnableColorScaling;
+		InViewportClient.ColorScale = InColorScale;*/
+
+		InViewportClient.bEditorCameraCut = bNewCameraCut;
+	}
+	else
+	{
+		InViewportClient.ViewFOV = InViewportClient.FOVAngle;
+
+		// @todo Sequencer previewing
+		/*InViewportClient.>FadeAmount = InFadeAmount;
+		InViewportClient.bEnableFading = true;*/
+
+		InViewportClient.bEditorCameraCut = false;
+	}
+
+	// Set the actor lock.
+	InViewportClient.SetMatineeActorLock(Cam);
+
+	// If viewing through a camera - enforce aspect ratio.
+	if (Cam)
+	{
+		if (Cam->GetCameraComponent()->AspectRatio == 0)
+		{
+			InViewportClient.AspectRatio = 1.7f;
+		}
+		else
+		{
+			InViewportClient.AspectRatio = Cam->GetCameraComponent()->AspectRatio;
+		}
+
+		//don't stop the camera from zooming when not playing back
+		InViewportClient.ViewFOV = Cam->GetCameraComponent()->FieldOfView;
+
+		// If there are selected actors, invalidate the viewports hit proxies, otherwise they won't be selectable afterwards
+		if (InViewportClient.Viewport && GEditor->GetSelectedActorCount() > 0)
+		{
+			InViewportClient.Viewport->InvalidateHitProxy();
+		}
+	}
+
+	// Update ControllingActorViewInfo, so it is in sync with the updated viewport
+	InViewportClient.UpdateViewForLockedActor();
+}
+
 void FSequencer::SaveCurrentMovieScene()
 {
 	UPackage* MovieScenePackage = GetFocusedMovieScene()->GetOutermost();
diff --git a/Engine/Source/Editor/Sequencer/Private/Sequencer.h b/Engine/Source/Editor/Sequencer/Private/Sequencer.h
--- a/Engine/Source/Editor/Sequencer/Private/Sequencer.h
+++ b/Engine/Source/Editor/Sequencer/Private/Sequencer.h
@@ -227,7 +227,7 @@ public:
 
 	/** IMovieScenePlayer interface */
 	virtual void GetRuntimeObjects( TSharedRef<FMovieSceneInstance> MovieSceneInstance, const FGuid& ObjectHandle, TArray< UObject* >& OutObjects ) const override;
-	virtual void UpdatePreviewViewports(UObject* ObjectToViewThrough) const override;
+	virtual void UpdateCameraCut(UObject* ObjectToViewThrough, bool bNewCameraCut) const override;
 	virtual EMovieScenePlayerStatus::Type GetPlaybackStatus() const override;
 	virtual void AddMovieSceneInstance( class UMovieSceneSection& MovieSceneSection, TSharedRef<FMovieSceneInstance> InstanceToAdd ) override;
 	virtual void RemoveMovieSceneInstance( class UMovieSceneSection& MovieSceneSection, TSharedRef<FMovieSceneInstance> InstanceToRemove ) override;
@@ -370,10 +370,14 @@ protected:
 
 	/** Called after the world has been saved. The sequencer updates to the animated state. */
 	void OnPostSaveWorld(uint32 SaveFlags, class UWorld* World, bool bSuccess);
+
+	/** Updates a viewport client from camera cut data */
+	void UpdatePreviewLevelViewportClientFromCameraCut( FLevelEditorViewportClient& InViewportClient, UObject* InCameraObject, bool bNewCameraCut ) const;
+
 private:
 	TMap< TWeakObjectPtr<UMovieSceneSection>, TSharedRef<FMovieSceneInstance> > MovieSceneSectionToInstanceMap;
 
-	/** Command list for seququencer commands */
+	/** Command list for sequencer commands */
 	TSharedRef<FUICommandList> SequencerCommandBindings;
 
 	/** List of tools we own */
diff --git a/Engine/Source/Editor/Sequencer/Private/SequencerActorBindingManager.cpp b/Engine/Source/Editor/Sequencer/Private/SequencerActorBindingManager.cpp
--- a/Engine/Source/Editor/Sequencer/Private/SequencerActorBindingManager.cpp
+++ b/Engine/Source/Editor/Sequencer/Private/SequencerActorBindingManager.cpp
@@ -90,6 +90,7 @@ void FSequencerActorBindingManager::SpawnOrDestroyObjectsForInstance( TSharedRef
 	UMovieScene* MovieScene = MovieSceneInstance->GetMovieScene();
 
 	// Remove any puppet objects that we no longer need
+	if( MovieScene )
 	{
 		for( auto PuppetObjectIndex = 0; PuppetObjectIndex < PuppetObjects.Num(); ++PuppetObjectIndex )
 		{
@@ -163,7 +164,7 @@ void FSequencerActorBindingManager::SpawnOrDestroyObjectsForInstance( TSharedRef
 		}
 	}
 
-	if( !bDestroyAll )
+	if( !bDestroyAll && MovieScene )
 	{
 		for( auto SpawnableIndex = 0; SpawnableIndex < MovieScene->GetSpawnableCount(); ++SpawnableIndex )
 		{
@@ -238,6 +239,12 @@ void FSequencerActorBindingManager::SpawnOrDestroyObjectsForInstance( TSharedRef
 	}
 }
 
+void FSequencerActorBindingManager::RemoveMovieSceneInstance( TSharedRef<FMovieSceneInstance> MovieSceneInstance )
+{
+	SpawnOrDestroyObjectsForInstance( MovieSceneInstance, true );
+	InstanceToPuppetObjectsMap.Remove( MovieSceneInstance );
+}
+
 void FSequencerActorBindingManager::DestroyAllSpawnedObjects()
 {
 	const bool bDestroyAll = true;
diff --git a/Engine/Source/Editor/Sequencer/Private/SequencerActorBindingManager.h b/Engine/Source/Editor/Sequencer/Private/SequencerActorBindingManager.h
--- a/Engine/Source/Editor/Sequencer/Private/SequencerActorBindingManager.h
+++ b/Engine/Source/Editor/Sequencer/Private/SequencerActorBindingManager.h
@@ -81,6 +81,7 @@ public:
 	virtual bool AllowsSpawnableObjects() const override { return true; }
 	virtual FGuid FindGuidForObject( const UMovieScene& MovieScene, UObject& Object ) const override;
 	virtual void SpawnOrDestroyObjectsForInstance( TSharedRef<FMovieSceneInstance> MovieSceneInstance, bool bDestroyAll ) override;
+	virtual void RemoveMovieSceneInstance( TSharedRef<FMovieSceneInstance> MovieSceneInstance ) override;
 	virtual void DestroyAllSpawnedObjects() override;
 	virtual bool CanPossessObject( UObject& Object ) const override;
 	virtual void BindPossessableObject( const FGuid& PossessableGuid, UObject& PossessedObject ) override;
diff --git a/Engine/Source/Editor/Sequencer/Private/SequencerAssetEditor.cpp b/Engine/Source/Editor/Sequencer/Private/SequencerAssetEditor.cpp
--- a/Engine/Source/Editor/Sequencer/Private/SequencerAssetEditor.cpp
+++ b/Engine/Source/Editor/Sequencer/Private/SequencerAssetEditor.cpp
@@ -93,8 +93,6 @@ void FSequencerAssetEditor::InitSequencerAssetEditor( const EToolkitMode::Type M
 		// We need to find out when the user loads a new map, because we might need to re-create puppet actors
 		// when previewing a MovieScene
 		LevelEditorModule.OnMapChanged().AddSP(Sequencer.ToSharedRef(), &FSequencer::OnMapChanged);
-
-		AttachTransportControlsToViewports();
 	}
 }
 
@@ -110,8 +108,6 @@ FSequencerAssetEditor::FSequencerAssetEditor()
 
 FSequencerAssetEditor::~FSequencerAssetEditor()
 {
-	DetachTransportControlsFromViewports();
-
 	Sequencer->OnClose();
 
 	// Unregister delegates
@@ -161,93 +157,4 @@ FString FSequencerAssetEditor::GetWorldCentricTabPrefix() const
 }
 
 
-void FSequencerAssetEditor::AttachTransportControlsToViewports()
-{
-	FLevelEditorModule* Module = FModuleManager::Get().LoadModulePtr<FLevelEditorModule>("LevelEditor");
-	if (Module)
-	{
-		TSharedPtr<ILevelEditor> LevelEditor = Module->GetFirstLevelEditor();
-		const TArray< TSharedPtr<ILevelViewport> >& LevelViewports = LevelEditor->GetViewports();
-		
-		FEditorWidgetsModule& EditorWidgetsModule = FModuleManager::Get().LoadModuleChecked<FEditorWidgetsModule>( "EditorWidgets" );
-
-		TSharedRef<FSequencer> SequencerRef = Sequencer.ToSharedRef();
-
-		FTransportControlArgs TransportControlArgs;
-		TransportControlArgs.OnForwardPlay.BindSP(SequencerRef, &FSequencer::OnPlay, true);
-		TransportControlArgs.OnRecord.BindSP(SequencerRef, &FSequencer::OnRecord);
-		TransportControlArgs.OnForwardStep.BindSP(SequencerRef, &FSequencer::OnStepForward);
-		TransportControlArgs.OnBackwardStep.BindSP(SequencerRef, &FSequencer::OnStepBackward);
-		TransportControlArgs.OnForwardEnd.BindSP(SequencerRef, &FSequencer::OnStepToEnd);
-		TransportControlArgs.OnBackwardEnd.BindSP(SequencerRef, &FSequencer::OnStepToBeginning);
-		TransportControlArgs.OnToggleLooping.BindSP(SequencerRef, &FSequencer::OnToggleLooping);
-		TransportControlArgs.OnGetLooping.BindSP(SequencerRef, &FSequencer::IsLooping);
-		TransportControlArgs.OnGetPlaybackMode.BindSP(SequencerRef, &FSequencer::GetPlaybackMode);
-
-		for ( const TSharedPtr<ILevelViewport>& LevelViewport : LevelViewports )
-		{
-			TWeakPtr<ILevelViewport> LevelViewportWeakPtr = LevelViewport;
-			
-			TSharedPtr<SWidget> TransportControl =
-				SNew(SHorizontalBox)
-				.Visibility(EVisibility::SelfHitTestInvisible)
-				+SHorizontalBox::Slot()
-				.FillWidth(1)
-				.HAlign(HAlign_Center)
-				.VAlign(VAlign_Bottom)
-				.Padding(4.f)
-				[
-					SNew(SBorder)
-					.Padding(4.f)
-					.Cursor( EMouseCursor::Default )
-					.BorderImage( FEditorStyle::GetBrush( "FilledBorder" ) )
-					.Visibility(this, &FSequencerAssetEditor::GetTransportControlVisibility, LevelViewportWeakPtr)
-					.Content()
-					[
-						EditorWidgetsModule.CreateTransportControl(TransportControlArgs)
-					]
-				];
-
-			LevelViewport->AddOverlayWidget(TransportControl.ToSharedRef());
-
-			TransportControls.Add(LevelViewportWeakPtr, TransportControl);
-		}
-	}
-}
-
-
-void FSequencerAssetEditor::DetachTransportControlsFromViewports()
-{
-	FLevelEditorModule* Module = FModuleManager::Get().LoadModulePtr<FLevelEditorModule>("LevelEditor");
-	if (Module)
-	{
-		TSharedPtr<ILevelEditor> LevelEditor = Module->GetFirstLevelEditor();
-		if (LevelEditor.IsValid())
-		{
-			const TArray< TSharedPtr<ILevelViewport> >& LevelViewports = LevelEditor->GetViewports();
-		
-			for (int32 i = 0; i < LevelViewports.Num(); ++i)
-			{
-				const TSharedPtr<ILevelViewport>& LevelViewport = LevelViewports[i];
-
-				TSharedPtr<SWidget>* TransportControl = TransportControls.Find(LevelViewport);
-				if (TransportControl && TransportControl->IsValid())
-				{
-					LevelViewport->RemoveOverlayWidget(TransportControl->ToSharedRef());
-				}
-			}
-		}
-	}
-}
-
-EVisibility FSequencerAssetEditor::GetTransportControlVisibility(TWeakPtr<ILevelViewport> LevelViewport) const
-{
-	TSharedPtr<ILevelViewport> LevelViewportPin = LevelViewport.Pin();
-
-	FLevelEditorViewportClient& ViewportClient = LevelViewportPin->GetLevelViewportClient();
-	return (ViewportClient.IsPerspective() && ViewportClient.AllowsCinematicPreview()) ? EVisibility::Visible : EVisibility::Collapsed;
-}
-
-
-
 #undef LOCTEXT_NAMESPACE
diff --git a/Engine/Source/Editor/Sequencer/Private/SequencerAssetEditor.h b/Engine/Source/Editor/Sequencer/Private/SequencerAssetEditor.h
--- a/Engine/Source/Editor/Sequencer/Private/SequencerAssetEditor.h
+++ b/Engine/Source/Editor/Sequencer/Private/SequencerAssetEditor.h
@@ -41,13 +41,7 @@ public:
 private:
 	EVisibility GetTransportControlVisibility(TWeakPtr<ILevelViewport> LevelViewport) const;
 	TSharedRef<SDockTab> SpawnTab_SequencerMain(const FSpawnTabArgs& Args);	
-protected:
 
-	/** Generates and attaches transport control widgets to the main level editor viewports */
-	void AttachTransportControlsToViewports();
-	/** Purges all transport control widgets from the main level editor viewports */
-	void DetachTransportControlsFromViewports();
-	
 private:
 	TSharedPtr<FSequencer> Sequencer;
 
diff --git a/Engine/Source/Editor/Sequencer/Private/SequencerNodeTree.cpp b/Engine/Source/Editor/Sequencer/Private/SequencerNodeTree.cpp
--- a/Engine/Source/Editor/Sequencer/Private/SequencerNodeTree.cpp
+++ b/Engine/Source/Editor/Sequencer/Private/SequencerNodeTree.cpp
@@ -27,17 +27,19 @@ void FSequencerNodeTree::Update()
 
 	UMovieScene* MovieScene = Sequencer.GetFocusedMovieScene();
 
+	TArray< TSharedRef<FSequencerDisplayNode> > NewRootNodes;
+
 	// Get the master tracks  so we can get sections from them
 	const TArray<UMovieSceneTrack*>& MasterTracks = MovieScene->GetMasterTracks();
 
-	for( int32 TrackIndex = 0; TrackIndex < MasterTracks.Num(); ++TrackIndex )
+	for( UMovieSceneTrack* Track : MasterTracks )
 	{
-		UMovieSceneTrack& Track = *MasterTracks[TrackIndex];
+		UMovieSceneTrack& TrackRef = *Track;
 
-		TSharedRef<FTrackNode> SectionNode = MakeShareable( new FTrackNode( Track.GetTrackName(), Track, NULL, *this ) );
-		RootNodes.Add( SectionNode );
+		TSharedRef<FTrackNode> SectionNode = MakeShareable( new FTrackNode( TrackRef.GetTrackName(), TrackRef, NULL, *this ) );
+		NewRootNodes.Add( SectionNode );
 	
-		MakeSectionInterfaces( Track, SectionNode );
+		MakeSectionInterfaces( TrackRef, SectionNode );
 
 		SectionNode->PinNode();
 	}
@@ -46,21 +48,21 @@ void FSequencerNodeTree::Update()
 	const TArray<FMovieSceneObjectBinding>& ObjectBindings = MovieScene->GetObjectBindings();
 
 	// Make nodes for all object bindings
-	for( int32 BindingIndex = 0; BindingIndex < ObjectBindings.Num(); ++BindingIndex )
+	for( const FMovieSceneObjectBinding& Binding : ObjectBindings )
 	{
-		TSharedRef<FObjectBindingNode> ObjectBindingNode = AddObjectBinding( ObjectBindings[BindingIndex].GetName(), ObjectBindings[BindingIndex].GetObjectGuid() );
+		TSharedRef<FObjectBindingNode> ObjectBindingNode = AddObjectBinding( Binding.GetName(), Binding.GetObjectGuid(), NewRootNodes );
 
-		const TArray<UMovieSceneTrack*>& Tracks = ObjectBindings[BindingIndex].GetTracks();
+		const TArray<UMovieSceneTrack*>& Tracks = Binding.GetTracks();
 
-		for( int32 TrackIndex = 0; TrackIndex < Tracks.Num(); ++TrackIndex )
+		for( UMovieSceneTrack* Track : Tracks )
 		{
-			UMovieSceneTrack& Track = *Tracks[TrackIndex];
+			UMovieSceneTrack& TrackRef = *Track;
 
-			FName SectionName = Track.GetTrackName();
+			FName SectionName = TrackRef.GetTrackName();
 			check( SectionName != NAME_None );
 
-			TSharedRef<FTrackNode> SectionAreaNode = ObjectBindingNode->AddSectionAreaNode( SectionName, Track );
-			MakeSectionInterfaces( Track, SectionAreaNode );
+			TSharedRef<FTrackNode> SectionAreaNode = ObjectBindingNode->AddSectionAreaNode( SectionName, TrackRef );
+			MakeSectionInterfaces( TrackRef, SectionAreaNode );
 		}
 
 	}
@@ -96,7 +98,24 @@ void FSequencerNodeTree::Update()
 	};
 
 	// Sort so that master tracks appear before object tracks
-	RootNodes.Sort( FRootNodeSorter() );
+	NewRootNodes.Sort( FRootNodeSorter() );
+
+	// Look for a shot track.  It will always come first if it exists
+	UMovieSceneTrack* ShotTrack = MovieScene->GetShotTrack();
+	if( ShotTrack )
+	{
+		TSharedRef<FTrackNode> SectionNode = MakeShareable( new FTrackNode( ShotTrack->GetTrackName(), *ShotTrack, NULL, *this ) );
+
+		// Shot track always comes first
+		RootNodes.Add( SectionNode );
+
+		MakeSectionInterfaces( *ShotTrack, SectionNode );
+
+		SectionNode->PinNode();
+	}
+
+	// Add all other nodes after the shot track
+	RootNodes.Append( NewRootNodes );
 
 	// Re-filter the tree after updating 
 	// @todo Sequencer - Newly added sections may need to be visible even when there is a filter
@@ -155,7 +174,7 @@ const TArray< TSharedRef<FSequencerDisplayNode> >& FSequencerNodeTree::GetRootNo
 }
 
 
-TSharedRef<FObjectBindingNode> FSequencerNodeTree::AddObjectBinding( const FString& ObjectName, const FGuid& ObjectBinding )
+TSharedRef<FObjectBindingNode> FSequencerNodeTree::AddObjectBinding( const FString& ObjectName, const FGuid& ObjectBinding, TArray< TSharedRef<FSequencerDisplayNode> >& OutNodeList )
 {
 	// The node name is the object guid
 	FName ObjectNodeName = *ObjectBinding.ToString();
@@ -164,7 +183,7 @@ TSharedRef<FObjectBindingNode> FSequencerNodeTree::AddObjectBinding( const FStri
 	TSharedRef< FObjectBindingNode > ObjectNode = MakeShareable( new FObjectBindingNode( ObjectNodeName, ObjectName, ObjectBinding, ParentNode, *this ) );
 
 	// Object binding nodes are always root nodes
-	RootNodes.Add( ObjectNode );
+	OutNodeList.Add( ObjectNode );
 
 	// Map the guid to the object binding node for fast lookup later
 	ObjectBindingMap.Add( ObjectBinding, ObjectNode );
diff --git a/Engine/Source/Editor/Sequencer/Private/SequencerNodeTree.h b/Engine/Source/Editor/Sequencer/Private/SequencerNodeTree.h
--- a/Engine/Source/Editor/Sequencer/Private/SequencerNodeTree.h
+++ b/Engine/Source/Editor/Sequencer/Private/SequencerNodeTree.h
@@ -93,7 +93,7 @@ private:
 	 * 
 	 * @param ObjectBinding	The object binding guid
 	 */
-	TSharedRef<FObjectBindingNode> AddObjectBinding( const FString& ObjectName, const FGuid& ObjectBinding );
+	TSharedRef<FObjectBindingNode> AddObjectBinding( const FString& ObjectName, const FGuid& ObjectBinding, TArray< TSharedRef<FSequencerDisplayNode> >& OutNodeList );
 private:
 	/** Tools for building movie scene section layouts.  One tool for each track */
 	TMap< class UMovieSceneTrack*, TSharedPtr<FMovieSceneTrackEditor> > EditorMap;
diff --git a/Engine/Source/Editor/Sequencer/Public/ISequencerObjectBindingManager.h b/Engine/Source/Editor/Sequencer/Public/ISequencerObjectBindingManager.h
--- a/Engine/Source/Editor/Sequencer/Public/ISequencerObjectBindingManager.h
+++ b/Engine/Source/Editor/Sequencer/Public/ISequencerObjectBindingManager.h
@@ -32,6 +32,8 @@ public:
 	 */
 	virtual void SpawnOrDestroyObjectsForInstance( TSharedRef<FMovieSceneInstance> MovieSceneInstance, bool bDestroyAll ) = 0;
 	
+	virtual void RemoveMovieSceneInstance( TSharedRef<FMovieSceneInstance> MovieSceneInstance ) = 0;
+
 	virtual void DestroyAllSpawnedObjects() = 0;
 	
 	virtual bool CanPossessObject( UObject& Object ) const = 0;
diff --git a/Engine/Source/Editor/UMGEditor/Private/Animation/UMGSequencerObjectBindingManager.h b/Engine/Source/Editor/UMGEditor/Private/Animation/UMGSequencerObjectBindingManager.h
--- a/Engine/Source/Editor/UMGEditor/Private/Animation/UMGSequencerObjectBindingManager.h
+++ b/Engine/Source/Editor/UMGEditor/Private/Animation/UMGSequencerObjectBindingManager.h
@@ -21,6 +21,7 @@ public:
 	virtual FGuid FindGuidForObject( const UMovieScene& MovieScene, UObject& Object ) const override;
 	virtual void SpawnOrDestroyObjectsForInstance( TSharedRef<FMovieSceneInstance> MovieSceneInstance, bool bDestroyAll ) override {}
 	virtual void DestroyAllSpawnedObjects() override {}
+	virtual void RemoveMovieSceneInstance( TSharedRef<FMovieSceneInstance> MovieSceneInstance ) override {}
 	virtual bool CanPossessObject( UObject& Object ) const override;
 	virtual void BindPossessableObject( const FGuid& PossessableGuid, UObject& PossessedObject ) override;
 	virtual void UnbindPossessableObjects( const FGuid& PossessableGuid ) override;
diff --git a/Engine/Source/Runtime/MovieScene/Private/MovieScene.cpp b/Engine/Source/Runtime/MovieScene/Private/MovieScene.cpp
--- a/Engine/Source/Runtime/MovieScene/Private/MovieScene.cpp
+++ b/Engine/Source/Runtime/MovieScene/Private/MovieScene.cpp
@@ -367,6 +367,23 @@ class UMovieSceneTrack* UMovieScene::AddMasterTrack( TSubclassOf<UMovieSceneTrac
 	return CreatedType;
 }
 
+UMovieSceneTrack* UMovieScene::AddShotTrack( TSubclassOf<UMovieSceneTrack> TrackClass )
+{
+	if( !ShotTrack )
+	{
+		Modify();
+
+		ShotTrack = NewObject<UMovieSceneTrack>(this, TrackClass, FName("Shots"), RF_Transactional);
+	}
+
+	return ShotTrack;
+}
+
+UMovieSceneTrack* UMovieScene::GetShotTrack()
+{
+	return ShotTrack;
+}
+
 bool UMovieScene::RemoveMasterTrack( UMovieSceneTrack* Track ) 
 {
 	Modify();
diff --git a/Engine/Source/Runtime/MovieScene/Private/MovieSceneInstance.cpp b/Engine/Source/Runtime/MovieScene/Private/MovieSceneInstance.cpp
--- a/Engine/Source/Runtime/MovieScene/Private/MovieSceneInstance.cpp
+++ b/Engine/Source/Runtime/MovieScene/Private/MovieSceneInstance.cpp
@@ -38,13 +38,20 @@ void FMovieSceneInstance::RestoreState()
 
 void FMovieSceneInstance::Update( float Position, float LastPosition, class IMovieScenePlayer& Player )
 {
-	// Update each type
-	TArray<UObject*> Objects;
+	// Update  shot track
+	TArray<UObject*> NoObjects;
+	if( ShotTrackInstance.IsValid() )
+	{
+		ShotTrackInstance->Update( Position, LastPosition, NoObjects, Player );
+	}
+
+	// Update each master track
 	for( FMovieSceneInstanceMap::TIterator It( MasterTrackInstances ); It; ++It )
 	{
-		It.Value()->Update( Position, LastPosition, Objects, Player );
+		It.Value()->Update( Position, LastPosition, NoObjects, Player );
 	}
 
+	// Update tracks bound to objects
 	TMap<FGuid, FMovieSceneObjectBindingInstance>::TIterator ObjectIt = ObjectBindingInstances.CreateIterator();
 	for(; ObjectIt; ++ObjectIt )
 	{
@@ -59,6 +66,22 @@ void FMovieSceneInstance::Update( float Position, float LastPosition, class IMov
 
 void FMovieSceneInstance::RefreshInstance( IMovieScenePlayer& Player )
 {
+	UMovieSceneTrack* ShotTrack = MovieScene->GetShotTrack();
+
+	TSharedRef<FMovieSceneInstance> ThisInstance = AsShared();
+
+	// Only if root movie scene. Any sub-movie scene that has a shot track is ignored
+	if( ShotTrack && Player.GetRootMovieSceneInstance() == ThisInstance )
+	{
+		FMovieSceneInstanceMap ShotTrackInstanceMap;
+		TArray<UObject*> Objects;
+		TArray<UMovieSceneTrack*> Tracks;
+		Tracks.Add(ShotTrack);
+		RefreshInstanceMap(Tracks, Objects, ShotTrackInstanceMap, Player);
+
+		ShotTrackInstance = ShotTrackInstanceMap.FindRef( ShotTrack );
+	}
+
 	// Get all the master tracks and create instances for them if needed
 	const TArray<UMovieSceneTrack*>& MasterTracks = MovieScene->GetMasterTracks();
 	TArray<UObject*> Objects;
diff --git a/Engine/Source/Runtime/MovieScene/Private/MovieScenePropertyTrack.cpp b/Engine/Source/Runtime/MovieScene/Private/MovieScenePropertyTrack.cpp
--- a/Engine/Source/Runtime/MovieScene/Private/MovieScenePropertyTrack.cpp
+++ b/Engine/Source/Runtime/MovieScene/Private/MovieScenePropertyTrack.cpp
@@ -19,15 +19,9 @@ void UMovieScenePropertyTrack::SetPropertyNameAndPath( FName InPropertyName, con
 	PropertyPath = InPropertyPath;
 }
 
-TArray<UMovieSceneSection*> UMovieScenePropertyTrack::GetAllSections() const
+const TArray<UMovieSceneSection*>& UMovieScenePropertyTrack::GetAllSections() const
 {
-	TArray<UMovieSceneSection*> OutSectionPayloads;
-	for( int32 SectionIndex = 0; SectionIndex < Sections.Num(); ++SectionIndex )
-	{
-		OutSectionPayloads.Add( Sections[SectionIndex] );
-	}
-	
-	return OutSectionPayloads;
+	return Sections;
 }
 
 
diff --git a/Engine/Source/Runtime/MovieScene/Private/RuntimeMovieScenePlayer.cpp b/Engine/Source/Runtime/MovieScene/Private/RuntimeMovieScenePlayer.cpp
--- a/Engine/Source/Runtime/MovieScene/Private/RuntimeMovieScenePlayer.cpp
+++ b/Engine/Source/Runtime/MovieScene/Private/RuntimeMovieScenePlayer.cpp
@@ -202,6 +202,10 @@ void URuntimeMovieScenePlayer::GetRuntimeObjects( TSharedRef<FMovieSceneInstance
 	}
 }
 
+void URuntimeMovieScenePlayer::UpdateCameraCut(UObject* ObjectToViewThrough, bool bNewCameraCut) const
+{
+}
+
 EMovieScenePlayerStatus::Type URuntimeMovieScenePlayer::GetPlaybackStatus() const
 {
 	return bIsPlaying ? EMovieScenePlayerStatus::Playing : EMovieScenePlayerStatus::Stopped;
diff --git a/Engine/Source/Runtime/MovieScene/Public/IMovieScenePlayer.h b/Engine/Source/Runtime/MovieScene/Public/IMovieScenePlayer.h
--- a/Engine/Source/Runtime/MovieScene/Public/IMovieScenePlayer.h
+++ b/Engine/Source/Runtime/MovieScene/Public/IMovieScenePlayer.h
@@ -36,9 +36,10 @@ public:
 	/**
 	 * Updates the perspective viewports with the actor to view through
 	 *
-	 * @param ObjectToViewThrough The object, probably a camera, that the viewports should lock to
+	 * @param ObjectToViewThrough	The object, probably a camera, that the viewports should lock to
+	 * @param bNewCameraCut			If true this is a new camera cut
 	 */
-	virtual void UpdatePreviewViewports(UObject* ObjectToViewThrough) const = 0;
+	virtual void UpdateCameraCut(UObject* ObjectToViewThrough, bool bNewCameraCut) const = 0;
 
 	/**
 	 * Adds a MovieScene instance to the player.  MovieScene players need to know about each instance for actor spawning
diff --git a/Engine/Source/Runtime/MovieScene/Public/MovieScene.h b/Engine/Source/Runtime/MovieScene/Public/MovieScene.h
--- a/Engine/Source/Runtime/MovieScene/Public/MovieScene.h
+++ b/Engine/Source/Runtime/MovieScene/Public/MovieScene.h
@@ -377,6 +377,19 @@ public:
 	 */
 	UMovieSceneTrack* AddMasterTrack(TSubclassOf<UMovieSceneTrack> TrackClass);
 	
+	/**
+	 * Adds a new shot track if it doesn't exist 
+	 * A shot track is a special kind of sub-movie scene track that allows for cutting between camera views
+	 * There is only one per movie scene. 
+	 *
+	 * @param TrackClass  The shot track class type
+	 * @return The created shot track
+	 */
+	UMovieSceneTrack* AddShotTrack( TSubclassOf<UMovieSceneTrack> TrackClass );
+	
+	/** @return The shot track if it exists */
+	UMovieSceneTrack* GetShotTrack();
+
 	/**
 	 * Removes a master track.
 	 *
@@ -469,6 +482,10 @@ private:
 	UPROPERTY()
 	TArray<UMovieSceneTrack*> MasterTracks;
 
+	/** The shot track is a specialized track for switching between cameras on a cinematic */
+	UPROPERTY()
+	UMovieSceneTrack* ShotTrack;
+
 #if WITH_EDITORONLY_DATA
 	/** Editor only data that needs to be saved between sessions for editing but has no runtime purpose */
 	UPROPERTY()
diff --git a/Engine/Source/Runtime/MovieScene/Public/MovieSceneInstance.h b/Engine/Source/Runtime/MovieScene/Public/MovieSceneInstance.h
--- a/Engine/Source/Runtime/MovieScene/Public/MovieSceneInstance.h
+++ b/Engine/Source/Runtime/MovieScene/Public/MovieSceneInstance.h
@@ -52,6 +52,8 @@ private:
 
 	/** MovieScene that is instanced */
 	TWeakObjectPtr<UMovieScene> MovieScene;
+	/** The shot track instance map */
+	TSharedPtr<class IMovieSceneTrackInstance> ShotTrackInstance;
 	/** All Master track instances */
 	FMovieSceneInstanceMap MasterTrackInstances;
 	/** All object binding instances */
diff --git a/Engine/Source/Runtime/MovieScene/Public/MovieScenePropertyTrack.h b/Engine/Source/Runtime/MovieScene/Public/MovieScenePropertyTrack.h
--- a/Engine/Source/Runtime/MovieScene/Public/MovieScenePropertyTrack.h
+++ b/Engine/Source/Runtime/MovieScene/Public/MovieScenePropertyTrack.h
@@ -21,7 +21,7 @@ public:
 	virtual void RemoveSection( UMovieSceneSection* Section ) override;
 	virtual bool IsEmpty() const override;
 	virtual TRange<float> GetSectionBoundaries() const override;
-	virtual TArray<UMovieSceneSection*> GetAllSections() const override;
+	virtual const TArray<UMovieSceneSection*>& GetAllSections() const override;
 	virtual bool HasShowableData() const override {return bSectionsAreShowable;}
 	
 	/**
diff --git a/Engine/Source/Runtime/MovieScene/Public/MovieSceneTrack.h b/Engine/Source/Runtime/MovieScene/Public/MovieSceneTrack.h
--- a/Engine/Source/Runtime/MovieScene/Public/MovieSceneTrack.h
+++ b/Engine/Source/Runtime/MovieScene/Public/MovieSceneTrack.h
@@ -25,7 +25,7 @@ public:
 	 * 
 	 * @return List of all the sections in the track
 	 */
-	virtual TArray<UMovieSceneSection*> GetAllSections() const PURE_VIRTUAL( UMovieSceneTrack::GetAllSections, return TArray<UMovieSceneSection*>(); );
+	virtual const TArray<UMovieSceneSection*>& GetAllSections() const PURE_VIRTUAL( UMovieSceneTrack::GetAllSections, static TArray<UMovieSceneSection*> Empty; return Empty; );
 
 	/**
 	 * Creates a runtime instance of this class
@@ -82,4 +82,13 @@ public:
 	 * @return The range of time boundaries
 	 */
 	virtual TRange<float> GetSectionBoundaries() const PURE_VIRTUAL( UMovieSceneTrack::GetSectionBoundaries, return TRange<float>::Empty(); );
+
+#if WITH_EDITOR
+	/**
+	 * Called if the section is moved in Sequencer
+	 *
+	 * @param Section	The section that moved
+	 */
+	virtual void OnSectionMoved( UMovieSceneSection& Section ) {}
+#endif
 };
diff --git a/Engine/Source/Runtime/MovieScene/Public/RuntimeMovieScenePlayer.h b/Engine/Source/Runtime/MovieScene/Public/RuntimeMovieScenePlayer.h
--- a/Engine/Source/Runtime/MovieScene/Public/RuntimeMovieScenePlayer.h
+++ b/Engine/Source/Runtime/MovieScene/Public/RuntimeMovieScenePlayer.h
@@ -58,7 +58,7 @@ protected:
 	virtual void SpawnActorsForMovie( TSharedRef<FMovieSceneInstance> MovieSceneInstance );
 	virtual void DestroyActorsForMovie( TSharedRef<FMovieSceneInstance> MovieSceneInstance );
 	virtual void GetRuntimeObjects( TSharedRef<FMovieSceneInstance> MovieSceneInstance, const FGuid& ObjectHandle, TArray< UObject* >& OutObjects ) const override;
-	virtual void UpdatePreviewViewports(UObject* ObjectToViewThrough) const override {}
+	virtual void UpdateCameraCut(UObject* ObjectToViewThrough, bool bNewCameraCut) const override;
 	virtual EMovieScenePlayerStatus::Type GetPlaybackStatus() const override;
 	virtual void AddMovieSceneInstance( class UMovieSceneSection& MovieSceneSection, TSharedRef<FMovieSceneInstance> InstanceToAdd ) override;
 	virtual void RemoveMovieSceneInstance( class UMovieSceneSection& MovieSceneSection, TSharedRef<FMovieSceneInstance> InstanceToRemove ) override;
diff --git a/Engine/Source/Runtime/MovieSceneTracks/Private/Sections/MovieSceneShotSection.cpp b/Engine/Source/Runtime/MovieSceneTracks/Private/Sections/MovieSceneShotSection.cpp
--- a/Engine/Source/Runtime/MovieSceneTracks/Private/Sections/MovieSceneShotSection.cpp
+++ b/Engine/Source/Runtime/MovieSceneTracks/Private/Sections/MovieSceneShotSection.cpp
@@ -18,12 +18,9 @@ FGuid UMovieSceneShotSection::GetCameraGuid() const
 	return CameraGuid;
 }
 
-void UMovieSceneShotSection::SetTitle(const FText& InTitle)
+void UMovieSceneShotSection::SetShotNameAndNumber(const FText& InDisplayName, int32 InShotNumber)
 {
-	Title = InTitle;
+	DisplayName = InDisplayName;
+	ShotNumber = InShotNumber;
 }
 
-FText UMovieSceneShotSection::GetTitle() const
-{
-	return Title;
-}
diff --git a/Engine/Source/Runtime/MovieSceneTracks/Private/TrackInstances/MovieSceneShotTrackInstance.cpp b/Engine/Source/Runtime/MovieSceneTracks/Private/TrackInstances/MovieSceneShotTrackInstance.cpp
--- a/Engine/Source/Runtime/MovieSceneTracks/Private/TrackInstances/MovieSceneShotTrackInstance.cpp
+++ b/Engine/Source/Runtime/MovieSceneTracks/Private/TrackInstances/MovieSceneShotTrackInstance.cpp
@@ -7,41 +7,57 @@
 #include "IMovieScenePlayer.h"
 
 
-FMovieSceneShotTrackInstance::FMovieSceneShotTrackInstance( UMovieSceneShotTrack& InDirectorTrack )
+FMovieSceneShotTrackInstance::FMovieSceneShotTrackInstance( UMovieSceneShotTrack& InShotTrack )
+	: FSubMovieSceneTrackInstance( InShotTrack )
 {
-	DirectorTrack = &InDirectorTrack;
 }
 
+void FMovieSceneShotTrackInstance::RefreshInstance( const TArray<UObject*>& RuntimeObjects, IMovieScenePlayer& Player )
+{
+	const TArray<UMovieSceneSection*>& ShotSections = SubMovieSceneTrack->GetAllSections();
+
+	RuntimeCameraObjects.Empty( ShotSections.Num() );
+
+	for( UMovieSceneSection* Section : ShotSections )
+	{
+		// @todo Sequencer - Sub-moviescenes: Get the cameras from the root movie scene instance.  We should support adding cameras for sub-moviescenes as shots
+		TArray<UObject*> CameraObjects;
+		Player.GetRuntimeObjects(Player.GetRootMovieSceneInstance(), CastChecked<UMovieSceneShotSection>( Section )->GetCameraGuid(), CameraObjects );
+		if( CameraObjects.Num() == 1 )
+		{
+			RuntimeCameraObjects.Add( CameraObjects[0] );
+		}
+		else
+		{
+			// No valid camera object was found, take up space.  There should be one entry per section
+			RuntimeCameraObjects.Add( nullptr );
+		}
+	}
+
+	FSubMovieSceneTrackInstance::RefreshInstance( RuntimeObjects, Player );
+}
 
 void FMovieSceneShotTrackInstance::Update( float Position, float LastPosition, const TArray<UObject*>& RuntimeObjects, class IMovieScenePlayer& Player ) 
 {
-	// possess the 'first' shot section in order of track index
-	ACameraActor* PossessedCamera = NULL;
-	UMovieSceneShotSection* FirstShotSection = NULL;
-
-	const TArray<UMovieSceneSection*>& ShotSections = DirectorTrack->GetAllSections();
+	const TArray<UMovieSceneSection*>& ShotSections = SubMovieSceneTrack->GetAllSections();
 
 	for (int32 ShotIndex = 0; ShotIndex < ShotSections.Num(); ++ShotIndex)
 	{
 		UMovieSceneShotSection* ShotSection = CastChecked<UMovieSceneShotSection>(ShotSections[ShotIndex]);
 
-		if (ShotSection->IsTimeWithinSection(Position) &&
-			(FirstShotSection == NULL || FirstShotSection->GetRowIndex() > ShotSection->GetRowIndex()))
+		// Note shot times are not inclusive
+		// @todo Sequencer: This could be faster with a binary search
+		if( ShotSection->GetStartTime() <= Position && ShotSection->GetEndTime() > Position )
 		{
-			TArray<UObject*> CameraObjects;
-			// @todo Sequencer - Sub-moviescenes: Get the cameras from the root movie scene instance.  We should support adding cameras for sub-moviescenes as shots
-			Player.GetRuntimeObjects( Player.GetRootMovieSceneInstance(), ShotSection->GetCameraGuid(), CameraObjects);
-
-			if (CameraObjects.Num() > 0)
-			{
-				check(CameraObjects.Num() == 1 && CameraObjects[0]->IsA(ACameraActor::StaticClass()));
+			UObject* Camera = RuntimeCameraObjects[ShotIndex].Get();
+			
+			const bool bNewCameraCut = CurrentCameraObject != Camera;
+			Player.UpdateCameraCut(Camera, bNewCameraCut);
 
-				PossessedCamera = Cast<ACameraActor>(CameraObjects[0]);
+			CurrentCameraObject = Camera;
 
-				FirstShotSection = ShotSection;
-			}
+			// no need to process any more shots.  Only one shot can be active at a time
+			break;
 		}
 	}
-
-	Player.UpdatePreviewViewports(PossessedCamera);
 }
diff --git a/Engine/Source/Runtime/MovieSceneTracks/Private/TrackInstances/SubMovieSceneTrackInstance.cpp b/Engine/Source/Runtime/MovieSceneTracks/Private/TrackInstances/SubMovieSceneTrackInstance.cpp
--- a/Engine/Source/Runtime/MovieSceneTracks/Private/TrackInstances/SubMovieSceneTrackInstance.cpp
+++ b/Engine/Source/Runtime/MovieSceneTracks/Private/TrackInstances/SubMovieSceneTrackInstance.cpp
@@ -16,7 +16,7 @@ FSubMovieSceneTrackInstance::FSubMovieSceneTrackInstance( USubMovieSceneTrack& I
 
 void FSubMovieSceneTrackInstance::RefreshInstance( const TArray<UObject*>& RuntimeObjects, class IMovieScenePlayer& Player )
 {
-	TArray<UMovieSceneSection*> AllSections = SubMovieSceneTrack->GetAllSections();
+	const TArray<UMovieSceneSection*>& AllSections = SubMovieSceneTrack->GetAllSections();
 
 	// Sections we encountered
 	TSet<UMovieSceneSection*> FoundSections;
@@ -26,20 +26,25 @@ void FSubMovieSceneTrackInstance::RefreshInstance( const TArray<UObject*>& Runti
 	{
 		USubMovieSceneSection* Section = CastChecked<USubMovieSceneSection>( AllSections[SectionIndex] );
 
-		FoundSections.Add( Section );
-
-		TSharedPtr<FMovieSceneInstance> Instance = SubMovieSceneInstances.FindRef(Section);
-		if( !Instance.IsValid() )
+		// If the section doesn't have a valid movie scene or no longer has one 
+		// (e.g user deleted it) then skip adding an instance for it
+		if( Section->GetMovieScene() )
 		{
-			Instance = MakeShareable( new FMovieSceneInstance( *Section->GetMovieScene() ) );
+			FoundSections.Add(Section);
 
-			Player.AddMovieSceneInstance( *Section, Instance.ToSharedRef() );
+			TSharedPtr<FMovieSceneInstance> Instance = SubMovieSceneInstances.FindRef(Section);
+			if (!Instance.IsValid())
+			{
+				Instance = MakeShareable(new FMovieSceneInstance(*Section->GetMovieScene()));
 
-			SubMovieSceneInstances.Add( Section, Instance.ToSharedRef() );
-		}
+				Player.AddMovieSceneInstance(*Section, Instance.ToSharedRef());
 
-		// Refresh the existing instance
-		Instance->RefreshInstance( Player );
+				SubMovieSceneInstances.Add(Section, Instance.ToSharedRef());
+			}
+
+			// Refresh the existing instance
+			Instance->RefreshInstance(Player);
+		}
 	}
 
 	TMap< TWeakObjectPtr<USubMovieSceneSection>, TSharedPtr<FMovieSceneInstance> >::TIterator It =  SubMovieSceneInstances.CreateIterator();
@@ -57,7 +62,7 @@ void FSubMovieSceneTrackInstance::RefreshInstance( const TArray<UObject*>& Runti
 
 void FSubMovieSceneTrackInstance::Update( float Position, float LastPosition, const TArray<UObject*>& RuntimeObjects, class IMovieScenePlayer& Player ) 
 {
-	TArray<UMovieSceneSection*> AllSections = SubMovieSceneTrack->GetAllSections();
+	const TArray<UMovieSceneSection*>& AllSections = SubMovieSceneTrack->GetAllSections();
 
 	TArray<UMovieSceneSection*> TraversedSections = MovieSceneHelpers::GetTraversedSections( AllSections, Position, LastPosition );
 
@@ -85,7 +90,7 @@ void FSubMovieSceneTrackInstance::Update( float Position, float LastPosition, co
 
 void FSubMovieSceneTrackInstance::SaveState(const TArray<UObject*>& RuntimeObjects)
 {
-	TArray<UMovieSceneSection*> AllSections = SubMovieSceneTrack->GetAllSections();
+	const TArray<UMovieSceneSection*>& AllSections = SubMovieSceneTrack->GetAllSections();
 
 	for( int32 SectionIndex = 0; SectionIndex < AllSections.Num(); ++SectionIndex )
 	{
@@ -101,7 +106,7 @@ void FSubMovieSceneTrackInstance::SaveState(const TArray<UObject*>& RuntimeObjec
 
 void FSubMovieSceneTrackInstance::RestoreState(const TArray<UObject*>& RuntimeObjects)
 {
-	TArray<UMovieSceneSection*> AllSections = SubMovieSceneTrack->GetAllSections();
+	const TArray<UMovieSceneSection*>& AllSections = SubMovieSceneTrack->GetAllSections();
 
 	for( int32 SectionIndex = 0; SectionIndex < AllSections.Num(); ++SectionIndex )
 	{
diff --git a/Engine/Source/Runtime/MovieSceneTracks/Private/Tracks/MovieSceneAnimationTrack.cpp b/Engine/Source/Runtime/MovieSceneTracks/Private/Tracks/MovieSceneAnimationTrack.cpp
--- a/Engine/Source/Runtime/MovieSceneTracks/Private/Tracks/MovieSceneAnimationTrack.cpp
+++ b/Engine/Source/Runtime/MovieSceneTracks/Private/Tracks/MovieSceneAnimationTrack.cpp
@@ -24,14 +24,9 @@ TSharedPtr<IMovieSceneTrackInstance> UMovieSceneAnimationTrack::CreateInstance()
 	return MakeShareable( new FMovieSceneAnimationTrackInstance( *this ) ); 
 }
 
-TArray<UMovieSceneSection*> UMovieSceneAnimationTrack::GetAllSections() const
+const TArray<UMovieSceneSection*>& UMovieSceneAnimationTrack::GetAllSections() const
 {
-	TArray<UMovieSceneSection*> OutSections;
-	for (int32 SectionIndex = 0; SectionIndex < AnimationSections.Num(); ++SectionIndex)
-	{
-		OutSections.Add( AnimationSections[SectionIndex] );
-	}
-	return OutSections;
+	return AnimationSections;
 }
 
 void UMovieSceneAnimationTrack::RemoveAllAnimationData()
diff --git a/Engine/Source/Runtime/MovieSceneTracks/Private/Tracks/MovieSceneAudioTrack.cpp b/Engine/Source/Runtime/MovieSceneTracks/Private/Tracks/MovieSceneAudioTrack.cpp
--- a/Engine/Source/Runtime/MovieSceneTracks/Private/Tracks/MovieSceneAudioTrack.cpp
+++ b/Engine/Source/Runtime/MovieSceneTracks/Private/Tracks/MovieSceneAudioTrack.cpp
@@ -27,14 +27,9 @@ TSharedPtr<IMovieSceneTrackInstance> UMovieSceneAudioTrack::CreateInstance()
 	return MakeShareable( new FMovieSceneAudioTrackInstance( *this ) ); 
 }
 
-TArray<UMovieSceneSection*> UMovieSceneAudioTrack::GetAllSections() const
+const TArray<UMovieSceneSection*>& UMovieSceneAudioTrack::GetAllSections() const
 {
-	TArray<UMovieSceneSection*> OutSections;
-	for (int32 SectionIndex = 0; SectionIndex < AudioSections.Num(); ++SectionIndex)
-	{
-		OutSections.Add( AudioSections[SectionIndex] );
-	}
-	return OutSections;
+	return AudioSections;
 }
 
 void UMovieSceneAudioTrack::RemoveAllAnimationData()
diff --git a/Engine/Source/Runtime/MovieSceneTracks/Private/Tracks/MovieSceneParticleTrack.cpp b/Engine/Source/Runtime/MovieSceneTracks/Private/Tracks/MovieSceneParticleTrack.cpp
--- a/Engine/Source/Runtime/MovieSceneTracks/Private/Tracks/MovieSceneParticleTrack.cpp
+++ b/Engine/Source/Runtime/MovieSceneTracks/Private/Tracks/MovieSceneParticleTrack.cpp
@@ -24,14 +24,9 @@ TSharedPtr<IMovieSceneTrackInstance> UMovieSceneParticleTrack::CreateInstance()
 	return MakeShareable( new FMovieSceneParticleTrackInstance( *this ) ); 
 }
 
-TArray<UMovieSceneSection*> UMovieSceneParticleTrack::GetAllSections() const
+const TArray<UMovieSceneSection*>& UMovieSceneParticleTrack::GetAllSections() const
 {
-	TArray<UMovieSceneSection*> OutSections;
-	for (int32 SectionIndex = 0; SectionIndex < ParticleSections.Num(); ++SectionIndex)
-	{
-		OutSections.Add( ParticleSections[SectionIndex] );
-	}
-	return OutSections;
+	return ParticleSections;
 }
 
 void UMovieSceneParticleTrack::RemoveAllAnimationData()
diff --git a/Engine/Source/Runtime/MovieSceneTracks/Private/Tracks/MovieSceneShotTrack.cpp b/Engine/Source/Runtime/MovieSceneTracks/Private/Tracks/MovieSceneShotTrack.cpp
--- a/Engine/Source/Runtime/MovieSceneTracks/Private/Tracks/MovieSceneShotTrack.cpp
+++ b/Engine/Source/Runtime/MovieSceneTracks/Private/Tracks/MovieSceneShotTrack.cpp
@@ -26,15 +26,47 @@ TSharedPtr<IMovieSceneTrackInstance> UMovieSceneShotTrack::CreateInstance()
 	return MakeShareable( new FMovieSceneShotTrackInstance( *this ) ); 
 }
 
-void UMovieSceneShotTrack::AddNewShot(FGuid CameraHandle, float Time)
+void UMovieSceneShotTrack::RemoveSection( UMovieSceneSection* Section )
 {
-	UMovieSceneShotSection* NewSection = NewObject<UMovieSceneShotSection>(this);
-	NewSection->InitialPlacement(SubMovieSceneSections, Time, Time + 1.f, SupportsMultipleRows()); //@todo find a better default end time
-	NewSection->SetCameraGuid(CameraHandle);
-	NewSection->SetTitle(LOCTEXT("AddNewShot", "New Shot"));
+	Super::RemoveSection( Section );
 
-	SubMovieSceneSections.Add(NewSection);
+	// @todo Sequencer: The movie scene owned by the section is now abandoned.  Should we offer to delete it?  
 }
 
+void UMovieSceneShotTrack::AddNewShot(FGuid CameraHandle, UMovieScene& ShotMovieScene, const TRange<float>& TimeRange, const FText& ShotName, int32 ShotNumber )
+{
+	Modify();
+
+	FName UniqueShotName = MakeUniqueObjectName( this, UMovieSceneShotSection::StaticClass(), *ShotName.ToString() );
+
+	UMovieSceneShotSection* NewSection = NewObject<UMovieSceneShotSection>( this, UniqueShotName, RF_Transactional );
+	NewSection->SetMovieScene( &ShotMovieScene );
+	NewSection->SetStartTime( TimeRange.GetLowerBoundValue() );
+	NewSection->SetEndTime( TimeRange.GetUpperBoundValue() );
+	NewSection->SetCameraGuid( CameraHandle );
+	NewSection->SetShotNameAndNumber( ShotName , ShotNumber );
+
+	SubMovieSceneSections.Add( NewSection );
+	
+	// When a new shot is added, sort all shots to ensure they are in the correct order
+	SortShots();
+}
+
+#if WITH_EDITOR
+void UMovieSceneShotTrack::OnSectionMoved( UMovieSceneSection& Section )
+{
+	// Sort all shots when one moves.  They should be in order of start time
+	SortShots();
+}
+#endif
+
+void UMovieSceneShotTrack::SortShots()
+{
+	SubMovieSceneSections.Sort(
+		[](const UMovieSceneSection& A, const UMovieSceneSection& B)
+	{
+		return A.GetStartTime() < B.GetStartTime();
+	});
+}
 
 #undef LOCTEXT_NAMESPACE
\ No newline at end of file
diff --git a/Engine/Source/Runtime/MovieSceneTracks/Private/Tracks/SubMovieSceneTrack.cpp b/Engine/Source/Runtime/MovieSceneTracks/Private/Tracks/SubMovieSceneTrack.cpp
--- a/Engine/Source/Runtime/MovieSceneTracks/Private/Tracks/SubMovieSceneTrack.cpp
+++ b/Engine/Source/Runtime/MovieSceneTracks/Private/Tracks/SubMovieSceneTrack.cpp
@@ -34,15 +34,9 @@ TSharedPtr<IMovieSceneTrackInstance> USubMovieSceneTrack::CreateInstance()
 	return MakeShareable( new FSubMovieSceneTrackInstance( *this ) ); 
 }
 
-TArray<UMovieSceneSection*> USubMovieSceneTrack::GetAllSections() const 
+const TArray<UMovieSceneSection*>& USubMovieSceneTrack::GetAllSections() const 
 {
-	TArray<UMovieSceneSection*> AllSections;
-	for( int32 Section = 0; Section < SubMovieSceneSections.Num(); ++Section )
-	{
-		AllSections.Add( SubMovieSceneSections[Section] );
-	}
-
-	return AllSections;
+	return SubMovieSceneSections;
 }
 
 void USubMovieSceneTrack::RemoveSection( UMovieSceneSection* Section )
diff --git a/Engine/Source/Runtime/MovieSceneTracks/Public/Sections/MovieSceneShotSection.h b/Engine/Source/Runtime/MovieSceneTracks/Public/Sections/MovieSceneShotSection.h
--- a/Engine/Source/Runtime/MovieSceneTracks/Public/Sections/MovieSceneShotSection.h
+++ b/Engine/Source/Runtime/MovieSceneTracks/Public/Sections/MovieSceneShotSection.h
@@ -20,18 +20,29 @@ public:
 	/** Gets the camera guid for this shot section */
 	virtual FGuid GetCameraGuid() const;
 
-	/** Sets the shot title */
-	virtual void SetTitle(const FText& NewTitle);
-
-	/** Gets the shot title */
-	virtual FText GetTitle() const;
-
+	/** 
+	 * Sets the shot display name and shot number
+	 * 
+ 	 * @param InDisplayName	The display name shown next to each shot in Sequencer
+	 * @param InShotNumber	The generated shot number.
+	 */
+	virtual void SetShotNameAndNumber(const FText& InDisplayName, int32 InShotNumber);
+
+	/** @return The display name for the shot */
+	FText GetShotDisplayName() const { return DisplayName; }
+
+	/** @return The shot number generated for this shot */
+	int32 GetShotNumber() const { return ShotNumber; }
 private:
 	/** The camera possessable or spawnable that this movie shot uses */
 	UPROPERTY()
 	FGuid CameraGuid;
 
-	/** The shot's title */
+	/** The shot's display name */
+	UPROPERTY()
+	FText DisplayName;
+	
+	 /** The Shot number.  Shot numbers are generally part of the display name but is stored seperatley so we can use it to generate subsequent shot numbers */
 	UPROPERTY()
-	FText Title;
+	int32 ShotNumber;
 };
diff --git a/Engine/Source/Runtime/MovieSceneTracks/Public/Sections/SubMovieSceneSection.h b/Engine/Source/Runtime/MovieSceneTracks/Public/Sections/SubMovieSceneSection.h
--- a/Engine/Source/Runtime/MovieSceneTracks/Public/Sections/SubMovieSceneSection.h
+++ b/Engine/Source/Runtime/MovieSceneTracks/Public/Sections/SubMovieSceneSection.h
@@ -23,6 +23,7 @@ public:
 	UMovieScene* GetMovieScene() const { return MovieScene; }
 private:
 	/** Movie scene being played by this section */
+	/** @todo Sequencer: Should this be lazy loaded? */
 	UPROPERTY()
 	UMovieScene* MovieScene;
 };
diff --git a/Engine/Source/Runtime/MovieSceneTracks/Public/TrackInstances/MovieSceneShotTrackInstance.h b/Engine/Source/Runtime/MovieSceneTracks/Public/TrackInstances/MovieSceneShotTrackInstance.h
--- a/Engine/Source/Runtime/MovieSceneTracks/Public/TrackInstances/MovieSceneShotTrackInstance.h
+++ b/Engine/Source/Runtime/MovieSceneTracks/Public/TrackInstances/MovieSceneShotTrackInstance.h
@@ -2,8 +2,7 @@
 
 #pragma once
 
-#include "IMovieSceneTrackInstance.h"
-
+#include "SubMovieSceneTrackInstance.h"
 
 class UMovieSceneShotTrack;
 
@@ -12,18 +11,20 @@ class UMovieSceneShotTrack;
  * Instance of a UMovieSceneShotTrack
  */
 class FMovieSceneShotTrackInstance
-	: public IMovieSceneTrackInstance
+	: public FSubMovieSceneTrackInstance
 {
 public:
 
-	FMovieSceneShotTrackInstance( UMovieSceneShotTrack& InDirectorTrack );
+	FMovieSceneShotTrackInstance( UMovieSceneShotTrack& InShotTrack );
 
 	/** IMovieSceneTrackInstance interface */
+	virtual void RefreshInstance( const TArray<UObject*>& RuntimeObjects, IMovieScenePlayer& Player ) override;
 	virtual void Update( float Position, float LastPosition, const TArray<UObject*>& RuntimeObjects, IMovieScenePlayer& Player ) override;
-	virtual void RefreshInstance( const TArray<UObject*>& RuntimeObjects, IMovieScenePlayer& Player  ) override {}
 
 private:
+	/** Runtime camera objects.  One for each shot.  Must be the same number of entries as sections */
+	TArray< TWeakObjectPtr<UObject> > RuntimeCameraObjects;
 
-	/** Track that is being instanced */
-	UMovieSceneShotTrack* DirectorTrack;
+	/** Current camera object we are looking through.  Used to determine when making a new cut */
+	TWeakObjectPtr<UObject> CurrentCameraObject;
 };
diff --git a/Engine/Source/Runtime/MovieSceneTracks/Public/TrackInstances/SubMovieSceneTrackInstance.h b/Engine/Source/Runtime/MovieSceneTracks/Public/TrackInstances/SubMovieSceneTrackInstance.h
--- a/Engine/Source/Runtime/MovieSceneTracks/Public/TrackInstances/SubMovieSceneTrackInstance.h
+++ b/Engine/Source/Runtime/MovieSceneTracks/Public/TrackInstances/SubMovieSceneTrackInstance.h
@@ -25,7 +25,7 @@ public:
 	virtual void RefreshInstance( const TArray<UObject*>& RuntimeObjects, IMovieScenePlayer& Player ) override;
 	virtual void SaveState (const TArray<UObject*>& RuntimeObjects) override;
 	virtual void RestoreState (const TArray<UObject*>& RuntimeObjects) override;
-private:
+protected:
 
 	/** Track that is being instanced */
 	TWeakObjectPtr<USubMovieSceneTrack> SubMovieSceneTrack;
diff --git a/Engine/Source/Runtime/MovieSceneTracks/Public/Tracks/MovieSceneAnimationTrack.h b/Engine/Source/Runtime/MovieSceneTracks/Public/Tracks/MovieSceneAnimationTrack.h
--- a/Engine/Source/Runtime/MovieSceneTracks/Public/Tracks/MovieSceneAnimationTrack.h
+++ b/Engine/Source/Runtime/MovieSceneTracks/Public/Tracks/MovieSceneAnimationTrack.h
@@ -25,7 +25,7 @@ public:
 	virtual bool IsEmpty() const override;
 	virtual TRange<float> GetSectionBoundaries() const override;
 	virtual bool SupportsMultipleRows() const override { return true; }
-	virtual TArray<UMovieSceneSection*> GetAllSections() const override;
+	virtual const TArray<UMovieSceneSection*>& GetAllSections() const override;
 
 	/** Adds a new animation to this track */
 	virtual void AddNewAnimation(float KeyTime, class UAnimSequence* AnimSequence);
diff --git a/Engine/Source/Runtime/MovieSceneTracks/Public/Tracks/MovieSceneAudioTrack.h b/Engine/Source/Runtime/MovieSceneTracks/Public/Tracks/MovieSceneAudioTrack.h
--- a/Engine/Source/Runtime/MovieSceneTracks/Public/Tracks/MovieSceneAudioTrack.h
+++ b/Engine/Source/Runtime/MovieSceneTracks/Public/Tracks/MovieSceneAudioTrack.h
@@ -30,7 +30,7 @@ public:
 	virtual bool IsEmpty() const override;
 	virtual TRange<float> GetSectionBoundaries() const override;
 	virtual bool SupportsMultipleRows() const override { return true; }
-	virtual TArray<UMovieSceneSection*> GetAllSections() const override;
+	virtual const TArray<UMovieSceneSection*>& GetAllSections() const override;
 
 	/** Adds a new sound cue to the audio */
 	virtual void AddNewSound(class USoundBase* Sound, float Time);
diff --git a/Engine/Source/Runtime/MovieSceneTracks/Public/Tracks/MovieSceneParticleTrack.h b/Engine/Source/Runtime/MovieSceneTracks/Public/Tracks/MovieSceneParticleTrack.h
--- a/Engine/Source/Runtime/MovieSceneTracks/Public/Tracks/MovieSceneParticleTrack.h
+++ b/Engine/Source/Runtime/MovieSceneTracks/Public/Tracks/MovieSceneParticleTrack.h
@@ -25,7 +25,7 @@ public:
 	virtual bool IsEmpty() const override;
 	virtual TRange<float> GetSectionBoundaries() const override;
 	virtual bool SupportsMultipleRows() const override { return false; }
-	virtual TArray<UMovieSceneSection*> GetAllSections() const override;
+	virtual const TArray<UMovieSceneSection*>& GetAllSections() const override;
 	
 	virtual void AddNewParticleSystem(float KeyTime, bool bTrigger);
 
diff --git a/Engine/Source/Runtime/MovieSceneTracks/Public/Tracks/MovieSceneShotTrack.h b/Engine/Source/Runtime/MovieSceneTracks/Public/Tracks/MovieSceneShotTrack.h
--- a/Engine/Source/Runtime/MovieSceneTracks/Public/Tracks/MovieSceneShotTrack.h
+++ b/Engine/Source/Runtime/MovieSceneTracks/Public/Tracks/MovieSceneShotTrack.h
@@ -19,7 +19,24 @@ public:
 	/** UMovieSceneTrack interface */
 	virtual FName GetTrackName() const override;
 	virtual TSharedPtr<IMovieSceneTrackInstance> CreateInstance() override;
+	virtual void RemoveSection( UMovieSceneSection* Section ) override;
 
-	/** Adds a new shot at the specified time */
-	MOVIESCENETRACKS_API void AddNewShot(FGuid CameraHandle, float Time);
+	/** 
+	 * Adds a new shot at the specified time
+	 *
+	 * @param CameraHandle		Handle to the camera that the shot switches to when active
+	 * @param ShotMovieScene	MovieScene for the shot (each shot has a unique movie scene containing per tracks only active during the shot)
+	 * @param TimeRange		The range within this track's movie scene where the shot is initially placed
+	 * @param ShotName		The display name of the shot
+	 * @param ShotNumber		The number of the shot.  This is used to assist with auto-generated shot names when new shots are added
+	 */
+	MOVIESCENETRACKS_API void AddNewShot(FGuid CameraHandle, UMovieScene& ShotMovieScene, const TRange<float>& TimeRange, const FText& ShotName, int32 ShotNumber );
+
+#if WITH_EDITOR
+	virtual void OnSectionMoved( UMovieSceneSection& Section ) override;
+#endif
+
+private:
+	/** Sorts shots according to their start time */
+	void SortShots();
 };
diff --git a/Engine/Source/Runtime/MovieSceneTracks/Public/Tracks/SubMovieSceneTrack.h b/Engine/Source/Runtime/MovieSceneTracks/Public/Tracks/SubMovieSceneTrack.h
--- a/Engine/Source/Runtime/MovieSceneTracks/Public/Tracks/SubMovieSceneTrack.h
+++ b/Engine/Source/Runtime/MovieSceneTracks/Public/Tracks/SubMovieSceneTrack.h
@@ -17,7 +17,7 @@ public:
 	virtual UMovieSceneSection* CreateNewSection() override;
 	virtual FName GetTrackName() const override;
 	virtual TSharedPtr<IMovieSceneTrackInstance> CreateInstance() override;
-	virtual TArray<UMovieSceneSection*> GetAllSections() const override;
+	virtual const TArray<UMovieSceneSection*>& GetAllSections() const override;
 	virtual void RemoveSection( UMovieSceneSection* Section ) override;
 	virtual void RemoveAllAnimationData() override;
 	virtual bool IsEmpty() const override;
diff --git a/Engine/Source/Runtime/UMG/Public/Animation/UMGSequencePlayer.h b/Engine/Source/Runtime/UMG/Public/Animation/UMGSequencePlayer.h
--- a/Engine/Source/Runtime/UMG/Public/Animation/UMGSequencePlayer.h
+++ b/Engine/Source/Runtime/UMG/Public/Animation/UMGSequencePlayer.h
@@ -36,7 +36,7 @@ public:
 
 	/** IMovieScenePlayer interface */
 	virtual void GetRuntimeObjects( TSharedRef<FMovieSceneInstance> MovieSceneInstance, const FGuid& ObjectHandle, TArray< UObject* >& OutObjects ) const override;
-	virtual void UpdatePreviewViewports(UObject* ObjectToViewThrough) const override {}
+	virtual void UpdateCameraCut(UObject* ObjectToViewThrough, bool bNewCameraCut) const override {}
 	virtual void AddMovieSceneInstance( class UMovieSceneSection& MovieSceneSection, TSharedRef<FMovieSceneInstance> InstanceToAdd ) override {}
 	virtual void RemoveMovieSceneInstance( class UMovieSceneSection& MovieSceneSection, TSharedRef<FMovieSceneInstance> InstanceToRemove ) override {}
 	virtual TSharedRef<FMovieSceneInstance> GetRootMovieSceneInstance() const override { return RootMovieSceneInstance.ToSharedRef(); }
