Fixed a couple of Clang warnings and a compile error

diff --git a/Engine/Source/Editor/KismetCompiler/Private/KismetCompiler.cpp b/Engine/Source/Editor/KismetCompiler/Private/KismetCompiler.cpp
--- a/Engine/Source/Editor/KismetCompiler/Private/KismetCompiler.cpp
+++ b/Engine/Source/Editor/KismetCompiler/Private/KismetCompiler.cpp
@@ -318,8 +318,8 @@ void FKismetCompilerContext::ValidateLink(const UEdGraphPin* PinA, const UEdGrap
 	{
 		const UEdGraphPin* InputPin = (EEdGraphPinDirection::EGPD_Input == PinA->Direction) ? PinA : PinB;
 		const UEdGraphPin* OutputPin = (EEdGraphPinDirection::EGPD_Output == PinA->Direction) ? PinA : PinB;
-		const bool bForbiddenConnection = InputPin && OutputPin && (OutputPin->PinType.PinCategory == Schema->PC_Interface) && (InputPin->PinType.PinCategory == Schema->PC_Object);
-		if (bForbiddenConnection)
+		const bool bInvalidConnection = InputPin && OutputPin && (OutputPin->PinType.PinCategory == Schema->PC_Interface) && (InputPin->PinType.PinCategory == Schema->PC_Object);
+		if (bInvalidConnection)
 		{
 			MessageLog.Error(*LOCTEXT("PinTypeMismatch_Error", "Can't connect pins @@ (Interface) and @@ (Object). Use an explicit cast node.").ToString(), OutputPin, InputPin);
 		}
diff --git a/Engine/Source/Editor/MovieSceneTools/Private/TrackEditors/ShotTrackEditor.cpp b/Engine/Source/Editor/MovieSceneTools/Private/TrackEditors/ShotTrackEditor.cpp
--- a/Engine/Source/Editor/MovieSceneTools/Private/TrackEditors/ShotTrackEditor.cpp
+++ b/Engine/Source/Editor/MovieSceneTools/Private/TrackEditors/ShotTrackEditor.cpp
@@ -274,7 +274,7 @@ void FShotSection::FilterToSelectedShotSections(bool bZoomToShotBounds)
 	Sequencer.Pin()->FilterToSelectedShotSections(bZoomToShotBounds);
 }
 
-void FShotSection::RenameShotCommitted(const FText& RenameText, ETextCommit::Type CommitInfo, UMovieSceneSection* Section)
+void FShotSection::RenameShotCommitted(const FText& RenameText, ETextCommit::Type CommitInfo, UMovieSceneSection* InSection)
 {
 /*
 	if (CommitInfo == ETextCommit::OnEnter)
@@ -560,7 +560,7 @@ void FShotTrackEditor::AddKeyInternal( float KeyTime, const FGuid ObjectGuid )
 
 	const FString NewShotPackagePath = LongPackagePath + TEXT("/") + MovieScene->GetName() + TEXT("_") + ShotName;
 
-	const FString EmptySuffix;
+	const FString EmptySuffix(TEXT(""));
 
 	// Create a package name and asset name for the new movie scene
 	FString PackageName;
