Sequencer Shot changes
Added the ability to specify a widget for each section interface. The widget is placed inside the Section widget created for it.  Added an example of this for shot sections.  Each shot section has an inline editable text box to rename the shot.
Fixed shots not refreshing when a new map is loaded

diff --git a/Engine/Source/Editor/MovieSceneTools/Private/TrackEditors/ShotTrackEditor.cpp b/Engine/Source/Editor/MovieSceneTools/Private/TrackEditors/ShotTrackEditor.cpp
--- a/Engine/Source/Editor/MovieSceneTools/Private/TrackEditors/ShotTrackEditor.cpp
+++ b/Engine/Source/Editor/MovieSceneTools/Private/TrackEditors/ShotTrackEditor.cpp
@@ -18,8 +18,8 @@
 #include "ShotTrackEditor.h"
 #include "CommonMovieSceneTools.h"
 #include "Camera/CameraActor.h"
-#include "STextEntryPopup.h"
 #include "AssetToolsModule.h"
+#include "SInlineEditableTextBlock.h"
 
 namespace AnimatableShotToolConstants
 {
@@ -159,10 +159,9 @@ bool FShotThumbnail::IsValid() const
 /////////////////////////////////////////////////
 // FShotSection
 
-FShotSection::FShotSection( TSharedPtr<ISequencer> InSequencer, TSharedPtr<FShotThumbnailPool> InThumbnailPool, UMovieSceneSection& InSection, UObject* InTargetObject )
-	: Section( &InSection )
+FShotSection::FShotSection( TSharedPtr<ISequencer> InSequencer, TSharedPtr<FShotThumbnailPool> InThumbnailPool, UMovieSceneSection& InSection )
+	: Section( CastChecked<UMovieSceneShotSection>( &InSection ) )
 	, Sequencer( InSequencer )
-	, Camera(Cast<ACameraActor>(InTargetObject))
 	, ThumbnailPool( InThumbnailPool )
 	, Thumbnails()
 	, ThumbnailWidth(0)
@@ -172,6 +171,9 @@ FShotSection::FShotSection( TSharedPtr<ISequencer> InSequencer, TSharedPtr<FShot
 	, InternalViewportScene(nullptr)
 	, InternalViewportClient(nullptr)
 {
+
+	Camera = UpdateCameraObject();
+
 	if (Camera.IsValid())
 	{
 		// @todo Sequencer optimize This code shouldn't even be called if this is offscreen
@@ -208,9 +210,25 @@ UMovieSceneSection* FShotSection::GetSectionObject()
 	return Section;
 }
 
+TSharedRef<SWidget> FShotSection::GenerateSectionWidget()
+{
+	return 
+		SNew( SBox )
+		.HAlign( HAlign_Left )
+		.VAlign( VAlign_Top )
+		.Padding( FMargin( 15.0f, 7.0f, 0.0f, 0.0f ) )
+		[
+			SNew( SInlineEditableTextBlock )
+			.ToolTipText( NSLOCTEXT("FShotTrackEditor", "RenameShot", "The name of this shot.  Click or hit F2 to rename") )
+			.Text( this, &FShotSection::GetShotName )
+			.ShadowOffset( FVector2D(1,1) )
+			.OnTextCommitted(this, &FShotSection::OnRenameShot )
+		];
+}
+
 FText FShotSection::GetSectionTitle() const
 {
-	return Cast<UMovieSceneShotSection>(Section)->GetShotDisplayName();
+	return FText::GetEmpty();
 }
 
 float FShotSection::GetSectionHeight() const
@@ -222,9 +240,9 @@ FReply FShotSection::OnSectionDoubleClicked( const FGeometry& SectionGeometry, c
 {
 	if( MouseEvent.GetEffectingButton() == EKeys::LeftMouseButton )
 	{
-		check( Section->IsA<UMovieSceneSection>() );
+		TSharedRef<FMovieSceneInstance> MovieSceneInstance = Sequencer.Pin()->GetInstanceForSubMovieSceneSection( *Section );
 
-		Sequencer.Pin()->FilterToSelectedShotSections();
+		Sequencer.Pin()->FocusSubMovieScene( MovieSceneInstance );
 	}
 
 	return FReply::Handled();
@@ -232,63 +250,20 @@ FReply FShotSection::OnSectionDoubleClicked( const FGeometry& SectionGeometry, c
 
 void FShotSection::BuildSectionContextMenu(FMenuBuilder& MenuBuilder)
 {
-/*
 	MenuBuilder.AddMenuEntry(
-		NSLOCTEXT("Sequencer", "RenameShot", "Rename"),
-		NSLOCTEXT("Sequencer", "RenameShotToolTip", "Renames this shot."),
-		FSlateIcon(),
-		FUIAction(FExecuteAction::CreateSP(this, &FShotSection::RenameShot))
-		);
-*/
-
-	MenuBuilder.AddMenuEntry(
-		NSLOCTEXT("Sequencer", "FilterToShots", "Filter To Shots"),
-		NSLOCTEXT("Sequencer", "FilterToShotsToolTip", "Filters to the selected shot sections"),
+		NSLOCTEXT("FShotTrackEditor", "FilterToShots", "Filter To Shots"),
+		NSLOCTEXT("FShotTrackEditor", "FilterToShotsToolTip", "Filters to the selected shot sections"),
 		FSlateIcon(),
 		FUIAction(FExecuteAction::CreateSP(this, &FShotSection::FilterToSelectedShotSections, true))
 		);
 }
 
-void FShotSection::RenameShot()
-{
-	auto ActualShotSection = CastChecked<UMovieSceneShotSection>(Section);
-
-	/*TSharedRef<STextEntryPopup> TextEntry = 
-		SNew(STextEntryPopup)
-		.Label(NSLOCTEXT("Sequencer", "RenameShotHeader", "Name"))
-		.DefaultText( ActualShotSection->GetTitle() )
-		.OnTextCommitted(this, &FShotSection::RenameShotCommitted, Section)
-		.ClearKeyboardFocusOnCommit( false );
-	
-	NameEntryPopupMenu = FSlateApplication::Get().PushMenu(
-		Sequencer.Pin()->GetSequencerWidget(),
-		FWidgetPath(),
-		TextEntry,
-		FSlateApplication::Get().GetCursorPos(),
-		FPopupTransitionEffect( FPopupTransitionEffect::TypeInPopup )
-		);*/
-}
 
 void FShotSection::FilterToSelectedShotSections(bool bZoomToShotBounds)
 {
 	Sequencer.Pin()->FilterToSelectedShotSections(bZoomToShotBounds);
 }
 
-void FShotSection::RenameShotCommitted(const FText& RenameText, ETextCommit::Type CommitInfo, UMovieSceneSection* InSection)
-{
-/*
-	if (CommitInfo == ETextCommit::OnEnter)
-	{
-		auto ShotSection = CastChecked<UMovieSceneShotSection>(Section);
-		ShotSection->SetTitle(RenameText);
-	}
-
-	if (NameEntryPopupMenu.IsValid())
-	{
-		NameEntryPopupMenu.Pin()->Dismiss();
-	}*/
-}
-
 int32 FShotSection::OnPaintSection( const FGeometry& AllottedGeometry, const FSlateRect& SectionClippingRect, FSlateWindowElementList& OutDrawElements, int32 LayerId, bool bParentEnabled ) const
 {
 	if (Camera.IsValid())
@@ -357,7 +332,7 @@ void FShotSection::Tick( const FGeometry& AllottedGeometry, const FGeometry& Par
 {
 	if (!Camera.IsValid())
 	{
-		return;
+		Camera = UpdateCameraObject();
 	}
 
 	FTimeToPixel TimeToPixelConverter( AllottedGeometry, TRange<float>( Section->GetStartTime(), Section->GetEndTime() ) );
@@ -461,6 +436,53 @@ void FShotSection::CalculateThumbnailWidthAndResize()
 	}
 }
 
+ACameraActor* FShotSection::UpdateCameraObject()
+{
+	TArray<UObject*> OutObjects;
+	// @todo Sequencer - Sub-MovieScenes The director track may be able to get cameras from sub-movie scenes
+	Sequencer.Pin()->GetRuntimeObjects( Sequencer.Pin()->GetRootMovieSceneInstance(),  Section->GetCameraGuid(), OutObjects);
+
+	ACameraActor* Camera = nullptr;
+	if( OutObjects.Num() > 0 )
+	{
+		Camera = Cast<ACameraActor>( OutObjects[0] );
+	}
+
+	return Camera;
+}
+
+FText FShotSection::GetShotName() const
+{
+	return Section->GetShotDisplayName();
+}
+
+void FShotSection::OnRenameShot( const FText& NewShotName, ETextCommit::Type CommitType )
+{
+	if (CommitType == ETextCommit::OnEnter && !GetShotName().EqualTo( NewShotName ) )
+	{
+		Section->Modify();
+		Section->SetShotNameAndNumber( NewShotName, Section->GetShotNumber() );
+
+		UMovieScene* ShotMovieScene = Section->GetMovieScene();
+		if( ShotMovieScene )
+		{
+			FAssetToolsModule& AssetToolsModule = FModuleManager::GetModuleChecked<FAssetToolsModule>("AssetTools");
+
+			TArray<FAssetRenameData> AssetsAndNames;
+			const FString PackagePath = FPackageName::GetLongPackagePath(ShotMovieScene->GetOutermost()->GetName());
+
+			// We want to prepend the owning movie scene name.  This is not the movie scene inside referenced by the shot but the movie scene containing 
+			// the shot track
+			UMovieScene* OwningMovieScene = Section->GetTypedOuter<UMovieScene>();
+			const FString NewAssetName = OwningMovieScene->GetName() + TEXT("_") + NewShotName.ToString();
+
+			new (AssetsAndNames) FAssetRenameData(ShotMovieScene, PackagePath, NewAssetName );
+
+			AssetToolsModule.Get().RenameAssets(AssetsAndNames);
+		}
+	}
+}
+
 /////////////////////////////////////////////////
 // FShotTrackEditor
 
@@ -487,12 +509,8 @@ bool FShotTrackEditor::SupportsType( TSubclassOf<UMovieSceneTrack> Type ) const
 TSharedRef<ISequencerSection> FShotTrackEditor::MakeSectionInterface( UMovieSceneSection& SectionObject, UMovieSceneTrack* Track )
 {
 	check( SupportsType( SectionObject.GetOuter()->GetClass() ) );
-	
-	TArray<UObject*> OutObjects;
-	// @todo Sequencer - Sub-MovieScenes The director track may be able to get cameras from sub-movie scenes
-	GetSequencer()->GetRuntimeObjects( GetSequencer()->GetRootMovieSceneInstance(), Cast<const UMovieSceneShotSection>(&SectionObject)->GetCameraGuid(), OutObjects);
-	check(OutObjects.Num() <= 1);
-	TSharedRef<ISequencerSection> NewSection( new FShotSection( GetSequencer(), ThumbnailPool, SectionObject, OutObjects.Num() ? OutObjects[0] : NULL ) );
+
+	TSharedRef<ISequencerSection> NewSection( new FShotSection( GetSequencer(), ThumbnailPool, SectionObject ) );
 
 	return NewSection;
 }
@@ -524,8 +542,8 @@ void FShotTrackEditor::BuildObjectBindingContextMenu(FMenuBuilder& MenuBuilder,
 		const TSharedPtr<ISequencer> ParentSequencer = GetSequencer();
 
 		MenuBuilder.AddMenuEntry(
-			NSLOCTEXT("Sequencer", "AddShot", "Add New Shot"),
-			NSLOCTEXT("Sequencer", "AddShotTooltip", "Adds a new shot using this camera at the scrubber location."),
+			NSLOCTEXT("FShotTrackEditor", "AddShot", "Add New Shot"),
+			NSLOCTEXT("FShotTrackEditor", "AddShotTooltip", "Adds a new shot using this camera at the scrubber location."),
 			FSlateIcon(),
 			FUIAction(FExecuteAction::CreateSP(ParentSequencer.Get(), &ISequencer::AddNewShot, ObjectBinding))
 			);
diff --git a/Engine/Source/Editor/MovieSceneTools/Private/TrackEditors/ShotTrackEditor.h b/Engine/Source/Editor/MovieSceneTools/Private/TrackEditors/ShotTrackEditor.h
--- a/Engine/Source/Editor/MovieSceneTools/Private/TrackEditors/ShotTrackEditor.h
+++ b/Engine/Source/Editor/MovieSceneTools/Private/TrackEditors/ShotTrackEditor.h
@@ -129,11 +129,12 @@ private:
 class FShotSection : public ISequencerSection, public TSharedFromThis<FShotSection>
 {
 public:
-	FShotSection( TSharedPtr<ISequencer> InSequencer, TSharedPtr<FShotThumbnailPool> InThumbnailPool, UMovieSceneSection& InSection, UObject* InTargetObject );
+	FShotSection( TSharedPtr<ISequencer> InSequencer, TSharedPtr<FShotThumbnailPool> InThumbnailPool, UMovieSceneSection& InSection );
 	~FShotSection();
 
 	/** ISequencerSection interface */
 	virtual UMovieSceneSection* GetSectionObject() override;
+	virtual TSharedRef<SWidget> GenerateSectionWidget() override;
 	virtual int32 OnPaintSection( const FGeometry& AllottedGeometry, const FSlateRect& SectionClippingRect, FSlateWindowElementList& OutDrawElements, int32 LayerId, bool bParentEnabled ) const override;
 	virtual void Tick( const FGeometry& AllottedGeometry, const FGeometry& ParentGeometry, const double InCurrentTime, const float InDeltaTime ) override;
 	virtual FText GetDisplayName() const override { return NSLOCTEXT("FShotSection", "", "Shots"); }
@@ -143,12 +144,6 @@ public:
 	virtual FReply OnSectionDoubleClicked( const FGeometry& SectionGeometry, const FPointerEvent& MouseEvent ) override;
 	virtual void BuildSectionContextMenu(FMenuBuilder& MenuBuilder) override;
 
-	/** Opens a renaming dialog for the passed in shot section */
-	void RenameShot();
-	 
-	/** Called when committing a rename shot text entry popup */
-	void RenameShotCommitted(const FText& RenameText, ETextCommit::Type CommitInfo, UMovieSceneSection* Section);
-
 	/** Filter to selected shot sections */
 	void FilterToSelectedShotSections(bool bZoomToShotBounds);
 
@@ -169,9 +164,13 @@ public:
 
 	/** @return The sequencer widget owning the shot section */
 	TSharedRef<SWidget> GetSequencerWidget() { return Sequencer.Pin()->GetSequencerWidget(); }
+private:
+	ACameraActor* UpdateCameraObject();
+	FText GetShotName() const;
+	void OnRenameShot( const FText& NewShotName, ETextCommit::Type CommitType );
 private:
 	/** The section we are visualizing */
-	UMovieSceneSection* Section;
+	class UMovieSceneShotSection* Section;
 	/** The parent sequencer we are a part of */
 	TWeakPtr<ISequencer> Sequencer;
 	/** The actual camera actor we are possessing */
diff --git a/Engine/Source/Editor/Sequencer/Private/SSection.cpp b/Engine/Source/Editor/Sequencer/Private/SSection.cpp
--- a/Engine/Source/Editor/Sequencer/Private/SSection.cpp
+++ b/Engine/Source/Editor/Sequencer/Private/SSection.cpp
@@ -19,6 +19,11 @@ void SSection::Construct( const FArguments& InArgs, TSharedRef<FTrackNode> Secti
 	SectionIndex = InSectionIndex;
 	ParentSectionArea = SectionNode;
 	SectionInterface = SectionNode->GetSections()[InSectionIndex];
+
+	ChildSlot
+	[
+		SectionInterface->GenerateSectionWidget()
+	];
 }
 
 void SSection::GetKeyAreas( const TSharedPtr<FTrackNode>& SectionAreaNode, TArray<FKeyAreaElement>& OutKeyAreas ) const
diff --git a/Engine/Source/Editor/Sequencer/Public/ISequencerSection.h b/Engine/Source/Editor/Sequencer/Public/ISequencerSection.h
--- a/Engine/Source/Editor/Sequencer/Public/ISequencerSection.h
+++ b/Engine/Source/Editor/Sequencer/Public/ISequencerSection.h
@@ -16,6 +16,14 @@ public:
 	 */
 	virtual UMovieSceneSection* GetSectionObject() = 0;
 
+	/**
+	 * Allows each section to have it's own unique widget for advanced editing functionality
+	 * OnPaintSection will still be called if a widget is provided.  OnPaintSection is still used for the background section display
+	 * 
+	 * @return The generated widget 
+	 */
+	virtual TSharedRef<SWidget> GenerateSectionWidget() { return SNullWidget::NullWidget; }
+
 	/**
 	 * Called when the section should be painted
 	 *
