Fix shadowed variable

diff --git a/Engine/Source/Editor/MovieSceneTools/Private/TrackEditors/ShotTrackEditor.cpp b/Engine/Source/Editor/MovieSceneTools/Private/TrackEditors/ShotTrackEditor.cpp
--- a/Engine/Source/Editor/MovieSceneTools/Private/TrackEditors/ShotTrackEditor.cpp
+++ b/Engine/Source/Editor/MovieSceneTools/Private/TrackEditors/ShotTrackEditor.cpp
@@ -436,19 +436,19 @@ void FShotSection::CalculateThumbnailWidthAndResize()
 	}
 }
 
-ACameraActor* FShotSection::UpdateCameraObject()
+ACameraActor* FShotSection::UpdateCameraObject() const
 {
 	TArray<UObject*> OutObjects;
 	// @todo Sequencer - Sub-MovieScenes The director track may be able to get cameras from sub-movie scenes
 	Sequencer.Pin()->GetRuntimeObjects( Sequencer.Pin()->GetRootMovieSceneInstance(),  Section->GetCameraGuid(), OutObjects);
 
-	ACameraActor* Camera = nullptr;
+	ACameraActor* ReturnCam = nullptr;
 	if( OutObjects.Num() > 0 )
 	{
-		Camera = Cast<ACameraActor>( OutObjects[0] );
+		ReturnCam = Cast<ACameraActor>( OutObjects[0] );
 	}
 
-	return Camera;
+	return ReturnCam;
 }
 
 FText FShotSection::GetShotName() const
diff --git a/Engine/Source/Editor/MovieSceneTools/Private/TrackEditors/ShotTrackEditor.h b/Engine/Source/Editor/MovieSceneTools/Private/TrackEditors/ShotTrackEditor.h
--- a/Engine/Source/Editor/MovieSceneTools/Private/TrackEditors/ShotTrackEditor.h
+++ b/Engine/Source/Editor/MovieSceneTools/Private/TrackEditors/ShotTrackEditor.h
@@ -165,7 +165,7 @@ public:
 	/** @return The sequencer widget owning the shot section */
 	TSharedRef<SWidget> GetSequencerWidget() { return Sequencer.Pin()->GetSequencerWidget(); }
 private:
-	ACameraActor* UpdateCameraObject();
+	ACameraActor* UpdateCameraObject() const;
 	FText GetShotName() const;
 	void OnRenameShot( const FText& NewShotName, ETextCommit::Type CommitType );
 private:
