Sequencer: Hotkeys for translation, rotation, scale.

diff --git a/Engine/Source/Editor/MovieSceneTools/Private/TrackEditors/TransformTrackEditor.cpp b/Engine/Source/Editor/MovieSceneTools/Private/TrackEditors/TransformTrackEditor.cpp
--- a/Engine/Source/Editor/MovieSceneTools/Private/TrackEditors/TransformTrackEditor.cpp
+++ b/Engine/Source/Editor/MovieSceneTools/Private/TrackEditors/TransformTrackEditor.cpp
@@ -83,6 +83,45 @@ private:
 	UMovieSceneSection& Section;
 };
 
+class F3DTransformTrackCommands : public TCommands<F3DTransformTrackCommands>
+{
+
+public:
+	F3DTransformTrackCommands() : TCommands<F3DTransformTrackCommands>
+	(
+		"3DTransformTrack",
+		NSLOCTEXT("Contexts", "3DTransformTrack", "3DTransformTrack"),
+		NAME_None, // "MainFrame" // @todo Fix this crash
+		FEditorStyle::GetStyleSetName() // Icon Style Set
+	)
+	{}
+		
+	/** Sets a transform key at the current time for the selected actor */
+	TSharedPtr< FUICommandInfo > AddTransformKey;
+
+	/** Sets a translation key at the current time for the selected actor */
+	TSharedPtr< FUICommandInfo > AddTranslationKey;
+
+	/** Sets a rotation key at the current time for the selected actor */
+	TSharedPtr< FUICommandInfo > AddRotationKey;
+
+	/** Sets a scale key at the current time for the selected actor */
+	TSharedPtr< FUICommandInfo > AddScaleKey;
+
+	/**
+	 * Initialize commands
+	 */
+	virtual void RegisterCommands() override;
+};
+
+void F3DTransformTrackCommands::RegisterCommands()
+{
+	UI_COMMAND( AddTransformKey, "Add Transform Key", "Add a transform key at the current time for the selected actor.", EUserInterfaceActionType::Button, FInputChord(EKeys::S) );
+	UI_COMMAND( AddTranslationKey, "Add Translation Key", "Add a translation key at the current time for the selected actor.", EUserInterfaceActionType::Button, FInputChord(EModifierKey::Shift, EKeys::W) );
+	UI_COMMAND( AddRotationKey, "Add Rotation Key", "Add a rotation key at the current time for the selected actor.", EUserInterfaceActionType::Button, FInputChord(EModifierKey::Shift, EKeys::E) );
+	UI_COMMAND( AddScaleKey, "Add Scale Key", "Add a scale key at the current time for the selected actor.", EUserInterfaceActionType::Button, FInputChord(EModifierKey::Shift, EKeys::R) );
+}
+
 F3DTransformTrackEditor::F3DTransformTrackEditor( TSharedRef<ISequencer> InSequencer )
 	: FMovieSceneTrackEditor( InSequencer ) 
 {
@@ -93,6 +132,8 @@ F3DTransformTrackEditor::F3DTransformTrackEditor( TSharedRef<ISequencer> InSeque
 	// Listen for the viewport's viewed through camera starts and stops movement
 	GEditor->OnBeginCameraMovement().AddRaw( this, &F3DTransformTrackEditor::OnPreTransformChanged );
 	GEditor->OnEndCameraMovement().AddRaw( this, &F3DTransformTrackEditor::OnTransformChanged );
+
+	F3DTransformTrackCommands::Register();
 }
 
 F3DTransformTrackEditor::~F3DTransformTrackEditor()
@@ -102,6 +143,8 @@ F3DTransformTrackEditor::~F3DTransformTrackEditor()
 
 	GEditor->OnBeginCameraMovement().RemoveAll( this );
 	GEditor->OnEndCameraMovement().RemoveAll( this );
+
+	F3DTransformTrackCommands::Unregister();
 }
 
 TSharedRef<FMovieSceneTrackEditor> F3DTransformTrackEditor::CreateTrackEditor( TSharedRef<ISequencer> InSequencer )
@@ -212,11 +255,16 @@ void F3DTransformTrackEditor::OnTransformChanged( UObject& InObject )
 		FTransformDataPair TransformPair(NewTransformData, ExistingTransform);
 
 		AnimatablePropertyChanged(UMovieScene3DTransformTrack::StaticClass(), true,
-			FOnKeyProperty::CreateRaw(this, &F3DTransformTrackEditor::OnTransformChangedInternals, &InObject, ObjectHandle, TransformPair, true));
+			FOnKeyProperty::CreateRaw(this, &F3DTransformTrackEditor::OnTransformChangedInternals, &InObject, ObjectHandle, TransformPair, true, false, F3DTransformTrackKey::Key_All));
 	}
 }
 
 void F3DTransformTrackEditor::AddKey(const FGuid& ObjectGuid, UObject* AdditionalAsset)
+{
+	AddKeyInternal(ObjectGuid, AdditionalAsset);
+}
+
+void F3DTransformTrackEditor::AddKeyInternal(const FGuid& ObjectGuid, UObject* AdditionalAsset, bool bForceKey, F3DTransformTrackKey::Type KeyType)
 {
 	TArray<UObject*> OutObjects;
 	GetSequencer()->GetRuntimeObjects( GetSequencer()->GetFocusedMovieSceneInstance(), ObjectGuid, OutObjects);
@@ -250,11 +298,64 @@ void F3DTransformTrackEditor::AddKey(const FGuid& ObjectGuid, UObject* Additiona
 			FTransformDataPair TransformPair(CurrentTransform, FTransformData());
 
 			AnimatablePropertyChanged(UMovieScene3DTransformTrack::StaticClass(), false,
-				FOnKeyProperty::CreateRaw(this, &F3DTransformTrackEditor::OnTransformChangedInternals, Object, ObjectHandle, TransformPair, false));
+				FOnKeyProperty::CreateRaw(this, &F3DTransformTrackEditor::OnTransformChangedInternals, Object, ObjectHandle, TransformPair, false, bForceKey, KeyType));
 		}
 	}
 }
 
+void F3DTransformTrackEditor::AddTransformKey()
+{
+	AddTransformKeyInternal();
+}
+
+void F3DTransformTrackEditor::AddTranslationKey()
+{
+	AddTransformKeyInternal(F3DTransformTrackKey::Key_Translation);
+}
+
+void F3DTransformTrackEditor::AddRotationKey()
+{
+	AddTransformKeyInternal(F3DTransformTrackKey::Key_Rotation);
+}
+
+void F3DTransformTrackEditor::AddScaleKey()
+{
+	AddTransformKeyInternal(F3DTransformTrackKey::Key_Scale);
+}
+
+void F3DTransformTrackEditor::AddTransformKeyInternal(F3DTransformTrackKey::Type KeyType)
+{
+	USelection* CurrentSelection = GEditor->GetSelectedActors();
+	TArray<UObject*> SelectedActors;
+	CurrentSelection->GetSelectedObjects( AActor::StaticClass(), SelectedActors );
+	for (TArray<UObject*>::TIterator It(SelectedActors); It; ++It)
+	{
+		FGuid ObjectGuid = GetSequencer()->GetHandleToObject(*It);
+		AddKeyInternal(ObjectGuid, NULL, true, KeyType);
+	}
+}
+
+void F3DTransformTrackEditor::BindCommands(TSharedRef<FUICommandList> SequencerCommandBindings)
+{
+	const F3DTransformTrackCommands& Commands = F3DTransformTrackCommands::Get();
+
+	SequencerCommandBindings->MapAction(
+		Commands.AddTransformKey,
+		FExecuteAction::CreateSP( this, &F3DTransformTrackEditor::AddTransformKey ) );
+
+	SequencerCommandBindings->MapAction(
+		Commands.AddTranslationKey,
+		FExecuteAction::CreateSP( this, &F3DTransformTrackEditor::AddTranslationKey ) );
+
+	SequencerCommandBindings->MapAction(
+		Commands.AddRotationKey,
+		FExecuteAction::CreateSP( this, &F3DTransformTrackEditor::AddRotationKey ) );
+
+	SequencerCommandBindings->MapAction(
+		Commands.AddScaleKey,
+		FExecuteAction::CreateSP( this, &F3DTransformTrackEditor::AddScaleKey ) );
+}
+
 void F3DTransformTrackEditor::BuildObjectBindingEditButtons(TSharedPtr<SHorizontalBox> EditBox, const FGuid& ObjectGuid, const UClass* ObjectClass)
 {
 	TArray<UObject*> OutObjects;
@@ -297,7 +398,7 @@ void F3DTransformTrackEditor::BuildObjectBindingEditButtons(TSharedPtr<SHorizont
 	];
 }
 
-void F3DTransformTrackEditor::OnTransformChangedInternals(float KeyTime, UObject* InObject, FGuid ObjectHandle, FTransformDataPair TransformPair, bool bAutoKeying)
+void F3DTransformTrackEditor::OnTransformChangedInternals(float KeyTime, UObject* InObject, FGuid ObjectHandle, FTransformDataPair TransformPair, bool bAutoKeying, bool bForceKey, F3DTransformTrackKey::Type KeyType)
 {
 	// Only unwind rotation if we're generating keys while recording (scene is actively playing back)
 	const bool bUnwindRotation = GetSequencer()->IsRecordingLive();
@@ -316,9 +417,15 @@ void F3DTransformTrackEditor::OnTransformChangedInternals(float KeyTime, UObject
 			TransformPair.LastTransformData.bValid = TransformTrack->Eval(KeyTime, KeyTime, TransformPair.LastTransformData.Translation, TransformPair.LastTransformData.Rotation, TransformPair.LastTransformData.Scale, bHasTranslationKeys, bHasRotationKeys, bHasScaleKeys);
 		}
 
+		// If forcing the keyframe, set the last data invalid so that the transform keyframe will be set regardless of the current time values being equal.
+		if ( bForceKey )
+		{
+			TransformPair.LastTransformData.bValid = false;
+		}
+
 		FTransformKey TransformKey = FTransformKey(KeyTime, TransformPair.TransformData, TransformPair.LastTransformData);
 
-		bool bSuccessfulAdd = TransformTrack->AddKeyToSection( ObjectHandle, TransformKey, bUnwindRotation );
+		bool bSuccessfulAdd = TransformTrack->AddKeyToSection( ObjectHandle, TransformKey, bUnwindRotation, KeyType );
 		if (bSuccessfulAdd)
 		{
 			TransformTrack->SetAsShowable();
diff --git a/Engine/Source/Editor/MovieSceneTools/Private/TrackEditors/TransformTrackEditor.h b/Engine/Source/Editor/MovieSceneTools/Private/TrackEditors/TransformTrackEditor.h
--- a/Engine/Source/Editor/MovieSceneTools/Private/TrackEditors/TransformTrackEditor.h
+++ b/Engine/Source/Editor/MovieSceneTools/Private/TrackEditors/TransformTrackEditor.h
@@ -30,9 +30,14 @@ public:
 	virtual bool SupportsType( TSubclassOf<UMovieSceneTrack> Type ) const override;
 	virtual TSharedRef<ISequencerSection> MakeSectionInterface( UMovieSceneSection& SectionObject, UMovieSceneTrack* Track ) override;
 	virtual void AddKey( const FGuid& ObjectGuid, UObject* AdditionalAsset = NULL ) override;
+	virtual void BindCommands(TSharedRef<FUICommandList> SequencerCommandBindings) override;
 	virtual void BuildObjectBindingEditButtons(TSharedPtr<SHorizontalBox> EditBox, const FGuid& ObjectBinding, const UClass* ObjectClass) override;
 
 private:
+
+	/** Custom add key implementation */
+	void AddKeyInternal(const FGuid& ObjectGuid, UObject* AdditionalAsset = NULL, bool bForceKey = false, F3DTransformTrackKey::Type KeyType = F3DTransformTrackKey::Key_All);
+
 	/**
 	 * Called before an actor or component transform changes
 	 *
@@ -48,7 +53,7 @@ private:
 	void OnTransformChanged( UObject& InObject );
 
 	/** Delegate for animatable property changed in OnTransformChanged */
-	void OnTransformChangedInternals(float KeyTime, UObject* InObject, FGuid ObjectHandle, struct FTransformDataPair TransformPair, bool bAutoKeying);
+	void OnTransformChangedInternals(float KeyTime, UObject* InObject, FGuid ObjectHandle, struct FTransformDataPair TransformPair, bool bAutoKeying, bool bForceKey, F3DTransformTrackKey::Type KeyType);
 
 	/** Delegate for camera button lock state */
 	ECheckBoxState IsCameraLocked(TWeakObjectPtr<ACameraActor> CameraActor) const; 
@@ -59,6 +64,21 @@ private:
 	/** Delegate for camera button lock tooltip */
 	FText GetLockCameraToolTip(TWeakObjectPtr<ACameraActor> CameraActor) const; 
 
+	/** Add a transform key */
+	void AddTransformKey();
+
+	/** Add a translation key */
+	void AddTranslationKey();
+
+	/** Add a rotation key */
+	void AddRotationKey();
+
+	/** Add a scale key */
+	void AddScaleKey();
+
+	/** Internal add transform key */
+	void AddTransformKeyInternal(F3DTransformTrackKey::Type KeyType = F3DTransformTrackKey::Key_All);
+
 private:
 	/** Mapping of objects to their existing transform data (for comparing against new transform data) */
 	TMap< TWeakObjectPtr<UObject>, FTransformData > ObjectToExistingTransform;
diff --git a/Engine/Source/Editor/Sequencer/Private/Sequencer.cpp b/Engine/Source/Editor/Sequencer/Private/Sequencer.cpp
--- a/Engine/Source/Editor/Sequencer/Private/Sequencer.cpp
+++ b/Engine/Source/Editor/Sequencer/Private/Sequencer.cpp
@@ -1697,6 +1697,11 @@ void FSequencer::BindSequencerCommands()
 		FExecuteAction::CreateLambda( [Settings]{ Settings->SetCurveVisibility( ESequencerCurveVisibility::AnimatedCurves ); } ),
 		FCanExecuteAction::CreateLambda( []{ return true; } ),
 		FIsActionChecked::CreateLambda( [Settings]{ return Settings->GetCurveVisibility() == ESequencerCurveVisibility::AnimatedCurves; } ) );
+
+	for (int32 i = 0; i < TrackEditors.Num(); ++i)
+	{
+		TrackEditors[i]->BindCommands(SequencerCommandBindings);
+	}
 }
 
 void FSequencer::BuildObjectBindingContextMenu(FMenuBuilder& MenuBuilder, const FGuid& ObjectBinding, const UClass* ObjectClass)
diff --git a/Engine/Source/Editor/Sequencer/Public/MovieSceneTrackEditor.h b/Engine/Source/Editor/Sequencer/Public/MovieSceneTrackEditor.h
--- a/Engine/Source/Editor/Sequencer/Public/MovieSceneTrackEditor.h
+++ b/Engine/Source/Editor/Sequencer/Public/MovieSceneTrackEditor.h
@@ -174,6 +174,13 @@ public:
 	 */
 	virtual bool HandleAssetAdded(UObject* Asset, const FGuid& TargetObjectGuid) {return false;}
 	
+	/**
+	 * Allows the track editors to bind commands
+	 *
+	 * @param SequencerCommandBindings	The command bindings to map to
+	*/
+	virtual void BindCommands(TSharedRef<FUICommandList> SequencerCommandBindings) {}
+
 	/**
 	 * Builds up the object binding context menu for the outliner
 	 *
diff --git a/Engine/Source/Runtime/MovieSceneTracks/Private/Tracks/MovieScene3DTransformTrack.cpp b/Engine/Source/Runtime/MovieSceneTracks/Private/Tracks/MovieScene3DTransformTrack.cpp
--- a/Engine/Source/Runtime/MovieSceneTracks/Private/Tracks/MovieScene3DTransformTrack.cpp
+++ b/Engine/Source/Runtime/MovieSceneTracks/Private/Tracks/MovieScene3DTransformTrack.cpp
@@ -72,7 +72,7 @@ TSharedPtr<IMovieSceneTrackInstance> UMovieScene3DTransformTrack::CreateInstance
 	return MakeShareable( new FMovieScene3DTransformTrackInstance( *this ) );
 }
 
-bool UMovieScene3DTransformTrack::AddKeyToSection( const FGuid& ObjectHandle, const FTransformKey& InKey, const bool bUnwindRotation )
+bool UMovieScene3DTransformTrack::AddKeyToSection( const FGuid& ObjectHandle, const FTransformKey& InKey, const bool bUnwindRotation, F3DTransformTrackKey::Type KeyType )
 {
 	const UMovieSceneSection* NearestSection = MovieSceneHelpers::FindSectionAtTime(Sections, InKey.GetKeyTime());
 	if (!NearestSection || CastChecked<UMovieScene3DTransformSection>(NearestSection)->NewKeyIsNewData(InKey))
@@ -82,9 +82,12 @@ bool UMovieScene3DTransformTrack::AddKeyToSection( const FGuid& ObjectHandle, co
 		UMovieScene3DTransformSection* NewSection = Cast<UMovieScene3DTransformSection>( FindOrAddSection( InKey.GetKeyTime() ) );
 
 		// key each component of the transform
-		NewSection->AddTranslationKeys( InKey );
-		NewSection->AddRotationKeys( InKey, bUnwindRotation );
-		NewSection->AddScaleKeys( InKey );
+		if (KeyType & F3DTransformTrackKey::Key_Translation)
+			NewSection->AddTranslationKeys( InKey );
+		if (KeyType & F3DTransformTrackKey::Key_Rotation)
+			NewSection->AddRotationKeys( InKey, bUnwindRotation );
+		if (KeyType & F3DTransformTrackKey::Key_Scale)
+			NewSection->AddScaleKeys( InKey );
 
 		return true;
 	}
diff --git a/Engine/Source/Runtime/MovieSceneTracks/Public/Tracks/MovieScene3DTransformTrack.h b/Engine/Source/Runtime/MovieSceneTracks/Public/Tracks/MovieScene3DTransformTrack.h
--- a/Engine/Source/Runtime/MovieSceneTracks/Public/Tracks/MovieScene3DTransformTrack.h
+++ b/Engine/Source/Runtime/MovieSceneTracks/Public/Tracks/MovieScene3DTransformTrack.h
@@ -7,6 +7,17 @@
 #include "MovieScenePropertyTrack.h"
 #include "MovieScene3DTransformTrack.generated.h"
 
+namespace F3DTransformTrackKey
+{
+	enum Type
+	{
+		Key_Translation = 0x00000001,
+		Key_Rotation    = 0x00000002,
+		Key_Scale       = 0x00000004,
+		Key_All         = Key_Translation | Key_Rotation | Key_Scale
+	};
+}
+
 /**
  * Stores information about a transform for the purpose of adding keys to a transform section
  */
@@ -141,9 +152,10 @@ public:
 	 * @param ObjectHandle		Handle to the object(s) being changed
 	 * @param InKey				The transform key to add
 	 * @param bUnwindRotation	When true, the value will be treated like a rotation value in degrees, and will automatically be unwound to prevent flipping 360 degrees from the previous key 	 
+	 * @param KeyType           Key type (trans, rot, scale)
 	 * @return True if the key was successfully added.
 	 */
-	virtual bool AddKeyToSection( const FGuid& ObjectHandle, const FTransformKey& InKey, const bool bUnwindRotation );
+	virtual bool AddKeyToSection( const FGuid& ObjectHandle, const FTransformKey& InKey, const bool bUnwindRotation, F3DTransformTrackKey::Type KeyType = F3DTransformTrackKey::Key_All );
 
 	/**
 	 * Evaluates the track at the playback position
