Sequencer - Don't add tracks for actors when not in auto-key mode.

diff --git a/Engine/Source/Editor/MovieSceneTools/Private/TrackEditors/TransformTrackEditor.cpp b/Engine/Source/Editor/MovieSceneTools/Private/TrackEditors/TransformTrackEditor.cpp
--- a/Engine/Source/Editor/MovieSceneTools/Private/TrackEditors/TransformTrackEditor.cpp
+++ b/Engine/Source/Editor/MovieSceneTools/Private/TrackEditors/TransformTrackEditor.cpp
@@ -215,47 +215,49 @@ struct FTransformDataPair
 void F3DTransformTrackEditor::OnTransformChanged( UObject& InObject )
 {
 	const TSharedPtr<ISequencer> Sequencer = GetSequencer();
-
-	USceneComponent* SceneComponentThatChanged = NULL;
-
-	// The runtime binding
-	FGuid ObjectHandle;
-
-	AActor* Actor = Cast<AActor>( &InObject );
-	if( Actor && Actor->GetRootComponent() )
+	if (Sequencer->GetAutoKeyEnabled())
 	{
-		// Get a handle bound to the key/section we are adding so we know what objects to change during playback
-		ObjectHandle = Sequencer->GetHandleToObject( Actor );
-		SceneComponentThatChanged = Actor->GetRootComponent();
-	}
+		USceneComponent* SceneComponentThatChanged = NULL;
 
-	else
-	{
-		// If the object wasn't an actor attempt to get it directly as a scene component 
-		SceneComponentThatChanged = Cast<USceneComponent>( &InObject );
-		if( SceneComponentThatChanged )
+		// The runtime binding
+		FGuid ObjectHandle;
+
+		AActor* Actor = Cast<AActor>( &InObject );
+		if( Actor && Actor->GetRootComponent() )
 		{
-			ObjectHandle = Sequencer->GetHandleToObject( SceneComponentThatChanged );
+			// Get a handle bound to the key/section we are adding so we know what objects to change during playback
+			ObjectHandle = Sequencer->GetHandleToObject( Actor );
+			SceneComponentThatChanged = Actor->GetRootComponent();
 		}
 
-	}
+		else
+		{
+			// If the object wasn't an actor attempt to get it directly as a scene component 
+			SceneComponentThatChanged = Cast<USceneComponent>( &InObject );
+			if( SceneComponentThatChanged )
+			{
+				ObjectHandle = Sequencer->GetHandleToObject( SceneComponentThatChanged );
+			}
 
-	if( SceneComponentThatChanged && ObjectHandle.IsValid() )
-	{
-		// Find an existing transform if possible.  If one exists we will compare against the new one to decide what components of the transform need keys
-		FTransformData ExistingTransform = ObjectToExistingTransform.FindRef( &InObject );
+		}
 
-		// Remove it from the list of cached transforms. 
-		// @todo sequencer livecapture: This can be made much for efficient by not removing cached state during live capture situation
-		ObjectToExistingTransform.Remove( &InObject );
+		if( SceneComponentThatChanged && ObjectHandle.IsValid() )
+		{
+			// Find an existing transform if possible.  If one exists we will compare against the new one to decide what components of the transform need keys
+			FTransformData ExistingTransform = ObjectToExistingTransform.FindRef( &InObject );
 
-		// Build new transform data
-		FTransformData NewTransformData( SceneComponentThatChanged );
+			// Remove it from the list of cached transforms. 
+			// @todo sequencer livecapture: This can be made much for efficient by not removing cached state during live capture situation
+			ObjectToExistingTransform.Remove( &InObject );
 
-		FTransformDataPair TransformPair(NewTransformData, ExistingTransform);
+			// Build new transform data
+			FTransformData NewTransformData( SceneComponentThatChanged );
 
-		AnimatablePropertyChanged(UMovieScene3DTransformTrack::StaticClass(), true,
-			FOnKeyProperty::CreateRaw(this, &F3DTransformTrackEditor::OnTransformChangedInternals, &InObject, ObjectHandle, TransformPair, true, false, F3DTransformTrackKey::Key_All));
+			FTransformDataPair TransformPair(NewTransformData, ExistingTransform);
+
+			AnimatablePropertyChanged(UMovieScene3DTransformTrack::StaticClass(), true,
+				FOnKeyProperty::CreateRaw(this, &F3DTransformTrackEditor::OnTransformChangedInternals, &InObject, ObjectHandle, TransformPair, true, false, F3DTransformTrackKey::Key_All));
+		}
 	}
 }
 
