Sequencer - Allow keying of UActorComponent properties.

diff --git a/Engine/Source/Editor/MovieSceneTools/Private/K2Node_PlayMovieScene.cpp b/Engine/Source/Editor/MovieSceneTools/Private/K2Node_PlayMovieScene.cpp
--- a/Engine/Source/Editor/MovieSceneTools/Private/K2Node_PlayMovieScene.cpp
+++ b/Engine/Source/Editor/MovieSceneTools/Private/K2Node_PlayMovieScene.cpp
@@ -8,7 +8,6 @@
 #include "Kismet2/BlueprintEditorUtils.h"
 #include "MovieSceneBindings.h"
 
-
 UEdGraphNode* FEdGraphSchemaAction_K2AddPlayMovieScene::PerformAction(class UEdGraph* ParentGraph, UEdGraphPin* FromPin, const FVector2D Location, bool bSelectNewNode/* = true*/)
 {
 	const FScopedTransaction Transaction( NSLOCTEXT("PlayMovieSceneNode", "UndoAddingNewNode", "Add PlayMovieScene Node") );
@@ -63,7 +62,7 @@ void UK2Node_PlayMovieScene::AllocateDefaultPins()
 		{
 			auto& BoundObject = *BoundObjectIter;
 
-			CreatePinForBoundObject( BoundObject );
+			CreatePinsForBoundObject( BoundObject );
 		}
 	}
 	
@@ -130,43 +129,51 @@ void UK2Node_PlayMovieScene::PinConnectionListChanged( UEdGraphPin* InPin )
 		{
 			auto& BoundObject = *BoundObjectIter;
 
-			TArray< UObject* > OtherObjects;
+			TArray< FMovieSceneBoundObjectInfo > OtherObjectInfos;
 
 			// Find the pin that goes with this object
-			auto* Pin = FindPinChecked( BoundObject.GetPossessableGuid().ToString( EGuidFormats::DigitsWithHyphens ) );
+			UEdGraphPin* ActorPin = FindPinChecked( BoundObject.GetPossessableGuid().ToString( EGuidFormats::DigitsWithHyphens ) );
+			UEdGraphPin* ComponentNamePin = FindPinChecked(BoundObject.GetPossessableGuid().ToString(EGuidFormats::DigitsWithHyphens) + "_ComponentName");
 
 			// Is the pin even linked to anything?
-			if( Pin->LinkedTo.Num() > 0 )
+			checkf(ActorPin->LinkedTo.Num() == ComponentNamePin->LinkedTo.Num(), TEXT("Actor pin link count must match the component name link count"));
+			for (int32 i = 0; i < ActorPin->LinkedTo.Num(); i++)
 			{
-				// @todo sequencer major: Add support for multiple connections to a single input pin to be treated as an "unordered array"
-				for( auto OtherPinIter( Pin->LinkedTo.CreateIterator() ); OtherPinIter; ++OtherPinIter )
+				UEdGraphPin* LinkedActorPin = ActorPin->LinkedTo[i];
+				UEdGraphPin* LinkedComponentNamePin = ComponentNamePin->LinkedTo[i];
+				if( LinkedActorPin != NULL )
 				{
-					auto* OtherPin = *OtherPinIter;
-					if( OtherPin != NULL )
+					// Look for an object bound to a literal.  We can use these for scrub preview in the editor!
+					UK2Node_Literal* LinkedActorLiteralNode = Cast< UK2Node_Literal >( LinkedActorPin->GetOwningNode() );
+					if (LinkedActorLiteralNode != NULL)
 					{
-						// Look for an object bound to a literal.  We can use these for scrub preview in the editor!
-						UK2Node_Literal* OtherLiteralPin = Cast< UK2Node_Literal >( OtherPin->GetOwningNode() );
-						if( OtherLiteralPin != NULL )
+						// @todo sequencer: Because we only recognize object literals, any dynamically bound actor won't be scrubable
+						//   in the editor.  Maybe we should support a "preview actor" that can be hooked up just for scrubbing and puppeting?
+						//  ==> Maybe use the pin's default value as the preview actor, even when overridden with a dynamically spawned actor?
+						UObject* LinkedActor = LinkedActorLiteralNode->GetObjectRef();
+						if ( LinkedActor != NULL )
 						{
-							// @todo sequencer: Because we only recognize object literals, any dynamically bound actor won't be scrubable
-							//   in the editor.  Maybe we should support a "preview actor" that can be hooked up just for scrubbing and puppeting?
-							//  ==> Maybe use the pin's default value as the preview actor, even when overridden with a dynamically spawned actor?
-							UObject* OtherObject = OtherLiteralPin->GetObjectRef();
-							if ( OtherObject != NULL )
+							FString LinkedComponentName;
+							if (LinkedComponentNamePin)
 							{
-								OtherObjects.Add( OtherObject );
-							}							
+								UK2Node_Literal* LinkedComponentNameLiteralNode = Cast< UK2Node_Literal >(LinkedComponentNamePin->GetOwningNode());
+								if (LinkedComponentNameLiteralNode != NULL)
+								{
+									LinkedComponentName = LinkedComponentNameLiteralNode->GetValuePin()->GetDefaultAsString();
+								}
+							}
+							OtherObjectInfos.Add( FMovieSceneBoundObjectInfo(LinkedActor, LinkedComponentName));
 						}
 					}
 				}
 			}
 
 			// Update our bindings to match the state of the node
-			if( BoundObject.GetObjects() != OtherObjects )	// @todo sequencer: Kind of weird to compare entire array (order change shouldn't cause us to invalidate).  Change to a set compare?
+			if( BoundObject.GetObjectInfos() != OtherObjectInfos )	// @todo sequencer: Kind of weird to compare entire array (order change shouldn't cause us to invalidate).  Change to a set compare?
 			{
 				// @todo sequencer: No type checking is happening here.  Should we? (interactive during pin drag?)
 				Modify();
-				BoundObject.SetObjects( OtherObjects );
+				BoundObject.SetObjectInfos( OtherObjectInfos );
 
 				bAnyBindingsChanged = true;
 			}
@@ -220,63 +227,62 @@ void UK2Node_PlayMovieScene::SetMovieScene( class UMovieScene* NewMovieScene )
 }
 
 
-void UK2Node_PlayMovieScene::BindPossessableToObjects( const FGuid& PossessableGuid, const TArray< UObject* >& Objects )
+void UK2Node_PlayMovieScene::BindPossessableToObjects(const FGuid& PossessableGuid, const TArray< FMovieSceneBoundObjectInfo >& ObjectInfos )
 {
 	CreateBindingsIfNeeded();
 
 	// Bind the possessable from the MovieScene asset to the supplied object
-	FMovieSceneBoundObject& BoundObject = MovieSceneBindings->AddBinding( PossessableGuid, Objects );
+	FMovieSceneBoundObject& BoundObject = MovieSceneBindings->AddBinding( PossessableGuid, ObjectInfos );
 
 	// Create a pin for the bound object
 	Modify();
-	CreatePinForBoundObject( BoundObject );
+	CreatePinsForBoundObject( BoundObject );
 	FBlueprintEditorUtils::MarkBlueprintAsStructurallyModified( this->GetBlueprint() );
 }
 
 
-TArray< UObject* > UK2Node_PlayMovieScene::FindBoundObjects( const FGuid& Guid )
+TArray< FMovieSceneBoundObjectInfo > UK2Node_PlayMovieScene::FindBoundObjectInfos( const FGuid& Guid )
 {
+	TArray< UObject* > BoundObjects;
 	if( MovieSceneBindings != NULL )
 	{
-		return MovieSceneBindings->FindBoundObjects( Guid );
+		return MovieSceneBindings->FindBoundObjects(Guid);
 	}
-
-	return TArray< UObject* >();
+	return TArray< FMovieSceneBoundObjectInfo >();
 }
 
-FGuid UK2Node_PlayMovieScene::FindGuidForObject( UObject* Object ) const
+FGuid UK2Node_PlayMovieScene::FindGuidForObjectInfo( FMovieSceneBoundObjectInfo ObjectInfo ) const
 {
 	if( MovieSceneBindings != NULL )
 	{
 		// Search all bindings for the object 
 		TArray<FMovieSceneBoundObject>& ObjectBindings = MovieSceneBindings->GetBoundObjects();
-		for( int32 BoundObjIndex = 0; BoundObjIndex < ObjectBindings.Num(); ++BoundObjIndex )
+		for(FMovieSceneBoundObject& BoundObject : MovieSceneBindings->GetBoundObjects())
 		{
-			// A binding slots can have multiple objects.  Ask the slot if the object is bound
-			FMovieSceneBoundObject& BindingSlot = ObjectBindings[BoundObjIndex];
-			if( BindingSlot.ContainsObject( Object ) )
+			if(BoundObject.GetObjectInfos().Contains(ObjectInfo))
 			{
-				return BindingSlot.GetPossessableGuid();
+				return BoundObject.GetPossessableGuid();
 			}
 		}
 	}
-
 	return FGuid();
 }
 
 
-void UK2Node_PlayMovieScene::CreatePinForBoundObject( FMovieSceneBoundObject& BoundObject )
+void UK2Node_PlayMovieScene::CreatePinsForBoundObject( FMovieSceneBoundObject& BoundObject )
 {
 	const UEdGraphSchema_K2* K2Schema = GetDefault<UEdGraphSchema_K2>();
 
 	// We use the GUID as the pin name as this uniquely identifies it
-	const FString PinName = BoundObject.GetPossessableGuid().ToString( EGuidFormats::DigitsWithHyphens );
+	const FString ActorPinName = BoundObject.GetPossessableGuid().ToString( EGuidFormats::DigitsWithHyphens );
+	const FString ComponentNamePinName = ActorPinName + "_ComponentName";
 
 	// For the friendly name, we use the possessable name from the MovieScene asset that is associated with this node
-	FText PinFriendlyName = FText::FromString(PinName);
+	FText ActorPinFriendlyName = FText::FromString(ActorPinName);
+	FText ComponentNamePinFriendlyName = FText::FromString(ComponentNamePinName);
 	{
 		UMovieScene* MovieScene = GetMovieScene();
-		if( MovieScene != NULL )		// @todo sequencer: Need to refresh the PinFriendlyName if the MovieScene asset changes, or if the possessable slot is renamed within Sequencer
+		if( MovieScene != NULL )		// @todo sequencer: Need to refresh the ActorPinFriendlyName if the MovieScene asset changes, or if the possessable slot is renamed within Sequencer
 		{
 			for( auto PossessableIndex = 0; PossessableIndex < MovieScene->GetPossessableCount(); ++PossessableIndex )
 			{
@@ -284,7 +290,8 @@ void UK2Node_PlayMovieScene::CreatePinForBoundObject( FMovieSceneBoundObject& Bo
 				if( Possessable.GetGuid() == BoundObject.GetPossessableGuid() )
 				{
 					// Found a name for this possessable
-					PinFriendlyName = FText::FromString(Possessable.GetName());
+					ActorPinFriendlyName = FText::FromString(Possessable.GetName());
+					ComponentNamePinFriendlyName = FText::Format(NSLOCTEXT("PlayMovieSceneNode", "ComponentNamePinFormat", "{0} Component Name"), FText::FromString(Possessable.GetName()));
 					break;
 				}
 			}
@@ -295,24 +302,26 @@ void UK2Node_PlayMovieScene::CreatePinForBoundObject( FMovieSceneBoundObject& Bo
 	UObject* PinSubCategoryObject = AActor::StaticClass();
 	const bool bIsArray = false;
 	const bool bIsReference = false;
-	UEdGraphPin* NewPin = CreatePin( EGPD_Input, K2Schema->PC_Object, PinSubCategory, PinSubCategoryObject, bIsArray, bIsReference, PinName );
-	check( NewPin != NULL );
+	UEdGraphPin* NewActorPin = CreatePin( EGPD_Input, K2Schema->PC_Object, PinSubCategory, PinSubCategoryObject, bIsArray, bIsReference, ActorPinName );
+	UEdGraphPin* NewComponentNamePin = CreatePin(EGPD_Input, K2Schema->PC_String, PinSubCategory, PinSubCategoryObject, bIsArray, bIsReference, ComponentNamePinName );
+	check( NewActorPin != NULL && NewComponentNamePin != NULL );
 
 	// Set the friendly name for this pin
-	NewPin->PinFriendlyName = PinFriendlyName;
+	NewActorPin->PinFriendlyName = ActorPinFriendlyName;
+	NewComponentNamePin->PinFriendlyName = ComponentNamePinFriendlyName;
 
 	// Place a literal for the bound object and hook it up to the pin
 	// @todo sequencer: Should we instead set the default on the pin to the object instead?
-	const TArray< UObject* >& Objects = BoundObject.GetObjects();
-	if( Objects.Num() > 0 )
+	const TArray< FMovieSceneBoundObjectInfo >& ObjectInfos = BoundObject.GetObjectInfos();
+	if( ObjectInfos.Num() > 0 )
 	{
-		for( auto ObjectIter( Objects.CreateConstIterator() ); ObjectIter; ++ObjectIter )
+		for(const FMovieSceneBoundObjectInfo& BoundObjectInfo : ObjectInfos )
 		{
-			auto* Object = *ObjectIter;
-			if( ensure( Object != NULL ) )
+			if( ensure( BoundObjectInfo.Object != NULL ) )
 			{
 				// Check to see if we have a literal for this object already
-				UK2Node_Literal* LiteralNode = NULL;
+				UK2Node_Literal* ActorLiteralNode = NULL;
+				UK2Node_Literal* ComponentNameLiteralNode = NULL;
 				{
 					TArray< UK2Node_Literal* > LiteralNodes;
 					GetGraph()->GetNodesOfClass( LiteralNodes );
@@ -320,29 +329,58 @@ void UK2Node_PlayMovieScene::CreatePinForBoundObject( FMovieSceneBoundObject& Bo
 					{
 						const auto CurLiteralNode = *NodeIt;
 
-						if( CurLiteralNode->GetObjectRef() == Object )
+						if( CurLiteralNode->GetObjectRef() == BoundObjectInfo.Object )
 						{
 							// Found one!
-							LiteralNode = CurLiteralNode;
-							break;
+							ActorLiteralNode = CurLiteralNode;
+						}
+						else
+						{
+							UEdGraphPin* ValuePin = CurLiteralNode->GetValuePin();
+							if (ValuePin->PinType.PinCategory == K2Schema->PC_String && ValuePin->GetDefaultAsString() == BoundObjectInfo.Tag)
+							{
+								ComponentNameLiteralNode = CurLiteralNode;
+							}
 						}
 					}
+
+					if (ActorLiteralNode != NULL && ComponentNameLiteralNode != NULL)
+					{
+						break;
+					}
 				}
 
-				if( LiteralNode == NULL )
+				if( ActorLiteralNode == NULL )
 				{
 					// No literal for this object yet, so we'll make one now.
 					UK2Node_Literal* LiteralNodeTemplate = NewObject<UK2Node_Literal>();
-					LiteralNodeTemplate->SetObjectRef( Object );
+					LiteralNodeTemplate->SetObjectRef( BoundObjectInfo.Object );
 
 					// Figure out a decent place to stick the node
 					// @todo sequencer: The placement of these is really unacceptable.  We should setup new logic specific for moviescene pin objects.
 					const FVector2D NewNodePos = GetGraph()->GetGoodPlaceForNewNode();
-					LiteralNode = FEdGraphSchemaAction_K2NewNode::SpawnNodeFromTemplate<UK2Node_Literal>(GetGraph(), LiteralNodeTemplate, NewNodePos);
+					ActorLiteralNode = FEdGraphSchemaAction_K2NewNode::SpawnNodeFromTemplate<UK2Node_Literal>(GetGraph(), LiteralNodeTemplate, NewNodePos);
 				}
 
 				// Hook up the object reference literal to our pin
-				LiteralNode->GetValuePin()->MakeLinkTo( NewPin );
+				ActorLiteralNode->GetValuePin()->MakeLinkTo(NewActorPin);
+
+				if (BoundObjectInfo.Tag.IsEmpty() == false)
+				{
+					if (ComponentNameLiteralNode == NULL)
+					{
+						// No literal for this object yet, so we'll make one now.
+						UK2Node_Literal* LiteralNodeTemplate = NewObject<UK2Node_Literal>();
+
+						// Figure out a decent place to stick the node
+						// @todo sequencer: The placement of these is really unacceptable.  We should setup new logic specific for moviescene pin objects.
+						const FVector2D NewNodePos = GetGraph()->GetGoodPlaceForNewNode();
+						ComponentNameLiteralNode = FEdGraphSchemaAction_K2NewNode::SpawnNodeFromTemplate<UK2Node_Literal>(GetGraph(), LiteralNodeTemplate, NewNodePos);
+						ComponentNameLiteralNode->GetSchema()->TrySetDefaultValue(*ComponentNameLiteralNode->GetValuePin(), BoundObjectInfo.Tag);
+					}
+
+					ComponentNameLiteralNode->GetValuePin()->MakeLinkTo(NewComponentNamePin);
+				}
 			}
 		}
 	}
diff --git a/Engine/Source/Editor/MovieSceneTools/Public/K2Node_PlayMovieScene.h b/Engine/Source/Editor/MovieSceneTools/Public/K2Node_PlayMovieScene.h
--- a/Engine/Source/Editor/MovieSceneTools/Public/K2Node_PlayMovieScene.h
+++ b/Engine/Source/Editor/MovieSceneTools/Public/K2Node_PlayMovieScene.h
@@ -33,6 +33,8 @@ struct FEdGraphSchemaAction_K2AddPlayMovieScene : public FEdGraphSchemaAction_K2
 };
 
 
+struct FMovieSceneBoundObjectInfo;
+
 UCLASS( MinimalAPI )
 class UK2Node_PlayMovieScene : public UK2Node
 {
@@ -70,9 +72,9 @@ class UK2Node_PlayMovieScene : public UK2Node
 	 * Binds a possessable in the MovieScene asset (based on the guid of that possessable) to the specified objects
 	 *
 	 * @param	PossessableGuid		The guid of the posessable in the MovieScene to bind to
-	 * @param	Objects				The object(s) to bind to
+	 * @param	ObjectInfos				The object(s) to bind to
 	 */
-	virtual void BindPossessableToObjects( const FGuid& PossessableGuid, const TArray< UObject* >& Objects );
+	virtual void BindPossessableToObjects( const FGuid& PossessableGuid, const TArray< FMovieSceneBoundObjectInfo >& ObjectInfos );
 
 	/**
 	 * Given a guid for a possessable, attempts to find the object bound to that guid
@@ -81,15 +83,15 @@ class UK2Node_PlayMovieScene : public UK2Node
 	 *
 	 * @return	The bound objects, if found, otherwise NULL
 	 */
-	virtual TArray< UObject* > FindBoundObjects( const FGuid& Guid );
+	virtual TArray< FMovieSceneBoundObjectInfo > FindBoundObjectInfos(const FGuid& Guid);
 
 	/**
 	 * Finds a guid for a bound posessable object if it exists 
 	 *
-	 * @param Object	The object to search for
+	 * @param ObjectInfo	The object to search for
 	 * @return A valid guid if the object was found.  Otherwise the guid is invalid
 	 */
-	virtual FGuid FindGuidForObject( UObject* Object ) const;
+	virtual FGuid FindGuidForObjectInfo( FMovieSceneBoundObjectInfo ObjectInfo ) const;
 
 	/**
 	 * Returns access to an event you can bind to to find out when the object bindings change
@@ -118,7 +120,7 @@ protected:
 	 *
 	 * @param	BoundObject		The object to create a pin for
 	 */
-	virtual void CreatePinForBoundObject( struct FMovieSceneBoundObject& BoundObject );
+	virtual void CreatePinsForBoundObject( struct FMovieSceneBoundObject& BoundObject );
 
 	/**
 	 * Creates a MovieSceneBindings object if we don't have one yet
diff --git a/Engine/Source/Editor/Sequencer/Private/SequencerActorBindingManager.cpp b/Engine/Source/Editor/Sequencer/Private/SequencerActorBindingManager.cpp
--- a/Engine/Source/Editor/Sequencer/Private/SequencerActorBindingManager.cpp
+++ b/Engine/Source/Editor/Sequencer/Private/SequencerActorBindingManager.cpp
@@ -14,6 +14,7 @@
 #include "NotificationManager.h"
 #include "Engine/Selection.h"
 #include "LevelEditor.h"
+#include "MovieSceneBindings.h"
 
 #define LOCTEXT_NAMESPACE "Sequencer"
 
@@ -71,7 +72,19 @@ FGuid FSequencerActorBindingManager::FindGuidForObject( const UMovieScene& Movie
 			// When editing within the level editor, make sure we're bound to the level script node which contains data about possessables.
 			if( PlayMovieSceneNode.IsValid() )
 			{
-				ObjectGuid = PlayMovieSceneNode->FindGuidForObject( &Object );
+				UObject* Actor;
+				FString ComponentName;
+				UActorComponent* ComponentObject = Cast<UActorComponent>(&Object);
+				if (ComponentObject != nullptr)
+				{
+					Actor = ComponentObject->GetOuter();
+					ComponentName = ComponentObject->GetName();
+				}
+				else
+				{
+					Actor = &Object;
+				}
+				ObjectGuid = PlayMovieSceneNode->FindGuidForObjectInfo( FMovieSceneBoundObjectInfo(Actor, ComponentName) );
 			}
 		}
 
@@ -260,30 +273,44 @@ void FSequencerActorBindingManager::DestroyAllSpawnedObjects()
 
 bool FSequencerActorBindingManager::CanPossessObject( UObject& Object ) const
 {
-	return Object.IsA<AActor>();
+	return Object.IsA<AActor>() || Object.IsA<UActorComponent>();
 }
 
 void FSequencerActorBindingManager::BindPossessableObject( const FGuid& PossessableGuid, UObject& PossessedObject )
 {
-	 // When editing within the level editor, make sure we're bound to the level script node which contains data about possessables.
-	 const bool bCreateIfNotFound = true;
-	 BindToPlayMovieSceneNode( bCreateIfNotFound );
+	// When editing within the level editor, make sure we're bound to the level script node which contains data about possessables.
+	const bool bCreateIfNotFound = true;
+	BindToPlayMovieSceneNode( bCreateIfNotFound );
+	
+	//const FString& PossessableName = Object->GetName();
+	//const FGuid PossessableGuid = FocusedMovieScene->AddPossessable( PossessableName, Object->GetClass() );
+	
+	//if (IsShotFilteringOn())
+	//{
+	//	AddUnfilterableObject(PossessableGuid);
+	// }
+	
+	// Bind the object to the handle
+
+	UObject* Actor;
+	FString ComponentName;
+	UActorComponent* ComponentObject = Cast<UActorComponent>(&PossessedObject);
+	if (ComponentObject != nullptr)
+	{
+		Actor = ComponentObject->GetOuter();
+		ComponentName = ComponentObject->GetName();
+	}
+	else
+	{
+		Actor = &PossessedObject;
+	}
+
+	TArray< FMovieSceneBoundObjectInfo > ObjectInfos;
+	ObjectInfos.Add( FMovieSceneBoundObjectInfo(Actor, ComponentName) );
+	PlayMovieSceneNode->BindPossessableToObjects( PossessableGuid, ObjectInfos );
 	
-	 //const FString& PossessableName = Object->GetName();
-	 //const FGuid PossessableGuid = FocusedMovieScene->AddPossessable( PossessableName, Object->GetClass() );
-	 
-	 //if (IsShotFilteringOn())
-	 //{
-	 //	 AddUnfilterableObject(PossessableGuid);
-	 // }
-	 
-	 // Bind the object to the handle
-	 TArray< UObject* > Objects;
-	 Objects.Add( &PossessedObject );
-	 PlayMovieSceneNode->BindPossessableToObjects( PossessableGuid, Objects );
-	 
-	 // A possessable was created so we need to respawn its puppet
-	 // SpawnOrDestroyPuppetObjects( FocusedMovieSceneInstance );
+	// A possessable was created so we need to respawn its puppet
+	// SpawnOrDestroyPuppetObjects( FocusedMovieSceneInstance );
 }
 
 void FSequencerActorBindingManager::UnbindPossessableObjects( const FGuid& PossessableGuid )
@@ -310,7 +337,27 @@ void FSequencerActorBindingManager::GetRuntimeObjects( const TSharedRef<FMovieSc
 	
 		if (PlayMovieSceneNode.IsValid())
 		{
-			OutRuntimeObjects = PlayMovieSceneNode->FindBoundObjects( ObjectGuid );
+			for (FMovieSceneBoundObjectInfo& BoundObjectInfo : PlayMovieSceneNode->FindBoundObjectInfos(ObjectGuid))
+			{
+				if (BoundObjectInfo.Tag.IsEmpty() == false)
+				{
+					AActor* Actor = Cast<AActor>(BoundObjectInfo.Object);
+					if (Actor != nullptr)
+					{
+						for (UActorComponent* ActorComponent : Actor->GetComponents())
+						{
+							if (ActorComponent->GetName() == BoundObjectInfo.Tag)
+							{
+								OutRuntimeObjects.Add(ActorComponent);
+							}
+						}
+					}
+				}
+				else
+				{
+					OutRuntimeObjects.Add(BoundObjectInfo.Object);
+				}
+			}
 		}
 	}
 }
diff --git a/Engine/Source/Editor/Sequencer/Private/SequencerObjectChangeListener.cpp b/Engine/Source/Editor/Sequencer/Private/SequencerObjectChangeListener.cpp
--- a/Engine/Source/Editor/Sequencer/Private/SequencerObjectChangeListener.cpp
+++ b/Engine/Source/Editor/Sequencer/Private/SequencerObjectChangeListener.cpp
@@ -80,7 +80,7 @@ bool FSequencerObjectChangeListener::IsObjectValidForListening( UObject* Object
 {
 	// @todo Sequencer - Pre/PostEditChange is sometimes called for inner objects of other objects (like actors with components)
 	// We only want the outer object so assume it's an actor for now
-	return bListenForActorsOnly ? Object->IsA<AActor>() : true;
+	return bListenForActorsOnly ? Object->IsA<AActor>() || Object->IsA<UActorComponent>() : true;
 }
 
 FOnAnimatablePropertyChanged& FSequencerObjectChangeListener::GetOnAnimatablePropertyChanged( FName PropertyTypeName )
diff --git a/Engine/Source/Runtime/MovieScene/Private/MovieSceneBindings.cpp b/Engine/Source/Runtime/MovieScene/Private/MovieSceneBindings.cpp
--- a/Engine/Source/Runtime/MovieScene/Private/MovieSceneBindings.cpp
+++ b/Engine/Source/Runtime/MovieScene/Private/MovieSceneBindings.cpp
@@ -3,9 +3,9 @@
 #include "MovieScenePrivatePCH.h"
 #include "MovieSceneBindings.h"
 
-FMovieSceneBoundObject::FMovieSceneBoundObject( const FGuid& InitPosessableGuid, const TArray< UObject* >& InitObjects )
+FMovieSceneBoundObject::FMovieSceneBoundObject(const FGuid& InitPosessableGuid, const TArray< FMovieSceneBoundObjectInfo >& InitObjectInfos)
 	: PossessableGuid( InitPosessableGuid ),
-	  Objects( InitObjects )
+	  ObjectInfos( InitObjectInfos )
 {
 }
 
@@ -36,15 +36,15 @@ void UMovieSceneBindings::SetRootMovieScene( class UMovieScene* NewMovieScene )
 }
 
 
-FMovieSceneBoundObject& UMovieSceneBindings::AddBinding( const FGuid& PosessableGuid, const TArray< UObject* >& Objects )
+FMovieSceneBoundObject& UMovieSceneBindings::AddBinding(const FGuid& PosessableGuid, const TArray< FMovieSceneBoundObjectInfo >& ObjectInfos)
 {
 	Modify();
-	BoundObjects.Add( FMovieSceneBoundObject( PosessableGuid, Objects ) );
+	BoundObjects.Add( FMovieSceneBoundObject( PosessableGuid, ObjectInfos ) );
 	return BoundObjects.Last();
 }
 
 
-TArray< UObject* > UMovieSceneBindings::FindBoundObjects( const FGuid& Guid ) const
+TArray< FMovieSceneBoundObjectInfo > UMovieSceneBindings::FindBoundObjects( const FGuid& Guid ) const
 {
 	for( TArray< FMovieSceneBoundObject >::TConstIterator BoundObjectIter( BoundObjects.CreateConstIterator() ); BoundObjectIter; ++BoundObjectIter )
 	{
@@ -52,11 +52,11 @@ TArray< UObject* > UMovieSceneBindings::FindBoundObjects( const FGuid& Guid ) co
 
 		if( BoundObject.GetPossessableGuid() == Guid )
 		{
-			return BoundObject.GetObjects();
+			return BoundObject.GetObjectInfos();
 		}
 	}
 
-	return TArray< UObject* >();
+	return TArray< FMovieSceneBoundObjectInfo >();
 }
 
 
diff --git a/Engine/Source/Runtime/MovieScene/Private/RuntimeMovieScenePlayer.cpp b/Engine/Source/Runtime/MovieScene/Private/RuntimeMovieScenePlayer.cpp
--- a/Engine/Source/Runtime/MovieScene/Private/RuntimeMovieScenePlayer.cpp
+++ b/Engine/Source/Runtime/MovieScene/Private/RuntimeMovieScenePlayer.cpp
@@ -198,7 +198,27 @@ void URuntimeMovieScenePlayer::GetRuntimeObjects( TSharedRef<FMovieSceneInstance
 	// Otherwise, check to see if we have one or more possessed actors that are mapped to this handle
 	else if( MovieSceneBindings != NULL )
 	{
-		OutObjects = MovieSceneBindings->FindBoundObjects( ObjectHandle );
+		for (FMovieSceneBoundObjectInfo& BoundObject : MovieSceneBindings->FindBoundObjects(ObjectHandle))
+		{
+			if (BoundObject.Tag.IsEmpty() == false)
+			{
+				AActor* Actor = Cast<AActor>(BoundObject.Object);
+				if (Actor != nullptr)
+				{
+					for (UActorComponent* ActorComponent : Actor->GetComponents())
+					{
+						if (ActorComponent->GetName() == BoundObject.Tag)
+						{
+							OutObjects.Add(ActorComponent);
+						}
+					}
+				}
+			}
+			else
+			{
+				OutObjects.Add(BoundObject.Object);
+			}
+		}
 	}
 }
 
diff --git a/Engine/Source/Runtime/MovieScene/Public/MovieSceneBindings.h b/Engine/Source/Runtime/MovieScene/Public/MovieSceneBindings.h
--- a/Engine/Source/Runtime/MovieScene/Public/MovieSceneBindings.h
+++ b/Engine/Source/Runtime/MovieScene/Public/MovieSceneBindings.h
@@ -4,6 +4,42 @@
 
 #include "MovieSceneBindings.generated.h"
 
+/**
+ * Identifies a bound object and an optional tag for sub-object lookup.
+ */
+USTRUCT()
+struct FMovieSceneBoundObjectInfo
+{
+	GENERATED_USTRUCT_BODY( FMovieSceneBoundObjectInfo )
+
+public:
+	FMovieSceneBoundObjectInfo()
+	{
+		Object = nullptr;
+	}
+
+	FMovieSceneBoundObjectInfo(UObject* InObject)
+	{
+		Object = InObject;
+	}
+
+	FMovieSceneBoundObjectInfo(UObject* InObject, FString InTag)
+	{
+		Object = InObject;
+		Tag = InTag;
+	}
+
+	inline bool operator==(const FMovieSceneBoundObjectInfo& Other) const
+	{
+		return (Object == Other.Object) && (Tag == Other.Tag);
+	}
+
+	UPROPERTY()
+	UObject* Object;
+
+	UPROPERTY()
+	FString Tag;
+};
 
 /**
  * MovieSceneBoundObject connects an object to a "possessable" slot in the MovieScene asset
@@ -24,9 +60,9 @@ public:
 	 * FMovieSceneBoundObject initialization constructor.  Binds the specified possessable from a MovieScene asset, to the specified object
 	 *
 	 * @param	InitPossessableGuid		The Guid of the possessable from the MovieScene asset to bind to
-	 * @param	InitObjects				The objects to bind
+	 * @param	InitObjectInfos			The object infos to bind
 	 */
-	FMovieSceneBoundObject( const FGuid& InitPosessableGuid, const TArray< UObject* >& InitObjects );
+	FMovieSceneBoundObject( const FGuid& InitPosessableGuid, const TArray< FMovieSceneBoundObjectInfo >& InitObjectInfos );
 
 	/** @return Returns the guid of the possessable that we're bound to */
 	const FGuid& GetPossessableGuid() const
@@ -35,36 +71,25 @@ public:
 	}
 
 	/** @return Returns the objects we're bound to */
-	TArray< UObject* >& GetObjects()
+	TArray< FMovieSceneBoundObjectInfo >& GetObjectInfos()
 	{
-		return Objects;
+		return ObjectInfos;
 	}
 
 	/** @return Returns the objects we're bound to */
-	const TArray< UObject* >& GetObjects() const
-	{
-		return Objects;
-	}
-
-	/**
-	 * Returns whether or not an object is bound
-	 *
-	 * @param Object	The object to search for
-	 * @return true if the passed in object is bound, false otherwise 
-	 */
-	bool ContainsObject( UObject* Object )
+	const TArray< FMovieSceneBoundObjectInfo >& GetObjectInfos() const
 	{
-		return Objects.Contains( Object );	
+		return ObjectInfos;
 	}
 
 	/**
 	 * Sets the actual objects bound to this slot.  Remember to call Modify() on the owner object if you change this!
 	 * 
-	 * @param	NewObjects	The list of new bound object (can be NULL)
+	 * @param	NewObjectInfos	The list of new bound object infos (can be NULL)
 	 */
-	void SetObjects( TArray< UObject* > NewObjects )
+	void SetObjectInfos( TArray< FMovieSceneBoundObjectInfo > NewObjectInfos )
 	{
-		Objects = NewObjects;
+		ObjectInfos = NewObjectInfos;
 	}
 
 private:
@@ -76,7 +101,7 @@ private:
 
 	/** The actual bound objects (we support binding multiple objects to the same possessable) */
 	UPROPERTY()
-	TArray<	UObject* > Objects;
+	TArray< FMovieSceneBoundObjectInfo > ObjectInfos;
 
 };
 
@@ -109,7 +134,7 @@ public:
 	 *
 	 * @return	Reference to the wrapper structure for this bound object
 	 */
-	virtual FMovieSceneBoundObject& AddBinding( const FGuid& PosessableGuid, const TArray< UObject* >& Objects );
+	virtual FMovieSceneBoundObject& AddBinding( const FGuid& PosessableGuid, const TArray< FMovieSceneBoundObjectInfo >& ObjectInfos );
 
 	/**
 	 * Given a guid for a possessable, attempts to find the object bound to that guid
@@ -118,7 +143,7 @@ public:
 	 *
 	 * @return	The bound object, if found, otherwise NULL
 	 */
-	virtual TArray< UObject* > FindBoundObjects( const FGuid& Guid ) const;
+	virtual TArray< FMovieSceneBoundObjectInfo > FindBoundObjects( const FGuid& Guid ) const;
 
 	/**
 	 * Gets all bound objects
