Sequencer - Fix PlayMovieSceneNode component name pin issues.

There are still some kinks to work out, but this gets things usable again.

diff --git a/Engine/Source/Editor/MovieSceneTools/Private/K2Node_PlayMovieScene.cpp b/Engine/Source/Editor/MovieSceneTools/Private/K2Node_PlayMovieScene.cpp
--- a/Engine/Source/Editor/MovieSceneTools/Private/K2Node_PlayMovieScene.cpp
+++ b/Engine/Source/Editor/MovieSceneTools/Private/K2Node_PlayMovieScene.cpp
@@ -6,6 +6,7 @@
 #include "BlueprintUtilities.h"
 #include "MovieScene.h"
 #include "Kismet2/BlueprintEditorUtils.h"
+#include "Kismet/KismetSystemLibrary.h"
 #include "MovieSceneBindings.h"
 
 UEdGraphNode* FEdGraphSchemaAction_K2AddPlayMovieScene::PerformAction(class UEdGraph* ParentGraph, UEdGraphPin* FromPin, const FVector2D Location, bool bSelectNewNode/* = true*/)
@@ -136,11 +137,10 @@ void UK2Node_PlayMovieScene::PinConnectionListChanged( UEdGraphPin* InPin )
 			UEdGraphPin* ComponentNamePin = FindPinChecked(BoundObject.GetPossessableGuid().ToString(EGuidFormats::DigitsWithHyphens) + "_ComponentName");
 
 			// Is the pin even linked to anything?
-			checkf(ActorPin->LinkedTo.Num() == ComponentNamePin->LinkedTo.Num(), TEXT("Actor pin link count must match the component name link count"));
 			for (int32 i = 0; i < ActorPin->LinkedTo.Num(); i++)
 			{
 				UEdGraphPin* LinkedActorPin = ActorPin->LinkedTo[i];
-				UEdGraphPin* LinkedComponentNamePin = ComponentNamePin->LinkedTo[i];
+				UEdGraphPin* LinkedComponentNamePin = ComponentNamePin->LinkedTo.Num() > i ? ComponentNamePin->LinkedTo[i] : NULL;
 				if( LinkedActorPin != NULL )
 				{
 					// Look for an object bound to a literal.  We can use these for scrub preview in the editor!
@@ -154,12 +154,12 @@ void UK2Node_PlayMovieScene::PinConnectionListChanged( UEdGraphPin* InPin )
 						if ( LinkedActor != NULL )
 						{
 							FString LinkedComponentName;
-							if (LinkedComponentNamePin)
+							if (LinkedComponentNamePin != NULL)
 							{
-								UK2Node_Literal* LinkedComponentNameLiteralNode = Cast< UK2Node_Literal >(LinkedComponentNamePin->GetOwningNode());
+								UK2Node_CallFunction* LinkedComponentNameLiteralNode = Cast< UK2Node_CallFunction >(LinkedComponentNamePin->GetOwningNode());
 								if (LinkedComponentNameLiteralNode != NULL)
 								{
-									LinkedComponentName = LinkedComponentNameLiteralNode->GetValuePin()->GetDefaultAsString();
+									LinkedComponentName = LinkedComponentNameLiteralNode->FindPinChecked("Value")->DefaultValue;
 								}
 							}
 							OtherObjectInfos.Add( FMovieSceneBoundObjectInfo(LinkedActor, LinkedComponentName));
@@ -303,7 +303,7 @@ void UK2Node_PlayMovieScene::CreatePinsForBoundObject( FMovieSceneBoundObject& B
 	const bool bIsArray = false;
 	const bool bIsReference = false;
 	UEdGraphPin* NewActorPin = CreatePin( EGPD_Input, K2Schema->PC_Object, PinSubCategory, PinSubCategoryObject, bIsArray, bIsReference, ActorPinName );
-	UEdGraphPin* NewComponentNamePin = CreatePin(EGPD_Input, K2Schema->PC_String, PinSubCategory, PinSubCategoryObject, bIsArray, bIsReference, ComponentNamePinName );
+	UEdGraphPin* NewComponentNamePin = CreatePin(EGPD_Input, K2Schema->PC_String, PinSubCategory, NULL, bIsArray, bIsReference, ComponentNamePinName );
 	check( NewActorPin != NULL && NewComponentNamePin != NULL );
 
 	// Set the friendly name for this pin
@@ -321,33 +321,34 @@ void UK2Node_PlayMovieScene::CreatePinsForBoundObject( FMovieSceneBoundObject& B
 			{
 				// Check to see if we have a literal for this object already
 				UK2Node_Literal* ActorLiteralNode = NULL;
-				UK2Node_Literal* ComponentNameLiteralNode = NULL;
+				UK2Node_CallFunction* ComponentNameLiteralNode = NULL;
 				{
 					TArray< UK2Node_Literal* > LiteralNodes;
 					GetGraph()->GetNodesOfClass( LiteralNodes );
-					for( auto NodeIt( LiteralNodes.CreateConstIterator() ); NodeIt; ++NodeIt )
+					for(UK2Node_Literal* LiteralNode : LiteralNodes)
 					{
-						const auto CurLiteralNode = *NodeIt;
-
-						if( CurLiteralNode->GetObjectRef() == BoundObjectInfo.Object )
+						if( LiteralNode->GetObjectRef() == BoundObjectInfo.Object )
 						{
 							// Found one!
-							ActorLiteralNode = CurLiteralNode;
+							ActorLiteralNode = LiteralNode;
+							break;
 						}
-						else
+					}
+
+					TArray< UK2Node_CallFunction* > CallFunctionNodes;
+					GetGraph()->GetNodesOfClass(CallFunctionNodes);
+					for (UK2Node_CallFunction* CallFunctionNode : CallFunctionNodes)
+					{
+						if (CallFunctionNode->FunctionReference.GetMemberName() == FName("MakeLiteralString"))
 						{
-							UEdGraphPin* ValuePin = CurLiteralNode->GetValuePin();
-							if (ValuePin->PinType.PinCategory == K2Schema->PC_String && ValuePin->GetDefaultAsString() == BoundObjectInfo.Tag)
+							UEdGraphPin* ValuePin = CallFunctionNode->FindPinChecked("Value");
+							if (ValuePin->DefaultValue == BoundObjectInfo.Tag)
 							{
-								ComponentNameLiteralNode = CurLiteralNode;
+								ComponentNameLiteralNode = CallFunctionNode;
+								break;
 							}
 						}
 					}
-
-					if (ActorLiteralNode != NULL && ComponentNameLiteralNode != NULL)
-					{
-						break;
-					}
 				}
 
 				if( ActorLiteralNode == NULL )
@@ -369,17 +370,18 @@ void UK2Node_PlayMovieScene::CreatePinsForBoundObject( FMovieSceneBoundObject& B
 				{
 					if (ComponentNameLiteralNode == NULL)
 					{
-						// No literal for this object yet, so we'll make one now.
-						UK2Node_Literal* LiteralNodeTemplate = NewObject<UK2Node_Literal>();
+						UK2Node_CallFunction* MakeLiteralStringNodeTemplate = NewObject<UK2Node_CallFunction>();
+						MakeLiteralStringNodeTemplate->FunctionReference.SetExternalMember("MakeLiteralString", UKismetSystemLibrary::StaticClass());
 
-						// Figure out a decent place to stick the node
 						// @todo sequencer: The placement of these is really unacceptable.  We should setup new logic specific for moviescene pin objects.
 						const FVector2D NewNodePos = GetGraph()->GetGoodPlaceForNewNode();
-						ComponentNameLiteralNode = FEdGraphSchemaAction_K2NewNode::SpawnNodeFromTemplate<UK2Node_Literal>(GetGraph(), LiteralNodeTemplate, NewNodePos);
-						ComponentNameLiteralNode->GetSchema()->TrySetDefaultValue(*ComponentNameLiteralNode->GetValuePin(), BoundObjectInfo.Tag);
+						ComponentNameLiteralNode = FEdGraphSchemaAction_K2NewNode::SpawnNodeFromTemplate<UK2Node_CallFunction>(GetGraph(), MakeLiteralStringNodeTemplate, NewNodePos);
+
+						UEdGraphPin* ValuePin = ComponentNameLiteralNode->FindPinChecked(FString("Value"));
+						ValuePin->DefaultValue = BoundObjectInfo.Tag;
 					}
 
-					ComponentNameLiteralNode->GetValuePin()->MakeLinkTo(NewComponentNamePin);
+					ComponentNameLiteralNode->GetReturnValuePin()->MakeLinkTo(NewComponentNamePin);
 				}
 			}
 		}
