Sequencer: Add initial path track.

- Path track editor.
- Path track section that evaluates to a usplinecomponent of an aactor.
- Add evaluation order to track instances so that path can evaluate before transform.

diff --git a/Engine/Source/Editor/MovieSceneTools/Private/MovieSceneToolsModule.cpp b/Engine/Source/Editor/MovieSceneTools/Private/MovieSceneToolsModule.cpp
--- a/Engine/Source/Editor/MovieSceneTools/Private/MovieSceneToolsModule.cpp
+++ b/Engine/Source/Editor/MovieSceneTools/Private/MovieSceneToolsModule.cpp
@@ -23,6 +23,7 @@
 #include "AudioTrackEditor.h"
 #include "AnimationTrackEditor.h"
 #include "ParticleTrackEditor.h"
+#include "PathTrackEditor.h"
 
 /**
  * MovieSceneTools module implementation (private)
@@ -44,6 +45,7 @@ class FMovieSceneToolsModule : public IMovieSceneTools
 		AudioTrackEditorCreateTrackEditorDelegateHandle         = SequencerModule.RegisterTrackEditor_Handle( FOnCreateTrackEditor::CreateStatic( &FAudioTrackEditor::CreateTrackEditor ) );
 		AnimationTrackEditorCreateTrackEditorDelegateHandle     = SequencerModule.RegisterTrackEditor_Handle( FOnCreateTrackEditor::CreateStatic( &FAnimationTrackEditor::CreateTrackEditor ) );
 		ParticleTrackEditorCreateTrackEditorDelegateHandle      = SequencerModule.RegisterTrackEditor_Handle( FOnCreateTrackEditor::CreateStatic( &FParticleTrackEditor::CreateTrackEditor ) );
+		PathTrackEditorCreateTrackEditorDelegateHandle          = SequencerModule.RegisterTrackEditor_Handle( FOnCreateTrackEditor::CreateStatic( &F3DPathTrackEditor::CreateTrackEditor ) );
 	}
 
 
@@ -69,6 +71,7 @@ class FMovieSceneToolsModule : public IMovieSceneTools
 			SequencerModule.UnRegisterTrackEditor_Handle( AudioTrackEditorCreateTrackEditorDelegateHandle );
 			SequencerModule.UnRegisterTrackEditor_Handle( AnimationTrackEditorCreateTrackEditorDelegateHandle );
 			SequencerModule.UnRegisterTrackEditor_Handle( ParticleTrackEditorCreateTrackEditorDelegateHandle );
+			SequencerModule.UnRegisterTrackEditor_Handle( PathTrackEditorCreateTrackEditorDelegateHandle );
 		}
 		
 	}
@@ -87,6 +90,7 @@ private:
 	FDelegateHandle AudioTrackEditorCreateTrackEditorDelegateHandle;
 	FDelegateHandle AnimationTrackEditorCreateTrackEditorDelegateHandle;
 	FDelegateHandle ParticleTrackEditorCreateTrackEditorDelegateHandle;
+	FDelegateHandle PathTrackEditorCreateTrackEditorDelegateHandle;
 };
 
 
diff --git a/Engine/Source/Editor/MovieSceneTools/Private/TrackEditors/PathTrackEditor.cpp b/Engine/Source/Editor/MovieSceneTools/Private/TrackEditors/PathTrackEditor.cpp
new file mode 100644
--- /dev/null
+++ b/Engine/Source/Editor/MovieSceneTools/Private/TrackEditors/PathTrackEditor.cpp
@@ -0,0 +1,349 @@
+// Copyright 1998-2015 Epic Games, Inc. All Rights Reserved.
+
+#include "MovieSceneToolsPrivatePCH.h"
+#include "ScopedTransaction.h"
+#include "MovieScene.h"
+#include "MovieSceneSection.h"
+#include "MovieScene3DPathTrack.h"
+#include "MovieScene3DPathSection.h"
+#include "ISequencerObjectBindingManager.h"
+#include "ISequencerSection.h"
+#include "ISectionLayoutBuilder.h"
+#include "IKeyArea.h"
+#include "MovieSceneToolHelpers.h"
+#include "MovieSceneTrackEditor.h"
+#include "PathTrackEditor.h"
+#include "Components/SplineComponent.h"
+
+void GetActorsWithSplineComponents(TSet<AActor*>& SelectedActors, TSet<AActor*>& UnselectedActors)
+{
+	if (GEditor->GetEditorWorldContext().World())
+	{
+		for( FActorIterator	ActorIt(GEditor->GetEditorWorldContext().World()); ActorIt; ++ActorIt )
+		{
+			AActor* Actor = *ActorIt;
+			TArray<USplineComponent*> SplineComponents;
+			Actor->GetComponents(SplineComponents);
+			if (SplineComponents.Num())
+			{
+				if (GEditor->GetSelectedActors()->IsSelected(Actor))
+				{
+					SelectedActors.Add(Actor);
+				}
+				else
+				{
+					UnselectedActors.Add(Actor);
+				}
+			}
+		}
+	}
+}
+
+/**
+ * Class that draws a path section in the sequencer
+ */
+class F3DPathSection : public ISequencerSection
+{
+public:
+	F3DPathSection( UMovieSceneSection& InSection, F3DPathTrackEditor* InPathTrackEditor )
+		: Section( InSection )
+		, PathTrackEditor(InPathTrackEditor)
+	{
+	}
+
+	/** ISequencerSection interface */
+	virtual UMovieSceneSection* GetSectionObject() override
+	{ 
+		return &Section;
+	}
+
+	virtual FText GetDisplayName() const override
+	{ 
+		return NSLOCTEXT("FPathSection", "DisplayName", "Path");
+	}
+	
+	virtual FText GetSectionTitle() const override { return FText::GetEmpty(); }
+
+	virtual void GenerateSectionLayout( class ISectionLayoutBuilder& LayoutBuilder ) const override
+	{
+		UMovieScene3DPathSection* PathSection = Cast<UMovieScene3DPathSection>( &Section );
+
+		LayoutBuilder.AddKeyArea("Timing", NSLOCTEXT("FPathSection", "TimingArea", "Timing"), MakeShareable( new FFloatCurveKeyArea ( &PathSection->GetTimingCurve( ), PathSection ) ) );
+	}
+
+	virtual int32 OnPaintSection( const FGeometry& AllottedGeometry, const FSlateRect& SectionClippingRect, FSlateWindowElementList& OutDrawElements, int32 LayerId, bool bParentEnabled ) const override 
+	{
+		// Add a box for the section
+		FSlateDrawElement::MakeBox( 
+			OutDrawElements,
+			LayerId,
+			AllottedGeometry.ToPaintGeometry(),
+			FEditorStyle::GetBrush("Sequencer.GenericSection.Background"),
+			SectionClippingRect
+		); 
+
+		return LayerId;
+	}
+	
+	virtual void BuildSectionContextMenu(FMenuBuilder& MenuBuilder) override
+	{
+		TSet<AActor*> SelectedActors;
+		TSet<AActor*> UnselectedActors;
+		GetActorsWithSplineComponents(SelectedActors, UnselectedActors);	
+
+		if (SelectedActors.Num()+UnselectedActors.Num() == 0)
+		{
+			return;
+		}
+
+		if (SelectedActors.Num())
+		{
+			MenuBuilder.BeginSection(NAME_None, FText::FromString("Selected Splines"));
+
+			for (TSet<AActor*>::TIterator It(SelectedActors); It; ++It)
+			{
+				AActor* ActorWithSplineComponent = *It;
+				FFormatNamedArguments Args;
+				Args.Add( TEXT("PathName"), FText::FromString( ActorWithSplineComponent->GetActorLabel() ) );
+
+				MenuBuilder.AddMenuEntry(
+					FText::Format( NSLOCTEXT("Sequencer", "SetPath", "Set {PathName}"), Args ),
+					FText::Format( NSLOCTEXT("Sequencer", "SetPathTooltip", "Set path track driven by {PathName}."), Args ),
+					FSlateIcon(),
+					FUIAction(FExecuteAction::CreateSP(PathTrackEditor, &F3DPathTrackEditor::SetPath, &Section, ActorWithSplineComponent))
+				);
+			}
+			
+			MenuBuilder.EndSection();
+		}
+
+		if (UnselectedActors.Num())
+		{
+			MenuBuilder.BeginSection(NAME_None, FText::FromString("Splines"));
+
+			for (TSet<AActor*>::TIterator It(UnselectedActors); It; ++It)
+			{
+				AActor* ActorWithSplineComponent = *It;
+				FFormatNamedArguments Args;
+				Args.Add( TEXT("PathName"), FText::FromString( ActorWithSplineComponent->GetActorLabel() ) );
+
+				MenuBuilder.AddMenuEntry(
+					FText::Format( NSLOCTEXT("Sequencer", "SetPath", "Set {PathName}"), Args ),
+					FText::Format( NSLOCTEXT("Sequencer", "SetPathTooltip", "Set path track driven by {PathName}."), Args ),
+					FSlateIcon(),
+					FUIAction(FExecuteAction::CreateSP(PathTrackEditor, &F3DPathTrackEditor::SetPath, &Section, ActorWithSplineComponent))
+				);
+			}
+
+			MenuBuilder.EndSection();
+		}
+
+		MenuBuilder.AddMenuSeparator();
+	}
+
+private:
+	/** The section we are visualizing */
+	UMovieSceneSection& Section;
+
+	/** The path track editor */
+	F3DPathTrackEditor* PathTrackEditor;
+};
+
+F3DPathTrackEditor::F3DPathTrackEditor( TSharedRef<ISequencer> InSequencer )
+	: FMovieSceneTrackEditor( InSequencer ) 
+{
+}
+
+F3DPathTrackEditor::~F3DPathTrackEditor()
+{
+}
+
+TSharedRef<FMovieSceneTrackEditor> F3DPathTrackEditor::CreateTrackEditor( TSharedRef<ISequencer> InSequencer )
+{
+	return MakeShareable( new F3DPathTrackEditor( InSequencer ) );
+}
+
+bool F3DPathTrackEditor::SupportsType( TSubclassOf<UMovieSceneTrack> Type ) const
+{
+	// We support animatable transforms
+	return Type == UMovieScene3DPathTrack::StaticClass();
+}
+
+TSharedRef<ISequencerSection> F3DPathTrackEditor::MakeSectionInterface( UMovieSceneSection& SectionObject, UMovieSceneTrack* Track )
+{
+	check( SupportsType( SectionObject.GetOuter()->GetClass() ) );
+
+	TSharedRef<ISequencerSection> NewSection( new F3DPathSection( SectionObject, this ) );
+
+	return NewSection;
+}
+
+void F3DPathTrackEditor::AddKey( const FGuid& ObjectGuid, UObject* AdditionalAsset )
+{
+	AddPath(ObjectGuid, AdditionalAsset);
+}
+
+void F3DPathTrackEditor::BuildObjectBindingContextMenu(FMenuBuilder& MenuBuilder, const FGuid& ObjectBinding, const UClass* ObjectClass)
+{
+	// build a menu with all the actors that have spline components, add the selected actors first
+	if (ObjectClass->IsChildOf(AActor::StaticClass()))
+	{
+		TSet<AActor*> SelectedActors;
+		TSet<AActor*> UnselectedActors;
+		GetActorsWithSplineComponents(SelectedActors, UnselectedActors);
+
+		if (SelectedActors.Num()+UnselectedActors.Num() > 0)
+		{
+			MenuBuilder.AddSubMenu(
+				NSLOCTEXT("Sequencer", "AddPath", "Add Path"), NSLOCTEXT("Sequencer", "AddPathTooltip", "Adds a path track."),
+				FNewMenuDelegate::CreateRaw(this, &F3DPathTrackEditor::AddPathSubMenu, ObjectBinding));
+		}
+		else
+		{
+			AActor* ActorWithSplineComponent = NULL;
+			FFormatNamedArguments Args;
+			MenuBuilder.AddMenuEntry(
+				FText::Format( NSLOCTEXT("Sequencer", "AddPath", "Add Path"), Args),
+				FText::Format( NSLOCTEXT("Sequencer", "AddPathTooltip", "Adds a path track."), Args ),
+				FSlateIcon(),
+				FUIAction(FExecuteAction::CreateSP(this, &F3DPathTrackEditor::AddPath, ObjectBinding, (UObject*)ActorWithSplineComponent))
+			);
+		}
+	}
+}
+
+void F3DPathTrackEditor::AddPathSubMenu(FMenuBuilder& MenuBuilder, FGuid ObjectBinding)
+{
+	TSet<AActor*> SelectedActors;
+	TSet<AActor*> UnselectedActors;
+	GetActorsWithSplineComponents(SelectedActors, UnselectedActors);	
+
+	if (SelectedActors.Num()+UnselectedActors.Num() == 0)
+	{
+		return;
+	}
+
+	if (SelectedActors.Num())
+	{
+		MenuBuilder.BeginSection(NAME_None, FText::FromString("Selected Splines"));
+
+		for (TSet<AActor*>::TIterator It(SelectedActors); It; ++It)
+		{
+			AActor* ActorWithSplineComponent = *It;
+			FFormatNamedArguments Args;
+			Args.Add( TEXT("PathName"), FText::FromString( ActorWithSplineComponent->GetActorLabel() ) );
+
+			MenuBuilder.AddMenuEntry(
+				FText::Format( NSLOCTEXT("Sequencer", "AddPath", "Add {PathName}"), Args ),
+				FText::Format( NSLOCTEXT("Sequencer", "AddPathTooltip", "Add path track driven by {PathName}."), Args ),
+				FSlateIcon(),
+				FUIAction(FExecuteAction::CreateSP(this, &F3DPathTrackEditor::AddPath, ObjectBinding, (UObject*)ActorWithSplineComponent))
+			);
+		}
+			
+		MenuBuilder.EndSection();
+	}
+
+	if (UnselectedActors.Num())
+	{
+		MenuBuilder.BeginSection(NAME_None, FText::FromString("Splines"));
+
+		for (TSet<AActor*>::TIterator It(UnselectedActors); It; ++It)
+		{
+			AActor* ActorWithSplineComponent = *It;
+			FFormatNamedArguments Args;
+			Args.Add( TEXT("PathName"), FText::FromString( ActorWithSplineComponent->GetActorLabel() ) );
+
+			MenuBuilder.AddMenuEntry(
+				FText::Format( NSLOCTEXT("Sequencer", "AddPath", "Add {PathName}"), Args ),
+				FText::Format( NSLOCTEXT("Sequencer", "AddPathTooltip", "Add path track driven by {PathName}."), Args ),
+				FSlateIcon(),
+				FUIAction(FExecuteAction::CreateSP(this, &F3DPathTrackEditor::AddPath, ObjectBinding, (UObject*)ActorWithSplineComponent))
+			);
+		}
+
+		MenuBuilder.EndSection();
+	}
+}
+
+void F3DPathTrackEditor::AddPath(FGuid ObjectGuid, UObject* AdditionalAsset)
+{
+	TArray<UObject*> OutObjects;
+	GetSequencer()->GetRuntimeObjects( GetSequencer()->GetFocusedMovieSceneInstance(), ObjectGuid, OutObjects);
+
+	AnimatablePropertyChanged( UMovieScene3DPathTrack::StaticClass(), false,
+		FOnKeyProperty::CreateRaw( this, &F3DPathTrackEditor::AddKeyInternal, OutObjects, AdditionalAsset) );
+}
+
+void F3DPathTrackEditor::SetPath(UMovieSceneSection* Section, AActor* ActorWithSplineComponent)
+{
+	const FScopedTransaction Transaction(NSLOCTEXT("Sequencer", "UndoSetPath", "Set Path"));
+
+	UMovieScene3DPathSection* PathSection = (UMovieScene3DPathSection*)(Section);
+
+	//@todo sequencer - replace this with movie scene object id stuff
+	FGuid SplineGuid = FindOrCreateHandleToObject(ActorWithSplineComponent);
+
+	FMovieSceneObjectId SplineId;
+	SplineId.ObjectGuid = SplineGuid;
+
+	if (SplineGuid.IsValid())
+	{
+		PathSection->SetPathId(SplineId);
+	}
+}
+
+void F3DPathTrackEditor::AddKeyInternal( float KeyTime, const TArray<UObject*> Objects, UObject* AdditionalAsset)
+{
+	AActor* ActorWithSplineComponent = NULL;
+	
+	if (AdditionalAsset != NULL)
+	{
+		ActorWithSplineComponent = Cast<AActor>(AdditionalAsset);
+	}
+
+	//@todo sequencer - replace this with movie scene object id stuff
+	FGuid SplineGuid;
+	if (ActorWithSplineComponent != NULL)
+	{
+		SplineGuid = FindOrCreateHandleToObject(ActorWithSplineComponent);
+	}
+
+	if (SplineGuid.IsValid())
+	{
+		FMovieSceneObjectId SplineId;
+		SplineId.ObjectGuid = SplineGuid;
+
+		for( int32 ObjectIndex = 0; ObjectIndex < Objects.Num(); ++ObjectIndex )
+		{
+			UObject* Object = Objects[ObjectIndex];
+
+			FGuid ObjectHandle = FindOrCreateHandleToObject( Object );
+			if (ObjectHandle.IsValid())
+			{
+				UMovieSceneTrack* Track = GetTrackForObject( ObjectHandle, UMovieScene3DPathTrack::StaticClass(), FName("Path"));
+
+				if (ensure(Track))
+				{
+					// Clamp to next path section's start time or the end of the current sequencer view range
+					float PathEndTime = GetSequencer()->GetViewRange().GetUpperBoundValue();
+	
+					for (int32 PathSectionIndex = 0; PathSectionIndex < Track->GetAllSections().Num(); ++PathSectionIndex)
+					{
+						float StartTime = Track->GetAllSections()[PathSectionIndex]->GetStartTime();
+						float EndTime = Track->GetAllSections()[PathSectionIndex]->GetEndTime();
+						if (KeyTime < StartTime)
+						{
+							if (PathEndTime > StartTime)
+							{
+								PathEndTime = StartTime;
+							}
+						}
+					}
+
+					Cast<UMovieScene3DPathTrack>(Track)->AddPath( KeyTime, PathEndTime, SplineId );
+				}
+			}
+		}
+	}
+}
diff --git a/Engine/Source/Editor/MovieSceneTools/Private/TrackEditors/PathTrackEditor.h b/Engine/Source/Editor/MovieSceneTools/Private/TrackEditors/PathTrackEditor.h
new file mode 100644
--- /dev/null
+++ b/Engine/Source/Editor/MovieSceneTools/Private/TrackEditors/PathTrackEditor.h
@@ -0,0 +1,50 @@
+// Copyright 1998-2015 Epic Games, Inc. All Rights Reserved.
+
+#pragma once
+
+#include "MovieScene3DPathSection.h"
+#include "MovieScene3DPathTrack.h"
+/**
+ * Tools for animating a transform with a path curve
+ */
+class F3DPathTrackEditor : public FMovieSceneTrackEditor
+{
+public:
+	/**
+	 * Constructor
+	 *
+	 * @param InSequencer	The sequencer instance to be used by this tool
+	 */
+	F3DPathTrackEditor( TSharedRef<ISequencer> InSequencer );
+	~F3DPathTrackEditor();
+
+	/**
+	 * Creates an instance of this class.  Called by a sequencer 
+	 *
+	 * @param OwningSequencer The sequencer instance to be used by this tool
+	 * @return The new instance of this class
+	 */
+	static TSharedRef<FMovieSceneTrackEditor> CreateTrackEditor( TSharedRef<ISequencer> OwningSequencer );
+
+	/** FMovieSceneTrackEditor Interface */
+	virtual bool SupportsType( TSubclassOf<UMovieSceneTrack> Type ) const override;
+	virtual TSharedRef<ISequencerSection> MakeSectionInterface( UMovieSceneSection& SectionObject, UMovieSceneTrack* Track ) override;
+	virtual void AddKey( const FGuid& ObjectGuid, UObject* AdditionalAsset = NULL ) override;
+	virtual void BuildObjectBindingContextMenu(FMenuBuilder& MenuBuilder, const FGuid& ObjectBinding, const UClass* ObjectClass) override;
+
+	/** Add path */
+	void AddPath(FGuid ObjectGuid, UObject* AdditionalAsset);
+
+	/** Set path */
+	void SetPath(UMovieSceneSection* Section, AActor* ActorWithSplineComponent);
+
+private:
+
+	/** Add path sub menu */
+	void AddPathSubMenu(FMenuBuilder& MenuBuilder, FGuid ObjectBinding);
+
+	/** Delegate for AnimatablePropertyChanged in AddKey */
+	void AddKeyInternal(float KeyTime, const TArray<UObject*> Objects, UObject* AdditionalAsset);
+};
+
+
diff --git a/Engine/Source/Editor/Sequencer/Private/Sequencer.h b/Engine/Source/Editor/Sequencer/Private/Sequencer.h
--- a/Engine/Source/Editor/Sequencer/Private/Sequencer.h
+++ b/Engine/Source/Editor/Sequencer/Private/Sequencer.h
@@ -102,7 +102,7 @@ public:
 	virtual FSequencerSelection& GetSelection() override;
 
 	/** @return The current view range */
-	FAnimatedRange GetViewRange() const;
+	virtual FAnimatedRange GetViewRange() const override;
 
 	/** Access the user-supplied settings object */
 	USequencerSettings* GetSettings() const { return Settings; }
diff --git a/Engine/Source/Editor/Sequencer/Public/ISequencer.h b/Engine/Source/Editor/Sequencer/Public/ISequencer.h
--- a/Engine/Source/Editor/Sequencer/Public/ISequencer.h
+++ b/Engine/Source/Editor/Sequencer/Public/ISequencer.h
@@ -4,6 +4,7 @@
 
 
 #include "Runtime/MovieScene/Public/IMovieScenePlayer.h"
+#include "Editor/SequencerWidgets/Public/ITimeSlider.h"
 #include "KeyPropertyParams.h"
 
 class UMovieScene;
@@ -104,6 +105,9 @@ public:
 	 */
 	virtual void SetGlobalTime(float Time) = 0;
 
+	/** @return The current view range */
+	virtual FAnimatedRange GetViewRange() const { return FAnimatedRange(); }
+
 	/**
 	 * Sets whether perspective viewport hijacking is enabled.
 	 */
diff --git a/Engine/Source/Runtime/MovieScene/Private/MovieSceneInstance.cpp b/Engine/Source/Runtime/MovieScene/Private/MovieSceneInstance.cpp
--- a/Engine/Source/Runtime/MovieScene/Private/MovieSceneInstance.cpp
+++ b/Engine/Source/Runtime/MovieScene/Private/MovieSceneInstance.cpp
@@ -147,6 +147,13 @@ void FMovieSceneInstance::RefreshInstance( IMovieScenePlayer& Player )
 	}
 }
 
+struct FTrackInstanceEvalSorter
+{
+	bool operator()( const TSharedPtr<IMovieSceneTrackInstance> A, const TSharedPtr<IMovieSceneTrackInstance> B ) const
+	{
+		return A->EvalOrder() > B->EvalOrder();
+	}
+};
 
 void FMovieSceneInstance::RefreshInstanceMap( const TArray<UMovieSceneTrack*>& Tracks, const TArray<UObject*>& RuntimeObjects, FMovieSceneInstanceMap& TrackInstances, IMovieScenePlayer& Player  )
 {
@@ -192,4 +199,7 @@ void FMovieSceneInstance::RefreshInstanceMap( const TArray<UMovieSceneTrack*>& T
 			It.RemoveCurrent();
 		}
 	}
-}
+
+	// Sort based on evaluation order
+	TrackInstances.ValueSort(FTrackInstanceEvalSorter());
+}
\ No newline at end of file
diff --git a/Engine/Source/Runtime/MovieScene/Public/IMovieSceneTrackInstance.h b/Engine/Source/Runtime/MovieScene/Public/IMovieSceneTrackInstance.h
--- a/Engine/Source/Runtime/MovieScene/Public/IMovieSceneTrackInstance.h
+++ b/Engine/Source/Runtime/MovieScene/Public/IMovieSceneTrackInstance.h
@@ -47,4 +47,9 @@ public:
 	 * Called before an instance is removed
 	 */
 	virtual void ClearInstance( IMovieScenePlayer& Player ) = 0;
+
+	/**
+	* Evaluation order
+	*/
+	virtual float EvalOrder() { return 0.f; }
 };
diff --git a/Engine/Source/Runtime/MovieSceneTracks/Private/Sections/MovieScene3DPathSection.cpp b/Engine/Source/Runtime/MovieSceneTracks/Private/Sections/MovieScene3DPathSection.cpp
new file mode 100644
--- /dev/null
+++ b/Engine/Source/Runtime/MovieSceneTracks/Private/Sections/MovieScene3DPathSection.cpp
@@ -0,0 +1,155 @@
+// Copyright 1998-2015 Epic Games, Inc. All Rights Reserved.
+
+#include "MovieSceneTracksPrivatePCH.h"
+#include "MovieScene3DPathSection.h"
+#include "MovieScene3DPathTrack.h"
+#include "Components/SplineComponent.h"
+
+UMovieScene3DPathSection::UMovieScene3DPathSection( const FObjectInitializer& ObjectInitializer )
+	: Super( ObjectInitializer )
+	, FrontAxisEnum(MovieScene3DPathSection_Axis::Y)
+	, UpAxisEnum(MovieScene3DPathSection_Axis::Z)
+	, bFollow (true)
+	, bReverse (false)
+	, bForceUpright (false)
+{
+}
+
+void UMovieScene3DPathSection::SetPathId(const FMovieSceneObjectId& InPathId)
+{
+	Modify();
+
+	PathId = InPathId;
+}
+
+FMovieSceneObjectId UMovieScene3DPathSection::GetPathId() const
+{
+	return PathId;
+}
+
+void UMovieScene3DPathSection::Eval( float Position, USplineComponent* SplineComponent, FVector& OutTranslation, FRotator& OutRotation ) const
+{
+	float Timing = TimingCurve.Eval( Position );
+
+	if (Timing < 0.f)
+	{
+		Timing = 0.f;
+	}
+	else if (Timing > 1.f)
+	{
+		Timing = 1.f;
+	}
+
+	if (bReverse)
+	{
+		Timing = 1.f - Timing;
+	}
+
+	const bool UseConstantVelocity = true;
+	OutTranslation = SplineComponent->GetWorldLocationAtTime(Timing, UseConstantVelocity);
+	OutRotation = SplineComponent->GetWorldRotationAtTime(Timing, UseConstantVelocity);
+	
+	FMatrix NewRotationMatrix = FRotationMatrix(OutRotation);
+
+	FVector UpAxis(0, 0, 0);
+	if (UpAxisEnum == MovieScene3DPathSection_Axis::X)
+	{
+		UpAxis = FVector(1, 0, 0);
+	}
+	else if (UpAxisEnum == MovieScene3DPathSection_Axis::NEG_X)
+	{
+		UpAxis = FVector(-1, 0, 0);
+	}
+	else if (UpAxisEnum == MovieScene3DPathSection_Axis::Y)
+	{
+		UpAxis = FVector(0, 1, 0);
+	}
+	else if (UpAxisEnum == MovieScene3DPathSection_Axis::NEG_Y)
+	{
+		UpAxis = FVector(0, -1, 0);
+	}
+	else if (UpAxisEnum == MovieScene3DPathSection_Axis::Z)
+	{
+		UpAxis = FVector(0, 0, 1);
+	}
+	else if (UpAxisEnum == MovieScene3DPathSection_Axis::NEG_Z)
+	{
+		UpAxis = FVector(0, 0, -1);
+	}
+
+	FVector FrontAxis(0, 0, 0);
+	if (FrontAxisEnum == MovieScene3DPathSection_Axis::X)
+	{
+		FrontAxis = FVector(1, 0, 0);
+	}
+	else if (FrontAxisEnum == MovieScene3DPathSection_Axis::NEG_X)
+	{
+		FrontAxis = FVector(-1, 0, 0);
+	}
+	else if (FrontAxisEnum == MovieScene3DPathSection_Axis::Y)
+	{
+		FrontAxis = FVector(0, 1, 0);
+	}
+	else if (FrontAxisEnum == MovieScene3DPathSection_Axis::NEG_Y)
+	{
+		FrontAxis = FVector(0, -1, 0);
+	}
+	else if (FrontAxisEnum == MovieScene3DPathSection_Axis::Z)
+	{
+		FrontAxis = FVector(0, 0, 1);
+	}
+	else if (FrontAxisEnum == MovieScene3DPathSection_Axis::NEG_Z)
+	{
+		FrontAxis = FVector(0, 0, -1);
+	}
+
+	// Negate the front axis because the spline rotation comes in reversed
+	FrontAxis *= FVector(-1, -1, -1);
+
+	FMatrix AxisRotator = FRotationMatrix::MakeFromXZ(FrontAxis, UpAxis);
+	NewRotationMatrix = AxisRotator * NewRotationMatrix;
+	OutRotation = NewRotationMatrix.Rotator();
+
+	if (bForceUpright)
+	{
+		OutRotation.Pitch = 0.f;
+		OutRotation.Roll = 0.f;
+	}
+}
+
+void UMovieScene3DPathSection::MoveSection( float DeltaPosition, TSet<FKeyHandle>& KeyHandles )
+{
+	Super::MoveSection( DeltaPosition, KeyHandles );
+
+	// Move the curve
+	TimingCurve.ShiftCurve(DeltaPosition, KeyHandles);
+}
+
+void UMovieScene3DPathSection::DilateSection( float DilationFactor, float Origin, TSet<FKeyHandle>& KeyHandles )
+{	
+	Super::DilateSection(DilationFactor, Origin, KeyHandles);
+	
+	TimingCurve.ScaleCurve(Origin, DilationFactor, KeyHandles);
+}
+
+void UMovieScene3DPathSection::GetKeyHandles(TSet<FKeyHandle>& KeyHandles) const
+{
+	for (auto It(TimingCurve.GetKeyHandleIterator()); It; ++It)
+	{
+		float Time = TimingCurve.GetKeyTime(It.Key());
+		if (IsTimeWithinSection(Time))
+		{
+			KeyHandles.Add(It.Key());
+		}
+	}
+}
+
+void UMovieScene3DPathSection::AddPath( float Time, float SequenceEndTime, const FMovieSceneObjectId& InPathId )
+{
+	Modify();
+	PathId = InPathId;
+
+	TimingCurve.UpdateOrAddKey(Time, 0);
+	TimingCurve.UpdateOrAddKey(SequenceEndTime, 1);
+}
+
diff --git a/Engine/Source/Runtime/MovieSceneTracks/Private/TrackInstances/MovieScene3DPathTrackInstance.cpp b/Engine/Source/Runtime/MovieSceneTracks/Private/TrackInstances/MovieScene3DPathTrackInstance.cpp
new file mode 100644
--- /dev/null
+++ b/Engine/Source/Runtime/MovieSceneTracks/Private/TrackInstances/MovieScene3DPathTrackInstance.cpp
@@ -0,0 +1,100 @@
+// Copyright 1998-2015 Epic Games, Inc. All Rights Reserved.
+
+#include "MovieSceneTracksPrivatePCH.h"
+#include "MovieSceneCommonHelpers.h"
+#include "MovieScene3DPathTrackInstance.h"
+#include "MovieScene3DPathTrack.h"
+#include "MovieScene3DPathSection.h"
+#include "IMovieScenePlayer.h"
+#include "Components/SplineComponent.h"
+
+FMovieScene3DPathTrackInstance::FMovieScene3DPathTrackInstance( UMovieScene3DPathTrack& InPathTrack )
+{
+	PathTrack = &InPathTrack;
+}
+
+void FMovieScene3DPathTrackInstance::SaveState(const TArray<UObject*>& RuntimeObjects)
+{
+	for (int32 ObjIndex = 0; ObjIndex < RuntimeObjects.Num(); ++ObjIndex)
+	{
+		USceneComponent* SceneComponent = MovieSceneHelpers::SceneComponentFromRuntimeObject(RuntimeObjects[ObjIndex]);
+		if (SceneComponent != NULL)
+		{
+			InitTransformMap.Add(RuntimeObjects[ObjIndex], SceneComponent->GetRelativeTransform());
+		}
+	}
+}
+
+void FMovieScene3DPathTrackInstance::RestoreState(const TArray<UObject*>& RuntimeObjects)
+{
+	for (int32 ObjIndex = 0; ObjIndex < RuntimeObjects.Num(); ++ObjIndex)
+	{
+		if (!IsValid(RuntimeObjects[ObjIndex]))
+		{
+			continue;
+		}
+
+		USceneComponent* SceneComponent = MovieSceneHelpers::SceneComponentFromRuntimeObject(RuntimeObjects[ObjIndex]);
+		if (SceneComponent != NULL)
+		{
+			FTransform *Transform = InitTransformMap.Find(RuntimeObjects[ObjIndex]);
+			if (Transform != NULL)
+			{
+				SceneComponent->SetRelativeTransform(*Transform);
+			}
+		}
+	}
+}
+
+void FMovieScene3DPathTrackInstance::Update( float Position, float LastPosition, const TArray<UObject*>& RuntimeObjects, class IMovieScenePlayer& Player ) 
+{
+	FVector Translation;
+	FRotator Rotation;
+
+	UMovieScene3DPathSection* FirstPathSection = NULL;
+
+	const TArray<UMovieSceneSection*>& PathSections = PathTrack->GetAllSections();
+
+	for (int32 PathIndex = 0; PathIndex < PathSections.Num(); ++PathIndex)
+	{
+		UMovieScene3DPathSection* PathSection = CastChecked<UMovieScene3DPathSection>(PathSections[PathIndex]);
+
+		if (PathSection->IsTimeWithinSection(Position) &&
+			(FirstPathSection == NULL || FirstPathSection->GetRowIndex() > PathSection->GetRowIndex()))
+		{
+			TArray<UObject*> PathObjects;
+			FMovieSceneObjectId PathId = PathSection->GetPathId();
+			if (PathId.IsValid())
+			{
+				Player.GetRuntimeObjects( Player.GetRootMovieSceneInstance(), PathId.ObjectGuid, PathObjects);
+
+				for (int32 PathObjectIndex = 0; PathObjectIndex < PathObjects.Num(); ++PathObjectIndex)
+				{
+					AActor* Actor = Cast<AActor>(PathObjects[PathObjectIndex]);
+					if (Actor)
+					{
+						TArray<USplineComponent*> SplineComponents;
+						Actor->GetComponents(SplineComponents);
+
+						if (SplineComponents.Num())
+						{
+							PathSection->Eval(Position, SplineComponents[0], Translation, Rotation);
+
+							for( int32 ObjIndex = 0; ObjIndex < RuntimeObjects.Num(); ++ObjIndex )
+							{
+								USceneComponent* SceneComponent = MovieSceneHelpers::SceneComponentFromRuntimeObject(RuntimeObjects[ObjIndex]);
+
+								if (SceneComponent != NULL)
+								{
+									// Set the relative translation and rotation.  Note they are set once instead of individually to avoid a redundant component transform update.
+									SceneComponent->SetRelativeLocationAndRotation(Translation, Rotation);
+								}
+							}
+						}
+					}	
+				}
+			}
+		}
+	}
+}
+
diff --git a/Engine/Source/Runtime/MovieSceneTracks/Private/Tracks/MovieScene3DPathTrack.cpp b/Engine/Source/Runtime/MovieSceneTracks/Private/Tracks/MovieScene3DPathTrack.cpp
new file mode 100644
--- /dev/null
+++ b/Engine/Source/Runtime/MovieSceneTracks/Private/Tracks/MovieScene3DPathTrack.cpp
@@ -0,0 +1,65 @@
+// Copyright 1998-2015 Epic Games, Inc. All Rights Reserved.
+
+#include "MovieSceneTracksPrivatePCH.h"
+#include "MovieScene3DPathSection.h"
+#include "MovieScene3DPathTrack.h"
+#include "IMovieScenePlayer.h"
+#include "MovieScene3DPathTrackInstance.h"
+
+UMovieScene3DPathTrack::UMovieScene3DPathTrack( const FObjectInitializer& ObjectInitializer )
+	: Super( ObjectInitializer )
+{
+}
+FName UMovieScene3DPathTrack::GetTrackName() const
+{
+	return FName("Path");
+}
+
+TSharedPtr<IMovieSceneTrackInstance> UMovieScene3DPathTrack::CreateInstance()
+{
+	return MakeShareable( new FMovieScene3DPathTrackInstance( *this ) ); 
+}
+
+const TArray<UMovieSceneSection*>& UMovieScene3DPathTrack::GetAllSections() const
+{
+	return PathSections;
+}
+
+void UMovieScene3DPathTrack::RemoveAllAnimationData()
+{
+}
+
+bool UMovieScene3DPathTrack::HasSection( UMovieSceneSection* Section ) const
+{
+	return PathSections.Find( Section ) != INDEX_NONE;
+}
+
+void UMovieScene3DPathTrack::RemoveSection( UMovieSceneSection* Section )
+{
+	PathSections.Remove( Section );
+}
+
+bool UMovieScene3DPathTrack::IsEmpty() const
+{
+	return PathSections.Num() == 0;
+}
+
+TRange<float> UMovieScene3DPathTrack::GetSectionBoundaries() const
+{
+	TArray< TRange<float> > Bounds;
+	for (int32 i = 0; i < PathSections.Num(); ++i)
+	{
+		Bounds.Add(PathSections[i]->GetRange());
+	}
+	return TRange<float>::Hull(Bounds);
+}
+
+void UMovieScene3DPathTrack::AddPath(float KeyTime, float PathEndTime, const FMovieSceneObjectId& PathId)
+{
+	// add the section
+	UMovieScene3DPathSection* NewSection = NewObject<UMovieScene3DPathSection>(this);
+	NewSection->AddPath(KeyTime, PathEndTime, PathId);
+	NewSection->InitialPlacement( PathSections, KeyTime, PathEndTime, SupportsMultipleRows() );
+
+	PathSections.Add(NewSection);
+}
diff --git a/Engine/Source/Runtime/MovieSceneTracks/Public/Sections/MovieScene3DPathSection.h b/Engine/Source/Runtime/MovieSceneTracks/Public/Sections/MovieScene3DPathSection.h
new file mode 100644
--- /dev/null
+++ b/Engine/Source/Runtime/MovieSceneTracks/Public/Sections/MovieScene3DPathSection.h
@@ -0,0 +1,95 @@
+// Copyright 1998-2015 Epic Games, Inc. All Rights Reserved.
+
+#pragma once
+
+#include "MovieSceneObjectId.h"
+#include "MovieSceneSection.h"
+#include "MovieScene3DPathSection.generated.h"
+
+class USplineComponent;
+
+UENUM(BlueprintType)
+enum class MovieScene3DPathSection_Axis : uint8
+{
+	X UMETA(DisplayName = "X"),
+	Y UMETA(DisplayName = "Y"),
+	Z UMETA(DisplayName = "Z"),
+	NEG_X UMETA(DisplayName = "-X"),
+	NEG_Y UMETA(DisplayName = "-Y"),
+	NEG_Z UMETA(DisplayName = "-Z")
+};
+
+/**
+ * A 3D Path section
+ */
+UCLASS(MinimalAPI)
+class UMovieScene3DPathSection
+	: public UMovieSceneSection
+{
+	GENERATED_UCLASS_BODY()
+
+public:
+
+	/** MovieSceneSection interface */
+	virtual void MoveSection( float DeltaPosition, TSet<FKeyHandle>& KeyHandles ) override;
+	virtual void DilateSection( float DilationFactor, float Origin, TSet<FKeyHandle>& KeyHandles ) override;
+	virtual void GetKeyHandles(TSet<FKeyHandle>& KeyHandles) const override;
+
+	/** Sets the path id for this path section */
+	virtual void SetPathId(const FMovieSceneObjectId& InId);
+
+	/** Gets the path id for this path section */
+	virtual FMovieSceneObjectId GetPathId() const;
+
+	/**
+	 * Evaluates the path track
+	 *
+	 * @param Time				The position in time within the movie scene
+	 */
+	void Eval( float Time, USplineComponent* SplineComponent, FVector& OutTranslation, FRotator& OutRotation ) const;
+
+	/** 
+	 * Adds a path to the section
+	 *
+	 * @param Time	The location in time where the path should be added
+	 * @param SequenceEndTime   The time at the end of the sequence, by default the path is set to end at this time
+	 * @param InPathId The id to the path
+	 */
+	void AddPath( float Time, float SequenceEndTime, const FMovieSceneObjectId& InPathId);
+
+	/** 
+	 * Returns the timing curve
+	 *
+	 * @return The timing curve
+	 */
+	MOVIESCENETRACKS_API FRichCurve& GetTimingCurve() { return TimingCurve; }
+	
+private:
+	/** The path possessable that this path uses */
+	UPROPERTY()
+	FMovieSceneObjectId PathId;
+
+	/** Timing Curve */
+	UPROPERTY()
+	FRichCurve TimingCurve;
+
+	/** Front Axis */
+	UPROPERTY()
+	MovieScene3DPathSection_Axis FrontAxisEnum;
+
+	/** Up Axis */
+	UPROPERTY()
+	MovieScene3DPathSection_Axis UpAxisEnum;
+
+	/** Follow Curve */
+	UPROPERTY()
+	uint32 bFollow:1;
+
+	/** Reverse Timing */
+	UPROPERTY()
+	uint32 bReverse:1;
+
+	/** Force Upright */
+	UPROPERTY()
+	uint32 bForceUpright:1;
+};
\ No newline at end of file
diff --git a/Engine/Source/Runtime/MovieSceneTracks/Public/TrackInstances/MovieScene3DPathTrackInstance.h b/Engine/Source/Runtime/MovieSceneTracks/Public/TrackInstances/MovieScene3DPathTrackInstance.h
new file mode 100644
--- /dev/null
+++ b/Engine/Source/Runtime/MovieSceneTracks/Public/TrackInstances/MovieScene3DPathTrackInstance.h
@@ -0,0 +1,32 @@
+// Copyright 1998-2015 Epic Games, Inc. All Rights Reserved.
+
+#pragma once
+
+#include "IMovieSceneTrackInstance.h"
+
+class UMovieScene3DPathTrack;
+
+/**
+ * Instance of a UMovieScene3DPathTrack
+ */
+class FMovieScene3DPathTrackInstance
+	: public IMovieSceneTrackInstance
+{
+public:
+	FMovieScene3DPathTrackInstance( UMovieScene3DPathTrack& InPathTrack );
+
+	/** IMovieSceneTrackInstance interface */
+	virtual void SaveState (const TArray<UObject*>& RuntimeObjects) override;
+	virtual void RestoreState (const TArray<UObject*>& RuntimeObjects) override;
+	virtual void Update( float Position, float LastPosition, const TArray<UObject*>& RuntimeObjects, IMovieScenePlayer& Player ) override;
+	virtual void RefreshInstance( const TArray<UObject*>& RuntimeObjects, IMovieScenePlayer& Player ) override {}
+	virtual void ClearInstance( IMovieScenePlayer& Player ) override {}
+	virtual float EvalOrder() { return 1.f; }
+
+private:
+	/** Track that is being instanced */
+	UMovieScene3DPathTrack* PathTrack;
+
+	/** Map from object to initial state */
+	TMap< TWeakObjectPtr<UObject>, FTransform > InitTransformMap;
+};
diff --git a/Engine/Source/Runtime/MovieSceneTracks/Public/Tracks/MovieScene3DPathTrack.h b/Engine/Source/Runtime/MovieSceneTracks/Public/Tracks/MovieScene3DPathTrack.h
new file mode 100644
--- /dev/null
+++ b/Engine/Source/Runtime/MovieSceneTracks/Public/Tracks/MovieScene3DPathTrack.h
@@ -0,0 +1,40 @@
+// Copyright 1998-2015 Epic Games, Inc. All Rights Reserved.
+
+#pragma once
+
+#include "MovieScene.h"
+#include "MovieSceneTrack.h"
+#include "MovieScene3DPathTrack.generated.h"
+
+/**
+ * Handles manipulation of path tracks in a movie scene
+ */
+UCLASS( MinimalAPI )
+class UMovieScene3DPathTrack : public UMovieSceneTrack
+{
+	GENERATED_UCLASS_BODY()
+public:
+	/** UMovieSceneTrack interface */
+	virtual FName GetTrackName() const override;
+	virtual TSharedPtr<IMovieSceneTrackInstance> CreateInstance() override;
+	virtual void RemoveAllAnimationData() override;
+	virtual bool HasSection( UMovieSceneSection* Section ) const override;
+	virtual void RemoveSection( UMovieSceneSection* Section ) override;
+	virtual bool IsEmpty() const override;
+	virtual TRange<float> GetSectionBoundaries() const override;
+	virtual const TArray<UMovieSceneSection*>& GetAllSections() const override;
+
+	/**
+	 * Adds a path
+	 *
+	 * @param Time				The time relative to the owning movie scene where the section should be
+	 * @param PathEndTime       Set the path to end at this time
+	 * @param PathId			The id to the spline
+	 */
+	virtual void AddPath( float Time, float PathEndTime, const FMovieSceneObjectId& PathId );
+
+private:
+	/** List of all path sections */
+	UPROPERTY()
+	TArray<UMovieSceneSection*> PathSections;
+};
