Sequencer - Prevent the PlayMovieScene node from connecting the actor and component name pins every time it's reconstructed.

diff --git a/Engine/Source/Editor/MovieSceneTools/Private/K2Node_PlayMovieScene.cpp b/Engine/Source/Editor/MovieSceneTools/Private/K2Node_PlayMovieScene.cpp
--- a/Engine/Source/Editor/MovieSceneTools/Private/K2Node_PlayMovieScene.cpp
+++ b/Engine/Source/Editor/MovieSceneTools/Private/K2Node_PlayMovieScene.cpp
@@ -63,7 +63,7 @@ void UK2Node_PlayMovieScene::AllocateDefaultPins()
 		{
 			auto& BoundObject = *BoundObjectIter;
 
-			CreatePinsForBoundObject( BoundObject );
+			CreatePinsForBoundObject( BoundObject, false );
 		}
 	}
 	
@@ -236,7 +236,7 @@ void UK2Node_PlayMovieScene::BindPossessableToObjects(const FGuid& PossessableGu
 
 	// Create a pin for the bound object
 	Modify();
-	CreatePinsForBoundObject( BoundObject );
+	CreatePinsForBoundObject( BoundObject, true );
 	FBlueprintEditorUtils::MarkBlueprintAsStructurallyModified( this->GetBlueprint() );
 }
 
@@ -269,7 +269,7 @@ FGuid UK2Node_PlayMovieScene::FindGuidForObjectInfo( FMovieSceneBoundObjectInfo
 }
 
 
-void UK2Node_PlayMovieScene::CreatePinsForBoundObject( FMovieSceneBoundObject& BoundObject )
+void UK2Node_PlayMovieScene::CreatePinsForBoundObject( FMovieSceneBoundObject& BoundObject, bool bLinkNewPinsToObject )
 {
 	const UEdGraphSchema_K2* K2Schema = GetDefault<UEdGraphSchema_K2>();
 
@@ -363,8 +363,11 @@ void UK2Node_PlayMovieScene::CreatePinsForBoundObject( FMovieSceneBoundObject& B
 					ActorLiteralNode = FEdGraphSchemaAction_K2NewNode::SpawnNodeFromTemplate<UK2Node_Literal>(GetGraph(), LiteralNodeTemplate, NewNodePos);
 				}
 
-				// Hook up the object reference literal to our pin
-				ActorLiteralNode->GetValuePin()->MakeLinkTo(NewActorPin);
+				if (bLinkNewPinsToObject)
+				{
+					// Hook up the object reference literal to our pin
+					ActorLiteralNode->GetValuePin()->MakeLinkTo(NewActorPin);
+				}
 
 				if (BoundObjectInfo.Tag.IsEmpty() == false)
 				{
@@ -381,7 +384,10 @@ void UK2Node_PlayMovieScene::CreatePinsForBoundObject( FMovieSceneBoundObject& B
 						ValuePin->DefaultValue = BoundObjectInfo.Tag;
 					}
 
-					ComponentNameLiteralNode->GetReturnValuePin()->MakeLinkTo(NewComponentNamePin);
+					if (bLinkNewPinsToObject)
+					{
+						ComponentNameLiteralNode->GetReturnValuePin()->MakeLinkTo(NewComponentNamePin);
+					}
 				}
 			}
 		}
diff --git a/Engine/Source/Editor/MovieSceneTools/Public/K2Node_PlayMovieScene.h b/Engine/Source/Editor/MovieSceneTools/Public/K2Node_PlayMovieScene.h
--- a/Engine/Source/Editor/MovieSceneTools/Public/K2Node_PlayMovieScene.h
+++ b/Engine/Source/Editor/MovieSceneTools/Public/K2Node_PlayMovieScene.h
@@ -118,9 +118,10 @@ protected:
 	/**
 	 * Creates a pin on this node for the specified bound object
 	 *
-	 * @param	BoundObject		The object to create a pin for
+	 * @param	BoundObject				The object to create a pin for
+	 * @param	bLinkNewPinsToObject	When true the new pins will be linked to the bound object.
 	 */
-	virtual void CreatePinsForBoundObject( struct FMovieSceneBoundObject& BoundObject );
+	virtual void CreatePinsForBoundObject( struct FMovieSceneBoundObject& BoundObject, bool bLinkNewPinsToObject );
 
 	/**
 	 * Creates a MovieSceneBindings object if we don't have one yet
