Sequencer: Support for bool key editor.

diff --git a/Engine/Source/Editor/MovieSceneTools/Private/CurveKeyEditors/SBoolCurveKeyEditor.cpp b/Engine/Source/Editor/MovieSceneTools/Private/CurveKeyEditors/SBoolCurveKeyEditor.cpp
new file mode 100644
--- /dev/null
+++ b/Engine/Source/Editor/MovieSceneTools/Private/CurveKeyEditors/SBoolCurveKeyEditor.cpp
@@ -0,0 +1,53 @@
+// Copyright 1998-2015 Epic Games, Inc. All Rights Reserved.
+
+#include "MovieSceneToolsPrivatePCH.h"
+#include "SBoolCurveKeyEditor.h"
+
+#define LOCTEXT_NAMESPACE "BoolCurveKeyEditor"
+
+void SBoolCurveKeyEditor::Construct(const FArguments& InArgs)
+{
+	Sequencer = InArgs._Sequencer;
+	OwningSection = InArgs._OwningSection;
+	Curve = InArgs._Curve;
+	float CurrentTime = Sequencer->GetCurrentLocalTime(*Sequencer->GetFocusedMovieScene());
+	ChildSlot
+	[
+		SNew(SCheckBox)
+		.IsChecked(this, &SBoolCurveKeyEditor::IsChecked)
+		.OnCheckStateChanged(this, &SBoolCurveKeyEditor::OnCheckStateChanged)
+	];
+}
+
+ECheckBoxState SBoolCurveKeyEditor::IsChecked() const
+{
+	float CurrentTime = Sequencer->GetCurrentLocalTime(*Sequencer->GetFocusedMovieScene());
+	return !!Curve->Evaluate(CurrentTime) ? ECheckBoxState::Checked : ECheckBoxState::Unchecked;
+}
+
+void SBoolCurveKeyEditor::OnCheckStateChanged(ECheckBoxState NewCheckboxState)
+{
+	FScopedTransaction Transaction(LOCTEXT("SetBoolKey", "Set bool key value"));
+	OwningSection->SetFlags(RF_Transactional);
+	OwningSection->Modify();
+
+	float CurrentTime = Sequencer->GetCurrentLocalTime(*Sequencer->GetFocusedMovieScene());
+
+	bool bKeyWillBeAdded = Curve->IsKeyHandleValid(Curve->FindKey(CurrentTime)) == false;
+	if (bKeyWillBeAdded)
+	{
+		if (OwningSection->GetStartTime() > CurrentTime)
+		{
+			OwningSection->SetStartTime(CurrentTime);
+		}
+		if (OwningSection->GetEndTime() < CurrentTime)
+		{
+			OwningSection->SetEndTime(CurrentTime);
+		}
+	}
+
+	Curve->UpdateOrAddKey(CurrentTime, NewCheckboxState == ECheckBoxState::Checked ? 1 : 0);
+	Sequencer->UpdateRuntimeInstances();
+}
+
+#undef LOCTEXT_NAMESPACE
\ No newline at end of file
diff --git a/Engine/Source/Editor/MovieSceneTools/Private/CurveKeyEditors/SBoolCurveKeyEditor.h b/Engine/Source/Editor/MovieSceneTools/Private/CurveKeyEditors/SBoolCurveKeyEditor.h
new file mode 100644
--- /dev/null
+++ b/Engine/Source/Editor/MovieSceneTools/Private/CurveKeyEditors/SBoolCurveKeyEditor.h
@@ -0,0 +1,27 @@
+// Copyright 1998-2015 Epic Games, Inc. All Rights Reserved.
+
+#pragma once
+
+/**
+ * A widget for editing a curve representing bool keys.
+ */
+class SBoolCurveKeyEditor : public SCompoundWidget
+{
+public:
+	SLATE_BEGIN_ARGS(SBoolCurveKeyEditor) {}
+		SLATE_ARGUMENT(ISequencer*, Sequencer)
+		SLATE_ARGUMENT(UMovieSceneSection*, OwningSection)
+		SLATE_ARGUMENT(FIntegralCurve*, Curve)
+	SLATE_END_ARGS();
+
+	void Construct(const FArguments& InArgs);
+
+private:
+
+	ECheckBoxState IsChecked() const;
+	void OnCheckStateChanged(ECheckBoxState);
+
+	ISequencer* Sequencer;
+	UMovieSceneSection* OwningSection;
+	FIntegralCurve* Curve;
+};
\ No newline at end of file
diff --git a/Engine/Source/Editor/MovieSceneTools/Private/MovieSceneToolHelpers.cpp b/Engine/Source/Editor/MovieSceneTools/Private/MovieSceneToolHelpers.cpp
--- a/Engine/Source/Editor/MovieSceneTools/Private/MovieSceneToolHelpers.cpp
+++ b/Engine/Source/Editor/MovieSceneTools/Private/MovieSceneToolHelpers.cpp
@@ -14,6 +14,7 @@
 #include "SFloatCurveKeyEditor.h"
 #include "SIntegralCurveKeyEditor.h"
 #include "SEnumCurveKeyEditor.h"
+#include "SBoolCurveKeyEditor.h"
 
 void FFloatCurveKeyArea::AddKeyUnique(float Time)
 {
@@ -74,3 +75,11 @@ TSharedRef<SWidget> FEnumKeyArea::CreateKeyEditor(ISequencer* Sequencer)
 		.Enum(Enum);
 };
 
+TSharedRef<SWidget> FBoolKeyArea::CreateKeyEditor(ISequencer* Sequencer)
+{
+	return SNew(SBoolCurveKeyEditor)
+		.Sequencer(Sequencer)
+		.OwningSection(OwningSection)
+		.Curve(&Curve);
+};
+
diff --git a/Engine/Source/Editor/MovieSceneTools/Private/TrackEditors/PropertyTrackEditor.cpp b/Engine/Source/Editor/MovieSceneTools/Private/TrackEditors/PropertyTrackEditor.cpp
--- a/Engine/Source/Editor/MovieSceneTools/Private/TrackEditors/PropertyTrackEditor.cpp
+++ b/Engine/Source/Editor/MovieSceneTools/Private/TrackEditors/PropertyTrackEditor.cpp
@@ -215,7 +215,7 @@ public:
 	virtual void GenerateSectionLayout( class ISectionLayoutBuilder& LayoutBuilder ) const override
 	{
 		UMovieSceneBoolSection* BoolSection = Cast<UMovieSceneBoolSection>( &SectionObject );
-		LayoutBuilder.SetSectionAsKeyArea( MakeShareable( new FIntegralKeyArea( BoolSection->GetCurve(), BoolSection ) ) );
+		LayoutBuilder.SetSectionAsKeyArea( MakeShareable( new FBoolKeyArea( BoolSection->GetCurve(), BoolSection ) ) );
 	}
 
 	virtual int32 OnPaintSection( const FGeometry& AllottedGeometry, const FSlateRect& SectionClippingRect, FSlateWindowElementList& OutDrawElements, int32 LayerId, bool bParentEnabled ) const override 
diff --git a/Engine/Source/Editor/MovieSceneTools/Public/MovieSceneToolHelpers.h b/Engine/Source/Editor/MovieSceneTools/Public/MovieSceneToolHelpers.h
--- a/Engine/Source/Editor/MovieSceneTools/Public/MovieSceneToolHelpers.h
+++ b/Engine/Source/Editor/MovieSceneTools/Public/MovieSceneToolHelpers.h
@@ -144,4 +144,20 @@ private:
 	UEnum* Enum;
 };
 
+/** A key area for displaying and editing integral curves representing bools. */
+class FBoolKeyArea : public FIntegralKeyArea
+{
+public:
+	FBoolKeyArea(FIntegralCurve& InCurve, UMovieSceneSection* InOwningSection)
+		: FIntegralKeyArea(InCurve, InOwningSection)
+	{}
+
+	virtual bool CanCreateKeyEditor() override 
+	{
+		return true;
+	}
+
+	virtual TSharedRef<SWidget> CreateKeyEditor(ISequencer* Sequencer) override;
+};
+
 
