Sequencer: Support for infinite sections.

The bool section is now an infinite section by default.

diff --git a/Engine/Source/Editor/MovieSceneTools/Private/TrackEditors/PropertyTrackEditor.cpp b/Engine/Source/Editor/MovieSceneTools/Private/TrackEditors/PropertyTrackEditor.cpp
--- a/Engine/Source/Editor/MovieSceneTools/Private/TrackEditors/PropertyTrackEditor.cpp
+++ b/Engine/Source/Editor/MovieSceneTools/Private/TrackEditors/PropertyTrackEditor.cpp
@@ -209,8 +209,9 @@ public:
 class FBoolPropertySection : public FPropertySection
 {
 public:
-	FBoolPropertySection( UMovieSceneSection& InSectionObject, FName SectionName )
-		: FPropertySection(InSectionObject, SectionName) {}
+	FBoolPropertySection( UMovieSceneSection& InSectionObject, FName SectionName, ISequencer* InSequencer )
+		: FPropertySection(InSectionObject, SectionName)
+	, Sequencer(InSequencer) {}
 
 	virtual void GenerateSectionLayout( class ISectionLayoutBuilder& LayoutBuilder ) const override
 	{
@@ -220,13 +221,21 @@ public:
 
 	virtual int32 OnPaintSection( const FGeometry& AllottedGeometry, const FSlateRect& SectionClippingRect, FSlateWindowElementList& OutDrawElements, int32 LayerId, bool bParentEnabled ) const override 
 	{
-		FTimeToPixel TimeToPixelConverter( AllottedGeometry, TRange<float>( SectionObject.GetStartTime(), SectionObject.GetEndTime() ) );
+		FTimeToPixel TimeToPixelConverter = SectionObject.IsInfinite() ? 
+			FTimeToPixel( AllottedGeometry, Sequencer->GetViewRange() ) : 
+			FTimeToPixel( AllottedGeometry, TRange<float>( SectionObject.GetStartTime(), SectionObject.GetEndTime() ) );
 
 		// custom drawing for bool curves
 		UMovieSceneBoolSection* BoolSection = Cast<UMovieSceneBoolSection>( &SectionObject );
 
 		TArray<float> SectionSwitchTimes;
+
+		if (SectionObject.IsInfinite() && Sequencer->GetViewRange().GetLowerBoundValue() < SectionObject.GetStartTime())
+		{
+			SectionSwitchTimes.Add(Sequencer->GetViewRange().GetLowerBoundValue());
+		}
 		SectionSwitchTimes.Add(SectionObject.GetStartTime());
+
 		FIntegralCurve& BoolCurve = BoolSection->GetCurve();
 		for (TArray<FIntegralKey>::TConstIterator It(BoolCurve.GetKeyIterator()); It; ++It)
 		{
@@ -236,6 +245,11 @@ public:
 				SectionSwitchTimes.Add(Time);
 			}
 		}
+
+		if (SectionObject.IsInfinite() && Sequencer->GetViewRange().GetUpperBoundValue() > SectionObject.GetEndTime())
+		{
+			SectionSwitchTimes.Add(Sequencer->GetViewRange().GetUpperBoundValue());
+		}
 		SectionSwitchTimes.Add(SectionObject.GetEndTime());
 
 		for (int32 i = 0; i < SectionSwitchTimes.Num() - 1; ++i)
@@ -272,6 +286,8 @@ public:
 
 		return LayerId+1;
 	}
+private:
+	ISequencer* Sequencer;
 };
 
 /**
@@ -362,11 +378,11 @@ TSharedRef<ISequencerSection> FPropertyTrackEditor::MakeSectionInterface( UMovie
 
 	TSharedRef<ISequencerSection> NewSection = TSharedRef<ISequencerSection>(
 		SectionClass == UMovieSceneColorTrack::StaticClass() ? new FColorPropertySection( SectionObject, Track->GetTrackName() ) :
-		SectionClass == UMovieSceneBoolTrack::StaticClass() ? new FBoolPropertySection( SectionObject, Track->GetTrackName() ) :
+		SectionClass == UMovieSceneBoolTrack::StaticClass() ? new FBoolPropertySection( SectionObject, Track->GetTrackName(), GetSequencer().Get() ) :
 		SectionClass == UMovieSceneByteTrack::StaticClass() ? new FBytePropertySection(SectionObject, Track->GetTrackName(), Cast<UMovieSceneByteTrack>(SectionObject.GetOuter())->GetEnum()) :
 		SectionClass == UMovieSceneVectorTrack::StaticClass() ? new FVectorPropertySection( SectionObject, Track->GetTrackName() ) :
 		SectionClass == UMovieSceneFloatTrack::StaticClass() ? new FFloatPropertySection( SectionObject, Track->GetTrackName() ) :
-		SectionClass == UMovieSceneVisibilityTrack::StaticClass() ? new FBoolPropertySection( SectionObject, Track->GetTrackName() ) :
+		SectionClass == UMovieSceneVisibilityTrack::StaticClass() ? new FBoolPropertySection( SectionObject, Track->GetTrackName(), GetSequencer().Get() ) :
 		new FPropertySection( SectionObject, Track->GetTrackName() )
 		);
 
diff --git a/Engine/Source/Editor/Sequencer/Private/SSection.cpp b/Engine/Source/Editor/Sequencer/Private/SSection.cpp
--- a/Engine/Source/Editor/Sequencer/Private/SSection.cpp
+++ b/Engine/Source/Editor/Sequencer/Private/SSection.cpp
@@ -96,7 +96,9 @@ FSelectedKey SSection::GetKeyUnderMouse( const FVector2D& MousePosition, const F
 
 			FVector2D LocalSpaceMousePosition = KeyAreaGeometry.AbsoluteToLocal( MousePosition );
 
-			FTimeToPixel TimeToPixelConverter( KeyAreaGeometry, TRange<float>( Section.GetStartTime(), Section.GetEndTime() ) );
+			FTimeToPixel TimeToPixelConverter = Section.IsInfinite() ? 			
+				FTimeToPixel( ParentGeometry, GetSequencer().GetViewRange()) : 
+				FTimeToPixel( KeyAreaGeometry, TRange<float>( Section.GetStartTime(), Section.GetEndTime() ) );
 
 			// Check each key until we find one under the mouse (if any)
 			TArray<FKeyHandle> KeyHandles = KeyArea->GetUnsortedKeyHandles();
@@ -292,7 +294,9 @@ void SSection::PaintKeys( const FGeometry& AllottedGeometry, const FSlateRect& M
 
 		FGeometry KeyAreaGeometry = GetKeyAreaGeometry( Element, AllottedGeometry );
 
-		FTimeToPixel TimeToPixelConverter( KeyAreaGeometry, TRange<float>( SectionObject.GetStartTime(), SectionObject.GetEndTime() ) );
+		FTimeToPixel TimeToPixelConverter = SectionObject.IsInfinite() ? 			
+			FTimeToPixel( ParentGeometry, GetSequencer().GetViewRange()) : 
+			FTimeToPixel( KeyAreaGeometry, TRange<float>( SectionObject.GetStartTime(), SectionObject.GetEndTime() ) );
 
 		// Draw a box for the key area 
 		// @todo Sequencer - Allow the IKeyArea to do this
@@ -634,7 +638,9 @@ FReply SSection::OnMouseMove( const FGeometry& MyGeometry, const FPointerEvent&
 			else if( DragOperation.IsValid() )
 			{
 				// Already in a drag, tell all operations to perform their drag implementations
-				FTimeToPixel TimeToPixelConverter( MyGeometry, TRange<float>( SectionInterface->GetSectionObject()->GetStartTime(), SectionInterface->GetSectionObject()->GetEndTime() ) );
+				FTimeToPixel TimeToPixelConverter = SectionInterface->GetSectionObject()->IsInfinite() ? 			
+					FTimeToPixel( ParentGeometry, GetSequencer().GetViewRange()) : 
+					FTimeToPixel( MyGeometry, TRange<float>( SectionInterface->GetSectionObject()->GetStartTime(), SectionInterface->GetSectionObject()->GetEndTime() ) );
 
 				DragOperation->OnDrag( MouseEvent, LocalMousePos, TimeToPixelConverter, ParentSectionArea );
 			}
diff --git a/Engine/Source/Editor/Sequencer/Private/SSequencerSectionAreaView.cpp b/Engine/Source/Editor/Sequencer/Private/SSequencerSectionAreaView.cpp
--- a/Engine/Source/Editor/Sequencer/Private/SSequencerSectionAreaView.cpp
+++ b/Engine/Source/Editor/Sequencer/Private/SSequencerSectionAreaView.cpp
@@ -31,11 +31,21 @@ namespace SequencerSectionUtils
 	FGeometry GetSectionGeometry( const FGeometry& AllottedGeometry, int32 RowIndex, int32 MaxTracks, float NodeHeight, TSharedPtr<ISequencerSection> SectionInterface, const FTimeToPixel& TimeToPixelConverter )
 	{
 		const UMovieSceneSection* Section = SectionInterface->GetSectionObject();
-		// Where to start drawing the section
-		float StartX =  TimeToPixelConverter.TimeToPixel( Section->GetStartTime() );
-		// Where to stop drawing the section
-		float EndX = TimeToPixelConverter.TimeToPixel( Section->GetEndTime() );
-		
+
+		float StartX, EndX = 0;
+
+		// If the section is infinite, occupy the entire width of the geometry where the section is located.
+		if (Section->IsInfinite())
+		{
+			StartX = AllottedGeometry.Position.X;
+			EndX = AllottedGeometry.Position.X + AllottedGeometry.Size.X;
+		}
+		else
+		{
+			StartX = TimeToPixelConverter.TimeToPixel( Section->GetStartTime() );
+			EndX = TimeToPixelConverter.TimeToPixel( Section->GetEndTime() );
+		}
+
 		// Actual section length without grips.
 		float SectionLengthActual = EndX-StartX;
 
diff --git a/Engine/Source/Editor/Sequencer/Private/SectionDragOperations.cpp b/Engine/Source/Editor/Sequencer/Private/SectionDragOperations.cpp
--- a/Engine/Source/Editor/Sequencer/Private/SectionDragOperations.cpp
+++ b/Engine/Source/Editor/Sequencer/Private/SectionDragOperations.cpp
@@ -293,8 +293,16 @@ void FMoveSection::OnDrag( const FPointerEvent& MouseEvent, const FVector2D& Loc
 			MovieSceneSections.Add(Sections[i]->GetSectionObject());
 		}
 		
-
 		FVector2D TotalDelta = LocalMousePos - DragOffset;
+		
+		// When the section is not infinite, the local mouse pos is relative to the start time of the section which continually updates as the section is dragged. 
+		// When the section is infinite, the local mouse pos is relative to the track extents, which do not change. 
+		// Reset the drag offset so that the motion is relative to the last moved position.
+		if (Section.Get()->IsInfinite())
+		{
+			DragOffset = LocalMousePos;
+		}
+
 		float DistanceMoved = TotalDelta.X / TimeToPixelConverter.GetPixelsPerInput();
 		
 		float DeltaTime = DistanceMoved;
diff --git a/Engine/Source/Runtime/MovieScene/Public/MovieSceneSection.h b/Engine/Source/Runtime/MovieScene/Public/MovieSceneSection.h
--- a/Engine/Source/Runtime/MovieScene/Public/MovieSceneSection.h
+++ b/Engine/Source/Runtime/MovieScene/Public/MovieSceneSection.h
@@ -166,6 +166,10 @@ public:
 	 */
 	virtual MOVIESCENE_API void InitialPlacement(const TArray<UMovieSceneSection*>& Sections, float InStartTime, float InEndTime, bool bAllowMultipleRows);
 
+	/** Whether or not this section is infinite. An infinite section will draw the entire width of the track. StartTime and EndTime will be ignored but not discarded. */
+	void SetIsInfinite(bool bInIsInfinite) { bIsInfinite = bInIsInfinite; }
+	bool IsInfinite() const { return bIsInfinite; }
+
 private:
 	/** The start time of the section */
 	UPROPERTY()
@@ -178,4 +182,8 @@ private:
 	/** The row index that this section sits on */
 	UPROPERTY()
 	int32 RowIndex;
+
+	/** Toggle to set this section to be infinite */
+	UPROPERTY()
+	uint32 bIsInfinite : 1;
 };
diff --git a/Engine/Source/Runtime/MovieSceneTracks/Private/Sections/MovieSceneBoolSection.cpp b/Engine/Source/Runtime/MovieSceneTracks/Private/Sections/MovieSceneBoolSection.cpp
--- a/Engine/Source/Runtime/MovieSceneTracks/Private/Sections/MovieSceneBoolSection.cpp
+++ b/Engine/Source/Runtime/MovieSceneTracks/Private/Sections/MovieSceneBoolSection.cpp
@@ -7,6 +7,7 @@
 UMovieSceneBoolSection::UMovieSceneBoolSection( const FObjectInitializer& ObjectInitializer )
 	: Super( ObjectInitializer )
 {
+	SetIsInfinite(true);
 }
 
 bool UMovieSceneBoolSection::Eval( float Position ) const
