Sequencer: Set the display name for the Visibility Section as "Visible" so that it doesn't read as "bHidden"

diff --git a/Engine/Source/Editor/MovieSceneTools/Private/TrackEditors/PropertyTrackEditor.cpp b/Engine/Source/Editor/MovieSceneTools/Private/TrackEditors/PropertyTrackEditor.cpp
--- a/Engine/Source/Editor/MovieSceneTools/Private/TrackEditors/PropertyTrackEditor.cpp
+++ b/Engine/Source/Editor/MovieSceneTools/Private/TrackEditors/PropertyTrackEditor.cpp
@@ -290,6 +290,19 @@ private:
 	ISequencer* Sequencer;
 };
 
+/**
+ * An implementation of visibility property sections
+ */
+class FVisibilityPropertySection : public FBoolPropertySection
+{
+public:
+	FVisibilityPropertySection( UMovieSceneSection& InSectionObject, FName SectionName, ISequencer* InSequencer )
+		: FBoolPropertySection(InSectionObject, SectionName, InSequencer) 
+	{ 
+		DisplayName = FText::FromString(TEXT("Visible"));
+	}
+};
+
 /**
 * An implementation of byte property sections
 */
@@ -382,7 +395,7 @@ TSharedRef<ISequencerSection> FPropertyTrackEditor::MakeSectionInterface( UMovie
 		SectionClass == UMovieSceneByteTrack::StaticClass() ? new FBytePropertySection(SectionObject, Track->GetTrackName(), Cast<UMovieSceneByteTrack>(SectionObject.GetOuter())->GetEnum()) :
 		SectionClass == UMovieSceneVectorTrack::StaticClass() ? new FVectorPropertySection( SectionObject, Track->GetTrackName() ) :
 		SectionClass == UMovieSceneFloatTrack::StaticClass() ? new FFloatPropertySection( SectionObject, Track->GetTrackName() ) :
-		SectionClass == UMovieSceneVisibilityTrack::StaticClass() ? new FBoolPropertySection( SectionObject, Track->GetTrackName(), GetSequencer().Get() ) :
+		SectionClass == UMovieSceneVisibilityTrack::StaticClass() ? new FVisibilityPropertySection( SectionObject, Track->GetTrackName(), GetSequencer().Get() ) :
 		new FPropertySection( SectionObject, Track->GetTrackName() )
 		);
 
