Sequencer: Disable setting transformation keyframes when the viewport's camera is being moved.

diff --git a/Engine/Source/Editor/MovieSceneTools/Private/TrackEditors/TransformTrackEditor.cpp b/Engine/Source/Editor/MovieSceneTools/Private/TrackEditors/TransformTrackEditor.cpp
--- a/Engine/Source/Editor/MovieSceneTools/Private/TrackEditors/TransformTrackEditor.cpp
+++ b/Engine/Source/Editor/MovieSceneTools/Private/TrackEditors/TransformTrackEditor.cpp
@@ -327,6 +327,17 @@ void F3DTransformTrackEditor::AddScaleKey()
 
 void F3DTransformTrackEditor::AddTransformKeyInternal(F3DTransformTrackKey::Type KeyType)
 {
+	// WASD hotkeys to fly the viewport can conflict with hotkeys for setting keyframes (ie. s). 
+	// If the viewport is moving, disregard setting keyframes.
+	for (int32 i = 0; i < GEditor->LevelViewportClients.Num(); ++i)
+	{		
+		FLevelEditorViewportClient* LevelVC = GEditor->LevelViewportClients[i];
+		if (LevelVC && LevelVC->IsMovingCamera())
+		{
+			return;
+		}
+	}
+
 	USelection* CurrentSelection = GEditor->GetSelectedActors();
 	TArray<UObject*> SelectedActors;
 	CurrentSelection->GetSelectedObjects( AActor::StaticClass(), SelectedActors );
