Sequencer: Fix color track values when channel is not keyed.

diff --git a/Engine/Source/Editor/MovieSceneTools/Private/TrackEditors/PropertyTrackEditor.cpp b/Engine/Source/Editor/MovieSceneTools/Private/TrackEditors/PropertyTrackEditor.cpp
--- a/Engine/Source/Editor/MovieSceneTools/Private/TrackEditors/PropertyTrackEditor.cpp
+++ b/Engine/Source/Editor/MovieSceneTools/Private/TrackEditors/PropertyTrackEditor.cpp
@@ -34,8 +34,10 @@
 class FColorPropertySection : public FPropertySection
 {
 public:
-	FColorPropertySection( UMovieSceneSection& InSectionObject, FName SectionName )
-		: FPropertySection(InSectionObject, SectionName) {}
+	FColorPropertySection( UMovieSceneSection& InSectionObject, FName SectionName, ISequencer* InSequencer, UMovieSceneTrack* InTrack  )
+		: FPropertySection(InSectionObject, SectionName)
+	, Sequencer(InSequencer)
+	, Track(InTrack) {}
 
 	virtual void GenerateSectionLayout( class ISectionLayoutBuilder& LayoutBuilder ) const override
 	{
@@ -101,6 +103,44 @@ public:
 private:
 	void ConsolidateColorCurves(TArray< TKeyValuePair<float, FLinearColor> >& OutColorKeys, const UMovieSceneColorSection* Section) const
 	{
+		// Get the default color of the first instance
+		FLinearColor DefaultColor(0.0f, 0.0f, 0.0f, 0.0f);
+
+		const TArray<FMovieSceneBinding>& MovieSceneBindings = Sequencer->GetFocusedMovieScene()->GetBindings();
+
+		bool bFoundColor = false;
+		for (int32 BindingIndex = 0; BindingIndex < MovieSceneBindings.Num() && !bFoundColor; ++BindingIndex)
+		{
+			const FMovieSceneBinding& MovieSceneBinding = MovieSceneBindings[BindingIndex];
+
+			for (int32 TrackIndex = 0; TrackIndex < MovieSceneBinding.GetTracks().Num() && !bFoundColor; ++TrackIndex)
+			{
+				if (MovieSceneBinding.GetTracks()[TrackIndex] == Track)
+				{
+					TArray<UObject*> RuntimeObjects;
+					Sequencer->GetObjectBindingManager()->GetRuntimeObjects(Sequencer->GetFocusedMovieSceneInstance(), MovieSceneBinding.GetObjectGuid(), RuntimeObjects);
+
+					for (UObject* RuntimeObject : RuntimeObjects)
+					{
+						UProperty* Property = RuntimeObject->GetClass()->FindPropertyByName(CastChecked<UMovieSceneColorTrack>(Track)->GetPropertyName());
+						if (Property != nullptr)
+						{
+							if (CastChecked<UMovieSceneColorTrack>(Track)->IsSlateColor())
+							{
+								DefaultColor = (*Property->ContainerPtrToValuePtr<FSlateColor>(RuntimeObject)).GetSpecifiedColor();
+							}
+							else
+							{
+								DefaultColor = *Property->ContainerPtrToValuePtr<FLinearColor>(RuntimeObject);
+							}
+							bFoundColor = true;
+							break;
+						}
+					}
+				}
+			}
+		}
+
 		// @todo Sequencer Optimize - This could all get cached, instead of recalculating everything every OnPaint
 
 		const FRichCurve* Curves[4] = {
@@ -154,9 +194,12 @@ private:
 		// @todo Sequencer Optimize - This another O(n^2) loop, since Eval is O(n)!
 		for (int32 i = 0; i < TimesWithKeys.Num(); ++i)
 		{
-			OutColorKeys.Add( TKeyValuePair<float, FLinearColor>(TimesWithKeys[i], Section->Eval(TimesWithKeys[i], FLinearColor(0,0,0,0) ) ) );
+			OutColorKeys.Add( TKeyValuePair<float, FLinearColor>(TimesWithKeys[i], Section->Eval(TimesWithKeys[i], DefaultColor) ) );
 		}
 	}
+
+	ISequencer* Sequencer;
+	UMovieSceneTrack* Track;
 };
 
 /**
@@ -390,7 +433,7 @@ TSharedRef<ISequencerSection> FPropertyTrackEditor::MakeSectionInterface( UMovie
 	UClass* SectionClass = SectionObject.GetOuter()->GetClass();
 
 	TSharedRef<ISequencerSection> NewSection = TSharedRef<ISequencerSection>(
-		SectionClass == UMovieSceneColorTrack::StaticClass() ? new FColorPropertySection( SectionObject, Track->GetTrackName() ) :
+		SectionClass == UMovieSceneColorTrack::StaticClass() ? new FColorPropertySection( SectionObject, Track->GetTrackName(), GetSequencer().Get(), Track ) :
 		SectionClass == UMovieSceneBoolTrack::StaticClass() ? new FBoolPropertySection( SectionObject, Track->GetTrackName(), GetSequencer().Get() ) :
 		SectionClass == UMovieSceneByteTrack::StaticClass() ? new FBytePropertySection(SectionObject, Track->GetTrackName(), Cast<UMovieSceneByteTrack>(SectionObject.GetOuter())->GetEnum()) :
 		SectionClass == UMovieSceneVectorTrack::StaticClass() ? new FVectorPropertySection( SectionObject, Track->GetTrackName() ) :
