Sequencer: Fix track editor release so that track editors aren't left with dangling connections when the active sequencer is destroyed and rebuilt (the track editors are also destroyed and rebuilt). This fixes a bug where if you create a movie scene and manipulate the viewport that's controlled by a movie scene camera, the sequencer crashes when you switch to another movie scene.

diff --git a/Engine/Source/Editor/MovieSceneTools/Private/TrackEditors/TransformTrackEditor.cpp b/Engine/Source/Editor/MovieSceneTools/Private/TrackEditors/TransformTrackEditor.cpp
--- a/Engine/Source/Editor/MovieSceneTools/Private/TrackEditors/TransformTrackEditor.cpp
+++ b/Engine/Source/Editor/MovieSceneTools/Private/TrackEditors/TransformTrackEditor.cpp
@@ -137,14 +137,18 @@ F3DTransformTrackEditor::F3DTransformTrackEditor( TSharedRef<ISequencer> InSeque
 }
 
 F3DTransformTrackEditor::~F3DTransformTrackEditor()
+{
+	OnRelease();
+	F3DTransformTrackCommands::Unregister();
+}
+
+void F3DTransformTrackEditor::OnRelease()
 {
 	GEditor->OnBeginObjectMovement().RemoveAll( this );
 	GEditor->OnEndObjectMovement().RemoveAll( this );
 
 	GEditor->OnBeginCameraMovement().RemoveAll( this );
 	GEditor->OnEndCameraMovement().RemoveAll( this );
-
-	F3DTransformTrackCommands::Unregister();
 }
 
 TSharedRef<FMovieSceneTrackEditor> F3DTransformTrackEditor::CreateTrackEditor( TSharedRef<ISequencer> InSequencer )
diff --git a/Engine/Source/Editor/MovieSceneTools/Private/TrackEditors/TransformTrackEditor.h b/Engine/Source/Editor/MovieSceneTools/Private/TrackEditors/TransformTrackEditor.h
--- a/Engine/Source/Editor/MovieSceneTools/Private/TrackEditors/TransformTrackEditor.h
+++ b/Engine/Source/Editor/MovieSceneTools/Private/TrackEditors/TransformTrackEditor.h
@@ -27,6 +27,7 @@ public:
 	static TSharedRef<FMovieSceneTrackEditor> CreateTrackEditor( TSharedRef<ISequencer> OwningSequencer );
 
 	/** FMovieSceneTrackEditor Interface */
+	virtual void OnRelease() override;
 	virtual bool SupportsType( TSubclassOf<UMovieSceneTrack> Type ) const override;
 	virtual TSharedRef<ISequencerSection> MakeSectionInterface( UMovieSceneSection& SectionObject, UMovieSceneTrack* Track ) override;
 	virtual void AddKey( const FGuid& ObjectGuid, UObject* AdditionalAsset = NULL ) override;
diff --git a/Engine/Source/Editor/Sequencer/Private/Sequencer.cpp b/Engine/Source/Editor/Sequencer/Private/Sequencer.cpp
--- a/Engine/Source/Editor/Sequencer/Private/Sequencer.cpp
+++ b/Engine/Source/Editor/Sequencer/Private/Sequencer.cpp
@@ -193,7 +193,11 @@ FSequencer::OnClose()
 
 		FEditorDelegates::PreSaveWorld.RemoveAll(this);
 		FEditorDelegates::PostSaveWorld.RemoveAll(this);
-
+		
+		for (auto TrackEditor : TrackEditors)
+		{
+			TrackEditor->OnRelease();
+		}
 		ActiveSequencer.Reset();
 
 		if( GLevelEditorModeTools().IsModeActive( FSequencerEdMode::EM_SequencerMode ) )
diff --git a/Engine/Source/Editor/Sequencer/Public/MovieSceneTrackEditor.h b/Engine/Source/Editor/Sequencer/Public/MovieSceneTrackEditor.h
--- a/Engine/Source/Editor/Sequencer/Public/MovieSceneTrackEditor.h
+++ b/Engine/Source/Editor/Sequencer/Public/MovieSceneTrackEditor.h
@@ -26,6 +26,9 @@ public:
 	/** Destructor */
 	virtual ~FMovieSceneTrackEditor(){};
 
+	/** Called when the instance of this track editor is released */
+	virtual void OnRelease(){};
+
 	/** @return The current movie scene */
 	UMovieScene* GetMovieScene() const
 	{
