Sequencer - Refactor the FMovieScenePropertyTrack
+ Move section interface classes to separate files so they can be reused by other tracks.
+ Change to a template class to avoid switching on types so much and implement separate classes for each property type.
+ Delete copy/pasted code from UMG property track editors and derive from property track.

diff --git a/Engine/Source/Editor/MovieSceneTools/MovieSceneTools.Build.cs b/Engine/Source/Editor/MovieSceneTools/MovieSceneTools.Build.cs
--- a/Engine/Source/Editor/MovieSceneTools/MovieSceneTools.Build.cs
+++ b/Engine/Source/Editor/MovieSceneTools/MovieSceneTools.Build.cs
@@ -11,6 +11,8 @@ public class MovieSceneTools : ModuleRules
                 "Editor/MovieSceneTools/Private",
                 "Editor/MovieSceneTools/Private/CurveKeyEditors",
                 "Editor/MovieSceneTools/Private/TrackEditors",
+				"Editor/MovieSceneTools/Private/TrackEditors/PropertyTrackEditors",
+				"Editor/MovieSceneTools/Private/Sections"
             }
         );
 
diff --git a/Engine/Source/Editor/MovieSceneTools/Private/MovieSceneToolsModule.cpp b/Engine/Source/Editor/MovieSceneTools/Private/MovieSceneToolsModule.cpp
--- a/Engine/Source/Editor/MovieSceneTools/Private/MovieSceneToolsModule.cpp
+++ b/Engine/Source/Editor/MovieSceneTools/Private/MovieSceneToolsModule.cpp
@@ -9,7 +9,12 @@
 #include "ScopedTransaction.h"
 #include "MovieScene.h"
 #include "MovieSceneTrackEditor.h"
-#include "PropertyTrackEditor.h"
+#include "BoolPropertyTrackEditor.h"
+#include "BytePropertyTrackEditor.h"
+#include "ColorPropertyTrackEditor.h"
+#include "FloatPropertyTrackEditor.h"
+#include "VectorPropertyTrackEditor.h"
+#include "VisibilityPropertyTrackEditor.h"
 #include "TransformTrackEditor.h"
 #include "ShotTrackEditor.h"
 #include "SubMovieSceneTrackEditor.h"
@@ -28,14 +33,21 @@ class FMovieSceneToolsModule : public IMovieSceneTools
 	{
 		// Register with the sequencer module that we provide auto-key handlers.
 		ISequencerModule& SequencerModule = FModuleManager::Get().LoadModuleChecked<ISequencerModule>( "Sequencer" );
-		PropertyTrackEditorCreateTrackEditorDelegateHandle      = SequencerModule.RegisterTrackEditor_Handle( FOnCreateTrackEditor::CreateStatic( &FPropertyTrackEditor::CreateTrackEditor ) );
-		TransformTrackEditorCreateTrackEditorDelegateHandle     = SequencerModule.RegisterTrackEditor_Handle( FOnCreateTrackEditor::CreateStatic( &F3DTransformTrackEditor::CreateTrackEditor ) );
-		EditorCreateTrackEditorDelegateHandle					= SequencerModule.RegisterTrackEditor_Handle( FOnCreateTrackEditor::CreateStatic( &FShotTrackEditor::CreateTrackEditor ) );
-		SubMovieSceneTrackEditorCreateTrackEditorDelegateHandle = SequencerModule.RegisterTrackEditor_Handle( FOnCreateTrackEditor::CreateStatic( &FSubMovieSceneTrackEditor::CreateTrackEditor ) );
-		AudioTrackEditorCreateTrackEditorDelegateHandle         = SequencerModule.RegisterTrackEditor_Handle( FOnCreateTrackEditor::CreateStatic( &FAudioTrackEditor::CreateTrackEditor ) );
-		AnimationTrackEditorCreateTrackEditorDelegateHandle     = SequencerModule.RegisterTrackEditor_Handle( FOnCreateTrackEditor::CreateStatic( &FSkeletalAnimationTrackEditor::CreateTrackEditor ) );
-		ParticleTrackEditorCreateTrackEditorDelegateHandle      = SequencerModule.RegisterTrackEditor_Handle( FOnCreateTrackEditor::CreateStatic( &FParticleTrackEditor::CreateTrackEditor ) );
-		PathTrackEditorCreateTrackEditorDelegateHandle          = SequencerModule.RegisterTrackEditor_Handle( FOnCreateTrackEditor::CreateStatic( &F3DPathTrackEditor::CreateTrackEditor ) );
+
+		BoolPropertyTrackEditorCreateTrackEditorDelegateHandle   = SequencerModule.RegisterTrackEditor_Handle( FOnCreateTrackEditor::CreateStatic( &FBoolPropertyTrackEditor::CreateTrackEditor ) );
+		BytePropertyTrackEditorCreateTrackEditorDelegateHandle = SequencerModule.RegisterTrackEditor_Handle( FOnCreateTrackEditor::CreateStatic( &FBytePropertyTrackEditor::CreateTrackEditor ) );
+		ColorPropertyTrackEditorCreateTrackEditorDelegateHandle  = SequencerModule.RegisterTrackEditor_Handle( FOnCreateTrackEditor::CreateStatic( &FColorPropertyTrackEditor::CreateTrackEditor ) );
+		FloatPropertyTrackEditorCreateTrackEditorDelegateHandle = SequencerModule.RegisterTrackEditor_Handle( FOnCreateTrackEditor::CreateStatic( &FFloatPropertyTrackEditor::CreateTrackEditor ) );
+		VectorPropertyTrackEditorCreateTrackEditorDelegateHandle = SequencerModule.RegisterTrackEditor_Handle( FOnCreateTrackEditor::CreateStatic( &FVectorPropertyTrackEditor::CreateTrackEditor ) );
+		VisibilityPropertyTrackEditorCreateTrackEditorDelegateHandle = SequencerModule.RegisterTrackEditor_Handle( FOnCreateTrackEditor::CreateStatic( &FVisibilityPropertyTrackEditor::CreateTrackEditor ) );
+
+		TransformTrackEditorCreateTrackEditorDelegateHandle      = SequencerModule.RegisterTrackEditor_Handle( FOnCreateTrackEditor::CreateStatic( &F3DTransformTrackEditor::CreateTrackEditor ) );
+		EditorCreateTrackEditorDelegateHandle                    = SequencerModule.RegisterTrackEditor_Handle( FOnCreateTrackEditor::CreateStatic( &FShotTrackEditor::CreateTrackEditor ) );
+		SubMovieSceneTrackEditorCreateTrackEditorDelegateHandle  = SequencerModule.RegisterTrackEditor_Handle( FOnCreateTrackEditor::CreateStatic( &FSubMovieSceneTrackEditor::CreateTrackEditor ) );
+		AudioTrackEditorCreateTrackEditorDelegateHandle          = SequencerModule.RegisterTrackEditor_Handle( FOnCreateTrackEditor::CreateStatic( &FAudioTrackEditor::CreateTrackEditor ) );
+		AnimationTrackEditorCreateTrackEditorDelegateHandle      = SequencerModule.RegisterTrackEditor_Handle( FOnCreateTrackEditor::CreateStatic( &FSkeletalAnimationTrackEditor::CreateTrackEditor ) );
+		ParticleTrackEditorCreateTrackEditorDelegateHandle       = SequencerModule.RegisterTrackEditor_Handle( FOnCreateTrackEditor::CreateStatic( &FParticleTrackEditor::CreateTrackEditor ) );
+		PathTrackEditorCreateTrackEditorDelegateHandle           = SequencerModule.RegisterTrackEditor_Handle( FOnCreateTrackEditor::CreateStatic( &F3DPathTrackEditor::CreateTrackEditor ) );
 	}
 
 	virtual void ShutdownModule() override
@@ -44,7 +56,14 @@ class FMovieSceneToolsModule : public IMovieSceneTools
 		{
 			// Unregister auto key handlers
 			ISequencerModule& SequencerModule = FModuleManager::Get().GetModuleChecked<ISequencerModule>( "Sequencer" );
-			SequencerModule.UnRegisterTrackEditor_Handle( PropertyTrackEditorCreateTrackEditorDelegateHandle );
+
+			SequencerModule.UnRegisterTrackEditor_Handle( BoolPropertyTrackEditorCreateTrackEditorDelegateHandle );
+			SequencerModule.UnRegisterTrackEditor_Handle( BytePropertyTrackEditorCreateTrackEditorDelegateHandle );
+			SequencerModule.UnRegisterTrackEditor_Handle( ColorPropertyTrackEditorCreateTrackEditorDelegateHandle );
+			SequencerModule.UnRegisterTrackEditor_Handle( FloatPropertyTrackEditorCreateTrackEditorDelegateHandle );
+			SequencerModule.UnRegisterTrackEditor_Handle( VectorPropertyTrackEditorCreateTrackEditorDelegateHandle );
+			SequencerModule.UnRegisterTrackEditor_Handle( VisibilityPropertyTrackEditorCreateTrackEditorDelegateHandle );
+
 			SequencerModule.UnRegisterTrackEditor_Handle( TransformTrackEditorCreateTrackEditorDelegateHandle );
 			SequencerModule.UnRegisterTrackEditor_Handle( EditorCreateTrackEditorDelegateHandle );
 			SequencerModule.UnRegisterTrackEditor_Handle( SubMovieSceneTrackEditorCreateTrackEditorDelegateHandle );
@@ -59,7 +78,13 @@ class FMovieSceneToolsModule : public IMovieSceneTools
 private:
 
 	/** Registered delegate handles */
-	FDelegateHandle PropertyTrackEditorCreateTrackEditorDelegateHandle;
+	FDelegateHandle BoolPropertyTrackEditorCreateTrackEditorDelegateHandle;
+	FDelegateHandle BytePropertyTrackEditorCreateTrackEditorDelegateHandle;
+	FDelegateHandle ColorPropertyTrackEditorCreateTrackEditorDelegateHandle;
+	FDelegateHandle FloatPropertyTrackEditorCreateTrackEditorDelegateHandle;
+	FDelegateHandle VectorPropertyTrackEditorCreateTrackEditorDelegateHandle;
+	FDelegateHandle VisibilityPropertyTrackEditorCreateTrackEditorDelegateHandle;
+
 	FDelegateHandle TransformTrackEditorCreateTrackEditorDelegateHandle;
 	FDelegateHandle EditorCreateTrackEditorDelegateHandle;
 	FDelegateHandle SubMovieSceneTrackEditorCreateTrackEditorDelegateHandle;
diff --git a/Engine/Source/Editor/MovieSceneTools/Private/MovieSceneToolsPrivatePCH.h b/Engine/Source/Editor/MovieSceneTools/Private/MovieSceneToolsPrivatePCH.h
--- a/Engine/Source/Editor/MovieSceneTools/Private/MovieSceneToolsPrivatePCH.h
+++ b/Engine/Source/Editor/MovieSceneTools/Private/MovieSceneToolsPrivatePCH.h
@@ -7,9 +7,14 @@
 #include "MovieScene.h"
 #include "ScopedTransaction.h"
 #include "ISequencerSection.h"
+#include "ISectionLayoutBuilder.h"
 #include "MovieSceneSequenceInstance.h"
 #include "MovieSceneTrackEditor.h"
 #include "IMovieSceneTools.h"
 #include "MovieSceneSection.h"
+#include "PropertySection.h"
+#include "PropertyTrackEditor.h"
+#include "CommonMovieSceneTools.h"
+#include "MovieSceneToolHelpers.h"
 
 #endif // __MovieSceneToolsPrivatePCH_h__
diff --git a/Engine/Source/Editor/MovieSceneTools/Private/Sections/BoolPropertySection.cpp b/Engine/Source/Editor/MovieSceneTools/Private/Sections/BoolPropertySection.cpp
new file mode 100644
--- /dev/null
+++ b/Engine/Source/Editor/MovieSceneTools/Private/Sections/BoolPropertySection.cpp
@@ -0,0 +1,80 @@
+// Copyright 1998-2015 Epic Games, Inc. All Rights Reserved.
+
+#include "MovieSceneToolsPrivatePCH.h"
+#include "MovieSceneBoolTrack.h"
+#include "MovieSceneBoolSection.h"
+#include "BoolPropertySection.h"
+
+void FBoolPropertySection::GenerateSectionLayout( class ISectionLayoutBuilder& LayoutBuilder ) const
+{
+	UMovieSceneBoolSection* BoolSection = Cast<UMovieSceneBoolSection>( &SectionObject );
+	LayoutBuilder.SetSectionAsKeyArea( MakeShareable( new FBoolKeyArea( BoolSection->GetCurve(), BoolSection ) ) );
+}
+
+int32 FBoolPropertySection::OnPaintSection( const FGeometry& AllottedGeometry, const FSlateRect& SectionClippingRect, FSlateWindowElementList& OutDrawElements, int32 LayerId, bool bParentEnabled ) const
+{
+	FTimeToPixel TimeToPixelConverter = SectionObject.IsInfinite() ?
+		FTimeToPixel( AllottedGeometry, Sequencer->GetViewRange() ) :
+		FTimeToPixel( AllottedGeometry, TRange<float>( SectionObject.GetStartTime(), SectionObject.GetEndTime() ) );
+
+	// custom drawing for bool curves
+	UMovieSceneBoolSection* BoolSection = Cast<UMovieSceneBoolSection>( &SectionObject );
+
+	TArray<float> SectionSwitchTimes;
+
+	if ( SectionObject.IsInfinite() && Sequencer->GetViewRange().GetLowerBoundValue() < SectionObject.GetStartTime() )
+	{
+		SectionSwitchTimes.Add( Sequencer->GetViewRange().GetLowerBoundValue() );
+	}
+	SectionSwitchTimes.Add( SectionObject.GetStartTime() );
+
+	FIntegralCurve& BoolCurve = BoolSection->GetCurve();
+	for ( TArray<FIntegralKey>::TConstIterator It( BoolCurve.GetKeyIterator() ); It; ++It )
+	{
+		float Time = It->Time;
+		if ( Time > SectionObject.GetStartTime() && Time < SectionObject.GetEndTime() )
+		{
+			SectionSwitchTimes.Add( Time );
+		}
+	}
+
+	if ( SectionObject.IsInfinite() && Sequencer->GetViewRange().GetUpperBoundValue() > SectionObject.GetEndTime() )
+	{
+		SectionSwitchTimes.Add( Sequencer->GetViewRange().GetUpperBoundValue() );
+	}
+	SectionSwitchTimes.Add( SectionObject.GetEndTime() );
+
+	for ( int32 i = 0; i < SectionSwitchTimes.Num() - 1; ++i )
+	{
+		float FirstTime = SectionSwitchTimes[i];
+		float NextTime = SectionSwitchTimes[i + 1];
+		float FirstTimeInPixels = TimeToPixelConverter.TimeToPixel( FirstTime );
+		float NextTimeInPixels = TimeToPixelConverter.TimeToPixel( NextTime );
+		if ( BoolSection->Eval( FirstTime ) )
+		{
+			FSlateDrawElement::MakeBox(
+				OutDrawElements,
+				LayerId,
+				AllottedGeometry.ToPaintGeometry( FVector2D( FirstTimeInPixels, 0 ), FVector2D( NextTimeInPixels - FirstTimeInPixels, AllottedGeometry.Size.Y ) ),
+				FEditorStyle::GetBrush( "Sequencer.GenericSection.Background" ),
+				SectionClippingRect,
+				ESlateDrawEffect::None,
+				FLinearColor( 0.f, 1.f, 0.f, 1.f )
+				);
+		}
+		else
+		{
+			FSlateDrawElement::MakeBox(
+				OutDrawElements,
+				LayerId,
+				AllottedGeometry.ToPaintGeometry( FVector2D( FirstTimeInPixels, 0 ), FVector2D( NextTimeInPixels - FirstTimeInPixels, AllottedGeometry.Size.Y ) ),
+				FEditorStyle::GetBrush( "Sequencer.GenericSection.Background" ),
+				SectionClippingRect,
+				ESlateDrawEffect::None,
+				FLinearColor( 1.f, 0.f, 0.f, 1.f )
+				);
+		}
+	}
+
+	return LayerId + 1;
+}
\ No newline at end of file
diff --git a/Engine/Source/Editor/MovieSceneTools/Private/Sections/BoolPropertySection.h b/Engine/Source/Editor/MovieSceneTools/Private/Sections/BoolPropertySection.h
new file mode 100644
--- /dev/null
+++ b/Engine/Source/Editor/MovieSceneTools/Private/Sections/BoolPropertySection.h
@@ -0,0 +1,20 @@
+// Copyright 1998-2015 Epic Games, Inc. All Rights Reserved.
+
+#pragma once
+
+/**
+* An implementation of bool property sections
+*/
+class FBoolPropertySection : public FPropertySection
+{
+public:
+	FBoolPropertySection( UMovieSceneSection& InSectionObject, FName SectionName, ISequencer* InSequencer )
+		: FPropertySection( InSectionObject, SectionName )
+		, Sequencer( InSequencer ) {}
+
+	virtual void GenerateSectionLayout( class ISectionLayoutBuilder& LayoutBuilder ) const override;
+
+	virtual int32 OnPaintSection( const FGeometry& AllottedGeometry, const FSlateRect& SectionClippingRect, FSlateWindowElementList& OutDrawElements, int32 LayerId, bool bParentEnabled ) const override;
+private:
+	ISequencer* Sequencer;
+};
\ No newline at end of file
diff --git a/Engine/Source/Editor/MovieSceneTools/Private/Sections/BytePropertySection.cpp b/Engine/Source/Editor/MovieSceneTools/Private/Sections/BytePropertySection.cpp
new file mode 100644
--- /dev/null
+++ b/Engine/Source/Editor/MovieSceneTools/Private/Sections/BytePropertySection.cpp
@@ -0,0 +1,19 @@
+// Copyright 1998-2015 Epic Games, Inc. All Rights Reserved.
+
+#include "MovieSceneToolsPrivatePCH.h"
+#include "MovieSceneByteTrack.h"
+#include "MovieSceneByteSection.h"
+#include "BytePropertySection.h"
+
+void FBytePropertySection::GenerateSectionLayout( class ISectionLayoutBuilder& LayoutBuilder ) const
+{
+	UMovieSceneByteSection* ByteSection = Cast<UMovieSceneByteSection>( &SectionObject );
+	if ( Enum != nullptr )
+	{
+		LayoutBuilder.SetSectionAsKeyArea( MakeShareable( new FEnumKeyArea( ByteSection->GetCurve(), ByteSection, Enum ) ) );
+	}
+	else
+	{
+		LayoutBuilder.SetSectionAsKeyArea( MakeShareable( new FIntegralKeyArea( ByteSection->GetCurve(), ByteSection ) ) );
+	}
+}
\ No newline at end of file
diff --git a/Engine/Source/Editor/MovieSceneTools/Private/Sections/BytePropertySection.h b/Engine/Source/Editor/MovieSceneTools/Private/Sections/BytePropertySection.h
new file mode 100644
--- /dev/null
+++ b/Engine/Source/Editor/MovieSceneTools/Private/Sections/BytePropertySection.h
@@ -0,0 +1,19 @@
+// Copyright 1998-2015 Epic Games, Inc. All Rights Reserved.
+
+#pragma once
+
+/**
+* An implementation of byte property sections
+*/
+class FBytePropertySection : public FPropertySection
+{
+public:
+	FBytePropertySection( UMovieSceneSection& InSectionObject, FName InSectionName, UEnum* InEnum )
+		: FPropertySection( InSectionObject, InSectionName )
+		, Enum( InEnum ) {}
+
+	virtual void GenerateSectionLayout( class ISectionLayoutBuilder& LayoutBuilder ) const override;
+
+private:
+	UEnum* Enum;
+};
\ No newline at end of file
diff --git a/Engine/Source/Editor/MovieSceneTools/Private/Sections/ColorPropertySection.cpp b/Engine/Source/Editor/MovieSceneTools/Private/Sections/ColorPropertySection.cpp
new file mode 100644
--- /dev/null
+++ b/Engine/Source/Editor/MovieSceneTools/Private/Sections/ColorPropertySection.cpp
@@ -0,0 +1,174 @@
+// Copyright 1998-2015 Epic Games, Inc. All Rights Reserved.
+
+#include "MovieSceneToolsPrivatePCH.h"
+#include "MovieSceneColorTrack.h"
+#include "MovieSceneColorSection.h"
+#include "ColorPropertySection.h"
+#include "MovieSceneSequence.h"
+
+void FColorPropertySection::GenerateSectionLayout( class ISectionLayoutBuilder& LayoutBuilder ) const
+{
+	UMovieSceneColorSection* ColorSection = Cast<UMovieSceneColorSection>( &SectionObject );
+
+	LayoutBuilder.AddKeyArea( "R", NSLOCTEXT( "FColorPropertySection", "RedArea", "Red" ), MakeShareable( new FFloatCurveKeyArea( &ColorSection->GetRedCurve(), ColorSection ) ) );
+	LayoutBuilder.AddKeyArea( "G", NSLOCTEXT( "FColorPropertySection", "GreenArea", "Green" ), MakeShareable( new FFloatCurveKeyArea( &ColorSection->GetGreenCurve(), ColorSection ) ) );
+	LayoutBuilder.AddKeyArea( "B", NSLOCTEXT( "FColorPropertySection", "BlueArea", "Blue" ), MakeShareable( new FFloatCurveKeyArea( &ColorSection->GetBlueCurve(), ColorSection ) ) );
+	LayoutBuilder.AddKeyArea( "A", NSLOCTEXT( "FColorPropertySection", "OpacityArea", "Opacity" ), MakeShareable( new FFloatCurveKeyArea( &ColorSection->GetAlphaCurve(), ColorSection ) ) );
+}
+
+int32 FColorPropertySection::OnPaintSection( const FGeometry& AllottedGeometry, const FSlateRect& SectionClippingRect, FSlateWindowElementList& OutDrawElements, int32 LayerId, bool bParentEnabled ) const
+{
+	const UMovieSceneColorSection* ColorSection = Cast<const UMovieSceneColorSection>( &SectionObject );
+
+	float StartTime = ColorSection->GetStartTime();
+	float EndTime = ColorSection->GetEndTime();
+	float SectionDuration = EndTime - StartTime;
+
+	if ( !FMath::IsNearlyZero( SectionDuration ) )
+	{
+		LayerId = FPropertySection::OnPaintSection( AllottedGeometry, SectionClippingRect, OutDrawElements, LayerId, bParentEnabled );
+
+		FVector2D GradientSize = FVector2D( AllottedGeometry.Size.X, (AllottedGeometry.Size.Y / 4) - 3.0f );
+
+		FPaintGeometry PaintGeometry = AllottedGeometry.ToPaintGeometry( FVector2D( 0, 0 ), GradientSize );
+
+		// If we are showing a background pattern and the colors is transparent, draw a checker pattern
+		const FSlateBrush* CheckerBrush = FEditorStyle::GetBrush( "Checker" );
+		FSlateDrawElement::MakeBox( OutDrawElements, LayerId, PaintGeometry, CheckerBrush, SectionClippingRect );
+
+		TArray<FSlateGradientStop> GradientStops;
+
+		TArray< TKeyValuePair<float, FLinearColor> > ColorKeys;
+		ConsolidateColorCurves( ColorKeys, ColorSection );
+
+		for ( int32 i = 0; i < ColorKeys.Num(); ++i )
+		{
+			float Time = ColorKeys[i].Key;
+			FLinearColor Color = ColorKeys[i].Value;
+			float TimeFraction = (Time - StartTime) / SectionDuration;
+
+			GradientStops.Add( FSlateGradientStop( FVector2D( TimeFraction * AllottedGeometry.Size.X, 0 ),
+				Color ) );
+		}
+
+		if ( GradientStops.Num() > 0 )
+		{
+			FSlateDrawElement::MakeGradient(
+				OutDrawElements,
+				LayerId + 1,
+				PaintGeometry,
+				GradientStops,
+				Orient_Vertical,
+				SectionClippingRect
+				);
+		}
+	}
+
+	return LayerId + 1;
+}
+
+
+void FColorPropertySection::ConsolidateColorCurves( TArray< TKeyValuePair<float, FLinearColor> >& OutColorKeys, const UMovieSceneColorSection* Section ) const
+{
+	// Get the default color of the first instance
+	FLinearColor DefaultColor( 0.0f, 0.0f, 0.0f, 0.0f );
+
+	static const FName SlateColorName( "SlateColor" );
+
+	UMovieSceneSequence* Sequence = Sequencer->GetFocusedMovieSceneSequence();
+
+	const TArray<FMovieSceneBinding>& MovieSceneBindings = Sequence->GetMovieScene()->GetBindings();
+
+	bool bFoundColor = false;
+	for ( int32 BindingIndex = 0; BindingIndex < MovieSceneBindings.Num() && !bFoundColor; ++BindingIndex )
+	{
+		const FMovieSceneBinding& MovieSceneBinding = MovieSceneBindings[BindingIndex];
+
+		for ( int32 TrackIndex = 0; TrackIndex < MovieSceneBinding.GetTracks().Num() && !bFoundColor; ++TrackIndex )
+		{
+			if ( MovieSceneBinding.GetTracks()[TrackIndex] == Track )
+			{
+				UObject* RuntimeObject = Sequence->FindObject( MovieSceneBinding.GetObjectGuid() );
+
+				if ( RuntimeObject != nullptr )
+				{
+					UProperty* Property = RuntimeObject->GetClass()->FindPropertyByName( CastChecked<UMovieSceneColorTrack>( Track )->GetPropertyName() );
+					UStructProperty* ColorStructProp = Cast<UStructProperty>( Property );
+					if ( ColorStructProp && ColorStructProp->Struct )
+					{
+						if ( ColorStructProp->Struct->GetFName() == SlateColorName )
+						{
+							DefaultColor = (*Property->ContainerPtrToValuePtr<FSlateColor>( RuntimeObject )).GetSpecifiedColor();
+						}
+						else if ( ColorStructProp->Struct->GetFName() == NAME_LinearColor )
+						{
+							DefaultColor = *Property->ContainerPtrToValuePtr<FLinearColor>( RuntimeObject );
+						}
+						else
+						{
+							DefaultColor = Property->ContainerPtrToValuePtr<FColor>( RuntimeObject )->ReinterpretAsLinear();
+						}
+						bFoundColor = true;
+						break;
+					}
+				}
+			}
+		}
+	}
+
+	// @todo Sequencer Optimize - This could all get cached, instead of recalculating everything every OnPaint
+
+	const FRichCurve* Curves[4] = {
+		&Section->GetRedCurve(),
+		&Section->GetGreenCurve(),
+		&Section->GetBlueCurve(),
+		&Section->GetAlphaCurve()
+	};
+
+	// @todo Sequencer Optimize - This is a O(n^2) loop!
+	// Our times are floats, which means we can't use a map and
+	// do a quick lookup to see if the keys already exist
+	// because the keys are ordered, we could take advantage of that, however
+	TArray<float> TimesWithKeys;
+	for ( int32 i = 0; i < 4; ++i )
+	{
+		const FRichCurve* Curve = Curves[i];
+		for ( auto It( Curve->GetKeyIterator() ); It; ++It )
+		{
+			float KeyTime = It->Time;
+
+			bool bShouldAddKey = true;
+
+			int32 InsertKeyIndex = INDEX_NONE;
+			for ( int32 k = 0; k < TimesWithKeys.Num(); ++k )
+			{
+				if ( FMath::IsNearlyEqual( TimesWithKeys[k], KeyTime ) )
+				{
+					bShouldAddKey = false;
+					break;
+				}
+				else if ( TimesWithKeys[k] > KeyTime )
+				{
+					InsertKeyIndex = k;
+					break;
+				}
+			}
+
+			if ( InsertKeyIndex == INDEX_NONE && bShouldAddKey )
+			{
+				InsertKeyIndex = TimesWithKeys.Num();
+			}
+
+			if ( bShouldAddKey )
+			{
+				TimesWithKeys.Insert( KeyTime, InsertKeyIndex );
+			}
+		}
+	}
+
+	// @todo Sequencer Optimize - This another O(n^2) loop, since Eval is O(n)!
+	for ( int32 i = 0; i < TimesWithKeys.Num(); ++i )
+	{
+		OutColorKeys.Add( TKeyValuePair<float, FLinearColor>( TimesWithKeys[i], Section->Eval( TimesWithKeys[i], DefaultColor ) ) );
+	}
+}
\ No newline at end of file
diff --git a/Engine/Source/Editor/MovieSceneTools/Private/Sections/ColorPropertySection.h b/Engine/Source/Editor/MovieSceneTools/Private/Sections/ColorPropertySection.h
new file mode 100644
--- /dev/null
+++ b/Engine/Source/Editor/MovieSceneTools/Private/Sections/ColorPropertySection.h
@@ -0,0 +1,27 @@
+// Copyright 1998-2015 Epic Games, Inc. All Rights Reserved.
+
+#pragma once
+
+class UMovieSceneColorSection;
+
+/**
+* A color section implementation
+*/
+class FColorPropertySection : public FPropertySection
+{
+public:
+	FColorPropertySection( UMovieSceneSection& InSectionObject, FName SectionName, ISequencer* InSequencer, UMovieSceneTrack* InTrack )
+		: FPropertySection( InSectionObject, SectionName )
+		, Sequencer( InSequencer )
+		, Track( InTrack ) {}
+
+	virtual void GenerateSectionLayout( class ISectionLayoutBuilder& LayoutBuilder ) const override;
+
+	virtual int32 OnPaintSection( const FGeometry& AllottedGeometry, const FSlateRect& SectionClippingRect, FSlateWindowElementList& OutDrawElements, int32 LayerId, bool bParentEnabled ) const override;
+
+private:
+	void ConsolidateColorCurves( TArray< TKeyValuePair<float, FLinearColor> >& OutColorKeys, const UMovieSceneColorSection* Section ) const;
+
+	ISequencer* Sequencer;
+	UMovieSceneTrack* Track;
+};
\ No newline at end of file
diff --git a/Engine/Source/Editor/MovieSceneTools/Private/Sections/FloatPropertySection.cpp b/Engine/Source/Editor/MovieSceneTools/Private/Sections/FloatPropertySection.cpp
new file mode 100644
--- /dev/null
+++ b/Engine/Source/Editor/MovieSceneTools/Private/Sections/FloatPropertySection.cpp
@@ -0,0 +1,11 @@
+// Copyright 1998-2015 Epic Games, Inc. All Rights Reserved.
+
+#include "MovieSceneToolsPrivatePCH.h"
+#include "FloatPropertySection.h"
+#include "MovieSceneFloatSection.h"
+
+void FFloatPropertySection::GenerateSectionLayout( class ISectionLayoutBuilder& LayoutBuilder ) const
+{
+	UMovieSceneFloatSection* FloatSection = Cast<UMovieSceneFloatSection>( &SectionObject );
+	LayoutBuilder.SetSectionAsKeyArea( MakeShareable( new FFloatCurveKeyArea( &FloatSection->GetFloatCurve(), FloatSection ) ) );
+}
\ No newline at end of file
diff --git a/Engine/Source/Editor/MovieSceneTools/Private/Sections/FloatPropertySection.h b/Engine/Source/Editor/MovieSceneTools/Private/Sections/FloatPropertySection.h
new file mode 100644
--- /dev/null
+++ b/Engine/Source/Editor/MovieSceneTools/Private/Sections/FloatPropertySection.h
@@ -0,0 +1,15 @@
+// Copyright 1998-2015 Epic Games, Inc. All Rights Reserved.
+
+#pragma once
+
+/**
+* An implementation of float property sections
+*/
+class FFloatPropertySection : public FPropertySection
+{
+public:
+	FFloatPropertySection( UMovieSceneSection& InSectionObject, FName SectionName )
+		: FPropertySection( InSectionObject, SectionName ) {}
+
+	virtual void GenerateSectionLayout( class ISectionLayoutBuilder& LayoutBuilder ) const override;
+};
\ No newline at end of file
diff --git a/Engine/Source/Editor/MovieSceneTools/Private/Sections/VectorPropertySection.cpp b/Engine/Source/Editor/MovieSceneTools/Private/Sections/VectorPropertySection.cpp
new file mode 100644
--- /dev/null
+++ b/Engine/Source/Editor/MovieSceneTools/Private/Sections/VectorPropertySection.cpp
@@ -0,0 +1,23 @@
+// Copyright 1998-2015 Epic Games, Inc. All Rights Reserved.
+
+#include "MovieSceneToolsPrivatePCH.h"
+#include "VectorPropertySection.h"
+#include "MovieSceneVectorSection.h"
+
+void FVectorPropertySection::GenerateSectionLayout( class ISectionLayoutBuilder& LayoutBuilder ) const
+{
+	UMovieSceneVectorSection* VectorSection = Cast<UMovieSceneVectorSection>( &SectionObject );
+	int32 ChannelsUsed = VectorSection->GetChannelsUsed();
+	check( ChannelsUsed >= 2 && ChannelsUsed <= 4 );
+
+	LayoutBuilder.AddKeyArea( "Vector.X", NSLOCTEXT( "FVectorPropertySection", "XArea", "X" ), MakeShareable( new FFloatCurveKeyArea( &VectorSection->GetCurve( 0 ), VectorSection ) ) );
+	LayoutBuilder.AddKeyArea( "Vector.Y", NSLOCTEXT( "FVectorPropertySection", "YArea", "Y" ), MakeShareable( new FFloatCurveKeyArea( &VectorSection->GetCurve( 1 ), VectorSection ) ) );
+	if ( ChannelsUsed >= 3 )
+	{
+		LayoutBuilder.AddKeyArea( "Vector.Z", NSLOCTEXT( "FVectorPropertySection", "ZArea", "Z" ), MakeShareable( new FFloatCurveKeyArea( &VectorSection->GetCurve( 2 ), VectorSection ) ) );
+	}
+	if ( ChannelsUsed >= 4 )
+	{
+		LayoutBuilder.AddKeyArea( "Vector.W", NSLOCTEXT( "FVectorPropertySection", "WArea", "W" ), MakeShareable( new FFloatCurveKeyArea( &VectorSection->GetCurve( 3 ), VectorSection ) ) );
+	}
+}
\ No newline at end of file
diff --git a/Engine/Source/Editor/MovieSceneTools/Private/Sections/VectorPropertySection.h b/Engine/Source/Editor/MovieSceneTools/Private/Sections/VectorPropertySection.h
new file mode 100644
--- /dev/null
+++ b/Engine/Source/Editor/MovieSceneTools/Private/Sections/VectorPropertySection.h
@@ -0,0 +1,15 @@
+// Copyright 1998-2015 Epic Games, Inc. All Rights Reserved.
+
+#pragma once
+
+/**
+* An implementation of vector property sections
+*/
+class FVectorPropertySection : public FPropertySection
+{
+public:
+	FVectorPropertySection( UMovieSceneSection& InSectionObject, FName SectionName )
+		: FPropertySection( InSectionObject, SectionName ) {}
+
+	virtual void GenerateSectionLayout( class ISectionLayoutBuilder& LayoutBuilder ) const override;
+};
\ No newline at end of file
diff --git a/Engine/Source/Editor/MovieSceneTools/Private/Sections/VisibilityPropertySection.h b/Engine/Source/Editor/MovieSceneTools/Private/Sections/VisibilityPropertySection.h
new file mode 100644
--- /dev/null
+++ b/Engine/Source/Editor/MovieSceneTools/Private/Sections/VisibilityPropertySection.h
@@ -0,0 +1,18 @@
+// Copyright 1998-2015 Epic Games, Inc. All Rights Reserved.
+
+#pragma once
+
+#include "BoolPropertySection.h"
+
+/**
+* An implementation of visibility property sections
+*/
+class FVisibilityPropertySection : public FBoolPropertySection
+{
+public:
+	FVisibilityPropertySection( UMovieSceneSection& InSectionObject, FName SectionName, ISequencer* InSequencer )
+		: FBoolPropertySection( InSectionObject, SectionName, InSequencer )
+	{
+		DisplayName = FText::FromString( TEXT( "Visible" ) );
+	}
+};
\ No newline at end of file
diff --git a/Engine/Source/Editor/MovieSceneTools/Private/TrackEditors/PropertyTrackEditor.cpp b/Engine/Source/Editor/MovieSceneTools/Private/TrackEditors/PropertyTrackEditor.cpp
deleted file mode 100644
diff --git a/Engine/Source/Editor/MovieSceneTools/Private/TrackEditors/PropertyTrackEditor.h b/Engine/Source/Editor/MovieSceneTools/Private/TrackEditors/PropertyTrackEditor.h
deleted file mode 100644
diff --git a/Engine/Source/Editor/MovieSceneTools/Private/TrackEditors/PropertyTrackEditors/BoolPropertyTrackEditor.cpp b/Engine/Source/Editor/MovieSceneTools/Private/TrackEditors/PropertyTrackEditors/BoolPropertyTrackEditor.cpp
new file mode 100644
--- /dev/null
+++ b/Engine/Source/Editor/MovieSceneTools/Private/TrackEditors/PropertyTrackEditors/BoolPropertyTrackEditor.cpp
@@ -0,0 +1,21 @@
+// Copyright 1998-2015 Epic Games, Inc. All Rights Reserved.
+
+#include "MovieSceneToolsPrivatePCH.h"
+#include "BoolPropertyTrackEditor.h"
+#include "BoolPropertySection.h"
+
+TSharedRef<FMovieSceneTrackEditor> FBoolPropertyTrackEditor::CreateTrackEditor( TSharedRef<ISequencer> OwningSequencer )
+{
+	return MakeShareable(new FBoolPropertyTrackEditor(OwningSequencer));
+}
+
+TSharedRef<ISequencerSection> FBoolPropertyTrackEditor::MakeSectionInterface( UMovieSceneSection& SectionObject, UMovieSceneTrack* Track )
+{
+	return MakeShareable(new FBoolPropertySection( SectionObject, Track->GetTrackName(), GetSequencer().Get() ));
+}
+
+bool FBoolPropertyTrackEditor::TryGenerateKeyFromPropertyChanged( const FPropertyChangedParams& PropertyChangedParams, bool& OutKey )
+{
+	OutKey = *PropertyChangedParams.GetPropertyValue<bool>();
+	return true;
+}
\ No newline at end of file
diff --git a/Engine/Source/Editor/MovieSceneTools/Private/TrackEditors/PropertyTrackEditors/BoolPropertyTrackEditor.h b/Engine/Source/Editor/MovieSceneTools/Private/TrackEditors/PropertyTrackEditors/BoolPropertyTrackEditor.h
new file mode 100644
--- /dev/null
+++ b/Engine/Source/Editor/MovieSceneTools/Private/TrackEditors/PropertyTrackEditors/BoolPropertyTrackEditor.h
@@ -0,0 +1,38 @@
+// Copyright 1998-2015 Epic Games, Inc. All Rights Reserved.
+
+#pragma once
+
+#include "MovieSceneBoolTrack.h"
+
+/**
+* A property track editor for bools.
+*/
+class FBoolPropertyTrackEditor : public FPropertyTrackEditor<UMovieSceneBoolTrack, bool>
+{
+public:
+	/**
+	 * Constructor
+	 *
+	 * @param InSequencer	The sequencer instance to be used by this tool
+	 */
+	FBoolPropertyTrackEditor( TSharedRef<ISequencer> InSequencer )
+		: FPropertyTrackEditor( InSequencer, NAME_BoolProperty )
+	{}
+
+	/**
+	 * Creates an instance of this class.  Called by a sequencer 
+	 *
+	 * @param OwningSequencer The sequencer instance to be used by this tool
+	 * @return The new instance of this class
+	 */
+	static TSharedRef<FMovieSceneTrackEditor> CreateTrackEditor( TSharedRef<ISequencer> OwningSequencer );
+
+	/** FMovieSceneTrackEditor Interface */
+	virtual TSharedRef<ISequencerSection> MakeSectionInterface( UMovieSceneSection& SectionObject, UMovieSceneTrack* Track ) override;
+
+protected:
+	/** FPropertyTrackEditor Interface */
+	virtual bool TryGenerateKeyFromPropertyChanged( const FPropertyChangedParams& PropertyChangedParams, bool& OutKey ) override;
+};
+
+
diff --git a/Engine/Source/Editor/MovieSceneTools/Private/TrackEditors/PropertyTrackEditors/BytePropertyTrackEditor.cpp b/Engine/Source/Editor/MovieSceneTools/Private/TrackEditors/PropertyTrackEditors/BytePropertyTrackEditor.cpp
new file mode 100644
--- /dev/null
+++ b/Engine/Source/Editor/MovieSceneTools/Private/TrackEditors/PropertyTrackEditors/BytePropertyTrackEditor.cpp
@@ -0,0 +1,68 @@
+// Copyright 1998-2015 Epic Games, Inc. All Rights Reserved.
+
+#include "MovieSceneToolsPrivatePCH.h"
+#include "BytePropertyTrackEditor.h"
+#include "BytePropertySection.h"
+#include "MovieSceneByteTrack.h"
+#include "MovieSceneSequence.h"
+
+TSharedRef<FMovieSceneTrackEditor> FBytePropertyTrackEditor::CreateTrackEditor( TSharedRef<ISequencer> OwningSequencer )
+{
+	return MakeShareable(new FBytePropertyTrackEditor(OwningSequencer));
+}
+
+TSharedRef<ISequencerSection> FBytePropertyTrackEditor::MakeSectionInterface( UMovieSceneSection& SectionObject, UMovieSceneTrack* Track )
+{
+	return MakeShareable( new FBytePropertySection( SectionObject, Track->GetTrackName(), Cast<UMovieSceneByteTrack>( SectionObject.GetOuter() )->GetEnum() ) );
+}
+
+UEnum* GetEnumForByteTrack(TSharedPtr<ISequencer> Sequencer, const FGuid& OwnerObjectHandle, FName PropertyName, UMovieSceneByteTrack* ByteTrack)
+{
+	
+	UObject* RuntimeObject = Sequencer->GetFocusedMovieSceneSequence()->FindObject(OwnerObjectHandle);
+
+	TSet<UEnum*> PropertyEnums;
+	if (RuntimeObject != nullptr)
+	{
+		UProperty* Property = RuntimeObject->GetClass()->FindPropertyByName(PropertyName);
+		if (Property != nullptr)
+		{
+			UByteProperty* ByteProperty = Cast<UByteProperty>(Property);
+			if (ByteProperty != nullptr && ByteProperty->Enum != nullptr)
+			{
+				PropertyEnums.Add(ByteProperty->Enum);
+			}
+		}
+	}
+
+	UEnum* TrackEnum;
+	if (PropertyEnums.Num() == 1)
+	{
+		TrackEnum = PropertyEnums.Array()[0];
+	}
+	else
+	{
+		TrackEnum = nullptr;
+	}
+	return TrackEnum;
+}
+
+UMovieSceneTrack* FBytePropertyTrackEditor::AddTrack(UMovieScene* FocusedMovieScene, const FGuid& ObjectHandle, TSubclassOf<class UMovieSceneTrack> TrackClass, FName UniqueTypeName)
+{
+	UMovieSceneTrack* NewTrack = FPropertyTrackEditor::AddTrack(FocusedMovieScene, ObjectHandle, TrackClass, UniqueTypeName);
+	
+	UMovieSceneByteTrack* ByteTrack = Cast<UMovieSceneByteTrack>(NewTrack);
+	UEnum* TrackEnum = GetEnumForByteTrack(GetSequencer(), ObjectHandle, UniqueTypeName, ByteTrack);
+	if (TrackEnum != nullptr)
+	{
+		ByteTrack->SetEnum(TrackEnum);
+	}
+
+	return NewTrack;
+}
+
+bool FBytePropertyTrackEditor::TryGenerateKeyFromPropertyChanged( const FPropertyChangedParams& PropertyChangedParams, uint8& OutKey )
+{
+	OutKey = *PropertyChangedParams.GetPropertyValue<uint8>();
+	return true;
+}
\ No newline at end of file
diff --git a/Engine/Source/Editor/MovieSceneTools/Private/TrackEditors/PropertyTrackEditors/BytePropertyTrackEditor.h b/Engine/Source/Editor/MovieSceneTools/Private/TrackEditors/PropertyTrackEditors/BytePropertyTrackEditor.h
new file mode 100644
--- /dev/null
+++ b/Engine/Source/Editor/MovieSceneTools/Private/TrackEditors/PropertyTrackEditors/BytePropertyTrackEditor.h
@@ -0,0 +1,39 @@
+// Copyright 1998-2015 Epic Games, Inc. All Rights Reserved.
+
+#pragma once
+
+#include "MovieSceneByteTrack.h"
+
+/**
+* A property track editor for byte and enums.
+*/
+class FBytePropertyTrackEditor : public FPropertyTrackEditor<UMovieSceneByteTrack, uint8>
+{
+public:
+	/**
+	 * Constructor
+	 *
+	 * @param InSequencer	The sequencer instance to be used by this tool
+	 */
+	FBytePropertyTrackEditor( TSharedRef<ISequencer> InSequencer)
+		: FPropertyTrackEditor( InSequencer, NAME_ByteProperty )
+	{}
+
+	/**
+	 * Creates an instance of this class.  Called by a sequencer 
+	 *
+	 * @param OwningSequencer The sequencer instance to be used by this tool
+	 * @return The new instance of this class
+	 */
+	static TSharedRef<FMovieSceneTrackEditor> CreateTrackEditor( TSharedRef<ISequencer> OwningSequencer );
+
+	/** FMovieSceneTrackEditor Interface */
+	virtual TSharedRef<ISequencerSection> MakeSectionInterface( UMovieSceneSection& SectionObject, UMovieSceneTrack* Track ) override;
+	virtual UMovieSceneTrack* AddTrack( UMovieScene* FocusedMovieScene, const FGuid& ObjectHandle, TSubclassOf<class UMovieSceneTrack> TrackClass, FName UniqueTypeName ) override;
+
+protected:
+	/** FPropertyTrackEditor Interface */
+	virtual bool TryGenerateKeyFromPropertyChanged( const FPropertyChangedParams& PropertyChangedParams, uint8& OutKey ) override;
+};
+
+
diff --git a/Engine/Source/Editor/MovieSceneTools/Private/TrackEditors/PropertyTrackEditors/ColorPropertyTrackEditor.cpp b/Engine/Source/Editor/MovieSceneTools/Private/TrackEditors/PropertyTrackEditors/ColorPropertyTrackEditor.cpp
new file mode 100644
--- /dev/null
+++ b/Engine/Source/Editor/MovieSceneTools/Private/TrackEditors/PropertyTrackEditors/ColorPropertyTrackEditor.cpp
@@ -0,0 +1,58 @@
+// Copyright 1998-2015 Epic Games, Inc. All Rights Reserved.
+
+#include "MovieSceneToolsPrivatePCH.h"
+#include "ColorPropertyTrackEditor.h"
+#include "ColorPropertySection.h"
+
+TSharedRef<FMovieSceneTrackEditor> FColorPropertyTrackEditor::CreateTrackEditor( TSharedRef<ISequencer> InSequencer )
+{
+	return MakeShareable( new FColorPropertyTrackEditor( InSequencer ) );
+}
+
+TSharedRef<ISequencerSection> FColorPropertyTrackEditor::MakeSectionInterface( UMovieSceneSection& SectionObject, UMovieSceneTrack* Track )
+{
+	return MakeShareable(new FColorPropertySection( SectionObject, Track->GetTrackName(), GetSequencer().Get(), Track ));
+}
+
+bool FColorPropertyTrackEditor::TryGenerateKeyFromPropertyChanged( const FPropertyChangedParams& PropertyChangedParams, FColorKey& OutKey )
+{
+	bool bIsFColor = false, bIsFLinearColor = false, bIsSlateColor = false;
+	const UStructProperty* StructProp = Cast<const UStructProperty>(PropertyChangedParams.PropertyPath.Last());
+	if (StructProp && StructProp->Struct)
+	{
+		FName StructName = StructProp->Struct->GetFName();
+
+		bIsFColor = StructName == NAME_Color;
+		bIsFLinearColor = StructName == NAME_LinearColor;
+		bIsSlateColor = StructName == FName("SlateColor");
+	}
+
+	if (!bIsFColor && !bIsFLinearColor && !bIsSlateColor) 
+	{
+		return false;
+	}
+
+	FName PropertyName = PropertyChangedParams.PropertyPath.Last()->GetFName();
+
+	FLinearColor ColorValue;
+	if (bIsFColor)
+	{
+		ColorValue = PropertyChangedParams.GetPropertyValue<FColor>()->ReinterpretAsLinear();
+	}
+	else
+	{
+		ColorValue = *PropertyChangedParams.GetPropertyValue<FLinearColor>();
+	}
+
+	if( StructProp->HasMetaData("HideAlphaChannel") )
+	{
+		ColorValue.A = 1;
+	}
+
+	OutKey.Value = ColorValue;
+	OutKey.CurveName = PropertyChangedParams.StructPropertyNameToKey;
+	OutKey.bAddKeyEvenIfUnchanged = !PropertyChangedParams.bRequireAutoKey;
+	OutKey.bIsSlateColor = bIsSlateColor;
+
+	return true;
+}
diff --git a/Engine/Source/Editor/MovieSceneTools/Private/TrackEditors/PropertyTrackEditors/ColorPropertyTrackEditor.h b/Engine/Source/Editor/MovieSceneTools/Private/TrackEditors/PropertyTrackEditors/ColorPropertyTrackEditor.h
new file mode 100644
--- /dev/null
+++ b/Engine/Source/Editor/MovieSceneTools/Private/TrackEditors/PropertyTrackEditors/ColorPropertyTrackEditor.h
@@ -0,0 +1,38 @@
+// Copyright 1998-2015 Epic Games, Inc. All Rights Reserved.
+
+#pragma once
+
+#include "MovieSceneColorTrack.h"
+
+/**
+* A property track editor for colors.
+*/
+class FColorPropertyTrackEditor : public FPropertyTrackEditor<UMovieSceneColorTrack, FColorKey>
+{
+public:
+	/**
+	 * Constructor
+	 *
+	 * @param InSequencer	The sequencer instance to be used by this tool
+	 */
+	FColorPropertyTrackEditor( TSharedRef<ISequencer> InSequencer )
+		: FPropertyTrackEditor<UMovieSceneColorTrack, FColorKey>( InSequencer, NAME_Color, NAME_LinearColor, "SlateColor" )
+	{}
+
+	/**
+	 * Creates an instance of this class.  Called by a sequencer 
+	 *
+	 * @param OwningSequencer The sequencer instance to be used by this tool
+	 * @return The new instance of this class
+	 */
+	static TSharedRef<FMovieSceneTrackEditor> CreateTrackEditor( TSharedRef<ISequencer> OwningSequencer );
+
+	/** FMovieSceneTrackEditor Interface */
+	virtual TSharedRef<ISequencerSection> MakeSectionInterface( UMovieSceneSection& SectionObject, UMovieSceneTrack* Track ) override;
+
+protected:
+	/** FPropertyTrackEditor Interface */
+	virtual bool TryGenerateKeyFromPropertyChanged( const FPropertyChangedParams& PropertyChangedParams, FColorKey& OutKey ) override;
+};
+
+
diff --git a/Engine/Source/Editor/MovieSceneTools/Private/TrackEditors/PropertyTrackEditors/VectorPropertyTrackEditor.cpp b/Engine/Source/Editor/MovieSceneTools/Private/TrackEditors/PropertyTrackEditors/VectorPropertyTrackEditor.cpp
new file mode 100644
--- /dev/null
+++ b/Engine/Source/Editor/MovieSceneTools/Private/TrackEditors/PropertyTrackEditors/VectorPropertyTrackEditor.cpp
@@ -0,0 +1,48 @@
+// Copyright 1998-2015 Epic Games, Inc. All Rights Reserved.
+
+#include "MovieSceneToolsPrivatePCH.h"
+#include "VectorPropertyTrackEditor.h"
+#include "VectorPropertySection.h"
+
+TSharedRef<FMovieSceneTrackEditor> FVectorPropertyTrackEditor::CreateTrackEditor( TSharedRef<ISequencer> InSequencer )
+{
+	return MakeShareable( new FVectorPropertyTrackEditor( InSequencer ) );
+}
+
+TSharedRef<ISequencerSection> FVectorPropertyTrackEditor::MakeSectionInterface( UMovieSceneSection& SectionObject, UMovieSceneTrack* Track )
+{
+	return MakeShareable(new FVectorPropertySection( SectionObject, Track->GetTrackName() ));
+}
+
+bool FVectorPropertyTrackEditor::TryGenerateKeyFromPropertyChanged( const FPropertyChangedParams& PropertyChangedParams, FVectorKey& OutKey )
+{
+	bool bIsVector2D = false, bIsVector = false, bIsVector4 = false;
+	const UStructProperty* StructProp = Cast<const UStructProperty>( PropertyChangedParams.PropertyPath.Last() );
+	if ( StructProp && StructProp->Struct )
+	{
+		FName StructName = StructProp->Struct->GetFName();
+
+		bIsVector2D = StructName == NAME_Vector2D;
+		bIsVector = StructName == NAME_Vector;
+		bIsVector4 = StructName == NAME_Vector4;
+	}
+	if ( !bIsVector2D && !bIsVector && !bIsVector4 )
+	{
+		return false;
+	}
+
+	// Get the vector value from the property
+	if ( bIsVector2D )
+	{
+		OutKey = FVectorKey(*PropertyChangedParams.GetPropertyValue<FVector2D>(), PropertyChangedParams.StructPropertyNameToKey, !PropertyChangedParams.bRequireAutoKey);
+	}
+	else if ( bIsVector )
+	{
+		OutKey = FVectorKey( *PropertyChangedParams.GetPropertyValue<FVector>(), PropertyChangedParams.StructPropertyNameToKey, !PropertyChangedParams.bRequireAutoKey );
+	}
+	else // if ( bIsVector4 )
+	{
+		OutKey = FVectorKey( *PropertyChangedParams.GetPropertyValue<FVector4>(), PropertyChangedParams.StructPropertyNameToKey, !PropertyChangedParams.bRequireAutoKey );
+	}
+	return true;
+}
diff --git a/Engine/Source/Editor/MovieSceneTools/Private/TrackEditors/PropertyTrackEditors/VectorPropertyTrackEditor.h b/Engine/Source/Editor/MovieSceneTools/Private/TrackEditors/PropertyTrackEditors/VectorPropertyTrackEditor.h
new file mode 100644
--- /dev/null
+++ b/Engine/Source/Editor/MovieSceneTools/Private/TrackEditors/PropertyTrackEditors/VectorPropertyTrackEditor.h
@@ -0,0 +1,38 @@
+// Copyright 1998-2015 Epic Games, Inc. All Rights Reserved.
+
+#pragma once
+
+#include "MovieSceneVectorTrack.h"
+
+/**
+ * A property track editor for vectors.
+ */
+class FVectorPropertyTrackEditor : public FPropertyTrackEditor<UMovieSceneVectorTrack, FVectorKey>
+{
+public:
+	/**
+	 * Constructor
+	 *
+	 * @param InSequencer	The sequencer instance to be used by this tool
+	 */
+	FVectorPropertyTrackEditor( TSharedRef<ISequencer> InSequencer )
+		: FPropertyTrackEditor<UMovieSceneVectorTrack, FVectorKey>( InSequencer, NAME_Vector, NAME_Vector4, NAME_Vector2D )
+	{}
+
+	/**
+	 * Creates an instance of this class.  Called by a sequencer 
+	 *
+	 * @param OwningSequencer The sequencer instance to be used by this tool
+	 * @return The new instance of this class
+	 */
+	static TSharedRef<FMovieSceneTrackEditor> CreateTrackEditor( TSharedRef<ISequencer> OwningSequencer );
+
+	/** FMovieSceneTrackEditor Interface */
+	virtual TSharedRef<ISequencerSection> MakeSectionInterface( UMovieSceneSection& SectionObject, UMovieSceneTrack* Track ) override;
+
+protected:
+	/** FPropertyTrackEditor Interface */
+	virtual bool TryGenerateKeyFromPropertyChanged( const FPropertyChangedParams& PropertyChangedParams, FVectorKey& OutKey ) override;
+};
+
+
diff --git a/Engine/Source/Editor/MovieSceneTools/Public/PropertyTrackEditor.h b/Engine/Source/Editor/MovieSceneTools/Public/PropertyTrackEditor.h
new file mode 100644
--- /dev/null
+++ b/Engine/Source/Editor/MovieSceneTools/Public/PropertyTrackEditor.h
@@ -0,0 +1,199 @@
+// Copyright 1998-2015 Epic Games, Inc. All Rights Reserved.
+
+#pragma once
+
+#include "MovieSceneTrackEditor.h"
+#include "ISequencerObjectChangeListener.h"
+
+class IPropertyHandle;
+class FPropertyChangedParams;
+
+/**
+ * Tools for animatable property types such as floats ands vectors
+ */
+ template<typename TrackType, typename KeyType>
+class FPropertyTrackEditor : public FMovieSceneTrackEditor
+{
+public:
+	/**
+	* Constructor
+	*
+	* @param InSequencer The sequencer instance to be used by this tool
+	*/
+	FPropertyTrackEditor( TSharedRef<ISequencer> InSequencer )
+		: FMovieSceneTrackEditor( InSequencer )
+	{ }
+
+	/**
+	* Constructor
+	*
+	* @param InSequencer The sequencer instance to be used by this tool
+	* @param WatchedPropertyTypeName The property type name that this property track editor should watch for changes.
+	*/
+	FPropertyTrackEditor( TSharedRef<ISequencer> InSequencer, FName WatchedPropertyTypeName )
+		: FMovieSceneTrackEditor( InSequencer )
+	{
+		AddWatchedPropertyType( WatchedPropertyTypeName );
+	}
+
+	/**
+	* Constructor
+	*
+	* @param InSequencer The sequencer instance to be used by this tool
+	* @param WatchedPropertyTypeName1 The first property type name that this property track editor should watch for changes.
+	* @param WatchedPropertyTypeName2 The second property type name that this property track editor should watch for changes.
+	*/
+	FPropertyTrackEditor( TSharedRef<ISequencer> InSequencer, FName WatchedPropertyTypeName1, FName WatchedPropertyTypeName2 )
+		: FMovieSceneTrackEditor( InSequencer )
+	{
+		AddWatchedPropertyType( WatchedPropertyTypeName1 );
+		AddWatchedPropertyType( WatchedPropertyTypeName2 );
+	}
+
+	/**
+	* Constructor
+	*
+	* @param InSequencer The sequencer instance to be used by this tool
+	* @param WatchedPropertyTypeName1 The first property type name that this property track editor should watch for changes.
+	* @param WatchedPropertyTypeName2 The second property type name that this property track editor should watch for changes.
+	* @param WatchedPropertyTypeName3 The third property type name that this property track editor should watch for changes.
+	*/
+	FPropertyTrackEditor( TSharedRef<ISequencer> InSequencer, FName WatchedPropertyTypeName1, FName WatchedPropertyTypeName2, FName WatchedPropertyTypeName3 )
+		: FMovieSceneTrackEditor( InSequencer )
+	{
+		AddWatchedPropertyType( WatchedPropertyTypeName1 );
+		AddWatchedPropertyType( WatchedPropertyTypeName2 );
+		AddWatchedPropertyType( WatchedPropertyTypeName3 );
+	}
+
+	/**
+	* Constructor
+	*
+	* @param InSequencer The sequencer instance to be used by this tool
+	* @param WatchedPropertyTypeNameS An array of property type names that this property track editor should watch for changes.
+	*/
+	FPropertyTrackEditor( TSharedRef<ISequencer> InSequencer, TArray<FName> WatchedPropertyTypeNames )
+		: FMovieSceneTrackEditor( InSequencer )
+	{
+		for ( FName WatchedPropertyTypeName : WatchedPropertyTypeNames )
+		{
+			AddWatchedPropertyType( WatchedPropertyTypeName );
+		}
+	}
+
+	~FPropertyTrackEditor()
+	{
+		TSharedPtr<ISequencer> Sequencer = GetSequencer();
+		if ( Sequencer.IsValid() )
+		{
+			ISequencerObjectChangeListener& ObjectChangeListener = Sequencer->GetObjectChangeListener();
+			for ( FName WatchedPropertyTypeName : WatchedPropertyTypeNames )
+			{
+				ObjectChangeListener.GetOnAnimatablePropertyChanged( WatchedPropertyTypeName ).RemoveAll( this );
+			}
+		}
+	}
+
+	/** FMovieSceneTrackEditor Interface */
+	virtual bool SupportsType( TSubclassOf<UMovieSceneTrack> Type ) const override
+	{
+		return Type == TrackType::StaticClass();
+	}
+
+	virtual void AddKey( const FGuid& ObjectGuid, UObject* AdditionalAsset = NULL ) override
+	{
+		ISequencerObjectChangeListener& ObjectChangeListener = GetSequencer()->GetObjectChangeListener();
+
+		TArray<UObject*> OutObjects;
+		GetSequencer()->GetRuntimeObjects( GetSequencer()->GetFocusedMovieSceneSequenceInstance(), ObjectGuid, OutObjects );
+		for ( int32 i = 0; i < OutObjects.Num(); ++i )
+		{
+			ObjectChangeListener.TriggerAllPropertiesChanged( OutObjects[i] );
+		}
+	}
+
+protected:
+	/**
+	 * Tries to generate a key for the track based on a property change event.  This will be called when a property changes
+	 * which matches a property type name which was supplied on construction.
+	 *
+	 * @param PropertyChangedParams Parameters associated with the property change.
+	 * @param OutKey If this call is successful this represents the key which was generated.
+	 *
+	 * @return true if generating a key was successful, otherwise false.
+	 */
+	virtual bool TryGenerateKeyFromPropertyChanged( const FPropertyChangedParams& PropertyChangedParams, KeyType& OutKey ) = 0;
+	
+private:
+	/** Adds a callback for property changes for the supplied property type name. */
+	void AddWatchedPropertyType( FName WatchedPropertyTypeName )
+	{
+		GetSequencer()->GetObjectChangeListener().GetOnAnimatablePropertyChanged( WatchedPropertyTypeName ).AddRaw( this, &FPropertyTrackEditor<TrackType, KeyType>::OnAnimatedPropertyChanged );
+		WatchedPropertyTypeNames.Add( WatchedPropertyTypeName );
+	}
+
+	/**
+	* Called by the details panel when an animatable property changes
+	*
+	* @param InObjectsThatChanged	List of objects that changed
+	* @param KeyPropertyParams		Parameters for the property change.
+	*/
+	virtual void OnAnimatedPropertyChanged( const FPropertyChangedParams& PropertyChangedParams )
+	{
+		// Get the value from the property
+		KeyType Key;
+		if ( TryGenerateKeyFromPropertyChanged( PropertyChangedParams, Key ) )
+		{
+			AnimatablePropertyChanged( TrackType::StaticClass(), PropertyChangedParams.bRequireAutoKey,
+				FOnKeyProperty::CreateRaw( this, &FPropertyTrackEditor::OnKeyProperty, PropertyChangedParams, Key ) );
+		}
+	}
+
+	/** Adds a key based on a property change. */
+	void OnKeyProperty( float KeyTime, FPropertyChangedParams PropertyChangedParams, KeyType Key )
+	{
+		for ( UObject* Object : PropertyChangedParams.ObjectsThatChanged )
+		{
+			FGuid ObjectHandle = FindOrCreateHandleToObject( Object );
+			if ( ObjectHandle.IsValid() )
+			{
+				FName PropertyName = PropertyChangedParams.PropertyPath.Last()->GetFName();
+
+				// Look for a customized track class for this property on the meta data
+				const FString& MetaSequencerTrackClass = PropertyChangedParams.PropertyPath.Last()->GetMetaData( TEXT( "SequencerTrackClass" ) );
+				TSubclassOf<UMovieSceneTrack> SequencerTrackClass = TrackType::StaticClass();
+				if ( !MetaSequencerTrackClass.IsEmpty() )
+				{
+					UClass* MetaClass = FindObject<UClass>( ANY_PACKAGE, *MetaSequencerTrackClass );
+					if ( !MetaClass )
+					{
+						MetaClass = LoadObject<UClass>( nullptr, *MetaSequencerTrackClass );
+					}
+					if ( MetaClass != NULL )
+						SequencerTrackClass = MetaClass;
+				}
+
+				UMovieSceneTrack* Track = GetTrackForObject( ObjectHandle, SequencerTrackClass, PropertyName );
+				if ( ensure( Track ) )
+				{
+					TrackType* TypedTrack = CastChecked<TrackType>( Track );
+
+					TypedTrack->SetPropertyNameAndPath( PropertyName, PropertyChangedParams.GetPropertyPathString() );
+					// Find or add a new section at the auto-key time and changing the property same property
+					// AddKeyToSection is not actually a virtual, it's redefined in each class with a different type
+					bool bSuccessfulAdd = TypedTrack->AddKeyToSection( KeyTime, Key );
+					if ( bSuccessfulAdd )
+					{
+						TypedTrack->SetAsShowable();
+					}
+				}
+			}
+		}
+	}
+
+private:
+	/** An array of property type names which are being watched for changes. */
+	TArray<FName> WatchedPropertyTypeNames;
+};
+
+
diff --git a/Engine/Source/Editor/UMGEditor/Private/Animation/MarginTrackEditor.cpp b/Engine/Source/Editor/UMGEditor/Private/Animation/MarginTrackEditor.cpp
--- a/Engine/Source/Editor/UMGEditor/Private/Animation/MarginTrackEditor.cpp
+++ b/Engine/Source/Editor/UMGEditor/Private/Animation/MarginTrackEditor.cpp
@@ -2,13 +2,10 @@
 
 #include "UMGEditorPrivatePCH.h"
 #include "MarginTrackEditor.h"
-#include "Editor/MovieSceneTools/Public/PropertySection.h"
-#include "Editor/MovieSceneTools/Public/MovieSceneToolHelpers.h"
-#include "Runtime/UMG/Public/Animation/MovieSceneMarginSection.h"
-#include "Runtime/UMG/Public/Animation/MovieSceneMarginTrack.h"
-#include "Editor/Sequencer/Public/ISectionLayoutBuilder.h"
-#include "Editor/Sequencer/Public/ISequencerObjectChangeListener.h"
-#include "Editor/PropertyEditor/Public/PropertyHandle.h"
+#include "MovieSceneMarginSection.h"
+#include "PropertySection.h"
+#include "ISectionLayoutBuilder.h"
+#include "MovieSceneToolHelpers.h"
 
 class FMarginPropertySection : public FPropertySection
 {
@@ -27,36 +24,11 @@ public:
 	}
 };
 
-FMarginTrackEditor::FMarginTrackEditor( TSharedRef<ISequencer> InSequencer )
-	: FMovieSceneTrackEditor( InSequencer ) 
-{
-	// Get the object change listener for the sequencer and register a delegates for when properties change that we care about
-	ISequencerObjectChangeListener& ObjectChangeListener = InSequencer->GetObjectChangeListener();
-	ObjectChangeListener.GetOnAnimatablePropertyChanged( "Margin" ).AddRaw( this, &FMarginTrackEditor::OnMarginChanged );
-}
-
-FMarginTrackEditor::~FMarginTrackEditor()
-{
-	TSharedPtr<ISequencer> Sequencer = GetSequencer();
-	if( Sequencer.IsValid() )
-	{
-		ISequencerObjectChangeListener& ObjectChangeListener = Sequencer->GetObjectChangeListener();
-		ObjectChangeListener.GetOnAnimatablePropertyChanged( "Margin" ).RemoveAll( this );
-	}
-}
-
-
-
 TSharedRef<FMovieSceneTrackEditor> FMarginTrackEditor::CreateTrackEditor( TSharedRef<ISequencer> InSequencer )
 {
 	return MakeShareable( new FMarginTrackEditor( InSequencer ) );
 }
 
-bool FMarginTrackEditor::SupportsType( TSubclassOf<UMovieSceneTrack> Type ) const
-{
-	return Type == UMovieSceneMarginTrack::StaticClass();
-}
-
 TSharedRef<ISequencerSection> FMarginTrackEditor::MakeSectionInterface( UMovieSceneSection& SectionObject, UMovieSceneTrack* Track )
 {
 	check( SupportsType( SectionObject.GetOuter()->GetClass() ) );
@@ -68,48 +40,10 @@ TSharedRef<ISequencerSection> FMarginTrackEditor::MakeSectionInterface( UMovieSc
 	return NewSection;
 }
 
-
-void FMarginTrackEditor::OnMarginChanged(  const FPropertyChangedParams& PropertyChangedParams )
-{
-	AnimatablePropertyChanged
-	(
-		UMovieSceneMarginTrack::StaticClass(), 
-		PropertyChangedParams.bRequireAutoKey,
-		FOnKeyProperty::CreateRaw(this, &FMarginTrackEditor::OnKeyMargin, &PropertyChangedParams ) 
-	);
-}
-
-
-void FMarginTrackEditor::OnKeyMargin( float KeyTime, const FPropertyChangedParams* PropertyChangedParams )
+bool FMarginTrackEditor::TryGenerateKeyFromPropertyChanged( const FPropertyChangedParams& PropertyChangedParams, FMarginKey& OutKey )
 {
-	FMargin MarginValue = *PropertyChangedParams->GetPropertyValue<FMargin>();
-	FName PropertyName = PropertyChangedParams->PropertyPath.Last()->GetFName();
-
-	for( int32 ObjectIndex = 0; ObjectIndex < PropertyChangedParams->ObjectsThatChanged.Num(); ++ObjectIndex )
-	{
-		UObject* Object = PropertyChangedParams->ObjectsThatChanged[ObjectIndex];
-
-		FMarginKey Key;
-		Key.bAddKeyEvenIfUnchanged = !PropertyChangedParams->bRequireAutoKey;
-		Key.CurveName = PropertyChangedParams->StructPropertyNameToKey;
-		Key.Value = MarginValue;
-
-		FGuid ObjectHandle = FindOrCreateHandleToObject( Object );
-		if (ObjectHandle.IsValid())
-		{
-			UMovieSceneTrack* Track = GetTrackForObject( ObjectHandle, UMovieSceneMarginTrack::StaticClass(), PropertyName );
-			if( ensure( Track ) )
-			{
-				UMovieSceneMarginTrack* MarginTrack = CastChecked<UMovieSceneMarginTrack>(Track);
-				MarginTrack->SetPropertyNameAndPath( PropertyName, PropertyChangedParams->GetPropertyPathString() );
-				// Find or add a new section at the auto-key time and changing the property same property
-				// AddKeyToSection is not actually a virtual, it's redefined in each class with a different type
-				bool bSuccessfulAdd = MarginTrack->AddKeyToSection( KeyTime, Key );
-				if (bSuccessfulAdd)
-				{
-					MarginTrack->SetAsShowable();
-				}
-			}
-		}
-	}
+	OutKey.bAddKeyEvenIfUnchanged = !PropertyChangedParams.bRequireAutoKey;
+	OutKey.CurveName = PropertyChangedParams.StructPropertyNameToKey;
+	OutKey.Value = *PropertyChangedParams.GetPropertyValue<FMargin>();
+	return true;
 }
\ No newline at end of file
diff --git a/Engine/Source/Editor/UMGEditor/Private/Animation/MarginTrackEditor.h b/Engine/Source/Editor/UMGEditor/Private/Animation/MarginTrackEditor.h
--- a/Engine/Source/Editor/UMGEditor/Private/Animation/MarginTrackEditor.h
+++ b/Engine/Source/Editor/UMGEditor/Private/Animation/MarginTrackEditor.h
@@ -2,11 +2,12 @@
 
 #pragma once
 
-#include "Editor/Sequencer/Public/MovieSceneTrackEditor.h"
+#include "PropertyTrackEditor.h"
+#include "MovieSceneMarginTrack.h"
 
-class IPropertyHandle;
+class UMovieSceneTrack;
 
-class FMarginTrackEditor : public FMovieSceneTrackEditor
+class FMarginTrackEditor : public FPropertyTrackEditor<UMovieSceneMarginTrack, FMarginKey>
 {
 public:
 	/**
@@ -14,8 +15,9 @@ public:
 	 *
 	 * @param InSequencer	The sequencer instance to be used by this tool
 	 */
-	FMarginTrackEditor( TSharedRef<ISequencer> InSequencer );
-	~FMarginTrackEditor();
+	FMarginTrackEditor( TSharedRef<ISequencer> InSequencer )
+		: FPropertyTrackEditor<UMovieSceneMarginTrack, FMarginKey>( InSequencer, "Margin" )
+	{ }
 
 	/**
 	 * Creates an instance of this class.  Called by a sequencer 
@@ -26,20 +28,11 @@ public:
 	static TSharedRef<FMovieSceneTrackEditor> CreateTrackEditor( TSharedRef<ISequencer> OwningSequencer );
 
 	/** FMovieSceneTrackEditor Interface */
-	virtual bool SupportsType( TSubclassOf<UMovieSceneTrack> Type ) const override;
 	virtual TSharedRef<ISequencerSection> MakeSectionInterface( UMovieSceneSection& SectionObject, UMovieSceneTrack* Track ) override;
 
-private:
-	/**
-	 * Called by the details panel when an animatable property changes
-	 *
-	 * @param InObjectsThatChanged	List of objects that changed
-	 * @param PropertyValue			Handle to the property value which changed
-	 */
-	void OnMarginChanged( const class FPropertyChangedParams& PropertyChangedParams );
-
-	/** Called After OnMarginChanged if we actually can key the margin */
-	void OnKeyMargin( float KeyTime, const class FPropertyChangedParams* PropertyChangedParams );
+protected:
+	/** FPropertyTrackEditor Interface */
+	virtual bool TryGenerateKeyFromPropertyChanged( const FPropertyChangedParams& PropertyChangedParams, FMarginKey& OutKey ) override;
 };
 
 
diff --git a/Engine/Source/Editor/UMGEditor/Private/Animation/Sequencer2DTransformTrackEditor.cpp b/Engine/Source/Editor/UMGEditor/Private/Animation/Sequencer2DTransformTrackEditor.cpp
--- a/Engine/Source/Editor/UMGEditor/Private/Animation/Sequencer2DTransformTrackEditor.cpp
+++ b/Engine/Source/Editor/UMGEditor/Private/Animation/Sequencer2DTransformTrackEditor.cpp
@@ -2,14 +2,10 @@
 
 #include "UMGEditorPrivatePCH.h"
 #include "Sequencer2DTransformTrackEditor.h"
-#include "Editor/MovieSceneTools/Public/PropertySection.h"
-#include "Editor/MovieSceneTools/Public/MovieSceneToolHelpers.h"
-#include "Runtime/UMG/Public/Animation/MovieScene2DTransformSection.h"
-#include "Runtime/UMG/Public/Animation/MovieScene2DTransformTrack.h"
-#include "Editor/Sequencer/Public/ISectionLayoutBuilder.h"
-#include "Editor/Sequencer/Public/ISequencerObjectChangeListener.h"
-#include "Editor/PropertyEditor/Public/PropertyHandle.h"
-#include "Slate/WidgetTransform.h"
+#include "MovieScene2DTransformSection.h"
+#include "PropertySection.h"
+#include "ISectionLayoutBuilder.h"
+#include "MovieSceneToolHelpers.h"
 
 class F2DTransformSection : public FPropertySection
 {
@@ -43,36 +39,11 @@ public:
 	}
 };
 
-F2DTransformTrackEditor::F2DTransformTrackEditor( TSharedRef<ISequencer> InSequencer )
-	: FMovieSceneTrackEditor( InSequencer ) 
-{
-	// Get the object change listener for the sequencer and register a delegates for when properties change that we care about
-	ISequencerObjectChangeListener& ObjectChangeListener = InSequencer->GetObjectChangeListener();
-	ObjectChangeListener.GetOnAnimatablePropertyChanged( "WidgetTransform" ).AddRaw( this, &F2DTransformTrackEditor::OnTransformChanged );
-}
-
-F2DTransformTrackEditor::~F2DTransformTrackEditor()
-{
-	TSharedPtr<ISequencer> Sequencer = GetSequencer();
-	if( Sequencer.IsValid() )
-	{
-		ISequencerObjectChangeListener& ObjectChangeListener = Sequencer->GetObjectChangeListener();
-		ObjectChangeListener.GetOnAnimatablePropertyChanged( "WidgetTransform" ).RemoveAll( this );
-	}
-}
-
-
-
 TSharedRef<FMovieSceneTrackEditor> F2DTransformTrackEditor::CreateTrackEditor( TSharedRef<ISequencer> InSequencer )
 {
 	return MakeShareable( new F2DTransformTrackEditor( InSequencer ) );
 }
 
-bool F2DTransformTrackEditor::SupportsType( TSubclassOf<UMovieSceneTrack> Type ) const
-{
-	return Type == UMovieScene2DTransformTrack::StaticClass();
-}
-
 TSharedRef<ISequencerSection> F2DTransformTrackEditor::MakeSectionInterface( UMovieSceneSection& SectionObject, UMovieSceneTrack* Track )
 {
 	check( SupportsType( SectionObject.GetOuter()->GetClass() ) );
@@ -84,49 +55,10 @@ TSharedRef<ISequencerSection> F2DTransformTrackEditor::MakeSectionInterface( UMo
 	return NewSection;
 }
 
-
-void F2DTransformTrackEditor::OnTransformChanged( const FPropertyChangedParams& PropertyChangedParams )
-{
-	AnimatablePropertyChanged
-	(
-		UMovieScene2DTransformTrack::StaticClass(), 
-		PropertyChangedParams.bRequireAutoKey,
-		FOnKeyProperty::CreateRaw(this, &F2DTransformTrackEditor::OnKeyTransform, &PropertyChangedParams ) 
-	);
-}
-
-
-void F2DTransformTrackEditor::OnKeyTransform( float KeyTime, const FPropertyChangedParams* PropertyChangedParams )
+bool F2DTransformTrackEditor::TryGenerateKeyFromPropertyChanged( const FPropertyChangedParams& PropertyChangedParams, F2DTransformKey& OutKey )
 {
-	FWidgetTransform TransformValue = *PropertyChangedParams->GetPropertyValue<FWidgetTransform>();
-
-	FName PropertyName = PropertyChangedParams->PropertyPath.Last()->GetFName();
-
-	for( int32 ObjectIndex = 0; ObjectIndex < PropertyChangedParams->ObjectsThatChanged.Num(); ++ObjectIndex )
-	{
-		UObject* Object = PropertyChangedParams->ObjectsThatChanged[ObjectIndex];
-
-		F2DTransformKey Key;
-		Key.bAddKeyEvenIfUnchanged = !PropertyChangedParams->bRequireAutoKey;
-		Key.CurveName = PropertyChangedParams->StructPropertyNameToKey;
-		Key.Value = TransformValue;
-
-		FGuid ObjectHandle = FindOrCreateHandleToObject( Object );
-		if (ObjectHandle.IsValid())
-		{
-			UMovieSceneTrack* Track = GetTrackForObject( ObjectHandle, UMovieScene2DTransformTrack::StaticClass(), PropertyName );
-			if( ensure( Track ) )
-			{
-				UMovieScene2DTransformTrack* TransformTrack = CastChecked<UMovieScene2DTransformTrack>(Track);
-				TransformTrack->SetPropertyNameAndPath( PropertyName, PropertyChangedParams->GetPropertyPathString() );
-				// Find or add a new section at the auto-key time and changing the property same property
-				// AddKeyToSection is not actually a virtual, it's redefined in each class with a different type
-				bool bSuccessfulAdd = TransformTrack->AddKeyToSection( KeyTime, Key );
-				if (bSuccessfulAdd)
-				{
-					TransformTrack->SetAsShowable();
-				}
-			}
-		}
-	}
+	OutKey.bAddKeyEvenIfUnchanged = !PropertyChangedParams.bRequireAutoKey;
+	OutKey.CurveName = PropertyChangedParams.StructPropertyNameToKey;
+	OutKey.Value = *PropertyChangedParams.GetPropertyValue<FWidgetTransform>();
+	return true;
 }
\ No newline at end of file
diff --git a/Engine/Source/Editor/UMGEditor/Private/Animation/Sequencer2DTransformTrackEditor.h b/Engine/Source/Editor/UMGEditor/Private/Animation/Sequencer2DTransformTrackEditor.h
--- a/Engine/Source/Editor/UMGEditor/Private/Animation/Sequencer2DTransformTrackEditor.h
+++ b/Engine/Source/Editor/UMGEditor/Private/Animation/Sequencer2DTransformTrackEditor.h
@@ -2,11 +2,12 @@
 
 #pragma once
 
-#include "Editor/Sequencer/Public/MovieSceneTrackEditor.h"
+#include "PropertyTrackEditor.h"
+#include "MovieScene2DTransformTrack.h"
 
-class IPropertyHandle;
+class UMovieSceneTrack;
 
-class F2DTransformTrackEditor : public FMovieSceneTrackEditor
+class F2DTransformTrackEditor : public FPropertyTrackEditor<UMovieScene2DTransformTrack, F2DTransformKey>
 {
 public:
 	/**
@@ -14,8 +15,9 @@ public:
 	 *
 	 * @param InSequencer	The sequencer instance to be used by this tool
 	 */
-	F2DTransformTrackEditor( TSharedRef<ISequencer> InSequencer );
-	~F2DTransformTrackEditor();
+	F2DTransformTrackEditor( TSharedRef<ISequencer> InSequencer )
+		: FPropertyTrackEditor<UMovieScene2DTransformTrack, F2DTransformKey>( InSequencer, "WidgetTransform" )
+	{ }
 
 	/**
 	 * Creates an instance of this class.  Called by a sequencer 
@@ -26,20 +28,11 @@ public:
 	static TSharedRef<FMovieSceneTrackEditor> CreateTrackEditor( TSharedRef<ISequencer> OwningSequencer );
 
 	/** FMovieSceneTrackEditor Interface */
-	virtual bool SupportsType( TSubclassOf<UMovieSceneTrack> Type ) const override;
 	virtual TSharedRef<ISequencerSection> MakeSectionInterface( UMovieSceneSection& SectionObject, UMovieSceneTrack* Track ) override;
 
-private:
-	/**
-	 * Called by the details panel when an animatable property changes
-	 *
-	 * @param InObjectsThatChanged	List of objects that changed
-	 * @param PropertyValue			Handle to the property value which changed
-	 */
-	void OnTransformChanged( const class FPropertyChangedParams& PropertyChangedParams );
-
-	/** Called After OnMarginChanged if we actually can key the margin */
-	void OnKeyTransform( float KeyTime, const class FPropertyChangedParams* PropertyChangedParams );
+protected:
+	/** FPropertyTrackEditor Interface */
+	virtual bool TryGenerateKeyFromPropertyChanged( const FPropertyChangedParams& PropertyChangedParams, F2DTransformKey& OutKey ) override;
 };
 
 
diff --git a/Engine/Source/Runtime/MovieSceneTracks/Private/Tracks/MovieSceneVectorTrack.cpp b/Engine/Source/Runtime/MovieSceneTracks/Private/Tracks/MovieSceneVectorTrack.cpp
--- a/Engine/Source/Runtime/MovieSceneTracks/Private/Tracks/MovieSceneVectorTrack.cpp
+++ b/Engine/Source/Runtime/MovieSceneTracks/Private/Tracks/MovieSceneVectorTrack.cpp
@@ -6,6 +6,36 @@
 #include "IMovieScenePlayer.h"
 #include "MovieSceneVectorTrackInstance.h"
 
+FVectorKey::FVectorKey( FVector2D InValue, FName InCurveName, bool InbAddKeyEvenIfUnChanged )
+{
+	Value.X = InValue.X;
+	Value.Y = InValue.Y;
+	Value.Z = 0;
+	Value.W = 0;
+	ChannelsUsed = 2;
+	CurveName = InCurveName;
+	bAddKeyEvenIfUnchanged = InbAddKeyEvenIfUnChanged;
+}
+
+FVectorKey::FVectorKey( FVector InValue, FName InCurveName, bool InbAddKeyEvenIfUnChanged )
+{
+	Value.X = InValue.X;
+	Value.Y = InValue.Y;
+	Value.Z = InValue.Z;
+	Value.W = 0;
+	ChannelsUsed = 3;
+	CurveName = InCurveName;
+	bAddKeyEvenIfUnchanged = InbAddKeyEvenIfUnChanged;
+}
+
+FVectorKey::FVectorKey( FVector4 InValue, FName InCurveName, bool InbAddKeyEvenIfUnChanged )
+{
+	Value = InValue;
+	ChannelsUsed = 4;
+	CurveName = InCurveName;
+	bAddKeyEvenIfUnchanged = InbAddKeyEvenIfUnChanged;
+}
+
 UMovieSceneVectorTrack::UMovieSceneVectorTrack( const FObjectInitializer& ObjectInitializer )
 	: Super( ObjectInitializer )
 {
@@ -43,19 +73,9 @@ bool UMovieSceneVectorTrack::AddKeyToSection( float Time, const FVector4& Value,
 	return false;
 }
 
-bool UMovieSceneVectorTrack::AddKeyToSection( float Time, const FVectorKey<FVector4>& Key )
-{
-	return AddKeyToSection(Time, Key.Value, 4, Key.CurveName, Key.bAddKeyEvenIfUnchanged );
-}
-
-bool UMovieSceneVectorTrack::AddKeyToSection( float Time, const FVectorKey<FVector>& Key )
-{
-	return AddKeyToSection(Time, FVector4(Key.Value.X, Key.Value.Y, Key.Value.Z, 0.f), 3, Key.CurveName, Key.bAddKeyEvenIfUnchanged );
-}
-
-bool UMovieSceneVectorTrack::AddKeyToSection( float Time, const FVectorKey<FVector2D>& Key )
+bool UMovieSceneVectorTrack::AddKeyToSection( float Time, FVectorKey Key )
 {
-	return AddKeyToSection(Time, FVector4(Key.Value.X, Key.Value.Y, 0.f, 0.f), 2, Key.CurveName, Key.bAddKeyEvenIfUnchanged );
+	return AddKeyToSection(Time, Key.Value, Key.ChannelsUsed, Key.CurveName, Key.bAddKeyEvenIfUnchanged );
 }
 
 bool UMovieSceneVectorTrack::Eval( float Position, float LastPosition, FVector4& InOutVector ) const
diff --git a/Engine/Source/Runtime/MovieSceneTracks/Public/Tracks/MovieSceneVectorTrack.h b/Engine/Source/Runtime/MovieSceneTracks/Public/Tracks/MovieSceneVectorTrack.h
--- a/Engine/Source/Runtime/MovieSceneTracks/Public/Tracks/MovieSceneVectorTrack.h
+++ b/Engine/Source/Runtime/MovieSceneTracks/Public/Tracks/MovieSceneVectorTrack.h
@@ -9,10 +9,15 @@
 
 #include "MovieSceneVectorTrack.generated.h"
 
-template<typename VecType>
-struct FVectorKey
+struct MOVIESCENETRACKS_API FVectorKey
 {
-	VecType Value;
+	FVectorKey() { }
+	FVectorKey( FVector InValue, FName InCurveName, bool InbAddKeyEvenIfUnChanged );
+	FVectorKey( FVector2D InValue, FName InCurveName, bool InbAddKeyEvenIfUnChanged );
+	FVectorKey( FVector4 InValue, FName InCurveName, bool InbAddKeyEvenIfUnChanged );
+
+	FVector4 Value;
+	uint32 ChannelsUsed;
 	FName CurveName;
 	bool bAddKeyEvenIfUnchanged;
 };
@@ -40,9 +45,7 @@ public:
 	 * @return True if the key was successfully added.
 	 */
 
-	virtual bool AddKeyToSection( float Time, const FVectorKey<FVector4>& Key );
-	virtual bool AddKeyToSection( float Time, const FVectorKey<FVector>& Key );
-	virtual bool AddKeyToSection( float Time, const FVectorKey<FVector2D>& Key );
+	virtual bool AddKeyToSection( float Time, FVectorKey Key );
 	
 	/**
 	 * Evaluates the track at the playback position
