Sequencer - Refactor the FMovieScenePropertyTrack - Missing from last checkin

diff --git a/Engine/Source/Editor/MovieSceneTools/Private/TrackEditors/PropertyTrackEditors/FloatPropertyTrackEditor.cpp b/Engine/Source/Editor/MovieSceneTools/Private/TrackEditors/PropertyTrackEditors/FloatPropertyTrackEditor.cpp
new file mode 100644
--- /dev/null
+++ b/Engine/Source/Editor/MovieSceneTools/Private/TrackEditors/PropertyTrackEditors/FloatPropertyTrackEditor.cpp
@@ -0,0 +1,21 @@
+// Copyright 1998-2015 Epic Games, Inc. All Rights Reserved.
+
+#include "MovieSceneToolsPrivatePCH.h"
+#include "FloatPropertyTrackEditor.h"
+#include "FloatPropertySection.h"
+
+TSharedRef<FMovieSceneTrackEditor> FFloatPropertyTrackEditor::CreateTrackEditor( TSharedRef<ISequencer> OwningSequencer )
+{
+	return MakeShareable(new FFloatPropertyTrackEditor(OwningSequencer));
+}
+
+TSharedRef<ISequencerSection> FFloatPropertyTrackEditor::MakeSectionInterface( UMovieSceneSection& SectionObject, UMovieSceneTrack* Track )
+{
+	return MakeShareable(new FFloatPropertySection( SectionObject, Track->GetTrackName() ));
+}
+
+bool FFloatPropertyTrackEditor::TryGenerateKeyFromPropertyChanged( const FPropertyChangedParams& PropertyChangedParams, float& OutKey )
+{
+	OutKey = *PropertyChangedParams.GetPropertyValue<float>();
+	return true;
+}
\ No newline at end of file
diff --git a/Engine/Source/Editor/MovieSceneTools/Private/TrackEditors/PropertyTrackEditors/FloatPropertyTrackEditor.h b/Engine/Source/Editor/MovieSceneTools/Private/TrackEditors/PropertyTrackEditors/FloatPropertyTrackEditor.h
new file mode 100644
--- /dev/null
+++ b/Engine/Source/Editor/MovieSceneTools/Private/TrackEditors/PropertyTrackEditors/FloatPropertyTrackEditor.h
@@ -0,0 +1,38 @@
+// Copyright 1998-2015 Epic Games, Inc. All Rights Reserved.
+
+#pragma once
+
+#include "MovieSceneFloatTrack.h"
+
+/**
+ * A property track editor for floats.
+ */
+class FFloatPropertyTrackEditor : public FPropertyTrackEditor<UMovieSceneFloatTrack, float>
+{
+public:
+	/**
+	 * Constructor
+	 *
+	 * @param InSequencer	The sequencer instance to be used by this tool
+	 */
+	FFloatPropertyTrackEditor( TSharedRef<ISequencer> InSequencer )
+		: FPropertyTrackEditor( InSequencer, NAME_FloatProperty )
+	{}
+
+	/**
+	 * Creates an instance of this class.  Called by a sequencer 
+	 *
+	 * @param OwningSequencer The sequencer instance to be used by this tool
+	 * @return The new instance of this class
+	 */
+	static TSharedRef<FMovieSceneTrackEditor> CreateTrackEditor( TSharedRef<ISequencer> OwningSequencer );
+
+	/** FMovieSceneTrackEditor Interface */
+	virtual TSharedRef<ISequencerSection> MakeSectionInterface( UMovieSceneSection& SectionObject, UMovieSceneTrack* Track ) override;
+
+protected:
+	/** FPropertyTrackEditor Interface */
+	virtual bool TryGenerateKeyFromPropertyChanged( const FPropertyChangedParams& PropertyChangedParams, float& OutKey ) override;
+};
+
+
diff --git a/Engine/Source/Editor/MovieSceneTools/Private/TrackEditors/PropertyTrackEditors/VisibilityPropertyTrackEditor.cpp b/Engine/Source/Editor/MovieSceneTools/Private/TrackEditors/PropertyTrackEditors/VisibilityPropertyTrackEditor.cpp
new file mode 100644
--- /dev/null
+++ b/Engine/Source/Editor/MovieSceneTools/Private/TrackEditors/PropertyTrackEditors/VisibilityPropertyTrackEditor.cpp
@@ -0,0 +1,21 @@
+// Copyright 1998-2015 Epic Games, Inc. All Rights Reserved.
+
+#include "MovieSceneToolsPrivatePCH.h"
+#include "VisibilityPropertyTrackEditor.h"
+#include "VisibilityPropertySection.h"
+
+TSharedRef<FMovieSceneTrackEditor> FVisibilityPropertyTrackEditor::CreateTrackEditor( TSharedRef<ISequencer> OwningSequencer )
+{
+	return MakeShareable(new FVisibilityPropertyTrackEditor(OwningSequencer));
+}
+
+TSharedRef<ISequencerSection> FVisibilityPropertyTrackEditor::MakeSectionInterface( UMovieSceneSection& SectionObject, UMovieSceneTrack* Track )
+{
+	return MakeShareable(new FVisibilityPropertySection( SectionObject, Track->GetTrackName(), GetSequencer().Get() ));
+}
+
+bool FVisibilityPropertyTrackEditor::TryGenerateKeyFromPropertyChanged( const FPropertyChangedParams& PropertyChangedParams, bool& OutKey )
+{
+	OutKey = *PropertyChangedParams.GetPropertyValue<bool>();
+	return true;
+}
\ No newline at end of file
diff --git a/Engine/Source/Editor/MovieSceneTools/Private/TrackEditors/PropertyTrackEditors/VisibilityPropertyTrackEditor.h b/Engine/Source/Editor/MovieSceneTools/Private/TrackEditors/PropertyTrackEditors/VisibilityPropertyTrackEditor.h
new file mode 100644
--- /dev/null
+++ b/Engine/Source/Editor/MovieSceneTools/Private/TrackEditors/PropertyTrackEditors/VisibilityPropertyTrackEditor.h
@@ -0,0 +1,34 @@
+// Copyright 1998-2015 Epic Games, Inc. All Rights Reserved.
+
+#pragma once
+
+#include "MovieSceneVisibilityTrack.h"
+
+/**
+ * A property track for bool properties which have been set to use a UMovieSceneVisibilityTrack through metadata.
+ */
+class FVisibilityPropertyTrackEditor : public FPropertyTrackEditor<UMovieSceneVisibilityTrack, bool>
+{
+public:
+	FVisibilityPropertyTrackEditor( TSharedRef<ISequencer> InSequencer)
+		// Don't supply any property type names to watch since the FBoolPropertyTrackEditor is already watching for bool property changes.
+		: FPropertyTrackEditor<UMovieSceneVisibilityTrack, bool>( InSequencer )
+	{ }
+
+	/**
+	 * Creates an instance of this class.  Called by a sequencer 
+	 *
+	 * @param OwningSequencer The sequencer instance to be used by this tool
+	 * @return The new instance of this class
+	 */
+	static TSharedRef<FMovieSceneTrackEditor> CreateTrackEditor( TSharedRef<ISequencer> OwningSequencer );
+
+	/** FMovieSceneTrackEditor Interface */
+	virtual TSharedRef<ISequencerSection> MakeSectionInterface( UMovieSceneSection& SectionObject, UMovieSceneTrack* Track ) override;
+
+protected:
+	/** FPropertyTrackEditor Interface */
+	virtual bool TryGenerateKeyFromPropertyChanged( const FPropertyChangedParams& PropertyChangedParams, bool& OutKey ) override;
+};
+
+
