[BuildWatcher] Fixing shadowed variable warnings on the Mac

diff --git a/Engine/Source/Editor/MovieSceneTools/Public/PropertyTrackEditor.h b/Engine/Source/Editor/MovieSceneTools/Public/PropertyTrackEditor.h
--- a/Engine/Source/Editor/MovieSceneTools/Public/PropertyTrackEditor.h
+++ b/Engine/Source/Editor/MovieSceneTools/Public/PropertyTrackEditor.h
@@ -72,10 +72,10 @@ public:
 	* @param InSequencer The sequencer instance to be used by this tool
 	* @param WatchedPropertyTypeNameS An array of property type names that this property track editor should watch for changes.
 	*/
-	FPropertyTrackEditor( TSharedRef<ISequencer> InSequencer, TArray<FName> WatchedPropertyTypeNames )
+	FPropertyTrackEditor( TSharedRef<ISequencer> InSequencer, TArray<FName> InWatchedPropertyTypeNames )
 		: FMovieSceneTrackEditor( InSequencer )
 	{
-		for ( FName WatchedPropertyTypeName : WatchedPropertyTypeNames )
+		for ( FName WatchedPropertyTypeName : InWatchedPropertyTypeNames )
 		{
 			AddWatchedPropertyType( WatchedPropertyTypeName );
 		}
