Sequencer: Add transactions for single float and integral key changes (not slider changes).

diff --git a/Engine/Source/Editor/MovieSceneTools/Private/CurveKeyEditors/SFloatCurveKeyEditor.cpp b/Engine/Source/Editor/MovieSceneTools/Private/CurveKeyEditors/SFloatCurveKeyEditor.cpp
--- a/Engine/Source/Editor/MovieSceneTools/Private/CurveKeyEditors/SFloatCurveKeyEditor.cpp
+++ b/Engine/Source/Editor/MovieSceneTools/Private/CurveKeyEditors/SFloatCurveKeyEditor.cpp
@@ -23,6 +23,7 @@ void SFloatCurveKeyEditor::Construct(const FArguments& InArgs)
 		.MinDesiredWidth(60.0f)
 		.Value(this, &SFloatCurveKeyEditor::OnGetKeyValue)
 		.OnValueChanged(this, &SFloatCurveKeyEditor::OnValueChanged)
+		.OnValueCommitted(this, &SFloatCurveKeyEditor::OnValueCommitted)
 		.OnBeginSliderMovement(this, &SFloatCurveKeyEditor::OnBeginSliderMovement)
 		.OnEndSliderMovement(this, &SFloatCurveKeyEditor::OnEndSliderMovement)
 	];
@@ -74,4 +75,14 @@ void SFloatCurveKeyEditor::OnValueChanged(float Value)
 	Sequencer->UpdateRuntimeInstances();
 }
 
+void SFloatCurveKeyEditor::OnValueCommitted(float Value, ETextCommit::Type CommitInfo)
+{
+	if (CommitInfo == ETextCommit::OnEnter)
+	{
+		const FScopedTransaction Transaction( LOCTEXT("SetFloatKey", "Set float key value") );
+
+		OnValueChanged(Value);
+	}
+}
+
 #undef LOCTEXT_NAMESPACE
\ No newline at end of file
diff --git a/Engine/Source/Editor/MovieSceneTools/Private/CurveKeyEditors/SFloatCurveKeyEditor.h b/Engine/Source/Editor/MovieSceneTools/Private/CurveKeyEditors/SFloatCurveKeyEditor.h
--- a/Engine/Source/Editor/MovieSceneTools/Private/CurveKeyEditors/SFloatCurveKeyEditor.h
+++ b/Engine/Source/Editor/MovieSceneTools/Private/CurveKeyEditors/SFloatCurveKeyEditor.h
@@ -24,6 +24,7 @@ private:
 	float OnGetKeyValue() const;
 
 	void OnValueChanged(float Value);
+	void OnValueCommitted(float Value, ETextCommit::Type CommitInfo);
 
 	ISequencer* Sequencer;
 	UMovieSceneSection* OwningSection;
diff --git a/Engine/Source/Editor/MovieSceneTools/Private/CurveKeyEditors/SIntegralCurveKeyEditor.cpp b/Engine/Source/Editor/MovieSceneTools/Private/CurveKeyEditors/SIntegralCurveKeyEditor.cpp
--- a/Engine/Source/Editor/MovieSceneTools/Private/CurveKeyEditors/SIntegralCurveKeyEditor.cpp
+++ b/Engine/Source/Editor/MovieSceneTools/Private/CurveKeyEditors/SIntegralCurveKeyEditor.cpp
@@ -23,6 +23,7 @@ void SIntegralCurveKeyEditor::Construct(const FArguments& InArgs)
 		.MinDesiredWidth(60.0f)
 		.Value(this, &SIntegralCurveKeyEditor::OnGetKeyValue)
 		.OnValueChanged(this, &SIntegralCurveKeyEditor::OnValueChanged)
+		.OnValueCommitted(this, &SIntegralCurveKeyEditor::OnValueCommitted)
 		.OnBeginSliderMovement(this, &SIntegralCurveKeyEditor::OnBeginSliderMovement)
 		.OnEndSliderMovement(this, &SIntegralCurveKeyEditor::OnEndSliderMovement)
 	];
@@ -72,4 +73,14 @@ void SIntegralCurveKeyEditor::OnValueChanged(int32 Value)
 	Sequencer->UpdateRuntimeInstances();
 }
 
+void SIntegralCurveKeyEditor::OnValueCommitted(int32 Value, ETextCommit::Type CommitInfo)
+{
+	if (CommitInfo == ETextCommit::OnEnter)
+	{
+		const FScopedTransaction Transaction( LOCTEXT("SetIntegralKey", "Set integral key value") );
+
+		OnValueChanged(Value);
+	}
+}
+
 #undef LOCTEXT_NAMESPACE
\ No newline at end of file
diff --git a/Engine/Source/Editor/MovieSceneTools/Private/CurveKeyEditors/SIntegralCurveKeyEditor.h b/Engine/Source/Editor/MovieSceneTools/Private/CurveKeyEditors/SIntegralCurveKeyEditor.h
--- a/Engine/Source/Editor/MovieSceneTools/Private/CurveKeyEditors/SIntegralCurveKeyEditor.h
+++ b/Engine/Source/Editor/MovieSceneTools/Private/CurveKeyEditors/SIntegralCurveKeyEditor.h
@@ -24,6 +24,7 @@ private:
 	int32 OnGetKeyValue() const;
 
 	void OnValueChanged(int32 Value);
+	void OnValueCommitted(int32 Value, ETextCommit::Type CommitInfo);
 
 	ISequencer* Sequencer;
 	UMovieSceneSection* OwningSection;
