Sequencer: Consolidate right click menu track operations to AddTrackMenu.

diff --git a/Engine/Source/Editor/MovieSceneTools/Private/TrackEditors/ParticleTrackEditor.cpp b/Engine/Source/Editor/MovieSceneTools/Private/TrackEditors/ParticleTrackEditor.cpp
--- a/Engine/Source/Editor/MovieSceneTools/Private/TrackEditors/ParticleTrackEditor.cpp
+++ b/Engine/Source/Editor/MovieSceneTools/Private/TrackEditors/ParticleTrackEditor.cpp
@@ -133,7 +133,7 @@ TSharedRef<ISequencerSection> FParticleTrackEditor::MakeSectionInterface( UMovie
 	return NewSection;
 }
 
-void FParticleTrackEditor::BuildObjectBindingContextMenu(FMenuBuilder& MenuBuilder, const FGuid& ObjectBinding, const UClass* ObjectClass)
+void FParticleTrackEditor::BuildObjectBindingTrackMenu(FMenuBuilder& MenuBuilder, const FGuid& ObjectBinding, const UClass* ObjectClass)
 {
 	if (ObjectClass->IsChildOf(AEmitter::StaticClass()))
 	{
diff --git a/Engine/Source/Editor/MovieSceneTools/Private/TrackEditors/ParticleTrackEditor.h b/Engine/Source/Editor/MovieSceneTools/Private/TrackEditors/ParticleTrackEditor.h
--- a/Engine/Source/Editor/MovieSceneTools/Private/TrackEditors/ParticleTrackEditor.h
+++ b/Engine/Source/Editor/MovieSceneTools/Private/TrackEditors/ParticleTrackEditor.h
@@ -28,7 +28,7 @@ public:
 	/** FMovieSceneTrackEditor Interface */
 	virtual bool SupportsType( TSubclassOf<UMovieSceneTrack> Type ) const override;
 	virtual TSharedRef<ISequencerSection> MakeSectionInterface( UMovieSceneSection& SectionObject, UMovieSceneTrack* Track ) override;
-	virtual void BuildObjectBindingContextMenu(FMenuBuilder& MenuBuilder, const FGuid& ObjectBinding, const UClass* ObjectClass) override;
+	virtual void BuildObjectBindingTrackMenu(FMenuBuilder& MenuBuilder, const FGuid& ObjectBinding, const UClass* ObjectClass) override;
 	
 	void AddParticleKey(const FGuid ObjectGuid, bool bTrigger);
 
diff --git a/Engine/Source/Editor/MovieSceneTools/Private/TrackEditors/PathTrackEditor.cpp b/Engine/Source/Editor/MovieSceneTools/Private/TrackEditors/PathTrackEditor.cpp
--- a/Engine/Source/Editor/MovieSceneTools/Private/TrackEditors/PathTrackEditor.cpp
+++ b/Engine/Source/Editor/MovieSceneTools/Private/TrackEditors/PathTrackEditor.cpp
@@ -182,7 +182,7 @@ void F3DPathTrackEditor::AddKey( const FGuid& ObjectGuid, UObject* AdditionalAss
 	AddPath(ObjectGuid, AdditionalAsset);
 }
 
-void F3DPathTrackEditor::BuildObjectBindingContextMenu(FMenuBuilder& MenuBuilder, const FGuid& ObjectBinding, const UClass* ObjectClass)
+void F3DPathTrackEditor::BuildObjectBindingTrackMenu(FMenuBuilder& MenuBuilder, const FGuid& ObjectBinding, const UClass* ObjectClass)
 {
 	// build a menu with all the actors that have spline components, add the selected actors first
 	if (ObjectClass->IsChildOf(AActor::StaticClass()))
@@ -205,7 +205,8 @@ void F3DPathTrackEditor::BuildObjectBindingContextMenu(FMenuBuilder& MenuBuilder
 				FText::Format( NSLOCTEXT("Sequencer", "AddPath", "Add Path"), Args),
 				FText::Format( NSLOCTEXT("Sequencer", "AddPathTooltip", "Adds a path track."), Args ),
 				FSlateIcon(),
-				FUIAction(FExecuteAction::CreateSP(this, &F3DPathTrackEditor::AddPath, ObjectBinding, (UObject*)ActorWithSplineComponent))
+				FUIAction(FExecuteAction::CreateSP(this, &F3DPathTrackEditor::AddPath, ObjectBinding, (UObject*)ActorWithSplineComponent),
+						  FCanExecuteAction::CreateLambda( []{ return false; } ))
 			);
 		}
 	}
@@ -233,7 +234,7 @@ void F3DPathTrackEditor::AddPathSubMenu(FMenuBuilder& MenuBuilder, FGuid ObjectB
 			Args.Add( TEXT("PathName"), FText::FromString( ActorWithSplineComponent->GetActorLabel() ) );
 
 			MenuBuilder.AddMenuEntry(
-				FText::Format( NSLOCTEXT("Sequencer", "AddPath", "Add {PathName}"), Args ),
+				FText::Format( NSLOCTEXT("Sequencer", "AddPath", "{PathName}"), Args ),
 				FText::Format( NSLOCTEXT("Sequencer", "AddPathTooltip", "Add path track driven by {PathName}."), Args ),
 				FSlateIcon(),
 				FUIAction(FExecuteAction::CreateSP(this, &F3DPathTrackEditor::AddPath, ObjectBinding, (UObject*)ActorWithSplineComponent))
diff --git a/Engine/Source/Editor/MovieSceneTools/Private/TrackEditors/PathTrackEditor.h b/Engine/Source/Editor/MovieSceneTools/Private/TrackEditors/PathTrackEditor.h
--- a/Engine/Source/Editor/MovieSceneTools/Private/TrackEditors/PathTrackEditor.h
+++ b/Engine/Source/Editor/MovieSceneTools/Private/TrackEditors/PathTrackEditor.h
@@ -30,7 +30,7 @@ public:
 	virtual bool SupportsType( TSubclassOf<UMovieSceneTrack> Type ) const override;
 	virtual TSharedRef<ISequencerSection> MakeSectionInterface( UMovieSceneSection& SectionObject, UMovieSceneTrack* Track ) override;
 	virtual void AddKey( const FGuid& ObjectGuid, UObject* AdditionalAsset = NULL ) override;
-	virtual void BuildObjectBindingContextMenu(FMenuBuilder& MenuBuilder, const FGuid& ObjectBinding, const UClass* ObjectClass) override;
+	virtual void BuildObjectBindingTrackMenu(FMenuBuilder& MenuBuilder, const FGuid& ObjectBinding, const UClass* ObjectClass) override;
 
 	/** Add path */
 	void AddPath(FGuid ObjectGuid, UObject* AdditionalAsset);
diff --git a/Engine/Source/Editor/MovieSceneTools/Private/TrackEditors/ShotTrackEditor.cpp b/Engine/Source/Editor/MovieSceneTools/Private/TrackEditors/ShotTrackEditor.cpp
--- a/Engine/Source/Editor/MovieSceneTools/Private/TrackEditors/ShotTrackEditor.cpp
+++ b/Engine/Source/Editor/MovieSceneTools/Private/TrackEditors/ShotTrackEditor.cpp
@@ -543,7 +543,7 @@ void FShotTrackEditor::Tick(float DeltaTime)
 	}
 }
 
-void FShotTrackEditor::BuildObjectBindingContextMenu(FMenuBuilder& MenuBuilder, const FGuid& ObjectBinding, const UClass* ObjectClass)
+void FShotTrackEditor::BuildObjectBindingTrackMenu(FMenuBuilder& MenuBuilder, const FGuid& ObjectBinding, const UClass* ObjectClass)
 {
 	if (ObjectClass->IsChildOf(ACameraActor::StaticClass()))
 	{
diff --git a/Engine/Source/Editor/MovieSceneTools/Private/TrackEditors/ShotTrackEditor.h b/Engine/Source/Editor/MovieSceneTools/Private/TrackEditors/ShotTrackEditor.h
--- a/Engine/Source/Editor/MovieSceneTools/Private/TrackEditors/ShotTrackEditor.h
+++ b/Engine/Source/Editor/MovieSceneTools/Private/TrackEditors/ShotTrackEditor.h
@@ -30,7 +30,7 @@ public:
 	virtual TSharedRef<ISequencerSection> MakeSectionInterface( UMovieSceneSection& SectionObject, UMovieSceneTrack* Track ) override;
 	virtual void AddKey(const FGuid& ObjectGuid, UObject* AdditionalAsset = NULL) override;
 	virtual void Tick(float DeltaTime) override;
-	virtual void BuildObjectBindingContextMenu(FMenuBuilder& MenuBuilder, const FGuid& ObjectBinding, const UClass* ObjectClass) override;
+	virtual void BuildObjectBindingTrackMenu(FMenuBuilder& MenuBuilder, const FGuid& ObjectBinding, const UClass* ObjectClass) override;
 
 private:
 	/** Delegate for AnimatablePropertyChanged in AddKey */
diff --git a/Engine/Source/Editor/MovieSceneTools/Private/TrackEditors/SkeletalAnimationTrackEditor.cpp b/Engine/Source/Editor/MovieSceneTools/Private/TrackEditors/SkeletalAnimationTrackEditor.cpp
--- a/Engine/Source/Editor/MovieSceneTools/Private/TrackEditors/SkeletalAnimationTrackEditor.cpp
+++ b/Engine/Source/Editor/MovieSceneTools/Private/TrackEditors/SkeletalAnimationTrackEditor.cpp
@@ -189,7 +189,7 @@ bool FSkeletalAnimationTrackEditor::HandleAssetAdded(UObject* Asset, const FGuid
 	return false;
 }
 
-void FSkeletalAnimationTrackEditor::BuildObjectBindingContextMenu(FMenuBuilder& MenuBuilder, const FGuid& ObjectBinding, const UClass* ObjectClass)
+void FSkeletalAnimationTrackEditor::BuildObjectBindingTrackMenu(FMenuBuilder& MenuBuilder, const FGuid& ObjectBinding, const UClass* ObjectClass)
 {
 	if (ObjectClass->IsChildOf(ASkeletalMeshActor::StaticClass()))
 	{
@@ -206,27 +206,46 @@ void FSkeletalAnimationTrackEditor::BuildObjectBindingContextMenu(FMenuBuilder&
 			TArray<FAssetData> AssetDataList;
 			AssetRegistryModule.Get().GetAssetsByClass(UAnimSequence::StaticClass()->GetFName(), AssetDataList);
 
-			for( int32 AssetIndex = 0; AssetIndex < AssetDataList.Num(); ++AssetIndex )
+			if (AssetDataList.Num())
 			{
-				const FAssetData& PossibleAnimSequence = AssetDataList[AssetIndex];
-				if( FAssetData(Skeleton).GetExportTextName() == PossibleAnimSequence.TagsAndValues.FindRef("Skeleton") )
-				{
-					FFormatNamedArguments Args;
-					Args.Add( TEXT("AnimationName"), FText::FromName( PossibleAnimSequence.AssetName ) );
-
-					UAnimSequence* AnimSequence = CastChecked<UAnimSequence>(AssetDataList[AssetIndex].GetAsset());
-					MenuBuilder.AddMenuEntry(
-						FText::Format( NSLOCTEXT("Sequencer", "AddAnimSequence", "Add {AnimationName}"), Args ),
-						FText::Format( NSLOCTEXT("Sequencer", "AddAnimSequenceTooltip", "Adds a {AnimationName} animation to this skeletal mesh."), Args ),
-						FSlateIcon(),
-						FUIAction(FExecuteAction::CreateSP(ParentSequencer.Get(), &ISequencer::AddAnimation, ObjectBinding, AnimSequence))
-						);
-				}
+				MenuBuilder.AddSubMenu(
+					NSLOCTEXT("Sequencer", "AddAnimation", "Add Animation"), NSLOCTEXT("Sequencer", "AddAnimationTooltip", "Adds an animation track."),
+					FNewMenuDelegate::CreateRaw(this, &FSkeletalAnimationTrackEditor::BuildAnimationSubMenu, ObjectBinding, Skeleton));
 			}
 		}
 	}
 }
 
+void FSkeletalAnimationTrackEditor::BuildAnimationSubMenu(FMenuBuilder& MenuBuilder, FGuid ObjectBinding, USkeleton* Skeleton)
+{
+	const TSharedPtr<ISequencer> ParentSequencer = GetSequencer();
+
+	// Load the asset registry module
+	FAssetRegistryModule& AssetRegistryModule = FModuleManager::LoadModuleChecked<FAssetRegistryModule>(TEXT("AssetRegistry"));
+
+	// Collect a full list of assets with the specified class
+	TArray<FAssetData> AssetDataList;
+	AssetRegistryModule.Get().GetAssetsByClass(UAnimSequence::StaticClass()->GetFName(), AssetDataList);
+
+	for( int32 AssetIndex = 0; AssetIndex < AssetDataList.Num(); ++AssetIndex )
+	{
+		const FAssetData& PossibleAnimSequence = AssetDataList[AssetIndex];
+		if( FAssetData(Skeleton).GetExportTextName() == PossibleAnimSequence.TagsAndValues.FindRef("Skeleton") )
+		{
+			FFormatNamedArguments Args;
+			Args.Add( TEXT("AnimationName"), FText::FromName( PossibleAnimSequence.AssetName ) );
+
+			UAnimSequence* AnimSequence = CastChecked<UAnimSequence>(AssetDataList[AssetIndex].GetAsset());
+			MenuBuilder.AddMenuEntry(
+				FText::Format( NSLOCTEXT("Sequencer", "AddAnimSequence", "Add {AnimationName}"), Args ),
+				FText::Format( NSLOCTEXT("Sequencer", "AddAnimSequenceTooltip", "Adds a {AnimationName} animation to this skeletal mesh."), Args ),
+				FSlateIcon(),
+				FUIAction(FExecuteAction::CreateSP(ParentSequencer.Get(), &ISequencer::AddAnimation, ObjectBinding, AnimSequence))
+				);
+		}
+	}
+}
+
 void FSkeletalAnimationTrackEditor::AddKeyInternal( float KeyTime, const TArray<UObject*> Objects, class UAnimSequence* AnimSequence )
 {
 	for( int32 ObjectIndex = 0; ObjectIndex < Objects.Num(); ++ObjectIndex )
diff --git a/Engine/Source/Editor/MovieSceneTools/Private/TrackEditors/SkeletalAnimationTrackEditor.h b/Engine/Source/Editor/MovieSceneTools/Private/TrackEditors/SkeletalAnimationTrackEditor.h
--- a/Engine/Source/Editor/MovieSceneTools/Private/TrackEditors/SkeletalAnimationTrackEditor.h
+++ b/Engine/Source/Editor/MovieSceneTools/Private/TrackEditors/SkeletalAnimationTrackEditor.h
@@ -30,9 +30,12 @@ public:
 	virtual TSharedRef<ISequencerSection> MakeSectionInterface( UMovieSceneSection& SectionObject, UMovieSceneTrack* Track ) override;
 	virtual void AddKey(const FGuid& ObjectGuid, UObject* AdditionalAsset) override;
 	virtual bool HandleAssetAdded(UObject* Asset, const FGuid& TargetObjectGuid) override;
-	virtual void BuildObjectBindingContextMenu(FMenuBuilder& MenuBuilder, const FGuid& ObjectBinding, const UClass* ObjectClass) override;
+	virtual void BuildObjectBindingTrackMenu(FMenuBuilder& MenuBuilder, const FGuid& ObjectBinding, const UClass* ObjectClass) override;
 
 private:
+	/** Animation sub menu */
+	void BuildAnimationSubMenu(FMenuBuilder& MenuBuilder, FGuid ObjectBinding, USkeleton* Skeleton);
+
 	/** Delegate for AnimatablePropertyChanged in AddKey */
 	void AddKeyInternal(float KeyTime, const TArray<UObject*> Objects, class UAnimSequence* AnimSequence);
 
diff --git a/Engine/Source/Editor/MovieSceneTools/Private/TrackEditors/TransformTrackEditor.cpp b/Engine/Source/Editor/MovieSceneTools/Private/TrackEditors/TransformTrackEditor.cpp
--- a/Engine/Source/Editor/MovieSceneTools/Private/TrackEditors/TransformTrackEditor.cpp
+++ b/Engine/Source/Editor/MovieSceneTools/Private/TrackEditors/TransformTrackEditor.cpp
@@ -416,6 +416,61 @@ void F3DTransformTrackEditor::BuildObjectBindingEditButtons(TSharedPtr<SHorizont
 	];
 }
 
+void F3DTransformTrackEditor::BuildObjectBindingTrackMenu(FMenuBuilder& MenuBuilder, const FGuid& ObjectBinding, const UClass* ObjectClass)
+{
+	if (ObjectClass->IsChildOf(AActor::StaticClass()))
+	{
+		FFormatNamedArguments Args;
+		MenuBuilder.AddMenuEntry(
+			FText::Format( NSLOCTEXT("Sequencer", "AddTransform", "Add Transform"), Args),
+			FText::Format( NSLOCTEXT("Sequencer", "AddPTransformTooltip", "Adds a transform track."), Args ),
+			FSlateIcon(),
+			FUIAction(FExecuteAction::CreateSP(this, &F3DTransformTrackEditor::AddTransform, ObjectBinding))
+			);
+	}
+}
+
+void F3DTransformTrackEditor::AddTransform(FGuid ObjectBinding)
+{
+	/*
+	TArray<UObject*> OutObjects;
+	GetSequencer()->GetRuntimeObjects( GetSequencer()->GetFocusedMovieSceneSequenceInstance(), ObjectBinding, OutObjects);
+
+	for ( UObject* Object : OutObjects )
+	{
+		FGuid ObjectHandle;
+		USceneComponent* SceneComponent = NULL;
+		AActor* Actor = Cast<AActor>( Object );
+		if( Actor && Actor->GetRootComponent() )
+		{
+			ObjectHandle = GetSequencer()->GetHandleToObject( Actor );
+			SceneComponent = Actor->GetRootComponent();
+		}
+		else
+		{
+			// If the object wasn't an actor attempt to get it directly as a scene component 
+			SceneComponent = Cast<USceneComponent>( Object );
+			if( SceneComponent )
+			{
+				ObjectHandle = GetSequencer()->GetHandleToObject( SceneComponent );
+			}
+		}
+		*/
+	FName Transform("Transform");
+	UMovieSceneTrack* Track = AddTrack(GetSequencer()->GetFocusedMovieSceneSequence()->GetMovieScene(), ObjectBinding, UMovieScene3DTransformTrack::StaticClass(), Transform);
+	UMovieScene3DTransformTrack* TransformTrack = CastChecked<UMovieScene3DTransformTrack>( Track );
+	// Transform name and path are the same
+	TransformTrack->SetPropertyNameAndPath( Transform, Transform.ToString() );
+
+	UMovieScene3DTransformSection* TransformSection = CastChecked<UMovieScene3DTransformSection>(TransformTrack->CreateNewSection());
+	TransformSection->SetStartTime(GetSequencer()->GetViewRange().GetLowerBoundValue());
+	TransformSection->SetEndTime(GetSequencer()->GetViewRange().GetUpperBoundValue());
+	TransformTrack->AddSection(TransformSection);
+	TransformTrack->SetAsShowable();
+	
+	NotifyMovieSceneDataChanged();
+}
+
 void F3DTransformTrackEditor::OnTransformChangedInternals(float KeyTime, UObject* InObject, FGuid ObjectHandle, FTransformDataPair TransformPair, bool bAutoKeying, bool bForceKey, F3DTransformTrackKey::Type KeyType)
 {
 	// Only unwind rotation if we're generating keys while recording (scene is actively playing back)
diff --git a/Engine/Source/Editor/MovieSceneTools/Private/TrackEditors/TransformTrackEditor.h b/Engine/Source/Editor/MovieSceneTools/Private/TrackEditors/TransformTrackEditor.h
--- a/Engine/Source/Editor/MovieSceneTools/Private/TrackEditors/TransformTrackEditor.h
+++ b/Engine/Source/Editor/MovieSceneTools/Private/TrackEditors/TransformTrackEditor.h
@@ -33,9 +33,13 @@ public:
 	virtual void AddKey( const FGuid& ObjectGuid, UObject* AdditionalAsset = NULL ) override;
 	virtual void BindCommands(TSharedRef<FUICommandList> SequencerCommandBindings) override;
 	virtual void BuildObjectBindingEditButtons(TSharedPtr<SHorizontalBox> EditBox, const FGuid& ObjectBinding, const UClass* ObjectClass) override;
+	virtual void BuildObjectBindingTrackMenu(FMenuBuilder& MenuBuilder, const FGuid& ObjectBinding, const UClass* ObjectClass) override;
 
 private:
 
+	/** Add a transform track */
+	void AddTransform(FGuid ObjectGuid);
+
 	/** Custom add key implementation */
 	void AddKeyInternal(const FGuid& ObjectGuid, UObject* AdditionalAsset = NULL, bool bForceKey = false, F3DTransformTrackKey::Type KeyType = F3DTransformTrackKey::Key_All);
 
diff --git a/Engine/Source/Editor/Sequencer/Private/Sequencer.cpp b/Engine/Source/Editor/Sequencer/Private/Sequencer.cpp
--- a/Engine/Source/Editor/Sequencer/Private/Sequencer.cpp
+++ b/Engine/Source/Editor/Sequencer/Private/Sequencer.cpp
@@ -1949,6 +1949,14 @@ void FSequencer::BuildObjectBindingContextMenu(FMenuBuilder& MenuBuilder, const
 	}
 }
 
+void FSequencer::BuildObjectBindingTrackMenu(FMenuBuilder& MenuBuilder, const FGuid& ObjectBinding, const UClass* ObjectClass)
+{
+	for (int32 i = 0; i < TrackEditors.Num(); ++i)
+	{
+		TrackEditors[i]->BuildObjectBindingTrackMenu(MenuBuilder, ObjectBinding, ObjectClass);
+	}
+}
+
 void FSequencer::BuildObjectBindingEditButtons(TSharedPtr<SHorizontalBox> EditBox, const FGuid& ObjectBinding, const UClass* ObjectClass)
 {
 	for (int32 i = 0; i < TrackEditors.Num(); ++i)
diff --git a/Engine/Source/Editor/Sequencer/Private/Sequencer.h b/Engine/Source/Editor/Sequencer/Private/Sequencer.h
--- a/Engine/Source/Editor/Sequencer/Private/Sequencer.h
+++ b/Engine/Source/Editor/Sequencer/Private/Sequencer.h
@@ -212,6 +212,15 @@ public:
 	 */
 	void BuildObjectBindingContextMenu(class FMenuBuilder& MenuBuilder, const FGuid& ObjectBinding, const class UClass* ObjectClass);
 
+	/**
+	 * Builds up the track menu for object binding nodes in the outliner
+	 * 
+	 * @param MenuBuilder	The track menu builder to add things to
+	 * @param ObjectBinding	The object binding of the selected node
+	 * @param ObjectClass	The class of the selected object
+	 */
+	void BuildObjectBindingTrackMenu(class FMenuBuilder& MenuBuilder, const FGuid& ObjectBinding, const class UClass* ObjectClass);
+
 	/**
 	 * Builds up the edit buttons for object binding nodes in the outliner
 	 * 
diff --git a/Engine/Source/Editor/Sequencer/Private/SequencerDisplayNode.cpp b/Engine/Source/Editor/Sequencer/Private/SequencerDisplayNode.cpp
--- a/Engine/Source/Editor/Sequencer/Private/SequencerDisplayNode.cpp
+++ b/Engine/Source/Editor/Sequencer/Private/SequencerDisplayNode.cpp
@@ -666,7 +666,7 @@ TSharedRef<SWidget> FObjectBindingNode::GenerateEditWidgetForOutliner()
 	[
 		SNew(SComboButton)
 		.ButtonStyle(FEditorStyle::Get(), "FlatButton.Light")
-		.OnGetMenuContent(this, &FObjectBindingNode::OnGetAddPropertyTrackMenuContent)
+		.OnGetMenuContent(this, &FObjectBindingNode::OnGetAddTrackMenuContent)
 		.ContentPadding(FMargin(2, 0))
 		.ButtonContent()
 		[
@@ -734,7 +734,7 @@ void GetKeyablePropertyPaths(UClass* Class, UStruct* PropertySource, TArray<UPro
 	}
 }
 
-TSharedRef<SWidget> FObjectBindingNode::OnGetAddPropertyTrackMenuContent()
+TSharedRef<SWidget> FObjectBindingNode::OnGetAddTrackMenuContent()
 {
 	TArray<TArray<UProperty*>> KeyablePropertyPaths;
 	FSequencer& Sequencer = GetSequencer();
@@ -764,6 +764,11 @@ TSharedRef<SWidget> FObjectBindingNode::OnGetAddPropertyTrackMenuContent()
 	TSharedRef<FUICommandList> CommandList(new FUICommandList);
 	FMenuBuilder AddTrackMenuBuilder(true, nullptr, SequencerModule.GetMenuExtensibilityManager()->GetAllExtenders(CommandList, TArrayBuilder<UObject*>().Add(BoundObject)));
 
+	const UClass* ObjectClass = GetClassForObjectBinding();
+	AddTrackMenuBuilder.BeginSection( NAME_None, LOCTEXT("TracksMenuHeader" , "Tracks"));
+	GetSequencer().BuildObjectBindingTrackMenu(AddTrackMenuBuilder, ObjectBinding, ObjectClass);
+	AddTrackMenuBuilder.EndSection();
+
 	AddTrackMenuBuilder.BeginSection( SequencerMenuExtensionPoints::AddTrackMenu_PropertiesSection, LOCTEXT("PropertiesMenuHeader" , "Properties"));
 	{
 		for ( TArray<UProperty*>& KeyablePropertyPath : KeyablePropertyPaths )
diff --git a/Engine/Source/Editor/Sequencer/Private/SequencerDisplayNode.h b/Engine/Source/Editor/Sequencer/Private/SequencerDisplayNode.h
--- a/Engine/Source/Editor/Sequencer/Private/SequencerDisplayNode.h
+++ b/Engine/Source/Editor/Sequencer/Private/SequencerDisplayNode.h
@@ -483,7 +483,7 @@ public:
 	virtual TSharedPtr<SWidget> OnSummonContextMenu(const FGeometry& MyGeometry, const FPointerEvent& MouseEvent) override;
 
 private:
-	TSharedRef<SWidget> OnGetAddPropertyTrackMenuContent();
+	TSharedRef<SWidget> OnGetAddTrackMenuContent();
 
 	void AddTrackForProperty(TArray<UProperty*> PropertyPath);
 
diff --git a/Engine/Source/Editor/Sequencer/Public/MovieSceneTrackEditor.h b/Engine/Source/Editor/Sequencer/Public/MovieSceneTrackEditor.h
--- a/Engine/Source/Editor/Sequencer/Public/MovieSceneTrackEditor.h
+++ b/Engine/Source/Editor/Sequencer/Public/MovieSceneTrackEditor.h
@@ -194,6 +194,15 @@ public:
 	 */
 	virtual void BuildObjectBindingContextMenu(FMenuBuilder& MenuBuilder, const FGuid& ObjectBinding, const UClass* ObjectClass) {}
 
+	/**
+	 * Builds up the object binding track menu for the outliner
+	 *
+	 * @param MenuBuilder	The menu builder to change
+	 * @param ObjectBinding	The object binding this is for
+	 * @param ObjectClass	The class of the object this is for
+	 */
+	virtual void BuildObjectBindingTrackMenu(FMenuBuilder& MenuBuilder, const FGuid& ObjectBinding, const UClass* ObjectClass) {}
+
 	/**
 	 * Builds up the object binding edit buttons for the outliner
 	 *
