Sequencer: Remove some commented out code.

diff --git a/Engine/Source/Editor/MovieSceneTools/Private/TrackEditors/TransformTrackEditor.cpp b/Engine/Source/Editor/MovieSceneTools/Private/TrackEditors/TransformTrackEditor.cpp
--- a/Engine/Source/Editor/MovieSceneTools/Private/TrackEditors/TransformTrackEditor.cpp
+++ b/Engine/Source/Editor/MovieSceneTools/Private/TrackEditors/TransformTrackEditor.cpp
@@ -432,30 +432,6 @@ void F3DTransformTrackEditor::BuildObjectBindingTrackMenu(FMenuBuilder& MenuBuil
 
 void F3DTransformTrackEditor::AddTransform(FGuid ObjectBinding)
 {
-	/*
-	TArray<UObject*> OutObjects;
-	GetSequencer()->GetRuntimeObjects( GetSequencer()->GetFocusedMovieSceneSequenceInstance(), ObjectBinding, OutObjects);
-
-	for ( UObject* Object : OutObjects )
-	{
-		FGuid ObjectHandle;
-		USceneComponent* SceneComponent = NULL;
-		AActor* Actor = Cast<AActor>( Object );
-		if( Actor && Actor->GetRootComponent() )
-		{
-			ObjectHandle = GetSequencer()->GetHandleToObject( Actor );
-			SceneComponent = Actor->GetRootComponent();
-		}
-		else
-		{
-			// If the object wasn't an actor attempt to get it directly as a scene component 
-			SceneComponent = Cast<USceneComponent>( Object );
-			if( SceneComponent )
-			{
-				ObjectHandle = GetSequencer()->GetHandleToObject( SceneComponent );
-			}
-		}
-		*/
 	FName Transform("Transform");
 	UMovieSceneTrack* Track = AddTrack(GetSequencer()->GetFocusedMovieSceneSequence()->GetMovieScene(), ObjectBinding, UMovieScene3DTransformTrack::StaticClass(), Transform);
 	UMovieScene3DTransformTrack* TransformTrack = CastChecked<UMovieScene3DTransformTrack>( Track );
