Sequencer: Fix for trying to add another Transform track.

diff --git a/Engine/Source/Editor/MovieSceneTools/Private/TrackEditors/TransformTrackEditor.cpp b/Engine/Source/Editor/MovieSceneTools/Private/TrackEditors/TransformTrackEditor.cpp
--- a/Engine/Source/Editor/MovieSceneTools/Private/TrackEditors/TransformTrackEditor.cpp
+++ b/Engine/Source/Editor/MovieSceneTools/Private/TrackEditors/TransformTrackEditor.cpp
@@ -425,8 +425,8 @@ void F3DTransformTrackEditor::BuildObjectBindingTrackMenu(FMenuBuilder& MenuBuil
 			FText::Format( NSLOCTEXT("Sequencer", "AddTransform", "Add Transform"), Args),
 			FText::Format( NSLOCTEXT("Sequencer", "AddPTransformTooltip", "Adds a transform track."), Args ),
 			FSlateIcon(),
-			FUIAction(FExecuteAction::CreateSP(this, &F3DTransformTrackEditor::AddTransform, ObjectBinding))
-			);
+			FUIAction(FExecuteAction::CreateSP(this, &F3DTransformTrackEditor::AddTransform, ObjectBinding),
+					  FCanExecuteAction::CreateSP(this, &F3DTransformTrackEditor::CanAddTransform, ObjectBinding)) );
 	}
 }
 
@@ -447,6 +447,16 @@ void F3DTransformTrackEditor::AddTransform(FGuid ObjectBinding)
 	NotifyMovieSceneDataChanged();
 }
 
+bool F3DTransformTrackEditor::CanAddTransform(FGuid ObjectBinding) const
+{
+	FName Transform("Transform");
+	if (GetSequencer()->GetFocusedMovieSceneSequence()->GetMovieScene()->FindTrack(UMovieScene3DTransformTrack::StaticClass(), ObjectBinding, Transform))
+	{
+		return false;
+	}
+	return true;
+}
+
 void F3DTransformTrackEditor::OnTransformChangedInternals(float KeyTime, UObject* InObject, FGuid ObjectHandle, FTransformDataPair TransformPair, bool bAutoKeying, bool bForceKey, F3DTransformTrackKey::Type KeyType)
 {
 	// Only unwind rotation if we're generating keys while recording (scene is actively playing back)
diff --git a/Engine/Source/Editor/MovieSceneTools/Private/TrackEditors/TransformTrackEditor.h b/Engine/Source/Editor/MovieSceneTools/Private/TrackEditors/TransformTrackEditor.h
--- a/Engine/Source/Editor/MovieSceneTools/Private/TrackEditors/TransformTrackEditor.h
+++ b/Engine/Source/Editor/MovieSceneTools/Private/TrackEditors/TransformTrackEditor.h
@@ -40,6 +40,9 @@ private:
 	/** Add a transform track */
 	void AddTransform(FGuid ObjectGuid);
 
+	/** Called to determine whether a transform track can be added */
+	bool CanAddTransform(FGuid ObjectGuid) const;
+
 	/** Custom add key implementation */
 	void AddKeyInternal(const FGuid& ObjectGuid, UObject* AdditionalAsset = NULL, bool bForceKey = false, F3DTransformTrackKey::Type KeyType = F3DTransformTrackKey::Key_All);
 
