Sequencer: Add attach track.

diff --git a/Engine/Source/Editor/MovieSceneTools/Private/MovieSceneToolsModule.cpp b/Engine/Source/Editor/MovieSceneTools/Private/MovieSceneToolsModule.cpp
--- a/Engine/Source/Editor/MovieSceneTools/Private/MovieSceneToolsModule.cpp
+++ b/Engine/Source/Editor/MovieSceneTools/Private/MovieSceneToolsModule.cpp
@@ -21,6 +21,7 @@
 #include "AudioTrackEditor.h"
 #include "SkeletalAnimationTrackEditor.h"
 #include "ParticleTrackEditor.h"
+#include "AttachTrackEditor.h"
 #include "PathTrackEditor.h"
 
 
@@ -47,6 +48,7 @@ class FMovieSceneToolsModule : public IMovieSceneTools
 		AudioTrackEditorCreateTrackEditorDelegateHandle          = SequencerModule.RegisterTrackEditor_Handle( FOnCreateTrackEditor::CreateStatic( &FAudioTrackEditor::CreateTrackEditor ) );
 		AnimationTrackEditorCreateTrackEditorDelegateHandle      = SequencerModule.RegisterTrackEditor_Handle( FOnCreateTrackEditor::CreateStatic( &FSkeletalAnimationTrackEditor::CreateTrackEditor ) );
 		ParticleTrackEditorCreateTrackEditorDelegateHandle       = SequencerModule.RegisterTrackEditor_Handle( FOnCreateTrackEditor::CreateStatic( &FParticleTrackEditor::CreateTrackEditor ) );
+		AttachTrackEditorCreateTrackEditorDelegateHandle         = SequencerModule.RegisterTrackEditor_Handle( FOnCreateTrackEditor::CreateStatic( &F3DAttachTrackEditor::CreateTrackEditor ) );
 		PathTrackEditorCreateTrackEditorDelegateHandle           = SequencerModule.RegisterTrackEditor_Handle( FOnCreateTrackEditor::CreateStatic( &F3DPathTrackEditor::CreateTrackEditor ) );
 	}
 
@@ -70,6 +72,7 @@ class FMovieSceneToolsModule : public IMovieSceneTools
 			SequencerModule.UnRegisterTrackEditor_Handle( AudioTrackEditorCreateTrackEditorDelegateHandle );
 			SequencerModule.UnRegisterTrackEditor_Handle( AnimationTrackEditorCreateTrackEditorDelegateHandle );
 			SequencerModule.UnRegisterTrackEditor_Handle( ParticleTrackEditorCreateTrackEditorDelegateHandle );
+			SequencerModule.UnRegisterTrackEditor_Handle( AttachTrackEditorCreateTrackEditorDelegateHandle );
 			SequencerModule.UnRegisterTrackEditor_Handle( PathTrackEditorCreateTrackEditorDelegateHandle );
 		}
 		
@@ -91,6 +94,7 @@ private:
 	FDelegateHandle AudioTrackEditorCreateTrackEditorDelegateHandle;
 	FDelegateHandle AnimationTrackEditorCreateTrackEditorDelegateHandle;
 	FDelegateHandle ParticleTrackEditorCreateTrackEditorDelegateHandle;
+	FDelegateHandle AttachTrackEditorCreateTrackEditorDelegateHandle;
 	FDelegateHandle PathTrackEditorCreateTrackEditorDelegateHandle;
 };
 
diff --git a/Engine/Source/Editor/MovieSceneTools/Private/TrackEditors/AttachTrackEditor.cpp b/Engine/Source/Editor/MovieSceneTools/Private/TrackEditors/AttachTrackEditor.cpp
new file mode 100644
--- /dev/null
+++ b/Engine/Source/Editor/MovieSceneTools/Private/TrackEditors/AttachTrackEditor.cpp
@@ -0,0 +1,357 @@
+// Copyright 1998-2015 Epic Games, Inc. All Rights Reserved.
+
+#include "MovieSceneToolsPrivatePCH.h"
+#include "ScopedTransaction.h"
+#include "MovieScene.h"
+#include "MovieSceneSection.h"
+#include "MovieScene3DAttachTrack.h"
+#include "MovieScene3DAttachSection.h"
+#include "ISequencerSection.h"
+#include "ISectionLayoutBuilder.h"
+#include "IKeyArea.h"
+#include "MovieSceneToolHelpers.h"
+#include "MovieSceneTrackEditor.h"
+#include "ActorEditorUtils.h"
+#include "GameFramework/WorldSettings.h"
+#include "AttachTrackEditor.h"
+
+void GetAttachActors(TSet<AActor*>& SelectedActors, TSet<AActor*>& UnselectedActors, TArray<UObject*> InObjects)
+{
+	if (GEditor->GetEditorWorldContext().World())
+	{
+		for( FActorIterator	ActorIt(GEditor->GetEditorWorldContext().World()); ActorIt; ++ActorIt )
+		{
+			AActor* Actor = *ActorIt;
+
+			if (Actor->IsListedInSceneOutliner() &&
+				!FActorEditorUtils::IsABuilderBrush(Actor) &&
+				!Actor->IsA( AWorldSettings::StaticClass() ) &&
+				!Actor->IsPendingKill() &&
+				!InObjects.Contains(Actor))
+			{			
+				if (GEditor->GetSelectedActors()->IsSelected(Actor))
+				{
+					SelectedActors.Add(Actor);
+				}
+				else
+				{
+					UnselectedActors.Add(Actor);
+				}
+			}
+		}
+	}
+	SelectedActors.Sort([](const AActor& A, const AActor& B)
+	{
+		return A.GetActorLabel() < B.GetActorLabel();
+	});
+	UnselectedActors.Sort([](const AActor& A, const AActor& B)
+	{
+		return A.GetActorLabel() < B.GetActorLabel();
+	});
+}
+
+/**
+ * Class that draws an attach section in the sequencer
+ */
+class F3DAttachSection : public ISequencerSection
+{
+public:
+	F3DAttachSection( UMovieSceneSection& InSection, F3DAttachTrackEditor* InAttachTrackEditor )
+		: Section( InSection )
+		, AttachTrackEditor(InAttachTrackEditor)
+	{
+	}
+
+	/** ISequencerSection interface */
+	virtual UMovieSceneSection* GetSectionObject() override
+	{ 
+		return &Section;
+	}
+
+	virtual FText GetDisplayName() const override
+	{ 
+		return NSLOCTEXT("FAttachSection", "DisplayName", "Attach");
+	}
+	
+	virtual FText GetSectionTitle() const override { return FText::GetEmpty(); }
+
+	virtual void GenerateSectionLayout( class ISectionLayoutBuilder& LayoutBuilder ) const override {}
+
+	virtual int32 OnPaintSection( const FGeometry& AllottedGeometry, const FSlateRect& SectionClippingRect, FSlateWindowElementList& OutDrawElements, int32 LayerId, bool bParentEnabled ) const override 
+	{
+		// Add a box for the section
+		FSlateDrawElement::MakeBox( 
+			OutDrawElements,
+			LayerId,
+			AllottedGeometry.ToPaintGeometry(),
+			FEditorStyle::GetBrush("Sequencer.GenericSection.Background"),
+			SectionClippingRect
+		); 
+
+		return LayerId;
+	}
+	
+	virtual void BuildSectionContextMenu(FMenuBuilder& MenuBuilder) override
+	{
+		TArray<UObject*> OutObjects;
+		//@todo
+
+		TSet<AActor*> SelectedActors;
+		TSet<AActor*> UnselectedActors;
+		GetAttachActors(SelectedActors, UnselectedActors, OutObjects);	
+		
+		if (SelectedActors.Num()+UnselectedActors.Num() == 0)
+		{
+			return;
+		}
+
+		if (SelectedActors.Num())
+		{
+			MenuBuilder.BeginSection(NAME_None, FText::FromString("Selected Actors"));
+
+			for (TSet<AActor*>::TIterator It(SelectedActors); It; ++It)
+			{
+				AActor* Actor = *It;
+				FFormatNamedArguments Args;
+				Args.Add( TEXT("AttachName"), FText::FromString( Actor->GetActorLabel() ) );
+
+				MenuBuilder.AddMenuEntry(
+					FText::Format( NSLOCTEXT("Sequencer", "SetAttach", "Set {AttachName}"), Args ),
+					FText::Format( NSLOCTEXT("Sequencer", "SetAttachTooltip", "Set attach track driven by {AttachName}."), Args ),
+					FSlateIcon(),
+					FUIAction(FExecuteAction::CreateSP(AttachTrackEditor, &F3DAttachTrackEditor::SetAttach, &Section, Actor))
+				);
+			}
+			
+			MenuBuilder.EndSection();
+		}
+
+		if (UnselectedActors.Num())
+		{
+			MenuBuilder.BeginSection(NAME_None, FText::FromString("Actors"));
+
+			for (TSet<AActor*>::TIterator It(UnselectedActors); It; ++It)
+			{
+				AActor* Actor = *It;
+				FFormatNamedArguments Args;
+				Args.Add( TEXT("AttachName"), FText::FromString( Actor->GetActorLabel() ) );
+
+				MenuBuilder.AddMenuEntry(
+					FText::Format( NSLOCTEXT("Sequencer", "SetAttach", "Set {AttachName}"), Args ),
+					FText::Format( NSLOCTEXT("Sequencer", "SetAttachTooltip", "Set attach track driven by {AttachName}."), Args ),
+					FSlateIcon(),
+					FUIAction(FExecuteAction::CreateSP(AttachTrackEditor, &F3DAttachTrackEditor::SetAttach, &Section, Actor))
+				);
+			}
+
+			MenuBuilder.EndSection();
+		}
+
+		MenuBuilder.AddMenuSeparator();
+	}
+
+private:
+	/** The section we are visualizing */
+	UMovieSceneSection& Section;
+
+	/** The attach track editor */
+	F3DAttachTrackEditor* AttachTrackEditor;
+};
+
+F3DAttachTrackEditor::F3DAttachTrackEditor( TSharedRef<ISequencer> InSequencer )
+	: FMovieSceneTrackEditor( InSequencer ) 
+{
+}
+
+F3DAttachTrackEditor::~F3DAttachTrackEditor()
+{
+}
+
+TSharedRef<FMovieSceneTrackEditor> F3DAttachTrackEditor::CreateTrackEditor( TSharedRef<ISequencer> InSequencer )
+{
+	return MakeShareable( new F3DAttachTrackEditor( InSequencer ) );
+}
+
+bool F3DAttachTrackEditor::SupportsType( TSubclassOf<UMovieSceneTrack> Type ) const
+{
+	// We support animatable transforms
+	return Type == UMovieScene3DAttachTrack::StaticClass();
+}
+
+TSharedRef<ISequencerSection> F3DAttachTrackEditor::MakeSectionInterface( UMovieSceneSection& SectionObject, UMovieSceneTrack* Track )
+{
+	check( SupportsType( SectionObject.GetOuter()->GetClass() ) );
+
+	TSharedRef<ISequencerSection> NewSection( new F3DAttachSection( SectionObject, this ) );
+
+	return NewSection;
+}
+
+void F3DAttachTrackEditor::AddKey( const FGuid& ObjectGuid, UObject* AdditionalAsset )
+{
+	AddAttach(ObjectGuid, AdditionalAsset);
+}
+
+void F3DAttachTrackEditor::BuildObjectBindingTrackMenu(FMenuBuilder& MenuBuilder, const FGuid& ObjectBinding, const UClass* ObjectClass)
+{
+	// build a menu with all the possible actors, add the selected actors first
+	if (ObjectClass->IsChildOf(AActor::StaticClass()))
+	{
+		TArray<UObject*> OutObjects;
+		GetSequencer()->GetRuntimeObjects( GetSequencer()->GetFocusedMovieSceneSequenceInstance(), ObjectBinding, OutObjects);
+
+		TSet<AActor*> SelectedActors;
+		TSet<AActor*> UnselectedActors;
+		GetAttachActors(SelectedActors, UnselectedActors, OutObjects);
+
+		if (SelectedActors.Num()+UnselectedActors.Num() > 0)
+		{
+			MenuBuilder.AddSubMenu(
+				NSLOCTEXT("Sequencer", "AddAttach", "Add Attach"), NSLOCTEXT("Sequencer", "AddAttachTooltip", "Adds an attach track."),
+				FNewMenuDelegate::CreateRaw(this, &F3DAttachTrackEditor::AddAttachSubMenu, ObjectBinding));
+		}
+		else
+		{
+			AActor* Actor = nullptr;
+			FFormatNamedArguments Args;
+			MenuBuilder.AddMenuEntry(
+				FText::Format( NSLOCTEXT("Sequencer", "AddAttach", "Add Attach"), Args),
+				FText::Format( NSLOCTEXT("Sequencer", "AddAttachTooltip", "Adds an attach track."), Args ),
+				FSlateIcon(),
+				FUIAction(FExecuteAction::CreateSP(this, &F3DAttachTrackEditor::AddAttach, ObjectBinding, (UObject*)Actor),
+						  FCanExecuteAction::CreateLambda( []{ return false; } ))
+			);
+		}
+	}
+}
+
+void F3DAttachTrackEditor::AddAttachSubMenu(FMenuBuilder& MenuBuilder, FGuid ObjectBinding)
+{
+	TArray<UObject*> OutObjects;
+	GetSequencer()->GetRuntimeObjects( GetSequencer()->GetFocusedMovieSceneSequenceInstance(), ObjectBinding, OutObjects);
+
+	TSet<AActor*> SelectedActors;
+	TSet<AActor*> UnselectedActors;
+	GetAttachActors(SelectedActors, UnselectedActors, OutObjects);	
+
+	if (SelectedActors.Num()+UnselectedActors.Num() == 0)
+	{
+		return;
+	}
+
+	if (SelectedActors.Num())
+	{
+		MenuBuilder.BeginSection(NAME_None, FText::FromString("Selected Actors"));
+
+		for (TSet<AActor*>::TIterator It(SelectedActors); It; ++It)
+		{
+			AActor* Actor = *It;
+			FFormatNamedArguments Args;
+			Args.Add( TEXT("AttachName"), FText::FromString( Actor->GetActorLabel() ) );
+
+			MenuBuilder.AddMenuEntry(
+				FText::Format( NSLOCTEXT("Sequencer", "AddAttach", "{AttachName}"), Args ),
+				FText::Format( NSLOCTEXT("Sequencer", "AddAttachTooltip", "Add attach track driven by {AttachName}."), Args ),
+				FSlateIcon(),
+				FUIAction(FExecuteAction::CreateSP(this, &F3DAttachTrackEditor::AddAttach, ObjectBinding, (UObject*)Actor))
+			);
+		}
+			
+		MenuBuilder.EndSection();
+	}
+
+	if (UnselectedActors.Num())
+	{
+		MenuBuilder.BeginSection(NAME_None, FText::FromString("Actors"));
+
+		for (TSet<AActor*>::TIterator It(UnselectedActors); It; ++It)
+		{
+			AActor* Actor = *It;
+			FFormatNamedArguments Args;
+			Args.Add( TEXT("AttachName"), FText::FromString( Actor->GetActorLabel() ) );
+
+			MenuBuilder.AddMenuEntry(
+				FText::Format( NSLOCTEXT("Sequencer", "AddAttach", "Add {AttachName}"), Args ),
+				FText::Format( NSLOCTEXT("Sequencer", "AddAttachTooltip", "Add attach track driven by {AttachName}."), Args ),
+				FSlateIcon(),
+				FUIAction(FExecuteAction::CreateSP(this, &F3DAttachTrackEditor::AddAttach, ObjectBinding, (UObject*)Actor))
+			);
+		}
+
+		MenuBuilder.EndSection();
+	}
+}
+
+void F3DAttachTrackEditor::AddAttach(FGuid ObjectGuid, UObject* AdditionalAsset)
+{
+	TArray<UObject*> OutObjects;
+	GetSequencer()->GetRuntimeObjects( GetSequencer()->GetFocusedMovieSceneSequenceInstance(), ObjectGuid, OutObjects);
+
+	AnimatablePropertyChanged( UMovieScene3DAttachTrack::StaticClass(), false,
+		FOnKeyProperty::CreateRaw( this, &F3DAttachTrackEditor::AddKeyInternal, OutObjects, AdditionalAsset) );
+}
+
+void F3DAttachTrackEditor::SetAttach(UMovieSceneSection* Section, AActor* Actor)
+{
+	const FScopedTransaction Transaction(NSLOCTEXT("Sequencer", "UndoSetAttach", "Set Attach"));
+
+	UMovieScene3DAttachSection* AttachSection = (UMovieScene3DAttachSection*)(Section);
+	FGuid ActorId = FindOrCreateHandleToObject(Actor);
+
+	if (ActorId.IsValid())
+	{
+		AttachSection->SetConstraintId(ActorId);
+	}
+}
+
+void F3DAttachTrackEditor::AddKeyInternal( float KeyTime, const TArray<UObject*> Objects, UObject* AdditionalAsset)
+{
+	AActor* Actor = nullptr;
+	
+	if (AdditionalAsset != nullptr)
+	{
+		Actor = Cast<AActor>(AdditionalAsset);
+	}
+
+	FGuid ActorId;
+
+	if (Actor != nullptr)
+	{
+		ActorId = FindOrCreateHandleToObject(Actor);
+	}
+
+	if (ActorId.IsValid())
+	{
+		for( int32 ObjectIndex = 0; ObjectIndex < Objects.Num(); ++ObjectIndex )
+		{
+			UObject* Object = Objects[ObjectIndex];
+
+			FGuid ObjectHandle = FindOrCreateHandleToObject( Object );
+			if (ObjectHandle.IsValid())
+			{
+				UMovieSceneTrack* Track = GetTrackForObject( ObjectHandle, UMovieScene3DAttachTrack::StaticClass(), FName("Attach"));
+
+				if (ensure(Track))
+				{
+					// Clamp to next attach section's start time or the end of the current sequencer view range
+					float AttachEndTime = GetSequencer()->GetViewRange().GetUpperBoundValue();
+	
+					for (int32 AttachSectionIndex = 0; AttachSectionIndex < Track->GetAllSections().Num(); ++AttachSectionIndex)
+					{
+						float StartTime = Track->GetAllSections()[AttachSectionIndex]->GetStartTime();
+						float EndTime = Track->GetAllSections()[AttachSectionIndex]->GetEndTime();
+						if (KeyTime < StartTime)
+						{
+							if (AttachEndTime > StartTime)
+							{
+								AttachEndTime = StartTime;
+							}
+						}
+					}
+
+					Cast<UMovieScene3DAttachTrack>(Track)->AddConstraint( KeyTime, AttachEndTime, ActorId );
+				}
+			}
+		}
+	}
+}
diff --git a/Engine/Source/Editor/MovieSceneTools/Private/TrackEditors/AttachTrackEditor.h b/Engine/Source/Editor/MovieSceneTools/Private/TrackEditors/AttachTrackEditor.h
new file mode 100644
--- /dev/null
+++ b/Engine/Source/Editor/MovieSceneTools/Private/TrackEditors/AttachTrackEditor.h
@@ -0,0 +1,50 @@
+// Copyright 1998-2015 Epic Games, Inc. All Rights Reserved.
+
+#pragma once
+
+#include "MovieScene3DAttachSection.h"
+#include "MovieScene3DAttachTrack.h"
+/**
+ * Tools for attaching an object to another object
+ */
+class F3DAttachTrackEditor : public FMovieSceneTrackEditor
+{
+public:
+	/**
+	 * Constructor
+	 *
+	 * @param InSequencer	The sequencer instance to be used by this tool
+	 */
+	F3DAttachTrackEditor( TSharedRef<ISequencer> InSequencer );
+	~F3DAttachTrackEditor();
+
+	/**
+	 * Creates an instance of this class.  Called by a sequencer 
+	 *
+	 * @param OwningSequencer The sequencer instance to be used by this tool
+	 * @return The new instance of this class
+	 */
+	static TSharedRef<FMovieSceneTrackEditor> CreateTrackEditor( TSharedRef<ISequencer> OwningSequencer );
+
+	/** FMovieSceneTrackEditor Interface */
+	virtual bool SupportsType( TSubclassOf<UMovieSceneTrack> Type ) const override;
+	virtual TSharedRef<ISequencerSection> MakeSectionInterface( UMovieSceneSection& SectionObject, UMovieSceneTrack* Track ) override;
+	virtual void AddKey( const FGuid& ObjectGuid, UObject* AdditionalAsset = NULL ) override;
+	virtual void BuildObjectBindingTrackMenu(FMenuBuilder& MenuBuilder, const FGuid& ObjectBinding, const UClass* ObjectClass) override;
+
+	/** Add attach */
+	void AddAttach(FGuid ObjectGuid, UObject* AdditionalAsset);
+
+	/** Set attach */
+	void SetAttach(UMovieSceneSection* Section, AActor* ActorWithSplineComponent);
+
+private:
+
+	/** Add attach sub menu */
+	void AddAttachSubMenu(FMenuBuilder& MenuBuilder, FGuid ObjectBinding);
+
+	/** Delegate for AnimatablePropertyChanged in AddKey */
+	void AddKeyInternal(float KeyTime, const TArray<UObject*> Objects, UObject* AdditionalAsset);
+};
+
+
diff --git a/Engine/Source/Runtime/MovieSceneTracks/Private/Sections/MovieScene3DAttachSection.cpp b/Engine/Source/Runtime/MovieSceneTracks/Private/Sections/MovieScene3DAttachSection.cpp
new file mode 100644
--- /dev/null
+++ b/Engine/Source/Runtime/MovieSceneTracks/Private/Sections/MovieScene3DAttachSection.cpp
@@ -0,0 +1,61 @@
+// Copyright 1998-2015 Epic Games, Inc. All Rights Reserved.
+
+#include "MovieSceneTracksPrivatePCH.h"
+#include "MovieScene3DAttachSection.h"
+#include "MovieScene3DAttachTrack.h"
+
+UMovieScene3DAttachSection::UMovieScene3DAttachSection( const FObjectInitializer& ObjectInitializer )
+	: Super( ObjectInitializer )
+{
+	bConstrainTx = true;
+	bConstrainTy = true;
+	bConstrainTz = true;
+	bConstrainRx = true;
+	bConstrainRy = true;
+	bConstrainRz = true;
+}
+
+void UMovieScene3DAttachSection::Eval( USceneComponent* SceneComponent, float Position, AActor* Actor, FVector& OutTranslation, FRotator& OutRotation ) const
+{
+	OutTranslation = Actor->GetActorLocation();
+	OutRotation = Actor->GetActorRotation();
+
+	//@todo eval with component/bone
+
+	if (!bConstrainTx)
+	{
+		OutTranslation[0] = SceneComponent->GetRelativeTransform().GetLocation()[0];
+	}
+
+	if (!bConstrainTy)
+	{
+		OutTranslation[1] = SceneComponent->GetRelativeTransform().GetLocation()[1];
+	}
+
+	if (!bConstrainTz)
+	{
+		OutTranslation[1] = SceneComponent->GetRelativeTransform().GetLocation()[2];
+	}
+
+	if (!bConstrainRx)
+	{
+		OutRotation.Roll = SceneComponent->GetRelativeTransform().GetRotation().Rotator().Roll;
+	}
+
+	if (!bConstrainRx)
+	{
+		OutRotation.Pitch = SceneComponent->GetRelativeTransform().GetRotation().Rotator().Pitch;
+	}
+
+	if (!bConstrainRx)
+	{
+		OutRotation.Yaw = SceneComponent->GetRelativeTransform().GetRotation().Rotator().Yaw;
+	}
+}
+
+void UMovieScene3DAttachSection::AddAttach( float Time, float SequenceEndTime, const FGuid& InAttachId )
+{
+	Modify();
+	ConstraintId = InAttachId;
+}
+
diff --git a/Engine/Source/Runtime/MovieSceneTracks/Private/TrackInstances/MovieScene3DAttachTrackInstance.cpp b/Engine/Source/Runtime/MovieSceneTracks/Private/TrackInstances/MovieScene3DAttachTrackInstance.cpp
new file mode 100644
--- /dev/null
+++ b/Engine/Source/Runtime/MovieSceneTracks/Private/TrackInstances/MovieScene3DAttachTrackInstance.cpp
@@ -0,0 +1,35 @@
+// Copyright 1998-2015 Epic Games, Inc. All Rights Reserved.
+
+#include "MovieSceneTracksPrivatePCH.h"
+#include "MovieSceneCommonHelpers.h"
+#include "MovieScene3DAttachTrackInstance.h"
+#include "MovieScene3DAttachTrack.h"
+#include "MovieScene3DAttachSection.h"
+#include "IMovieScenePlayer.h"
+
+FMovieScene3DAttachTrackInstance::FMovieScene3DAttachTrackInstance( UMovieScene3DAttachTrack& InAttachTrack )
+	: FMovieScene3DConstraintTrackInstance( InAttachTrack )
+{
+}
+
+void FMovieScene3DAttachTrackInstance::UpdateConstraint( float Position, const TArray<UObject*>& RuntimeObjects, AActor* Actor, UMovieScene3DConstraintSection* ConstraintSection ) 
+{
+	FVector Translation;
+	FRotator Rotation;
+
+	UMovieScene3DAttachSection* AttachSection = CastChecked<UMovieScene3DAttachSection>(ConstraintSection);
+
+	for( int32 ObjIndex = 0; ObjIndex < RuntimeObjects.Num(); ++ObjIndex )
+	{
+		USceneComponent* SceneComponent = MovieSceneHelpers::SceneComponentFromRuntimeObject(RuntimeObjects[ObjIndex]);
+
+		if (SceneComponent != NULL)
+		{
+			AttachSection->Eval(SceneComponent, Position, Actor, Translation, Rotation);
+
+			// Set the relative translation and rotation.  Note they are set once instead of individually to avoid a redundant component transform update.
+			SceneComponent->SetRelativeLocationAndRotation(Translation, Rotation);
+		}
+	}
+}
+
diff --git a/Engine/Source/Runtime/MovieSceneTracks/Private/Tracks/MovieScene3DAttachTrack.cpp b/Engine/Source/Runtime/MovieSceneTracks/Private/Tracks/MovieScene3DAttachTrack.cpp
new file mode 100644
--- /dev/null
+++ b/Engine/Source/Runtime/MovieSceneTracks/Private/Tracks/MovieScene3DAttachTrack.cpp
@@ -0,0 +1,31 @@
+// Copyright 1998-2015 Epic Games, Inc. All Rights Reserved.
+
+#include "MovieSceneTracksPrivatePCH.h"
+#include "MovieScene3DAttachSection.h"
+#include "MovieScene3DAttachTrack.h"
+#include "IMovieScenePlayer.h"
+#include "MovieScene3DAttachTrackInstance.h"
+
+UMovieScene3DAttachTrack::UMovieScene3DAttachTrack( const FObjectInitializer& ObjectInitializer )
+	: Super( ObjectInitializer )
+{
+}
+FName UMovieScene3DAttachTrack::GetTrackName() const
+{
+	return FName("Attach");
+}
+
+TSharedPtr<IMovieSceneTrackInstance> UMovieScene3DAttachTrack::CreateInstance()
+{
+	return MakeShareable( new FMovieScene3DAttachTrackInstance( *this ) ); 
+}
+
+void UMovieScene3DAttachTrack::AddConstraint(float KeyTime, float ConstraintEndTime, const FGuid& ConstraintId)
+{
+	// add the section
+	UMovieScene3DAttachSection* NewSection = NewObject<UMovieScene3DAttachSection>(this);
+	NewSection->AddAttach(KeyTime, ConstraintEndTime, ConstraintId);
+	NewSection->InitialPlacement( ConstraintSections, KeyTime, ConstraintEndTime, SupportsMultipleRows() );
+
+	ConstraintSections.Add(NewSection);
+}
diff --git a/Engine/Source/Runtime/MovieSceneTracks/Public/Sections/MovieScene3DAttachSection.h b/Engine/Source/Runtime/MovieSceneTracks/Public/Sections/MovieScene3DAttachSection.h
new file mode 100644
--- /dev/null
+++ b/Engine/Source/Runtime/MovieSceneTracks/Public/Sections/MovieScene3DAttachSection.h
@@ -0,0 +1,56 @@
+// Copyright 1998-2015 Epic Games, Inc. All Rights Reserved.
+
+#pragma once
+
+#include "MovieScene3DConstraintSection.h"
+#include "MovieScene3DAttachSection.generated.h"
+
+/**
+ * A 3D Attach section
+ */
+UCLASS(MinimalAPI)
+class UMovieScene3DAttachSection
+	: public UMovieScene3DConstraintSection
+{
+	GENERATED_UCLASS_BODY()
+
+public:
+
+	/**
+	 * Evaluates the attach track
+	 *
+	 * @param Time				The position in time within the movie scene
+	 */
+	void Eval( USceneComponent* SceneComponent, float Time, AActor* Actor, FVector& OutTranslation, FRotator& OutRotation ) const;
+
+	/** 
+	 * Adds an attach to the section
+	 *
+	 * @param Time	The location in time where the attach should be added
+	 * @param SequenceEndTime   The time at the end of the sequence, by default the attach is set to end at this time
+	 * @param InAttachId The id to the path
+	 */
+	void AddAttach( float Time, float SequenceEndTime, const FGuid& InAttachId);
+	
+private:
+
+	//@todo need additional actor component/bone to attach to
+
+	UPROPERTY()
+	uint32 bConstrainTx:1;
+
+	UPROPERTY()
+	uint32 bConstrainTy:1;
+
+	UPROPERTY()
+	uint32 bConstrainTz:1;
+
+	UPROPERTY()
+	uint32 bConstrainRx:1;
+
+	UPROPERTY()
+	uint32 bConstrainRy:1;
+
+	UPROPERTY()
+	uint32 bConstrainRz:1;
+};
\ No newline at end of file
diff --git a/Engine/Source/Runtime/MovieSceneTracks/Public/TrackInstances/MovieScene3DAttachTrackInstance.h b/Engine/Source/Runtime/MovieSceneTracks/Public/TrackInstances/MovieScene3DAttachTrackInstance.h
new file mode 100644
--- /dev/null
+++ b/Engine/Source/Runtime/MovieSceneTracks/Public/TrackInstances/MovieScene3DAttachTrackInstance.h
@@ -0,0 +1,20 @@
+// Copyright 1998-2015 Epic Games, Inc. All Rights Reserved.
+
+#pragma once
+
+#include "MovieScene3DConstraintTrackInstance.h"
+
+class UMovieScene3DAttachTrack;
+
+/**
+ * Instance of a UMovieScene3DAttachTrack
+ */
+class FMovieScene3DAttachTrackInstance
+	: public FMovieScene3DConstraintTrackInstance
+{
+public:
+	FMovieScene3DAttachTrackInstance( UMovieScene3DAttachTrack& InAttachTrack );
+
+	/** FMovieScene3DConstraintTrackInstance interface */
+	virtual void UpdateConstraint(float Position, const TArray<UObject*>& RuntimeObjects, AActor* Actor, UMovieScene3DConstraintSection* ConstraintSection) override;
+};
diff --git a/Engine/Source/Runtime/MovieSceneTracks/Public/Tracks/MovieScene3DAttachTrack.h b/Engine/Source/Runtime/MovieSceneTracks/Public/Tracks/MovieScene3DAttachTrack.h
new file mode 100644
--- /dev/null
+++ b/Engine/Source/Runtime/MovieSceneTracks/Public/Tracks/MovieScene3DAttachTrack.h
@@ -0,0 +1,29 @@
+// Copyright 1998-2015 Epic Games, Inc. All Rights Reserved.
+
+#pragma once
+
+#include "MovieScene.h"
+#include "MovieScene3DConstraintTrack.h"
+#include "MovieScene3DAttachTrack.generated.h"
+
+/**
+ * Handles manipulation of path tracks in a movie scene
+ */
+UCLASS( MinimalAPI )
+class UMovieScene3DAttachTrack : public UMovieScene3DConstraintTrack
+{
+	GENERATED_UCLASS_BODY()
+public:
+	/** UMovieSceneTrack interface */
+	virtual FName GetTrackName() const override;
+	virtual TSharedPtr<IMovieSceneTrackInstance> CreateInstance() override;
+
+	/**
+	 * Adds a constraint
+	 *
+	 * @param Time				    The time relative to the owning movie scene where the section should be
+	 * @param ConstraintEndTime     Set the constraint to end at this time
+	 * @param ConstraintId			The id to the constraint
+	 */
+	virtual void AddConstraint( float Time, float ConstraintEndTime, const FGuid& ConstraintId ) override;
+};
