Sequencer: Refix getting the bool value from the property.

diff --git a/Engine/Source/Editor/MovieSceneTools/Private/TrackEditors/PropertyTrackEditors/BoolPropertyTrackEditor.cpp b/Engine/Source/Editor/MovieSceneTools/Private/TrackEditors/PropertyTrackEditors/BoolPropertyTrackEditor.cpp
--- a/Engine/Source/Editor/MovieSceneTools/Private/TrackEditors/PropertyTrackEditors/BoolPropertyTrackEditor.cpp
+++ b/Engine/Source/Editor/MovieSceneTools/Private/TrackEditors/PropertyTrackEditors/BoolPropertyTrackEditor.cpp
@@ -16,6 +16,12 @@ TSharedRef<ISequencerSection> FBoolPropertyTrackEditor::MakeSectionInterface( UM
 
 bool FBoolPropertyTrackEditor::TryGenerateKeyFromPropertyChanged( const FPropertyChangedParams& PropertyChangedParams, bool& OutKey )
 {
-	OutKey = *PropertyChangedParams.GetPropertyValue<bool>();
-	return true;
+	const UBoolProperty* BoolProperty = Cast<const UBoolProperty>(PropertyChangedParams.PropertyPath.Last());
+	if (BoolProperty)
+	{
+		OutKey = BoolProperty->GetPropertyValue(BoolProperty->ContainerPtrToValuePtr<void>(PropertyChangedParams.ObjectsThatChanged.Last()));
+		return true;
+	}
+
+	return false;
 }
\ No newline at end of file
diff --git a/Engine/Source/Editor/MovieSceneTools/Private/TrackEditors/PropertyTrackEditors/VisibilityPropertyTrackEditor.cpp b/Engine/Source/Editor/MovieSceneTools/Private/TrackEditors/PropertyTrackEditors/VisibilityPropertyTrackEditor.cpp
--- a/Engine/Source/Editor/MovieSceneTools/Private/TrackEditors/PropertyTrackEditors/VisibilityPropertyTrackEditor.cpp
+++ b/Engine/Source/Editor/MovieSceneTools/Private/TrackEditors/PropertyTrackEditors/VisibilityPropertyTrackEditor.cpp
@@ -16,6 +16,12 @@ TSharedRef<ISequencerSection> FVisibilityPropertyTrackEditor::MakeSectionInterfa
 
 bool FVisibilityPropertyTrackEditor::TryGenerateKeyFromPropertyChanged( const FPropertyChangedParams& PropertyChangedParams, bool& OutKey )
 {
-	OutKey = *PropertyChangedParams.GetPropertyValue<bool>();
-	return true;
+	const UBoolProperty* BoolProperty = Cast<const UBoolProperty>(PropertyChangedParams.PropertyPath.Last());
+	if (BoolProperty)
+	{
+		OutKey = BoolProperty->GetPropertyValue(BoolProperty->ContainerPtrToValuePtr<void>(PropertyChangedParams.ObjectsThatChanged.Last()));
+		return true;
+	}
+
+	return false;
 }
\ No newline at end of file
