Sequencer Shot UI improvements
- Shots now attempt to stay connected to each other when modifed.
-- Shots appear "connected" in the shot track so that when one moves, the previous one shrinks and next next shot grows.
-- Holes are filled for deleted shots
- Various low level shot rendering fixes

diff --git a/Engine/Source/Editor/EditorStyle/Private/SlateEditorStyle.cpp b/Engine/Source/Editor/EditorStyle/Private/SlateEditorStyle.cpp
--- a/Engine/Source/Editor/EditorStyle/Private/SlateEditorStyle.cpp
+++ b/Engine/Source/Editor/EditorStyle/Private/SlateEditorStyle.cpp
@@ -1589,8 +1589,8 @@ void FSlateEditorStyle::FStyle::SetupGeneralStyles()
 		Set( "Sequencer.Key", new IMAGE_BRUSH( "Common/Key", FVector2D(7.0f,12.0f) ) );
 		Set( "Sequencer.PartialKey", new IMAGE_BRUSH( "Sequencer/PartialKey", FVector2D(11.f, 11.f) ) );
 		Set( "Sequencer.GenericSection.Background", new BOX_BRUSH( "Sequencer/GenericSectionBackground", FMargin(5.f/16.f) ) );
-		Set( "Sequencer.SectionGripLeft", new BOX_BRUSH( "Sequencer/SectionGripLeft", FMargin(5.f/16.f) ) );
-		Set( "Sequencer.SectionGripRight", new BOX_BRUSH( "Sequencer/SectionGripRight", FMargin(5.f/16.f) ) );
+		Set( "Sequencer.DefaultSectionGripLeft", new BOX_BRUSH( "Sequencer/SectionGripLeft", FMargin(5.f/16.f) ) );
+		Set( "Sequencer.DefaultSectionGripRight", new BOX_BRUSH( "Sequencer/SectionGripRight", FMargin(5.f/16.f) ) );
 		Set( "Sequencer.SectionArea.Background", new FSlateColorBrush( FColor::White ) );
 		Set( "Sequencer.AnimationOutliner.SelectionBorder", new BOX_BRUSH( "Common/RoundedSelection_16x", 4.0f/16.0f, SelectionColor_Pressed ) );
 		Set( "Sequencer.AnimationOutliner.SelectionBorderInactive", new BOX_BRUSH("Common/RoundedSelection_16x", 4.0f / 16.0f, SelectionColor_Inactive));
@@ -1616,6 +1616,7 @@ void FSlateEditorStyle::FStyle::SetupGeneralStyles()
 		Set( "Sequencer.LockCamera", new IMAGE_BRUSH( "Icons/icon_ToggleActorPilotCameraView_16x", Icon16x16 ) );
 		Set( "Sequencer.UnlockCamera", new IMAGE_BRUSH( "Icons/icon_ToggleActorPilotCameraView_16x", Icon16x16, FLinearColor(1.f, 1.f, 1.f, 0.5f) ) );
 		Set( "Sequencer.Section.SelectionBorder", new BORDER_BRUSH( TEXT("Sequencer/SectionHighlight"), FMargin(7.f/16.f) ) );
+		Set( "Sequencer.ShotTrack.SectionHandle", new IMAGE_BRUSH( "Old/White", Icon16x16, FLinearColor::Black ) );
 
 		const FButtonStyle DetailsKeyButton = FButtonStyle(NoBorder)
 			.SetNormal( IMAGE_BRUSH("Sequencer/AddKey_Details", FVector2D(11,11) )  )
diff --git a/Engine/Source/Editor/MovieSceneTools/Private/TrackEditors/ShotTrackEditor.cpp b/Engine/Source/Editor/MovieSceneTools/Private/TrackEditors/ShotTrackEditor.cpp
--- a/Engine/Source/Editor/MovieSceneTools/Private/TrackEditors/ShotTrackEditor.cpp
+++ b/Engine/Source/Editor/MovieSceneTools/Private/TrackEditors/ShotTrackEditor.cpp
@@ -21,12 +21,16 @@
 #include "AssetToolsModule.h"
 #include "SInlineEditableTextBlock.h"
 
-namespace AnimatableShotToolConstants
+
+namespace ShotTrackConstants
 {
 	// @todo Sequencer Perhaps allow this to be customizable
-	const uint32 TrackHeight = 90;
+	const uint32 ThumbnailHeight = 90;
+	const uint32 TrackHeight = ThumbnailHeight+10; // some extra padding
 	const uint32 TrackWidth = 90;
 	const double ThumbnailFadeInDuration = 0.25f;
+	const float SectionGripSize = 4.0f;
+	const FName ShotTrackGripBrushName = FName("Sequencer.ShotTrack.SectionHandle");
 
 }
 
@@ -88,6 +92,9 @@ void FShotThumbnailPool::DrawThumbnails()
 		
 		if( ThumbnailsDrawn > 0 )
 		{
+			// Ensure all buffers are updated
+			FlushRenderingCommands();
+
 			// restore the global time
 			Sequencer.Pin()->SetGlobalTime(SavedTime);
 		}
@@ -99,13 +106,14 @@ void FShotThumbnailPool::DrawThumbnails()
 /////////////////////////////////////////////////
 // FShotThumbnail
 
-FShotThumbnail::FShotThumbnail(TSharedPtr<FShotSection> InSection, TRange<float> InTimeRange)
+FShotThumbnail::FShotThumbnail(TSharedPtr<FShotSection> InSection, const FIntPoint& InSize, TRange<float> InTimeRange)
 	: OwningSection(InSection)
+	, Size(InSize)
 	, Texture(NULL)
 	, TimeRange(InTimeRange)
-	, FadeInCurve( 0.0f, AnimatableShotToolConstants::ThumbnailFadeInDuration )
+	, FadeInCurve( 0.0f, ShotTrackConstants::ThumbnailFadeInDuration )
 {
-	Texture = new FSlateTexture2DRHIRef(GetSize().X, GetSize().Y, PF_B8G8R8A8, NULL, TexCreate_Dynamic, true);
+	Texture = new FSlateTexture2DRHIRef(GetSize().X, GetSize().Y, PF_B8G8R8A8, nullptr, TexCreate_Dynamic, true);
 
 	BeginInitResource( Texture );
 
@@ -123,8 +131,14 @@ FShotThumbnail::~FShotThumbnail()
 	}
 }
 
-FIntPoint FShotThumbnail::GetSize() const {return FIntPoint(OwningSection.Pin()->GetThumbnailWidth(), AnimatableShotToolConstants::TrackHeight);}
-bool FShotThumbnail::RequiresVsync() const {return false;}
+FIntPoint FShotThumbnail::GetSize() const
+{
+	return Size;
+}
+bool FShotThumbnail::RequiresVsync() const 
+{
+	return false;
+}
 
 FSlateShaderResource* FShotThumbnail::GetViewportRenderTargetTexture() const
 {
@@ -193,11 +207,11 @@ FShotSection::FShotSection( TSharedPtr<ISequencer> InSequencer, TSharedPtr<FShot
 		InternalViewportClient->EngineShowFlags = FEngineShowFlags(ESFIM_Game);
 		InternalViewportClient->EngineShowFlags.DisableAdvancedFeatures();
 		InternalViewportClient->SetActorLock(Camera.Get());
+		InternalViewportClient->SetAllowCinematicPreview( false );
+		InternalViewportClient->SetRealtime( false );
 
 		InternalViewportScene = MakeShareable(new FSceneViewport(InternalViewportClient.Get(), nullptr));
 		InternalViewportClient->Viewport = InternalViewportScene.Get();
-
-		CalculateThumbnailWidthAndResize();
 	}
 
 	WhiteBrush = FEditorStyle::GetBrush("WhiteBrush");
@@ -205,11 +219,14 @@ FShotSection::FShotSection( TSharedPtr<ISequencer> InSequencer, TSharedPtr<FShot
 
 FShotSection::~FShotSection()
 {
-	ThumbnailPool.Pin()->RemoveThumbnailsNeedingRedraw(Thumbnails);
-
-	if (InternalViewportClient.IsValid())
+	if( ThumbnailPool.IsValid() )
 	{
-		InternalViewportClient->Viewport = nullptr;
+		ThumbnailPool.Pin()->RemoveThumbnailsNeedingRedraw(Thumbnails);
+
+		if(InternalViewportClient.IsValid())
+		{
+			InternalViewportClient->Viewport = nullptr;
+		}
 	}
 }
 
@@ -241,7 +258,22 @@ FText FShotSection::GetSectionTitle() const
 
 float FShotSection::GetSectionHeight() const
 {
-	return AnimatableShotToolConstants::TrackHeight+10;
+	return ShotTrackConstants::TrackHeight;
+}
+
+float FShotSection::GetSectionGripSize() const
+{
+	return ShotTrackConstants::SectionGripSize;
+}
+
+FName FShotSection::GetSectionGripLeftBrushName() const
+{
+	return ShotTrackConstants::ShotTrackGripBrushName;
+}
+
+FName FShotSection::GetSectionGripRightBrushName() const
+{
+	return ShotTrackConstants::ShotTrackGripBrushName;
 }
 
 FReply FShotSection::OnSectionDoubleClicked( const FGeometry& SectionGeometry, const FPointerEvent& MouseEvent )
@@ -276,22 +308,24 @@ int32 FShotSection::OnPaintSection( const FGeometry& AllottedGeometry, const FSl
 {
 	if (Camera.IsValid())
 	{
+		FGeometry ThumbnailAreaGeometry = AllottedGeometry.MakeChild( FVector2D(GetSectionGripSize(), 0.0f), AllottedGeometry.GetDrawSize() - FVector2D( GetSectionGripSize()*2, 0.0f ) );
+
 		FSlateDrawElement::MakeBox(
 				OutDrawElements,
 				LayerId,
-				AllottedGeometry.ToPaintGeometry(),
-				WhiteBrush,
-				SectionClippingRect,
-				ESlateDrawEffect::None,
-				FLinearColor::Black
+				ThumbnailAreaGeometry.ToPaintGeometry(),
+				FEditorStyle::GetBrush("Sequencer.GenericSection.Background"),
+				SectionClippingRect
 			);
 
+		// @todo Sequencer: Need a way to visualize the key here
+
 		int32 ThumbnailIndex = 0;
 		for( const TSharedPtr<FShotThumbnail>& Thumbnail : Thumbnails )
 		{
-			FGeometry TruncatedGeometry = AllottedGeometry.MakeChild(
+			FGeometry TruncatedGeometry = ThumbnailAreaGeometry.MakeChild(
 				Thumbnail->GetSize(),
-				FSlateLayoutTransform(AllottedGeometry.Scale, FVector2D(ThumbnailIndex * ThumbnailWidth, 5.0f))
+				FSlateLayoutTransform(ThumbnailAreaGeometry.Scale, FVector2D( ThumbnailIndex * ThumbnailWidth, 5.f))
 				);
 			
 			FSlateDrawElement::MakeViewport(
@@ -310,12 +344,12 @@ int32 FShotSection::OnPaintSection( const FGeometry& AllottedGeometry, const FSl
 			{
 				FSlateDrawElement::MakeBox(
 					OutDrawElements,
-					LayerId + 1,
+					LayerId+1,
 					TruncatedGeometry.ToPaintGeometry(),
 					WhiteBrush,
 					SectionClippingRect,
 					ESlateDrawEffect::None,
-					FColor(255, 255, 255, Thumbnail->GetFadeInCurve() * 255)
+					FLinearColor(1.0f, 1.0f, 1.0f, Thumbnail->GetFadeInCurve())
 					);
 			}
 
@@ -333,13 +367,14 @@ int32 FShotSection::OnPaintSection( const FGeometry& AllottedGeometry, const FSl
 		); 
 	}
 
-	return LayerId + 1;
+	return LayerId + 2;
 }
 
 void FShotSection::Tick( const FGeometry& AllottedGeometry, const FGeometry& ParentGeometry, const double InCurrentTime, const float InDeltaTime )
 {
 	if (!Camera.IsValid())
 	{
+		// @todo Sequencer: This is called too often and is expensive
 		Camera = UpdateCameraObject();
 	}
 
@@ -350,7 +385,9 @@ void FShotSection::Tick( const FGeometry& AllottedGeometry, const FGeometry& Par
 		TimeToPixelConverter.PixelToTime(-AllottedGeometry.Position.X),
 		TimeToPixelConverter.PixelToTime(-AllottedGeometry.Position.X + ParentGeometry.Size.X));
 
-	FIntPoint AllocatedSize = AllottedGeometry.Size.IntPoint();
+	FIntPoint AllocatedSize = AllottedGeometry.MakeChild( FVector2D( GetSectionGripSize(), 0.0f ), FVector2D( AllottedGeometry.Size.X - (GetSectionGripSize()*2), AllottedGeometry.Size.Y) ).Size.IntPoint();
+	AllocatedSize.X = FMath::Max(AllocatedSize.X, 1);
+
 	float StartTime = Section->GetStartTime();
 
 	if (!FSlateApplication::Get().HasAnyMouseCaptor() && ( AllocatedSize.X != StoredSize.X || !FMath::IsNearlyEqual(StartTime, StoredStartTime) ) )
@@ -369,79 +406,66 @@ void FShotSection::RegenerateViewportThumbnails(const FIntPoint& Size)
 	StoredSize = Size;
 	StoredStartTime = Section->GetStartTime();
 	
-	ThumbnailPool.Pin()->RemoveThumbnailsNeedingRedraw(Thumbnails);
-	Thumbnails.Empty();
-
-	if (Size.X == 0 || Size.Y == 0)
+	if( Camera.IsValid() )
 	{
-		return;
-	}
+		ThumbnailPool.Pin()->RemoveThumbnailsNeedingRedraw(Thumbnails);
+		Thumbnails.Empty();
 
-	CalculateThumbnailWidthAndResize();
+		if(Size.X == 0 || Size.Y == 0)
+		{
+			return;
+		}
 
-	int32 ThumbnailCount = FMath::DivideAndRoundUp(StoredSize.X, (int32)ThumbnailWidth);
-	
-	float StartTime = Section->GetStartTime();
-	float EndTime = Section->GetEndTime();
-	float DeltaTime = EndTime - StartTime;
+		float AspectRatio = Camera->CameraComponent->AspectRatio;
 
-	int32 TotalX = StoredSize.X;
-	
-	// @todo sequencer optimize this to use a single viewport and a single thumbnail
-	// instead paste the textures into the thumbnail at specific offsets
-	for (int32 i = 0; i < ThumbnailCount; ++i)
-	{
-		int32 StartX = i * ThumbnailWidth;
-		float FractionThrough = (float)StartX / (float)TotalX;
-		float Time = StartTime + DeltaTime * FractionThrough;
+		int32 NewThumbnailWidth = FMath::TruncToInt(ShotTrackConstants::TrackWidth*AspectRatio);
 
-		int32 NextStartX = (i+1) * ThumbnailWidth;
-		float NextFractionThrough = (float)NextStartX / (float)TotalX;
-		float NextTime = FMath::Min(StartTime + DeltaTime * NextFractionThrough, EndTime);
-	
-		check(FractionThrough >= 0.f && FractionThrough <= 1.f && NextFractionThrough >= 0.f);
-		TSharedPtr<FShotThumbnail> NewThumbnail = MakeShareable(new FShotThumbnail(SharedThis(this), TRange<float>(Time, NextTime)));
+		int32 ThumbnailCount = FMath::DivideAndRoundUp(StoredSize.X, NewThumbnailWidth);
 
-		Thumbnails.Add(NewThumbnail);
-	}
+		if(NewThumbnailWidth != ThumbnailWidth)
+		{
+			ThumbnailWidth = NewThumbnailWidth;
+			InternalViewportScene->UpdateViewportRHI(false, ThumbnailWidth, ShotTrackConstants::ThumbnailHeight, EWindowMode::Windowed);
+		}
+
+		float StartTime = Section->GetStartTime();
+		float EndTime = Section->GetEndTime();
+		float DeltaTime = EndTime - StartTime;
+
+		int32 TotalX = StoredSize.X;
+
+		// @todo sequencer optimize this to use a single viewport and a single thumbnail
+		// instead paste the textures into the thumbnail at specific offsets
+		for(int32 ThumbnailIndex = 0; ThumbnailIndex < ThumbnailCount; ++ThumbnailIndex)
+		{
+			int32 StartX = ThumbnailIndex * ThumbnailWidth;
+			float FractionThrough = (float)StartX / (float)TotalX;
+			float Time = StartTime + DeltaTime * FractionThrough;
+
+			int32 NextStartX = (ThumbnailIndex+1) * ThumbnailWidth;
+			float NextFractionThrough = (float)NextStartX / (float)TotalX;
+			float NextTime = FMath::Min(StartTime + DeltaTime * NextFractionThrough, EndTime);
+
+			FIntPoint ThumbnailSize(NextStartX-StartX, ShotTrackConstants::ThumbnailHeight);
+
+			check(FractionThrough >= 0.f && FractionThrough <= 1.f && NextFractionThrough >= 0.f);
+			TSharedPtr<FShotThumbnail> NewThumbnail = MakeShareable(new FShotThumbnail(SharedThis(this), ThumbnailSize, TRange<float>(Time, NextTime)));
+
+			Thumbnails.Add(NewThumbnail);
+		}
 
-	// @todo Sequencer Optimize Only say a thumbnail needs redraw if it is onscreen
-	ThumbnailPool.Pin()->AddThumbnailsNeedingRedraw(Thumbnails);
+		// @todo Sequencer Optimize Only say a thumbnail needs redraw if it is onscreen
+		ThumbnailPool.Pin()->AddThumbnailsNeedingRedraw(Thumbnails);
+	}
 }
 
 void FShotSection::DrawViewportThumbnail(TSharedPtr<FShotThumbnail> ShotThumbnail)
 {
 	Sequencer.Pin()->SetGlobalTime(ShotThumbnail->GetTime());
 	InternalViewportClient->UpdateViewForLockedActor();
-			
 	GWorld->SendAllEndOfFrameUpdates();
 	InternalViewportScene->Draw(false);
 	ShotThumbnail->CopyTextureIn((FSlateRenderTargetRHI*)InternalViewportScene->GetViewportRenderTargetTexture());
-
-	FlushRenderingCommands();
-}
-
-void FShotSection::CalculateThumbnailWidthAndResize()
-{
-	// @todo Sequencer We also should detect when the property is changed or keyframed to update
-
-	// get the aspect ratio at the first frame
-	// if it changes over time, we ignore this
-	Sequencer.Pin()->SetPerspectiveViewportPossessionEnabled(false);
-	float SavedTime = Sequencer.Pin()->GetGlobalTime();
-	
-	Sequencer.Pin()->SetGlobalTime(Section->GetStartTime());
-	uint32 OutThumbnailWidth = AnimatableShotToolConstants::TrackWidth * Camera->GetCameraComponent()->AspectRatio;
-
-	// restore the global time
-	Sequencer.Pin()->SetGlobalTime(SavedTime);
-	Sequencer.Pin()->SetPerspectiveViewportPossessionEnabled(true);
-
-	if (OutThumbnailWidth != ThumbnailWidth)
-	{
-		ThumbnailWidth = OutThumbnailWidth;
-		InternalViewportScene->UpdateViewportRHI( false, ThumbnailWidth, AnimatableShotToolConstants::TrackHeight, EWindowMode::Windowed );
-	}
 }
 
 ACameraActor* FShotSection::UpdateCameraObject() const
@@ -468,7 +492,6 @@ void FShotSection::OnRenameShot( const FText& NewShotName, ETextCommit::Type Com
 {
 	if (CommitType == ETextCommit::OnEnter && !GetShotName().EqualTo( NewShotName ) )
 	{
-		Section->Modify();
 		Section->SetShotNameAndNumber( NewShotName, Section->GetShotNumber() );
 
 		UMovieSceneSequence* ShotMovieSceneAnim = Section->GetMovieSceneAnimation();
@@ -602,33 +625,8 @@ void FShotTrackEditor::AddKeyInternal( float KeyTime, const FGuid ObjectGuid )
 
 	UObject* NewAsset = AssetToolsModule.Get().CreateAsset( AssetName, PackagePath, AssetClass, Factory );
 	
-	float EndTime = KeyTime;
-	if( AllSections.IsValidIndex( SectionIndex + 1 ) )
-	{
-		// The new section ends where the next section begins
-		EndTime = AllSections[SectionIndex+1]->GetStartTime();
-	}
-	else
-	{
-		// There is no section after this one so the new section should be the length of the movie scene
-		EndTime = MovieScene->GetTimeRange().GetUpperBoundValue();
-	}
-
-
-	if (AllSections.IsValidIndex(SectionIndex - 1))
-	{
-		// Shift the previous shot's end time so that it ends when the new shot starts
-		AllSections[SectionIndex - 1]->SetEndTime(KeyTime);
-	}
-
-	if (AllSections.IsValidIndex(SectionIndex + 1))
-	{
-		// Shift the next shot's start time so that it starts when the new shot ends
-		AllSections[SectionIndex + 1]->SetStartTime(EndTime);
-	}
 
-	TRange<float> TimeRange( KeyTime, EndTime );
-	ShotTrack->AddNewShot( ObjectGuid, *CastChecked<UMovieSceneSequence>( NewAsset ), TimeRange, FText::FromString( ShotName ), ShotNumber );
+	ShotTrack->AddNewShot( ObjectGuid, *CastChecked<UMovieSceneSequence>( NewAsset ), KeyTime, FText::FromString( ShotName ), ShotNumber );
 }
 
 int32 FShotTrackEditor::GenerateShotNumber( UMovieSceneShotTrack& ShotTrack, int32 SectionIndex ) const
@@ -688,6 +686,7 @@ int32 FShotTrackEditor::GenerateShotNumber( UMovieSceneShotTrack& ShotTrack, int
 
 int32 FShotTrackEditor::FindIndexForNewShot( const TArray<UMovieSceneSection*>& ShotSections, float NewShotTime ) const
 {
+	// @todo Sequencer: This could be a binary search. All shots should be sorted 
 	int32 SectionIndex = 0;
 	for (SectionIndex = 0; SectionIndex < ShotSections.Num() && ShotSections[SectionIndex]->GetStartTime() < NewShotTime; ++SectionIndex);
 
diff --git a/Engine/Source/Editor/MovieSceneTools/Private/TrackEditors/ShotTrackEditor.h b/Engine/Source/Editor/MovieSceneTools/Private/TrackEditors/ShotTrackEditor.h
--- a/Engine/Source/Editor/MovieSceneTools/Private/TrackEditors/ShotTrackEditor.h
+++ b/Engine/Source/Editor/MovieSceneTools/Private/TrackEditors/ShotTrackEditor.h
@@ -86,7 +86,7 @@ private:
 class FShotThumbnail : public ISlateViewport, public TSharedFromThis<FShotThumbnail>
 {
 public:
-	FShotThumbnail(TSharedPtr<class FShotSection> InSection, TRange<float> InTimeRange);
+	FShotThumbnail(TSharedPtr<class FShotSection> InSection, const FIntPoint& InSize, TRange<float> InTimeRange);
 	~FShotThumbnail();
 
 	/* ISlateViewport interface */
@@ -115,6 +115,8 @@ private:
 	/** Parent shot section we are a thumbnail of */
 	TWeakPtr<class FShotSection> OwningSection;
 	
+	FIntPoint Size;
+
 	/** The Texture RHI that holds the thumbnail */
 	class FSlateTexture2DRHIRef* Texture;
 
@@ -143,6 +145,10 @@ public:
 	virtual FText GetDisplayName() const override { return NSLOCTEXT("FShotSection", "", "Shots"); }
 	virtual FText GetSectionTitle() const override;
 	virtual float GetSectionHeight() const override;
+	virtual float GetSectionGripSize() const override;
+	virtual FName GetSectionGripLeftBrushName() const override;
+	virtual FName GetSectionGripRightBrushName() const override;
+	virtual bool AreSectionsConnected() const override { return true; }
 	virtual void GenerateSectionLayout( class ISectionLayoutBuilder& LayoutBuilder ) const override {}
 	virtual FReply OnSectionDoubleClicked( const FGeometry& SectionGeometry, const FPointerEvent& MouseEvent ) override;
 	virtual void BuildSectionContextMenu(FMenuBuilder& MenuBuilder) override;
@@ -159,17 +165,17 @@ public:
 	/** Draws the passed in viewport thumbnail and copies it to the thumbnail's texture */
 	void DrawViewportThumbnail(TSharedPtr<FShotThumbnail> ShotThumbnail);
 
-	/** Calculates and sets the thumbnail width, and resizes if it is different than before */
-	void CalculateThumbnailWidthAndResize();
-
 	/** Gets the time range of what in the sequencer is visible */
 	TRange<float> GetVisibleTimeRange() const {return VisibleTimeRange;}
 
 	/** @return The sequencer widget owning the shot section */
 	TSharedRef<SWidget> GetSequencerWidget() { return Sequencer.Pin()->GetSequencerWidget(); }
 private:
+
 	ACameraActor* UpdateCameraObject() const;
+
 	FText GetShotName() const;
+
 	void OnRenameShot( const FText& NewShotName, ETextCommit::Type CommitType );
 private:
 	/** The section we are visualizing */
diff --git a/Engine/Source/Editor/Sequencer/Private/SSection.cpp b/Engine/Source/Editor/Sequencer/Private/SSection.cpp
--- a/Engine/Source/Editor/Sequencer/Private/SSection.cpp
+++ b/Engine/Source/Editor/Sequencer/Private/SSection.cpp
@@ -65,7 +65,12 @@ FSelectedKey SSection::GetKeyUnderMouse( const FVector2D& MousePosition, const F
 		// Is the key area under the mouse
 		if( KeyAreaGeometryPadded.IsUnderLocation( MousePosition ) )
 		{
-			FGeometry SectionGeometry = AllottedGeometry.MakeChild(FVector2D(SequencerSectionConstants::SectionGripSize, 0), AllottedGeometry.GetDrawSize() - FVector2D(SequencerSectionConstants::SectionGripSize*2, 0.0f));
+			// Do not display section handles for connected sections
+			const bool bDisplaySectionHandles = SectionInterface->AreSectionsConnected();
+
+			const float SectionGripSize = bDisplaySectionHandles ? SectionInterface->GetSectionGripSize() : 0.0f;
+
+			FGeometry SectionGeometry = AllottedGeometry.MakeChild(FVector2D(SectionGripSize, 0), AllottedGeometry.GetDrawSize() - FVector2D(SectionGripSize*2, 0.0f));
 			FGeometry KeyAreaGeometry = GetKeyAreaGeometry( Element, SectionGeometry );
 
 			FVector2D LocalSpaceMousePosition = KeyAreaGeometry.AbsoluteToLocal( MousePosition );
@@ -117,11 +122,11 @@ void SSection::CheckForEdgeInteraction( const FPointerEvent& MouseEvent, const F
 	// Make areas to the left and right of the geometry.  We will use these areas to determine if someone dragged the left or right edge of a section
 	FGeometry SectionRectLeft = SectionGeometry.MakeChild(
 		FVector2D::ZeroVector,
-		FVector2D( SequencerSectionConstants::SectionGripSize, SectionGeometry.Size.Y )
+		FVector2D( SectionInterface->GetSectionGripSize(), SectionGeometry.Size.Y )
 		);
 
 	FGeometry SectionRectRight = SectionGeometry.MakeChild(
-		FVector2D( SectionGeometry.Size.X - SequencerSectionConstants::SectionGripSize, 0 ), 
+		FVector2D( SectionGeometry.Size.X - SectionInterface->GetSectionGripSize(), 0 ), 
 		SectionGeometry.Size 
 		);
 
@@ -162,14 +167,16 @@ int32 SSection::OnPaint( const FPaintArgs& Args, const FGeometry& AllottedGeomet
 {
 	int32 StartLayer = SCompoundWidget::OnPaint( Args, AllottedGeometry, MyClippingRect, OutDrawElements, LayerId, InWidgetStyle, bParentEnabled );
 
-	FGeometry SectionGeometry = AllottedGeometry.MakeChild( FVector2D( SequencerSectionConstants::SectionGripSize, 0 ), AllottedGeometry.GetDrawSize() - FVector2D( SequencerSectionConstants::SectionGripSize*2, 0.0f ) );
+	FGeometry SectionGeometry = MakeSectionGeometryWithoutHandles( AllottedGeometry, SectionInterface );
 
 	FSlateRect SectionClipRect = SectionGeometry.GetClippingRect().IntersectionWith( MyClippingRect );
 
 	// Ask the interface to draw the section
 	int32 PostSectionLayer = SectionInterface->OnPaintSection( SectionGeometry, SectionClipRect, OutDrawElements, LayerId, bParentEnabled );
-	
-	DrawSectionBorders(AllottedGeometry, MyClippingRect, OutDrawElements, PostSectionLayer );
+
+	const bool bDisplaySectionHandles = true;
+
+	DrawSectionHandlesAndSelection(AllottedGeometry, MyClippingRect, OutDrawElements, PostSectionLayer, bDisplaySectionHandles );
 
 	PaintKeys( SectionGeometry, MyClippingRect, OutDrawElements, PostSectionLayer, InWidgetStyle );
 
@@ -321,7 +328,7 @@ void SSection::PaintKeys( const FGeometry& AllottedGeometry, const FSlateRect& M
 	}
 }
 
-void SSection::DrawSectionBorders( const FGeometry& AllottedGeometry, const FSlateRect& MyClippingRect, FSlateWindowElementList& OutDrawElements, int32 LayerId ) const
+void SSection::DrawSectionHandlesAndSelection( const FGeometry& AllottedGeometry, const FSlateRect& MyClippingRect, FSlateWindowElementList& OutDrawElements, int32 LayerId, bool bDisplaySectionHandles ) const
 {
 	UMovieSceneSection* SectionObject = SectionInterface->GetSectionObject();
 
@@ -329,39 +336,41 @@ void SSection::DrawSectionBorders( const FGeometry& AllottedGeometry, const FSla
 	const bool bSelected = Selection.IsSelected(SectionObject);
 	const bool bActive = Selection.GetActiveSelection() == FSequencerSelection::EActiveSelection::KeyAndSection;
 
-	static const FName SelectionColorName("SelectionColor");
-	static const FName SelectionInactiveColorName("SelectionColorInactive");
-
-	FLinearColor SelectionColor = FEditorStyle::GetSlateColor(SelectionColorName).GetColor(FWidgetStyle());
-	FLinearColor SelectionInactiveColor = FEditorStyle::GetSlateColor(SelectionInactiveColorName).GetColor(FWidgetStyle());
+	FLinearColor SelectionColor = FEditorStyle::GetSlateColor(SequencerSectionConstants::SelectionColorName).GetColor(FWidgetStyle());
+	FLinearColor SelectionInactiveColor = FEditorStyle::GetSlateColor(SequencerSectionConstants::SelectionInactiveColorName).GetColor(FWidgetStyle());
 	FLinearColor TransparentSelectionColor = SelectionColor;
 
-	static const FName SectionGripLeftName("Sequencer.SectionGripLeft");
-	static const FName SectionGripRightName("Sequencer.SectionGripRight");
-
-	// Left Grip
-	FSlateDrawElement::MakeBox(
-		OutDrawElements,
-		LayerId,
-		// Center the key along Y.  Ensure the middle of the key is at the actual key time
-		AllottedGeometry.ToPaintGeometry( FVector2D( 0.0f, 0.0f ), FVector2D( SequencerSectionConstants::SectionGripSize, AllottedGeometry.GetDrawSize().Y) ) ,
-		FEditorStyle::GetBrush(SectionGripLeftName),
-		MyClippingRect,
-		ESlateDrawEffect::None,
-		(bLeftEdgePressed || bLeftEdgeHovered) ? TransparentSelectionColor : FLinearColor::White
-	);
-
-	// Right Grip
-	FSlateDrawElement::MakeBox(
-		OutDrawElements,
-		LayerId,
-		// Center the key along Y.  Ensure the middle of the key is at the actual key time
-		AllottedGeometry.ToPaintGeometry( FVector2D( AllottedGeometry.Size.X-SequencerSectionConstants::SectionGripSize, 0.0f), FVector2D(SequencerSectionConstants::SectionGripSize, AllottedGeometry.GetDrawSize().Y)),
-		FEditorStyle::GetBrush(SectionGripRightName),
-		MyClippingRect,
-		ESlateDrawEffect::None,
-		(bRightEdgePressed || bRightEdgeHovered) ? TransparentSelectionColor : FLinearColor::White
+	if( bDisplaySectionHandles )
+	{
+		const FSlateBrush* LeftGripBrush = FEditorStyle::GetBrush(SectionInterface->GetSectionGripLeftBrushName());
+		const FSlateBrush* RightGripBrush = FEditorStyle::GetBrush(SectionInterface->GetSectionGripRightBrushName());
+
+		// Left Grip
+		FSlateDrawElement::MakeBox
+		(
+			OutDrawElements,
+			LayerId,
+			// Center the key along Y.  Ensure the middle of the key is at the actual key time
+			AllottedGeometry.ToPaintGeometry(FVector2D(0.0f, 0.0f), FVector2D(SectionInterface->GetSectionGripSize(), AllottedGeometry.GetDrawSize().Y)),
+			LeftGripBrush,
+			MyClippingRect,
+			ESlateDrawEffect::None,
+			(bLeftEdgePressed || bLeftEdgeHovered) ? TransparentSelectionColor : LeftGripBrush->GetTint(FWidgetStyle())
 		);
+		
+		// Right Grip
+		FSlateDrawElement::MakeBox
+		(
+			OutDrawElements,
+			LayerId,
+			// Center the key along Y.  Ensure the middle of the key is at the actual key time
+			AllottedGeometry.ToPaintGeometry(FVector2D(AllottedGeometry.Size.X-SectionInterface->GetSectionGripSize(), 0.0f), FVector2D(SectionInterface->GetSectionGripSize(), AllottedGeometry.GetDrawSize().Y)),
+			RightGripBrush,
+			MyClippingRect,
+			ESlateDrawEffect::None,
+			(bRightEdgePressed || bRightEdgeHovered) ? TransparentSelectionColor : RightGripBrush->GetTint(FWidgetStyle())
+		);
+	}
 
 
 	// draw selection box
@@ -429,7 +438,9 @@ void SSection::Tick( const FGeometry& AllottedGeometry, const double InCurrentTi
 		{
 			Layout = FKeyAreaLayout(*ParentSectionArea, SectionIndex);
 		}
-		SectionInterface->Tick(AllottedGeometry, ParentGeometry, InCurrentTime, InDeltaTime);
+		FGeometry SectionGeometry = MakeSectionGeometryWithoutHandles( AllottedGeometry, SectionInterface );
+
+		SectionInterface->Tick(SectionGeometry, ParentGeometry, InCurrentTime, InDeltaTime);
 	}
 }
 
@@ -477,6 +488,15 @@ void SSection::ResetHoveredState()
 	GetSequencer().GetEditTool().SetHotspot(nullptr);
 }
 
+FGeometry SSection::MakeSectionGeometryWithoutHandles( const FGeometry& AllottedGeometry, const TSharedPtr<ISequencerSection>& SectionInterface ) const
+{
+	const bool bSectionsAreConnected = SectionInterface->AreSectionsConnected();
+
+	const float SectionGripSize = !bSectionsAreConnected ? SectionInterface->GetSectionGripSize() : 0.0f;
+
+	return AllottedGeometry.MakeChild( FVector2D( SectionGripSize, 0 ), AllottedGeometry.GetDrawSize() - FVector2D( SectionGripSize*2, 0.0f ) );
+}
+
 FReply SSection::OnMouseButtonUp( const FGeometry& MyGeometry, const FPointerEvent& MouseEvent )
 {
 	if( ( MouseEvent.GetEffectingButton() == EKeys::LeftMouseButton || MouseEvent.GetEffectingButton() == EKeys::RightMouseButton ) && MyGeometry.IsUnderLocation( MouseEvent.GetScreenSpacePosition() ) )
@@ -635,6 +655,7 @@ void SSection::HandleSectionSelection( const FPointerEvent& MouseEvent, bool bSe
 {
 	// handle selecting sections 
 	UMovieSceneSection* Section = SectionInterface->GetSectionObject();
+
 	bool bSectionIsSelected = GetSequencer().GetSelection().IsSelected(Section);
 
 	// Clear previous section selection if:
@@ -657,6 +678,15 @@ void SSection::HandleSectionSelection( const FPointerEvent& MouseEvent, bool bSe
 		{
 			GetSequencer().GetSelection().EmptySelectedSections();
 		}
+	}
+	
+	if( GetSequencer().GetSelection().IsSelected(Section) )
+	{
+		// Control should toggle the selection
+		GetSequencer().GetSelection().RemoveFromSelection(Section);
+	}
+	else
+	{
 		GetSequencer().GetSelection().AddToSelection(Section);
 	}
 }
diff --git a/Engine/Source/Editor/Sequencer/Private/SSection.h b/Engine/Source/Editor/Sequencer/Private/SSection.h
--- a/Engine/Source/Editor/Sequencer/Private/SSection.h
+++ b/Engine/Source/Editor/Sequencer/Private/SSection.h
@@ -4,16 +4,6 @@
 
 #include "KeyAreaLayout.h"
 
-namespace SequencerSectionConstants
-{
-	/** How far the user has to drag the mouse before we consider the action dragging rather than a click */
-	const float SectionDragStartDistance = 5.0f;
-
-	/** The size of each key */
-	const FVector2D KeySize(11.0f, 11.0f);
-
-	const float SectionGripSize = 7.0f;
-}
 
 class SSection : public SCompoundWidget
 {
@@ -88,7 +78,7 @@ private:
 	/**
 	 * Draw the borders around each section.  The borders are the areas which can be used to expand or collapse the size of the section.
 	 */
-	void DrawSectionBorders( const FGeometry& AllottedGeometry, const FSlateRect& MyClippingRect, FSlateWindowElementList& OutDrawElements, int32 LayerId ) const;
+	void DrawSectionHandlesAndSelection( const FGeometry& AllottedGeometry, const FSlateRect& MyClippingRect, FSlateWindowElementList& OutDrawElements, int32 LayerId, bool bDisplaySectionHandles ) const;
 
 	/** Summons a context menu over the associated section */
 	TSharedPtr<SWidget> OnSummonContextMenu( const FGeometry& MyGeometry, const FPointerEvent& MouseEvent );
@@ -139,6 +129,12 @@ private:
 	 * Resets the hovered state of elements in the widget (section edges, keys)
 	 */
 	void ResetHoveredState();
+
+	/** 
+	 * Creates geometry for a section without space for the handles
+	 */
+	FGeometry MakeSectionGeometryWithoutHandles( const FGeometry& AllottedGeometry, const TSharedPtr<ISequencerSection>& SectionInterface ) const;
+
 private:
 	/** Interface to section data */
 	TSharedPtr<ISequencerSection> SectionInterface;
diff --git a/Engine/Source/Editor/Sequencer/Private/SSequencerSectionAreaView.cpp b/Engine/Source/Editor/Sequencer/Private/SSequencerSectionAreaView.cpp
--- a/Engine/Source/Editor/Sequencer/Private/SSequencerSectionAreaView.cpp
+++ b/Engine/Source/Editor/Sequencer/Private/SSequencerSectionAreaView.cpp
@@ -31,6 +31,9 @@ namespace SequencerSectionUtils
 	FGeometry GetSectionGeometry( const FGeometry& AllottedGeometry, int32 RowIndex, int32 MaxTracks, float NodeHeight, TSharedPtr<ISequencerSection> SectionInterface, const FTimeToPixel& TimeToPixelConverter )
 	{
 		const UMovieSceneSection* Section = SectionInterface->GetSectionObject();
+		float PixelStartX = TimeToPixelConverter.TimeToPixel( Section->GetStartTime() );
+		// Note the -1 pixel at the end is because the section does not actually end at the end time if there is a section starting at that same time.  It is more important that a section lines up correctly with it's true start time
+		float PixelEndX = TimeToPixelConverter.TimeToPixel( Section->GetEndTime() );
 
 		float StartX, EndX = 0;
 
@@ -47,14 +50,22 @@ namespace SequencerSectionUtils
 		}
 
 		// Actual section length without grips.
-		float SectionLengthActual = EndX-StartX;
+		float SectionLengthActual = FMath::Max(1.0f,PixelEndX-PixelStartX);
 
-		float SectionLengthWithGrips = SectionLengthActual+SequencerSectionConstants::SectionGripSize*2;
+		float SectionLengthWithGrips = SectionLengthActual;
+		float ExtentSize = 0;
+		if( !SectionInterface->AreSectionsConnected() )
+		{
+			// Extend the section to include areas for the grips
+			ExtentSize = SectionInterface->GetSectionGripSize();
+			// Connected sections do not display grips outside of their section area
+			SectionLengthWithGrips += (ExtentSize*2);
+		}
 
 		float ActualHeight = NodeHeight / MaxTracks;
 
 		// Compute allotted geometry area that can be used to draw the section
-		return AllottedGeometry.MakeChild( FVector2D( StartX-SequencerSectionConstants::SectionGripSize, ActualHeight * RowIndex ), FVector2D( SectionLengthWithGrips, ActualHeight ) );
+		return AllottedGeometry.MakeChild( FVector2D( PixelStartX-ExtentSize, ActualHeight * RowIndex ), FVector2D( SectionLengthWithGrips, ActualHeight ) );
 	}
 
 }
@@ -92,7 +103,7 @@ void SSequencerSectionAreaView::GenerateSectionWidgets()
 
 	if( SectionAreaNode.IsValid() )
 	{
-		const TArray< TSharedRef<ISequencerSection> >& Sections = SectionAreaNode->GetSections();
+		TArray< TSharedRef<ISequencerSection> >& Sections = SectionAreaNode->GetSections();
 
 		for ( int32 SectionIndex = 0; SectionIndex < Sections.Num(); ++SectionIndex )
 		{
diff --git a/Engine/Source/Editor/Sequencer/Private/Sequencer.cpp b/Engine/Source/Editor/Sequencer/Private/Sequencer.cpp
--- a/Engine/Source/Editor/Sequencer/Private/Sequencer.cpp
+++ b/Engine/Source/Editor/Sequencer/Private/Sequencer.cpp
@@ -1826,7 +1826,7 @@ void FSequencer::DeleteSelectedItems()
 bool FSequencer::CanDeleteSelectedItems() const
 {
 	// Only delete sequencer items if the delete key was pressed in the sequencer widget and not in, for example, the viewport.
-	return SequencerWidget->HasFocusedDescendants();
+	return SequencerWidget->HasFocusedDescendants() || SequencerWidget->HasKeyboardFocus();
 }
 
 void FSequencer::AssignActor(FGuid InObjectBinding, FObjectBindingNode* ObjectBindingNode)
diff --git a/Engine/Source/Editor/Sequencer/Private/Tools/EditToolDragOperations.cpp b/Engine/Source/Editor/Sequencer/Private/Tools/EditToolDragOperations.cpp
--- a/Engine/Source/Editor/Sequencer/Private/Tools/EditToolDragOperations.cpp
+++ b/Engine/Source/Editor/Sequencer/Private/Tools/EditToolDragOperations.cpp
@@ -64,16 +64,16 @@ TRange<float> GetSectionBoundaries(UMovieSceneSection* Section, TSharedPtr<FTrac
 	const TArray< TSharedRef<ISequencerSection> >& Sections = SequencerNode->GetSections();
 	for (int32 SectionIndex = 0; SectionIndex < Sections.Num(); ++SectionIndex)
 	{
-		const UMovieSceneSection* InSection = Sections[SectionIndex]->GetSectionObject();
-		if (Section != InSection && Section->GetRowIndex() == InSection->GetRowIndex())
+		const UMovieSceneSection* TestSection = Sections[SectionIndex]->GetSectionObject();
+		if (Section != TestSection && Section->GetRowIndex() == TestSection->GetRowIndex())
 		{
-			if (InSection->GetEndTime() <= Section->GetStartTime() && InSection->GetEndTime() > LowerBound)
+			if (TestSection->GetEndTime() <= Section->GetStartTime() && TestSection->GetEndTime() > LowerBound)
 			{
-				LowerBound = InSection->GetEndTime();
+				LowerBound = TestSection->GetEndTime();
 			}
-			if (InSection->GetStartTime() >= Section->GetEndTime() && InSection->GetStartTime() < UpperBound)
+			if (TestSection->GetStartTime() >= Section->GetEndTime() && TestSection->GetStartTime() < UpperBound)
 			{
-				UpperBound = InSection->GetStartTime();
+				UpperBound = TestSection->GetStartTime();
 			}
 		}
 	}
@@ -97,7 +97,7 @@ void GetSectionSnapTimes(TArray<float>& OutSnapTimes, UMovieSceneSection* Sectio
 		}
 	}
 
-	// snap to director track if it exists, and we are not the director track
+	// snap to shots if it exists, and we are not the director track
 	UMovieSceneTrack* OuterTrack = Cast<UMovieSceneTrack>(Section->GetOuter());
 	UMovieScene* MovieScene = Cast<UMovieScene>(OuterTrack->GetOuter());
 	UMovieSceneTrack* ShotTrack = MovieScene->FindMasterTrack(UMovieSceneShotTrack::StaticClass());
@@ -261,7 +261,7 @@ void FResizeSection::OnDrag(const FPointerEvent& MouseEvent, FVector2D LocalMous
 		{
 			// Dragging the end of a section
 			// Ensure we aren't shrinking past the start time
-			NewTime = FMath::Clamp( NewTime, Section->GetStartTime(), SectionBoundaries.GetUpperBoundValue() );
+			NewTime = FMath::Max( NewTime, Section->GetStartTime() );
 
 			if (bIsDilating)
 			{
@@ -278,8 +278,8 @@ void FResizeSection::OnDrag(const FPointerEvent& MouseEvent, FVector2D LocalMous
 		{
 			// Dragging the start of a section
 			// Ensure we arent expanding past the end time
-			NewTime = FMath::Clamp( NewTime, SectionBoundaries.GetLowerBoundValue(), Section->GetEndTime() );
-			
+			NewTime = FMath::Min( NewTime, Section->GetEndTime() );
+
 			if (bIsDilating)
 			{
 				float NewSize = Section->GetEndTime() - NewTime;
@@ -291,6 +291,13 @@ void FResizeSection::OnDrag(const FPointerEvent& MouseEvent, FVector2D LocalMous
 				Section->SetStartTime( NewTime );
 			}
 		}
+
+		UMovieSceneTrack* OuterTrack = Section->GetTypedOuter<UMovieSceneTrack>();
+		if (OuterTrack)
+		{
+			OuterTrack->Modify();
+			OuterTrack->OnSectionMoved(*Section);
+		}
 	}
 }
 
diff --git a/Engine/Source/Editor/Sequencer/Public/ISequencerSection.h b/Engine/Source/Editor/Sequencer/Public/ISequencerSection.h
--- a/Engine/Source/Editor/Sequencer/Public/ISequencerSection.h
+++ b/Engine/Source/Editor/Sequencer/Public/ISequencerSection.h
@@ -4,6 +4,27 @@
 
 class IKeyArea;
 
+namespace SequencerSectionConstants
+{
+	/** How far the user has to drag the mouse before we consider the action dragging rather than a click */
+	const float SectionDragStartDistance = 5.0f;
+
+	/** The size of each key */
+	const FVector2D KeySize(11.0f, 11.0f);
+
+	const float DefaultSectionGripSize = 7.0f;
+
+	const float DefaultSectionHeight = 15.f;
+
+	const FName SelectionColorName("SelectionColor");
+
+	const FName SelectionInactiveColorName("SelectionColorInactive");
+	
+	const FName DefaultSectionGripLeftImageName("Sequencer.DefaultSectionGripLeft");
+	
+	const FName DefaultSectionGripRightImageName("Sequencer.DefaultSectionGripRight");
+}
+
 /**
  * Interface that should be implemented for the UI portion of a section
  */
@@ -65,13 +86,22 @@ public:
 	/**
 	 * @return The height of the section
 	 */
-	virtual float GetSectionHeight() const { return 15.0f; }
+	virtual float GetSectionHeight() const { return SequencerSectionConstants::DefaultSectionHeight; }
 	
+	virtual float GetSectionGripSize() const { return SequencerSectionConstants::DefaultSectionGripSize; }
+
+	virtual FName GetSectionGripLeftBrushName() const { return SequencerSectionConstants::DefaultSectionGripLeftImageName; }
+
+	virtual FName GetSectionGripRightBrushName() const { return SequencerSectionConstants::DefaultSectionGripRightImageName; }
+
 	/**
 	 * @return Whether or not the user can resize this section.
 	 */
 	virtual bool SectionIsResizable() const {return true;}
 
+
+	virtual bool AreSectionsConnected() const { return false; }
+
 	/**
 	 * Ticks the section during the Slate tick
 	 *
diff --git a/Engine/Source/Runtime/MovieScene/Public/MovieSceneSection.h b/Engine/Source/Runtime/MovieScene/Public/MovieSceneSection.h
--- a/Engine/Source/Runtime/MovieScene/Public/MovieSceneSection.h
+++ b/Engine/Source/Runtime/MovieScene/Public/MovieSceneSection.h
@@ -33,14 +33,24 @@ public:
 	 * 
 	 * @param InEndTime	The new end time
 	 */
-	void SetStartTime( float NewStartTime ) { StartTime = NewStartTime; }
+	void SetStartTime( float NewStartTime )
+	{ 
+		Modify();
+
+		StartTime = NewStartTime;
+	}
 
 	/**
 	 * Sets a new end time for this section
 	 * 
 	 * @param InEndTime	The new end time
 	 */
-	void SetEndTime( float NewEndTime ) { EndTime = NewEndTime; }
+	void SetEndTime( float NewEndTime )
+	{ 
+		Modify();
+
+		EndTime = NewEndTime;
+	}
 	
 	/**
 	 * @return The range of times of the section
diff --git a/Engine/Source/Runtime/MovieSceneTracks/Private/Sections/MovieSceneShotSection.cpp b/Engine/Source/Runtime/MovieSceneTracks/Private/Sections/MovieSceneShotSection.cpp
--- a/Engine/Source/Runtime/MovieSceneTracks/Private/Sections/MovieSceneShotSection.cpp
+++ b/Engine/Source/Runtime/MovieSceneTracks/Private/Sections/MovieSceneShotSection.cpp
@@ -20,6 +20,8 @@ FGuid UMovieSceneShotSection::GetCameraGuid() const
 
 void UMovieSceneShotSection::SetShotNameAndNumber(const FText& InDisplayName, int32 InShotNumber)
 {
+	Modify();
+
 	DisplayName = InDisplayName;
 	ShotNumber = InShotNumber;
 }
diff --git a/Engine/Source/Runtime/MovieSceneTracks/Private/TrackInstances/SubMovieSceneTrackInstance.cpp b/Engine/Source/Runtime/MovieSceneTracks/Private/TrackInstances/SubMovieSceneTrackInstance.cpp
--- a/Engine/Source/Runtime/MovieSceneTracks/Private/TrackInstances/SubMovieSceneTrackInstance.cpp
+++ b/Engine/Source/Runtime/MovieSceneTracks/Private/TrackInstances/SubMovieSceneTrackInstance.cpp
@@ -70,16 +70,21 @@ void FSubMovieSceneTrackInstance::Update( float Position, float LastPosition, co
 	{
 		USubMovieSceneSection* Section = CastChecked<USubMovieSceneSection>( TraversedSections[SectionIndex] );
 
-		TSharedPtr<FMovieSceneSequenceInstance>& Instance = SubMovieSceneInstances.FindChecked( Section );
+		TSharedPtr<FMovieSceneSequenceInstance> Instance = SubMovieSceneInstances.FindRef( Section );
+
+		FMovieSceneSequenceInstance* InstancePtr = Instance.Get();
 
-		TRange<float> TimeRange = Instance->GetMovieSceneTimeRange();
+		if( InstancePtr )
+		{
+			TRange<float> TimeRange = InstancePtr->GetMovieSceneTimeRange();
 
-		// Position for the movie scene needs to be in local space
-		float LocalDelta = TimeRange.GetLowerBoundValue() - Section->GetStartTime();
-		float LocalPosition = Position + LocalDelta;
+			// Position for the movie scene needs to be in local space
+			float LocalDelta = TimeRange.GetLowerBoundValue() - Section->GetStartTime();
+			float LocalPosition = Position + LocalDelta;
 
 
-		Instance->Update( LocalPosition, LastPosition + LocalDelta, Player );
+			InstancePtr->Update(LocalPosition, LastPosition + LocalDelta, Player);
+		}
 	}
 
 }
diff --git a/Engine/Source/Runtime/MovieSceneTracks/Private/Tracks/MovieSceneShotTrack.cpp b/Engine/Source/Runtime/MovieSceneTracks/Private/Tracks/MovieSceneShotTrack.cpp
--- a/Engine/Source/Runtime/MovieSceneTracks/Private/Tracks/MovieSceneShotTrack.cpp
+++ b/Engine/Source/Runtime/MovieSceneTracks/Private/Tracks/MovieSceneShotTrack.cpp
@@ -33,35 +33,44 @@ void UMovieSceneShotTrack::AddSection( UMovieSceneSection* Section )
 
 void UMovieSceneShotTrack::RemoveSection( UMovieSceneSection* Section )
 {
+	FixupSurroundingShots( *Section, true );
+
 	Super::RemoveSection( Section );
 
+	SortShots();
+
 	// @todo Sequencer: The movie scene owned by the section is now abandoned.  Should we offer to delete it?  
+
 }
 
-void UMovieSceneShotTrack::AddNewShot(FGuid CameraHandle, UMovieSceneSequence& ShotMovieSceneAnimation, const TRange<float>& TimeRange, const FText& ShotName, int32 ShotNumber )
+void UMovieSceneShotTrack::AddNewShot(FGuid CameraHandle, UMovieSceneSequence& ShotMovieSceneSequence, float StartTime, const FText& ShotName, int32 ShotNumber )
 {
 	Modify();
 
 	FName UniqueShotName = MakeUniqueObjectName( this, UMovieSceneShotSection::StaticClass(), *ShotName.ToString() );
 
 	UMovieSceneShotSection* NewSection = NewObject<UMovieSceneShotSection>( this, UniqueShotName, RF_Transactional );
-	NewSection->SetMovieSceneAnimation( &ShotMovieSceneAnimation );
-	NewSection->SetStartTime( TimeRange.GetLowerBoundValue() );
-	NewSection->SetEndTime( TimeRange.GetUpperBoundValue() );
+	NewSection->SetMovieSceneAnimation( &ShotMovieSceneSequence );
+	NewSection->SetStartTime( StartTime );
+	NewSection->SetEndTime( FindEndTimeForShot( StartTime ) );
 	NewSection->SetCameraGuid( CameraHandle );
 	NewSection->SetShotNameAndNumber( ShotName , ShotNumber );
 
 	SubMovieSceneSections.Add( NewSection );
-	
+
 	// When a new shot is added, sort all shots to ensure they are in the correct order
 	SortShots();
+
+
+
+	// Once shots are sorted fixup the surrounding shots to fix any gaps
+	FixupSurroundingShots( *NewSection, false );
 }
 
 #if WITH_EDITOR
 void UMovieSceneShotTrack::OnSectionMoved( UMovieSceneSection& Section )
 {
-	// Sort all shots when one moves.  They should be in order of start time
-	SortShots();
+	FixupSurroundingShots( Section, false );
 }
 #endif
 
@@ -74,4 +83,63 @@ void UMovieSceneShotTrack::SortShots()
 	});
 }
 
+void UMovieSceneShotTrack::FixupSurroundingShots( UMovieSceneSection& Section, bool bDelete )
+{
+	// Find the previous section and extend it to take the place of the section being deleted
+	int32 SectionIndex = INDEX_NONE;
+	if( SubMovieSceneSections.Find( &Section, SectionIndex ) )
+	{
+		int32 PrevSectionIndex = SectionIndex - 1;
+		if( SubMovieSceneSections.IsValidIndex( PrevSectionIndex ) )
+		{
+			// Extend the previous section
+			SubMovieSceneSections[PrevSectionIndex]->SetEndTime( bDelete ? Section.GetEndTime() : Section.GetStartTime() );
+		}
+
+		if( !bDelete )
+		{
+			int32 NextSectionIndex = SectionIndex + 1;
+			if(SubMovieSceneSections.IsValidIndex(NextSectionIndex))
+			{
+				// Shift the next shot's start time so that it starts when the new shot ends
+				SubMovieSceneSections[NextSectionIndex]->SetStartTime(Section.GetEndTime());
+			}
+		}
+	}
+
+	SortShots();
+}
+
+float UMovieSceneShotTrack::FindEndTimeForShot( float StartTime )
+{
+	float EndTime = 0;
+	bool bFoundEndTime = false;
+	for( UMovieSceneSection* Section : SubMovieSceneSections )
+	{
+		if( Section->GetStartTime() >= StartTime )
+		{
+			EndTime = Section->GetStartTime();
+			bFoundEndTime = true;
+			break;
+		}
+	}
+
+	if( !bFoundEndTime )
+	{
+		UMovieScene* OwnerScene = GetTypedOuter<UMovieScene>();
+
+		// End time should just end where the movie scene ends.  Ensure it is at least the same as start time (this should only happen when the movie scene has an initial time range smaller than the start time
+		EndTime = FMath::Max( OwnerScene->GetTimeRange().GetUpperBoundValue(), StartTime );
+	}
+
+			
+	if( StartTime == EndTime )
+	{
+		// Give the shot a reasonable length of time to start out with.  A 0 time shot is not usable
+		EndTime = StartTime + .5f;
+	}
+
+	return EndTime;
+}
+
 #undef LOCTEXT_NAMESPACE
\ No newline at end of file
diff --git a/Engine/Source/Runtime/MovieSceneTracks/Public/Tracks/MovieSceneShotTrack.h b/Engine/Source/Runtime/MovieSceneTracks/Public/Tracks/MovieSceneShotTrack.h
--- a/Engine/Source/Runtime/MovieSceneTracks/Public/Tracks/MovieSceneShotTrack.h
+++ b/Engine/Source/Runtime/MovieSceneTracks/Public/Tracks/MovieSceneShotTrack.h
@@ -24,20 +24,22 @@ public:
 
 	/** 
 	 * Adds a new shot at the specified time
-	 *
-	 * @param CameraHandle		Handle to the camera that the shot switches to when active
-	 * @param ShotMovieScene	MovieScene for the shot (each shot has a unique movie scene containing per tracks only active during the shot)
-	 * @param TimeRange		The range within this track's movie scene where the shot is initially placed
-	 * @param ShotName		The display name of the shot
-	 * @param ShotNumber		The number of the shot.  This is used to assist with auto-generated shot names when new shots are added
+	 *	
+	 * @param CameraHandle				Handle to the camera that the shot switches to when active
+	 * @param ShotMovieSceneSequence	MovieScene for the shot (each shot has a unique movie scene containing per tracks only active during the shot)
+	 * @param TimeRange					The range within this track's movie scene where the shot is initially placed
+	 * @param ShotName					The display name of the shot
+	 * @param ShotNumber				The number of the shot.  This is used to assist with auto-generated shot names when new shots are added
 	 */
-	MOVIESCENETRACKS_API void AddNewShot(FGuid CameraHandle, UMovieSceneSequence& ShotMovieSceneAnimation, const TRange<float>& TimeRange, const FText& ShotName, int32 ShotNumber );
+	MOVIESCENETRACKS_API void AddNewShot(FGuid CameraHandle, UMovieSceneSequence& ShotMovieSceneSequence, float Time, const FText& ShotName, int32 ShotNumber );
 
 #if WITH_EDITOR
 	virtual void OnSectionMoved( UMovieSceneSection& Section ) override;
 #endif
 
 private:
+	float FindEndTimeForShot( float StartTime );
+	void FixupSurroundingShots( UMovieSceneSection& Section, bool bDelete );
 	/** Sorts shots according to their start time */
 	void SortShots();
 };
