Sequencer: Change autokey so that it doesn't automatically create a track/key if a track doesn't already exist.

diff --git a/Engine/Source/Editor/MovieSceneTools/Private/TrackEditors/TransformTrackEditor.cpp b/Engine/Source/Editor/MovieSceneTools/Private/TrackEditors/TransformTrackEditor.cpp
--- a/Engine/Source/Editor/MovieSceneTools/Private/TrackEditors/TransformTrackEditor.cpp
+++ b/Engine/Source/Editor/MovieSceneTools/Private/TrackEditors/TransformTrackEditor.cpp
@@ -218,6 +218,8 @@ struct FTransformDataPair
 
 void F3DTransformTrackEditor::OnTransformChanged( UObject& InObject )
 {
+	const bool bCreateHandleIfMissing = false;
+
 	const TSharedPtr<ISequencer> Sequencer = GetSequencer();
 	if (Sequencer->GetAutoKeyEnabled())
 	{
@@ -230,7 +232,7 @@ void F3DTransformTrackEditor::OnTransformChanged( UObject& InObject )
 		if( Actor && Actor->GetRootComponent() )
 		{
 			// Get a handle bound to the key/section we are adding so we know what objects to change during playback
-			ObjectHandle = Sequencer->GetHandleToObject( Actor );
+			ObjectHandle = Sequencer->GetHandleToObject( Actor, bCreateHandleIfMissing );
 			SceneComponentThatChanged = Actor->GetRootComponent();
 		}
 
@@ -240,27 +242,31 @@ void F3DTransformTrackEditor::OnTransformChanged( UObject& InObject )
 			SceneComponentThatChanged = Cast<USceneComponent>( &InObject );
 			if( SceneComponentThatChanged )
 			{
-				ObjectHandle = Sequencer->GetHandleToObject( SceneComponentThatChanged );
+				ObjectHandle = Sequencer->GetHandleToObject( SceneComponentThatChanged, bCreateHandleIfMissing );
 			}
 
 		}
 
 		if( SceneComponentThatChanged && ObjectHandle.IsValid() )
 		{
-			// Find an existing transform if possible.  If one exists we will compare against the new one to decide what components of the transform need keys
-			FTransformData ExistingTransform = ObjectToExistingTransform.FindRef( &InObject );
+			FName Transform("Transform");
+			if (Sequencer->GetFocusedMovieSceneSequence()->GetMovieScene()->FindTrack(UMovieScene3DTransformTrack::StaticClass(), ObjectHandle, Transform))
+			{
+				// Find an existing transform if possible.  If one exists we will compare against the new one to decide what components of the transform need keys
+				FTransformData ExistingTransform = ObjectToExistingTransform.FindRef( &InObject );
 
-			// Remove it from the list of cached transforms. 
-			// @todo sequencer livecapture: This can be made much for efficient by not removing cached state during live capture situation
-			ObjectToExistingTransform.Remove( &InObject );
+				// Remove it from the list of cached transforms. 
+				// @todo sequencer livecapture: This can be made much for efficient by not removing cached state during live capture situation
+				ObjectToExistingTransform.Remove( &InObject );
 
-			// Build new transform data
-			FTransformData NewTransformData( SceneComponentThatChanged );
+				// Build new transform data
+				FTransformData NewTransformData( SceneComponentThatChanged );
 
-			FTransformDataPair TransformPair(NewTransformData, ExistingTransform);
+				FTransformDataPair TransformPair(NewTransformData, ExistingTransform);
 
-			AnimatablePropertyChanged(UMovieScene3DTransformTrack::StaticClass(), true,
-				FOnKeyProperty::CreateRaw(this, &F3DTransformTrackEditor::OnTransformChangedInternals, &InObject, ObjectHandle, TransformPair, true, false, F3DTransformTrackKey::Key_All));
+				AnimatablePropertyChanged(UMovieScene3DTransformTrack::StaticClass(), true,
+					FOnKeyProperty::CreateRaw(this, &F3DTransformTrackEditor::OnTransformChangedInternals, &InObject, ObjectHandle, TransformPair, true, false, F3DTransformTrackKey::Key_All));
+			}
 		}
 	}
 }
diff --git a/Engine/Source/Editor/MovieSceneTools/Public/PropertyTrackEditor.h b/Engine/Source/Editor/MovieSceneTools/Public/PropertyTrackEditor.h
--- a/Engine/Source/Editor/MovieSceneTools/Public/PropertyTrackEditor.h
+++ b/Engine/Source/Editor/MovieSceneTools/Public/PropertyTrackEditor.h
@@ -132,6 +132,27 @@ private:
 		WatchedPropertyTypeNames.Add( WatchedPropertyTypeName );
 	}
 
+	/** Get the property name and track class from the property */
+	void GetPropertyAndTrackClass(const UProperty* Property, FName& PropertyName, TSubclassOf<UMovieSceneTrack>& SequencerTrackClass)
+	{
+		PropertyName = Property->GetFName();
+
+		// Look for a customized track class for this property on the meta data
+		const FString& MetaSequencerTrackClass = Property->GetMetaData( TEXT( "SequencerTrackClass" ) );
+		if ( !MetaSequencerTrackClass.IsEmpty() )
+		{
+			UClass* MetaClass = FindObject<UClass>( ANY_PACKAGE, *MetaSequencerTrackClass );
+			if ( !MetaClass )
+			{
+				MetaClass = LoadObject<UClass>( nullptr, *MetaSequencerTrackClass );
+			}
+			if ( MetaClass != NULL )
+			{
+				SequencerTrackClass = MetaClass;
+			}
+		}
+	}
+
 	/**
 	* Called by the details panel when an animatable property changes
 	*
@@ -140,41 +161,52 @@ private:
 	*/
 	virtual void OnAnimatedPropertyChanged( const FPropertyChangedParams& PropertyChangedParams )
 	{
-		// Get the value from the property
-		KeyType Key;
-		if ( TryGenerateKeyFromPropertyChanged( PropertyChangedParams, Key ) )
+		bool bCreateHandleIfMissing = !PropertyChangedParams.bRequireAutoKey;
+
+		for ( UObject* Object : PropertyChangedParams.ObjectsThatChanged )
 		{
-			AnimatablePropertyChanged( TrackType::StaticClass(), PropertyChangedParams.bRequireAutoKey,
-				FOnKeyProperty::CreateRaw( this, &FPropertyTrackEditor::OnKeyProperty, PropertyChangedParams, Key ) );
+			FGuid ObjectHandle = FindOrCreateHandleToObject( Object, bCreateHandleIfMissing );
+
+			if ( ObjectHandle.IsValid() )
+			{
+				FName PropertyName;
+				TSubclassOf<UMovieSceneTrack> SequencerTrackClass = TrackType::StaticClass();
+				GetPropertyAndTrackClass(PropertyChangedParams.PropertyPath.Last(), PropertyName, SequencerTrackClass);
+
+				bool bCreateTrackIfMissing = !PropertyChangedParams.bRequireAutoKey;
+
+				if (bCreateTrackIfMissing || GetSequencer()->GetFocusedMovieSceneSequence()->GetMovieScene()->FindTrack(SequencerTrackClass, ObjectHandle, PropertyName))
+				{
+					// Get the value from the property
+					KeyType Key;
+					if ( TryGenerateKeyFromPropertyChanged( PropertyChangedParams, Key ) )
+					{
+						AnimatablePropertyChanged( TrackType::StaticClass(), PropertyChangedParams.bRequireAutoKey,
+							FOnKeyProperty::CreateRaw( this, &FPropertyTrackEditor::OnKeyProperty, PropertyChangedParams, Key ) );
+					}
+				}
+			}
 		}
 	}
 
 	/** Adds a key based on a property change. */
 	void OnKeyProperty( float KeyTime, FPropertyChangedParams PropertyChangedParams, KeyType Key )
 	{
+		bool bCreateHandleIfMissing = !PropertyChangedParams.bRequireAutoKey;
+
 		for ( UObject* Object : PropertyChangedParams.ObjectsThatChanged )
 		{
-			FGuid ObjectHandle = FindOrCreateHandleToObject( Object );
+			FGuid ObjectHandle = FindOrCreateHandleToObject( Object, bCreateHandleIfMissing );
 			if ( ObjectHandle.IsValid() )
 			{
-				FName PropertyName = PropertyChangedParams.PropertyPath.Last()->GetFName();
-
-				// Look for a customized track class for this property on the meta data
-				const FString& MetaSequencerTrackClass = PropertyChangedParams.PropertyPath.Last()->GetMetaData( TEXT( "SequencerTrackClass" ) );
+				FName PropertyName;
 				TSubclassOf<UMovieSceneTrack> SequencerTrackClass = TrackType::StaticClass();
-				if ( !MetaSequencerTrackClass.IsEmpty() )
-				{
-					UClass* MetaClass = FindObject<UClass>( ANY_PACKAGE, *MetaSequencerTrackClass );
-					if ( !MetaClass )
-					{
-						MetaClass = LoadObject<UClass>( nullptr, *MetaSequencerTrackClass );
-					}
-					if ( MetaClass != NULL )
-						SequencerTrackClass = MetaClass;
-				}
+				GetPropertyAndTrackClass(PropertyChangedParams.PropertyPath.Last(), PropertyName, SequencerTrackClass);
+
+				bool bCreateTrackIfMissing = !PropertyChangedParams.bRequireAutoKey;
 
-				UMovieSceneTrack* Track = GetTrackForObject( ObjectHandle, SequencerTrackClass, PropertyName );
-				if ( ensure( Track ) )
+				UMovieSceneTrack* Track = GetTrackForObject( ObjectHandle, SequencerTrackClass, PropertyName, bCreateTrackIfMissing );
+				if ( Track )
 				{
 					TrackType* TypedTrack = CastChecked<TrackType>( Track );
 
diff --git a/Engine/Source/Editor/Sequencer/Private/Sequencer.cpp b/Engine/Source/Editor/Sequencer/Private/Sequencer.cpp
--- a/Engine/Source/Editor/Sequencer/Private/Sequencer.cpp
+++ b/Engine/Source/Editor/Sequencer/Private/Sequencer.cpp
@@ -678,7 +678,7 @@ void FSequencer::SetPerspectiveViewportPossessionEnabled(bool bEnabled)
 	bPerspectiveViewportPossessionEnabled = bEnabled;
 }
 
-FGuid FSequencer::GetHandleToObject( UObject* Object )
+FGuid FSequencer::GetHandleToObject( UObject* Object, bool bCreateHandleIfMissing )
 {
 	if (Object == nullptr)
 	{
@@ -715,7 +715,7 @@ FGuid FSequencer::GetHandleToObject( UObject* Object )
 
 	// If the object guid was not found attempt to add it
 	// Note: Only possessed actors can be added like this
-	if (!ObjectGuid.IsValid() && FocusedMovieSceneSequence->CanPossessObject(*Object))
+	if (!ObjectGuid.IsValid() && FocusedMovieSceneSequence->CanPossessObject(*Object) && bCreateHandleIfMissing)
 	{
 		// @todo sequencer: Undo doesn't seem to be working at all
 		const FScopedTransaction Transaction(LOCTEXT("UndoPossessingObject", "Possess Object with MovieScene"));
diff --git a/Engine/Source/Editor/Sequencer/Private/Sequencer.h b/Engine/Source/Editor/Sequencer/Private/Sequencer.h
--- a/Engine/Source/Editor/Sequencer/Private/Sequencer.h
+++ b/Engine/Source/Editor/Sequencer/Private/Sequencer.h
@@ -77,7 +77,7 @@ public:
 	virtual float GetGlobalTime() override;
 	virtual void SetGlobalTime(float Time) override;
 	virtual void SetPerspectiveViewportPossessionEnabled(bool bEnabled) override;
-	virtual FGuid GetHandleToObject(UObject* Object) override;
+	virtual FGuid GetHandleToObject(UObject* Object, bool bCreateHandleIfMissing = true) override;
 	virtual ISequencerObjectChangeListener& GetObjectChangeListener() override;
 	virtual void NotifyMovieSceneDataChanged() override;
 	virtual void UpdateRuntimeInstances() override;
diff --git a/Engine/Source/Editor/Sequencer/Public/ISequencer.h b/Engine/Source/Editor/Sequencer/Public/ISequencer.h
--- a/Engine/Source/Editor/Sequencer/Public/ISequencer.h
+++ b/Engine/Source/Editor/Sequencer/Public/ISequencer.h
@@ -145,9 +145,10 @@ public:
 	 * Gets a handle to runtime information about the object being manipulated by a movie scene
 	 * 
 	 * @param Object The object to get a handle for.
+	 * @param bCreateHandleIfMissing Create a handle if it doesn't exist.
 	 * @return The handle to the object.
 	 */
-	virtual FGuid GetHandleToObject(UObject* Object) = 0;
+	virtual FGuid GetHandleToObject(UObject* Object, bool bCreateHandleIfMissing = true) = 0;
 
 	/**
 	 * @return Returns the object change listener for sequencer instance
diff --git a/Engine/Source/Editor/Sequencer/Public/MovieSceneTrackEditor.h b/Engine/Source/Editor/Sequencer/Public/MovieSceneTrackEditor.h
--- a/Engine/Source/Editor/Sequencer/Public/MovieSceneTrackEditor.h
+++ b/Engine/Source/Editor/Sequencer/Public/MovieSceneTrackEditor.h
@@ -104,12 +104,12 @@ public:
 	 * @param Object	The object to create a binding for
 	 * @return A handle to the binding or an invalid handle if the object could not be bound
 	 */
-	FGuid FindOrCreateHandleToObject( UObject* Object )
+	FGuid FindOrCreateHandleToObject( UObject* Object, bool bCreateHandleIfMissing = true )
 	{
-		return GetSequencer()->GetHandleToObject( Object );
+		return GetSequencer()->GetHandleToObject( Object, bCreateHandleIfMissing );
 	}
 
-	UMovieSceneTrack* GetTrackForObject( const FGuid& ObjectHandle, TSubclassOf<class UMovieSceneTrack> TrackClass, FName UniqueTypeName )
+	UMovieSceneTrack* GetTrackForObject( const FGuid& ObjectHandle, TSubclassOf<class UMovieSceneTrack> TrackClass, FName UniqueTypeName, bool bCreateTrackIfMissing = true )
 	{
 		check( UniqueTypeName != NAME_None );
 
@@ -117,7 +117,7 @@ public:
 
 		UMovieSceneTrack* Type = MovieScene->FindTrack( TrackClass, ObjectHandle, UniqueTypeName );
 
-		if( !Type )
+		if( !Type && bCreateTrackIfMissing )
 		{
 			Type = AddTrack( MovieScene, ObjectHandle, TrackClass, UniqueTypeName );
 		}
