Sequencer: Add missing undef namespace.

diff --git a/Engine/Source/Editor/MovieSceneTools/Private/TrackEditors/TransformTrackEditor.cpp b/Engine/Source/Editor/MovieSceneTools/Private/TrackEditors/TransformTrackEditor.cpp
--- a/Engine/Source/Editor/MovieSceneTools/Private/TrackEditors/TransformTrackEditor.cpp
+++ b/Engine/Source/Editor/MovieSceneTools/Private/TrackEditors/TransformTrackEditor.cpp
@@ -572,3 +572,5 @@ FText F3DTransformTrackEditor::GetLockCameraToolTip(TWeakObjectPtr<ACameraActor>
 		FText::Format(LOCTEXT("UnlockCamera", "Unlock {0} from Viewport"), FText::FromString(CameraActor.Get()->GetActorLabel())) :
 		FText::Format(LOCTEXT("LockCamera", "Lock {0} to Selected Viewport"), FText::FromString(CameraActor.Get()->GetActorLabel()));
 }
+
+#undef LOCTEXT_NAMESPACE
\ No newline at end of file
