Sequencer: Add default value for curves.

This gives sections a default value which is evaluated and set when there are no keys in the section. The typical workflow for setting up a transform now becomes first creating a transform track with a transform section. The current transformation of the object is set as the default value of the transform section. The default value is updated as the object is moved. But once a keyframe is set, the default value is no longer used. When all keys are removed, the default value is again in use.

Internal notes: Consolidated special flags to a KeyParams struct.

diff --git a/Engine/Source/Editor/MovieSceneTools/Private/CurveKeyEditors/SBoolCurveKeyEditor.cpp b/Engine/Source/Editor/MovieSceneTools/Private/CurveKeyEditors/SBoolCurveKeyEditor.cpp
--- a/Engine/Source/Editor/MovieSceneTools/Private/CurveKeyEditors/SBoolCurveKeyEditor.cpp
+++ b/Engine/Source/Editor/MovieSceneTools/Private/CurveKeyEditors/SBoolCurveKeyEditor.cpp
@@ -46,7 +46,14 @@ void SBoolCurveKeyEditor::OnCheckStateChanged(ECheckBoxState NewCheckboxState)
 		}
 	}
 
-	Curve->UpdateOrAddKey(CurrentTime, NewCheckboxState == ECheckBoxState::Checked ? 1 : 0);
+	if (Curve->GetNumKeys() == 0)
+	{
+		Curve->SetDefaultValue(NewCheckboxState == ECheckBoxState::Checked ? 1 : 0);
+	}
+	else
+	{
+		Curve->UpdateOrAddKey(CurrentTime, NewCheckboxState == ECheckBoxState::Checked ? 1 : 0);
+	}
 	Sequencer->UpdateRuntimeInstances();
 }
 
diff --git a/Engine/Source/Editor/MovieSceneTools/Private/CurveKeyEditors/SEnumCurveKeyEditor.cpp b/Engine/Source/Editor/MovieSceneTools/Private/CurveKeyEditors/SEnumCurveKeyEditor.cpp
--- a/Engine/Source/Editor/MovieSceneTools/Private/CurveKeyEditors/SEnumCurveKeyEditor.cpp
+++ b/Engine/Source/Editor/MovieSceneTools/Private/CurveKeyEditors/SEnumCurveKeyEditor.cpp
@@ -72,7 +72,14 @@ void SEnumCurveKeyEditor::OnComboSelectionChanged(TSharedPtr<FString> InSelected
 
 	int32 SelectedIndex;
 	EnumValues.Find(InSelectedItem, SelectedIndex);
-	Curve->UpdateOrAddKey(CurrentTime, SelectedIndex);
+	if (Curve->GetNumKeys() == 0)
+	{
+		Curve->SetDefaultValue(SelectedIndex);
+	}
+	else
+	{
+		Curve->UpdateOrAddKey(CurrentTime, SelectedIndex);
+	}
 	Sequencer->UpdateRuntimeInstances();
 }
 
diff --git a/Engine/Source/Editor/MovieSceneTools/Private/CurveKeyEditors/SFloatCurveKeyEditor.cpp b/Engine/Source/Editor/MovieSceneTools/Private/CurveKeyEditors/SFloatCurveKeyEditor.cpp
--- a/Engine/Source/Editor/MovieSceneTools/Private/CurveKeyEditors/SFloatCurveKeyEditor.cpp
+++ b/Engine/Source/Editor/MovieSceneTools/Private/CurveKeyEditors/SFloatCurveKeyEditor.cpp
@@ -62,7 +62,15 @@ void SFloatCurveKeyEditor::OnValueChanged(float Value)
 	}
 	else
 	{
-		Curve->AddKey(CurrentTime, Value, false, CurrentKeyHandle);
+		if (Curve->GetNumKeys() == 0)
+		{
+			Curve->SetDefaultValue(Value);
+		}
+		else
+		{
+			Curve->AddKey(CurrentTime, Value, false, CurrentKeyHandle);
+		}
+
 		if (OwningSection->GetStartTime() > CurrentTime)
 		{
 			OwningSection->SetStartTime(CurrentTime);
diff --git a/Engine/Source/Editor/MovieSceneTools/Private/CurveKeyEditors/SIntegralCurveKeyEditor.cpp b/Engine/Source/Editor/MovieSceneTools/Private/CurveKeyEditors/SIntegralCurveKeyEditor.cpp
--- a/Engine/Source/Editor/MovieSceneTools/Private/CurveKeyEditors/SIntegralCurveKeyEditor.cpp
+++ b/Engine/Source/Editor/MovieSceneTools/Private/CurveKeyEditors/SIntegralCurveKeyEditor.cpp
@@ -69,7 +69,14 @@ void SIntegralCurveKeyEditor::OnValueChanged(int32 Value)
 		}
 	}
 
-	Curve->UpdateOrAddKey(CurrentTime, Value);
+	if (Curve->GetNumKeys() == 0)
+	{
+		Curve->SetDefaultValue(Value);
+	}
+	else
+	{
+		Curve->UpdateOrAddKey(CurrentTime, Value);
+	}
 	Sequencer->UpdateRuntimeInstances();
 }
 
diff --git a/Engine/Source/Editor/MovieSceneTools/Private/MovieSceneToolsPrivatePCH.h b/Engine/Source/Editor/MovieSceneTools/Private/MovieSceneToolsPrivatePCH.h
--- a/Engine/Source/Editor/MovieSceneTools/Private/MovieSceneToolsPrivatePCH.h
+++ b/Engine/Source/Editor/MovieSceneTools/Private/MovieSceneToolsPrivatePCH.h
@@ -8,6 +8,7 @@
 #include "ScopedTransaction.h"
 #include "ISequencerSection.h"
 #include "ISectionLayoutBuilder.h"
+#include "KeyParams.h"
 #include "MovieSceneSequenceInstance.h"
 #include "MovieSceneTrackEditor.h"
 #include "IMovieSceneTools.h"
diff --git a/Engine/Source/Editor/MovieSceneTools/Private/Sections/ColorPropertySection.cpp b/Engine/Source/Editor/MovieSceneTools/Private/Sections/ColorPropertySection.cpp
--- a/Engine/Source/Editor/MovieSceneTools/Private/Sections/ColorPropertySection.cpp
+++ b/Engine/Source/Editor/MovieSceneTools/Private/Sections/ColorPropertySection.cpp
@@ -169,6 +169,12 @@ void FColorPropertySection::ConsolidateColorCurves( TArray< TKeyValuePair<float,
 		}
 	}
 
+	// Enforce at least one key for the default value
+	if (TimesWithKeys.Num() == 0)
+	{
+		TimesWithKeys.Add(0);
+	}
+
 	// @todo Sequencer Optimize - This another O(n^2) loop, since Eval is O(n)!
 	for ( int32 i = 0; i < TimesWithKeys.Num(); ++i )
 	{
diff --git a/Engine/Source/Editor/MovieSceneTools/Private/TrackEditors/AttachTrackEditor.cpp b/Engine/Source/Editor/MovieSceneTools/Private/TrackEditors/AttachTrackEditor.cpp
--- a/Engine/Source/Editor/MovieSceneTools/Private/TrackEditors/AttachTrackEditor.cpp
+++ b/Engine/Source/Editor/MovieSceneTools/Private/TrackEditors/AttachTrackEditor.cpp
@@ -290,7 +290,7 @@ void F3DAttachTrackEditor::AddAttach(FGuid ObjectGuid, UObject* AdditionalAsset)
 	TArray<UObject*> OutObjects;
 	GetSequencer()->GetRuntimeObjects( GetSequencer()->GetFocusedMovieSceneSequenceInstance(), ObjectGuid, OutObjects);
 
-	AnimatablePropertyChanged( UMovieScene3DAttachTrack::StaticClass(), false,
+	AnimatablePropertyChanged( UMovieScene3DAttachTrack::StaticClass(),
 		FOnKeyProperty::CreateRaw( this, &F3DAttachTrackEditor::AddKeyInternal, OutObjects, AdditionalAsset) );
 }
 
diff --git a/Engine/Source/Editor/MovieSceneTools/Private/TrackEditors/AudioTrackEditor.cpp b/Engine/Source/Editor/MovieSceneTools/Private/TrackEditors/AudioTrackEditor.cpp
--- a/Engine/Source/Editor/MovieSceneTools/Private/TrackEditors/AudioTrackEditor.cpp
+++ b/Engine/Source/Editor/MovieSceneTools/Private/TrackEditors/AudioTrackEditor.cpp
@@ -478,12 +478,12 @@ bool FAudioTrackEditor::HandleAssetAdded(UObject* Asset, const FGuid& TargetObje
 			TArray<UObject*> OutObjects;
 			GetSequencer()->GetRuntimeObjects(GetSequencer()->GetFocusedMovieSceneSequenceInstance(), TargetObjectGuid, OutObjects);
 
-			AnimatablePropertyChanged(UMovieSceneAudioTrack::StaticClass(), false,
+			AnimatablePropertyChanged(UMovieSceneAudioTrack::StaticClass(),
 				FOnKeyProperty::CreateRaw(this, &FAudioTrackEditor::AddNewAttachedSound, Sound, OutObjects));
 		}
 		else
 		{
-			AnimatablePropertyChanged(UMovieSceneAudioTrack::StaticClass(), false,
+			AnimatablePropertyChanged(UMovieSceneAudioTrack::StaticClass(),
 				FOnKeyProperty::CreateRaw(this, &FAudioTrackEditor::AddNewMasterSound, Sound));
 		}
 
diff --git a/Engine/Source/Editor/MovieSceneTools/Private/TrackEditors/ParticleTrackEditor.cpp b/Engine/Source/Editor/MovieSceneTools/Private/TrackEditors/ParticleTrackEditor.cpp
--- a/Engine/Source/Editor/MovieSceneTools/Private/TrackEditors/ParticleTrackEditor.cpp
+++ b/Engine/Source/Editor/MovieSceneTools/Private/TrackEditors/ParticleTrackEditor.cpp
@@ -161,7 +161,7 @@ void FParticleTrackEditor::AddParticleKey(const FGuid ObjectGuid, bool bTrigger)
 	TArray<UObject*> OutObjects;
 	GetSequencer()->GetRuntimeObjects( GetSequencer()->GetFocusedMovieSceneSequenceInstance(), ObjectGuid, OutObjects );
 
-	AnimatablePropertyChanged( UMovieSceneParticleTrack::StaticClass(), false,
+	AnimatablePropertyChanged( UMovieSceneParticleTrack::StaticClass(), 
 		FOnKeyProperty::CreateRaw( this, &FParticleTrackEditor::AddKeyInternal, OutObjects, bTrigger) );
 }
 
diff --git a/Engine/Source/Editor/MovieSceneTools/Private/TrackEditors/PathTrackEditor.cpp b/Engine/Source/Editor/MovieSceneTools/Private/TrackEditors/PathTrackEditor.cpp
--- a/Engine/Source/Editor/MovieSceneTools/Private/TrackEditors/PathTrackEditor.cpp
+++ b/Engine/Source/Editor/MovieSceneTools/Private/TrackEditors/PathTrackEditor.cpp
@@ -271,7 +271,7 @@ void F3DPathTrackEditor::AddPath(FGuid ObjectGuid, UObject* AdditionalAsset)
 	TArray<UObject*> OutObjects;
 	GetSequencer()->GetRuntimeObjects( GetSequencer()->GetFocusedMovieSceneSequenceInstance(), ObjectGuid, OutObjects);
 
-	AnimatablePropertyChanged( UMovieScene3DPathTrack::StaticClass(), false,
+	AnimatablePropertyChanged( UMovieScene3DPathTrack::StaticClass(),
 		FOnKeyProperty::CreateRaw( this, &F3DPathTrackEditor::AddKeyInternal, OutObjects, AdditionalAsset) );
 }
 
diff --git a/Engine/Source/Editor/MovieSceneTools/Private/TrackEditors/PropertyTrackEditors/ColorPropertyTrackEditor.cpp b/Engine/Source/Editor/MovieSceneTools/Private/TrackEditors/PropertyTrackEditors/ColorPropertyTrackEditor.cpp
--- a/Engine/Source/Editor/MovieSceneTools/Private/TrackEditors/PropertyTrackEditors/ColorPropertyTrackEditor.cpp
+++ b/Engine/Source/Editor/MovieSceneTools/Private/TrackEditors/PropertyTrackEditors/ColorPropertyTrackEditor.cpp
@@ -51,7 +51,6 @@ bool FColorPropertyTrackEditor::TryGenerateKeyFromPropertyChanged( const FProper
 
 	OutKey.Value = ColorValue;
 	OutKey.CurveName = PropertyChangedParams.StructPropertyNameToKey;
-	OutKey.bAddKeyEvenIfUnchanged = !PropertyChangedParams.bRequireAutoKey;
 	OutKey.bIsSlateColor = bIsSlateColor;
 
 	return true;
diff --git a/Engine/Source/Editor/MovieSceneTools/Private/TrackEditors/PropertyTrackEditors/VectorPropertyTrackEditor.cpp b/Engine/Source/Editor/MovieSceneTools/Private/TrackEditors/PropertyTrackEditors/VectorPropertyTrackEditor.cpp
--- a/Engine/Source/Editor/MovieSceneTools/Private/TrackEditors/PropertyTrackEditors/VectorPropertyTrackEditor.cpp
+++ b/Engine/Source/Editor/MovieSceneTools/Private/TrackEditors/PropertyTrackEditors/VectorPropertyTrackEditor.cpp
@@ -34,15 +34,15 @@ bool FVectorPropertyTrackEditor::TryGenerateKeyFromPropertyChanged( const FPrope
 	// Get the vector value from the property
 	if ( bIsVector2D )
 	{
-		OutKey = FVectorKey(*PropertyChangedParams.GetPropertyValue<FVector2D>(), PropertyChangedParams.StructPropertyNameToKey, !PropertyChangedParams.bRequireAutoKey);
+		OutKey = FVectorKey(*PropertyChangedParams.GetPropertyValue<FVector2D>(), PropertyChangedParams.StructPropertyNameToKey );
 	}
 	else if ( bIsVector )
 	{
-		OutKey = FVectorKey( *PropertyChangedParams.GetPropertyValue<FVector>(), PropertyChangedParams.StructPropertyNameToKey, !PropertyChangedParams.bRequireAutoKey );
+		OutKey = FVectorKey( *PropertyChangedParams.GetPropertyValue<FVector>(), PropertyChangedParams.StructPropertyNameToKey );
 	}
 	else // if ( bIsVector4 )
 	{
-		OutKey = FVectorKey( *PropertyChangedParams.GetPropertyValue<FVector4>(), PropertyChangedParams.StructPropertyNameToKey, !PropertyChangedParams.bRequireAutoKey );
+		OutKey = FVectorKey( *PropertyChangedParams.GetPropertyValue<FVector4>(), PropertyChangedParams.StructPropertyNameToKey );
 	}
 	return true;
 }
diff --git a/Engine/Source/Editor/MovieSceneTools/Private/TrackEditors/ShotTrackEditor.cpp b/Engine/Source/Editor/MovieSceneTools/Private/TrackEditors/ShotTrackEditor.cpp
--- a/Engine/Source/Editor/MovieSceneTools/Private/TrackEditors/ShotTrackEditor.cpp
+++ b/Engine/Source/Editor/MovieSceneTools/Private/TrackEditors/ShotTrackEditor.cpp
@@ -558,7 +558,7 @@ void FShotTrackEditor::AddKey(const FGuid& ObjectGuid, UObject* AdditionalAsset)
 	bool bValidKey = OutObjects.Num() == 1 && OutObjects[0]->IsA<ACameraActor>();
 	if (bValidKey)
 	{
-		AnimatablePropertyChanged( UMovieSceneShotTrack::StaticClass(), false, FOnKeyProperty::CreateRaw( this, &FShotTrackEditor::AddKeyInternal, ObjectGuid ) );
+		AnimatablePropertyChanged( UMovieSceneShotTrack::StaticClass(), FOnKeyProperty::CreateRaw( this, &FShotTrackEditor::AddKeyInternal, ObjectGuid ) );
 	}
 }
 
diff --git a/Engine/Source/Editor/MovieSceneTools/Private/TrackEditors/SkeletalAnimationTrackEditor.cpp b/Engine/Source/Editor/MovieSceneTools/Private/TrackEditors/SkeletalAnimationTrackEditor.cpp
--- a/Engine/Source/Editor/MovieSceneTools/Private/TrackEditors/SkeletalAnimationTrackEditor.cpp
+++ b/Engine/Source/Editor/MovieSceneTools/Private/TrackEditors/SkeletalAnimationTrackEditor.cpp
@@ -162,7 +162,7 @@ void FSkeletalAnimationTrackEditor::AddKey(const FGuid& ObjectGuid, UObject* Add
 		TArray<UObject*> OutObjects;
 		GetSequencer()->GetRuntimeObjects( GetSequencer()->GetFocusedMovieSceneSequenceInstance(), ObjectGuid, OutObjects);
 
-		AnimatablePropertyChanged( UMovieSceneSkeletalAnimationTrack::StaticClass(), false,
+		AnimatablePropertyChanged( UMovieSceneSkeletalAnimationTrack::StaticClass(),
 			FOnKeyProperty::CreateRaw( this, &FSkeletalAnimationTrackEditor::AddKeyInternal, OutObjects, AnimSequence) );
 	}
 }
@@ -182,7 +182,7 @@ bool FSkeletalAnimationTrackEditor::HandleAssetAdded(UObject* Asset, const FGuid
 				TArray<UObject*> OutObjects;
 				GetSequencer()->GetRuntimeObjects(GetSequencer()->GetFocusedMovieSceneSequenceInstance(), TargetObjectGuid, OutObjects);
 
-				AnimatablePropertyChanged(UMovieSceneSkeletalAnimationTrack::StaticClass(), false,
+				AnimatablePropertyChanged(UMovieSceneSkeletalAnimationTrack::StaticClass(),
 					FOnKeyProperty::CreateRaw(this, &FSkeletalAnimationTrackEditor::AddKeyInternal, OutObjects, AnimSequence));
 
 				return true;
diff --git a/Engine/Source/Editor/MovieSceneTools/Private/TrackEditors/TransformTrackEditor.cpp b/Engine/Source/Editor/MovieSceneTools/Private/TrackEditors/TransformTrackEditor.cpp
--- a/Engine/Source/Editor/MovieSceneTools/Private/TrackEditors/TransformTrackEditor.cpp
+++ b/Engine/Source/Editor/MovieSceneTools/Private/TrackEditors/TransformTrackEditor.cpp
@@ -224,51 +224,56 @@ void F3DTransformTrackEditor::OnTransformChanged( UObject& InObject )
 	const bool bCreateHandleIfMissing = false;
 
 	const TSharedPtr<ISequencer> Sequencer = GetSequencer();
-	if (Sequencer->GetAutoKeyEnabled())
-	{
-		USceneComponent* SceneComponentThatChanged = NULL;
 
-		// The runtime binding
-		FGuid ObjectHandle;
+	USceneComponent* SceneComponentThatChanged = NULL;
 
-		AActor* Actor = Cast<AActor>( &InObject );
-		if( Actor && Actor->GetRootComponent() )
-		{
-			// Get a handle bound to the key/section we are adding so we know what objects to change during playback
-			ObjectHandle = Sequencer->GetHandleToObject( Actor, bCreateHandleIfMissing );
-			SceneComponentThatChanged = Actor->GetRootComponent();
-		}
+	// The runtime binding
+	FGuid ObjectHandle;
 
-		else
-		{
-			// If the object wasn't an actor attempt to get it directly as a scene component 
-			SceneComponentThatChanged = Cast<USceneComponent>( &InObject );
-			if( SceneComponentThatChanged )
-			{
-				ObjectHandle = Sequencer->GetHandleToObject( SceneComponentThatChanged, bCreateHandleIfMissing );
-			}
+	AActor* Actor = Cast<AActor>( &InObject );
+	if( Actor && Actor->GetRootComponent() )
+	{
+		// Get a handle bound to the key/section we are adding so we know what objects to change during playback
+		ObjectHandle = Sequencer->GetHandleToObject( Actor, bCreateHandleIfMissing );
+		SceneComponentThatChanged = Actor->GetRootComponent();
+	}
 
+	else
+	{
+		// If the object wasn't an actor attempt to get it directly as a scene component 
+		SceneComponentThatChanged = Cast<USceneComponent>( &InObject );
+		if( SceneComponentThatChanged )
+		{
+			ObjectHandle = Sequencer->GetHandleToObject( SceneComponentThatChanged, bCreateHandleIfMissing );
 		}
 
-		if( SceneComponentThatChanged && ObjectHandle.IsValid() )
+	}
+
+	if( SceneComponentThatChanged && ObjectHandle.IsValid() )
+	{
+		FName Transform("Transform");
+		if (Sequencer->GetFocusedMovieSceneSequence()->GetMovieScene()->FindTrack(UMovieScene3DTransformTrack::StaticClass(), ObjectHandle, Transform))
 		{
-			FName Transform("Transform");
-			if (Sequencer->GetFocusedMovieSceneSequence()->GetMovieScene()->FindTrack(UMovieScene3DTransformTrack::StaticClass(), ObjectHandle, Transform))
-			{
-				// Find an existing transform if possible.  If one exists we will compare against the new one to decide what components of the transform need keys
-				FTransformData ExistingTransform = ObjectToExistingTransform.FindRef( &InObject );
+			// Find an existing transform if possible.  If one exists we will compare against the new one to decide what components of the transform need keys
+			FTransformData ExistingTransform = ObjectToExistingTransform.FindRef( &InObject );
+
+			// Remove it from the list of cached transforms. 
+			// @todo sequencer livecapture: This can be made much for efficient by not removing cached state during live capture situation
+			ObjectToExistingTransform.Remove( &InObject );
 
-				// Remove it from the list of cached transforms. 
-				// @todo sequencer livecapture: This can be made much for efficient by not removing cached state during live capture situation
-				ObjectToExistingTransform.Remove( &InObject );
+			// Build new transform data
+			FTransformData NewTransformData( SceneComponentThatChanged );
 
-				// Build new transform data
-				FTransformData NewTransformData( SceneComponentThatChanged );
+			FTransformDataPair TransformPair(NewTransformData, ExistingTransform);
 
-				FTransformDataPair TransformPair(NewTransformData, ExistingTransform);
+			FKeyParams KeyParams;
+			KeyParams.bAutoKeying = Sequencer->GetAutoKeyEnabled();
+			KeyParams.bAddKeyEvenIfUnchanged = false;
 
-				AnimatablePropertyChanged(UMovieScene3DTransformTrack::StaticClass(), true,
-					FOnKeyProperty::CreateRaw(this, &F3DTransformTrackEditor::OnTransformChangedInternals, &InObject, ObjectHandle, TransformPair, true, false, F3DTransformTrackKey::Key_All));
+			if (CanKeyProperty(FCanKeyProperty::CreateRaw(this, &F3DTransformTrackEditor::CanKeyPropertyInternal, ObjectHandle, TransformPair, KeyParams)))
+			{
+				AnimatablePropertyChanged(UMovieScene3DTransformTrack::StaticClass(),
+					FOnKeyProperty::CreateRaw(this, &F3DTransformTrackEditor::OnTransformChangedInternals, &InObject, ObjectHandle, TransformPair, KeyParams, F3DTransformTrackKey::Key_All));
 			}
 		}
 	}
@@ -279,7 +284,7 @@ void F3DTransformTrackEditor::AddKey(const FGuid& ObjectGuid, UObject* Additiona
 	AddKeyInternal(ObjectGuid, AdditionalAsset);
 }
 
-void F3DTransformTrackEditor::AddKeyInternal(const FGuid& ObjectGuid, UObject* AdditionalAsset, bool bForceKey, F3DTransformTrackKey::Type KeyType)
+void F3DTransformTrackEditor::AddKeyInternal(const FGuid& ObjectGuid, UObject* AdditionalAsset, const bool bAddKeyEvenIfUnchanged, F3DTransformTrackKey::Type KeyType)
 {
 	TArray<UObject*> OutObjects;
 	GetSequencer()->GetRuntimeObjects( GetSequencer()->GetFocusedMovieSceneSequenceInstance(), ObjectGuid, OutObjects);
@@ -312,8 +317,11 @@ void F3DTransformTrackEditor::AddKeyInternal(const FGuid& ObjectGuid, UObject* A
 
 			FTransformDataPair TransformPair(CurrentTransform, FTransformData());
 
-			AnimatablePropertyChanged(UMovieScene3DTransformTrack::StaticClass(), false,
-				FOnKeyProperty::CreateRaw(this, &F3DTransformTrackEditor::OnTransformChangedInternals, Object, ObjectHandle, TransformPair, false, bForceKey, KeyType));
+			FKeyParams KeyParams;
+			KeyParams.bAddKeyEvenIfUnchanged = bAddKeyEvenIfUnchanged;
+
+			AnimatablePropertyChanged(UMovieScene3DTransformTrack::StaticClass(),
+				FOnKeyProperty::CreateRaw(this, &F3DTransformTrackEditor::OnTransformChangedInternals, Object, ObjectHandle, TransformPair, KeyParams, KeyType));
 		}
 	}
 }
@@ -441,6 +449,8 @@ void F3DTransformTrackEditor::BuildObjectBindingTrackMenu(FMenuBuilder& MenuBuil
 
 void F3DTransformTrackEditor::AddTransform(FGuid ObjectBinding)
 {
+	FScopedTransaction DeleteSectionTransaction( NSLOCTEXT("Sequencer", "AddTransformTrack_Transaction", "Add Transform Track") );
+
 	FName Transform("Transform");
 	UMovieSceneTrack* Track = AddTrack(GetSequencer()->GetFocusedMovieSceneSequence()->GetMovieScene(), ObjectBinding, UMovieScene3DTransformTrack::StaticClass(), Transform);
 	UMovieScene3DTransformTrack* TransformTrack = CastChecked<UMovieScene3DTransformTrack>( Track );
@@ -448,11 +458,13 @@ void F3DTransformTrackEditor::AddTransform(FGuid ObjectBinding)
 	TransformTrack->SetPropertyNameAndPath( Transform, Transform.ToString() );
 
 	UMovieScene3DTransformSection* TransformSection = CastChecked<UMovieScene3DTransformSection>(TransformTrack->CreateNewSection());
-	TransformSection->SetStartTime(GetSequencer()->GetViewRange().GetLowerBoundValue());
-	TransformSection->SetEndTime(GetSequencer()->GetViewRange().GetUpperBoundValue());
+	TransformSection->SetStartTime(GetSequencer()->GetGlobalTime());
+	TransformSection->SetEndTime(GetSequencer()->GetGlobalTime());
 	TransformTrack->AddSection(TransformSection);
 	TransformTrack->SetAsShowable();
 	
+	AddKey(ObjectBinding, nullptr);
+
 	NotifyMovieSceneDataChanged();
 }
 
@@ -466,7 +478,28 @@ bool F3DTransformTrackEditor::CanAddTransform(FGuid ObjectBinding) const
 	return true;
 }
 
-void F3DTransformTrackEditor::OnTransformChangedInternals(float KeyTime, UObject* InObject, FGuid ObjectHandle, FTransformDataPair TransformPair, bool bAutoKeying, bool bForceKey, F3DTransformTrackKey::Type KeyType)
+FTransformKey F3DTransformTrackEditor::GetTransformKey(const UMovieScene3DTransformTrack* TransformTrack, float KeyTime, FTransformDataPair TransformPair, FKeyParams KeyParams) const
+{
+	if (!TransformPair.LastTransformData.IsValid())
+	{
+		TransformPair.LastTransformData.bValid = TransformTrack->Eval(KeyTime, KeyTime, TransformPair.LastTransformData.Translation, TransformPair.LastTransformData.Rotation, TransformPair.LastTransformData.Scale);
+	}
+
+	// If forcing the keyframe, set the last data invalid so that the transform keyframe will be set regardless of the current time values being equal.
+	if ( KeyParams.bAddKeyEvenIfUnchanged || IsAllowedKeyAll())
+	{
+		TransformPair.LastTransformData.bValid = false;
+	}
+
+	FTransformKey TransformKey = FTransformKey(KeyTime, TransformPair.TransformData, TransformPair.LastTransformData);
+	TransformKey.KeyParams.bAutoKeying = KeyParams.bAutoKeying;
+	TransformKey.KeyParams.bAddKeyEvenIfUnchanged = KeyParams.bAddKeyEvenIfUnchanged;
+	TransformKey.KeyParams.bCreateTrackIfMissing = KeyParams.bCreateTrackIfMissing;
+	TransformKey.KeyParams.bCreateHandleIfMissing = KeyParams.bCreateHandleIfMissing;
+	return TransformKey;
+}
+
+void F3DTransformTrackEditor::OnTransformChangedInternals(float KeyTime, UObject* InObject, FGuid ObjectHandle, FTransformDataPair TransformPair, FKeyParams KeyParams, F3DTransformTrackKey::Type KeyType)
 {
 	// Only unwind rotation if we're generating keys while recording (scene is actively playing back)
 	const bool bUnwindRotation = GetSequencer()->IsRecordingLive();
@@ -479,28 +512,38 @@ void F3DTransformTrackEditor::OnTransformChangedInternals(float KeyTime, UObject
 		// Transform name and path are the same
 		TransformTrack->SetPropertyNameAndPath( Transform, Transform.ToString() );
 	
-		if (!TransformPair.LastTransformData.IsValid())
-		{
-			TArray<bool> bHasTranslationKeys;
-			TArray<bool> bHasRotationKeys;
-			TArray<bool> bHasScaleKeys;			
-			TransformPair.LastTransformData.bValid = TransformTrack->Eval(KeyTime, KeyTime, TransformPair.LastTransformData.Translation, TransformPair.LastTransformData.Rotation, TransformPair.LastTransformData.Scale, bHasTranslationKeys, bHasRotationKeys, bHasScaleKeys);
-		}
+		FTransformKey TransformKey = GetTransformKey(TransformTrack, KeyTime, TransformPair, KeyParams);
 
-		// If forcing the keyframe, set the last data invalid so that the transform keyframe will be set regardless of the current time values being equal.
-		if ( bForceKey || IsAllowedKeyAll() )
+		bool bSuccessfulAdd = TransformTrack->AddKeyToSection( ObjectHandle, TransformKey, bUnwindRotation, KeyType );
+		if (bSuccessfulAdd)
 		{
-			TransformPair.LastTransformData.bValid = false;
+			TransformTrack->SetAsShowable();
 		}
+	}
+}
 
-		FTransformKey TransformKey = FTransformKey(KeyTime, TransformPair.TransformData, TransformPair.LastTransformData);
+bool F3DTransformTrackEditor::CanKeyPropertyInternal(float KeyTime, FGuid ObjectHandle, FTransformDataPair TransformPair, FKeyParams KeyParams) const
+{
+	FName Transform("Transform");
+	if (ObjectHandle.IsValid())
+	{
+		const UMovieScene* MovieScene = GetSequencer()->GetFocusedMovieSceneSequence()->GetMovieScene();
 
-		bool bSuccessfulAdd = TransformTrack->AddKeyToSection( ObjectHandle, TransformKey, bUnwindRotation, KeyType );
-		if (bSuccessfulAdd)
+		const UMovieSceneTrack* Track = MovieScene->FindTrack( UMovieScene3DTransformTrack::StaticClass(), ObjectHandle, Transform );
+
+		if (Track)
 		{
-			TransformTrack->SetAsShowable();
+			const UMovieScene3DTransformTrack* TransformTrack = CastChecked<UMovieScene3DTransformTrack>( Track );
+
+			FTransformKey TransformKey = GetTransformKey(TransformTrack, KeyTime, TransformPair, KeyParams);
+
+			if (TransformTrack->CanKeyTrack(TransformKey))
+			{
+				return true;
+			}
 		}
 	}
+	return false;
 }
 
 ECheckBoxState F3DTransformTrackEditor::IsCameraLocked(TWeakObjectPtr<ACameraActor> CameraActor) const
diff --git a/Engine/Source/Editor/MovieSceneTools/Private/TrackEditors/TransformTrackEditor.h b/Engine/Source/Editor/MovieSceneTools/Private/TrackEditors/TransformTrackEditor.h
--- a/Engine/Source/Editor/MovieSceneTools/Private/TrackEditors/TransformTrackEditor.h
+++ b/Engine/Source/Editor/MovieSceneTools/Private/TrackEditors/TransformTrackEditor.h
@@ -44,7 +44,7 @@ private:
 	bool CanAddTransform(FGuid ObjectGuid) const;
 
 	/** Custom add key implementation */
-	void AddKeyInternal(const FGuid& ObjectGuid, UObject* AdditionalAsset = NULL, bool bForceKey = false, F3DTransformTrackKey::Type KeyType = F3DTransformTrackKey::Key_All);
+	void AddKeyInternal(const FGuid& ObjectGuid, UObject* AdditionalAsset = NULL, const bool bAddKeyEvenIfUnchanged = false, F3DTransformTrackKey::Type KeyType = F3DTransformTrackKey::Key_All);
 
 	/**
 	 * Called before an actor or component transform changes
@@ -60,8 +60,14 @@ private:
 	 */
 	void OnTransformChanged( UObject& InObject );
 
+	/** Get transform key */
+	FTransformKey GetTransformKey(const UMovieScene3DTransformTrack* TransformTrack, float KeyTime, struct FTransformDataPair TransformPair, FKeyParams KeyParams) const;
+
 	/** Delegate for animatable property changed in OnTransformChanged */
-	void OnTransformChangedInternals(float KeyTime, UObject* InObject, FGuid ObjectHandle, struct FTransformDataPair TransformPair, bool bAutoKeying, bool bForceKey, F3DTransformTrackKey::Type KeyType);
+	void OnTransformChangedInternals(float KeyTime, UObject* InObject, FGuid ObjectHandle, struct FTransformDataPair TransformPair, FKeyParams KeyParams, F3DTransformTrackKey::Type KeyType);
+
+	/** Delegate to determine whether the property can be keyed */
+	bool CanKeyPropertyInternal(float KeyTime, FGuid ObjectHandle, struct FTransformDataPair TransformPair, FKeyParams KeyParams) const;
 
 	/** Delegate for camera button lock state */
 	ECheckBoxState IsCameraLocked(TWeakObjectPtr<ACameraActor> CameraActor) const; 
diff --git a/Engine/Source/Editor/MovieSceneTools/Public/PropertyTrackEditor.h b/Engine/Source/Editor/MovieSceneTools/Public/PropertyTrackEditor.h
--- a/Engine/Source/Editor/MovieSceneTools/Public/PropertyTrackEditor.h
+++ b/Engine/Source/Editor/MovieSceneTools/Public/PropertyTrackEditor.h
@@ -161,7 +161,7 @@ private:
 	*/
 	virtual void OnAnimatedPropertyChanged( const FPropertyChangedParams& PropertyChangedParams )
 	{
-		bool bCreateHandleIfMissing = !PropertyChangedParams.bRequireAutoKey;
+		bool bCreateHandleIfMissing = PropertyChangedParams.KeyParams.bCreateHandleIfMissing;
 
 		for ( UObject* Object : PropertyChangedParams.ObjectsThatChanged )
 		{
@@ -173,7 +173,7 @@ private:
 				TSubclassOf<UMovieSceneTrack> SequencerTrackClass = TrackType::StaticClass();
 				GetPropertyAndTrackClass(PropertyChangedParams.PropertyPath.Last(), PropertyName, SequencerTrackClass);
 
-				bool bCreateTrackIfMissing = !PropertyChangedParams.bRequireAutoKey;
+				bool bCreateTrackIfMissing = PropertyChangedParams.KeyParams.bCreateTrackIfMissing;
 
 				if (bCreateTrackIfMissing || GetSequencer()->GetFocusedMovieSceneSequence()->GetMovieScene()->FindTrack(SequencerTrackClass, ObjectHandle, PropertyName))
 				{
@@ -181,7 +181,7 @@ private:
 					KeyType Key;
 					if ( TryGenerateKeyFromPropertyChanged( PropertyChangedParams, Key ) )
 					{
-						AnimatablePropertyChanged( TrackType::StaticClass(), PropertyChangedParams.bRequireAutoKey,
+						AnimatablePropertyChanged( TrackType::StaticClass(),
 							FOnKeyProperty::CreateRaw( this, &FPropertyTrackEditor::OnKeyProperty, PropertyChangedParams, Key ) );
 					}
 				}
@@ -192,7 +192,7 @@ private:
 	/** Adds a key based on a property change. */
 	void OnKeyProperty( float KeyTime, FPropertyChangedParams PropertyChangedParams, KeyType Key )
 	{
-		bool bCreateHandleIfMissing = !PropertyChangedParams.bRequireAutoKey;
+		bool bCreateHandleIfMissing = PropertyChangedParams.KeyParams.bCreateHandleIfMissing;
 
 		for ( UObject* Object : PropertyChangedParams.ObjectsThatChanged )
 		{
@@ -203,7 +203,7 @@ private:
 				TSubclassOf<UMovieSceneTrack> SequencerTrackClass = TrackType::StaticClass();
 				GetPropertyAndTrackClass(PropertyChangedParams.PropertyPath.Last(), PropertyName, SequencerTrackClass);
 
-				bool bCreateTrackIfMissing = !PropertyChangedParams.bRequireAutoKey;
+				bool bCreateTrackIfMissing = PropertyChangedParams.KeyParams.bCreateTrackIfMissing;
 
 				UMovieSceneTrack* Track = GetTrackForObject( ObjectHandle, SequencerTrackClass, PropertyName, bCreateTrackIfMissing );
 				if ( Track )
@@ -213,7 +213,7 @@ private:
 					TypedTrack->SetPropertyNameAndPath( PropertyName, PropertyChangedParams.GetPropertyPathString() );
 					// Find or add a new section at the auto-key time and changing the property same property
 					// AddKeyToSection is not actually a virtual, it's redefined in each class with a different type
-					bool bSuccessfulAdd = TypedTrack->AddKeyToSection( KeyTime, Key );
+					bool bSuccessfulAdd = TypedTrack->AddKeyToSection( KeyTime, Key, PropertyChangedParams.KeyParams );
 					if ( bSuccessfulAdd )
 					{
 						TypedTrack->SetAsShowable();
diff --git a/Engine/Source/Editor/Sequencer/Private/SequencerDisplayNode.cpp b/Engine/Source/Editor/Sequencer/Private/SequencerDisplayNode.cpp
--- a/Engine/Source/Editor/Sequencer/Private/SequencerDisplayNode.cpp
+++ b/Engine/Source/Editor/Sequencer/Private/SequencerDisplayNode.cpp
@@ -800,7 +800,13 @@ void FObjectBindingNode::AddTrackForProperty(TArray<UProperty*> PropertyPath)
 			KeyableBoundObjects.Add(BoundObject);
 		}
 	}
-	Sequencer.KeyProperty(FKeyPropertyParams(KeyableBoundObjects, PropertyPath));
+
+	FKeyPropertyParams KeyPropertyParams(KeyableBoundObjects, PropertyPath);
+	KeyPropertyParams.KeyParams.bCreateTrackIfMissing = true;
+	KeyPropertyParams.KeyParams.bCreateHandleIfMissing = false;
+	KeyPropertyParams.KeyParams.bAddKeyEvenIfUnchanged = false;
+
+	Sequencer.KeyProperty(KeyPropertyParams);
 }
 
 float FSectionCategoryNode::GetNodeHeight() const 
diff --git a/Engine/Source/Editor/Sequencer/Private/SequencerObjectChangeListener.cpp b/Engine/Source/Editor/Sequencer/Private/SequencerObjectChangeListener.cpp
--- a/Engine/Source/Editor/Sequencer/Private/SequencerObjectChangeListener.cpp
+++ b/Engine/Source/Editor/Sequencer/Private/SequencerObjectChangeListener.cpp
@@ -24,12 +24,12 @@ FSequencerObjectChangeListener::~FSequencerObjectChangeListener()
 	GEditor->OnActorMoved().RemoveAll( this );
 }
 
-void FSequencerObjectChangeListener::OnPropertyChanged(const TArray<UObject*>& ChangedObjects, const IPropertyHandle& PropertyHandle, bool bRequireAutoKey) const
+void FSequencerObjectChangeListener::OnPropertyChanged(const TArray<UObject*>& ChangedObjects, const IPropertyHandle& PropertyHandle) const
 {
-	BroadcastPropertyChanged(FKeyPropertyParams(ChangedObjects, PropertyHandle), bRequireAutoKey);
+	BroadcastPropertyChanged(FKeyPropertyParams(ChangedObjects, PropertyHandle));
 }
 
-void FSequencerObjectChangeListener::BroadcastPropertyChanged( FKeyPropertyParams KeyPropertyParams, bool bRequireAutoKey ) const
+void FSequencerObjectChangeListener::BroadcastPropertyChanged( FKeyPropertyParams KeyPropertyParams ) const
 {
 	// Filter out objects that actually have the property path that will be keyable. 
 	// Otherwise, this might try to key objects that don't have the requested property.
@@ -57,7 +57,10 @@ void FSequencerObjectChangeListener::BroadcastPropertyChanged( FKeyPropertyParam
 	}
 
 	FPropertyChangedParams Params;
-	Params.bRequireAutoKey = bRequireAutoKey;
+	Params.KeyParams.bCreateHandleIfMissing = KeyPropertyParams.KeyParams.bCreateHandleIfMissing;
+	Params.KeyParams.bCreateTrackIfMissing = KeyPropertyParams.KeyParams.bCreateTrackIfMissing;
+	Params.KeyParams.bAddKeyEvenIfUnchanged = KeyPropertyParams.KeyParams.bAddKeyEvenIfUnchanged;
+	Params.KeyParams.bAutoKeying = Sequencer.Pin()->GetAutoKeyEnabled();
 	Params.ObjectsThatChanged = KeyableObjects;
 	Params.StructPropertyNameToKey = NAME_None;
 
@@ -70,10 +73,7 @@ void FSequencerObjectChangeListener::BroadcastPropertyChanged( FKeyPropertyParam
 		FOnAnimatablePropertyChanged Delegate = ClassToPropertyChangedMap.FindRef(ParentStructProperty->Struct->GetFName());
 		if (Delegate.IsBound())
 		{
-			if (!bRequireAutoKey)
-			{
-				Params.StructPropertyNameToKey = KeyPropertyParams.PropertyPath.Last()->GetFName();
-			}
+			Params.StructPropertyNameToKey = KeyPropertyParams.PropertyPath.Last()->GetFName();
 			bFoundAndBroadcastedDelegate = true;
 			Delegate.Broadcast(Params);
 		}
@@ -191,14 +191,13 @@ bool FSequencerObjectChangeListener::CanKeyProperty(FCanKeyPropertyParams CanKey
 
 void FSequencerObjectChangeListener::KeyProperty(FKeyPropertyParams KeyPropertyParams) const
 {
-	const bool bRequireAutoKey = false;
-	BroadcastPropertyChanged(KeyPropertyParams, bRequireAutoKey);
+	BroadcastPropertyChanged(KeyPropertyParams);
 }
 
 void FSequencerObjectChangeListener::OnObjectPreEditChange( UObject* Object, const FEditPropertyChain& PropertyChain )
 {
-	// We only care if we are in auto-key mode and not attempting to change properties of a CDO (which cannot be animated)
-	if( Sequencer.IsValid() && Sequencer.Pin()->GetAutoKeyEnabled() && !Object->HasAnyFlags(RF_ClassDefaultObject) && PropertyChain.GetActiveMemberNode() )
+	// We only care if we are not attempting to change properties of a CDO (which cannot be animated)
+	if( Sequencer.IsValid() && !Object->HasAnyFlags(RF_ClassDefaultObject) && PropertyChain.GetActiveMemberNode() )
 	{
 		// Register with the property editor module that we'd like to know about animated float properties that change
 		FPropertyEditorModule& PropertyEditor = FModuleManager::Get().LoadModuleChecked<FPropertyEditorModule>( "PropertyEditor" );
@@ -214,7 +213,7 @@ void FSequencerObjectChangeListener::OnObjectPreEditChange( UObject* Object, con
 				
 				ActivePropertyChangeListeners.Add( Object, PropertyChangeListener );
 
-				PropertyChangeListener->GetOnPropertyChangedDelegate().AddRaw( this, &FSequencerObjectChangeListener::OnPropertyChanged, true );
+				PropertyChangeListener->GetOnPropertyChangedDelegate().AddRaw( this, &FSequencerObjectChangeListener::OnPropertyChanged );
 
 				FPropertyListenerSettings Settings;
 				// Ignore array and object properties
@@ -239,8 +238,8 @@ void FSequencerObjectChangeListener::OnObjectPostEditChange( UObject* Object, FP
 
 		bool bShouldPropagateChanges = bIsObjectValid;
 
-		// We only care if we are in auto-key mode and not attempting to change properties of a CDO (which cannot be animated)
-		if( Sequencer.IsValid() && Sequencer.Pin()->GetAutoKeyEnabled() && bIsObjectValid && !Object->HasAnyFlags(RF_ClassDefaultObject) )
+		// We only care if we are not attempting to change properties of a CDO (which cannot be animated)
+		if( Sequencer.IsValid() && bIsObjectValid && !Object->HasAnyFlags(RF_ClassDefaultObject) )
 		{
 			TSharedPtr< IPropertyChangeListener > Listener;
 			ActivePropertyChangeListeners.RemoveAndCopyValue( Object, Listener );
@@ -300,7 +299,7 @@ void FSequencerObjectChangeListener::TriggerAllPropertiesChanged(UObject* Object
 
 				PropertyChangeListener = PropertyEditor.CreatePropertyChangeListener();
 				
-				PropertyChangeListener->GetOnPropertyChangedDelegate().AddRaw( this, &FSequencerObjectChangeListener::OnPropertyChanged, false );
+				PropertyChangeListener->GetOnPropertyChangedDelegate().AddRaw( this, &FSequencerObjectChangeListener::OnPropertyChanged );
 
 				FPropertyListenerSettings Settings;
 				// Ignore array and object properties
diff --git a/Engine/Source/Editor/Sequencer/Private/SequencerObjectChangeListener.h b/Engine/Source/Editor/Sequencer/Private/SequencerObjectChangeListener.h
--- a/Engine/Source/Editor/Sequencer/Private/SequencerObjectChangeListener.h
+++ b/Engine/Source/Editor/Sequencer/Private/SequencerObjectChangeListener.h
@@ -50,10 +50,10 @@ private:
 	 *
 	 * @param Object	The object that PostEditChange was called on
 	 */
-	void OnPropertyChanged( const TArray<UObject*>& ChangedObjects, const IPropertyHandle& PropertyHandle, bool bRequireAutoKey ) const;
+	void OnPropertyChanged( const TArray<UObject*>& ChangedObjects, const IPropertyHandle& PropertyHandle ) const;
 
 	/** Broadcasts the property change callback to the appropriate listeners. */
-	void BroadcastPropertyChanged(FKeyPropertyParams KeyPropertyParams, bool bRequireAutoKey) const;
+	void BroadcastPropertyChanged(FKeyPropertyParams KeyPropertyParams) const;
 
 	/**
 	 * @return True if an object is valid for listening to property changes 
diff --git a/Engine/Source/Editor/Sequencer/Public/KeyPropertyParams.h b/Engine/Source/Editor/Sequencer/Public/KeyPropertyParams.h
--- a/Engine/Source/Editor/Sequencer/Public/KeyPropertyParams.h
+++ b/Engine/Source/Editor/Sequencer/Public/KeyPropertyParams.h
@@ -2,6 +2,8 @@
 
 #pragma once
 
+#include "KeyParams.h"
+
 class IPropertyHandle;
 class UClass;
 class UProperty;
@@ -62,6 +64,8 @@ struct SEQUENCER_API FKeyPropertyParams
 	/** An array of UProperty objects which represents a path of properties to get from the root object to the
 	property to be keyed. */
 	TArray<UProperty*> PropertyPath;
+	/** Keyframing params */
+	FKeyParams KeyParams;
 };
 
 /**
@@ -95,6 +99,6 @@ public:
 	/** Represents the FName of an inner property which should be keyed for a struct property.  If all inner 
 	properties should be keyed, this will be FName::None. */
 	FName StructPropertyNameToKey;
-	/** If autokey is required. */
-	bool bRequireAutoKey;
+	/** Keyframing params */
+	FKeyParams KeyParams;
 };
\ No newline at end of file
diff --git a/Engine/Source/Editor/Sequencer/Public/MovieSceneTrackEditor.h b/Engine/Source/Editor/Sequencer/Public/MovieSceneTrackEditor.h
--- a/Engine/Source/Editor/Sequencer/Public/MovieSceneTrackEditor.h
+++ b/Engine/Source/Editor/Sequencer/Public/MovieSceneTrackEditor.h
@@ -11,6 +11,7 @@
 class ISequencerSection;
 
 DECLARE_DELEGATE_OneParam(FOnKeyProperty, float)
+DECLARE_DELEGATE_RetVal_OneParam(bool, FCanKeyProperty, float)
 
 /**
  * Base class for handling key and section drawing and manipulation for a UMovieSceneTrack class
@@ -71,16 +72,25 @@ public:
 		}
 	}
 	
-	void AnimatablePropertyChanged( TSubclassOf<class UMovieSceneTrack> TrackClass, bool bMustBeAutokeying, FOnKeyProperty OnKeyProperty )
+	bool CanKeyProperty( FCanKeyProperty InCanKeyProperty )
+	{
+		UMovieSceneSequence* MovieSceneSequence = GetMovieSceneSequence();
+		check( MovieSceneSequence );
+
+		float KeyTime = GetTimeForKey(MovieSceneSequence);
+
+		return InCanKeyProperty.Execute(KeyTime);
+	}
+	
+	void AnimatablePropertyChanged( TSubclassOf<class UMovieSceneTrack> TrackClass, FOnKeyProperty OnKeyProperty )
 	{
 		check(OnKeyProperty.IsBound());
 
 		// Get the movie scene we want to autokey
 		UMovieSceneSequence* MovieSceneSequence = GetMovieSceneSequence();
 		float KeyTime = GetTimeForKey( MovieSceneSequence );
-		bool bIsKeyingValid = !bMustBeAutokeying || IsAllowedToAutoKey();
 
-		if( bIsKeyingValid )
+		if( !Sequencer.Pin()->IsRecordingLive() )
 		{
 			check( MovieSceneSequence );
 
diff --git a/Engine/Source/Editor/UMGEditor/Private/Animation/MarginTrackEditor.cpp b/Engine/Source/Editor/UMGEditor/Private/Animation/MarginTrackEditor.cpp
--- a/Engine/Source/Editor/UMGEditor/Private/Animation/MarginTrackEditor.cpp
+++ b/Engine/Source/Editor/UMGEditor/Private/Animation/MarginTrackEditor.cpp
@@ -42,7 +42,6 @@ TSharedRef<ISequencerSection> FMarginTrackEditor::MakeSectionInterface( UMovieSc
 
 bool FMarginTrackEditor::TryGenerateKeyFromPropertyChanged( const FPropertyChangedParams& PropertyChangedParams, FMarginKey& OutKey )
 {
-	OutKey.bAddKeyEvenIfUnchanged = !PropertyChangedParams.bRequireAutoKey;
 	OutKey.CurveName = PropertyChangedParams.StructPropertyNameToKey;
 	OutKey.Value = *PropertyChangedParams.GetPropertyValue<FMargin>();
 	return true;
diff --git a/Engine/Source/Editor/UMGEditor/Private/Animation/Sequencer2DTransformTrackEditor.cpp b/Engine/Source/Editor/UMGEditor/Private/Animation/Sequencer2DTransformTrackEditor.cpp
--- a/Engine/Source/Editor/UMGEditor/Private/Animation/Sequencer2DTransformTrackEditor.cpp
+++ b/Engine/Source/Editor/UMGEditor/Private/Animation/Sequencer2DTransformTrackEditor.cpp
@@ -57,7 +57,6 @@ TSharedRef<ISequencerSection> F2DTransformTrackEditor::MakeSectionInterface( UMo
 
 bool F2DTransformTrackEditor::TryGenerateKeyFromPropertyChanged( const FPropertyChangedParams& PropertyChangedParams, F2DTransformKey& OutKey )
 {
-	OutKey.bAddKeyEvenIfUnchanged = !PropertyChangedParams.bRequireAutoKey;
 	OutKey.CurveName = PropertyChangedParams.StructPropertyNameToKey;
 	OutKey.Value = *PropertyChangedParams.GetPropertyValue<FWidgetTransform>();
 	return true;
diff --git a/Engine/Source/Editor/UMGEditor/Private/Animation/UMGDetailKeyframeHandler.cpp b/Engine/Source/Editor/UMGEditor/Private/Animation/UMGDetailKeyframeHandler.cpp
--- a/Engine/Source/Editor/UMGEditor/Private/Animation/UMGDetailKeyframeHandler.cpp
+++ b/Engine/Source/Editor/UMGEditor/Private/Animation/UMGDetailKeyframeHandler.cpp
@@ -26,5 +26,10 @@ void FUMGDetailKeyframeHandler::OnKeyPropertyClicked(const IPropertyHandle& Keye
 	TArray<UObject*> Objects;
 	KeyedPropertyHandle.GetOuterObjects( Objects );
 
-	BlueprintEditor.Pin()->GetSequencer()->KeyProperty(FKeyPropertyParams(Objects, KeyedPropertyHandle));
+	FKeyPropertyParams KeyPropertyParams(Objects, KeyedPropertyHandle);
+	KeyPropertyParams.KeyParams.bCreateHandleIfMissing = true;
+	KeyPropertyParams.KeyParams.bCreateTrackIfMissing = true;
+	KeyPropertyParams.KeyParams.bAddKeyEvenIfUnchanged = true;
+
+	BlueprintEditor.Pin()->GetSequencer()->KeyProperty(KeyPropertyParams);
 }
diff --git a/Engine/Source/Runtime/Engine/Classes/Curves/CurveBase.h b/Engine/Source/Runtime/Engine/Classes/Curves/CurveBase.h
--- a/Engine/Source/Runtime/Engine/Classes/Curves/CurveBase.h
+++ b/Engine/Source/Runtime/Engine/Classes/Curves/CurveBase.h
@@ -286,6 +286,7 @@ public:
 	: FIndexedCurve()
 	, PreInfinityExtrap(RCCE_Constant)
 	, PostInfinityExtrap(RCCE_Constant)
+	, DefaultValue(MAX_flt)
 	{
 	}
 
@@ -329,7 +330,7 @@ public:
 	void DeleteKey(FKeyHandle KeyHandle);
 
 	/** Finds the key at InTime, and updates its value. If it can't find the key, it adds one at that time */
-	FKeyHandle UpdateOrAddKey(float InTime, float InValue);
+	FKeyHandle UpdateOrAddKey(float InTime, float InValue, const bool bUnwindRotation = false);
 
 	/** Move a key to a new time. This may change the index of the key, so the new key index is returned. */
 	FKeyHandle SetKeyTime(FKeyHandle KeyHandle, float NewTime);
@@ -346,6 +347,12 @@ public:
 	/** Returns the value of the specified key */
 	float GetKeyValue(FKeyHandle KeyHandle) const;
 
+	/** Set the default value of the curve */
+	void SetDefaultValue(float InDefaultValue) { DefaultValue = InDefaultValue; }
+
+	/** Get the default value for the curve */
+	float GetDefaultValue() const { return DefaultValue; }
+
 	/** Shifts all keys forwards or backwards in time by an even amount, preserving order */
 	void ShiftCurve(float DeltaTime);
 	void ShiftCurve(float DeltaTime, TSet<FKeyHandle>& KeyHandles);
@@ -382,7 +389,7 @@ public:
 	void RemapTimeValue(float& InTime, float& CycleValueOffset) const;
 
 	/** Evaluate this rich curve at the specified time */
-	float Eval(float InTime, float DefaultValue = 0.0f) const;
+	float Eval(float InTime, float InDefaultValue = 0.0f) const;
 
 	/** Auto set tangents for any 'auto' keys in curve */
 	void AutoSetTangents(float Tension = 0.f);
@@ -404,6 +411,10 @@ public:
 	/** Sorted array of keys */
 	UPROPERTY()
 	TArray<FRichCurveKey> Keys;
+
+	/** Default value */
+	UPROPERTY()
+	float DefaultValue;
 };
 
 //////////////////////////////////////////////////////////////////////////
@@ -586,6 +597,15 @@ struct ENGINE_API FIntegralCurve : public FIndexedCurve
 {
 	GENERATED_USTRUCT_BODY()
 public:
+	FIntegralCurve() 
+	: FIndexedCurve()
+	, DefaultValue(MAX_int32)
+	{
+	}
+
+	virtual ~FIntegralCurve()
+	{
+	}
 	
 	/** Get number of keys in curve. */
 	virtual int32 GetNumKeys() const override;
@@ -594,7 +614,7 @@ public:
 	virtual bool IsKeyHandleValid(FKeyHandle KeyHandle) const override;
 
 	/** Evaluates the value of an array of keys at a time */
-	int32 Evaluate(float Time, int32 DefaultValue = 0) const;
+	int32 Evaluate(float Time, int32 InDefaultValue = 0) const;
 
 	/** Const iterator for the keys, so the indices and handles stay valid */
 	TArray<FIntegralKey>::TConstIterator GetKeyIterator() const;
@@ -618,6 +638,12 @@ public:
 	/** Get the time for the Key with the specified index. */
 	float GetKeyTime(FKeyHandle KeyHandle) const;
 	
+	/** Set the default value for the curve */
+	void SetDefaultValue(int32 InDefaultValue) { DefaultValue = InDefaultValue; }
+
+	/** Get the default value for the curve */
+	int32 GetDefaultValue() const { return DefaultValue; }
+
 	/** Shifts all keys forwards or backwards in time by an even amount, preserving order */
 	void ShiftCurve(float DeltaTime);
 	void ShiftCurve(float DeltaTime, TSet<FKeyHandle>& KeyHandles);
@@ -636,4 +662,8 @@ private:
 	/** The keys, ordered by time */
 	UPROPERTY()
 	TArray<FIntegralKey> Keys;
+
+	/** Default value */
+	UPROPERTY()
+	int32 DefaultValue;
 };
diff --git a/Engine/Source/Runtime/Engine/Private/CurveBase.cpp b/Engine/Source/Runtime/Engine/Private/CurveBase.cpp
--- a/Engine/Source/Runtime/Engine/Private/CurveBase.cpp
+++ b/Engine/Source/Runtime/Engine/Private/CurveBase.cpp
@@ -392,7 +392,7 @@ void FRichCurve::DeleteKey(FKeyHandle InKeyHandle)
 	}
 }
 
-FKeyHandle FRichCurve::UpdateOrAddKey(float InTime, float InValue)
+FKeyHandle FRichCurve::UpdateOrAddKey(float InTime, float InValue, const bool bUnwindRotation)
 {
 	// Search for a key that already exists at the time and if found, update its value
 	for( int32 KeyIndex = 0; KeyIndex < Keys.Num(); ++KeyIndex )
@@ -414,7 +414,7 @@ FKeyHandle FRichCurve::UpdateOrAddKey(float InTime, float InValue)
 	}
 
 	// A key wasnt found, add it now
-	return AddKey( InTime, InValue );
+	return AddKey( InTime, InValue, bUnwindRotation );
 }
 
 FKeyHandle FRichCurve::SetKeyTime( FKeyHandle KeyHandle, float NewTime )
@@ -956,19 +956,20 @@ void FRichCurve::RemapTimeValue(float& InTime, float& CycleValueOffset) const
 	}
 }
 
-float FRichCurve::Eval(float InTime, float DefaultValue) const
+float FRichCurve::Eval(float InTime, float InDefaultValue) const
 {
 	// Remap time if extrapolation is present and compute offset value to use if cycling 
 	float CycleValueOffset = 0;
 	RemapTimeValue(InTime, CycleValueOffset);
 
 	const int32 NumKeys = Keys.Num();
-	float InterpVal = DefaultValue;
+
+	// If the default value hasn't been initialized, use the incoming default value
+	float InterpVal = DefaultValue == MAX_flt ? InDefaultValue : DefaultValue;
 
 	if (NumKeys == 0)
 	{
-		// If no keys in curve, return the Default value we passed in.
-		InterpVal = DefaultValue;
+		// If no keys in curve, return the Default value.
 	} 
 	else if (NumKeys < 2 || (InTime <= Keys[0].Time))
 	{
@@ -1279,12 +1280,14 @@ bool FIntegralCurve::IsKeyHandleValid(FKeyHandle KeyHandle) const
 	return bValid;
 }
 
-int32 FIntegralCurve::Evaluate(float Time, int32 DefaultValue) const
+int32 FIntegralCurve::Evaluate(float Time, int32 InDefaultValue) const
 {
-	int32 ReturnVal = DefaultValue;
+	// If the default value hasn't been initialized, use the incoming default value
+	int32 ReturnVal = DefaultValue == MAX_int32 ? InDefaultValue : DefaultValue;
+
 	if( Keys.Num() == 0 )
 	{
-		ReturnVal = DefaultValue;
+		// If no keys in curve, return the Default value.
 	}
 	else if( Keys.Num() < 2 || Time < Keys[0].Time )
 	{
diff --git a/Engine/Source/Runtime/MovieScene/Private/MovieSceneSection.cpp b/Engine/Source/Runtime/MovieScene/Private/MovieSceneSection.cpp
--- a/Engine/Source/Runtime/MovieScene/Private/MovieSceneSection.cpp
+++ b/Engine/Source/Runtime/MovieScene/Private/MovieSceneSection.cpp
@@ -63,11 +63,18 @@ void UMovieSceneSection::InitialPlacement(const TArray<UMovieSceneSection*>& Sec
 	}
 }
 
-void UMovieSceneSection::AddKeyToCurve( FRichCurve& InCurve, float Time, float Value )
+void UMovieSceneSection::AddKeyToCurve( FRichCurve& InCurve, float Time, float Value, FKeyParams KeyParams, const bool bUnwindRotation )
 {
 	if(IsTimeWithinSection(Time))
 	{
 		Modify();
-		InCurve.UpdateOrAddKey(Time, Value);
+		if (InCurve.GetNumKeys() == 0 && !KeyParams.bAddKeyEvenIfUnchanged)
+		{
+			InCurve.SetDefaultValue(Value);
+		}
+		else
+		{
+			InCurve.UpdateOrAddKey(Time, Value, bUnwindRotation);
+		}
 	}
 }
\ No newline at end of file
diff --git a/Engine/Source/Runtime/MovieScene/Public/MovieSceneSection.h b/Engine/Source/Runtime/MovieScene/Public/MovieSceneSection.h
--- a/Engine/Source/Runtime/MovieScene/Public/MovieSceneSection.h
+++ b/Engine/Source/Runtime/MovieScene/Public/MovieSceneSection.h
@@ -2,6 +2,7 @@
 
 #pragma once
 
+#include "KeyParams.h"
 #include "MovieSceneSection.generated.h"
 
 /**
@@ -152,8 +153,10 @@ public:
 	 * @param InCurve	The curve to add keys to
 	 * @param Time		The time where the key should be added
 	 * @param Value		The value at the given time
+	 * @param KeyParams The keying parameters
+	 * @param bUnwindRotation Unwind rotation
 	 */
-	void MOVIESCENE_API AddKeyToCurve( FRichCurve& InCurve, float Time, float Value );
+	void MOVIESCENE_API AddKeyToCurve( FRichCurve& InCurve, float Time, float Value, FKeyParams KeyParams, const bool bUnwindRotation = false);
 
 	/**
 	 * Checks to see if this section overlaps with an array of other sections
diff --git a/Engine/Source/Runtime/MovieSceneTracks/Private/MovieSceneTracksPrivatePCH.h b/Engine/Source/Runtime/MovieSceneTracks/Private/MovieSceneTracksPrivatePCH.h
--- a/Engine/Source/Runtime/MovieSceneTracks/Private/MovieSceneTracksPrivatePCH.h
+++ b/Engine/Source/Runtime/MovieSceneTracks/Private/MovieSceneTracksPrivatePCH.h
@@ -8,6 +8,7 @@
 
 #include "Core.h"
 #include "Engine.h"
+#include "KeyParams.h"
 #include "MovieScene.h"
 #include "MovieSceneCommonHelpers.h"
 #include "MovieScenePropertyTrack.h"
diff --git a/Engine/Source/Runtime/MovieSceneTracks/Private/Sections/MovieScene3DTransformSection.cpp b/Engine/Source/Runtime/MovieSceneTracks/Private/Sections/MovieScene3DTransformSection.cpp
--- a/Engine/Source/Runtime/MovieSceneTracks/Private/Sections/MovieScene3DTransformSection.cpp
+++ b/Engine/Source/Runtime/MovieSceneTracks/Private/Sections/MovieScene3DTransformSection.cpp
@@ -66,70 +66,25 @@ void UMovieScene3DTransformSection::GetKeyHandles(TSet<FKeyHandle>& KeyHandles)
 	}
 }	
 
-void UMovieScene3DTransformSection::EvalTranslation( float Time, FVector& OutTranslation, TArray<bool>& OutHasTranslationKeys ) const
+void UMovieScene3DTransformSection::EvalTranslation( float Time, FVector& OutTranslation ) const
 {
-	OutHasTranslationKeys.Add(Translation[0].GetNumKeys() > 0);
-	OutHasTranslationKeys.Add(Translation[1].GetNumKeys() > 0);
-	OutHasTranslationKeys.Add(Translation[2].GetNumKeys() > 0);
-
-	// Evaluate each component of translation if there are any keys
-
-	if (OutHasTranslationKeys[0])
-	{
-		OutTranslation.X = Translation[0].Eval( Time );
-	}
-	if (OutHasTranslationKeys[1])
-	{
-		OutTranslation.Y = Translation[1].Eval( Time );
-	}
-	if (OutHasTranslationKeys[2])
-	{
-		OutTranslation.Z = Translation[2].Eval( Time );
-	}
+	OutTranslation.X = Translation[0].Eval( Time );
+	OutTranslation.Y = Translation[1].Eval( Time );
+	OutTranslation.Z = Translation[2].Eval( Time );
 }
 
-void UMovieScene3DTransformSection::EvalRotation( float Time, FRotator& OutRotation, TArray<bool>& OutHasRotationKeys ) const
+void UMovieScene3DTransformSection::EvalRotation( float Time, FRotator& OutRotation ) const
 {
-	OutHasRotationKeys.Add(Rotation[0].GetNumKeys() > 0);
-	OutHasRotationKeys.Add(Rotation[1].GetNumKeys() > 0);
-	OutHasRotationKeys.Add(Rotation[2].GetNumKeys() > 0);
-
-	// Evaluate each component of rotation if there are any keys
-
-	if (OutHasRotationKeys[0])
-	{
-		OutRotation.Roll = Rotation[0].Eval( Time );
-	}
-	if (OutHasRotationKeys[1])
-	{
-		OutRotation.Pitch = Rotation[1].Eval( Time );
-	}
-	if (OutHasRotationKeys[2])
-	{
-		OutRotation.Yaw = Rotation[2].Eval( Time );
-	}
+	OutRotation.Roll = Rotation[0].Eval( Time );
+	OutRotation.Pitch = Rotation[1].Eval( Time );
+	OutRotation.Yaw = Rotation[2].Eval( Time );
 }
 
-void UMovieScene3DTransformSection::EvalScale( float Time, FVector& OutScale, TArray<bool>& OutHasScaleKeys ) const
+void UMovieScene3DTransformSection::EvalScale( float Time, FVector& OutScale ) const
 {
-	OutHasScaleKeys.Add(Scale[0].GetNumKeys() > 0);
-	OutHasScaleKeys.Add(Scale[1].GetNumKeys() > 0);
-	OutHasScaleKeys.Add(Scale[2].GetNumKeys() > 0);
-
-	// Evaluate each component of scale if there are any keys
-
-	if (OutHasScaleKeys[0])
-	{
-		OutScale.X = Scale[0].Eval( Time );
-	}
-	if (OutHasScaleKeys[1])
-	{
-		OutScale.Y = Scale[1].Eval( Time );
-	}
-	if (OutHasScaleKeys[2])
-	{
-		OutScale.Z = Scale[2].Eval( Time );
-	}
+	OutScale.X = Scale[0].Eval( Time );
+	OutScale.Y = Scale[1].Eval( Time );
+	OutScale.Z = Scale[2].Eval( Time );
 }
 
 /**
@@ -170,28 +125,28 @@ FRichCurve& UMovieScene3DTransformSection::GetScaleCurve( EAxis::Type Axis )
 	return *ChooseCurve( Axis, Scale );
 }
 
-
-
 void UMovieScene3DTransformSection::AddTranslationKeys( const FTransformKey& TransformKey )
 {
 	const float Time = TransformKey.GetKeyTime();
 
-	const bool bUnwindRotation = false; // Not needed for translation
+	bool bTxKeyExists = Translation[0].IsKeyHandleValid(Translation[0].FindKey(Time));
+	bool bTyKeyExists = Translation[1].IsKeyHandleValid(Translation[1].FindKey(Time));
+	bool bTzKeyExists = Translation[2].IsKeyHandleValid(Translation[2].FindKey(Time));
 
-	// Key each component.  Note: we always key each component if no key exists
-	if( Translation[0].GetNumKeys() == 0 || TransformKey.ShouldKeyTranslation( EAxis::X ) )
+	// Key each component. Don't add a keyframe if there are existing keys and auto key is not enabled.
+	if( TransformKey.KeyParams.bAddKeyEvenIfUnchanged || (TransformKey.ShouldKeyTranslation( EAxis::X ) && !(!bTxKeyExists && !TransformKey.KeyParams.bAutoKeying && Translation[0].GetNumKeys() > 0)) )
 	{
-		AddKeyToCurve( Translation[0], Time, TransformKey.GetTranslationValue().X );
+		AddKeyToCurve( Translation[0], Time, TransformKey.GetTranslationValue().X, TransformKey.KeyParams );
 	}
 
-	if( Translation[1].GetNumKeys() == 0 || TransformKey.ShouldKeyTranslation( EAxis::Y ) )
+	if( TransformKey.KeyParams.bAddKeyEvenIfUnchanged || (TransformKey.ShouldKeyTranslation( EAxis::Y ) && !(!bTyKeyExists && !TransformKey.KeyParams.bAutoKeying && Translation[1].GetNumKeys() > 0)) )
 	{ 
-		AddKeyToCurve( Translation[1], Time, TransformKey.GetTranslationValue().Y );
+		AddKeyToCurve( Translation[1], Time, TransformKey.GetTranslationValue().Y, TransformKey.KeyParams );
 	}
 
-	if( Translation[2].GetNumKeys() == 0 || TransformKey.ShouldKeyTranslation( EAxis::Z ) )
+	if( TransformKey.KeyParams.bAddKeyEvenIfUnchanged || (TransformKey.ShouldKeyTranslation( EAxis::Z ) && !(!bTzKeyExists && !TransformKey.KeyParams.bAutoKeying && Translation[2].GetNumKeys() > 0)) )
 	{ 
-		AddKeyToCurve( Translation[2], Time, TransformKey.GetTranslationValue().Z );
+		AddKeyToCurve( Translation[2], Time, TransformKey.GetTranslationValue().Z, TransformKey.KeyParams );
 	}
 }
 
@@ -199,20 +154,24 @@ void UMovieScene3DTransformSection::AddRotationKeys( const FTransformKey& Transf
 {
 	const float Time = TransformKey.GetKeyTime();
 
-	// Key each component.  Note: we always key each component if no key exists
-	if( Rotation[0].GetNumKeys() == 0 || TransformKey.ShouldKeyRotation( EAxis::X ) )
+	bool bRxKeyExists = Rotation[0].IsKeyHandleValid(Rotation[0].FindKey(Time));
+	bool bRyKeyExists = Rotation[1].IsKeyHandleValid(Rotation[1].FindKey(Time));
+	bool bRzKeyExists = Rotation[2].IsKeyHandleValid(Rotation[2].FindKey(Time));
+
+	// Key each component. Don't add a keyframe if there are existing keys and auto key is not enabled.
+	if( TransformKey.KeyParams.bAddKeyEvenIfUnchanged || (TransformKey.ShouldKeyRotation( EAxis::X ) && !(!bRxKeyExists && !TransformKey.KeyParams.bAutoKeying && Rotation[0].GetNumKeys() > 0)) )
 	{
-		AddKeyToCurve( Rotation[0], Time, TransformKey.GetRotationValue().Roll );
+		AddKeyToCurve( Rotation[0], Time, TransformKey.GetRotationValue().Roll, TransformKey.KeyParams, bUnwindRotation );
 	}
 
-	if( Rotation[1].GetNumKeys() == 0 || TransformKey.ShouldKeyRotation( EAxis::Y ) )
+	if( TransformKey.KeyParams.bAddKeyEvenIfUnchanged || (TransformKey.ShouldKeyRotation( EAxis::Y ) && !(!bRyKeyExists && !TransformKey.KeyParams.bAutoKeying && Rotation[1].GetNumKeys() > 0)) )
 	{ 
-		AddKeyToCurve( Rotation[1], Time, TransformKey.GetRotationValue().Pitch );
+		AddKeyToCurve( Rotation[1], Time, TransformKey.GetRotationValue().Pitch, TransformKey.KeyParams, bUnwindRotation );
 	}
 
-	if( Rotation[2].GetNumKeys() == 0 || TransformKey.ShouldKeyRotation( EAxis::Z ) )
+	if( TransformKey.KeyParams.bAddKeyEvenIfUnchanged || (TransformKey.ShouldKeyRotation( EAxis::Z ) && !(!bRzKeyExists && !TransformKey.KeyParams.bAutoKeying && Rotation[2].GetNumKeys() > 0)) )
 	{ 
-		AddKeyToCurve( Rotation[2], Time, TransformKey.GetRotationValue().Yaw );
+		AddKeyToCurve( Rotation[2], Time, TransformKey.GetRotationValue().Yaw, TransformKey.KeyParams, bUnwindRotation );
 	}
 }
 
@@ -220,27 +179,31 @@ void UMovieScene3DTransformSection::AddScaleKeys( const FTransformKey& Transform
 {
 	const float Time = TransformKey.GetKeyTime();
 
-	const bool bUnwindRotation = false; // Not needed for scale
+	bool bSxKeyExists = Scale[0].IsKeyHandleValid(Scale[0].FindKey(Time));
+	bool bSyKeyExists = Scale[1].IsKeyHandleValid(Scale[1].FindKey(Time));
+	bool bSzKeyExists = Scale[2].IsKeyHandleValid(Scale[2].FindKey(Time));
 
-	// Key each component.  Note: we always key each component if no key exists
-	if( Scale[0].GetNumKeys() == 0 || TransformKey.ShouldKeyScale( EAxis::X ) )
+	// Key each component. Don't add a keyframe if there are existing keys and auto key is not enabled.
+	if( TransformKey.KeyParams.bAddKeyEvenIfUnchanged || (TransformKey.ShouldKeyScale( EAxis::X ) && !(!bSxKeyExists && !TransformKey.KeyParams.bAutoKeying && Scale[0].GetNumKeys() > 0)) )
 	{
-		AddKeyToCurve( Scale[0], Time, TransformKey.GetScaleValue().X );
+		AddKeyToCurve( Scale[0], Time, TransformKey.GetScaleValue().X, TransformKey.KeyParams );
 	}
 
-	if( Scale[1].GetNumKeys() == 0 || TransformKey.ShouldKeyScale( EAxis::Y ) )
+	if( TransformKey.KeyParams.bAddKeyEvenIfUnchanged || (TransformKey.ShouldKeyScale( EAxis::Y ) && !(!bSyKeyExists && !TransformKey.KeyParams.bAutoKeying && Scale[1].GetNumKeys() > 0)) )
 	{ 
-		AddKeyToCurve( Scale[1], Time, TransformKey.GetScaleValue().Y );
+		AddKeyToCurve( Scale[1], Time, TransformKey.GetScaleValue().Y, TransformKey.KeyParams );
 	}
 
-	if( Scale[2].GetNumKeys() == 0 || TransformKey.ShouldKeyScale( EAxis::Z ) )
+	if( TransformKey.KeyParams.bAddKeyEvenIfUnchanged || (TransformKey.ShouldKeyScale( EAxis::Z ) && !(!bSzKeyExists && !TransformKey.KeyParams.bAutoKeying && Scale[2].GetNumKeys() > 0)) )
 	{ 
-		AddKeyToCurve( Scale[2], Time, TransformKey.GetScaleValue().Z );
+		AddKeyToCurve( Scale[2], Time, TransformKey.GetScaleValue().Z, TransformKey.KeyParams );
 	}
 }
 
 bool UMovieScene3DTransformSection::NewKeyIsNewData(const FTransformKey& TransformKey) const
 {
+	const float KeyTime = TransformKey.GetKeyTime();
+
 	bool bHasEmptyKeys = false;
 	for (int32 i = 0; i < 3; ++i)
 	{
@@ -249,5 +212,35 @@ bool UMovieScene3DTransformSection::NewKeyIsNewData(const FTransformKey& Transfo
 			Rotation[i].GetNumKeys() == 0 ||
 			Scale[i].GetNumKeys() == 0;
 	}
-	return bHasEmptyKeys || TransformKey.ShouldKeyAny();
+
+	if (bHasEmptyKeys || TransformKey.ShouldKeyAny())
+	{
+		// Don't add a keyframe if there are existing keys and auto key is not enabled.
+		bool bTxKeyExists = Translation[0].IsKeyHandleValid(Translation[0].FindKey(KeyTime));
+		bool bTyKeyExists = Translation[1].IsKeyHandleValid(Translation[1].FindKey(KeyTime));
+		bool bTzKeyExists = Translation[2].IsKeyHandleValid(Translation[2].FindKey(KeyTime));
+
+		bool bRxKeyExists = Rotation[0].IsKeyHandleValid(Rotation[0].FindKey(KeyTime));
+		bool bRyKeyExists = Rotation[1].IsKeyHandleValid(Rotation[1].FindKey(KeyTime));
+		bool bRzKeyExists = Rotation[2].IsKeyHandleValid(Rotation[2].FindKey(KeyTime));
+
+		bool bSxKeyExists = Scale[0].IsKeyHandleValid(Scale[0].FindKey(KeyTime));
+		bool bSyKeyExists = Scale[1].IsKeyHandleValid(Scale[1].FindKey(KeyTime));
+		bool bSzKeyExists = Scale[2].IsKeyHandleValid(Scale[2].FindKey(KeyTime));
+
+		if ( !(!bTxKeyExists && !TransformKey.KeyParams.bAutoKeying && Translation[0].GetNumKeys() > 0) ||
+			 !(!bTyKeyExists && !TransformKey.KeyParams.bAutoKeying && Translation[1].GetNumKeys() > 0) ||
+			 !(!bTzKeyExists && !TransformKey.KeyParams.bAutoKeying && Translation[2].GetNumKeys() > 0) ||
+			 !(!bRxKeyExists && !TransformKey.KeyParams.bAutoKeying && Rotation[0].GetNumKeys() > 0) ||
+			 !(!bRyKeyExists && !TransformKey.KeyParams.bAutoKeying && Rotation[1].GetNumKeys() > 0) ||
+			 !(!bRzKeyExists && !TransformKey.KeyParams.bAutoKeying && Rotation[2].GetNumKeys() > 0) ||
+			 !(!bSxKeyExists && !TransformKey.KeyParams.bAutoKeying && Scale[0].GetNumKeys() > 0) ||
+			 !(!bSyKeyExists && !TransformKey.KeyParams.bAutoKeying && Scale[1].GetNumKeys() > 0) ||
+			 !(!bSzKeyExists && !TransformKey.KeyParams.bAutoKeying && Scale[2].GetNumKeys() > 0) )
+		{
+			return true;
+		}
+	}
+
+	return false;
 }
diff --git a/Engine/Source/Runtime/MovieSceneTracks/Private/Sections/MovieSceneBoolSection.cpp b/Engine/Source/Runtime/MovieSceneTracks/Private/Sections/MovieSceneBoolSection.cpp
--- a/Engine/Source/Runtime/MovieSceneTracks/Private/Sections/MovieSceneBoolSection.cpp
+++ b/Engine/Source/Runtime/MovieSceneTracks/Private/Sections/MovieSceneBoolSection.cpp
@@ -44,13 +44,31 @@ void UMovieSceneBoolSection::GetKeyHandles(TSet<FKeyHandle>& KeyHandles) const
 	}
 }
 
-void UMovieSceneBoolSection::AddKey( float Time, bool Value )
+void UMovieSceneBoolSection::AddKey( float Time, bool Value, FKeyParams KeyParams )
 {
 	Modify();
-	BoolCurve.UpdateOrAddKey(Time, Value ? 1 : 0);
+
+	if (BoolCurve.GetNumKeys() == 0 && !KeyParams.bAddKeyEvenIfUnchanged)
+	{
+		BoolCurve.SetDefaultValue(Value ? 1 : 0);
+	}
+	else
+	{
+		BoolCurve.UpdateOrAddKey(Time, Value ? 1 : 0);
+	}
 }
 
-bool UMovieSceneBoolSection::NewKeyIsNewData(float Time, bool Value) const
+bool UMovieSceneBoolSection::NewKeyIsNewData(float Time, bool Value, FKeyParams KeyParams) const
 {
-	return BoolCurve.GetNumKeys() == 0 || Eval(Time) != Value;
+	if (BoolCurve.GetNumKeys() == 0 || Eval(Time) != Value)
+	{
+		bool bKeyExists = BoolCurve.IsKeyHandleValid(BoolCurve.FindKey(Time));
+
+		// Don't add a keyframe if there are existing keys and auto key is not enabled.
+		if (!(!bKeyExists && !KeyParams.bAutoKeying && BoolCurve.GetNumKeys() > 0))
+		{
+			return true;
+		}
+	}
+	return false;
 }
diff --git a/Engine/Source/Runtime/MovieSceneTracks/Private/Sections/MovieSceneByteSection.cpp b/Engine/Source/Runtime/MovieSceneTracks/Private/Sections/MovieSceneByteSection.cpp
--- a/Engine/Source/Runtime/MovieSceneTracks/Private/Sections/MovieSceneByteSection.cpp
+++ b/Engine/Source/Runtime/MovieSceneTracks/Private/Sections/MovieSceneByteSection.cpp
@@ -38,13 +38,30 @@ void UMovieSceneByteSection::GetKeyHandles(TSet<FKeyHandle>& KeyHandles) const
 	}
 }
 
-void UMovieSceneByteSection::AddKey( float Time, uint8 Value )
+void UMovieSceneByteSection::AddKey( float Time, uint8 Value, FKeyParams KeyParams )
 {
 	Modify();
-	ByteCurve.UpdateOrAddKey(Time, Value ? 1 : 0);
+	if (ByteCurve.GetNumKeys() == 0 && !KeyParams.bAddKeyEvenIfUnchanged)
+	{
+		ByteCurve.SetDefaultValue(Value);
+	}
+	else
+	{
+		ByteCurve.UpdateOrAddKey(Time, Value ? 1 : 0);
+	}
 }
 
-bool UMovieSceneByteSection::NewKeyIsNewData(float Time, uint8 Value) const
+bool UMovieSceneByteSection::NewKeyIsNewData(float Time, uint8 Value, FKeyParams KeyParams) const
 {
-	return ByteCurve.GetNumKeys() == 0 || Eval(Time) != Value;
+	if (ByteCurve.GetNumKeys() == 0 || Eval(Time) != Value)
+	{
+		bool bKeyExists = ByteCurve.IsKeyHandleValid(ByteCurve.FindKey(Time));
+
+		// Don't add a keyframe if there are existing keys and auto key is not enabled.
+		if (!(!bKeyExists && !KeyParams.bAutoKeying && ByteCurve.GetNumKeys() > 0))
+		{
+			return true;
+		}
+	}
+	return false;
 }
diff --git a/Engine/Source/Runtime/MovieSceneTracks/Private/Sections/MovieSceneColorSection.cpp b/Engine/Source/Runtime/MovieSceneTracks/Private/Sections/MovieSceneColorSection.cpp
--- a/Engine/Source/Runtime/MovieSceneTracks/Private/Sections/MovieSceneColorSection.cpp
+++ b/Engine/Source/Runtime/MovieSceneTracks/Private/Sections/MovieSceneColorSection.cpp
@@ -73,54 +73,71 @@ FLinearColor UMovieSceneColorSection::Eval( float Position, const FLinearColor&
 						AlphaCurve.Eval(Position, DefaultColor.A));
 }
 
-void UMovieSceneColorSection::AddKey( float Time, const FColorKey& Key )
+void UMovieSceneColorSection::AddKey( float Time, const FColorKey& Key, FKeyParams KeyParams )
 {
 	Modify();
 
-	if( Key.CurveName == NAME_None )
+	static FName RedName("R");
+	static FName GreenName("G");
+	static FName BlueName("B");
+	static FName AlphaName("A");
+
+	FName CurveName = Key.CurveName;
+
+	bool bRedKeyExists = RedCurve.IsKeyHandleValid(RedCurve.FindKey(Time));
+	bool bGreenKeyExists = GreenCurve.IsKeyHandleValid(GreenCurve.FindKey(Time));
+	bool bBlueKeyExists = BlueCurve.IsKeyHandleValid(BlueCurve.FindKey(Time));
+	bool bAlphaKeyExists = AlphaCurve.IsKeyHandleValid(AlphaCurve.FindKey(Time));
+
+	if ( (CurveName == NAME_None || CurveName == RedName) &&
+			KeyParams.bAddKeyEvenIfUnchanged || !(!bRedKeyExists && !KeyParams.bAutoKeying && RedCurve.GetNumKeys() > 0) )
+	{
+		AddKeyToCurve(RedCurve, Time, Key.Value.R, KeyParams);
+	}
+
+	if ( (CurveName == NAME_None || CurveName == GreenName) &&
+			KeyParams.bAddKeyEvenIfUnchanged || !(!bGreenKeyExists && !KeyParams.bAutoKeying && GreenCurve.GetNumKeys() > 0) )
+	{
+		AddKeyToCurve(GreenCurve, Time, Key.Value.G, KeyParams);
+	}
+
+	if ( (CurveName == NAME_None || CurveName == BlueName) &&
+			KeyParams.bAddKeyEvenIfUnchanged || !(!bBlueKeyExists && !KeyParams.bAutoKeying && BlueCurve.GetNumKeys() > 0) )
 	{
-		AddKeyToCurve(RedCurve, Time, Key.Value.R);
-		AddKeyToCurve(GreenCurve, Time, Key.Value.G);
-		AddKeyToCurve(BlueCurve, Time, Key.Value.B);
-		AddKeyToCurve(AlphaCurve, Time, Key.Value.A);
+		AddKeyToCurve(BlueCurve, Time, Key.Value.B, KeyParams);
 	}
-	else
+
+	if ( (CurveName == NAME_None || CurveName == AlphaName) &&
+			KeyParams.bAddKeyEvenIfUnchanged || !(!bAlphaKeyExists && !KeyParams.bAutoKeying && AlphaCurve.GetNumKeys() > 0) )
 	{
-		AddKeyToNamedCurve( Time, Key );
+		AddKeyToCurve(AlphaCurve, Time, Key.Value.A, KeyParams);
 	}
 }
 
-bool UMovieSceneColorSection::NewKeyIsNewData(float Time, FLinearColor Value) const
+bool UMovieSceneColorSection::NewKeyIsNewData(float Time, FLinearColor Value, FKeyParams KeyParams) const
 {
-	return RedCurve.GetNumKeys() == 0 ||
+	bool bHasEmptyKeys = 
+		RedCurve.GetNumKeys() == 0 ||
 		GreenCurve.GetNumKeys() == 0 ||
 		BlueCurve.GetNumKeys() == 0 ||
 		AlphaCurve.GetNumKeys() == 0 ||
 		!Eval(Time,Value).Equals(Value);
-}
 
-void UMovieSceneColorSection::AddKeyToNamedCurve(float Time, const FColorKey& Key)
-{
-	static FName R("R");
-	static FName G("G");
-	static FName B("B");
-	static FName A("A");
-
-	FName CurveName = Key.CurveName;
-	if (CurveName == R)
-	{
-		AddKeyToCurve(RedCurve, Time, Key.Value.R);
-	}
-	else if (CurveName == G)
-	{
-		AddKeyToCurve(GreenCurve, Time, Key.Value.G);
-	}
-	else if (CurveName == B)
-	{
-		AddKeyToCurve(BlueCurve, Time, Key.Value.B);
-	}
-	else if (CurveName == A)
+	if (bHasEmptyKeys)
 	{
-		AddKeyToCurve(AlphaCurve, Time, Key.Value.A);
+		// Don't add a keyframe if there are existing keys and auto key is not enabled.
+		bool bRedKeyExists = RedCurve.IsKeyHandleValid(RedCurve.FindKey(Time));
+		bool bGreenKeyExists = GreenCurve.IsKeyHandleValid(GreenCurve.FindKey(Time));
+		bool bBlueKeyExists = BlueCurve.IsKeyHandleValid(BlueCurve.FindKey(Time));
+		bool bAlphaKeyExists = AlphaCurve.IsKeyHandleValid(AlphaCurve.FindKey(Time));
+
+		if ( !(!bRedKeyExists && !KeyParams.bAutoKeying && RedCurve.GetNumKeys() > 0) ||
+			 !(!bGreenKeyExists && !KeyParams.bAutoKeying && GreenCurve.GetNumKeys() > 0) ||
+			 !(!bBlueKeyExists && !KeyParams.bAutoKeying && BlueCurve.GetNumKeys() > 0) ||
+			 !(!bAlphaKeyExists && !KeyParams.bAutoKeying && AlphaCurve.GetNumKeys() > 0) )
+		{
+			return true;
+		}
 	}
+	return false;
 }
diff --git a/Engine/Source/Runtime/MovieSceneTracks/Private/Sections/MovieSceneFloatSection.cpp b/Engine/Source/Runtime/MovieSceneTracks/Private/Sections/MovieSceneFloatSection.cpp
--- a/Engine/Source/Runtime/MovieSceneTracks/Private/Sections/MovieSceneFloatSection.cpp
+++ b/Engine/Source/Runtime/MovieSceneTracks/Private/Sections/MovieSceneFloatSection.cpp
@@ -41,13 +41,30 @@ void UMovieSceneFloatSection::GetKeyHandles(TSet<FKeyHandle>& KeyHandles) const
 	}
 }
 
-void UMovieSceneFloatSection::AddKey( float Time, float Value )
+void UMovieSceneFloatSection::AddKey( float Time, float Value, FKeyParams KeyParams )
 {
 	Modify();
-	FloatCurve.UpdateOrAddKey(Time, Value);
+	if (FloatCurve.GetNumKeys() == 0 && !KeyParams.bAddKeyEvenIfUnchanged)
+	{
+		FloatCurve.SetDefaultValue(Value);
+	}
+	else
+	{
+		FloatCurve.UpdateOrAddKey(Time, Value);
+	}
 }
 
-bool UMovieSceneFloatSection::NewKeyIsNewData(float Time, float Value) const
+bool UMovieSceneFloatSection::NewKeyIsNewData(float Time, float Value, FKeyParams KeyParams) const
 {
-	return FloatCurve.GetNumKeys() == 0 || !FMath::IsNearlyEqual(FloatCurve.Eval(Time), Value);
+	if (FloatCurve.GetNumKeys() == 0 || Eval(Time) != Value)
+	{
+		bool bKeyExists = FloatCurve.IsKeyHandleValid(FloatCurve.FindKey(Time));
+
+		// Don't add a keyframe if there are existing keys and auto key is not enabled.
+		if (!(!bKeyExists && !KeyParams.bAutoKeying && FloatCurve.GetNumKeys() > 0))
+		{
+			return true;
+		}
+	}
+	return false;
 }
diff --git a/Engine/Source/Runtime/MovieSceneTracks/Private/Sections/MovieSceneVectorSection.cpp b/Engine/Source/Runtime/MovieSceneTracks/Private/Sections/MovieSceneVectorSection.cpp
--- a/Engine/Source/Runtime/MovieSceneTracks/Private/Sections/MovieSceneVectorSection.cpp
+++ b/Engine/Source/Runtime/MovieSceneTracks/Private/Sections/MovieSceneVectorSection.cpp
@@ -19,26 +19,42 @@ FVector4 UMovieSceneVectorSection::Eval( float Position, const FVector4& Default
 		Curves[3].Eval( Position, DefaultVector.W ) );
 }
 
-void UMovieSceneVectorSection::AddKey( float Time, FName CurveName, const FVector4& Value )
+void UMovieSceneVectorSection::AddKey( float Time, FName CurveName, const FVector4& Value, FKeyParams KeyParams )
 {
 	Modify();
 
-	if( CurveName == NAME_None )
-	{
-		check(ChannelsUsed >= 2 && ChannelsUsed <= 4);
+	static FName X("X");
+	static FName Y("Y");
+	static FName Z("Z");
+	static FName W("W");
+	
+	check(ChannelsUsed >= 2 && ChannelsUsed <= 4);
 
-		for (int32 i = 0; i < ChannelsUsed; ++i)
+	for (int32 i = 0; i < ChannelsUsed; ++i)
+	{
+		if ( CurveName == NAME_None || 
+			 (CurveName == X && i == 0) ||
+			 (CurveName == Y && i == 1) ||
+			 (CurveName == Z && i == 2) ||
+			 (CurveName == W && i == 3) )
 		{
-			Curves[i].UpdateOrAddKey(Time, Value[i]);
+			if (Curves[i].GetNumKeys() == 0 && !KeyParams.bAddKeyEvenIfUnchanged)
+			{
+				Curves[i].SetDefaultValue(Value[i]);
+			}
+			else
+			{
+				bool bKeyExists = Curves[i].IsKeyHandleValid(Curves[i].FindKey(Time));
+				if ( KeyParams.bAddKeyEvenIfUnchanged || !(!bKeyExists && !KeyParams.bAutoKeying && Curves[i].GetNumKeys() > 0) )
+				{
+					Curves[i].UpdateOrAddKey(Time, Value[i]);
+				}
+			}
 		}
 	}
-	else
-	{
-		AddKeyToNamedCurve( Time, CurveName, Value );
-	}
 }
 
-bool UMovieSceneVectorSection::NewKeyIsNewData(float Time, const FVector4& Value) const
+bool UMovieSceneVectorSection::NewKeyIsNewData(float Time, const FVector4& Value, FKeyParams KeyParams) const
 {
 	check(ChannelsUsed >= 2 && ChannelsUsed <= 4);
 
@@ -53,32 +69,20 @@ bool UMovieSceneVectorSection::NewKeyIsNewData(float Time, const FVector4& Value
 			break;
 		}
 	}
-	return bNewData;
-}
 
-void UMovieSceneVectorSection::AddKeyToNamedCurve(float Time, FName CurveName, const FVector4& Value)
-{
-	static FName X("X");
-	static FName Y("Y");
-	static FName Z("Z");
-	static FName W("W");
-
-	if (CurveName == X)
-	{
-		Curves[0].UpdateOrAddKey(Time, Value.X);
-	}
-	else if (CurveName == Y)
-	{
-		Curves[1].UpdateOrAddKey(Time, Value.Y);
-	}
-	else if (CurveName == Z)
+	if (bNewData)
 	{
-		Curves[2].UpdateOrAddKey(Time, Value.Z);
-	}
-	else if (CurveName == W)
-	{
-		Curves[3].UpdateOrAddKey(Time, Value.W);
+		for (int32 i = 0; i < ChannelsUsed; ++i)
+		{
+			// Don't add a keyframe if there are existing keys and auto key is not enabled.
+			bool bKeyExists = Curves[i].IsKeyHandleValid(Curves[i].FindKey(Time));
+			if ( !(!bKeyExists && !KeyParams.bAutoKeying && Curves[i].GetNumKeys() > 0) )
+			{
+				return true;
+			}
+		}
 	}
+	return bNewData;
 }
 
 void UMovieSceneVectorSection::MoveSection( float DeltaTime, TSet<FKeyHandle>& KeyHandles )
diff --git a/Engine/Source/Runtime/MovieSceneTracks/Private/Sections/MovieSceneVisibilitySection.cpp b/Engine/Source/Runtime/MovieSceneTracks/Private/Sections/MovieSceneVisibilitySection.cpp
--- a/Engine/Source/Runtime/MovieSceneTracks/Private/Sections/MovieSceneVisibilitySection.cpp
+++ b/Engine/Source/Runtime/MovieSceneTracks/Private/Sections/MovieSceneVisibilitySection.cpp
@@ -9,14 +9,14 @@ UMovieSceneVisibilitySection::UMovieSceneVisibilitySection( const FObjectInitial
 {
 }
 
-void UMovieSceneVisibilitySection::AddKey( float Time, bool bHiddenInGame )
+void UMovieSceneVisibilitySection::AddKey( float Time, bool bHiddenInGame, FKeyParams KeyParams )
 {
 	// The property that's being changed is bHiddenInGame. But we want to store the inverse of that as visibility, so invert the incoming value.
-	Super::AddKey(Time, !bHiddenInGame);
+	Super::AddKey(Time, !bHiddenInGame, KeyParams);
 }
 
-bool UMovieSceneVisibilitySection::NewKeyIsNewData( float Time, bool bHiddenInGame ) const
+bool UMovieSceneVisibilitySection::NewKeyIsNewData( float Time, bool bHiddenInGame, FKeyParams KeyParams ) const
 {
 	// The property that's being changed is bHiddenInGame. But we want to store the inverse of that as visibility, so invert the incoming value.
-	return Super::NewKeyIsNewData(Time, !bHiddenInGame);
+	return Super::NewKeyIsNewData(Time, !bHiddenInGame, KeyParams);
 }
diff --git a/Engine/Source/Runtime/MovieSceneTracks/Private/TrackInstances/MovieScene3DTransformTrackInstance.cpp b/Engine/Source/Runtime/MovieSceneTracks/Private/TrackInstances/MovieScene3DTransformTrackInstance.cpp
--- a/Engine/Source/Runtime/MovieSceneTracks/Private/TrackInstances/MovieScene3DTransformTrackInstance.cpp
+++ b/Engine/Source/Runtime/MovieSceneTracks/Private/TrackInstances/MovieScene3DTransformTrackInstance.cpp
@@ -49,11 +49,8 @@ void FMovieScene3DTransformTrackInstance::Update( float Position, float LastPosi
 	FVector Translation;
 	FRotator Rotation;
 	FVector Scale;
-	TArray<bool> bHasTranslationKeys;
-	TArray<bool> bHasRotationKeys;
-	TArray<bool> bHasScaleKeys;
 
-	if( TransformTrack->Eval( Position, LastPosition, Translation, Rotation, Scale, bHasTranslationKeys, bHasRotationKeys, bHasScaleKeys ) )
+	if( TransformTrack->Eval( Position, LastPosition, Translation, Rotation, Scale ) )
 	{
 		for( int32 ObjIndex = 0; ObjIndex < RuntimeObjects.Num(); ++ObjIndex )
 		{
@@ -61,25 +58,9 @@ void FMovieScene3DTransformTrackInstance::Update( float Position, float LastPosi
 
 			if (SceneComponent != NULL)
 			{
-				// Set the relative translation and rotation.  Note they are set once instead of individually to avoid a redundant component transform update.
-				FVector NewTranslation;
-				NewTranslation[0] = bHasTranslationKeys[0] ? Translation[0] : SceneComponent->RelativeLocation[0];
-				NewTranslation[1] = bHasTranslationKeys[1] ? Translation[1] : SceneComponent->RelativeLocation[1];
-				NewTranslation[2] = bHasTranslationKeys[2] ? Translation[2] : SceneComponent->RelativeLocation[2];
+				SceneComponent->SetRelativeLocationAndRotation(Translation, Rotation);
 
-				FRotator NewRotation;
-				NewRotation.Roll = bHasRotationKeys[0] ? Rotation.Roll : SceneComponent->RelativeRotation.Roll;
-				NewRotation.Pitch = bHasRotationKeys[1] ? Rotation.Pitch : SceneComponent->RelativeRotation.Pitch;
-				NewRotation.Yaw = bHasRotationKeys[2] ? Rotation.Yaw : SceneComponent->RelativeRotation.Yaw;
-
-				SceneComponent->SetRelativeLocationAndRotation(NewTranslation, NewRotation);
-
-				FVector NewScale;
-				NewScale[0] = bHasScaleKeys[0] ? Scale[0] : SceneComponent->RelativeScale3D[0];
-				NewScale[1] = bHasScaleKeys[1] ? Scale[1] : SceneComponent->RelativeScale3D[1];
-				NewScale[2] = bHasScaleKeys[2] ? Scale[2] : SceneComponent->RelativeScale3D[2];
-
-				SceneComponent->SetRelativeScale3D(NewScale);
+				SceneComponent->SetRelativeScale3D(Scale);
 			}
 		}
 	}
diff --git a/Engine/Source/Runtime/MovieSceneTracks/Private/Tracks/MovieScene3DTransformTrack.cpp b/Engine/Source/Runtime/MovieSceneTracks/Private/Tracks/MovieScene3DTransformTrack.cpp
--- a/Engine/Source/Runtime/MovieSceneTracks/Private/Tracks/MovieScene3DTransformTrack.cpp
+++ b/Engine/Source/Runtime/MovieSceneTracks/Private/Tracks/MovieScene3DTransformTrack.cpp
@@ -75,7 +75,7 @@ TSharedPtr<IMovieSceneTrackInstance> UMovieScene3DTransformTrack::CreateInstance
 bool UMovieScene3DTransformTrack::AddKeyToSection( const FGuid& ObjectHandle, const FTransformKey& InKey, const bool bUnwindRotation, F3DTransformTrackKey::Type KeyType )
 {
 	const UMovieSceneSection* NearestSection = MovieSceneHelpers::FindNearestSectionAtTime(Sections, InKey.GetKeyTime());
-	if (!NearestSection || CastChecked<UMovieScene3DTransformSection>(NearestSection)->NewKeyIsNewData(InKey))
+	if (!NearestSection || InKey.KeyParams.bAddKeyEvenIfUnchanged || CastChecked<UMovieScene3DTransformSection>(NearestSection)->NewKeyIsNewData(InKey))
 	{
 		Modify();
 
@@ -83,11 +83,17 @@ bool UMovieScene3DTransformTrack::AddKeyToSection( const FGuid& ObjectHandle, co
 
 		// key each component of the transform
 		if (KeyType & F3DTransformTrackKey::Key_Translation)
+		{
 			NewSection->AddTranslationKeys( InKey );
+		}
 		if (KeyType & F3DTransformTrackKey::Key_Rotation)
+		{
 			NewSection->AddRotationKeys( InKey, bUnwindRotation );
+		}
 		if (KeyType & F3DTransformTrackKey::Key_Scale)
+		{
 			NewSection->AddScaleKeys( InKey );
+		}
 
 		return true;
 	}
@@ -95,7 +101,7 @@ bool UMovieScene3DTransformTrack::AddKeyToSection( const FGuid& ObjectHandle, co
 }
 
 
-bool UMovieScene3DTransformTrack::Eval( float Position, float LastPosition, FVector& OutTranslation, FRotator& OutRotation, FVector& OutScale, TArray<bool>& OutHasTranslationKeys, TArray<bool>& OutHasRotationKeys, TArray<bool>& OutHasScaleKeys ) const
+bool UMovieScene3DTransformTrack::Eval( float Position, float LastPosition, FVector& OutTranslation, FRotator& OutRotation, FVector& OutScale ) const
 {
 	const UMovieSceneSection* Section = MovieSceneHelpers::FindNearestSectionAtTime( Sections, Position );
 
@@ -106,10 +112,20 @@ bool UMovieScene3DTransformTrack::Eval( float Position, float LastPosition, FVec
 		Position = FMath::Clamp(Position, Section->GetStartTime(), Section->GetEndTime());
 
 		// Evaluate translation,rotation, and scale curves.  If no keys were found on one of these, that component of the transform will remain unchained
-		TransformSection->EvalTranslation( Position, OutTranslation, OutHasTranslationKeys );
-		TransformSection->EvalRotation( Position, OutRotation, OutHasRotationKeys );
-		TransformSection->EvalScale( Position, OutScale, OutHasScaleKeys );
+		TransformSection->EvalTranslation( Position, OutTranslation );
+		TransformSection->EvalRotation( Position, OutRotation );
+		TransformSection->EvalScale( Position, OutScale );
 	}
 
 	return Section != NULL;
+}
+
+bool UMovieScene3DTransformTrack::CanKeyTrack(const FTransformKey& InKey ) const
+{
+	const UMovieSceneSection* NearestSection = MovieSceneHelpers::FindNearestSectionAtTime(Sections, InKey.GetKeyTime());
+	if (!NearestSection || CastChecked<UMovieScene3DTransformSection>(NearestSection)->NewKeyIsNewData(InKey))
+	{
+		return true;
+	}
+	return false;
 }
\ No newline at end of file
diff --git a/Engine/Source/Runtime/MovieSceneTracks/Private/Tracks/MovieSceneBoolTrack.cpp b/Engine/Source/Runtime/MovieSceneTracks/Private/Tracks/MovieSceneBoolTrack.cpp
--- a/Engine/Source/Runtime/MovieSceneTracks/Private/Tracks/MovieSceneBoolTrack.cpp
+++ b/Engine/Source/Runtime/MovieSceneTracks/Private/Tracks/MovieSceneBoolTrack.cpp
@@ -22,16 +22,16 @@ TSharedPtr<IMovieSceneTrackInstance> UMovieSceneBoolTrack::CreateInstance()
 	return MakeShareable( new FMovieSceneBoolTrackInstance( *this ) );
 }
 
-bool UMovieSceneBoolTrack::AddKeyToSection( float Time, bool Value )
+bool UMovieSceneBoolTrack::AddKeyToSection( float Time, bool Value, FKeyParams KeyParams)
 {
 	const UMovieSceneSection* NearestSection = MovieSceneHelpers::FindNearestSectionAtTime( Sections, Time );
-	if (!NearestSection || CastChecked<UMovieSceneBoolSection>(NearestSection)->NewKeyIsNewData(Time, Value))
+	if (!NearestSection || KeyParams.bAddKeyEvenIfUnchanged || CastChecked<UMovieSceneBoolSection>(NearestSection)->NewKeyIsNewData(Time, Value, KeyParams))
 	{
 		Modify();
 
 		UMovieSceneBoolSection* NewSection = CastChecked<UMovieSceneBoolSection>(FindOrAddSection(  Time ));
 	
-		NewSection->AddKey( Time, Value );
+		NewSection->AddKey( Time, Value, KeyParams );
 
 		return true;
 	}
diff --git a/Engine/Source/Runtime/MovieSceneTracks/Private/Tracks/MovieSceneByteTrack.cpp b/Engine/Source/Runtime/MovieSceneTracks/Private/Tracks/MovieSceneByteTrack.cpp
--- a/Engine/Source/Runtime/MovieSceneTracks/Private/Tracks/MovieSceneByteTrack.cpp
+++ b/Engine/Source/Runtime/MovieSceneTracks/Private/Tracks/MovieSceneByteTrack.cpp
@@ -22,16 +22,16 @@ TSharedPtr<IMovieSceneTrackInstance> UMovieSceneByteTrack::CreateInstance()
 	return MakeShareable( new FMovieSceneByteTrackInstance( *this ) );
 }
 
-bool UMovieSceneByteTrack::AddKeyToSection( float Time, uint8 Value )
+bool UMovieSceneByteTrack::AddKeyToSection( float Time, uint8 Value, FKeyParams KeyParams )
 {
 	const UMovieSceneSection* NearestSection = MovieSceneHelpers::FindNearestSectionAtTime( Sections, Time );
-	if (!NearestSection || CastChecked<UMovieSceneByteSection>(NearestSection)->NewKeyIsNewData(Time, Value))
+	if (!NearestSection || KeyParams.bAddKeyEvenIfUnchanged || CastChecked<UMovieSceneByteSection>(NearestSection)->NewKeyIsNewData(Time, Value, KeyParams))
 	{
 		Modify();
 
 		UMovieSceneByteSection* NewSection = CastChecked<UMovieSceneByteSection>(FindOrAddSection(  Time ));
 	
-		NewSection->AddKey( Time, Value );
+		NewSection->AddKey( Time, Value, KeyParams );
 
 		return true;
 	}
diff --git a/Engine/Source/Runtime/MovieSceneTracks/Private/Tracks/MovieSceneColorTrack.cpp b/Engine/Source/Runtime/MovieSceneTracks/Private/Tracks/MovieSceneColorTrack.cpp
--- a/Engine/Source/Runtime/MovieSceneTracks/Private/Tracks/MovieSceneColorTrack.cpp
+++ b/Engine/Source/Runtime/MovieSceneTracks/Private/Tracks/MovieSceneColorTrack.cpp
@@ -22,15 +22,15 @@ TSharedPtr<IMovieSceneTrackInstance> UMovieSceneColorTrack::CreateInstance()
 }
 
 
-bool UMovieSceneColorTrack::AddKeyToSection( float Time, const FColorKey& Key )
+bool UMovieSceneColorTrack::AddKeyToSection( float Time, const FColorKey& Key, FKeyParams KeyParams )
 {
 	const UMovieSceneSection* NearestSection = MovieSceneHelpers::FindNearestSectionAtTime( Sections, Time );
-	if (!NearestSection || Key.bAddKeyEvenIfUnchanged || CastChecked<UMovieSceneColorSection>(NearestSection)->NewKeyIsNewData(Time, Key.Value))
+	if (!NearestSection || KeyParams.bAddKeyEvenIfUnchanged || CastChecked<UMovieSceneColorSection>(NearestSection)->NewKeyIsNewData(Time, Key.Value, KeyParams))
 	{
 		Modify();
 		UMovieSceneColorSection* NewSection = CastChecked<UMovieSceneColorSection>( FindOrAddSection( Time ) );
 
-		NewSection->AddKey( Time, Key );
+		NewSection->AddKey( Time, Key, KeyParams );
 
 		return true;
 	}
diff --git a/Engine/Source/Runtime/MovieSceneTracks/Private/Tracks/MovieSceneEventTrack.cpp b/Engine/Source/Runtime/MovieSceneTracks/Private/Tracks/MovieSceneEventTrack.cpp
--- a/Engine/Source/Runtime/MovieSceneTracks/Private/Tracks/MovieSceneEventTrack.cpp
+++ b/Engine/Source/Runtime/MovieSceneTracks/Private/Tracks/MovieSceneEventTrack.cpp
@@ -9,7 +9,7 @@
 /* UMovieSceneEventTrack interface
  *****************************************************************************/
 
-bool UMovieSceneEventTrack::AddKeyToSection(float Time, FName EventName)
+bool UMovieSceneEventTrack::AddKeyToSection(float Time, FName EventName, FKeyParams KeyParams)
 {
 	// @todo sequencer: gmp: add key to event section
 	// should be inserted into sorted array
diff --git a/Engine/Source/Runtime/MovieSceneTracks/Private/Tracks/MovieSceneFloatTrack.cpp b/Engine/Source/Runtime/MovieSceneTracks/Private/Tracks/MovieSceneFloatTrack.cpp
--- a/Engine/Source/Runtime/MovieSceneTracks/Private/Tracks/MovieSceneFloatTrack.cpp
+++ b/Engine/Source/Runtime/MovieSceneTracks/Private/Tracks/MovieSceneFloatTrack.cpp
@@ -23,16 +23,16 @@ TSharedPtr<IMovieSceneTrackInstance> UMovieSceneFloatTrack::CreateInstance()
 }
 
 
-bool UMovieSceneFloatTrack::AddKeyToSection( float Time, float Value )
+bool UMovieSceneFloatTrack::AddKeyToSection( float Time, float Value, FKeyParams KeyParams )
 {
 	const UMovieSceneSection* NearestSection = MovieSceneHelpers::FindNearestSectionAtTime( Sections, Time );
-	if (!NearestSection || CastChecked<UMovieSceneFloatSection>(NearestSection)->NewKeyIsNewData(Time, Value))
+	if (!NearestSection || KeyParams.bAddKeyEvenIfUnchanged || CastChecked<UMovieSceneFloatSection>(NearestSection)->NewKeyIsNewData(Time, Value, KeyParams))
 	{
 		Modify();
 
 		UMovieSceneFloatSection* NewSection = Cast<UMovieSceneFloatSection>( FindOrAddSection( Time ) );
 
-		NewSection->AddKey( Time, Value );
+		NewSection->AddKey( Time, Value, KeyParams );
 
 		return true;
 	}
diff --git a/Engine/Source/Runtime/MovieSceneTracks/Private/Tracks/MovieSceneVectorTrack.cpp b/Engine/Source/Runtime/MovieSceneTracks/Private/Tracks/MovieSceneVectorTrack.cpp
--- a/Engine/Source/Runtime/MovieSceneTracks/Private/Tracks/MovieSceneVectorTrack.cpp
+++ b/Engine/Source/Runtime/MovieSceneTracks/Private/Tracks/MovieSceneVectorTrack.cpp
@@ -6,7 +6,7 @@
 #include "IMovieScenePlayer.h"
 #include "MovieSceneVectorTrackInstance.h"
 
-FVectorKey::FVectorKey( const FVector2D& InValue, FName InCurveName, bool InbAddKeyEvenIfUnChanged )
+FVectorKey::FVectorKey( const FVector2D& InValue, FName InCurveName )
 {
 	Value.X = InValue.X;
 	Value.Y = InValue.Y;
@@ -14,10 +14,9 @@ FVectorKey::FVectorKey( const FVector2D& InValue, FName InCurveName, bool InbAdd
 	Value.W = 0;
 	ChannelsUsed = 2;
 	CurveName = InCurveName;
-	bAddKeyEvenIfUnchanged = InbAddKeyEvenIfUnChanged;
 }
 
-FVectorKey::FVectorKey( const FVector& InValue, FName InCurveName, bool InbAddKeyEvenIfUnChanged )
+FVectorKey::FVectorKey( const FVector& InValue, FName InCurveName )
 {
 	Value.X = InValue.X;
 	Value.Y = InValue.Y;
@@ -25,15 +24,13 @@ FVectorKey::FVectorKey( const FVector& InValue, FName InCurveName, bool InbAddKe
 	Value.W = 0;
 	ChannelsUsed = 3;
 	CurveName = InCurveName;
-	bAddKeyEvenIfUnchanged = InbAddKeyEvenIfUnChanged;
 }
 
-FVectorKey::FVectorKey( const FVector4& InValue, FName InCurveName, bool InbAddKeyEvenIfUnChanged )
+FVectorKey::FVectorKey( const FVector4& InValue, FName InCurveName )
 {
 	Value = InValue;
 	ChannelsUsed = 4;
 	CurveName = InCurveName;
-	bAddKeyEvenIfUnchanged = InbAddKeyEvenIfUnChanged;
 }
 
 UMovieSceneVectorTrack::UMovieSceneVectorTrack( const FObjectInitializer& ObjectInitializer )
@@ -52,10 +49,10 @@ TSharedPtr<IMovieSceneTrackInstance> UMovieSceneVectorTrack::CreateInstance()
 	return MakeShareable( new FMovieSceneVectorTrackInstance( *this ) );
 }
 
-bool UMovieSceneVectorTrack::AddKeyToSection( float Time, const FVector4& Value, int32 InChannelsUsed, FName CurveName, bool bAddKeyEvenIfUnchanged )
+bool UMovieSceneVectorTrack::AddKeyToSection( float Time, const FVector4& Value, int32 InChannelsUsed, FName CurveName, FKeyParams KeyParams )
 {
 	const UMovieSceneSection* NearestSection = MovieSceneHelpers::FindNearestSectionAtTime( Sections, Time );
-	if (!NearestSection || bAddKeyEvenIfUnchanged || CastChecked<UMovieSceneVectorSection>(NearestSection)->NewKeyIsNewData(Time, Value))
+	if (!NearestSection || KeyParams.bAddKeyEvenIfUnchanged || CastChecked<UMovieSceneVectorSection>(NearestSection)->NewKeyIsNewData(Time, Value, KeyParams))
 	{
 		Modify();
 
@@ -63,7 +60,7 @@ bool UMovieSceneVectorTrack::AddKeyToSection( float Time, const FVector4& Value,
 		// @todo Sequencer - I don't like setting the channels used here. It should only be checked here, and set on section creation
 		NewSection->SetChannelsUsed(InChannelsUsed);
 
-		NewSection->AddKey( Time, CurveName, Value );
+		NewSection->AddKey( Time, CurveName, Value, KeyParams );
 
 		// We dont support one track containing multiple sections of differing channel amounts
 		NumChannelsUsed = InChannelsUsed;
@@ -73,9 +70,9 @@ bool UMovieSceneVectorTrack::AddKeyToSection( float Time, const FVector4& Value,
 	return false;
 }
 
-bool UMovieSceneVectorTrack::AddKeyToSection( float Time, const FVectorKey& Key )
+bool UMovieSceneVectorTrack::AddKeyToSection( float Time, const FVectorKey& Key, FKeyParams KeyParams )
 {
-	return AddKeyToSection(Time, Key.Value, Key.ChannelsUsed, Key.CurveName, Key.bAddKeyEvenIfUnchanged );
+	return AddKeyToSection(Time, Key.Value, Key.ChannelsUsed, Key.CurveName, KeyParams );
 }
 
 bool UMovieSceneVectorTrack::Eval( float Position, float LastPosition, FVector4& InOutVector ) const
diff --git a/Engine/Source/Runtime/MovieSceneTracks/Private/Tracks/MovieSceneVisibilityTrack.cpp b/Engine/Source/Runtime/MovieSceneTracks/Private/Tracks/MovieSceneVisibilityTrack.cpp
--- a/Engine/Source/Runtime/MovieSceneTracks/Private/Tracks/MovieSceneVisibilityTrack.cpp
+++ b/Engine/Source/Runtime/MovieSceneTracks/Private/Tracks/MovieSceneVisibilityTrack.cpp
@@ -21,16 +21,16 @@ TSharedPtr<IMovieSceneTrackInstance> UMovieSceneVisibilityTrack::CreateInstance(
 	return MakeShareable( new FMovieSceneVisibilityTrackInstance( *this ) );
 }
 
-bool UMovieSceneVisibilityTrack::AddKeyToSection( float Time, bool Value )
+bool UMovieSceneVisibilityTrack::AddKeyToSection( float Time, bool Value, FKeyParams KeyParams )
 {
 	const UMovieSceneSection* NearestSection = MovieSceneHelpers::FindNearestSectionAtTime( Sections, Time );
-	if (!NearestSection || CastChecked<UMovieSceneVisibilitySection>(NearestSection)->NewKeyIsNewData(Time, Value))
+	if (!NearestSection || KeyParams.bAddKeyEvenIfUnchanged || CastChecked<UMovieSceneVisibilitySection>(NearestSection)->NewKeyIsNewData(Time, Value, KeyParams))
 	{
 		Modify();
 
 		UMovieSceneVisibilitySection* NewSection = CastChecked<UMovieSceneVisibilitySection>(FindOrAddSection(  Time ));
 	
-		NewSection->AddKey( Time, Value );
+		NewSection->AddKey( Time, Value, KeyParams );
 
 		return true;
 	}
diff --git a/Engine/Source/Runtime/MovieSceneTracks/Public/Sections/MovieScene3DTransformSection.h b/Engine/Source/Runtime/MovieSceneTracks/Public/Sections/MovieScene3DTransformSection.h
--- a/Engine/Source/Runtime/MovieSceneTracks/Public/Sections/MovieScene3DTransformSection.h
+++ b/Engine/Source/Runtime/MovieSceneTracks/Public/Sections/MovieScene3DTransformSection.h
@@ -27,27 +27,24 @@ public:
 	 *
 	 * @param Time				The position in time within the movie scene
 	 * @param OutTranslation	The evaluated translation.  Note: will remain unchanged if there were no keys to evaluate
-	 * @param OutHasTranslationKeys True if the out has translation keys.
 	 */
-	void EvalTranslation( float Time, FVector& OutTranslation, TArray<bool>& OutHasTranslationKeys ) const;
+	void EvalTranslation( float Time, FVector& OutTranslation ) const;
 
 	/**
 	 * Evaluates the rotation component of the transform
 	 *
 	 * @param Time				The position in time within the movie scene
 	 * @param OutRotation		The evaluated rotation.  Note: will remain unchanged if there were no keys to evaluate
-	 * @param OutHasRotationKeys True if the out has rotation keys.
 	 */
-	void EvalRotation( float Time, FRotator& OutRotation, TArray<bool>& OutHasRotationKeys ) const;
+	void EvalRotation( float Time, FRotator& OutRotation ) const;
 
 	/**
 	 * Evaluates the scale component of the transform
 	 *
 	 * @param Time				The position in time within the movie scene
 	 * @param OutScale			The evaluated scale.  Note: will remain unchanged if there were no keys to evaluate
-	 * @param OutHasScaleKeys   True if the out has scale keys.
 	 */
-	void EvalScale( float Time, FVector& OutScale, TArray<bool>& OutHasScaleKeys ) const;
+	void EvalScale( float Time, FVector& OutScale ) const;
 
 	/** 
 	 * Returns the translation curve for a specific axis
@@ -98,10 +95,11 @@ public:
 	/** 
 	 * Determines if a new key would be new data, or just a duplicate of existing data
 	 *
-	 * @param TransformKey The new key that would be added
+	 * @param TransformKey	Information about the transform that should be keyed
 	 * @return True if the new key would be new data, false if duplicate
 	 */
-	bool NewKeyIsNewData(const class FTransformKey& TransformKey) const;
+	bool NewKeyIsNewData(const FTransformKey& TransformKey ) const;
+
 private:
 	/** Translation curves */
 	UPROPERTY()
diff --git a/Engine/Source/Runtime/MovieSceneTracks/Public/Sections/MovieSceneBoolSection.h b/Engine/Source/Runtime/MovieSceneTracks/Public/Sections/MovieSceneBoolSection.h
--- a/Engine/Source/Runtime/MovieSceneTracks/Public/Sections/MovieSceneBoolSection.h
+++ b/Engine/Source/Runtime/MovieSceneTracks/Public/Sections/MovieSceneBoolSection.h
@@ -24,17 +24,19 @@ public:
 	 *
 	 * @param Time	The location in time where the key should be added
 	 * @param Value	The value of the key
+	 * @param KeyParams The keying parameters
 	 */
-	void AddKey( float Time, bool Value );
+	void AddKey( float Time, bool Value, FKeyParams KeyParams );
 	
 	/** 
 	 * Determines if a new key would be new data, or just a duplicate of existing data
 	 *
 	 * @param Time	The location in time where the key would be added
 	 * @param Value	The value of the new key
+	 * @param KeyParams The keying parameters
 	 * @return True if the new key would be new data, false if duplicate
 	 */
-	bool NewKeyIsNewData(float Time, bool Value) const;
+	bool NewKeyIsNewData(float Time, bool Value, FKeyParams KeyParams) const;
 
 	/**
 	 * UMovieSceneSection interface 
diff --git a/Engine/Source/Runtime/MovieSceneTracks/Public/Sections/MovieSceneByteSection.h b/Engine/Source/Runtime/MovieSceneTracks/Public/Sections/MovieSceneByteSection.h
--- a/Engine/Source/Runtime/MovieSceneTracks/Public/Sections/MovieSceneByteSection.h
+++ b/Engine/Source/Runtime/MovieSceneTracks/Public/Sections/MovieSceneByteSection.h
@@ -24,17 +24,19 @@ public:
 	 *
 	 * @param Time	The location in time where the key should be added
 	 * @param Value	The value of the key
+	 * @param KeyParams The keying parameters
 	 */
-	void AddKey( float Time, uint8 Value );
+	void AddKey( float Time, uint8 Value, FKeyParams KeyParams );
 	
 	/** 
 	 * Determines if a new key would be new data, or just a duplicate of existing data
 	 *
 	 * @param Time	The location in time where the key would be added
 	 * @param Value	The value of the new key
+	 * @param KeyParams The keying parameters
 	 * @return True if the new key would be new data, false if duplicate
 	 */
-	bool NewKeyIsNewData(float Time, uint8 Value) const;
+	bool NewKeyIsNewData(float Time, uint8 Value, FKeyParams KeyParams) const;
 
 	/**
 	 * UMovieSceneSection interface 
diff --git a/Engine/Source/Runtime/MovieSceneTracks/Public/Sections/MovieSceneColorSection.h b/Engine/Source/Runtime/MovieSceneTracks/Public/Sections/MovieSceneColorSection.h
--- a/Engine/Source/Runtime/MovieSceneTracks/Public/Sections/MovieSceneColorSection.h
+++ b/Engine/Source/Runtime/MovieSceneTracks/Public/Sections/MovieSceneColorSection.h
@@ -31,17 +31,19 @@ public:
 	 *
 	 * @param Time	The location in time where the key should be added
 	 * @param Value	The value of the key
+     * @param KeyParams The keying parameters
 	 */
-	void AddKey( float Time, const FColorKey& Key );
+	void AddKey( float Time, const FColorKey& Key, FKeyParams KeyParams );
 	
 	/** 
 	 * Determines if a new key would be new data, or just a duplicate of existing data
 	 *
 	 * @param Time	The location in time where the key would be added
 	 * @param Value	The value of the new key
+     * @param KeyParams The keying parameters
 	 * @return True if the new key would be new data, false if duplicate
 	 */
-	bool NewKeyIsNewData(float Time, FLinearColor Value) const;
+	bool NewKeyIsNewData(float Time, FLinearColor Value, FKeyParams KeyParams) const;
 
 	/**
 	 * Gets the red color curve
@@ -74,9 +76,6 @@ public:
 	FRichCurve& GetAlphaCurve() { return AlphaCurve; }
 	const FRichCurve& GetAlphaCurve() const { return AlphaCurve; }
 
-private:
-	void AddKeyToNamedCurve(float Time, const FColorKey& Key);
-
 private:
 	/** Red curve data */
 	UPROPERTY()
diff --git a/Engine/Source/Runtime/MovieSceneTracks/Public/Sections/MovieSceneFloatSection.h b/Engine/Source/Runtime/MovieSceneTracks/Public/Sections/MovieSceneFloatSection.h
--- a/Engine/Source/Runtime/MovieSceneTracks/Public/Sections/MovieSceneFloatSection.h
+++ b/Engine/Source/Runtime/MovieSceneTracks/Public/Sections/MovieSceneFloatSection.h
@@ -24,17 +24,19 @@ public:
 	 *
 	 * @param Time	The location in time where the key should be added
 	 * @param Value	The value of the key
+	 * @param KeyParams The keying parameters
 	 */
-	void AddKey( float Time, float Value );
+	void AddKey( float Time, float Value, FKeyParams KeyParams );
 	
 	/** 
 	 * Determines if a new key would be new data, or just a duplicate of existing data
 	 *
 	 * @param Time	The location in time where the key would be added
 	 * @param Value	The value of the new key
+	 * @param KeyParams The keying parameters
 	 * @return True if the new key would be new data, false if duplicate
 	 */
-	bool NewKeyIsNewData(float Time, float Value) const;
+	bool NewKeyIsNewData(float Time, float Value, FKeyParams KeyParams) const;
 
 	/**
 	 * UMovieSceneSection interface 
diff --git a/Engine/Source/Runtime/MovieSceneTracks/Public/Sections/MovieSceneVectorSection.h b/Engine/Source/Runtime/MovieSceneTracks/Public/Sections/MovieSceneVectorSection.h
--- a/Engine/Source/Runtime/MovieSceneTracks/Public/Sections/MovieSceneVectorSection.h
+++ b/Engine/Source/Runtime/MovieSceneTracks/Public/Sections/MovieSceneVectorSection.h
@@ -32,17 +32,19 @@ public:
 	 *
 	 * @param Time	The location in time where the key should be added
 	 * @param Value	The value of the key
+	 * @param KeyParams The keying parameters
 	 */
-	void AddKey( float Time, FName CurveName, const FVector4& Value );
+	void AddKey( float Time, FName CurveName, const FVector4& Value, FKeyParams KeyParams );
 	
 	/** 
 	 * Determines if a new key would be new data, or just a duplicate of existing data
 	 *
 	 * @param Time	The location in time where the key would be added
 	 * @param Value	The value of the new key
+	 * @param KeyParams The keying parameters
 	 * @return True if the new key would be new data, false if duplicate
 	 */
-	bool NewKeyIsNewData(float Time, const FVector4& Value) const;
+	bool NewKeyIsNewData(float Time, const FVector4& Value, FKeyParams KeyParams) const;
 
 	/** Gets one of four curves in this section */
 	FRichCurve& GetCurve(const int32& Index) { return Curves[Index]; }
@@ -54,7 +56,7 @@ public:
 	int32 GetChannelsUsed() const {return ChannelsUsed;}
 
 private:
-	void AddKeyToNamedCurve(float Time, FName CurveName, const FVector4& Value);
+	void AddKeyToNamedCurve(float Time, FName CurveName, const FVector4& Value, FKeyParams KeyParams);
 
 private:
 	/** Vector t */
diff --git a/Engine/Source/Runtime/MovieSceneTracks/Public/Sections/MovieSceneVisibilitySection.h b/Engine/Source/Runtime/MovieSceneTracks/Public/Sections/MovieSceneVisibilitySection.h
--- a/Engine/Source/Runtime/MovieSceneTracks/Public/Sections/MovieSceneVisibilitySection.h
+++ b/Engine/Source/Runtime/MovieSceneTracks/Public/Sections/MovieSceneVisibilitySection.h
@@ -23,15 +23,17 @@ public:
 	 *
 	 * @param Time	The location in time where the key should be added
 	 * @param Value	The value of the key
+	 * @param KeyParams The keying parameters
 	 */
-	void AddKey( float Time, bool Value );
+	void AddKey( float Time, bool Value, FKeyParams KeyParams );
 
 	/** 
 	 * Determines if a new key would be new data, or just a duplicate of existing data
 	 *
 	 * @param Time	The location in time where the key would be added
 	 * @param Value	The value of the new key
+	 * @param KeyParams The keying parameters
 	 * @return True if the new key would be new data, false if duplicate
 	 */
-	bool NewKeyIsNewData(float Time, bool Value) const;
+	bool NewKeyIsNewData(float Time, bool Value, FKeyParams KeyParams) const;
 };
diff --git a/Engine/Source/Runtime/MovieSceneTracks/Public/Tracks/MovieScene3DTransformTrack.h b/Engine/Source/Runtime/MovieSceneTracks/Public/Tracks/MovieScene3DTransformTrack.h
--- a/Engine/Source/Runtime/MovieSceneTracks/Public/Tracks/MovieScene3DTransformTrack.h
+++ b/Engine/Source/Runtime/MovieSceneTracks/Public/Tracks/MovieScene3DTransformTrack.h
@@ -68,7 +68,7 @@ public:
 	 * @param PreviousTransform	The previous transform values (for comparing against new values).  If this is set only changed values will be keyed
 	 * 
 	 */
-	FTransformKey( float InKeyTime, const FTransformData& InNewTransform, const FTransformData& InPreviousTransform = FTransformData() )
+	FTransformKey( float InKeyTime, const FTransformData& InNewTransform, const FTransformData& InPreviousTransform = FTransformData())
 		: NewTransform( InNewTransform )
 		, PreviousTransform( InPreviousTransform )
 		, KeyTime( InKeyTime )
@@ -117,6 +117,10 @@ public:
 	 * @return Whether or not we should key on the given axis
 	 */
 	bool ShouldKeyAny() const;
+
+	/** Keying params */
+	FKeyParams KeyParams;
+
 private:
 	/**
 	 * Determines whether or not a vector component is different on the given axis
@@ -152,6 +156,7 @@ public:
 	 * @param ObjectHandle		Handle to the object(s) being changed
 	 * @param InKey				The transform key to add
 	 * @param bUnwindRotation	When true, the value will be treated like a rotation value in degrees, and will automatically be unwound to prevent flipping 360 degrees from the previous key 	 
+	 * @param bAddKeyEvenIfUnchanged Add a key even if its value is the same as the evaluated curve
 	 * @param KeyType           Key type (trans, rot, scale)
 	 * @return True if the key was successfully added.
 	 */
@@ -165,10 +170,15 @@ public:
 	 * @param OutTranslation	The evalulated translation component of the transform
 	 * @param OutRotation		The evalulated rotation component of the transform
 	 * @param OutScale		The evalulated scale component of the transform
-	 * @param OutHasTranslationKeys true if a translation key was evaluated 
-	 * @param OutHasRotationKeys 	true if a rotation key was evaluated
-	 * @param OutHasScaleKeys 	true if a scale key was evaluated
 	 * @return true if anything was evaluated. Note: if false is returned OutBool remains unchanged
 	 */
-	virtual bool Eval( float Position, float LastPosition, FVector& OutTranslation, FRotator& OutRotation, FVector& OutScale, TArray<bool>& OutHasTranslationKeys, TArray<bool>& OutHasRotationKeys, TArray<bool>& OutHasScaleKeys ) const;
+	virtual bool Eval( float Position, float LastPosition, FVector& OutTranslation, FRotator& OutRotation, FVector& OutScale ) const;
+
+	/*
+	* Get whether the track can be keyed at a particular time.
+	*
+    * @param InKey The transform key to add
+	* @return Whether the track can be keyed
+	*/
+	virtual bool CanKeyTrack(const FTransformKey& InKey) const;
 };
diff --git a/Engine/Source/Runtime/MovieSceneTracks/Public/Tracks/MovieSceneBoolTrack.h b/Engine/Source/Runtime/MovieSceneTracks/Public/Tracks/MovieSceneBoolTrack.h
--- a/Engine/Source/Runtime/MovieSceneTracks/Public/Tracks/MovieSceneBoolTrack.h
+++ b/Engine/Source/Runtime/MovieSceneTracks/Public/Tracks/MovieSceneBoolTrack.h
@@ -25,9 +25,10 @@ public:
 	 *
 	 * @param Time				The time relative to the owning movie scene where the section should be
 	 * @param Value				The value of the key
+	 * @param KeyParams         The keying parameters
 	 * @return True if the key was successfully added.
 	 */
-	virtual bool AddKeyToSection( float Time, bool Value );
+	virtual bool AddKeyToSection( float Time, bool Value, FKeyParams KeyParams );
 
 	/**
 	 * Evaluates the track at the playback position
diff --git a/Engine/Source/Runtime/MovieSceneTracks/Public/Tracks/MovieSceneByteTrack.h b/Engine/Source/Runtime/MovieSceneTracks/Public/Tracks/MovieSceneByteTrack.h
--- a/Engine/Source/Runtime/MovieSceneTracks/Public/Tracks/MovieSceneByteTrack.h
+++ b/Engine/Source/Runtime/MovieSceneTracks/Public/Tracks/MovieSceneByteTrack.h
@@ -25,9 +25,10 @@ public:
 	 *
 	 * @param Time				The time relative to the owning movie scene where the section should be
 	 * @param Value				The value of the key
+	 * @param KeyParams         The keying parameters
 	 * @return True if the key was successfully added.
 	 */
-	virtual bool AddKeyToSection( float Time, uint8 Value );
+	virtual bool AddKeyToSection( float Time, uint8 Value, FKeyParams KeyParams );
 
 	/**
 	 * Evaluates the track at the playback position
diff --git a/Engine/Source/Runtime/MovieSceneTracks/Public/Tracks/MovieSceneColorTrack.h b/Engine/Source/Runtime/MovieSceneTracks/Public/Tracks/MovieSceneColorTrack.h
--- a/Engine/Source/Runtime/MovieSceneTracks/Public/Tracks/MovieSceneColorTrack.h
+++ b/Engine/Source/Runtime/MovieSceneTracks/Public/Tracks/MovieSceneColorTrack.h
@@ -11,7 +11,6 @@ struct FColorKey
 {
 	FLinearColor Value;
 	FName CurveName;
-	bool bAddKeyEvenIfUnchanged;
 	bool bIsSlateColor;
 };
 
@@ -33,9 +32,10 @@ public:
 	 *
 	 * @param Time				The time relative to the owning movie scene where the section should be
 	 * @param Value				The value of the key
+	 * @param KeyParams         The keying parameters
 	 * @return True if the key was successfully added.
 	 */
-	virtual bool AddKeyToSection( float Time, const FColorKey& Key );
+	virtual bool AddKeyToSection( float Time, const FColorKey& Key, FKeyParams KeyParams );
 	
 	/**
 	 * Evaluates the track at the playback position
diff --git a/Engine/Source/Runtime/MovieSceneTracks/Public/Tracks/MovieSceneEventTrack.h b/Engine/Source/Runtime/MovieSceneTracks/Public/Tracks/MovieSceneEventTrack.h
--- a/Engine/Source/Runtime/MovieSceneTracks/Public/Tracks/MovieSceneEventTrack.h
+++ b/Engine/Source/Runtime/MovieSceneTracks/Public/Tracks/MovieSceneEventTrack.h
@@ -24,8 +24,9 @@ public:
 	 *
 	 * @param Time The time at which the event should be triggered.
 	 * @param EventName The name of the event to be triggered.
+	 * @param KeyParams The keying parameters
 	 */
-	bool AddKeyToSection(float Time, FName EventName);
+	bool AddKeyToSection(float Time, FName EventName, FKeyParams KeyParams);
 
 	/**
 	 * Trigger the events that fall into the given time range.
diff --git a/Engine/Source/Runtime/MovieSceneTracks/Public/Tracks/MovieSceneFloatTrack.h b/Engine/Source/Runtime/MovieSceneTracks/Public/Tracks/MovieSceneFloatTrack.h
--- a/Engine/Source/Runtime/MovieSceneTracks/Public/Tracks/MovieSceneFloatTrack.h
+++ b/Engine/Source/Runtime/MovieSceneTracks/Public/Tracks/MovieSceneFloatTrack.h
@@ -25,9 +25,10 @@ public:
 	 *
 	 * @param Time				The time relative to the owning movie scene where the section should be
 	 * @param Value				The value of the key
+	 * @param KeyParams         The keying parameters
 	 * @return True if the key was successfully added.
 	 */
-	virtual bool AddKeyToSection( float Time, float Value );
+	virtual bool AddKeyToSection( float Time, float Value, FKeyParams KeyParams );
 
 	/**
 	 * Evaluates the track at the playback position
diff --git a/Engine/Source/Runtime/MovieSceneTracks/Public/Tracks/MovieSceneVectorTrack.h b/Engine/Source/Runtime/MovieSceneTracks/Public/Tracks/MovieSceneVectorTrack.h
--- a/Engine/Source/Runtime/MovieSceneTracks/Public/Tracks/MovieSceneVectorTrack.h
+++ b/Engine/Source/Runtime/MovieSceneTracks/Public/Tracks/MovieSceneVectorTrack.h
@@ -12,14 +12,13 @@
 struct MOVIESCENETRACKS_API FVectorKey
 {
 	FVectorKey() { }
-	FVectorKey( const FVector& InValue, FName InCurveName, bool InbAddKeyEvenIfUnChanged );
-	FVectorKey( const FVector2D& InValue, FName InCurveName, bool InbAddKeyEvenIfUnChanged );
-	FVectorKey( const FVector4& InValue, FName InCurveName, bool InbAddKeyEvenIfUnChanged );
+	FVectorKey( const FVector& InValue, FName InCurveName );
+	FVectorKey( const FVector2D& InValue, FName InCurveName );
+	FVectorKey( const FVector4& InValue, FName InCurveName );
 
 	FVector4 Value;
 	uint32 ChannelsUsed;
 	FName CurveName;
-	bool bAddKeyEvenIfUnchanged;
 };
 
 /**
@@ -39,9 +38,10 @@ public:
 	 *
 	 * @param Time				The time relative to the owning movie scene where the section should be
 	 * @param Key				The vector key to add
+	 * @param KeyParams         The keying parameters 
 	 * @return True if the key was successfully added.
 	 */
-	virtual bool AddKeyToSection( float Time, const FVectorKey& Key );
+	virtual bool AddKeyToSection( float Time, const FVectorKey& Key, FKeyParams KeyParams );
 	
 	/**
 	 * Evaluates the track at the playback position
@@ -57,7 +57,7 @@ public:
 	int32 GetNumChannelsUsed() const { return NumChannelsUsed; }
 
 private:
-	virtual bool AddKeyToSection( float Time, const FVector4& InKey, int32 InChannelsUsed, FName CurveName, bool bAddKeyEvenIfUnchanged );
+	virtual bool AddKeyToSection( float Time, const FVector4& InKey, int32 InChannelsUsed, FName CurveName, FKeyParams KeyParams );
 private:
 	/** The number of channels used by the vector (2,3, or 4) */
 	UPROPERTY()
diff --git a/Engine/Source/Runtime/MovieSceneTracks/Public/Tracks/MovieSceneVisibilityTrack.h b/Engine/Source/Runtime/MovieSceneTracks/Public/Tracks/MovieSceneVisibilityTrack.h
--- a/Engine/Source/Runtime/MovieSceneTracks/Public/Tracks/MovieSceneVisibilityTrack.h
+++ b/Engine/Source/Runtime/MovieSceneTracks/Public/Tracks/MovieSceneVisibilityTrack.h
@@ -26,7 +26,8 @@ public:
 	 *
 	 * @param Time				The time relative to the owning movie scene where the section should be
 	 * @param Value				The value of the key
+	 * @param KeyParams         The keying parameters
 	 * @return True if the key was successfully added.
 	 */
-	virtual bool AddKeyToSection( float Time, bool Value );
+	virtual bool AddKeyToSection( float Time, bool Value, FKeyParams KeyParams );
 };
diff --git a/Engine/Source/Runtime/UMG/Private/Animation/MovieScene2DTransformSection.cpp b/Engine/Source/Runtime/UMG/Private/Animation/MovieScene2DTransformSection.cpp
--- a/Engine/Source/Runtime/UMG/Private/Animation/MovieScene2DTransformSection.cpp
+++ b/Engine/Source/Runtime/UMG/Private/Animation/MovieScene2DTransformSection.cpp
@@ -128,23 +128,52 @@ FWidgetTransform UMovieScene2DTransformSection::Eval( float Position, const FWid
 }
 
 
-bool UMovieScene2DTransformSection::NewKeyIsNewData( float Time, const FWidgetTransform& Transform ) const
+bool UMovieScene2DTransformSection::NewKeyIsNewData( float Time, const FWidgetTransform& Transform, FKeyParams KeyParams ) const
 {
-	bool bIsNewData = false;
+	bool bHasEmptyKeys = false;
 	for(int32 Axis = 0; Axis < 2; ++Axis)
 	{
-		bIsNewData |= Translation[Axis].GetNumKeys() == 0;
-		bIsNewData |= Scale[Axis].GetNumKeys() == 0;
-		bIsNewData |= Shear[Axis].GetNumKeys() == 0;
+		bHasEmptyKeys = bHasEmptyKeys ||
+			Translation[Axis].GetNumKeys() == 0 ||
+			Scale[Axis].GetNumKeys() == 0 ||
+			Shear[Axis].GetNumKeys() == 0;
 	}
 
-	bIsNewData |= Rotation.GetNumKeys() == 0;
+	bHasEmptyKeys |= Rotation.GetNumKeys() == 0;
 
-	return bIsNewData || Eval(Time,Transform) != Transform;
+	if (bHasEmptyKeys || Eval(Time, Transform) != Transform)
+	{
+		// Don't add a keyframe if there are existing keys and auto key is not enabled.
+		bool bTxKeyExists = Translation[0].IsKeyHandleValid(Translation[0].FindKey(Time));
+		bool bTyKeyExists = Translation[1].IsKeyHandleValid(Translation[1].FindKey(Time));
+
+		bool bRKeyExists = Rotation.IsKeyHandleValid(Rotation.FindKey(Time));
+
+		bool bHxKeyExists = Shear[0].IsKeyHandleValid(Shear[0].FindKey(Time));
+		bool bHyKeyExists = Shear[1].IsKeyHandleValid(Shear[1].FindKey(Time));
+
+		bool bSxKeyExists = Scale[0].IsKeyHandleValid(Scale[0].FindKey(Time));
+		bool bSyKeyExists = Scale[1].IsKeyHandleValid(Scale[1].FindKey(Time));
+
+		if ( !(!bTxKeyExists && !KeyParams.bAutoKeying && Translation[0].GetNumKeys() > 0) ||
+			 !(!bTyKeyExists && !KeyParams.bAutoKeying && Translation[1].GetNumKeys() > 0) ||
+			 !(!bRKeyExists && !KeyParams.bAutoKeying && Rotation.GetNumKeys() > 0) ||
+			 !(!bHxKeyExists && !KeyParams.bAutoKeying && Shear[0].GetNumKeys() > 0) ||
+			 !(!bHyKeyExists && !KeyParams.bAutoKeying && Shear[1].GetNumKeys() > 0) ||
+			 !(!bSxKeyExists && !KeyParams.bAutoKeying && Scale[0].GetNumKeys() > 0) ||
+			 !(!bSyKeyExists && !KeyParams.bAutoKeying && Scale[1].GetNumKeys() > 0) )
+		{
+			return true;
+		}
+	}
+
+	return false;
 }
 
-void UMovieScene2DTransformSection::AddKeyToNamedCurve(float Time, const F2DTransformKey& TransformKey)
+void UMovieScene2DTransformSection::AddKey( float Time, const struct F2DTransformKey& TransformKey, FKeyParams KeyParams )
 {
+	Modify();
+
 	static FName TranslationName("Translation");
 	static FName ScaleName("Scale");
 	static FName ShearName("Shear");
@@ -152,45 +181,56 @@ void UMovieScene2DTransformSection::AddKeyToNamedCurve(float Time, const F2DTran
 
 	FName CurveName = TransformKey.CurveName;
 
-	if(CurveName == TranslationName)
+	bool bTxKeyExists = Translation[0].IsKeyHandleValid(Translation[0].FindKey(Time));
+	bool bTyKeyExists = Translation[1].IsKeyHandleValid(Translation[1].FindKey(Time));
+
+	bool bRKeyExists = Rotation.IsKeyHandleValid(Rotation.FindKey(Time));
+
+	bool bHxKeyExists = Shear[0].IsKeyHandleValid(Shear[0].FindKey(Time));
+	bool bHyKeyExists = Shear[1].IsKeyHandleValid(Shear[1].FindKey(Time));
+
+	bool bSxKeyExists = Scale[0].IsKeyHandleValid(Scale[0].FindKey(Time));
+	bool bSyKeyExists = Scale[1].IsKeyHandleValid(Scale[1].FindKey(Time));
+
+	if ( (CurveName == NAME_None || CurveName == TranslationName) &&
+			KeyParams.bAddKeyEvenIfUnchanged || !(!bTxKeyExists && !KeyParams.bAutoKeying && Translation[0].GetNumKeys() > 0) )
 	{
-		AddKeyToCurve(Translation[0], Time, TransformKey.Value.Translation.X);
-		AddKeyToCurve(Translation[1], Time, TransformKey.Value.Translation.Y);
+		AddKeyToCurve(Translation[0], Time, TransformKey.Value.Translation.X, KeyParams);
 	}
-	else if(CurveName == ScaleName)
+
+	if ( (CurveName == NAME_None || CurveName == TranslationName) &&
+			KeyParams.bAddKeyEvenIfUnchanged || !(!bTyKeyExists && !KeyParams.bAutoKeying && Translation[1].GetNumKeys() > 0) )
 	{
-		AddKeyToCurve(Scale[0], Time, TransformKey.Value.Scale.X);
-		AddKeyToCurve(Scale[1], Time, TransformKey.Value.Scale.Y);
+		AddKeyToCurve(Translation[1], Time, TransformKey.Value.Translation.Y, KeyParams);
 	}
-	else if(CurveName == ShearName)
+
+	if ( (CurveName == NAME_None || CurveName == ScaleName) &&
+			KeyParams.bAddKeyEvenIfUnchanged || !(!bSxKeyExists && !KeyParams.bAutoKeying && Scale[0].GetNumKeys() > 0) )
 	{
-		AddKeyToCurve(Shear[0], Time, TransformKey.Value.Shear.X);
-		AddKeyToCurve(Shear[1], Time, TransformKey.Value.Shear.Y);
+		AddKeyToCurve(Scale[0], Time, TransformKey.Value.Scale.X, KeyParams);
 	}
-	else if(CurveName == AngleName)
+
+	if ( (CurveName == NAME_None || CurveName == ScaleName) &&
+			KeyParams.bAddKeyEvenIfUnchanged || !(!bSyKeyExists && !KeyParams.bAutoKeying && Scale[1].GetNumKeys() > 0) )
 	{
-		AddKeyToCurve(Rotation, Time, TransformKey.Value.Angle);
+		AddKeyToCurve(Scale[1], Time, TransformKey.Value.Scale.Y, KeyParams);
+	}
+		
+	if ( (CurveName == NAME_None || CurveName == ShearName) &&
+			KeyParams.bAddKeyEvenIfUnchanged || !(!bHxKeyExists && !KeyParams.bAutoKeying && Shear[0].GetNumKeys() > 0) )
+	{
+		AddKeyToCurve(Shear[0], Time, TransformKey.Value.Shear.X, KeyParams);
 	}
-}
-
-
-
-void UMovieScene2DTransformSection::AddKey( float Time, const struct F2DTransformKey& TransformKey )
-{
-	Modify();
 
-	if(TransformKey.CurveName == NAME_None)
+	if ( (CurveName == NAME_None || CurveName == ShearName) &&
+			KeyParams.bAddKeyEvenIfUnchanged || !(!bHyKeyExists && !KeyParams.bAutoKeying && Shear[1].GetNumKeys() > 0) )
 	{
-		AddKeyToCurve(Translation[0], Time, TransformKey.Value.Translation.X);
-		AddKeyToCurve(Translation[1], Time, TransformKey.Value.Translation.Y);
-		AddKeyToCurve(Scale[0], Time, TransformKey.Value.Scale.X);
-		AddKeyToCurve(Scale[1], Time, TransformKey.Value.Scale.Y);
-		AddKeyToCurve(Shear[0], Time, TransformKey.Value.Shear.X);
-		AddKeyToCurve(Shear[1], Time, TransformKey.Value.Shear.Y);
-		AddKeyToCurve(Rotation, Time, TransformKey.Value.Angle);
+		AddKeyToCurve(Shear[1], Time, TransformKey.Value.Shear.Y, KeyParams);
 	}
-	else
+
+	if ( (CurveName == NAME_None || CurveName == AngleName) &&
+			KeyParams.bAddKeyEvenIfUnchanged || !(!bRKeyExists && !KeyParams.bAutoKeying && Rotation.GetNumKeys() > 0) )
 	{
-		AddKeyToNamedCurve(Time, TransformKey);
+		AddKeyToCurve(Rotation, Time, TransformKey.Value.Angle, KeyParams);
 	}
 }
diff --git a/Engine/Source/Runtime/UMG/Private/Animation/MovieScene2DTransformTrack.cpp b/Engine/Source/Runtime/UMG/Private/Animation/MovieScene2DTransformTrack.cpp
--- a/Engine/Source/Runtime/UMG/Private/Animation/MovieScene2DTransformTrack.cpp
+++ b/Engine/Source/Runtime/UMG/Private/Animation/MovieScene2DTransformTrack.cpp
@@ -39,16 +39,16 @@ bool UMovieScene2DTransformTrack::Eval(float Position, float LastPosition, FWidg
 }
 
 
-bool UMovieScene2DTransformTrack::AddKeyToSection(float Time, const F2DTransformKey& TransformKey)
+bool UMovieScene2DTransformTrack::AddKeyToSection(float Time, const F2DTransformKey& TransformKey, FKeyParams KeyParams)
 {
 	const UMovieSceneSection* NearestSection = MovieSceneHelpers::FindNearestSectionAtTime(Sections, Time);
-	if(!NearestSection || TransformKey.bAddKeyEvenIfUnchanged || CastChecked<UMovieScene2DTransformSection>(NearestSection)->NewKeyIsNewData(Time, TransformKey.Value))
+	if(!NearestSection || KeyParams.bAddKeyEvenIfUnchanged || CastChecked<UMovieScene2DTransformSection>(NearestSection)->NewKeyIsNewData(Time, TransformKey.Value, KeyParams))
 	{
 		Modify();
 
 		UMovieScene2DTransformSection* NewSection = CastChecked<UMovieScene2DTransformSection>(FindOrAddSection(Time));
 
-		NewSection->AddKey(Time, TransformKey);
+		NewSection->AddKey(Time, TransformKey, KeyParams);
 
 		return true;
 	}
diff --git a/Engine/Source/Runtime/UMG/Private/Animation/MovieSceneMarginSection.cpp b/Engine/Source/Runtime/UMG/Private/Animation/MovieSceneMarginSection.cpp
--- a/Engine/Source/Runtime/UMG/Private/Animation/MovieSceneMarginSection.cpp
+++ b/Engine/Source/Runtime/UMG/Private/Animation/MovieSceneMarginSection.cpp
@@ -74,56 +74,71 @@ FMargin UMovieSceneMarginSection::Eval( float Position, const FMargin& DefaultVa
 					BottomCurve.Eval(Position, DefaultValue.Bottom));
 }
 
-void UMovieSceneMarginSection::AddKey( float Time, const FMarginKey& MarginKey )
+void UMovieSceneMarginSection::AddKey( float Time, const FMarginKey& MarginKey, FKeyParams KeyParams )
 {
 	Modify();
 
-	if( MarginKey.CurveName == NAME_None )
+	static FName LeftName("Left");
+	static FName TopName("Top");
+	static FName RightName("Right");
+	static FName BottomName("Bottom");
+
+	FName CurveName = MarginKey.CurveName;
+
+	bool bTopKeyExists = TopCurve.IsKeyHandleValid(TopCurve.FindKey(Time));
+	bool bLeftKeyExists = LeftCurve.IsKeyHandleValid(LeftCurve.FindKey(Time));
+	bool bRightKeyExists = RightCurve.IsKeyHandleValid(RightCurve.FindKey(Time));
+	bool bBottomKeyExists = BottomCurve.IsKeyHandleValid(BottomCurve.FindKey(Time));
+
+	if ( (CurveName == NAME_None || CurveName == LeftName) &&
+			KeyParams.bAddKeyEvenIfUnchanged || !(!bLeftKeyExists && !KeyParams.bAutoKeying && LeftCurve.GetNumKeys() > 0) )
 	{
-		AddKeyToCurve(LeftCurve, Time, MarginKey.Value.Left);
-		AddKeyToCurve(TopCurve, Time, MarginKey.Value.Top);
-		AddKeyToCurve(RightCurve, Time, MarginKey.Value.Right);
-		AddKeyToCurve(BottomCurve, Time, MarginKey.Value.Bottom);
+		AddKeyToCurve(LeftCurve, Time, MarginKey.Value.Left, KeyParams);
 	}
-	else
+
+	if ( (CurveName == NAME_None || CurveName == TopName) &&
+			KeyParams.bAddKeyEvenIfUnchanged || !(!bTopKeyExists && !KeyParams.bAutoKeying && TopCurve.GetNumKeys() > 0) )
 	{
-		AddKeyToNamedCurve( Time, MarginKey );
+		AddKeyToCurve(TopCurve, Time, MarginKey.Value.Top, KeyParams);
+	}
+
+	if ( (CurveName == NAME_None || CurveName == RightName) &&
+			KeyParams.bAddKeyEvenIfUnchanged || !(!bRightKeyExists && !KeyParams.bAutoKeying && RightCurve.GetNumKeys() > 0) )
+	{
+		AddKeyToCurve(RightCurve, Time, MarginKey.Value.Right, KeyParams);
+	}
+
+	if ( (CurveName == NAME_None || CurveName == BottomName) &&
+			KeyParams.bAddKeyEvenIfUnchanged || !(!bBottomKeyExists && !KeyParams.bAutoKeying && BottomCurve.GetNumKeys() > 0) )
+	{
+		AddKeyToCurve(BottomCurve, Time, MarginKey.Value.Bottom, KeyParams);
 	}
 }
 
-bool UMovieSceneMarginSection::NewKeyIsNewData(float Time, const FMargin& Value) const
+bool UMovieSceneMarginSection::NewKeyIsNewData(float Time, const FMargin& Value, FKeyParams KeyParams) const
 {
-	return TopCurve.GetNumKeys() == 0 ||
+	bool bHasEmptyKeys = TopCurve.GetNumKeys() == 0 ||
 		LeftCurve.GetNumKeys() == 0 ||
 		RightCurve.GetNumKeys() == 0 ||
 		BottomCurve.GetNumKeys() == 0 ||
 		(Eval(Time,Value) != Value);
-}
 
-void UMovieSceneMarginSection::AddKeyToNamedCurve( float Time, const FMarginKey& MarginKey )
-{
-	static FName Left("Left");
-	static FName Top("Top");
-	static FName Right("Right");
-	static FName Bottom("Bottom");
-
-	FName CurveName = MarginKey.CurveName;
-
-	if( CurveName == Left )
-	{
-		AddKeyToCurve(LeftCurve, Time, MarginKey.Value.Left);
-	}
-	else if( CurveName == Top )
-	{
-		AddKeyToCurve(TopCurve, Time, MarginKey.Value.Top);
-	}
-	else if( CurveName == Right )
-	{
-		AddKeyToCurve(RightCurve, Time, MarginKey.Value.Right);
-	}
-	else if( CurveName == Bottom )
+	if (bHasEmptyKeys)
 	{
-		AddKeyToCurve(BottomCurve, Time, MarginKey.Value.Bottom);
+		// Don't add a keyframe if there are existing keys and auto key is not enabled.
+		bool bTopKeyExists = TopCurve.IsKeyHandleValid(TopCurve.FindKey(Time));
+		bool bLeftKeyExists = LeftCurve.IsKeyHandleValid(LeftCurve.FindKey(Time));
+		bool bRightKeyExists = RightCurve.IsKeyHandleValid(RightCurve.FindKey(Time));
+		bool bBottomKeyExists = BottomCurve.IsKeyHandleValid(BottomCurve.FindKey(Time));
+
+		if ( !(!bTopKeyExists && !KeyParams.bAutoKeying && TopCurve.GetNumKeys() > 0) ||
+			 !(!bLeftKeyExists && !KeyParams.bAutoKeying && LeftCurve.GetNumKeys() > 0) ||
+			 !(!bRightKeyExists && !KeyParams.bAutoKeying && RightCurve.GetNumKeys() > 0) ||
+			 !(!bBottomKeyExists && !KeyParams.bAutoKeying && BottomCurve.GetNumKeys() > 0) )
+		{
+			return true;
+		}
 	}
+	return false;
 }
 
diff --git a/Engine/Source/Runtime/UMG/Private/Animation/MovieSceneMarginTrack.cpp b/Engine/Source/Runtime/UMG/Private/Animation/MovieSceneMarginTrack.cpp
--- a/Engine/Source/Runtime/UMG/Private/Animation/MovieSceneMarginTrack.cpp
+++ b/Engine/Source/Runtime/UMG/Private/Animation/MovieSceneMarginTrack.cpp
@@ -23,16 +23,16 @@ TSharedPtr<IMovieSceneTrackInstance> UMovieSceneMarginTrack::CreateInstance()
 }
 
 
-bool UMovieSceneMarginTrack::AddKeyToSection( float Time, const FMarginKey& MarginKey )
+bool UMovieSceneMarginTrack::AddKeyToSection( float Time, const FMarginKey& MarginKey, FKeyParams KeyParams )
 {
 	const UMovieSceneSection* NearestSection = MovieSceneHelpers::FindNearestSectionAtTime( Sections, Time );
-	if (!NearestSection || MarginKey.bAddKeyEvenIfUnchanged || CastChecked<UMovieSceneMarginSection>(NearestSection)->NewKeyIsNewData(Time, MarginKey.Value) )
+	if (!NearestSection || KeyParams.bAddKeyEvenIfUnchanged || CastChecked<UMovieSceneMarginSection>(NearestSection)->NewKeyIsNewData(Time, MarginKey.Value, KeyParams) )
 	{
 		Modify();
 
 		UMovieSceneMarginSection* NewSection = CastChecked<UMovieSceneMarginSection>( FindOrAddSection( Time ) );
 
-		NewSection->AddKey( Time, MarginKey );
+		NewSection->AddKey( Time, MarginKey, KeyParams );
 
 		return true;
 	}
diff --git a/Engine/Source/Runtime/UMG/Public/Animation/MovieScene2DTransformSection.h b/Engine/Source/Runtime/UMG/Public/Animation/MovieScene2DTransformSection.h
--- a/Engine/Source/Runtime/UMG/Public/Animation/MovieScene2DTransformSection.h
+++ b/Engine/Source/Runtime/UMG/Public/Animation/MovieScene2DTransformSection.h
@@ -32,11 +32,9 @@ public:
 
 	FWidgetTransform Eval( float Position, const FWidgetTransform& DefaultValue ) const;
 
-	bool NewKeyIsNewData( float Time, const FWidgetTransform& Transform ) const;
+	bool NewKeyIsNewData( float Time, const FWidgetTransform& Transform, FKeyParams KeyParams ) const;
 
-	void AddKey( float Time, const struct F2DTransformKey& TransformKey );
-private:
-	void AddKeyToNamedCurve(float Time, const F2DTransformKey& TransformKey);
+	void AddKey( float Time, const struct F2DTransformKey& TransformKey, FKeyParams KeyParams );
 private:
 	/** Translation curves*/
 	UPROPERTY()
diff --git a/Engine/Source/Runtime/UMG/Public/Animation/MovieScene2DTransformTrack.h b/Engine/Source/Runtime/UMG/Public/Animation/MovieScene2DTransformTrack.h
--- a/Engine/Source/Runtime/UMG/Public/Animation/MovieScene2DTransformTrack.h
+++ b/Engine/Source/Runtime/UMG/Public/Animation/MovieScene2DTransformTrack.h
@@ -10,7 +10,6 @@ struct F2DTransformKey
 {
 	FWidgetTransform Value;
 	FName CurveName;
-	bool bAddKeyEvenIfUnchanged;
 };
 
 
@@ -37,7 +36,14 @@ public:
 	 */
 	bool Eval( float Position, float LastPostion, FWidgetTransform& InOutTransform ) const;
 
-
-	UMG_API bool AddKeyToSection( float Time, const F2DTransformKey& TransformKey );
+	/**
+	 * Adds a key to a section.  Will create the section if it doesn't exist
+	 *
+	 * @param Time				The time relative to the owning movie scene where the section should be
+	 * @param Value				The value of the key
+	 * @param KeyParams         The keying parameters 
+	 * @return True if the key was successfully added.
+	 */
+	UMG_API bool AddKeyToSection( float Time, const F2DTransformKey& TransformKey, FKeyParams KeyParams );
 
 };
\ No newline at end of file
diff --git a/Engine/Source/Runtime/UMG/Public/Animation/MovieSceneMarginSection.h b/Engine/Source/Runtime/UMG/Public/Animation/MovieSceneMarginSection.h
--- a/Engine/Source/Runtime/UMG/Public/Animation/MovieSceneMarginSection.h
+++ b/Engine/Source/Runtime/UMG/Public/Animation/MovieSceneMarginSection.h
@@ -33,17 +33,19 @@ public:
 	 *
 	 * @param Time	The location in time where the key should be added
 	 * @param Value	The value of the key
+	 * @param KeyParams The keying parameters
 	 */
-	void AddKey( float Time, const FMarginKey& MarginKey );
+	void AddKey( float Time, const FMarginKey& MarginKey, FKeyParams KeyParams );
 	
 	/** 
 	 * Determines if a new key would be new data, or just a duplicate of existing data
 	 *
 	 * @param Time	The location in time where the key would be added
 	 * @param Value	The value of the new key
+	 * @param KeyParams The keying parameters
 	 * @return True if the new key would be new data, false if duplicate
 	 */
-	bool NewKeyIsNewData(float Time, const FMargin& Value) const;
+	bool NewKeyIsNewData(float Time, const FMargin& Value, FKeyParams KeyParams) const;
 
 	/**
 	 * Gets the top curve
@@ -76,9 +78,6 @@ public:
 	FRichCurve& GetBottomCurve() { return BottomCurve; }
 	const FRichCurve& GetBottomCurve() const { return BottomCurve; }
 
-private:
-	void AddKeyToNamedCurve( float Time, const FMarginKey& MarginKey );
-
 private:
 	/** Red curve data */
 	UPROPERTY()
diff --git a/Engine/Source/Runtime/UMG/Public/Animation/MovieSceneMarginTrack.h b/Engine/Source/Runtime/UMG/Public/Animation/MovieSceneMarginTrack.h
--- a/Engine/Source/Runtime/UMG/Public/Animation/MovieSceneMarginTrack.h
+++ b/Engine/Source/Runtime/UMG/Public/Animation/MovieSceneMarginTrack.h
@@ -9,7 +9,6 @@ struct FMarginKey
 {
 	FMargin Value;
 	FName CurveName;
-	bool bAddKeyEvenIfUnchanged;
 };
 
 /**
@@ -30,9 +29,10 @@ public:
 	 *
 	 * @param Time				The time relative to the owning movie scene where the section should be
 	 * @param Value				The value of the key
+	 * @param KeyParams         The keying parameters
 	 * @return True if the key was successfully added.
 	 */
-	UMG_API bool AddKeyToSection( float Time, const FMarginKey& MarginKey );
+	UMG_API bool AddKeyToSection( float Time, const FMarginKey& MarginKey, FKeyParams KeyParams );
 
 	/**
 	 * Evaluates the track at the playback position
