MovieSceneTools: Cosmetic changes

diff --git a/Engine/Source/Editor/MovieSceneTools/Private/MovieSceneToolsModule.cpp b/Engine/Source/Editor/MovieSceneTools/Private/MovieSceneToolsModule.cpp
--- a/Engine/Source/Editor/MovieSceneTools/Private/MovieSceneToolsModule.cpp
+++ b/Engine/Source/Editor/MovieSceneTools/Private/MovieSceneToolsModule.cpp
@@ -28,54 +28,60 @@
 /**
  * Implements the MovieSceneTools module.
  */
-class FMovieSceneToolsModule : public IMovieSceneTools
+class FMovieSceneToolsModule
+	: public IMovieSceneTools
 {
+public:
+
+	// IModuleInterface interface
+
 	virtual void StartupModule() override
 	{
-		// Register with the sequencer module that we provide auto-key handlers.
+		// register with the sequencer module that we provide auto-key handlers
 		ISequencerModule& SequencerModule = FModuleManager::Get().LoadModuleChecked<ISequencerModule>( "Sequencer" );
 
-		BoolPropertyTrackEditorCreateTrackEditorDelegateHandle   = SequencerModule.RegisterTrackEditor_Handle( FOnCreateTrackEditor::CreateStatic( &FBoolPropertyTrackEditor::CreateTrackEditor ) );
+		BoolPropertyTrackEditorCreateTrackEditorDelegateHandle = SequencerModule.RegisterTrackEditor_Handle( FOnCreateTrackEditor::CreateStatic( &FBoolPropertyTrackEditor::CreateTrackEditor ) );
 		BytePropertyTrackEditorCreateTrackEditorDelegateHandle = SequencerModule.RegisterTrackEditor_Handle( FOnCreateTrackEditor::CreateStatic( &FBytePropertyTrackEditor::CreateTrackEditor ) );
-		ColorPropertyTrackEditorCreateTrackEditorDelegateHandle  = SequencerModule.RegisterTrackEditor_Handle( FOnCreateTrackEditor::CreateStatic( &FColorPropertyTrackEditor::CreateTrackEditor ) );
+		ColorPropertyTrackEditorCreateTrackEditorDelegateHandle = SequencerModule.RegisterTrackEditor_Handle( FOnCreateTrackEditor::CreateStatic( &FColorPropertyTrackEditor::CreateTrackEditor ) );
 		FloatPropertyTrackEditorCreateTrackEditorDelegateHandle = SequencerModule.RegisterTrackEditor_Handle( FOnCreateTrackEditor::CreateStatic( &FFloatPropertyTrackEditor::CreateTrackEditor ) );
 		VectorPropertyTrackEditorCreateTrackEditorDelegateHandle = SequencerModule.RegisterTrackEditor_Handle( FOnCreateTrackEditor::CreateStatic( &FVectorPropertyTrackEditor::CreateTrackEditor ) );
 		VisibilityPropertyTrackEditorCreateTrackEditorDelegateHandle = SequencerModule.RegisterTrackEditor_Handle( FOnCreateTrackEditor::CreateStatic( &FVisibilityPropertyTrackEditor::CreateTrackEditor ) );
 
-		TransformTrackEditorCreateTrackEditorDelegateHandle      = SequencerModule.RegisterTrackEditor_Handle( FOnCreateTrackEditor::CreateStatic( &F3DTransformTrackEditor::CreateTrackEditor ) );
-		EditorCreateTrackEditorDelegateHandle                    = SequencerModule.RegisterTrackEditor_Handle( FOnCreateTrackEditor::CreateStatic( &FShotTrackEditor::CreateTrackEditor ) );
-		SubMovieSceneTrackEditorCreateTrackEditorDelegateHandle  = SequencerModule.RegisterTrackEditor_Handle( FOnCreateTrackEditor::CreateStatic( &FSubMovieSceneTrackEditor::CreateTrackEditor ) );
-		AudioTrackEditorCreateTrackEditorDelegateHandle          = SequencerModule.RegisterTrackEditor_Handle( FOnCreateTrackEditor::CreateStatic( &FAudioTrackEditor::CreateTrackEditor ) );
-		AnimationTrackEditorCreateTrackEditorDelegateHandle      = SequencerModule.RegisterTrackEditor_Handle( FOnCreateTrackEditor::CreateStatic( &FSkeletalAnimationTrackEditor::CreateTrackEditor ) );
-		ParticleTrackEditorCreateTrackEditorDelegateHandle       = SequencerModule.RegisterTrackEditor_Handle( FOnCreateTrackEditor::CreateStatic( &FParticleTrackEditor::CreateTrackEditor ) );
-		AttachTrackEditorCreateTrackEditorDelegateHandle         = SequencerModule.RegisterTrackEditor_Handle( FOnCreateTrackEditor::CreateStatic( &F3DAttachTrackEditor::CreateTrackEditor ) );
-		PathTrackEditorCreateTrackEditorDelegateHandle           = SequencerModule.RegisterTrackEditor_Handle( FOnCreateTrackEditor::CreateStatic( &F3DPathTrackEditor::CreateTrackEditor ) );
+		AnimationTrackEditorCreateTrackEditorDelegateHandle = SequencerModule.RegisterTrackEditor_Handle( FOnCreateTrackEditor::CreateStatic( &FSkeletalAnimationTrackEditor::CreateTrackEditor ) );
+		AttachTrackEditorCreateTrackEditorDelegateHandle = SequencerModule.RegisterTrackEditor_Handle( FOnCreateTrackEditor::CreateStatic( &F3DAttachTrackEditor::CreateTrackEditor ) );
+		AudioTrackEditorCreateTrackEditorDelegateHandle = SequencerModule.RegisterTrackEditor_Handle( FOnCreateTrackEditor::CreateStatic( &FAudioTrackEditor::CreateTrackEditor ) );
+		EditorCreateTrackEditorDelegateHandle = SequencerModule.RegisterTrackEditor_Handle( FOnCreateTrackEditor::CreateStatic( &FShotTrackEditor::CreateTrackEditor ) );
+		ParticleTrackEditorCreateTrackEditorDelegateHandle = SequencerModule.RegisterTrackEditor_Handle( FOnCreateTrackEditor::CreateStatic( &FParticleTrackEditor::CreateTrackEditor ) );
+		PathTrackEditorCreateTrackEditorDelegateHandle = SequencerModule.RegisterTrackEditor_Handle( FOnCreateTrackEditor::CreateStatic( &F3DPathTrackEditor::CreateTrackEditor ) );
+		SubMovieSceneTrackEditorCreateTrackEditorDelegateHandle = SequencerModule.RegisterTrackEditor_Handle( FOnCreateTrackEditor::CreateStatic( &FSubMovieSceneTrackEditor::CreateTrackEditor ) );
+		TransformTrackEditorCreateTrackEditorDelegateHandle = SequencerModule.RegisterTrackEditor_Handle( FOnCreateTrackEditor::CreateStatic( &F3DTransformTrackEditor::CreateTrackEditor ) );
 	}
 
 	virtual void ShutdownModule() override
 	{
-		if( FModuleManager::Get().IsModuleLoaded( "Sequencer" ) )
+		if (!FModuleManager::Get().IsModuleLoaded("Sequencer"))
 		{
-			// Unregister auto key handlers
-			ISequencerModule& SequencerModule = FModuleManager::Get().GetModuleChecked<ISequencerModule>( "Sequencer" );
-
-			SequencerModule.UnRegisterTrackEditor_Handle( BoolPropertyTrackEditorCreateTrackEditorDelegateHandle );
-			SequencerModule.UnRegisterTrackEditor_Handle( BytePropertyTrackEditorCreateTrackEditorDelegateHandle );
-			SequencerModule.UnRegisterTrackEditor_Handle( ColorPropertyTrackEditorCreateTrackEditorDelegateHandle );
-			SequencerModule.UnRegisterTrackEditor_Handle( FloatPropertyTrackEditorCreateTrackEditorDelegateHandle );
-			SequencerModule.UnRegisterTrackEditor_Handle( VectorPropertyTrackEditorCreateTrackEditorDelegateHandle );
-			SequencerModule.UnRegisterTrackEditor_Handle( VisibilityPropertyTrackEditorCreateTrackEditorDelegateHandle );
-
-			SequencerModule.UnRegisterTrackEditor_Handle( TransformTrackEditorCreateTrackEditorDelegateHandle );
-			SequencerModule.UnRegisterTrackEditor_Handle( EditorCreateTrackEditorDelegateHandle );
-			SequencerModule.UnRegisterTrackEditor_Handle( SubMovieSceneTrackEditorCreateTrackEditorDelegateHandle );
-			SequencerModule.UnRegisterTrackEditor_Handle( AudioTrackEditorCreateTrackEditorDelegateHandle );
-			SequencerModule.UnRegisterTrackEditor_Handle( AnimationTrackEditorCreateTrackEditorDelegateHandle );
-			SequencerModule.UnRegisterTrackEditor_Handle( ParticleTrackEditorCreateTrackEditorDelegateHandle );
-			SequencerModule.UnRegisterTrackEditor_Handle( AttachTrackEditorCreateTrackEditorDelegateHandle );
-			SequencerModule.UnRegisterTrackEditor_Handle( PathTrackEditorCreateTrackEditorDelegateHandle );
+			return;
 		}
-		
+
+		// unregister auto key handlers
+		ISequencerModule& SequencerModule = FModuleManager::Get().GetModuleChecked<ISequencerModule>( "Sequencer" );
+
+		SequencerModule.UnRegisterTrackEditor_Handle( BoolPropertyTrackEditorCreateTrackEditorDelegateHandle );
+		SequencerModule.UnRegisterTrackEditor_Handle( BytePropertyTrackEditorCreateTrackEditorDelegateHandle );
+		SequencerModule.UnRegisterTrackEditor_Handle( ColorPropertyTrackEditorCreateTrackEditorDelegateHandle );
+		SequencerModule.UnRegisterTrackEditor_Handle( FloatPropertyTrackEditorCreateTrackEditorDelegateHandle );
+		SequencerModule.UnRegisterTrackEditor_Handle( VectorPropertyTrackEditorCreateTrackEditorDelegateHandle );
+		SequencerModule.UnRegisterTrackEditor_Handle( VisibilityPropertyTrackEditorCreateTrackEditorDelegateHandle );
+
+		SequencerModule.UnRegisterTrackEditor_Handle( AnimationTrackEditorCreateTrackEditorDelegateHandle );
+		SequencerModule.UnRegisterTrackEditor_Handle( AttachTrackEditorCreateTrackEditorDelegateHandle );
+		SequencerModule.UnRegisterTrackEditor_Handle( AudioTrackEditorCreateTrackEditorDelegateHandle );
+		SequencerModule.UnRegisterTrackEditor_Handle( EditorCreateTrackEditorDelegateHandle );
+		SequencerModule.UnRegisterTrackEditor_Handle( ParticleTrackEditorCreateTrackEditorDelegateHandle );
+		SequencerModule.UnRegisterTrackEditor_Handle( PathTrackEditorCreateTrackEditorDelegateHandle );
+		SequencerModule.UnRegisterTrackEditor_Handle( SubMovieSceneTrackEditorCreateTrackEditorDelegateHandle );
+		SequencerModule.UnRegisterTrackEditor_Handle( TransformTrackEditorCreateTrackEditorDelegateHandle );
 	}
 
 private:
@@ -88,14 +94,14 @@ private:
 	FDelegateHandle VectorPropertyTrackEditorCreateTrackEditorDelegateHandle;
 	FDelegateHandle VisibilityPropertyTrackEditorCreateTrackEditorDelegateHandle;
 
-	FDelegateHandle TransformTrackEditorCreateTrackEditorDelegateHandle;
-	FDelegateHandle EditorCreateTrackEditorDelegateHandle;
-	FDelegateHandle SubMovieSceneTrackEditorCreateTrackEditorDelegateHandle;
-	FDelegateHandle AudioTrackEditorCreateTrackEditorDelegateHandle;
 	FDelegateHandle AnimationTrackEditorCreateTrackEditorDelegateHandle;
-	FDelegateHandle ParticleTrackEditorCreateTrackEditorDelegateHandle;
 	FDelegateHandle AttachTrackEditorCreateTrackEditorDelegateHandle;
+	FDelegateHandle AudioTrackEditorCreateTrackEditorDelegateHandle;
+	FDelegateHandle EditorCreateTrackEditorDelegateHandle;
+	FDelegateHandle ParticleTrackEditorCreateTrackEditorDelegateHandle;
 	FDelegateHandle PathTrackEditorCreateTrackEditorDelegateHandle;
+	FDelegateHandle SubMovieSceneTrackEditorCreateTrackEditorDelegateHandle;
+	FDelegateHandle TransformTrackEditorCreateTrackEditorDelegateHandle;
 };
 
 
