MovieSceneTools: Better names for track editor delegate handles

diff --git a/Engine/Source/Editor/MovieSceneTools/Private/MovieSceneToolsModule.cpp b/Engine/Source/Editor/MovieSceneTools/Private/MovieSceneToolsModule.cpp
--- a/Engine/Source/Editor/MovieSceneTools/Private/MovieSceneToolsModule.cpp
+++ b/Engine/Source/Editor/MovieSceneTools/Private/MovieSceneToolsModule.cpp
@@ -37,24 +37,25 @@ public:
 
 	virtual void StartupModule() override
 	{
-		// register with the sequencer module that we provide auto-key handlers
 		ISequencerModule& SequencerModule = FModuleManager::Get().LoadModuleChecked<ISequencerModule>( "Sequencer" );
 
-		BoolPropertyTrackEditorCreateTrackEditorDelegateHandle = SequencerModule.RegisterTrackEditor_Handle( FOnCreateTrackEditor::CreateStatic( &FBoolPropertyTrackEditor::CreateTrackEditor ) );
-		BytePropertyTrackEditorCreateTrackEditorDelegateHandle = SequencerModule.RegisterTrackEditor_Handle( FOnCreateTrackEditor::CreateStatic( &FBytePropertyTrackEditor::CreateTrackEditor ) );
-		ColorPropertyTrackEditorCreateTrackEditorDelegateHandle = SequencerModule.RegisterTrackEditor_Handle( FOnCreateTrackEditor::CreateStatic( &FColorPropertyTrackEditor::CreateTrackEditor ) );
-		FloatPropertyTrackEditorCreateTrackEditorDelegateHandle = SequencerModule.RegisterTrackEditor_Handle( FOnCreateTrackEditor::CreateStatic( &FFloatPropertyTrackEditor::CreateTrackEditor ) );
-		VectorPropertyTrackEditorCreateTrackEditorDelegateHandle = SequencerModule.RegisterTrackEditor_Handle( FOnCreateTrackEditor::CreateStatic( &FVectorPropertyTrackEditor::CreateTrackEditor ) );
-		VisibilityPropertyTrackEditorCreateTrackEditorDelegateHandle = SequencerModule.RegisterTrackEditor_Handle( FOnCreateTrackEditor::CreateStatic( &FVisibilityPropertyTrackEditor::CreateTrackEditor ) );
-
-		AnimationTrackEditorCreateTrackEditorDelegateHandle = SequencerModule.RegisterTrackEditor_Handle( FOnCreateTrackEditor::CreateStatic( &FSkeletalAnimationTrackEditor::CreateTrackEditor ) );
-		AttachTrackEditorCreateTrackEditorDelegateHandle = SequencerModule.RegisterTrackEditor_Handle( FOnCreateTrackEditor::CreateStatic( &F3DAttachTrackEditor::CreateTrackEditor ) );
-		AudioTrackEditorCreateTrackEditorDelegateHandle = SequencerModule.RegisterTrackEditor_Handle( FOnCreateTrackEditor::CreateStatic( &FAudioTrackEditor::CreateTrackEditor ) );
-		EditorCreateTrackEditorDelegateHandle = SequencerModule.RegisterTrackEditor_Handle( FOnCreateTrackEditor::CreateStatic( &FShotTrackEditor::CreateTrackEditor ) );
-		ParticleTrackEditorCreateTrackEditorDelegateHandle = SequencerModule.RegisterTrackEditor_Handle( FOnCreateTrackEditor::CreateStatic( &FParticleTrackEditor::CreateTrackEditor ) );
-		PathTrackEditorCreateTrackEditorDelegateHandle = SequencerModule.RegisterTrackEditor_Handle( FOnCreateTrackEditor::CreateStatic( &F3DPathTrackEditor::CreateTrackEditor ) );
-		SubMovieSceneTrackEditorCreateTrackEditorDelegateHandle = SequencerModule.RegisterTrackEditor_Handle( FOnCreateTrackEditor::CreateStatic( &FSubMovieSceneTrackEditor::CreateTrackEditor ) );
-		TransformTrackEditorCreateTrackEditorDelegateHandle = SequencerModule.RegisterTrackEditor_Handle( FOnCreateTrackEditor::CreateStatic( &F3DTransformTrackEditor::CreateTrackEditor ) );
+		// register property track editors
+		BoolPropertyTrackCreateEditorHandle = SequencerModule.RegisterTrackEditor_Handle( FOnCreateTrackEditor::CreateStatic( &FBoolPropertyTrackEditor::CreateTrackEditor ) );
+		BytePropertyTrackCreateEditorHandle = SequencerModule.RegisterTrackEditor_Handle( FOnCreateTrackEditor::CreateStatic( &FBytePropertyTrackEditor::CreateTrackEditor ) );
+		ColorPropertyTrackCreateEditorHandle = SequencerModule.RegisterTrackEditor_Handle( FOnCreateTrackEditor::CreateStatic( &FColorPropertyTrackEditor::CreateTrackEditor ) );
+		FloatPropertyTrackCreateEditorHandle = SequencerModule.RegisterTrackEditor_Handle( FOnCreateTrackEditor::CreateStatic( &FFloatPropertyTrackEditor::CreateTrackEditor ) );
+		VectorPropertyTrackCreateEditorHandle = SequencerModule.RegisterTrackEditor_Handle( FOnCreateTrackEditor::CreateStatic( &FVectorPropertyTrackEditor::CreateTrackEditor ) );
+		VisibilityPropertyTrackCreateEditorHandle = SequencerModule.RegisterTrackEditor_Handle( FOnCreateTrackEditor::CreateStatic( &FVisibilityPropertyTrackEditor::CreateTrackEditor ) );
+
+		// register specialty track editors
+		AnimationTrackCreateEditorHandle = SequencerModule.RegisterTrackEditor_Handle( FOnCreateTrackEditor::CreateStatic( &FSkeletalAnimationTrackEditor::CreateTrackEditor ) );
+		AttachTrackCreateEditorHandle = SequencerModule.RegisterTrackEditor_Handle( FOnCreateTrackEditor::CreateStatic( &F3DAttachTrackEditor::CreateTrackEditor ) );
+		AudioTrackCreateEditorHandle = SequencerModule.RegisterTrackEditor_Handle( FOnCreateTrackEditor::CreateStatic( &FAudioTrackEditor::CreateTrackEditor ) );
+		ParticleTrackCreateEditorHandle = SequencerModule.RegisterTrackEditor_Handle( FOnCreateTrackEditor::CreateStatic( &FParticleTrackEditor::CreateTrackEditor ) );
+		PathTrackCreateEditorHandle = SequencerModule.RegisterTrackEditor_Handle( FOnCreateTrackEditor::CreateStatic( &F3DPathTrackEditor::CreateTrackEditor ) );
+		ShotTrackCreateEditorHandle = SequencerModule.RegisterTrackEditor_Handle( FOnCreateTrackEditor::CreateStatic( &FShotTrackEditor::CreateTrackEditor ) );
+		SubMovieSceneTrackCreateEditorHandle = SequencerModule.RegisterTrackEditor_Handle( FOnCreateTrackEditor::CreateStatic( &FSubMovieSceneTrackEditor::CreateTrackEditor ) );
+		TransformTrackCreateEditorHandle = SequencerModule.RegisterTrackEditor_Handle( FOnCreateTrackEditor::CreateStatic( &F3DTransformTrackEditor::CreateTrackEditor ) );
 	}
 
 	virtual void ShutdownModule() override
@@ -64,44 +65,45 @@ public:
 			return;
 		}
 
-		// unregister auto key handlers
 		ISequencerModule& SequencerModule = FModuleManager::Get().GetModuleChecked<ISequencerModule>( "Sequencer" );
 
-		SequencerModule.UnRegisterTrackEditor_Handle( BoolPropertyTrackEditorCreateTrackEditorDelegateHandle );
-		SequencerModule.UnRegisterTrackEditor_Handle( BytePropertyTrackEditorCreateTrackEditorDelegateHandle );
-		SequencerModule.UnRegisterTrackEditor_Handle( ColorPropertyTrackEditorCreateTrackEditorDelegateHandle );
-		SequencerModule.UnRegisterTrackEditor_Handle( FloatPropertyTrackEditorCreateTrackEditorDelegateHandle );
-		SequencerModule.UnRegisterTrackEditor_Handle( VectorPropertyTrackEditorCreateTrackEditorDelegateHandle );
-		SequencerModule.UnRegisterTrackEditor_Handle( VisibilityPropertyTrackEditorCreateTrackEditorDelegateHandle );
-
-		SequencerModule.UnRegisterTrackEditor_Handle( AnimationTrackEditorCreateTrackEditorDelegateHandle );
-		SequencerModule.UnRegisterTrackEditor_Handle( AttachTrackEditorCreateTrackEditorDelegateHandle );
-		SequencerModule.UnRegisterTrackEditor_Handle( AudioTrackEditorCreateTrackEditorDelegateHandle );
-		SequencerModule.UnRegisterTrackEditor_Handle( EditorCreateTrackEditorDelegateHandle );
-		SequencerModule.UnRegisterTrackEditor_Handle( ParticleTrackEditorCreateTrackEditorDelegateHandle );
-		SequencerModule.UnRegisterTrackEditor_Handle( PathTrackEditorCreateTrackEditorDelegateHandle );
-		SequencerModule.UnRegisterTrackEditor_Handle( SubMovieSceneTrackEditorCreateTrackEditorDelegateHandle );
-		SequencerModule.UnRegisterTrackEditor_Handle( TransformTrackEditorCreateTrackEditorDelegateHandle );
+		// unregister property track editors
+		SequencerModule.UnRegisterTrackEditor_Handle( BoolPropertyTrackCreateEditorHandle );
+		SequencerModule.UnRegisterTrackEditor_Handle( BytePropertyTrackCreateEditorHandle );
+		SequencerModule.UnRegisterTrackEditor_Handle( ColorPropertyTrackCreateEditorHandle );
+		SequencerModule.UnRegisterTrackEditor_Handle( FloatPropertyTrackCreateEditorHandle );
+		SequencerModule.UnRegisterTrackEditor_Handle( VectorPropertyTrackCreateEditorHandle );
+		SequencerModule.UnRegisterTrackEditor_Handle( VisibilityPropertyTrackCreateEditorHandle );
+
+		// unregister specialty track editors
+		SequencerModule.UnRegisterTrackEditor_Handle( AnimationTrackCreateEditorHandle );
+		SequencerModule.UnRegisterTrackEditor_Handle( AttachTrackCreateEditorHandle );
+		SequencerModule.UnRegisterTrackEditor_Handle( AudioTrackCreateEditorHandle );
+		SequencerModule.UnRegisterTrackEditor_Handle( ParticleTrackCreateEditorHandle );
+		SequencerModule.UnRegisterTrackEditor_Handle( PathTrackCreateEditorHandle );
+		SequencerModule.UnRegisterTrackEditor_Handle( ShotTrackCreateEditorHandle );
+		SequencerModule.UnRegisterTrackEditor_Handle( SubMovieSceneTrackCreateEditorHandle );
+		SequencerModule.UnRegisterTrackEditor_Handle( TransformTrackCreateEditorHandle );
 	}
 
 private:
 
 	/** Registered delegate handles */
-	FDelegateHandle BoolPropertyTrackEditorCreateTrackEditorDelegateHandle;
-	FDelegateHandle BytePropertyTrackEditorCreateTrackEditorDelegateHandle;
-	FDelegateHandle ColorPropertyTrackEditorCreateTrackEditorDelegateHandle;
-	FDelegateHandle FloatPropertyTrackEditorCreateTrackEditorDelegateHandle;
-	FDelegateHandle VectorPropertyTrackEditorCreateTrackEditorDelegateHandle;
-	FDelegateHandle VisibilityPropertyTrackEditorCreateTrackEditorDelegateHandle;
-
-	FDelegateHandle AnimationTrackEditorCreateTrackEditorDelegateHandle;
-	FDelegateHandle AttachTrackEditorCreateTrackEditorDelegateHandle;
-	FDelegateHandle AudioTrackEditorCreateTrackEditorDelegateHandle;
-	FDelegateHandle EditorCreateTrackEditorDelegateHandle;
-	FDelegateHandle ParticleTrackEditorCreateTrackEditorDelegateHandle;
-	FDelegateHandle PathTrackEditorCreateTrackEditorDelegateHandle;
-	FDelegateHandle SubMovieSceneTrackEditorCreateTrackEditorDelegateHandle;
-	FDelegateHandle TransformTrackEditorCreateTrackEditorDelegateHandle;
+	FDelegateHandle BoolPropertyTrackCreateEditorHandle;
+	FDelegateHandle BytePropertyTrackCreateEditorHandle;
+	FDelegateHandle ColorPropertyTrackCreateEditorHandle;
+	FDelegateHandle FloatPropertyTrackCreateEditorHandle;
+	FDelegateHandle VectorPropertyTrackCreateEditorHandle;
+	FDelegateHandle VisibilityPropertyTrackCreateEditorHandle;
+
+	FDelegateHandle AnimationTrackCreateEditorHandle;
+	FDelegateHandle AttachTrackCreateEditorHandle;
+	FDelegateHandle AudioTrackCreateEditorHandle;
+	FDelegateHandle ParticleTrackCreateEditorHandle;
+	FDelegateHandle PathTrackCreateEditorHandle;
+	FDelegateHandle ShotTrackCreateEditorHandle;
+	FDelegateHandle SubMovieSceneTrackCreateEditorHandle;
+	FDelegateHandle TransformTrackCreateEditorHandle;
 };
 
 
