MovieSceneTools: Started to implement event track editor

diff --git a/Engine/Source/Editor/MovieSceneTools/Private/MovieSceneToolsModule.cpp b/Engine/Source/Editor/MovieSceneTools/Private/MovieSceneToolsModule.cpp
--- a/Engine/Source/Editor/MovieSceneTools/Private/MovieSceneToolsModule.cpp
+++ b/Engine/Source/Editor/MovieSceneTools/Private/MovieSceneToolsModule.cpp
@@ -9,12 +9,14 @@
 #include "ScopedTransaction.h"
 #include "MovieScene.h"
 #include "MovieSceneTrackEditor.h"
+
 #include "BoolPropertyTrackEditor.h"
 #include "BytePropertyTrackEditor.h"
 #include "ColorPropertyTrackEditor.h"
 #include "FloatPropertyTrackEditor.h"
 #include "VectorPropertyTrackEditor.h"
 #include "VisibilityPropertyTrackEditor.h"
+
 #include "TransformTrackEditor.h"
 #include "ShotTrackEditor.h"
 #include "SubMovieSceneTrackEditor.h"
@@ -22,6 +24,7 @@
 #include "SkeletalAnimationTrackEditor.h"
 #include "ParticleTrackEditor.h"
 #include "AttachTrackEditor.h"
+#include "EventTrackEditor.h"
 #include "PathTrackEditor.h"
 
 
@@ -51,6 +54,7 @@ public:
 		AnimationTrackCreateEditorHandle = SequencerModule.RegisterTrackEditor_Handle( FOnCreateTrackEditor::CreateStatic( &FSkeletalAnimationTrackEditor::CreateTrackEditor ) );
 		AttachTrackCreateEditorHandle = SequencerModule.RegisterTrackEditor_Handle( FOnCreateTrackEditor::CreateStatic( &F3DAttachTrackEditor::CreateTrackEditor ) );
 		AudioTrackCreateEditorHandle = SequencerModule.RegisterTrackEditor_Handle( FOnCreateTrackEditor::CreateStatic( &FAudioTrackEditor::CreateTrackEditor ) );
+		EventTrackCreateEditorHandle = SequencerModule.RegisterTrackEditor_Handle( FOnCreateTrackEditor::CreateStatic( &FEventTrackEditor::CreateTrackEditor ) );
 		ParticleTrackCreateEditorHandle = SequencerModule.RegisterTrackEditor_Handle( FOnCreateTrackEditor::CreateStatic( &FParticleTrackEditor::CreateTrackEditor ) );
 		PathTrackCreateEditorHandle = SequencerModule.RegisterTrackEditor_Handle( FOnCreateTrackEditor::CreateStatic( &F3DPathTrackEditor::CreateTrackEditor ) );
 		ShotTrackCreateEditorHandle = SequencerModule.RegisterTrackEditor_Handle( FOnCreateTrackEditor::CreateStatic( &FShotTrackEditor::CreateTrackEditor ) );
@@ -79,6 +83,7 @@ public:
 		SequencerModule.UnRegisterTrackEditor_Handle( AnimationTrackCreateEditorHandle );
 		SequencerModule.UnRegisterTrackEditor_Handle( AttachTrackCreateEditorHandle );
 		SequencerModule.UnRegisterTrackEditor_Handle( AudioTrackCreateEditorHandle );
+		SequencerModule.UnRegisterTrackEditor_Handle( EventTrackCreateEditorHandle );
 		SequencerModule.UnRegisterTrackEditor_Handle( ParticleTrackCreateEditorHandle );
 		SequencerModule.UnRegisterTrackEditor_Handle( PathTrackCreateEditorHandle );
 		SequencerModule.UnRegisterTrackEditor_Handle( ShotTrackCreateEditorHandle );
@@ -99,6 +104,7 @@ private:
 	FDelegateHandle AnimationTrackCreateEditorHandle;
 	FDelegateHandle AttachTrackCreateEditorHandle;
 	FDelegateHandle AudioTrackCreateEditorHandle;
+	FDelegateHandle EventTrackCreateEditorHandle;
 	FDelegateHandle ParticleTrackCreateEditorHandle;
 	FDelegateHandle PathTrackCreateEditorHandle;
 	FDelegateHandle ShotTrackCreateEditorHandle;
diff --git a/Engine/Source/Editor/MovieSceneTools/Private/Sections/EventTrackSection.cpp b/Engine/Source/Editor/MovieSceneTools/Private/Sections/EventTrackSection.cpp
new file mode 100644
--- /dev/null
+++ b/Engine/Source/Editor/MovieSceneTools/Private/Sections/EventTrackSection.cpp
@@ -0,0 +1,53 @@
+// Copyright 1998-2015 Epic Games, Inc. All Rights Reserved.
+
+#include "MovieSceneToolsPrivatePCH.h"
+#include "EventTrackSection.h"
+#include "MovieSceneEventSection.h"
+
+
+#define LOCTEXT_NAMESPACE "FEventTrackSection"
+
+
+/* FEventTrackEditor static functions
+ *****************************************************************************/
+
+FEventTrackSection::FEventTrackSection(UMovieSceneSection& InSection, TSharedPtr<ISequencer> InSequencer)
+	: Section(nullptr/*CastChecked<UMovieSceneEventSection>(&InSection)*/)
+	, Sequencer(InSequencer)
+{ }
+
+
+/* FEventTrackEditor static functions
+ *****************************************************************************/
+
+UMovieSceneSection* FEventTrackSection::GetSectionObject()
+{
+	return Section;
+}
+
+
+int32 FEventTrackSection::OnPaintSection(const FGeometry& AllottedGeometry, const FSlateRect& SectionClippingRect, FSlateWindowElementList& OutDrawElements, int32 LayerId, bool bParentEnabled) const
+{
+	return LayerId + 1;
+}
+
+
+FText FEventTrackSection::GetDisplayName() const
+{
+	return LOCTEXT("DisplayName", "Events");
+}
+
+
+FText FEventTrackSection::GetSectionTitle() const
+{
+	return FText::GetEmpty();
+}
+
+
+void FEventTrackSection::GenerateSectionLayout(class ISectionLayoutBuilder& LayoutBuilder) const
+{
+
+}
+
+
+#undef LOCTEXT_NAMESPACE
diff --git a/Engine/Source/Editor/MovieSceneTools/Private/Sections/EventTrackSection.h b/Engine/Source/Editor/MovieSceneTools/Private/Sections/EventTrackSection.h
new file mode 100644
--- /dev/null
+++ b/Engine/Source/Editor/MovieSceneTools/Private/Sections/EventTrackSection.h
@@ -0,0 +1,51 @@
+// Copyright 1998-2015 Epic Games, Inc. All Rights Reserved.
+
+#pragma once
+
+
+class UMovieSceneEventSection;
+
+
+/**
+ * Event track section, which paints and ticks the appropriate section.
+ */
+class FEventTrackSection
+	: public ISequencerSection
+{
+public:
+
+	/**
+	 * Creates and initializes a new instance
+	 *
+	 * @param InSection The section being visualized.
+	 * @param InSequencer The sequencer that owns this section.
+	 */
+	FEventTrackSection(UMovieSceneSection& InSection, TSharedPtr<ISequencer> InSequencer);
+
+public:
+
+	// ISequencerSection interface
+
+	virtual UMovieSceneSection* GetSectionObject() override;
+//	virtual TSharedRef<SWidget> GenerateSectionWidget() override;
+	virtual int32 OnPaintSection( const FGeometry& AllottedGeometry, const FSlateRect& SectionClippingRect, FSlateWindowElementList& OutDrawElements, int32 LayerId, bool bParentEnabled ) const override;
+//	virtual void Tick( const FGeometry& AllottedGeometry, const FGeometry& ParentGeometry, const double InCurrentTime, const float InDeltaTime ) override;
+	virtual FText GetDisplayName() const override;
+	virtual FText GetSectionTitle() const override;
+//	virtual float GetSectionHeight() const override;
+//	virtual float GetSectionGripSize() const override;
+//	virtual FName GetSectionGripLeftBrushName() const override;
+//	virtual FName GetSectionGripRightBrushName() const override;
+//	virtual bool AreSectionsConnected() const override { return true; }
+	virtual void GenerateSectionLayout( class ISectionLayoutBuilder& LayoutBuilder ) const override;
+//	virtual FReply OnSectionDoubleClicked( const FGeometry& SectionGeometry, const FPointerEvent& MouseEvent ) override;
+//	virtual void BuildSectionContextMenu(FMenuBuilder& MenuBuilder) override;
+
+private:
+
+	/** The section being visualized. */
+	UMovieSceneEventSection* Section;
+
+	/** The parent sequencer we are a part of. */
+	TWeakPtr<ISequencer> Sequencer;
+};
diff --git a/Engine/Source/Editor/MovieSceneTools/Private/TrackEditors/EventTrackEditor.cpp b/Engine/Source/Editor/MovieSceneTools/Private/TrackEditors/EventTrackEditor.cpp
new file mode 100644
--- /dev/null
+++ b/Engine/Source/Editor/MovieSceneTools/Private/TrackEditors/EventTrackEditor.cpp
@@ -0,0 +1,43 @@
+// Copyright 1998-2015 Epic Games, Inc. All Rights Reserved.
+
+#include "MovieSceneToolsPrivatePCH.h"
+#include "EventTrackEditor.h"
+#include "EventTrackSection.h"
+
+
+/* FEventTrackEditor static functions
+ *****************************************************************************/
+
+TSharedRef<FMovieSceneTrackEditor> FEventTrackEditor::CreateTrackEditor(TSharedRef<ISequencer> InSequencer)
+{
+	return MakeShareable(new FEventTrackEditor(InSequencer));
+}
+
+
+/* FEventTrackEditor structors
+ *****************************************************************************/
+
+FEventTrackEditor::FEventTrackEditor(TSharedRef<ISequencer> InSequencer)
+	: FMovieSceneTrackEditor(InSequencer)
+{ }
+
+
+/* FMovieSceneTrackEditor interface
+ *****************************************************************************/
+
+void FEventTrackEditor::AddKey(const FGuid& ObjectGuid, UObject* AdditionalAsset)
+{
+	// todo gmp: Sequencer: implement event track section
+}
+
+
+TSharedRef<ISequencerSection> FEventTrackEditor::MakeSectionInterface(UMovieSceneSection& SectionObject, UMovieSceneTrack* Track)
+{
+	return MakeShareable(new FEventTrackSection(SectionObject/*, Track->GetTrackName()*/, GetSequencer()));
+}
+
+
+bool FEventTrackEditor::SupportsType(TSubclassOf<UMovieSceneTrack> Type) const
+{
+	return (Type == UMovieSceneEventTrack::StaticClass());
+}
diff --git a/Engine/Source/Editor/MovieSceneTools/Private/TrackEditors/EventTrackEditor.h b/Engine/Source/Editor/MovieSceneTools/Private/TrackEditors/EventTrackEditor.h
new file mode 100644
--- /dev/null
+++ b/Engine/Source/Editor/MovieSceneTools/Private/TrackEditors/EventTrackEditor.h
@@ -0,0 +1,44 @@
+// Copyright 1998-2015 Epic Games, Inc. All Rights Reserved.
+
+#pragma once
+
+#include "MovieSceneEventTrack.h"
+#include "MovieSceneTrackEditor.h"
+
+
+class ISequencer;
+
+
+/**
+* A property track editor for named events.
+*/
+class FEventTrackEditor
+	: public FMovieSceneTrackEditor
+{
+public:
+
+	/**
+	 * Factory function to create an instance of this class (called by a sequencer).
+	 *
+	 * @param InSequencer The sequencer instance to be used by this tool.
+	 * @return The new instance of this class.
+	 */
+	static TSharedRef<FMovieSceneTrackEditor> CreateTrackEditor(TSharedRef<ISequencer> InSequencer);
+
+public:
+
+	/**
+	 * Creates and initializes a new instance.
+	 *
+	 * @param InSequencer The sequencer instance to be used by this tool.
+	 */
+	FEventTrackEditor(TSharedRef<ISequencer> InSequencer);	
+
+public:
+
+	// FMovieSceneTrackEditor interface
+
+	virtual void AddKey(const FGuid& ObjectGuid, UObject* AdditionalAsset = nullptr) override;
+	virtual TSharedRef<ISequencerSection> MakeSectionInterface(UMovieSceneSection& SectionObject, UMovieSceneTrack* Track) override;
+	virtual bool SupportsType( TSubclassOf<UMovieSceneTrack> Type ) const override;
+};
