Sequencer - Material parameter animation support
+ Add animation support for scalar material parameters on actor components.
+ Add the ability to specify an edit widget for track nodes in the outliner through the track editor.

diff --git a/Engine/Plugins/MovieScene/ActorAnimationEditor/Source/ActorAnimationEditor/ActorAnimationEditor.Build.cs b/Engine/Plugins/MovieScene/ActorAnimationEditor/Source/ActorAnimationEditor/ActorAnimationEditor.Build.cs
--- a/Engine/Plugins/MovieScene/ActorAnimationEditor/Source/ActorAnimationEditor/ActorAnimationEditor.Build.cs
+++ b/Engine/Plugins/MovieScene/ActorAnimationEditor/Source/ActorAnimationEditor/ActorAnimationEditor.Build.cs
@@ -21,6 +21,7 @@ public class ActorAnimationEditor : ModuleRules
                 "EditorStyle",
                 "Engine",
 				"MovieScene",
+				"MovieSceneTracks",
 				"Sequencer",
                 "Slate",
                 "SlateCore",
diff --git a/Engine/Plugins/MovieScene/ActorAnimationEditor/Source/ActorAnimationEditor/Private/ActorAnimationEditorToolkit.cpp b/Engine/Plugins/MovieScene/ActorAnimationEditor/Source/ActorAnimationEditor/Private/ActorAnimationEditorToolkit.cpp
--- a/Engine/Plugins/MovieScene/ActorAnimationEditor/Source/ActorAnimationEditor/Private/ActorAnimationEditorToolkit.cpp
+++ b/Engine/Plugins/MovieScene/ActorAnimationEditor/Source/ActorAnimationEditor/Private/ActorAnimationEditorToolkit.cpp
@@ -9,6 +9,9 @@
 #include "ISequencer.h"
 #include "ISequencerModule.h"
 #include "SDockTab.h"
+#include "MovieSceneBinding.h"
+#include "MovieScene.h"
+#include "MovieSceneMaterialTrack.h"
 
 
 #define LOCTEXT_NAMESPACE "ActorAnimationEditor"
@@ -172,6 +175,35 @@ void FActorAnimationEditorToolkit::HandleAddComponentActionExecute( UActorCompon
 }
 
 
+void FActorAnimationEditorToolkit::HandleAddComponentMaterialActionExecute( UPrimitiveComponent* Component, int32 MaterialIndex )
+{
+	FGuid ObjectHandle = Sequencer->GetHandleToObject( Component );
+	for ( const FMovieSceneBinding& Binding : Sequencer->GetFocusedMovieSceneSequence()->GetMovieScene()->GetBindings() )
+	{
+		if ( Binding.GetObjectGuid() == ObjectHandle )
+		{
+			bool bHasMaterialTrack = false;
+			for ( UMovieSceneTrack* Track : Binding.GetTracks() )
+			{
+				UMovieSceneComponentMaterialTrack* MaterialTrack = Cast<UMovieSceneComponentMaterialTrack>( Track );
+				if ( MaterialTrack->GetMaterialIndex() == MaterialIndex )
+				{
+					bHasMaterialTrack = true;
+					break;
+				}
+			}
+			if ( bHasMaterialTrack == false )
+			{
+				UMovieSceneComponentMaterialTrack* MaterialTrack = Cast<UMovieSceneComponentMaterialTrack>(
+					Sequencer->GetFocusedMovieSceneSequence()->GetMovieScene()->AddTrack( UMovieSceneComponentMaterialTrack::StaticClass(), ObjectHandle));
+				MaterialTrack->SetMaterialIndex(MaterialIndex);
+				Sequencer->NotifyMovieSceneDataChanged();
+			}
+		}
+	}
+}
+
+
 void FActorAnimationEditorToolkit::HandleMapChanged(class UWorld* NewWorld, EMapChangeType MapChangeType)
 {
 	Sequencer->NotifyMapChanged(NewWorld, MapChangeType);
@@ -218,23 +250,42 @@ void FActorAnimationEditorToolkit::HandleTrackMenuExtensionAddTrack( FMenuBuilde
 	}
 
 	AActor* Actor = Cast<AActor>(ContextObjects[0]);
-	
-	if (Actor == nullptr)
+	if (Actor != nullptr)
 	{
-		return;
+		AddTrackMenuBuilder.BeginSection("Components", LOCTEXT("ComponentsSection", "Components"));
+		{
+			for (UActorComponent* Component : Actor->GetComponents())
+			{
+				FUIAction AddComponentAction(FExecuteAction::CreateSP(this, &FActorAnimationEditorToolkit::HandleAddComponentActionExecute, Component));
+				FText AddComponentLabel = FText::FromString(Component->GetName());
+				FText AddComponentToolTip = FText::Format(LOCTEXT("ComponentToolTipFormat", "Add {0} component"), FText::FromString(Component->GetName()));
+				AddTrackMenuBuilder.AddMenuEntry(AddComponentLabel, AddComponentToolTip, FSlateIcon(), AddComponentAction);
+			}
+		}
+		AddTrackMenuBuilder.EndSection();
 	}
-
-	AddTrackMenuBuilder.BeginSection("Components", LOCTEXT("ComponentsSection", "Components"));
+	else
 	{
-		for (UActorComponent* Component : Actor->GetComponents())
+		UPrimitiveComponent* Component = Cast<UPrimitiveComponent>(ContextObjects[0]);
+		if (Component != nullptr)
 		{
-			FUIAction AddComponentAction(FExecuteAction::CreateSP(this, &FActorAnimationEditorToolkit::HandleAddComponentActionExecute, Component));
-			FText AddComponentLabel = FText::FromString(Component->GetName());
-			FText AddComponentToolTip = FText::Format(LOCTEXT("ComponentToolTipFormat", "Add {0} component"), FText::FromString(Component->GetName()));
-			AddTrackMenuBuilder.AddMenuEntry(AddComponentLabel, AddComponentToolTip, FSlateIcon(), AddComponentAction);
+			int32 NumMaterials = Component->GetNumMaterials();
+			if (NumMaterials > 0)
+			{
+				AddTrackMenuBuilder.BeginSection("Materials", LOCTEXT("MaterialSection", "Materials"));
+				{
+					for (int32 MaterialIndex = 0; MaterialIndex < NumMaterials; MaterialIndex++)
+					{
+						FUIAction AddComponentMaterialAction(FExecuteAction::CreateSP(this, &FActorAnimationEditorToolkit::HandleAddComponentMaterialActionExecute, Component, MaterialIndex));
+						FText AddComponentMaterialLabel = FText::Format(LOCTEXT("ComponentMaterialIndexLabelFormat", "Element {0}"), FText::AsNumber(MaterialIndex));
+						FText AddComponentMaterialToolTip = FText::Format(LOCTEXT("ComponentMaterialIndexToolTipFormat", "Add material element {0}" ), FText::AsNumber(MaterialIndex));
+						AddTrackMenuBuilder.AddMenuEntry( AddComponentMaterialLabel, AddComponentMaterialToolTip, FSlateIcon(), AddComponentMaterialAction );
+					}
+				}
+				AddTrackMenuBuilder.EndSection();
+			}
 		}
 	}
-	AddTrackMenuBuilder.EndSection();
 }
 
 
diff --git a/Engine/Plugins/MovieScene/ActorAnimationEditor/Source/ActorAnimationEditor/Private/ActorAnimationEditorToolkit.h b/Engine/Plugins/MovieScene/ActorAnimationEditor/Source/ActorAnimationEditor/Private/ActorAnimationEditorToolkit.h
--- a/Engine/Plugins/MovieScene/ActorAnimationEditor/Source/ActorAnimationEditor/Private/ActorAnimationEditorToolkit.h
+++ b/Engine/Plugins/MovieScene/ActorAnimationEditor/Source/ActorAnimationEditor/Private/ActorAnimationEditorToolkit.h
@@ -66,6 +66,9 @@ private:
 	/** Callback for executing the Add Component action. */
 	void HandleAddComponentActionExecute(UActorComponent* Component);
 
+	/** Callback for executing the add component material track. */
+	void HandleAddComponentMaterialActionExecute(UPrimitiveComponent* Component, int32 MaterialIndex);
+
 	/** Callback for map changes. */
 	void HandleMapChanged(UWorld* NewWorld, EMapChangeType MapChangeType);
 
diff --git a/Engine/Source/Editor/MovieSceneTools/Private/MovieSceneToolsModule.cpp b/Engine/Source/Editor/MovieSceneTools/Private/MovieSceneToolsModule.cpp
--- a/Engine/Source/Editor/MovieSceneTools/Private/MovieSceneToolsModule.cpp
+++ b/Engine/Source/Editor/MovieSceneTools/Private/MovieSceneToolsModule.cpp
@@ -26,6 +26,7 @@
 #include "AttachTrackEditor.h"
 #include "EventTrackEditor.h"
 #include "PathTrackEditor.h"
+#include "MaterialTrackEditor.h"
 
 
 /**
@@ -60,6 +61,7 @@ public:
 		ShotTrackCreateEditorHandle = SequencerModule.RegisterTrackEditor_Handle( FOnCreateTrackEditor::CreateStatic( &FShotTrackEditor::CreateTrackEditor ) );
 		SubMovieSceneTrackCreateEditorHandle = SequencerModule.RegisterTrackEditor_Handle( FOnCreateTrackEditor::CreateStatic( &FSubMovieSceneTrackEditor::CreateTrackEditor ) );
 		TransformTrackCreateEditorHandle = SequencerModule.RegisterTrackEditor_Handle( FOnCreateTrackEditor::CreateStatic( &F3DTransformTrackEditor::CreateTrackEditor ) );
+		ComponentMaterialTrackCreatEditorHandle = SequencerModule.RegisterTrackEditor_Handle( FOnCreateTrackEditor::CreateStatic( &FComponentMaterialTrackEditor::CreateTrackEditor ) );
 	}
 
 	virtual void ShutdownModule() override
@@ -89,6 +91,7 @@ public:
 		SequencerModule.UnRegisterTrackEditor_Handle( ShotTrackCreateEditorHandle );
 		SequencerModule.UnRegisterTrackEditor_Handle( SubMovieSceneTrackCreateEditorHandle );
 		SequencerModule.UnRegisterTrackEditor_Handle( TransformTrackCreateEditorHandle );
+		SequencerModule.UnRegisterTrackEditor_Handle( ComponentMaterialTrackCreatEditorHandle );
 	}
 
 private:
@@ -110,6 +113,7 @@ private:
 	FDelegateHandle ShotTrackCreateEditorHandle;
 	FDelegateHandle SubMovieSceneTrackCreateEditorHandle;
 	FDelegateHandle TransformTrackCreateEditorHandle;
+	FDelegateHandle ComponentMaterialTrackCreatEditorHandle;
 };
 
 
diff --git a/Engine/Source/Editor/MovieSceneTools/Private/Sections/MaterialParameterSection.cpp b/Engine/Source/Editor/MovieSceneTools/Private/Sections/MaterialParameterSection.cpp
new file mode 100644
--- /dev/null
+++ b/Engine/Source/Editor/MovieSceneTools/Private/Sections/MaterialParameterSection.cpp
@@ -0,0 +1,18 @@
+// Copyright 1998-2015 Epic Games, Inc. All Rights Reserved.
+
+#include "MovieSceneToolsPrivatePCH.h"
+#include "MaterialParameterSection.h"
+#include "MovieSceneMaterialParameterSection.h"
+
+void FMaterialParameterSection::GenerateSectionLayout( class ISectionLayoutBuilder& LayoutBuilder ) const
+{
+	UMovieSceneMaterialParameterSection* ParameterSection = Cast<UMovieSceneMaterialParameterSection>( &SectionObject );
+	TArray<FScalarParameterNameAndCurve>* ScalarParameterNamesAndCurves = ParameterSection->GetScalarParameterNamesAndCurves();
+	for ( int32 i = 0; i < ScalarParameterNamesAndCurves->Num(); i++)
+	{
+		LayoutBuilder.AddKeyArea( 
+			(*ScalarParameterNamesAndCurves)[i].ParameterName, 
+			FText::FromName( (*ScalarParameterNamesAndCurves)[i].ParameterName ),
+			MakeShareable( new FFloatCurveKeyArea( &(*ScalarParameterNamesAndCurves)[i].ParameterCurve, ParameterSection ) ) );
+	}
+}
\ No newline at end of file
diff --git a/Engine/Source/Editor/MovieSceneTools/Private/Sections/MaterialParameterSection.h b/Engine/Source/Editor/MovieSceneTools/Private/Sections/MaterialParameterSection.h
new file mode 100644
--- /dev/null
+++ b/Engine/Source/Editor/MovieSceneTools/Private/Sections/MaterialParameterSection.h
@@ -0,0 +1,15 @@
+// Copyright 1998-2015 Epic Games, Inc. All Rights Reserved.
+
+#pragma once
+
+/**
+* A movie scene section for material parameters.
+*/
+class FMaterialParameterSection : public FPropertySection
+{
+public:
+	FMaterialParameterSection( UMovieSceneSection& InSectionObject, FName SectionName )
+		: FPropertySection( InSectionObject, SectionName ) {}
+
+	virtual void GenerateSectionLayout( class ISectionLayoutBuilder& LayoutBuilder ) const override;
+};
\ No newline at end of file
diff --git a/Engine/Source/Editor/MovieSceneTools/Private/TrackEditors/MaterialTrackEditor.cpp b/Engine/Source/Editor/MovieSceneTools/Private/TrackEditors/MaterialTrackEditor.cpp
new file mode 100644
--- /dev/null
+++ b/Engine/Source/Editor/MovieSceneTools/Private/TrackEditors/MaterialTrackEditor.cpp
@@ -0,0 +1,142 @@
+// Copyright 1998-2015 Epic Games, Inc. All Rights Reserved.
+
+#pragma once
+
+#include "MovieSceneToolsPrivatePCH.h"
+#include "MaterialTrackEditor.h"
+#include "MovieSceneMaterialTrack.h"
+#include "MaterialParameterSection.h"
+#include "MovieSceneMaterialParameterSection.h"
+
+#define LOCTEXT_NAMESPACE "MaterialTrackEditor"
+
+FMaterialTrackEditor::FMaterialTrackEditor( TSharedRef<ISequencer> InSequencer )
+	: FMovieSceneTrackEditor( InSequencer )
+{
+}
+
+TSharedRef<ISequencerSection> FMaterialTrackEditor::MakeSectionInterface( UMovieSceneSection& SectionObject, UMovieSceneTrack* Track )
+{
+	UMovieSceneMaterialParameterSection* ParameterSection = Cast<UMovieSceneMaterialParameterSection>(&SectionObject);
+	checkf( ParameterSection != nullptr, TEXT("Unsupported section type.") );
+	return MakeShareable(new FMaterialParameterSection( *ParameterSection, ParameterSection->GetFName()));
+}
+
+TSharedPtr<SWidget> FMaterialTrackEditor::BuildOutlinerEditWidget( const FGuid& ObjectBinding, UMovieSceneTrack* Track )
+{
+	UMovieSceneMaterialTrack* MaterialTrack = Cast<UMovieSceneMaterialTrack>(Track);
+	return
+		SNew( SComboButton )
+		.ButtonStyle( FEditorStyle::Get(), "FlatButton.Light" )
+		.OnGetMenuContent( this, &FMaterialTrackEditor::OnGetAddParameterMenuContent, ObjectBinding, MaterialTrack )
+		.ContentPadding( FMargin( 2, 0 ) )
+		.ButtonContent()
+		[
+			SNew( SHorizontalBox )
+
+			+ SHorizontalBox::Slot()
+			.VAlign( VAlign_Center )
+			.AutoWidth()
+			[
+				SNew( STextBlock )
+				.Font( FEditorStyle::Get().GetFontStyle( "FontAwesome.8" ) )
+				.Text( FText::FromString( FString( TEXT( "\xf067" ) ) ) /*fa-plus*/ )
+			]
+
+			+ SHorizontalBox::Slot()
+			.AutoWidth()
+			.Padding( 4, 0, 0, 0 )
+			[
+				SNew( STextBlock )
+				.Font( FEditorStyle::GetFontStyle( "Sequencer.AnimationOutliner.RegularFont" ) )
+				.Text( LOCTEXT( "AddTrackButton", "Parameter" ) )
+			]
+		];
+}
+
+TSharedRef<SWidget> FMaterialTrackEditor::OnGetAddParameterMenuContent( FGuid ObjectBinding, UMovieSceneMaterialTrack* MaterialTrack )
+{
+	FMenuBuilder AddParameterMenuBuilder( true, nullptr );
+
+	UMaterial* Material = GetMaterialForTrack( ObjectBinding, MaterialTrack );
+	if ( Material != nullptr )
+	{
+		TArray<FName> ScalarParameterNames;
+		TArray<FGuid> ScalarParmeterGuids;
+		Material->GetAllScalarParameterNames( ScalarParameterNames, ScalarParmeterGuids );
+
+		for ( const FName& ScalarParameterName : ScalarParameterNames )
+		{
+			FUIAction AddParameterMenuAction( FExecuteAction::CreateSP( this, &FMaterialTrackEditor::AddParameterSection, ObjectBinding, MaterialTrack, ScalarParameterName ) );
+			AddParameterMenuBuilder.AddMenuEntry( FText::FromName( ScalarParameterName ), FText(), FSlateIcon(), AddParameterMenuAction );
+		}
+	}
+
+	return AddParameterMenuBuilder.MakeWidget();
+}
+
+UMaterial* FMaterialTrackEditor::GetMaterialForTrack( FGuid ObjectBinding, UMovieSceneMaterialTrack* MaterialTrack )
+{
+	UMaterialInterface* MaterialInterface = GetMaterialInterfaceForTrack( ObjectBinding, MaterialTrack );
+	if ( MaterialInterface != nullptr )
+	{
+		UMaterial* Material = Cast<UMaterial>( MaterialInterface );
+		if ( Material != nullptr )
+		{
+			return Material;
+		}
+		else
+		{
+			UMaterialInstance* MaterialInstance = Cast<UMaterialInstance>( MaterialInterface );
+			if ( MaterialInstance != nullptr )
+			{
+				return MaterialInstance->GetMaterial();
+			}
+		}
+	}
+	return nullptr;
+}
+
+void FMaterialTrackEditor::AddParameterSection( FGuid ObjectBinding, UMovieSceneMaterialTrack* MaterialTrack, FName ParameterName )
+{
+	UMovieSceneSequence* MovieSceneSequence = GetMovieSceneSequence();
+	float KeyTime = GetTimeForKey( MovieSceneSequence );
+
+	UMaterial* Material = GetMaterialForTrack(ObjectBinding, MaterialTrack);
+	if (Material != nullptr)
+	{
+		float ParameterValue;
+		Material->GetScalarParameterValue(ParameterName, ParameterValue);
+		MaterialTrack->AddScalarParameterKey(ParameterName, KeyTime, ParameterValue);
+	}
+	NotifyMovieSceneDataChanged();
+}
+
+FComponentMaterialTrackEditor::FComponentMaterialTrackEditor( TSharedRef<ISequencer> InSequencer )
+	: FMaterialTrackEditor( InSequencer )
+{
+}
+
+TSharedRef<FMovieSceneTrackEditor> FComponentMaterialTrackEditor::CreateTrackEditor( TSharedRef<ISequencer> OwningSequencer )
+{
+	return MakeShareable( new FComponentMaterialTrackEditor( OwningSequencer ) );
+}
+
+bool FComponentMaterialTrackEditor::SupportsType( TSubclassOf<UMovieSceneTrack> Type ) const
+{
+	return Type == UMovieSceneComponentMaterialTrack::StaticClass();
+}
+
+UMaterialInterface* FComponentMaterialTrackEditor::GetMaterialInterfaceForTrack( FGuid ObjectBinding, UMovieSceneMaterialTrack* MaterialTrack )
+{
+	UObject* ComponentObject = GetSequencer()->GetFocusedMovieSceneSequence()->FindObject( ObjectBinding );
+	UPrimitiveComponent* Component = Cast<UPrimitiveComponent>( ComponentObject );
+	UMovieSceneComponentMaterialTrack* ComponentMaterialTrack = Cast<UMovieSceneComponentMaterialTrack>( MaterialTrack );
+	if ( Component != nullptr && ComponentMaterialTrack != nullptr )
+	{
+		return Component->GetMaterial( ComponentMaterialTrack->GetMaterialIndex() );
+	}
+	return nullptr;
+}
+
+#undef LOCTEXT_NAMESPACE
\ No newline at end of file
diff --git a/Engine/Source/Editor/MovieSceneTools/Private/TrackEditors/MaterialTrackEditor.h b/Engine/Source/Editor/MovieSceneTools/Private/TrackEditors/MaterialTrackEditor.h
new file mode 100644
--- /dev/null
+++ b/Engine/Source/Editor/MovieSceneTools/Private/TrackEditors/MaterialTrackEditor.h
@@ -0,0 +1,56 @@
+// Copyright 1998-2015 Epic Games, Inc. All Rights Reserved.
+
+#pragma once
+
+class UMovieSceneMaterialTrack;
+
+/**
+ * Track editor for material parameters.
+ */
+class FMaterialTrackEditor : public FMovieSceneTrackEditor
+{
+public:
+	FMaterialTrackEditor( TSharedRef<ISequencer> InSequencer );
+
+	virtual ~FMaterialTrackEditor() { }
+
+	/* FMovieSceneTrackEditor interface */
+	virtual TSharedRef<ISequencerSection> MakeSectionInterface( UMovieSceneSection& SectionObject, UMovieSceneTrack* Track ) override;
+	virtual TSharedPtr<SWidget> BuildOutlinerEditWidget( const FGuid& ObjectBinding, UMovieSceneTrack* Track ) override;
+
+protected:
+	/** Gets a material interface for a specific object binding and material track. */
+	virtual UMaterialInterface* GetMaterialInterfaceForTrack( FGuid ObjectBinding, UMovieSceneMaterialTrack* MaterialTrack ) = 0;
+
+private:
+	/** Provides the contents of the add parameter menu. */
+	TSharedRef<SWidget> OnGetAddParameterMenuContent( FGuid ObjectBinding, UMovieSceneMaterialTrack* MaterialTrack );
+
+	/** Gets a material for a specific object binding and track */
+	UMaterial* GetMaterialForTrack( FGuid ObjectBinding, UMovieSceneMaterialTrack* MaterialTrack );
+
+	/** Adds a parameter section and initial key to a material track.
+	 * @param ObjectBinding The object binding which owns the material track.
+	 * @param MaterialTrack The track to Add the section to.
+	 * @param ParameterName The name of the parameter to add an initial key for.
+	 */
+	void AddParameterSection( FGuid ObjectBinding, UMovieSceneMaterialTrack* MaterialTrack, FName ParameterName );
+};
+
+/**
+ * A specialized material track editor for component materials
+ */
+class FComponentMaterialTrackEditor : public FMaterialTrackEditor
+{
+public:
+	FComponentMaterialTrackEditor( TSharedRef<ISequencer> InSequencer );
+
+	static TSharedRef<FMovieSceneTrackEditor> CreateTrackEditor( TSharedRef<ISequencer> OwningSequencer );
+
+	/** FMovieSceneTrackEditor Interface */
+	virtual bool SupportsType( TSubclassOf<UMovieSceneTrack> Type ) const override;
+
+protected:
+	/** FMaterialtrackEditor interface */
+	virtual UMaterialInterface* GetMaterialInterfaceForTrack( FGuid ObjectBinding, UMovieSceneMaterialTrack* MaterialTrack ) override;
+};
\ No newline at end of file
diff --git a/Engine/Source/Editor/Sequencer/Private/SequencerDisplayNode.cpp b/Engine/Source/Editor/Sequencer/Private/SequencerDisplayNode.cpp
--- a/Engine/Source/Editor/Sequencer/Private/SequencerDisplayNode.cpp
+++ b/Engine/Source/Editor/Sequencer/Private/SequencerDisplayNode.cpp
@@ -10,6 +10,7 @@
 #include "MovieSceneSection.h"
 #include "MovieSceneSequence.h"
 #include "MovieSceneTrack.h"
+#include "MovieSceneTrackEditor.h"
 #include "CommonMovieSceneTools.h"
 #include "IKeyArea.h"
 #include "GroupedKeyArea.h"
@@ -248,7 +249,7 @@ TSharedRef<FSectionCategoryNode> FSequencerDisplayNode::AddCategoryNode( FName C
 	return CategoryNode.ToSharedRef();
 }
 
-TSharedRef<FTrackNode> FSequencerDisplayNode::AddSectionAreaNode( FName SectionName, UMovieSceneTrack& AssociatedTrack )
+TSharedRef<FTrackNode> FSequencerDisplayNode::AddSectionAreaNode( FName SectionName, UMovieSceneTrack& AssociatedTrack, FMovieSceneTrackEditor& AssociatedEditor )
 {
 	TSharedPtr<FTrackNode> SectionNode;
 
@@ -265,7 +266,7 @@ TSharedRef<FTrackNode> FSequencerDisplayNode::AddSectionAreaNode( FName SectionN
 	if( !SectionNode.IsValid() )
 	{
 		// No existing node found make a new one
-		SectionNode = MakeShareable( new FTrackNode( SectionName, AssociatedTrack, SharedThis( this ), ParentTree ) );
+		SectionNode = MakeShareable( new FTrackNode( SectionName, AssociatedTrack, AssociatedEditor, SharedThis( this ), ParentTree ) );
 		ChildNodes.Add( SectionNode.ToSharedRef() );
 	}
 
@@ -497,9 +498,10 @@ void FSectionKeyAreaNode::AddKeyArea( TSharedRef< IKeyArea> KeyArea )
 	KeyAreas.Add( KeyArea );
 }
 
-FTrackNode::FTrackNode( FName NodeName, UMovieSceneTrack& InAssociatedType, TSharedPtr<FSequencerDisplayNode> InParentNode, FSequencerNodeTree& InParentTree )
+FTrackNode::FTrackNode( FName NodeName, UMovieSceneTrack& InAssociatedType, FMovieSceneTrackEditor& InAssociatedEditor, TSharedPtr<FSequencerDisplayNode> InParentNode, FSequencerNodeTree& InParentTree )
 	: FSequencerDisplayNode( NodeName, InParentNode, InParentTree )
 	, AssociatedType( &InAssociatedType )
+	, AssociatedEditor( InAssociatedEditor )
 {
 }
 
@@ -516,7 +518,7 @@ FNodePadding FTrackNode::GetNodePadding() const
 FText FTrackNode::GetDisplayName() const
 {
 	// @todo Sequencer - IS there a better way to get the section interface name for the animation outliner?
-	return Sections.Num() > 0 ? Sections[0]->GetDisplayName() : FText::FromName(NodeName);
+	return TopLevelKeyNode.IsValid() && Sections.Num() > 0 ? Sections[0]->GetDisplayName() : FText::FromName(NodeName);
 }
 
 void FTrackNode::SetSectionAsKeyArea( TSharedRef<IKeyArea>& KeyArea )
@@ -532,13 +534,23 @@ void FTrackNode::SetSectionAsKeyArea( TSharedRef<IKeyArea>& KeyArea )
 
 bool FTrackNode::GetShotFilteredVisibilityToCache() const
 {
+	if (AssociatedType->HasShowableData() == false)
+	{
+		return false;
+	}
+
+	if ( Sections.Num() == 0 )
+	{
+		return AssociatedType->IsVisibleWhenEmpty();
+	}
+
 	// if no child sections are visible, neither is the entire section
 	bool bAnySectionsVisible = false;
 	for (int32 i = 0; i < Sections.Num() && !bAnySectionsVisible; ++i)
 	{
 		bAnySectionsVisible = GetSequencer().IsSectionVisible(Sections[i]->GetSectionObject());
 	}
-	return AssociatedType->HasShowableData() && bAnySectionsVisible;
+	return bAnySectionsVisible;
 }
 
 void FTrackNode::GetChildKeyAreaNodesRecursively(TArray< TSharedRef<class FSectionKeyAreaNode> >& OutNodes) const
@@ -565,6 +577,21 @@ TSharedRef<SWidget> FTrackNode::GenerateEditWidgetForOutliner()
 			}
 		}
 	}
+	else
+	{
+		FGuid ObjectBinding;
+		TSharedPtr<FSequencerDisplayNode> ParentNode = GetParent();
+		if ( ParentNode.IsValid() && ParentNode->GetType() == ESequencerNode::Object )
+		{
+			ObjectBinding = StaticCastSharedPtr<FObjectBindingNode>(ParentNode)->GetObjectBinding();
+		}
+
+		TSharedPtr<SWidget> EditWidget = AssociatedEditor.BuildOutlinerEditWidget( ObjectBinding, AssociatedType.Get() );
+		if ( EditWidget.IsValid() )
+		{
+			return EditWidget.ToSharedRef();
+		}
+	}
 	return FSequencerDisplayNode::GenerateEditWidgetForOutliner();
 }
 
diff --git a/Engine/Source/Editor/Sequencer/Private/SequencerDisplayNode.h b/Engine/Source/Editor/Sequencer/Private/SequencerDisplayNode.h
--- a/Engine/Source/Editor/Sequencer/Private/SequencerDisplayNode.h
+++ b/Engine/Source/Editor/Sequencer/Private/SequencerDisplayNode.h
@@ -72,9 +72,10 @@ public:
 	 * Adds a new section area for this node.
 	 * 
 	 * @param SectionName		Name of the section area
-	 * @param AssociatedType	The track associated with sections in this node
+	 * @param AssociatedTrack	The track associated with sections in this node
+	 * @param AssociatedEditor	The track editor for the associated track
 	 */
-	TSharedRef<class FTrackNode> AddSectionAreaNode( FName SectionName, UMovieSceneTrack& AssociatedTrack );
+	TSharedRef<class FTrackNode> AddSectionAreaNode( FName SectionName, UMovieSceneTrack& AssociatedTrack, FMovieSceneTrackEditor& AssociatedEditor );
 
 	/**
 	 * Adds a key area to this node
@@ -387,12 +388,13 @@ public:
 	/**
 	 * Constructor
 	 * 
-	 * @param InNodeName	The name identifier of then node
-	 * @param InDisplayName	Display name of the section area
-	 * @param InParentNode	The parent of this node or NULL if this is a root node
-	 * @param InParentTree	The tree this node is in
+	 * @param InNodeName The name identifier of then node
+	 * @param InAssociatedType The track that this node represents.
+	 * @param InAssociatedEditor The track editor for the track that this node represents.
+	 * @param InParentNode The parent of this node or NULL if this is a root node
+	 * @param InParentTree The tree this node is in
 	 */
-	FTrackNode( FName NodeName, UMovieSceneTrack& InAssociatedType, TSharedPtr<FSequencerDisplayNode> InParentNode, FSequencerNodeTree& InParentTree );
+	FTrackNode( FName NodeName, UMovieSceneTrack& InAssociatedType, FMovieSceneTrackEditor& InAssociatedEditor, TSharedPtr<FSequencerDisplayNode> InParentNode, FSequencerNodeTree& InParentTree );
 
 	/** FSequencerDisplayNodeInterface */
 	virtual ESequencerNode::Type GetType() const override { return ESequencerNode::Track; }
@@ -445,6 +447,8 @@ private:
 	TSharedPtr< FSectionKeyAreaNode > TopLevelKeyNode;
 	/** The type associated with the sections in this node */
 	TWeakObjectPtr<UMovieSceneTrack> AssociatedType;
+	/** The track editor for the track associated with this node. */
+	FMovieSceneTrackEditor& AssociatedEditor;
 };
 
 /**
diff --git a/Engine/Source/Editor/Sequencer/Private/SequencerNodeTree.cpp b/Engine/Source/Editor/Sequencer/Private/SequencerNodeTree.cpp
--- a/Engine/Source/Editor/Sequencer/Private/SequencerNodeTree.cpp
+++ b/Engine/Source/Editor/Sequencer/Private/SequencerNodeTree.cpp
@@ -38,7 +38,7 @@ void FSequencerNodeTree::Update()
 	{
 		UMovieSceneTrack& TrackRef = *Track;
 
-		TSharedRef<FTrackNode> SectionNode = MakeShareable( new FTrackNode( TrackRef.GetTrackName(), TrackRef, NULL, *this ) );
+		TSharedRef<FTrackNode> SectionNode = MakeShareable( new FTrackNode( TrackRef.GetTrackName(), TrackRef, *FindOrAddTypeEditor(TrackRef), NULL, *this ) );
 		NewRootNodes.Add( SectionNode );
 	
 		MakeSectionInterfaces( TrackRef, SectionNode );
@@ -70,7 +70,7 @@ void FSequencerNodeTree::Update()
 			FName SectionName = TrackRef.GetTrackName();
 			check( SectionName != NAME_None );
 
-			TSharedRef<FTrackNode> SectionAreaNode = ObjectBindingNode->AddSectionAreaNode( SectionName, TrackRef );
+			TSharedRef<FTrackNode> SectionAreaNode = ObjectBindingNode->AddSectionAreaNode( SectionName, TrackRef, *FindOrAddTypeEditor( TrackRef ) );
 			MakeSectionInterfaces( TrackRef, SectionAreaNode );
 		}
 	}
@@ -103,7 +103,7 @@ void FSequencerNodeTree::Update()
 	UMovieSceneTrack* ShotTrack = MovieScene->GetShotTrack();
 	if( ShotTrack )
 	{
-		TSharedRef<FTrackNode> SectionNode = MakeShareable( new FTrackNode( ShotTrack->GetTrackName(), *ShotTrack, NULL, *this ) );
+		TSharedRef<FTrackNode> SectionNode = MakeShareable( new FTrackNode( ShotTrack->GetTrackName(), *ShotTrack, *FindOrAddTypeEditor( *ShotTrack ), NULL, *this ) );
 
 		// Shot track always comes first
 		RootNodes.Add( SectionNode );
diff --git a/Engine/Source/Editor/Sequencer/Public/MovieSceneTrackEditor.h b/Engine/Source/Editor/Sequencer/Public/MovieSceneTrackEditor.h
--- a/Engine/Source/Editor/Sequencer/Public/MovieSceneTrackEditor.h
+++ b/Engine/Source/Editor/Sequencer/Public/MovieSceneTrackEditor.h
@@ -228,6 +228,14 @@ public:
 	 */
 	virtual void BuildObjectBindingEditButtons(TSharedPtr<SHorizontalBox> EditBox, const FGuid& ObjectBinding, const UClass* ObjectClass) {}
 
+	/**
+	 * Builds an edit widget for the outliner nodes which represent tracks which are edited by this editor.
+	 * @param ObjectBinding The object binding associated with the track being edited by this editor.
+	 * @param Track The track being edited by this editor.
+	 * @returns The the widget to display in the outliner, or an empty shared ptr if not widget is to be displayed.
+	 */
+	virtual TSharedPtr<SWidget> BuildOutlinerEditWidget(const FGuid& ObjectBinding, UMovieSceneTrack* Track) { return TSharedPtr<SWidget>(); }
+
 private:
 	/** The sequencer bound to this handler.  Used to access movie scene and time info during auto-key */
 	TWeakPtr<ISequencer> Sequencer;
diff --git a/Engine/Source/Runtime/MovieScene/Public/MovieSceneTrack.h b/Engine/Source/Runtime/MovieScene/Public/MovieSceneTrack.h
--- a/Engine/Source/Runtime/MovieScene/Public/MovieSceneTrack.h
+++ b/Engine/Source/Runtime/MovieScene/Public/MovieSceneTrack.h
@@ -49,6 +49,11 @@ public:
 	 */
 	virtual bool HasShowableData() const {return true;}
 
+	/**
+	 * @returns True if this track should generate a display node even when it's empty.
+	 */
+	virtual bool IsVisibleWhenEmpty() const { return false; }
+
 	/**
 	 * @return Whether or not this track supports multiple row indices
 	 */
diff --git a/Engine/Source/Runtime/MovieSceneTracks/Private/Sections/MovieSceneMaterialParameterSection.cpp b/Engine/Source/Runtime/MovieSceneTracks/Private/Sections/MovieSceneMaterialParameterSection.cpp
new file mode 100644
--- /dev/null
+++ b/Engine/Source/Runtime/MovieSceneTracks/Private/Sections/MovieSceneMaterialParameterSection.cpp
@@ -0,0 +1,46 @@
+// Copyright 1998-2015 Epic Games, Inc. All Rights Reserved.
+
+#include "MovieSceneTracksPrivatePCH.h"
+#include "MovieSceneMaterialParameterSection.h"
+
+FScalarParameterNameAndCurve::FScalarParameterNameAndCurve( FName InParameterName )
+{
+	ParameterName = InParameterName;
+}
+
+UMovieSceneMaterialParameterSection::UMovieSceneMaterialParameterSection( const FObjectInitializer& ObjectInitializer )
+	: Super( ObjectInitializer )
+{
+}
+
+void UMovieSceneMaterialParameterSection::AddScalarParameterKey( FName InParameterName, float InTime, float InValue )
+{
+	FRichCurve* ExistingCurve = GetScalarParameterCurve(InParameterName);
+	if ( ExistingCurve != nullptr )
+	{
+		ExistingCurve->AddKey(InTime, InValue);
+	}
+	else
+	{
+		int32 NewIndex = ScalarParameterNamesAndCurves.Add( FScalarParameterNameAndCurve() );
+		ScalarParameterNamesAndCurves[NewIndex].ParameterName = InParameterName;
+		ScalarParameterNamesAndCurves[NewIndex].ParameterCurve.AddKey( InTime, InValue );
+	}
+}
+
+TArray<FScalarParameterNameAndCurve>* UMovieSceneMaterialParameterSection::GetScalarParameterNamesAndCurves()
+{
+	return &ScalarParameterNamesAndCurves;
+}
+
+FRichCurve* UMovieSceneMaterialParameterSection::GetScalarParameterCurve( FName InParameterName )
+{
+	for ( FScalarParameterNameAndCurve& ScalarParameterNameAndCurve : ScalarParameterNamesAndCurves )
+	{
+		if ( ScalarParameterNameAndCurve.ParameterName == InParameterName )
+		{
+			return &ScalarParameterNameAndCurve.ParameterCurve;
+		}
+	}
+	return nullptr;
+}
\ No newline at end of file
diff --git a/Engine/Source/Runtime/MovieSceneTracks/Private/TrackInstances/MovieSceneMaterialTrackInstance.cpp b/Engine/Source/Runtime/MovieSceneTracks/Private/TrackInstances/MovieSceneMaterialTrackInstance.cpp
new file mode 100644
--- /dev/null
+++ b/Engine/Source/Runtime/MovieSceneTracks/Private/TrackInstances/MovieSceneMaterialTrackInstance.cpp
@@ -0,0 +1,112 @@
+// Copyright 1998-2015 Epic Games, Inc. All Rights Reserved.
+
+#include "MovieSceneTracksPrivatePCH.h"
+#include "MovieSceneMaterialTrack.h"
+#include "MovieSceneMaterialTrackInstance.h"
+
+FMovieSceneMaterialTrackInstance::FMovieSceneMaterialTrackInstance( UMovieSceneMaterialTrack& InMaterialTrack )
+{
+	MaterialTrack = &InMaterialTrack;
+}
+
+void FMovieSceneMaterialTrackInstance::SaveState(const TArray<UObject*>& RuntimeObjects)
+{
+	for( UObject* Object : RuntimeObjects )
+	{
+		UMaterialInterface* CurrentMaterial = GetMaterialForObject(Object);
+		if ( CurrentMaterial->IsA( UMaterialInstanceDynamic::StaticClass() ) )
+		{
+			TWeakObjectPtr<UMaterialInterface>* OriginalMaterial = DynamicMaterialToOriginalMaterialMap.Find( FObjectKey( CurrentMaterial ) );
+			if ( OriginalMaterial != nullptr )
+			{
+				CurrentMaterial = OriginalMaterial->Get();
+			}
+			else
+			{
+				CurrentMaterial = nullptr;
+			}
+		}
+		if (CurrentMaterial != nullptr)
+		{
+			RuntimeObjectToOriginalMaterialMap.Add( FObjectKey(Object), CurrentMaterial );
+		}
+	}
+}
+
+void FMovieSceneMaterialTrackInstance::RestoreState(const TArray<UObject*>& RuntimeObjects)
+{
+	for( UObject* Object : RuntimeObjects )
+	{
+		if (!IsValid(Object))
+		{
+			continue;
+		}
+
+		TWeakObjectPtr<UMaterialInterface>* OriginalMaterial = RuntimeObjectToOriginalMaterialMap.Find( FObjectKey( Object ) );
+		if ( OriginalMaterial != nullptr && OriginalMaterial->IsValid() )
+		{
+			SetMaterialForObject(Object, OriginalMaterial->Get());
+		}
+	}
+}
+
+void FMovieSceneMaterialTrackInstance::Update( float Position, float LastPosition, const TArray<UObject*>& RuntimeObjects, class IMovieScenePlayer& Player ) 
+{
+	TArray<FScalarParameterNameAndValue> ScalarValues;
+	MaterialTrack->Eval(Position, ScalarValues);
+
+	for ( UMaterialInstanceDynamic* DynamicMaterialInstance : DynamicMaterialInstances )
+	{
+		for ( const FScalarParameterNameAndValue& ScalarValue : ScalarValues )
+		{
+			DynamicMaterialInstance->SetScalarParameterValue(ScalarValue.ParameterName, ScalarValue.Value);
+		}
+	}
+}
+
+void FMovieSceneMaterialTrackInstance::RefreshInstance( const TArray<UObject*>& RuntimeObjects, class IMovieScenePlayer& Player )
+{
+	DynamicMaterialInstances.Empty();
+	DynamicMaterialToOriginalMaterialMap.Empty();
+	for ( UObject* Object : RuntimeObjects )
+	{
+		UMaterialInterface* Material = GetMaterialForObject( Object );
+		UMaterialInstanceDynamic* DynamicMaterialInstance = Cast<UMaterialInstanceDynamic>(Material);
+		if ( DynamicMaterialInstance == nullptr )
+		{
+			DynamicMaterialInstance = UMaterialInstanceDynamic::Create(Material, Object);
+			SetMaterialForObject( Object, DynamicMaterialInstance );
+			DynamicMaterialToOriginalMaterialMap.Add( FObjectKey( DynamicMaterialInstance ), Material );
+		}
+		DynamicMaterialInstances.Add( DynamicMaterialInstance );
+	}
+}
+
+void FMovieSceneMaterialTrackInstance::ClearInstance( IMovieScenePlayer& Player )
+{
+	DynamicMaterialInstances.Empty();
+	DynamicMaterialToOriginalMaterialMap.Empty();
+}
+
+FMovieSceneComponentMaterialTrackInstance::FMovieSceneComponentMaterialTrackInstance( UMovieSceneComponentMaterialTrack& InMaterialTrack )
+	: FMovieSceneMaterialTrackInstance( InMaterialTrack )
+	, MaterialIndex( InMaterialTrack.GetMaterialIndex() )
+{
+}
+
+UMaterialInterface* FMovieSceneComponentMaterialTrackInstance::GetMaterialForObject( UObject* Object ) const 
+{
+	UPrimitiveComponent* Component = Cast<UPrimitiveComponent>(Object);
+	if ( Component != nullptr )
+	{
+		return Component->GetMaterial( MaterialIndex );
+	}
+	return nullptr;
+}
+
+void FMovieSceneComponentMaterialTrackInstance::SetMaterialForObject( UObject* Object, UMaterialInterface* Material )
+{
+	checkf(Object->IsA(UPrimitiveComponent::StaticClass()), TEXT("Unsupported object type."));
+	UPrimitiveComponent* Component = Cast<UPrimitiveComponent>( Object );
+	Component->SetMaterial(MaterialIndex, Material);
+}
diff --git a/Engine/Source/Runtime/MovieSceneTracks/Private/Tracks/MovieSceneMaterialTrack.cpp b/Engine/Source/Runtime/MovieSceneTracks/Private/Tracks/MovieSceneMaterialTrack.cpp
new file mode 100644
--- /dev/null
+++ b/Engine/Source/Runtime/MovieSceneTracks/Private/Tracks/MovieSceneMaterialTrack.cpp
@@ -0,0 +1,98 @@
+// Copyright 1998-2015 Epic Games, Inc. All Rights Reserved.
+
+#include "MovieSceneTracksPrivatePCH.h"
+#include "MovieSceneMaterialParameterSection.h"
+#include "MovieSceneMaterialTrack.h"
+#include "IMovieScenePlayer.h"
+#include "MovieSceneMaterialTrackInstance.h"
+
+UMovieSceneMaterialTrack::UMovieSceneMaterialTrack( const FObjectInitializer& ObjectInitializer )
+	: Super( ObjectInitializer )
+{
+}
+
+UMovieSceneSection* UMovieSceneMaterialTrack::CreateNewSection()
+{
+	return NewObject<UMovieSceneMaterialParameterSection>( this, UMovieSceneMaterialParameterSection::StaticClass(), NAME_None, RF_Transactional );
+}
+
+void UMovieSceneMaterialTrack::RemoveAllAnimationData()
+{
+	Sections.Empty();
+}
+
+bool UMovieSceneMaterialTrack::HasSection( UMovieSceneSection* Section ) const
+{
+	return Sections.Contains( Section );
+}
+
+void UMovieSceneMaterialTrack::AddSection( UMovieSceneSection* Section )
+{
+	Sections.Add( Section );
+}
+
+void UMovieSceneMaterialTrack::RemoveSection( UMovieSceneSection* Section )
+{
+	Sections.Remove( Section );
+}
+
+bool UMovieSceneMaterialTrack::IsEmpty() const
+{
+	return Sections.Num() == 0;
+}
+
+TRange<float> UMovieSceneMaterialTrack::GetSectionBoundaries() const
+{
+	TArray< TRange<float> > Bounds;
+
+	for ( int32 SectionIndex = 0; SectionIndex < Sections.Num(); ++SectionIndex )
+	{
+		Bounds.Add( Sections[SectionIndex]->GetRange() );
+	}
+
+	return TRange<float>::Hull( Bounds );
+}
+
+const TArray<UMovieSceneSection*>& UMovieSceneMaterialTrack::GetAllSections() const
+{
+	return Sections;
+}
+
+void UMovieSceneMaterialTrack::AddScalarParameterKey( FName ParameterName, float Time, float Value )
+{
+	UMovieSceneMaterialParameterSection* NearestSection = Cast<UMovieSceneMaterialParameterSection>(MovieSceneHelpers::FindNearestSectionAtTime(Sections, Time));
+	if ( NearestSection == nullptr )
+	{
+		NearestSection = Cast<UMovieSceneMaterialParameterSection>(CreateNewSection());
+		NearestSection->SetStartTime( Time );
+		NearestSection->SetEndTime( Time );
+		Sections.Add( NearestSection );
+	}
+	NearestSection->AddScalarParameterKey(ParameterName, Time, Value);
+}
+
+void UMovieSceneMaterialTrack::Eval( float Position, TArray<FScalarParameterNameAndValue>& OutScalarValues ) const
+{
+	for ( UMovieSceneSection* Section : Sections )
+	{
+		UMovieSceneMaterialParameterSection* MaterialParameterSection = Cast<UMovieSceneMaterialParameterSection>( Section );
+		if ( MaterialParameterSection != nullptr )
+		{
+			for ( const FScalarParameterNameAndCurve& ScalarParameterNameAndCurve : *MaterialParameterSection->GetScalarParameterNamesAndCurves() )
+			{
+				OutScalarValues.Add( FScalarParameterNameAndValue( ScalarParameterNameAndCurve.ParameterName, ScalarParameterNameAndCurve.ParameterCurve.Eval( Position ) ) );
+			}
+		}
+	}
+}
+
+UMovieSceneComponentMaterialTrack::UMovieSceneComponentMaterialTrack( const FObjectInitializer& ObjectInitializer )
+	: Super( ObjectInitializer )
+{
+	
+}
+
+TSharedPtr<IMovieSceneTrackInstance> UMovieSceneComponentMaterialTrack::CreateInstance()
+{
+	return MakeShareable( new FMovieSceneComponentMaterialTrackInstance( *this ) );
+}
diff --git a/Engine/Source/Runtime/MovieSceneTracks/Public/Sections/MovieSceneMaterialParameterSection.h b/Engine/Source/Runtime/MovieSceneTracks/Public/Sections/MovieSceneMaterialParameterSection.h
new file mode 100644
--- /dev/null
+++ b/Engine/Source/Runtime/MovieSceneTracks/Public/Sections/MovieSceneMaterialParameterSection.h
@@ -0,0 +1,56 @@
+// Copyright 1998-2015 Epic Games, Inc. All Rights Reserved.
+
+#pragma once
+
+#include "MovieSceneFloatSection.h"
+#include "MovieSceneMaterialParameterSection.generated.h"
+
+/**
+ * Structure representing an animated scalar parameter and it's associated animation curve.
+ */
+USTRUCT()
+struct FScalarParameterNameAndCurve
+{
+	GENERATED_USTRUCT_BODY()
+
+	FScalarParameterNameAndCurve() { }
+
+	/** Creates a new FScalarParameterNameAndCurve for a specific scalar parameter. */
+	FScalarParameterNameAndCurve( FName InParameterName );
+
+	/** The name of the scalar parameter which is being animated. */
+	UPROPERTY()
+	FName ParameterName;
+
+	/** The curve which contains the animation data for the scalar parameter. */
+	UPROPERTY()
+	FRichCurve ParameterCurve;
+};
+
+/**
+ * A movie scene section representing animation data for material parameters. 
+ */
+UCLASS( MinimalAPI )
+class UMovieSceneMaterialParameterSection : public UMovieSceneSection
+{
+	GENERATED_UCLASS_BODY()
+
+public:
+	/** Adds a a key for a specific scalar parameter at the specified time with the specified value. */
+	void AddScalarParameterKey(FName InParameterName, float InTime, float InValue);
+
+	/** Gets the animated scalar parameters and their associated curves. */
+	MOVIESCENETRACKS_API TArray<FScalarParameterNameAndCurve>* GetScalarParameterNamesAndCurves();
+
+private:
+	/** Gets a curve for a specific scalar parameter name. */
+	FRichCurve* GetScalarParameterCurve( FName InParameterName );
+
+private:
+	/**
+	 * The scalar parameter names and their associated curves.
+	 * NOTE: This isn't stored in a map since we want order added preserved and random access doesn't need to be fast.
+	 */
+	UPROPERTY()
+	TArray<FScalarParameterNameAndCurve> ScalarParameterNamesAndCurves;
+};
\ No newline at end of file
diff --git a/Engine/Source/Runtime/MovieSceneTracks/Public/TrackInstances/MovieSceneMaterialTrackInstance.h b/Engine/Source/Runtime/MovieSceneTracks/Public/TrackInstances/MovieSceneMaterialTrackInstance.h
new file mode 100644
--- /dev/null
+++ b/Engine/Source/Runtime/MovieSceneTracks/Public/TrackInstances/MovieSceneMaterialTrackInstance.h
@@ -0,0 +1,68 @@
+// Copyright 1998-2015 Epic Games, Inc. All Rights Reserved.
+
+#pragma once
+
+#include "IMovieSceneTrackInstance.h"
+#include "ObjectKey.h"
+
+class FTrackInstancePropertyBindings;
+class UMovieSceneMaterialTrack;
+class UMovieSceneComponentMaterialTrack;
+
+/**
+ * A movie scene track instance for material tracks.
+ */
+class FMovieSceneMaterialTrackInstance
+	: public IMovieSceneTrackInstance
+{
+public:
+	FMovieSceneMaterialTrackInstance( UMovieSceneMaterialTrack& InMaterialTrack );
+
+	/** IMovieSceneTrackInstance interface */
+	virtual void SaveState (const TArray<UObject*>& RuntimeObjects) override;
+	virtual void RestoreState (const TArray<UObject*>& RuntimeObjects) override;
+	virtual void Update( float Position, float LastPosition, const TArray<UObject*>& RuntimeObjects, IMovieScenePlayer& Player ) override;
+	virtual void RefreshInstance( const TArray<UObject*>& RuntimeObjects, IMovieScenePlayer& Player ) override;
+	virtual void ClearInstance( IMovieScenePlayer& Player ) override;
+
+protected:
+	/**
+	 * Gets the material for a specific runtime object.
+	 */
+	virtual UMaterialInterface* GetMaterialForObject( UObject* Object ) const = 0;
+
+	/**
+	 * Sets the material for a specific runtime object.
+	 */
+	virtual void SetMaterialForObject( UObject* Object, UMaterialInterface* Material) = 0;
+
+private:
+	/** Track that is being instanced */
+	UMovieSceneMaterialTrack* MaterialTrack;
+
+	/** Map from runtime object to original material */
+	TMap<FObjectKey, TWeakObjectPtr<UMaterialInterface>> RuntimeObjectToOriginalMaterialMap;
+
+	/** Map from dynamic material to original material */
+	TMap<FObjectKey, TWeakObjectPtr<UMaterialInterface>> DynamicMaterialToOriginalMaterialMap;
+
+	/** A list of dynamic material instances which need to be updated. */
+	TArray<UMaterialInstanceDynamic*> DynamicMaterialInstances;
+};
+
+/** A material track instance specialized for component materials. */
+class FMovieSceneComponentMaterialTrackInstance : public FMovieSceneMaterialTrackInstance
+{
+public:
+	/** Creates a new track instance with a specific component material track. */
+	FMovieSceneComponentMaterialTrackInstance( UMovieSceneComponentMaterialTrack& InMaterialTrack );
+
+protected:
+	/** FMovieSceneMaterialTrackInstance interface. */
+	virtual UMaterialInterface* GetMaterialForObject( UObject* Object ) const override;
+	virtual void SetMaterialForObject( UObject* Object, UMaterialInterface* Material ) override;
+
+private:
+	/** The index of the material on the associated component. */
+	int32 MaterialIndex;
+};
diff --git a/Engine/Source/Runtime/MovieSceneTracks/Public/Tracks/MovieSceneMaterialTrack.h b/Engine/Source/Runtime/MovieSceneTracks/Public/Tracks/MovieSceneMaterialTrack.h
new file mode 100644
--- /dev/null
+++ b/Engine/Source/Runtime/MovieSceneTracks/Public/Tracks/MovieSceneMaterialTrack.h
@@ -0,0 +1,102 @@
+// Copyright 1998-2015 Epic Games, Inc. All Rights Reserved.
+
+#pragma once
+
+#include "MovieSceneTrack.h"
+#include "MovieSceneMaterialTrack.generated.h"
+
+/**
+ * Structure representing the animated value of a scalar parameter.
+ */
+struct FScalarParameterNameAndValue
+{
+	/** Creates a new FScalarParameterAndValue with a parameter name and a value. */
+	FScalarParameterNameAndValue( FName InParameterName, float InValue )
+	{
+		ParameterName = InParameterName;
+		Value = InValue;
+	}
+
+	/** The name of the scalar parameter. */
+	FName ParameterName;
+
+	/** The animated value of the scalar parameter. */
+	float Value;
+};
+
+/**
+ * Handles manipulation of material parameters in a movie scene.
+ */
+UCLASS( MinimalAPI )
+class UMovieSceneMaterialTrack : public UMovieSceneTrack
+{
+	GENERATED_UCLASS_BODY()
+
+public:
+	/** UMovieSceneTrack interface */
+	virtual UMovieSceneSection* CreateNewSection() override;
+	virtual void RemoveAllAnimationData() override;
+	virtual bool HasSection( UMovieSceneSection* Section ) const override;
+	virtual void AddSection( UMovieSceneSection* Section ) override;
+	virtual void RemoveSection( UMovieSceneSection* Section ) override;
+	virtual bool IsEmpty() const override;
+	virtual TRange<float> GetSectionBoundaries() const override;
+	virtual const TArray<UMovieSceneSection*>& GetAllSections() const override;
+	virtual bool HasShowableData() const override { return true; }
+	virtual bool IsVisibleWhenEmpty() const override { return true; }
+
+	/**
+	 * Adds a scalar parameter key to the track. 
+	 * @param ParameterName The name of the parameter to add a key for.
+	 * @param Time The time to add the new key.
+	 * @param The value for the new key.
+	 */
+	void MOVIESCENETRACKS_API AddScalarParameterKey( FName ParameterName, float Position, float Value );
+
+	/**
+	 * Gets the animated values for this track.
+	 * @param Position The playback position to use for evaluation.
+	 * @param OutScalarValues An array of FScalarParameterNameAndValue objects representing each animated scalar parameter and it's animated value.
+	 */
+	void Eval( float Position, TArray<FScalarParameterNameAndValue>& OutScalarValues ) const;
+
+private:
+	/** The sections owned by this track .*/
+	UPROPERTY()
+	TArray<UMovieSceneSection*> Sections;
+};
+
+/**
+ * A material track which is specialized for animation materials which are owned by actor components.
+ */
+UCLASS( MinimalAPI )
+class UMovieSceneComponentMaterialTrack : public UMovieSceneMaterialTrack
+{
+	GENERATED_UCLASS_BODY()
+
+public:
+	/** UMovieSceneTrack interface */
+	virtual TSharedPtr<IMovieSceneTrackInstance> CreateInstance() override;
+	virtual FName GetTrackName() const override { return TrackName; }
+
+	/** Gets the index of the material in the component. */
+	int32 GetMaterialIndex() const { return MaterialIndex; }
+
+	/** Sets the index of the material in the component. */
+	void SetMaterialIndex(int32 InMaterialIndex) 
+	{
+		MaterialIndex = InMaterialIndex;
+		TrackName = *FString::Printf(TEXT("Material Element %i"), MaterialIndex);
+	}
+
+private:
+	/** The name of this track .*/
+	UPROPERTY()
+	FName TrackName;
+
+	/** The index of this material this track is animating. */
+	UPROPERTY()
+	int32 MaterialIndex;
+};
+
+
