Sequencer: Clean up right click menu so that it doesn't create an extra separator.

diff --git a/Engine/Source/Editor/MovieSceneTools/Private/TrackEditors/AttachTrackEditor.cpp b/Engine/Source/Editor/MovieSceneTools/Private/TrackEditors/AttachTrackEditor.cpp
--- a/Engine/Source/Editor/MovieSceneTools/Private/TrackEditors/AttachTrackEditor.cpp
+++ b/Engine/Source/Editor/MovieSceneTools/Private/TrackEditors/AttachTrackEditor.cpp
@@ -149,8 +149,6 @@ public:
 
 			MenuBuilder.EndSection();
 		}
-
-		MenuBuilder.AddMenuSeparator();
 	}
 
 private:
diff --git a/Engine/Source/Editor/MovieSceneTools/Private/TrackEditors/PathTrackEditor.cpp b/Engine/Source/Editor/MovieSceneTools/Private/TrackEditors/PathTrackEditor.cpp
--- a/Engine/Source/Editor/MovieSceneTools/Private/TrackEditors/PathTrackEditor.cpp
+++ b/Engine/Source/Editor/MovieSceneTools/Private/TrackEditors/PathTrackEditor.cpp
@@ -136,8 +136,6 @@ public:
 
 			MenuBuilder.EndSection();
 		}
-
-		MenuBuilder.AddMenuSeparator();
 	}
 
 private:
