Sequencer - Add vector parameter support to the material track, plus general material track fixes
+ Add undo support for adding and removing material tracks and material track parameters.
+ Add the ability to request deleting category and key area nodes, and use that to implement deleting previously added material parameters.
+ Fix general sequencer tree sorting so that only objects tracks are sorted by name.

diff --git a/Engine/Plugins/MovieScene/ActorAnimationEditor/Source/ActorAnimationEditor/Private/ActorAnimationEditorToolkit.cpp b/Engine/Plugins/MovieScene/ActorAnimationEditor/Source/ActorAnimationEditor/Private/ActorAnimationEditorToolkit.cpp
--- a/Engine/Plugins/MovieScene/ActorAnimationEditor/Source/ActorAnimationEditor/Private/ActorAnimationEditorToolkit.cpp
+++ b/Engine/Plugins/MovieScene/ActorAnimationEditor/Source/ActorAnimationEditor/Private/ActorAnimationEditorToolkit.cpp
@@ -12,6 +12,7 @@
 #include "MovieSceneBinding.h"
 #include "MovieScene.h"
 #include "MovieSceneMaterialTrack.h"
+#include "ScopedTransaction.h"
 
 
 #define LOCTEXT_NAMESPACE "ActorAnimationEditor"
@@ -194,9 +195,14 @@ void FActorAnimationEditorToolkit::HandleAddComponentMaterialActionExecute( UPri
 			}
 			if ( bHasMaterialTrack == false )
 			{
-				UMovieSceneComponentMaterialTrack* MaterialTrack = Cast<UMovieSceneComponentMaterialTrack>(
-					Sequencer->GetFocusedMovieSceneSequence()->GetMovieScene()->AddTrack( UMovieSceneComponentMaterialTrack::StaticClass(), ObjectHandle));
-				MaterialTrack->SetMaterialIndex(MaterialIndex);
+				const FScopedTransaction Transaction( LOCTEXT( "AddComponentMaterialTrack", "Add component material track" ) );
+
+				UMovieScene* FocusedMovieScene = Sequencer->GetFocusedMovieSceneSequence()->GetMovieScene();
+				FocusedMovieScene->Modify();
+				UMovieSceneComponentMaterialTrack* MaterialTrack = Cast<UMovieSceneComponentMaterialTrack>( FocusedMovieScene->AddTrack( UMovieSceneComponentMaterialTrack::StaticClass(), ObjectHandle ) );
+				MaterialTrack->Modify();
+				MaterialTrack->SetMaterialIndex( MaterialIndex );
+				
 				Sequencer->NotifyMovieSceneDataChanged();
 			}
 		}
diff --git a/Engine/Source/Editor/MovieSceneTools/Private/Sections/MaterialParameterSection.cpp b/Engine/Source/Editor/MovieSceneTools/Private/Sections/MaterialParameterSection.cpp
--- a/Engine/Source/Editor/MovieSceneTools/Private/Sections/MaterialParameterSection.cpp
+++ b/Engine/Source/Editor/MovieSceneTools/Private/Sections/MaterialParameterSection.cpp
@@ -4,15 +4,98 @@
 #include "MaterialParameterSection.h"
 #include "MovieSceneMaterialParameterSection.h"
 
+#define LOCTEXT_NAMESPACE "MaterialParameterSection"
+
+DECLARE_DELEGATE(FLayoutParameterDelegate);
+
+struct FIndexAndLayoutParameterDelegate
+{
+	int32 Index;
+	FLayoutParameterDelegate LayoutParameterDelegate;
+
+	bool operator<(FIndexAndLayoutParameterDelegate const& Other) const
+	{
+		return Index < Other.Index;
+	}
+};
+
+void LayoutScalarParameter( ISectionLayoutBuilder* LayoutBuilder, FScalarParameterNameAndCurve* ScalarParameterNameAndCurve, UMovieSceneMaterialParameterSection* ParameterSection )
+{
+	LayoutBuilder->AddKeyArea(
+		ScalarParameterNameAndCurve->ParameterName,
+		FText::FromName( ScalarParameterNameAndCurve->ParameterName ),
+		MakeShareable( new FFloatCurveKeyArea( &ScalarParameterNameAndCurve->ParameterCurve, ParameterSection ) ) );
+}
+
+void LayoutVectorParameter( ISectionLayoutBuilder* LayoutBuilder, FVectorParameterNameAndCurves* VectorParameterNameAndCurves, UMovieSceneMaterialParameterSection* ParameterSection )
+{
+	LayoutBuilder->PushCategory( VectorParameterNameAndCurves->ParameterName, FText::FromName( VectorParameterNameAndCurves->ParameterName ) );
+	{
+		LayoutBuilder->AddKeyArea( "R", NSLOCTEXT( "MaterialParameterSection", "RedArea", "Red" ),
+			MakeShareable( new FFloatCurveKeyArea( &VectorParameterNameAndCurves->RedCurve, ParameterSection ) ) );
+		LayoutBuilder->AddKeyArea( "G", NSLOCTEXT( "MaterialParameterSection", "GreenArea", "Green" ), 
+			MakeShareable( new FFloatCurveKeyArea( &VectorParameterNameAndCurves->GreenCurve, ParameterSection ) ) );
+		LayoutBuilder->AddKeyArea( "B", NSLOCTEXT( "MaterialParameterSection", "BlueArea", "Blue" ), 
+			MakeShareable( new FFloatCurveKeyArea( &VectorParameterNameAndCurves->BlueCurve, ParameterSection ) ) );
+		LayoutBuilder->AddKeyArea( "A", NSLOCTEXT( "MaterialParameterSection", "OpacityArea", "Opacity" ), 
+			MakeShareable( new FFloatCurveKeyArea( &VectorParameterNameAndCurves->AlphaCurve, ParameterSection ) ) );
+	}
+	LayoutBuilder->PopCategory();
+}
+
 void FMaterialParameterSection::GenerateSectionLayout( class ISectionLayoutBuilder& LayoutBuilder ) const
 {
+	// TODO - This seems a bit convoluted here to have to use delegates to get the layout to happen in the right order based on the indices.
 	UMovieSceneMaterialParameterSection* ParameterSection = Cast<UMovieSceneMaterialParameterSection>( &SectionObject );
+	TArray<FIndexAndLayoutParameterDelegate> IndexAndLayoutDelegates;
+
+	// Collect scalar parameters.
 	TArray<FScalarParameterNameAndCurve>* ScalarParameterNamesAndCurves = ParameterSection->GetScalarParameterNamesAndCurves();
 	for ( int32 i = 0; i < ScalarParameterNamesAndCurves->Num(); i++)
 	{
-		LayoutBuilder.AddKeyArea( 
-			(*ScalarParameterNamesAndCurves)[i].ParameterName, 
-			FText::FromName( (*ScalarParameterNamesAndCurves)[i].ParameterName ),
-			MakeShareable( new FFloatCurveKeyArea( &(*ScalarParameterNamesAndCurves)[i].ParameterCurve, ParameterSection ) ) );
+		FIndexAndLayoutParameterDelegate IndexAndLayoutDelegate;
+		IndexAndLayoutDelegate.Index = (*ScalarParameterNamesAndCurves)[i].Index;
+		IndexAndLayoutDelegate.LayoutParameterDelegate.BindStatic(&LayoutScalarParameter, &LayoutBuilder, &(*ScalarParameterNamesAndCurves)[i], ParameterSection);
+		IndexAndLayoutDelegates.Add(IndexAndLayoutDelegate);
 	}
-}
\ No newline at end of file
+
+	// Collect vector parameters
+	TArray<FVectorParameterNameAndCurves>* VectorParameterNamesAndCurves = ParameterSection->GetVectorParameterNamesAndCurves();
+	for ( int32 i = 0; i < VectorParameterNamesAndCurves->Num(); i++ )
+	{
+		FIndexAndLayoutParameterDelegate IndexAndLayoutDelegate;
+		IndexAndLayoutDelegate.Index = (*VectorParameterNamesAndCurves)[i].Index;
+		IndexAndLayoutDelegate.LayoutParameterDelegate.BindStatic( &LayoutVectorParameter, &LayoutBuilder, &(*VectorParameterNamesAndCurves)[i], ParameterSection );
+		IndexAndLayoutDelegates.Add( IndexAndLayoutDelegate );
+	}
+
+	// Sort and perform layout
+	IndexAndLayoutDelegates.Sort();
+	for ( FIndexAndLayoutParameterDelegate& IndexAndLayoutDelegate : IndexAndLayoutDelegates )
+	{
+		IndexAndLayoutDelegate.LayoutParameterDelegate.Execute();
+	}
+}
+
+bool FMaterialParameterSection::RequestDeleteCategory( const TArray<FName>& CategoryNamePath )
+{
+	const FScopedTransaction Transaction( LOCTEXT( "DeleteVectorParameter", "Delete vector parameter" ) );
+	UMovieSceneMaterialParameterSection* ParameterSection = Cast<UMovieSceneMaterialParameterSection>( &SectionObject );
+	ParameterSection->Modify();
+	return ParameterSection->RemoveVectorParameter( CategoryNamePath[0] );
+}
+
+bool FMaterialParameterSection::RequestDeleteKeyArea( const TArray<FName>& KeyAreaNamePath )
+{
+	// Only handle paths with a single name, in all other cases the user is deleting a component of a vector parameter.
+	if ( KeyAreaNamePath.Num() == 1)
+	{
+		const FScopedTransaction Transaction( LOCTEXT( "DeleteScalarParameter", "Delete scalar parameter" ) );
+		UMovieSceneMaterialParameterSection* ParameterSection = Cast<UMovieSceneMaterialParameterSection>( &SectionObject );
+		ParameterSection->Modify();
+		return ParameterSection->RemoveScalarParameter( KeyAreaNamePath[0] );
+	}
+	return false;
+}
+
+#undef LOCTEXT_NAMESPACE
\ No newline at end of file
diff --git a/Engine/Source/Editor/MovieSceneTools/Private/Sections/MaterialParameterSection.h b/Engine/Source/Editor/MovieSceneTools/Private/Sections/MaterialParameterSection.h
--- a/Engine/Source/Editor/MovieSceneTools/Private/Sections/MaterialParameterSection.h
+++ b/Engine/Source/Editor/MovieSceneTools/Private/Sections/MaterialParameterSection.h
@@ -11,5 +11,8 @@ public:
 	FMaterialParameterSection( UMovieSceneSection& InSectionObject, FName SectionName )
 		: FPropertySection( InSectionObject, SectionName ) {}
 
+	// ISequencerSection interface.
 	virtual void GenerateSectionLayout( class ISectionLayoutBuilder& LayoutBuilder ) const override;
+	virtual bool RequestDeleteCategory( const TArray<FName>& CategoryNamePath ) override;
+	virtual bool RequestDeleteKeyArea( const TArray<FName>& KeyAreaNamePath ) override;
 };
\ No newline at end of file
diff --git a/Engine/Source/Editor/MovieSceneTools/Private/TrackEditors/MaterialTrackEditor.cpp b/Engine/Source/Editor/MovieSceneTools/Private/TrackEditors/MaterialTrackEditor.cpp
--- a/Engine/Source/Editor/MovieSceneTools/Private/TrackEditors/MaterialTrackEditor.cpp
+++ b/Engine/Source/Editor/MovieSceneTools/Private/TrackEditors/MaterialTrackEditor.cpp
@@ -52,6 +52,23 @@ TSharedPtr<SWidget> FMaterialTrackEditor::BuildOutlinerEditWidget( const FGuid&
 		];
 }
 
+struct FParameterNameAndAction
+{
+	FName ParameterName;
+	FUIAction Action;
+
+	FParameterNameAndAction( FName InParameterName, FUIAction InAction )
+	{
+		ParameterName = InParameterName;
+		Action = InAction;
+	}
+
+	bool operator<(FParameterNameAndAction const& Other) const
+	{
+		return ParameterName < Other.ParameterName;
+	}
+};
+
 TSharedRef<SWidget> FMaterialTrackEditor::OnGetAddParameterMenuContent( FGuid ObjectBinding, UMovieSceneMaterialTrack* MaterialTrack )
 {
 	FMenuBuilder AddParameterMenuBuilder( true, nullptr );
@@ -59,14 +76,35 @@ TSharedRef<SWidget> FMaterialTrackEditor::OnGetAddParameterMenuContent( FGuid Ob
 	UMaterial* Material = GetMaterialForTrack( ObjectBinding, MaterialTrack );
 	if ( Material != nullptr )
 	{
+		TArray<FParameterNameAndAction> ParameterNamesAndActions;
+
+		// Collect scalar parameters.
 		TArray<FName> ScalarParameterNames;
 		TArray<FGuid> ScalarParmeterGuids;
 		Material->GetAllScalarParameterNames( ScalarParameterNames, ScalarParmeterGuids );
-
 		for ( const FName& ScalarParameterName : ScalarParameterNames )
 		{
-			FUIAction AddParameterMenuAction( FExecuteAction::CreateSP( this, &FMaterialTrackEditor::AddParameterSection, ObjectBinding, MaterialTrack, ScalarParameterName ) );
-			AddParameterMenuBuilder.AddMenuEntry( FText::FromName( ScalarParameterName ), FText(), FSlateIcon(), AddParameterMenuAction );
+			FUIAction AddParameterMenuAction( FExecuteAction::CreateSP( this, &FMaterialTrackEditor::AddScalarParameter, ObjectBinding, MaterialTrack, ScalarParameterName ) );
+			FParameterNameAndAction NameAndAction( ScalarParameterName, AddParameterMenuAction );
+			ParameterNamesAndActions.Add(NameAndAction);
+		}
+
+		// Collect vector parameters.
+		TArray<FName> VectorParameterNames;
+		TArray<FGuid> VectorParmeterGuids;
+		Material->GetAllVectorParameterNames( VectorParameterNames, VectorParmeterGuids );
+		for ( const FName& VectorParameterName : VectorParameterNames )
+		{
+			FUIAction AddParameterMenuAction( FExecuteAction::CreateSP( this, &FMaterialTrackEditor::AddVectorParameter, ObjectBinding, MaterialTrack, VectorParameterName ) );
+			FParameterNameAndAction NameAndAction( VectorParameterName, AddParameterMenuAction );
+			ParameterNamesAndActions.Add( NameAndAction );
+		}
+
+		// Sort and generate menu.
+		ParameterNamesAndActions.Sort();
+		for ( FParameterNameAndAction NameAndAction : ParameterNamesAndActions )
+		{
+			AddParameterMenuBuilder.AddMenuEntry( FText::FromName( NameAndAction.ParameterName ), FText(), FSlateIcon(), NameAndAction.Action );
 		}
 	}
 
@@ -95,7 +133,7 @@ UMaterial* FMaterialTrackEditor::GetMaterialForTrack( FGuid ObjectBinding, UMovi
 	return nullptr;
 }
 
-void FMaterialTrackEditor::AddParameterSection( FGuid ObjectBinding, UMovieSceneMaterialTrack* MaterialTrack, FName ParameterName )
+void FMaterialTrackEditor::AddScalarParameter( FGuid ObjectBinding, UMovieSceneMaterialTrack* MaterialTrack, FName ParameterName )
 {
 	UMovieSceneSequence* MovieSceneSequence = GetMovieSceneSequence();
 	float KeyTime = GetTimeForKey( MovieSceneSequence );
@@ -103,13 +141,40 @@ void FMaterialTrackEditor::AddParameterSection( FGuid ObjectBinding, UMovieScene
 	UMaterial* Material = GetMaterialForTrack(ObjectBinding, MaterialTrack);
 	if (Material != nullptr)
 	{
+		const FScopedTransaction Transaction( LOCTEXT( "AddScalarParameter", "Add scalar parameter" ) );
 		float ParameterValue;
 		Material->GetScalarParameterValue(ParameterName, ParameterValue);
+		MaterialTrack->Modify();
+		for ( UMovieSceneSection* Section : MaterialTrack->GetAllSections() )
+		{
+			Section->Modify();
+		}
 		MaterialTrack->AddScalarParameterKey(ParameterName, KeyTime, ParameterValue);
 	}
 	NotifyMovieSceneDataChanged();
 }
 
+void FMaterialTrackEditor::AddVectorParameter( FGuid ObjectBinding, UMovieSceneMaterialTrack* MaterialTrack, FName ParameterName )
+{
+	UMovieSceneSequence* MovieSceneSequence = GetMovieSceneSequence();
+	float KeyTime = GetTimeForKey( MovieSceneSequence );
+
+	UMaterial* Material = GetMaterialForTrack( ObjectBinding, MaterialTrack );
+	if ( Material != nullptr )
+	{
+		const FScopedTransaction Transaction( LOCTEXT( "AddVectorParameter", "Add vector parameter" ) );
+		FLinearColor ParameterValue;
+		Material->GetVectorParameterValue( ParameterName, ParameterValue );
+		MaterialTrack->Modify();
+		for ( UMovieSceneSection* Section : MaterialTrack->GetAllSections() )
+		{
+			Section->Modify();
+		}
+		MaterialTrack->AddVectorParameterKey( ParameterName, KeyTime, ParameterValue );
+	}
+	NotifyMovieSceneDataChanged();
+}
+
 FComponentMaterialTrackEditor::FComponentMaterialTrackEditor( TSharedRef<ISequencer> InSequencer )
 	: FMaterialTrackEditor( InSequencer )
 {
diff --git a/Engine/Source/Editor/MovieSceneTools/Private/TrackEditors/MaterialTrackEditor.h b/Engine/Source/Editor/MovieSceneTools/Private/TrackEditors/MaterialTrackEditor.h
--- a/Engine/Source/Editor/MovieSceneTools/Private/TrackEditors/MaterialTrackEditor.h
+++ b/Engine/Source/Editor/MovieSceneTools/Private/TrackEditors/MaterialTrackEditor.h
@@ -29,12 +29,19 @@ private:
 	/** Gets a material for a specific object binding and track */
 	UMaterial* GetMaterialForTrack( FGuid ObjectBinding, UMovieSceneMaterialTrack* MaterialTrack );
 
-	/** Adds a parameter section and initial key to a material track.
+	/** Adds a scalar parameter and initial key to a material track.
 	 * @param ObjectBinding The object binding which owns the material track.
 	 * @param MaterialTrack The track to Add the section to.
 	 * @param ParameterName The name of the parameter to add an initial key for.
 	 */
-	void AddParameterSection( FGuid ObjectBinding, UMovieSceneMaterialTrack* MaterialTrack, FName ParameterName );
+	void AddScalarParameter( FGuid ObjectBinding, UMovieSceneMaterialTrack* MaterialTrack, FName ParameterName );
+
+	/** Adds a vector parameter and initial key to a material track.
+	* @param ObjectBinding The object binding which owns the material track.
+	* @param MaterialTrack The track to Add the section to.
+	* @param ParameterName The name of the parameter to add an initial key for.
+	*/
+	void AddVectorParameter( FGuid ObjectBinding, UMovieSceneMaterialTrack* MaterialTrack, FName ParameterName );
 };
 
 /**
diff --git a/Engine/Source/Editor/Sequencer/Private/Sequencer.cpp b/Engine/Source/Editor/Sequencer/Private/Sequencer.cpp
--- a/Engine/Source/Editor/Sequencer/Private/Sequencer.cpp
+++ b/Engine/Source/Editor/Sequencer/Private/Sequencer.cpp
@@ -1445,6 +1445,29 @@ bool FSequencer::OnHandleAssetDropped(UObject* DroppedAsset, const FGuid& Target
 	return bWasConsumed;
 }
 
+// Takes a display node and traverses it's parents to find the nearest track node if any.  Also collects the names of the nodes which make
+// up the path from the track node to the display node being checked.  The name path includes the name of the node being checked, but not
+// the name of the track node.
+void GetParentTrackNodeAndNamePath(TSharedRef<const FSequencerDisplayNode> DisplayNode, TSharedPtr<FTrackNode>& OutParentTrack, TArray<FName>& OutNamePath )
+{
+	TArray<FName> PathToTrack;
+	PathToTrack.Add( DisplayNode->GetNodeName() );
+	TSharedPtr<FSequencerDisplayNode> CurrentParent = DisplayNode->GetParent();
+	while ( CurrentParent.IsValid() && CurrentParent->GetType() != ESequencerNode::Track )
+	{
+		PathToTrack.Add( CurrentParent->GetNodeName() );
+		CurrentParent = CurrentParent->GetParent();
+	}
+	if ( CurrentParent.IsValid() )
+	{
+		OutParentTrack = StaticCastSharedPtr<FTrackNode>( CurrentParent );
+		for ( int32 i = PathToTrack.Num() - 1; i >= 0; i-- )
+		{
+			OutNamePath.Add( PathToTrack[i] );
+		}
+	}
+}
+
 void FSequencer::OnRequestNodeDeleted( TSharedRef<const FSequencerDisplayNode>& NodeToBeDeleted )
 {
 	bool bAnySpawnablesRemoved = false;
@@ -1503,8 +1526,33 @@ void FSequencer::OnRequestNodeDeleted( TSharedRef<const FSequencerDisplayNode>&
 		
 		bAnythingRemoved = true;
 	}
+	else if ( NodeToBeDeleted->GetType() == ESequencerNode::Category )
+	{
+		TSharedPtr<FTrackNode> ParentTrackNode;
+		TArray<FName> PathFromTrack;
+		GetParentTrackNodeAndNamePath(NodeToBeDeleted, ParentTrackNode, PathFromTrack);
+		if ( ParentTrackNode.IsValid() )
+		{
+			for ( TSharedRef<ISequencerSection> Section : ParentTrackNode->GetSections() )
+			{
+				bAnythingRemoved |= Section->RequestDeleteCategory( PathFromTrack );
+			}
+		}
+	}
+	else if ( NodeToBeDeleted->GetType() == ESequencerNode::KeyArea )
+	{
+		TSharedPtr<FTrackNode> ParentTrackNode;
+		TArray<FName> PathFromTrack;
+		GetParentTrackNodeAndNamePath( NodeToBeDeleted, ParentTrackNode, PathFromTrack );
+		if ( ParentTrackNode.IsValid() )
+		{
+			for ( TSharedRef<ISequencerSection> Section : ParentTrackNode->GetSections() )
+			{
+				bAnythingRemoved |= Section->RequestDeleteKeyArea( PathFromTrack );
+			}
+		}
+	}
 
-	
 	if( bAnythingRemoved )
 	{
 		NotifyMovieSceneDataChanged();
@@ -1575,6 +1623,7 @@ void FSequencer::OnPreSaveWorld(uint32 SaveFlags, class UWorld* World)
 void FSequencer::OnPostSaveWorld(uint32 SaveFlags, class UWorld* World, bool bSuccess)
 {
 	// Reset the time after saving so that an update will be triggered to put objects back to their animated state.
+	RootMovieSceneSequenceInstance->RefreshInstance(*this);
 	SetGlobalTime(GetGlobalTime());
 }
 
diff --git a/Engine/Source/Editor/Sequencer/Private/SequencerNodeTree.cpp b/Engine/Source/Editor/Sequencer/Private/SequencerNodeTree.cpp
--- a/Engine/Source/Editor/Sequencer/Private/SequencerNodeTree.cpp
+++ b/Engine/Source/Editor/Sequencer/Private/SequencerNodeTree.cpp
@@ -87,7 +87,11 @@ void FSequencerNodeTree::Update()
 			{
 				return false;
 			}
-			return A->GetDisplayName().ToString() < B->GetDisplayName().ToString();
+			if ( A->GetType() == ESequencerNode::Object && B->GetType() == ESequencerNode::Object )
+			{
+				return A->GetDisplayName().ToString() < B->GetDisplayName().ToString();
+			}
+			return 0;
 		}
 	};
 
diff --git a/Engine/Source/Editor/Sequencer/Public/ISequencerSection.h b/Engine/Source/Editor/Sequencer/Public/ISequencerSection.h
--- a/Engine/Source/Editor/Sequencer/Public/ISequencerSection.h
+++ b/Engine/Source/Editor/Sequencer/Public/ISequencerSection.h
@@ -118,4 +118,20 @@ public:
 	 * @param MenuBuilder	The menu builder to change
 	 */
 	virtual void BuildSectionContextMenu(FMenuBuilder& MenuBuilder) {}
+
+	/**
+	 * Called when the user requests that a category from this section be deleted. 
+	 *
+	 * @param CategoryNamePath An array of category names which is the path to the category to be deleted.
+	 * @returns Whether or not the category was deleted.
+	 */
+	virtual bool RequestDeleteCategory( const TArray<FName>& CategoryNamePath ) { return false; }
+
+	/**
+	* Called when the user requests that a key area from this section be deleted.
+	*
+	* @param KeyAreaNamePath An array of names representing the path of to the key area to delete, starting with any categories which contain the key area.
+	* @returns Whether or not the key area was deleted.
+	*/
+	virtual bool RequestDeleteKeyArea( const TArray<FName>& KeyAreaNamePath ) { return false; }
 };
diff --git a/Engine/Source/Runtime/MovieSceneTracks/Private/Sections/MovieSceneMaterialParameterSection.cpp b/Engine/Source/Runtime/MovieSceneTracks/Private/Sections/MovieSceneMaterialParameterSection.cpp
--- a/Engine/Source/Runtime/MovieSceneTracks/Private/Sections/MovieSceneMaterialParameterSection.cpp
+++ b/Engine/Source/Runtime/MovieSceneTracks/Private/Sections/MovieSceneMaterialParameterSection.cpp
@@ -8,6 +8,11 @@ FScalarParameterNameAndCurve::FScalarParameterNameAndCurve( FName InParameterNam
 	ParameterName = InParameterName;
 }
 
+FVectorParameterNameAndCurves::FVectorParameterNameAndCurves( FName InParameterName )
+{
+	ParameterName = InParameterName;
+}
+
 UMovieSceneMaterialParameterSection::UMovieSceneMaterialParameterSection( const FObjectInitializer& ObjectInitializer )
 	: Super( ObjectInitializer )
 {
@@ -15,17 +20,73 @@ UMovieSceneMaterialParameterSection::UMovieSceneMaterialParameterSection( const
 
 void UMovieSceneMaterialParameterSection::AddScalarParameterKey( FName InParameterName, float InTime, float InValue )
 {
-	FRichCurve* ExistingCurve = GetScalarParameterCurve(InParameterName);
-	if ( ExistingCurve != nullptr )
+	FRichCurve* ExistingCurve = nullptr;
+	for ( FScalarParameterNameAndCurve& ScalarParameterNameAndCurve : ScalarParameterNamesAndCurves )
+	{
+		if ( ScalarParameterNameAndCurve.ParameterName == InParameterName )
+		{
+			ExistingCurve = &ScalarParameterNameAndCurve.ParameterCurve;
+			break;
+		}
+	}
+	if ( ExistingCurve == nullptr )
+	{
+		int32 NewIndex = ScalarParameterNamesAndCurves.Add( FScalarParameterNameAndCurve( InParameterName ) );
+		ScalarParameterNamesAndCurves[NewIndex].Index = ScalarParameterNamesAndCurves.Num() + VectorParameterNamesAndCurves.Num() - 1;
+		ExistingCurve = &ScalarParameterNamesAndCurves[NewIndex].ParameterCurve;
+	}
+	ExistingCurve->AddKey(InTime, InValue);
+}
+
+void UMovieSceneMaterialParameterSection::AddVectorParameterKey( FName InParameterName, float InTime, FLinearColor InValue )
+{
+	FVectorParameterNameAndCurves* ExistingCurves = nullptr;
+	for ( FVectorParameterNameAndCurves& VectorParameterNameAndCurve : VectorParameterNamesAndCurves )
+	{
+		if ( VectorParameterNameAndCurve.ParameterName == InParameterName )
+		{
+			ExistingCurves = &VectorParameterNameAndCurve;
+			break;
+		}
+	}
+	if ( ExistingCurves == nullptr )
+	{
+		int32 NewIndex = VectorParameterNamesAndCurves.Add( FVectorParameterNameAndCurves( InParameterName ) );
+		VectorParameterNamesAndCurves[NewIndex].Index = VectorParameterNamesAndCurves.Num() + VectorParameterNamesAndCurves.Num() - 1;
+		ExistingCurves = &VectorParameterNamesAndCurves[NewIndex];
+	}
+	ExistingCurves->RedCurve.AddKey( InTime, InValue.R );
+	ExistingCurves->GreenCurve.AddKey( InTime, InValue.G );
+	ExistingCurves->BlueCurve.AddKey( InTime, InValue.B );
+	ExistingCurves->AlphaCurve.AddKey( InTime, InValue.A );
+}
+
+bool UMovieSceneMaterialParameterSection::RemoveScalarParameter( FName InParameterName )
+{
+	for ( int32 i = 0; i < ScalarParameterNamesAndCurves.Num(); i++ )
 	{
-		ExistingCurve->AddKey(InTime, InValue);
+		if ( ScalarParameterNamesAndCurves[i].ParameterName == InParameterName )
+		{
+			ScalarParameterNamesAndCurves.RemoveAt(i);
+			UpdateParameterIndicesFromRemoval(i);
+			return true;
+		}
 	}
-	else
+	return false;
+}
+
+bool UMovieSceneMaterialParameterSection::RemoveVectorParameter( FName InParameterName )
+{
+	for ( int32 i = 0; i < VectorParameterNamesAndCurves.Num(); i++ )
 	{
-		int32 NewIndex = ScalarParameterNamesAndCurves.Add( FScalarParameterNameAndCurve() );
-		ScalarParameterNamesAndCurves[NewIndex].ParameterName = InParameterName;
-		ScalarParameterNamesAndCurves[NewIndex].ParameterCurve.AddKey( InTime, InValue );
+		if ( VectorParameterNamesAndCurves[i].ParameterName == InParameterName )
+		{
+			VectorParameterNamesAndCurves.RemoveAt( i );
+			UpdateParameterIndicesFromRemoval( i );
+			return true;
+		}
 	}
+	return false;
 }
 
 TArray<FScalarParameterNameAndCurve>* UMovieSceneMaterialParameterSection::GetScalarParameterNamesAndCurves()
@@ -33,14 +94,25 @@ TArray<FScalarParameterNameAndCurve>* UMovieSceneMaterialParameterSection::GetSc
 	return &ScalarParameterNamesAndCurves;
 }
 
-FRichCurve* UMovieSceneMaterialParameterSection::GetScalarParameterCurve( FName InParameterName )
+TArray<FVectorParameterNameAndCurves>* UMovieSceneMaterialParameterSection::GetVectorParameterNamesAndCurves()
 {
-	for ( FScalarParameterNameAndCurve& ScalarParameterNameAndCurve : ScalarParameterNamesAndCurves )
+	return &VectorParameterNamesAndCurves;
+}
+
+void UMovieSceneMaterialParameterSection::UpdateParameterIndicesFromRemoval( int32 RemovedIndex )
+{
+	for ( FScalarParameterNameAndCurve& ScalarParameterAndCurve : ScalarParameterNamesAndCurves )
 	{
-		if ( ScalarParameterNameAndCurve.ParameterName == InParameterName )
+		if ( ScalarParameterAndCurve.Index > RemovedIndex )
+		{
+			ScalarParameterAndCurve.Index--;
+		}
+	}
+	for ( FVectorParameterNameAndCurves& VectorParameterAndCurves : VectorParameterNamesAndCurves )
+	{
+		if ( VectorParameterAndCurves.Index > RemovedIndex )
 		{
-			return &ScalarParameterNameAndCurve.ParameterCurve;
+			VectorParameterAndCurves.Index--;
 		}
 	}
-	return nullptr;
 }
\ No newline at end of file
diff --git a/Engine/Source/Runtime/MovieSceneTracks/Private/TrackInstances/MovieSceneMaterialTrackInstance.cpp b/Engine/Source/Runtime/MovieSceneTracks/Private/TrackInstances/MovieSceneMaterialTrackInstance.cpp
--- a/Engine/Source/Runtime/MovieSceneTracks/Private/TrackInstances/MovieSceneMaterialTrackInstance.cpp
+++ b/Engine/Source/Runtime/MovieSceneTracks/Private/TrackInstances/MovieSceneMaterialTrackInstance.cpp
@@ -53,13 +53,27 @@ void FMovieSceneMaterialTrackInstance::RestoreState(const TArray<UObject*>& Runt
 void FMovieSceneMaterialTrackInstance::Update( float Position, float LastPosition, const TArray<UObject*>& RuntimeObjects, class IMovieScenePlayer& Player ) 
 {
 	TArray<FScalarParameterNameAndValue> ScalarValues;
-	MaterialTrack->Eval(Position, ScalarValues);
+	TArray<FVectorParameterNameAndValue> VectorValues;
+	MaterialTrack->Eval( Position, ScalarValues, VectorValues );
 
-	for ( UMaterialInstanceDynamic* DynamicMaterialInstance : DynamicMaterialInstances )
+	// Iterate from back to front to allow for fast remove of invalid weak pointers.
+	for ( int32 i = DynamicMaterialInstances.Num() - 1; i >= 0; i--)
 	{
-		for ( const FScalarParameterNameAndValue& ScalarValue : ScalarValues )
+		if ( DynamicMaterialInstances[i].IsValid() )
 		{
-			DynamicMaterialInstance->SetScalarParameterValue(ScalarValue.ParameterName, ScalarValue.Value);
+			UMaterialInstanceDynamic* DynamicMaterialInstance = DynamicMaterialInstances[i].Get();
+			for ( const FScalarParameterNameAndValue& ScalarValue : ScalarValues )
+			{
+				DynamicMaterialInstance->SetScalarParameterValue( ScalarValue.ParameterName, ScalarValue.Value );
+			}
+			for ( const FVectorParameterNameAndValue& VectorValue : VectorValues )
+			{
+				DynamicMaterialInstance->SetVectorParameterValue( VectorValue.ParameterName, VectorValue.Value );
+			}
+		}
+		else
+		{
+			DynamicMaterialInstances.RemoveAtSwap(i);
 		}
 	}
 }
@@ -67,7 +81,6 @@ void FMovieSceneMaterialTrackInstance::Update( float Position, float LastPositio
 void FMovieSceneMaterialTrackInstance::RefreshInstance( const TArray<UObject*>& RuntimeObjects, class IMovieScenePlayer& Player )
 {
 	DynamicMaterialInstances.Empty();
-	DynamicMaterialToOriginalMaterialMap.Empty();
 	for ( UObject* Object : RuntimeObjects )
 	{
 		UMaterialInterface* Material = GetMaterialForObject( Object );
@@ -86,6 +99,7 @@ void FMovieSceneMaterialTrackInstance::ClearInstance( IMovieScenePlayer& Player
 {
 	DynamicMaterialInstances.Empty();
 	DynamicMaterialToOriginalMaterialMap.Empty();
+	RuntimeObjectToOriginalMaterialMap.Empty();
 }
 
 FMovieSceneComponentMaterialTrackInstance::FMovieSceneComponentMaterialTrackInstance( UMovieSceneComponentMaterialTrack& InMaterialTrack )
diff --git a/Engine/Source/Runtime/MovieSceneTracks/Private/Tracks/MovieSceneMaterialTrack.cpp b/Engine/Source/Runtime/MovieSceneTracks/Private/Tracks/MovieSceneMaterialTrack.cpp
--- a/Engine/Source/Runtime/MovieSceneTracks/Private/Tracks/MovieSceneMaterialTrack.cpp
+++ b/Engine/Source/Runtime/MovieSceneTracks/Private/Tracks/MovieSceneMaterialTrack.cpp
@@ -71,7 +71,20 @@ void UMovieSceneMaterialTrack::AddScalarParameterKey( FName ParameterName, float
 	NearestSection->AddScalarParameterKey(ParameterName, Time, Value);
 }
 
-void UMovieSceneMaterialTrack::Eval( float Position, TArray<FScalarParameterNameAndValue>& OutScalarValues ) const
+void UMovieSceneMaterialTrack::AddVectorParameterKey( FName ParameterName, float Time, FLinearColor Value )
+{
+	UMovieSceneMaterialParameterSection* NearestSection = Cast<UMovieSceneMaterialParameterSection>( MovieSceneHelpers::FindNearestSectionAtTime( Sections, Time ) );
+	if ( NearestSection == nullptr )
+	{
+		NearestSection = Cast<UMovieSceneMaterialParameterSection>( CreateNewSection() );
+		NearestSection->SetStartTime( Time );
+		NearestSection->SetEndTime( Time );
+		Sections.Add( NearestSection );
+	}
+	NearestSection->AddVectorParameterKey( ParameterName, Time, Value );
+}
+
+void UMovieSceneMaterialTrack::Eval( float Position, TArray<FScalarParameterNameAndValue>& OutScalarValues, TArray<FVectorParameterNameAndValue>& OutVectorValues ) const
 {
 	for ( UMovieSceneSection* Section : Sections )
 	{
@@ -82,6 +95,14 @@ void UMovieSceneMaterialTrack::Eval( float Position, TArray<FScalarParameterName
 			{
 				OutScalarValues.Add( FScalarParameterNameAndValue( ScalarParameterNameAndCurve.ParameterName, ScalarParameterNameAndCurve.ParameterCurve.Eval( Position ) ) );
 			}
+			for ( const FVectorParameterNameAndCurves& VectorParameterNameAndCurves : *MaterialParameterSection->GetVectorParameterNamesAndCurves() )
+			{
+				OutVectorValues.Add( FVectorParameterNameAndValue( VectorParameterNameAndCurves.ParameterName, FLinearColor(
+					VectorParameterNameAndCurves.RedCurve.Eval( Position ),
+					VectorParameterNameAndCurves.GreenCurve.Eval( Position ),
+					VectorParameterNameAndCurves.BlueCurve.Eval( Position ),
+					VectorParameterNameAndCurves.AlphaCurve.Eval( Position ) ) ) );
+			}
 		}
 	}
 }
diff --git a/Engine/Source/Runtime/MovieSceneTracks/Public/Sections/MovieSceneMaterialParameterSection.h b/Engine/Source/Runtime/MovieSceneTracks/Public/Sections/MovieSceneMaterialParameterSection.h
--- a/Engine/Source/Runtime/MovieSceneTracks/Public/Sections/MovieSceneMaterialParameterSection.h
+++ b/Engine/Source/Runtime/MovieSceneTracks/Public/Sections/MovieSceneMaterialParameterSection.h
@@ -13,7 +13,10 @@ struct FScalarParameterNameAndCurve
 {
 	GENERATED_USTRUCT_BODY()
 
-	FScalarParameterNameAndCurve() { }
+	FScalarParameterNameAndCurve()
+	{
+		Index = 0;
+	}
 
 	/** Creates a new FScalarParameterNameAndCurve for a specific scalar parameter. */
 	FScalarParameterNameAndCurve( FName InParameterName );
@@ -22,11 +25,54 @@ struct FScalarParameterNameAndCurve
 	UPROPERTY()
 	FName ParameterName;
 
+	UPROPERTY()
+	int32 Index;
+
 	/** The curve which contains the animation data for the scalar parameter. */
 	UPROPERTY()
 	FRichCurve ParameterCurve;
 };
 
+/**
+* Structure representing an animated vector parameter and it's associated animation curve.
+*/
+USTRUCT()
+struct FVectorParameterNameAndCurves
+{
+	GENERATED_USTRUCT_BODY()
+
+	FVectorParameterNameAndCurves() 
+	{
+		Index = 0;
+	}
+
+	/** Creates a new FVectorParameterNameAndCurve for a specific vector parameter. */
+	FVectorParameterNameAndCurves( FName InParameterName );
+
+	/** The name of the vector parameter which is being animated. */
+	UPROPERTY()
+	FName ParameterName;
+
+	UPROPERTY()
+	int32 Index;
+
+	/** The curve which contains the animation data for the red component of the vector parameter. */
+	UPROPERTY()
+	FRichCurve RedCurve;
+
+	/** The curve which contains the animation data for the green component of the vector parameter. */
+	UPROPERTY()
+	FRichCurve GreenCurve;
+
+	/** The curve which contains the animation data for the blue component of the vector parameter. */
+	UPROPERTY()
+	FRichCurve BlueCurve;
+
+	/** The curve which contains the animation data for the alpha component of the vector parameter. */
+	UPROPERTY()
+	FRichCurve AlphaCurve;
+};
+
 /**
  * A movie scene section representing animation data for material parameters. 
  */
@@ -39,18 +85,44 @@ public:
 	/** Adds a a key for a specific scalar parameter at the specified time with the specified value. */
 	void AddScalarParameterKey(FName InParameterName, float InTime, float InValue);
 
+	/** Adds a a key for a specific vector parameter at the specified time with the specified value. */
+	void AddVectorParameterKey( FName InParameterName, float InTime, FLinearColor InValue );
+
+	/** 
+	 * Removes a scalar parameter from this section. 
+	 * 
+	 * @param InParameterName The name of the scalar parameter to remove.
+	 * @returns True if a parameter with that name was found and removed, otherwise false.
+	 */
+	MOVIESCENETRACKS_API bool RemoveScalarParameter( FName InParameterName );
+
+	/**
+	* Removes a vector parameter from this section.
+	*
+	* @param InParameterName The name of the vector parameter to remove.
+	* @returns True if a parameter with that name was found and removed, otherwise false.
+	*/
+	MOVIESCENETRACKS_API bool RemoveVectorParameter( FName InParameterName );
+
 	/** Gets the animated scalar parameters and their associated curves. */
 	MOVIESCENETRACKS_API TArray<FScalarParameterNameAndCurve>* GetScalarParameterNamesAndCurves();
 
+	/** Gets the animated scalar parameters and their associated curves. */
+	MOVIESCENETRACKS_API TArray<FVectorParameterNameAndCurves>* GetVectorParameterNamesAndCurves();
+
 private:
-	/** Gets a curve for a specific scalar parameter name. */
-	FRichCurve* GetScalarParameterCurve( FName InParameterName );
+	void UpdateParameterIndicesFromRemoval(int32 RemovedIndex);
 
 private:
 	/**
 	 * The scalar parameter names and their associated curves.
-	 * NOTE: This isn't stored in a map since we want order added preserved and random access doesn't need to be fast.
 	 */
 	UPROPERTY()
 	TArray<FScalarParameterNameAndCurve> ScalarParameterNamesAndCurves;
+
+	/**
+	* The vector parameter names and their associated curves.
+	*/
+	UPROPERTY()
+	TArray<FVectorParameterNameAndCurves> VectorParameterNamesAndCurves;
 };
\ No newline at end of file
diff --git a/Engine/Source/Runtime/MovieSceneTracks/Public/TrackInstances/MovieSceneMaterialTrackInstance.h b/Engine/Source/Runtime/MovieSceneTracks/Public/TrackInstances/MovieSceneMaterialTrackInstance.h
--- a/Engine/Source/Runtime/MovieSceneTracks/Public/TrackInstances/MovieSceneMaterialTrackInstance.h
+++ b/Engine/Source/Runtime/MovieSceneTracks/Public/TrackInstances/MovieSceneMaterialTrackInstance.h
@@ -47,7 +47,7 @@ private:
 	TMap<FObjectKey, TWeakObjectPtr<UMaterialInterface>> DynamicMaterialToOriginalMaterialMap;
 
 	/** A list of dynamic material instances which need to be updated. */
-	TArray<UMaterialInstanceDynamic*> DynamicMaterialInstances;
+	TArray<TWeakObjectPtr<UMaterialInstanceDynamic>> DynamicMaterialInstances;
 };
 
 /** A material track instance specialized for component materials. */
diff --git a/Engine/Source/Runtime/MovieSceneTracks/Public/Tracks/MovieSceneMaterialTrack.h b/Engine/Source/Runtime/MovieSceneTracks/Public/Tracks/MovieSceneMaterialTrack.h
--- a/Engine/Source/Runtime/MovieSceneTracks/Public/Tracks/MovieSceneMaterialTrack.h
+++ b/Engine/Source/Runtime/MovieSceneTracks/Public/Tracks/MovieSceneMaterialTrack.h
@@ -24,6 +24,25 @@ struct FScalarParameterNameAndValue
 	float Value;
 };
 
+/**
+* Structure representing the animated value of a vector parameter.
+*/
+struct FVectorParameterNameAndValue
+{
+	/** Creates a new FVectorParameterAndValue with a parameter name and a value. */
+	FVectorParameterNameAndValue( FName InParameterName, FLinearColor InValue )
+	{
+		ParameterName = InParameterName;
+		Value = InValue;
+	}
+
+	/** The name of the vector parameter. */
+	FName ParameterName;
+
+	/** The animated value of the vector parameter. */
+	FLinearColor Value;
+};
+
 /**
  * Handles manipulation of material parameters in a movie scene.
  */
@@ -53,12 +72,21 @@ public:
 	 */
 	void MOVIESCENETRACKS_API AddScalarParameterKey( FName ParameterName, float Position, float Value );
 
+	/**
+	* Adds a Vector parameter key to the track.
+	* @param ParameterName The name of the parameter to add a key for.
+	* @param Time The time to add the new key.
+	* @param The value for the new key.
+	*/
+	void MOVIESCENETRACKS_API AddVectorParameterKey( FName ParameterName, float Position, FLinearColor Value );
+
 	/**
 	 * Gets the animated values for this track.
 	 * @param Position The playback position to use for evaluation.
 	 * @param OutScalarValues An array of FScalarParameterNameAndValue objects representing each animated scalar parameter and it's animated value.
+	 * @param OutVectorValues An array of FVectorParameterNameAndValue objects representing each animated vector parameter and it's animated value.
 	 */
-	void Eval( float Position, TArray<FScalarParameterNameAndValue>& OutScalarValues ) const;
+	void Eval( float Position, TArray<FScalarParameterNameAndValue>& OutScalarValues, TArray<FVectorParameterNameAndValue>& OutVectorValues ) const;
 
 private:
 	/** The sections owned by this track .*/
