Sequencer: Removed unnecessary "Add" text from menus.

diff --git a/Engine/Source/Editor/MovieSceneTools/Private/TrackEditors/AttachTrackEditor.cpp b/Engine/Source/Editor/MovieSceneTools/Private/TrackEditors/AttachTrackEditor.cpp
--- a/Engine/Source/Editor/MovieSceneTools/Private/TrackEditors/AttachTrackEditor.cpp
+++ b/Engine/Source/Editor/MovieSceneTools/Private/TrackEditors/AttachTrackEditor.cpp
@@ -272,7 +272,7 @@ void F3DAttachTrackEditor::AddAttachSubMenu(FMenuBuilder& MenuBuilder, FGuid Obj
 			Args.Add( TEXT("AttachName"), FText::FromString( Actor->GetActorLabel() ) );
 
 			MenuBuilder.AddMenuEntry(
-				FText::Format( NSLOCTEXT("Sequencer", "AddAttach", "Add {AttachName}"), Args ),
+				FText::Format( NSLOCTEXT("Sequencer", "AddAttach", "{AttachName}"), Args ),
 				FText::Format( NSLOCTEXT("Sequencer", "AddAttachTooltip", "Add attach track driven by {AttachName}."), Args ),
 				FSlateIcon(),
 				FUIAction(FExecuteAction::CreateSP(this, &F3DAttachTrackEditor::AddAttach, ObjectBinding, (UObject*)Actor))
diff --git a/Engine/Source/Editor/MovieSceneTools/Private/TrackEditors/PathTrackEditor.cpp b/Engine/Source/Editor/MovieSceneTools/Private/TrackEditors/PathTrackEditor.cpp
--- a/Engine/Source/Editor/MovieSceneTools/Private/TrackEditors/PathTrackEditor.cpp
+++ b/Engine/Source/Editor/MovieSceneTools/Private/TrackEditors/PathTrackEditor.cpp
@@ -253,7 +253,7 @@ void F3DPathTrackEditor::AddPathSubMenu(FMenuBuilder& MenuBuilder, FGuid ObjectB
 			Args.Add( TEXT("PathName"), FText::FromString( ActorWithSplineComponent->GetActorLabel() ) );
 
 			MenuBuilder.AddMenuEntry(
-				FText::Format( NSLOCTEXT("Sequencer", "AddPath", "Add {PathName}"), Args ),
+				FText::Format( NSLOCTEXT("Sequencer", "AddPath", "{PathName}"), Args ),
 				FText::Format( NSLOCTEXT("Sequencer", "AddPathTooltip", "Add path track driven by {PathName}."), Args ),
 				FSlateIcon(),
 				FUIAction(FExecuteAction::CreateSP(this, &F3DPathTrackEditor::AddPath, ObjectBinding, (UObject*)ActorWithSplineComponent))
diff --git a/Engine/Source/Editor/MovieSceneTools/Private/TrackEditors/SkeletalAnimationTrackEditor.cpp b/Engine/Source/Editor/MovieSceneTools/Private/TrackEditors/SkeletalAnimationTrackEditor.cpp
--- a/Engine/Source/Editor/MovieSceneTools/Private/TrackEditors/SkeletalAnimationTrackEditor.cpp
+++ b/Engine/Source/Editor/MovieSceneTools/Private/TrackEditors/SkeletalAnimationTrackEditor.cpp
@@ -240,7 +240,7 @@ void FSkeletalAnimationTrackEditor::BuildAnimationSubMenu(FMenuBuilder& MenuBuil
 
 			UAnimSequence* AnimSequence = CastChecked<UAnimSequence>(AssetDataList[AssetIndex].GetAsset());
 			MenuBuilder.AddMenuEntry(
-				FText::Format( NSLOCTEXT("Sequencer", "AddAnimSequence", "Add {AnimationName}"), Args ),
+				FText::Format( NSLOCTEXT("Sequencer", "AddAnimSequence", "{AnimationName}"), Args ),
 				FText::Format( NSLOCTEXT("Sequencer", "AddAnimSequenceTooltip", "Adds a {AnimationName} animation to this skeletal mesh."), Args ),
 				FSlateIcon(),
 				FUIAction(FExecuteAction::CreateSP(ParentSequencer.Get(), &ISequencer::AddAnimation, ObjectBinding, AnimSequence))
