Sequencer: Add setting for key interpolation. Defaults to auto.

diff --git a/Engine/Source/Editor/MovieSceneTools/Private/MovieSceneToolHelpers.cpp b/Engine/Source/Editor/MovieSceneTools/Private/MovieSceneToolHelpers.cpp
--- a/Engine/Source/Editor/MovieSceneTools/Private/MovieSceneToolHelpers.cpp
+++ b/Engine/Source/Editor/MovieSceneTools/Private/MovieSceneToolHelpers.cpp
@@ -16,7 +16,7 @@
 #include "SEnumCurveKeyEditor.h"
 #include "SBoolCurveKeyEditor.h"
 
-TArray<FKeyHandle> FFloatCurveKeyArea::AddKeyUnique(float Time, float TimeToCopyFrom)
+TArray<FKeyHandle> FFloatCurveKeyArea::AddKeyUnique(float Time, EMovieSceneKeyInterpolation InKeyInterpolation, float TimeToCopyFrom)
 {
 	TArray<FKeyHandle> AddedKeyHandles;
 
@@ -41,6 +41,35 @@ TArray<FKeyHandle> FFloatCurveKeyArea::AddKeyUnique(float Time, float TimeToCopy
 		Curve->AddKey(Time, Value, false, CurrentKeyHandle);
 		AddedKeyHandles.Add(CurrentKeyHandle);
 		
+		// Set the default key interpolation
+		switch (InKeyInterpolation)
+		{
+			case MSKI_Auto:
+				Curve->SetKeyInterpMode(CurrentKeyHandle, RCIM_Cubic);
+				Curve->SetKeyTangentMode(CurrentKeyHandle, RCTM_Auto);
+				break;
+			case MSKI_User:
+				Curve->SetKeyInterpMode(CurrentKeyHandle, RCIM_Cubic);
+				Curve->SetKeyTangentMode(CurrentKeyHandle, RCTM_User);
+				break;
+			case MSKI_Break:
+				Curve->SetKeyInterpMode(CurrentKeyHandle, RCIM_Cubic);
+				Curve->SetKeyTangentMode(CurrentKeyHandle, RCTM_Break);
+				break;
+			case MSKI_Linear:
+				Curve->SetKeyInterpMode(CurrentKeyHandle, RCIM_Linear);
+				Curve->SetKeyTangentMode(CurrentKeyHandle, RCTM_Auto);
+				break;
+			case MSKI_Constant:
+				Curve->SetKeyInterpMode(CurrentKeyHandle, RCIM_Constant);
+				Curve->SetKeyTangentMode(CurrentKeyHandle, RCTM_Auto);
+				break;
+			default:
+				Curve->SetKeyInterpMode(CurrentKeyHandle, RCIM_Cubic);
+				Curve->SetKeyTangentMode(CurrentKeyHandle, RCTM_Auto);
+				break;
+		}
+		
 		// Copy the properties from the key if it exists
 		FKeyHandle KeyHandleToCopy = Curve->FindKey(TimeToCopyFrom);
 		if (Curve->IsKeyHandleValid(KeyHandleToCopy))
@@ -68,7 +97,7 @@ TSharedRef<SWidget> FFloatCurveKeyArea::CreateKeyEditor(ISequencer* Sequencer)
 		.Curve(Curve);
 };
 
-TArray<FKeyHandle> FIntegralKeyArea::AddKeyUnique(float Time, float TimeToCopyFrom)
+TArray<FKeyHandle> FIntegralKeyArea::AddKeyUnique(float Time, EMovieSceneKeyInterpolation InKeyInterpolation, float TimeToCopyFrom)
 {
 	TArray<FKeyHandle> AddedKeyHandles;
 
diff --git a/Engine/Source/Editor/MovieSceneTools/Private/TrackEditors/TransformTrackEditor.cpp b/Engine/Source/Editor/MovieSceneTools/Private/TrackEditors/TransformTrackEditor.cpp
--- a/Engine/Source/Editor/MovieSceneTools/Private/TrackEditors/TransformTrackEditor.cpp
+++ b/Engine/Source/Editor/MovieSceneTools/Private/TrackEditors/TransformTrackEditor.cpp
@@ -269,6 +269,7 @@ void F3DTransformTrackEditor::OnTransformChanged( UObject& InObject )
 			FKeyParams KeyParams;
 			KeyParams.bAutoKeying = Sequencer->GetAutoKeyEnabled();
 			KeyParams.bAddKeyEvenIfUnchanged = false;
+			KeyParams.KeyInterpolation = Sequencer->GetKeyInterpolation();
 
 			if (CanKeyProperty(FCanKeyProperty::CreateRaw(this, &F3DTransformTrackEditor::CanKeyPropertyInternal, ObjectHandle, TransformPair, KeyParams)))
 			{
@@ -319,6 +320,7 @@ void F3DTransformTrackEditor::AddKeyInternal(const FGuid& ObjectGuid, UObject* A
 
 			FKeyParams KeyParams;
 			KeyParams.bAddKeyEvenIfUnchanged = bAddKeyEvenIfUnchanged;
+			KeyParams.KeyInterpolation = GetSequencer()->GetKeyInterpolation();
 
 			AnimatablePropertyChanged(UMovieScene3DTransformTrack::StaticClass(),
 				FOnKeyProperty::CreateRaw(this, &F3DTransformTrackEditor::OnTransformChangedInternals, Object, ObjectHandle, TransformPair, KeyParams, KeyType));
diff --git a/Engine/Source/Editor/MovieSceneTools/Public/MovieSceneToolHelpers.h b/Engine/Source/Editor/MovieSceneTools/Public/MovieSceneToolHelpers.h
--- a/Engine/Source/Editor/MovieSceneTools/Public/MovieSceneToolHelpers.h
+++ b/Engine/Source/Editor/MovieSceneTools/Public/MovieSceneToolHelpers.h
@@ -110,7 +110,7 @@ public:
 		return RCCE_None;
 	}
 
-	virtual TArray<FKeyHandle> AddKeyUnique(float Time, float TimeToCopyFrom = FLT_MAX) override;
+	virtual TArray<FKeyHandle> AddKeyUnique(float Time, EMovieSceneKeyInterpolation InKeyInterpolation, float TimeToCopyFrom = FLT_MAX) override;
 
 	virtual FRichCurve* GetRichCurve() override { return Curve; }
 
@@ -199,7 +199,7 @@ public:
 		return RCCE_None;
 	}
 
-	virtual TArray<FKeyHandle> AddKeyUnique(float Time, float TimeToCopyFrom = FLT_MAX) override;
+	virtual TArray<FKeyHandle> AddKeyUnique(float Time, EMovieSceneKeyInterpolation InKeyInterpolation, float TimeToCopyFrom = FLT_MAX) override;
 
 	virtual FRichCurve* GetRichCurve() override { return nullptr; };
 
diff --git a/Engine/Source/Editor/Sequencer/Private/GroupedKeyArea.cpp b/Engine/Source/Editor/Sequencer/Private/GroupedKeyArea.cpp
--- a/Engine/Source/Editor/Sequencer/Private/GroupedKeyArea.cpp
+++ b/Engine/Source/Editor/Sequencer/Private/GroupedKeyArea.cpp
@@ -348,7 +348,7 @@ ERichCurveExtrapolation FGroupedKeyArea::GetExtrapolationMode(bool bPreInfinity)
 	return ExtrapMode;
 }
 
-TArray<FKeyHandle> FGroupedKeyArea::AddKeyUnique(float Time, float TimeToCopyFrom)
+TArray<FKeyHandle> FGroupedKeyArea::AddKeyUnique(float Time, EMovieSceneKeyInterpolation InKeyInterpolation, float TimeToCopyFrom)
 {
 	TArray<FKeyHandle> AddedKeyHandles;
 
@@ -366,7 +366,7 @@ TArray<FKeyHandle> FGroupedKeyArea::AddKeyUnique(float Time, float TimeToCopyFro
 			}
 		}
 
-		TArray<FKeyHandle> AddedGroupKeyHandles = Area->AddKeyUnique(Time, TimeToCopyFrom);
+		TArray<FKeyHandle> AddedGroupKeyHandles = Area->AddKeyUnique(Time, InKeyInterpolation, TimeToCopyFrom);
 		AddedKeyHandles.Append(AddedGroupKeyHandles);
 	}
 
diff --git a/Engine/Source/Editor/Sequencer/Private/GroupedKeyArea.h b/Engine/Source/Editor/Sequencer/Private/GroupedKeyArea.h
--- a/Engine/Source/Editor/Sequencer/Private/GroupedKeyArea.h
+++ b/Engine/Source/Editor/Sequencer/Private/GroupedKeyArea.h
@@ -81,7 +81,7 @@ public:
 	virtual ERichCurveTangentMode   GetKeyTangentMode(FKeyHandle KeyHandle) const override;
 	virtual void                    SetExtrapolationMode(ERichCurveExtrapolation ExtrapMode, bool bPreInfinity) override;
 	virtual ERichCurveExtrapolation GetExtrapolationMode(bool bPreInfinity) const override;
-	virtual TArray<FKeyHandle>		AddKeyUnique(float Time, float TimeToCopyFrom = FLT_MAX) override;
+	virtual TArray<FKeyHandle>		AddKeyUnique(float Time, EMovieSceneKeyInterpolation InKeyInterpolation, float TimeToCopyFrom = FLT_MAX) override;
 	virtual FRichCurve*				GetRichCurve() override;
 	virtual UMovieSceneSection*		GetOwningSection() override;
 	virtual bool					CanCreateKeyEditor() override;
diff --git a/Engine/Source/Editor/Sequencer/Private/SAnimationOutlinerTreeNode.cpp b/Engine/Source/Editor/Sequencer/Private/SAnimationOutlinerTreeNode.cpp
--- a/Engine/Source/Editor/Sequencer/Private/SAnimationOutlinerTreeNode.cpp
+++ b/Engine/Source/Editor/Sequencer/Private/SAnimationOutlinerTreeNode.cpp
@@ -256,7 +256,7 @@ FReply SAnimationOutlinerTreeNode::OnAddKeyClicked()
 		{
 			OwningSection->SetFlags(RF_Transactional);
 			OwningSection->Modify();
-			KeyArea->AddKeyUnique(CurrentTime);
+			KeyArea->AddKeyUnique(CurrentTime, Sequencer.GetKeyInterpolation());
 		}
 	}
 
diff --git a/Engine/Source/Editor/Sequencer/Private/SSection.cpp b/Engine/Source/Editor/Sequencer/Private/SSection.cpp
--- a/Engine/Source/Editor/Sequencer/Private/SSection.cpp
+++ b/Engine/Source/Editor/Sequencer/Private/SSection.cpp
@@ -146,7 +146,7 @@ FSelectedKey SSection::CreateKeyUnderMouse( const FVector2D& MousePosition, cons
 				// Offset by 1 pixel worth of time
 				const float TimeFuzz = (GetSequencer().GetViewRange().GetUpperBoundValue() - GetSequencer().GetViewRange().GetLowerBoundValue()) / ParentGeometry.GetLocalSize().X;
 
-				TArray<FKeyHandle> KeyHandles = KeyArea->AddKeyUnique(KeyTime+TimeFuzz, KeyTime);
+				TArray<FKeyHandle> KeyHandles = KeyArea->AddKeyUnique(KeyTime+TimeFuzz, GetSequencer().GetKeyInterpolation(), KeyTime);
 
 				Layout = FKeyAreaLayout(*ParentSectionArea, SectionIndex);
 
@@ -169,7 +169,7 @@ FSelectedKey SSection::CreateKeyUnderMouse( const FVector2D& MousePosition, cons
 			}
 			else
 			{
-				KeyArea->AddKeyUnique(KeyTime);
+				KeyArea->AddKeyUnique(KeyTime, GetSequencer().GetKeyInterpolation());
 
 				Layout = FKeyAreaLayout(*ParentSectionArea, SectionIndex);
 						
diff --git a/Engine/Source/Editor/Sequencer/Private/Sequencer.cpp b/Engine/Source/Editor/Sequencer/Private/Sequencer.cpp
--- a/Engine/Source/Editor/Sequencer/Private/Sequencer.cpp
+++ b/Engine/Source/Editor/Sequencer/Private/Sequencer.cpp
@@ -649,6 +649,16 @@ void FSequencer::SetKeyInterpPropertiesOnly(bool bKeyInterpPropertiesOnly)
 	Settings->SetKeyInterpPropertiesOnly(bKeyInterpPropertiesOnly);
 }
 
+EMovieSceneKeyInterpolation FSequencer::GetKeyInterpolation() const
+{
+	return Settings->GetKeyInterpolation();
+}
+
+void FSequencer::SetKeyInterpolation(EMovieSceneKeyInterpolation InKeyInterpolation)
+{
+	Settings->SetKeyInterpolation(InKeyInterpolation);
+}
+
 bool FSequencer::IsRecordingLive() const 
 {
 	return PlaybackState == EMovieScenePlayerStatus::Recording && GIsPlayInEditorWorld;
diff --git a/Engine/Source/Editor/Sequencer/Private/Sequencer.h b/Engine/Source/Editor/Sequencer/Private/Sequencer.h
--- a/Engine/Source/Editor/Sequencer/Private/Sequencer.h
+++ b/Engine/Source/Editor/Sequencer/Private/Sequencer.h
@@ -76,6 +76,8 @@ public:
 	virtual void SetKeyAllEnabled(bool bKeyAllEnabled) override;
 	virtual bool GetKeyInterpPropertiesOnly() const override;
 	virtual void SetKeyInterpPropertiesOnly(bool bKeyInterpPropertiesOnly) override;
+	virtual EMovieSceneKeyInterpolation GetKeyInterpolation() const override;
+	virtual void SetKeyInterpolation(EMovieSceneKeyInterpolation) override;
 	virtual bool IsRecordingLive() const override;
 	virtual float GetCurrentLocalTime(UMovieSceneSequence& InMovieSceneSequence) override;
 	virtual float GetGlobalTime() override;
diff --git a/Engine/Source/Editor/Sequencer/Private/SequencerDetailKeyframeHandler.cpp b/Engine/Source/Editor/Sequencer/Private/SequencerDetailKeyframeHandler.cpp
--- a/Engine/Source/Editor/Sequencer/Private/SequencerDetailKeyframeHandler.cpp
+++ b/Engine/Source/Editor/Sequencer/Private/SequencerDetailKeyframeHandler.cpp
@@ -27,6 +27,7 @@ void FSequencerDetailKeyframeHandler::OnKeyPropertyClicked(const IPropertyHandle
 	KeyPropertyParams.KeyParams.bCreateHandleIfMissing = true;
 	KeyPropertyParams.KeyParams.bCreateTrackIfMissing = true;
 	KeyPropertyParams.KeyParams.bAddKeyEvenIfUnchanged = true;
+	KeyPropertyParams.KeyParams.KeyInterpolation = Sequencer.Pin()->GetKeyInterpolation();
 
 	Sequencer.Pin()->KeyProperty(KeyPropertyParams);
 }
diff --git a/Engine/Source/Editor/Sequencer/Private/SequencerDisplayNode.cpp b/Engine/Source/Editor/Sequencer/Private/SequencerDisplayNode.cpp
--- a/Engine/Source/Editor/Sequencer/Private/SequencerDisplayNode.cpp
+++ b/Engine/Source/Editor/Sequencer/Private/SequencerDisplayNode.cpp
@@ -900,6 +900,7 @@ void FObjectBindingNode::AddTrackForProperty(TArray<UProperty*> PropertyPath)
 	KeyPropertyParams.KeyParams.bCreateTrackIfMissing = true;
 	KeyPropertyParams.KeyParams.bCreateHandleIfMissing = false;
 	KeyPropertyParams.KeyParams.bAddKeyEvenIfUnchanged = false;
+	KeyPropertyParams.KeyParams.KeyInterpolation = Sequencer.GetKeyInterpolation();
 
 	Sequencer.KeyProperty(KeyPropertyParams);
 }
diff --git a/Engine/Source/Editor/Sequencer/Private/SequencerObjectChangeListener.cpp b/Engine/Source/Editor/Sequencer/Private/SequencerObjectChangeListener.cpp
--- a/Engine/Source/Editor/Sequencer/Private/SequencerObjectChangeListener.cpp
+++ b/Engine/Source/Editor/Sequencer/Private/SequencerObjectChangeListener.cpp
@@ -61,6 +61,7 @@ void FSequencerObjectChangeListener::BroadcastPropertyChanged( FKeyPropertyParam
 	Params.KeyParams.bCreateTrackIfMissing = KeyPropertyParams.KeyParams.bCreateTrackIfMissing;
 	Params.KeyParams.bAddKeyEvenIfUnchanged = KeyPropertyParams.KeyParams.bAddKeyEvenIfUnchanged;
 	Params.KeyParams.bAutoKeying = Sequencer.Pin()->GetAutoKeyEnabled();
+	Params.KeyParams.KeyInterpolation = Sequencer.Pin()->GetKeyInterpolation();
 	Params.ObjectsThatChanged = KeyableObjects;
 	Params.StructPropertyNameToKey = NAME_None;
 
diff --git a/Engine/Source/Editor/Sequencer/Private/SequencerSettings.cpp b/Engine/Source/Editor/Sequencer/Private/SequencerSettings.cpp
--- a/Engine/Source/Editor/Sequencer/Private/SequencerSettings.cpp
+++ b/Engine/Source/Editor/Sequencer/Private/SequencerSettings.cpp
@@ -8,6 +8,7 @@ USequencerSettings::USequencerSettings( const FObjectInitializer& ObjectInitiali
 	bAutoKeyEnabled = false;
 	bKeyAllEnabled = false;
 	bKeyInterpPropertiesOnly = false;
+	KeyInterpolation = MSKI_Auto;
 	bShowFrameNumbers = true;
 	bShowRangeSlider = false;
 	bLockInOutToStartEndRange = false;
@@ -69,6 +70,20 @@ void USequencerSettings::SetKeyInterpPropertiesOnly(bool InbKeyInterpPropertiesO
 	}
 }
 
+EMovieSceneKeyInterpolation USequencerSettings::GetKeyInterpolation() const
+{
+	return KeyInterpolation;
+}
+
+void USequencerSettings::SetKeyInterpolation(EMovieSceneKeyInterpolation InKeyInterpolation)
+{
+	if ( KeyInterpolation != InKeyInterpolation)
+	{
+		KeyInterpolation = InKeyInterpolation;
+		SaveConfig();
+	}
+}
+
 bool USequencerSettings::GetShowFrameNumbers() const
 {
 	return bShowFrameNumbers;
diff --git a/Engine/Source/Editor/Sequencer/Private/SequencerSettings.h b/Engine/Source/Editor/Sequencer/Private/SequencerSettings.h
--- a/Engine/Source/Editor/Sequencer/Private/SequencerSettings.h
+++ b/Engine/Source/Editor/Sequencer/Private/SequencerSettings.h
@@ -59,6 +59,11 @@ public:
 	/** Sets whether or not to key interp properties only. */
 	void SetKeyInterpPropertiesOnly(bool InbKeyInterpPropertiesOnly); 
 
+	/** Gets default key interpolation. */
+	EMovieSceneKeyInterpolation GetKeyInterpolation() const;
+	/** Sets default key interpolation */
+	void SetKeyInterpolation(EMovieSceneKeyInterpolation InKeyInterpolation);
+
 	/** Gets whether or not to show frame numbers. */
 	bool GetShowFrameNumbers() const;
 	/** Sets whether or not to show frame numbers. */
@@ -160,6 +165,9 @@ protected:
 	UPROPERTY( config )
 	bool bKeyInterpPropertiesOnly;
 
+	UPROPERTY( config )
+	TEnumAsByte<EMovieSceneKeyInterpolation> KeyInterpolation;
+
 	UPROPERTY( config )
 	bool bShowFrameNumbers;
 
diff --git a/Engine/Source/Editor/Sequencer/Public/IKeyArea.h b/Engine/Source/Editor/Sequencer/Public/IKeyArea.h
--- a/Engine/Source/Editor/Sequencer/Public/IKeyArea.h
+++ b/Engine/Source/Editor/Sequencer/Public/IKeyArea.h
@@ -111,10 +111,11 @@ public:
 	 * Adds a key at the specified time if there isn't already a key present.  The value of the added key should
 	 * be the value which would be returned if the animation containing this key area was evaluated at the specified time.
 	 *
+	 * @param InKeyInterpolation KeyInterpolation
 	 * @param TimeToCopyFrom Optional time to copy key parameters from
 	 * @param return      The new keys that were added
 	 */
-	virtual TArray<FKeyHandle> AddKeyUnique(float Time, float TimeToCopyFrom = FLT_MAX) = 0;
+	virtual TArray<FKeyHandle> AddKeyUnique(float Time, EMovieSceneKeyInterpolation InKeyInterpolation, float TimeToCopyFrom = FLT_MAX) = 0;
 
 	/** Gets the rich curve associated with this key area.  This can be null but it must be valid in order to be
 	  * edited by the curve editor. */
diff --git a/Engine/Source/Editor/Sequencer/Public/ISequencer.h b/Engine/Source/Editor/Sequencer/Public/ISequencer.h
--- a/Engine/Source/Editor/Sequencer/Public/ISequencer.h
+++ b/Engine/Source/Editor/Sequencer/Public/ISequencer.h
@@ -79,12 +79,18 @@ public:
 	/** Sets whether key all is enabled in this sequencer. */
 	virtual void SetKeyAllEnabled(bool bKeyAllEnabled) = 0;
 
-	/** @return Returns whether or not to key only  interp properties in this sequencer */
+	/** @return Returns whether or not to key only interp properties in this sequencer */
 	virtual bool GetKeyInterpPropertiesOnly() const = 0;
 
 	/** Sets whether or not to key only interp properties in this sequencer. */
 	virtual void SetKeyInterpPropertiesOnly(bool bKeyInterpPropertiesOnly) = 0;
 
+	/** @return Returns default key interpolation */
+	virtual EMovieSceneKeyInterpolation GetKeyInterpolation() const = 0;
+
+	/** Set default key interpolation */
+	virtual void SetKeyInterpolation(EMovieSceneKeyInterpolation) = 0;
+
 	/** @return Returns whether sequencer is currently recording live data from simulated actors */
 	virtual bool IsRecordingLive() const = 0;
 
diff --git a/Engine/Source/Editor/UMGEditor/Private/Animation/UMGDetailKeyframeHandler.cpp b/Engine/Source/Editor/UMGEditor/Private/Animation/UMGDetailKeyframeHandler.cpp
--- a/Engine/Source/Editor/UMGEditor/Private/Animation/UMGDetailKeyframeHandler.cpp
+++ b/Engine/Source/Editor/UMGEditor/Private/Animation/UMGDetailKeyframeHandler.cpp
@@ -30,6 +30,7 @@ void FUMGDetailKeyframeHandler::OnKeyPropertyClicked(const IPropertyHandle& Keye
 	KeyPropertyParams.KeyParams.bCreateHandleIfMissing = true;
 	KeyPropertyParams.KeyParams.bCreateTrackIfMissing = true;
 	KeyPropertyParams.KeyParams.bAddKeyEvenIfUnchanged = true;
+	KeyPropertyParams.KeyParams.KeyInterpolation = BlueprintEditor.Pin()->GetSequencer()->GetKeyInterpolation();
 
 	BlueprintEditor.Pin()->GetSequencer()->KeyProperty(KeyPropertyParams);
 }
diff --git a/Engine/Source/Runtime/MovieScene/Private/MovieSceneSection.cpp b/Engine/Source/Runtime/MovieScene/Private/MovieSceneSection.cpp
--- a/Engine/Source/Runtime/MovieScene/Private/MovieSceneSection.cpp
+++ b/Engine/Source/Runtime/MovieScene/Private/MovieSceneSection.cpp
@@ -74,7 +74,41 @@ void UMovieSceneSection::AddKeyToCurve( FRichCurve& InCurve, float Time, float V
 		}
 		else
 		{
-			InCurve.UpdateOrAddKey(Time, Value, bUnwindRotation);
+			FKeyHandle ExistingKeyHandle = InCurve.FindKey(Time);
+			
+			FKeyHandle NewKeyHandle = InCurve.UpdateOrAddKey(Time, Value, bUnwindRotation);
+
+			if (!InCurve.IsKeyHandleValid(ExistingKeyHandle) && InCurve.IsKeyHandleValid(NewKeyHandle))
+			{
+				// Set the default key interpolation only if a new key was created
+				switch (KeyParams.KeyInterpolation)
+				{
+					case MSKI_Auto:
+						InCurve.SetKeyInterpMode(NewKeyHandle, RCIM_Cubic);
+						InCurve.SetKeyTangentMode(NewKeyHandle, RCTM_Auto);
+						break;
+					case MSKI_User:
+						InCurve.SetKeyInterpMode(NewKeyHandle, RCIM_Cubic);
+						InCurve.SetKeyTangentMode(NewKeyHandle, RCTM_User);
+						break;
+					case MSKI_Break:
+						InCurve.SetKeyInterpMode(NewKeyHandle, RCIM_Cubic);
+						InCurve.SetKeyTangentMode(NewKeyHandle, RCTM_Break);
+						break;
+					case MSKI_Linear:
+						InCurve.SetKeyInterpMode(NewKeyHandle, RCIM_Linear);
+						InCurve.SetKeyTangentMode(NewKeyHandle, RCTM_Auto);
+						break;
+					case MSKI_Constant:
+						InCurve.SetKeyInterpMode(NewKeyHandle, RCIM_Constant);
+						InCurve.SetKeyTangentMode(NewKeyHandle, RCTM_Auto);
+						break;
+					default:
+						InCurve.SetKeyInterpMode(NewKeyHandle, RCIM_Cubic);
+						InCurve.SetKeyTangentMode(NewKeyHandle, RCTM_Auto);
+						break;
+				}
+			}
 		}
 	}
 }
\ No newline at end of file
diff --git a/Engine/Source/Runtime/MovieScene/Public/KeyParams.h b/Engine/Source/Runtime/MovieScene/Public/KeyParams.h
--- a/Engine/Source/Runtime/MovieScene/Public/KeyParams.h
+++ b/Engine/Source/Runtime/MovieScene/Public/KeyParams.h
@@ -2,6 +2,25 @@
 
 #pragma once
 
+UENUM()
+enum EMovieSceneKeyInterpolation
+{
+	/** Auto. */
+	MSKI_Auto UMETA(DisplayName="Auto"),
+
+	/** User. */
+	MSKI_User UMETA(DisplayName="User"),
+
+	/** Break. */
+	MSKI_Break UMETA(DisplayName="Break"),
+
+	/** Linear. */
+	MSKI_Linear UMETA(DisplayName="Linear"),
+
+	/** Constant. */
+	MSKI_Constant UMETA(DisplayName="Constant"),
+};
+
 /**
  * Parameters for determining keying behavior
  */
@@ -13,6 +32,7 @@ struct MOVIESCENE_API FKeyParams
 		bCreateTrackIfMissing = false;
 		bAddKeyEvenIfUnchanged = false;
 		bAutoKeying = false;
+		KeyInterpolation = MSKI_Auto;
 	}
 
 	FKeyParams(const FKeyParams& InKeyParams)
@@ -21,6 +41,7 @@ struct MOVIESCENE_API FKeyParams
 		bCreateTrackIfMissing = InKeyParams.bCreateTrackIfMissing;
 		bAddKeyEvenIfUnchanged = InKeyParams.bAddKeyEvenIfUnchanged;
 		bAutoKeying = InKeyParams.bAutoKeying;
+		KeyInterpolation = InKeyParams.KeyInterpolation;
 	}
 
 	/** Create handle if it doesn't exist. */
@@ -31,4 +52,6 @@ struct MOVIESCENE_API FKeyParams
 	bool bAddKeyEvenIfUnchanged;
 	/** Auto keying on.*/
 	bool bAutoKeying;
+	/** Key interpolation. */
+	EMovieSceneKeyInterpolation KeyInterpolation;
 };
