Sequencer: Fix crash when trying to key bool property when there's multiple property paths (ie. properties on components)

diff --git a/Engine/Source/Editor/MovieSceneTools/Private/TrackEditors/PropertyTrackEditors/BoolPropertyTrackEditor.cpp b/Engine/Source/Editor/MovieSceneTools/Private/TrackEditors/PropertyTrackEditors/BoolPropertyTrackEditor.cpp
--- a/Engine/Source/Editor/MovieSceneTools/Private/TrackEditors/PropertyTrackEditors/BoolPropertyTrackEditor.cpp
+++ b/Engine/Source/Editor/MovieSceneTools/Private/TrackEditors/PropertyTrackEditors/BoolPropertyTrackEditor.cpp
@@ -18,9 +18,14 @@ TSharedRef<ISequencerSection> FBoolPropertyTrackEditor::MakeSectionInterface( UM
 bool FBoolPropertyTrackEditor::TryGenerateKeyFromPropertyChanged( const UMovieSceneTrack* InTrack, const FPropertyChangedParams& PropertyChangedParams, bool& OutKey )
 {
 	const UBoolProperty* BoolProperty = Cast<const UBoolProperty>(PropertyChangedParams.PropertyPath.Last());
-	if (BoolProperty)
+	if (BoolProperty && PropertyChangedParams.ObjectsThatChanged.Num() != 0)
 	{
-		OutKey = BoolProperty->GetPropertyValue(BoolProperty->ContainerPtrToValuePtr<void>(PropertyChangedParams.ObjectsThatChanged.Last()));
+		void* CurrentObject = PropertyChangedParams.ObjectsThatChanged[0];
+		for (int32 i = 0; i < PropertyChangedParams.PropertyPath.Num(); ++i)
+		{
+			CurrentObject = PropertyChangedParams.PropertyPath[i]->ContainerPtrToValuePtr<void>(CurrentObject, 0);
+		}
+		OutKey = BoolProperty->GetPropertyValue(CurrentObject);
 	
 		if (InTrack)
 		{
diff --git a/Engine/Source/Editor/MovieSceneTools/Private/TrackEditors/PropertyTrackEditors/VisibilityPropertyTrackEditor.cpp b/Engine/Source/Editor/MovieSceneTools/Private/TrackEditors/PropertyTrackEditors/VisibilityPropertyTrackEditor.cpp
--- a/Engine/Source/Editor/MovieSceneTools/Private/TrackEditors/PropertyTrackEditors/VisibilityPropertyTrackEditor.cpp
+++ b/Engine/Source/Editor/MovieSceneTools/Private/TrackEditors/PropertyTrackEditors/VisibilityPropertyTrackEditor.cpp
@@ -18,9 +18,14 @@ TSharedRef<ISequencerSection> FVisibilityPropertyTrackEditor::MakeSectionInterfa
 bool FVisibilityPropertyTrackEditor::TryGenerateKeyFromPropertyChanged( const UMovieSceneTrack* InTrack, const FPropertyChangedParams& PropertyChangedParams, bool& OutKey )
 {
 	const UBoolProperty* BoolProperty = Cast<const UBoolProperty>(PropertyChangedParams.PropertyPath.Last());
-	if (BoolProperty)
+	if (BoolProperty && PropertyChangedParams.ObjectsThatChanged.Num() != 0)
 	{
-		OutKey = BoolProperty->GetPropertyValue(BoolProperty->ContainerPtrToValuePtr<void>(PropertyChangedParams.ObjectsThatChanged.Last()));
+		void* CurrentObject = PropertyChangedParams.ObjectsThatChanged[0];
+		for (int32 i = 0; i < PropertyChangedParams.PropertyPath.Num(); ++i)
+		{
+			CurrentObject = PropertyChangedParams.PropertyPath[i]->ContainerPtrToValuePtr<void>(CurrentObject, 0);
+		}
+		OutKey = BoolProperty->GetPropertyValue(CurrentObject);
 
 		if (InTrack)
 		{
diff --git a/Engine/Source/Runtime/MovieSceneTracks/Private/Tracks/MovieSceneVisibilityTrack.cpp b/Engine/Source/Runtime/MovieSceneTracks/Private/Tracks/MovieSceneVisibilityTrack.cpp
--- a/Engine/Source/Runtime/MovieSceneTracks/Private/Tracks/MovieSceneVisibilityTrack.cpp
+++ b/Engine/Source/Runtime/MovieSceneTracks/Private/Tracks/MovieSceneVisibilityTrack.cpp
@@ -39,3 +39,9 @@ bool UMovieSceneVisibilityTrack::AddKeyToSection( float Time, bool Value, FKeyPa
 	}
 	return false;
 }
+
+bool UMovieSceneVisibilityTrack::CanKeyTrack(float Time, bool Value, FKeyParams KeyParams) const
+{
+	// The property that's being changed is bHiddenInGame. But we want to store the inverse of that as visibility, so invert the incoming value.
+	return Super::CanKeyTrack(Time, !Value, KeyParams);
+}
\ No newline at end of file
diff --git a/Engine/Source/Runtime/MovieSceneTracks/Public/Tracks/MovieSceneVisibilityTrack.h b/Engine/Source/Runtime/MovieSceneTracks/Public/Tracks/MovieSceneVisibilityTrack.h
--- a/Engine/Source/Runtime/MovieSceneTracks/Public/Tracks/MovieSceneVisibilityTrack.h
+++ b/Engine/Source/Runtime/MovieSceneTracks/Public/Tracks/MovieSceneVisibilityTrack.h
@@ -30,4 +30,14 @@ public:
 	 * @return True if the key was successfully added.
 	 */
 	virtual bool AddKeyToSection( float Time, bool Value, FKeyParams KeyParams );
+
+	/*
+	 * Get whether the track can be keyed at a particular time.
+	 *
+	 * @param Time				The time relative to the owning movie scene where the section should be
+	 * @param Value				The value of the key
+	 * @param KeyParams         The keying parameters
+	 * @return Whether the track can be keyed
+	 */
+	virtual bool CanKeyTrack( float Time, bool Value, FKeyParams KeyParams) const;
 };
