Sequencer: Started to implement scaffolding for slomo tracks

diff --git a/Engine/Source/Editor/MovieSceneTools/Private/TrackEditors/SlomoTrackEditor.cpp b/Engine/Source/Editor/MovieSceneTools/Private/TrackEditors/SlomoTrackEditor.cpp
new file mode 100644
--- /dev/null
+++ b/Engine/Source/Editor/MovieSceneTools/Private/TrackEditors/SlomoTrackEditor.cpp
@@ -0,0 +1,30 @@
+// Copyright 1998-2015 Epic Games, Inc. All Rights Reserved.
+
+#include "MovieSceneToolsPrivatePCH.h"
+#include "SlomoTrackEditor.h"
+
+
+/* FSlomoTrackEditor static functions
+ *****************************************************************************/
+
+TSharedRef<FMovieSceneTrackEditor> FSlomoTrackEditor::CreateTrackEditor(TSharedRef<ISequencer> InSequencer)
+{
+	return MakeShareable(new FSlomoTrackEditor(InSequencer));
+}
+
+
+/* FSlomoTrackEditor structors
+ *****************************************************************************/
+
+FSlomoTrackEditor::FSlomoTrackEditor(TSharedRef<ISequencer> InSequencer)
+	: FFloatPropertyTrackEditor(InSequencer)
+{ }
+
+
+/* FMovieSceneTrackEditor interface
+ *****************************************************************************/
+
+bool FSlomoTrackEditor::SupportsType(TSubclassOf<UMovieSceneTrack> Type) const
+{
+	return (Type == UMovieSceneSlomoTrack::StaticClass());
+}
diff --git a/Engine/Source/Editor/MovieSceneTools/Private/TrackEditors/SlomoTrackEditor.h b/Engine/Source/Editor/MovieSceneTools/Private/TrackEditors/SlomoTrackEditor.h
new file mode 100644
--- /dev/null
+++ b/Engine/Source/Editor/MovieSceneTools/Private/TrackEditors/SlomoTrackEditor.h
@@ -0,0 +1,45 @@
+// Copyright 1998-2015 Epic Games, Inc. All Rights Reserved.
+
+#pragma once
+
+#include "MovieSceneSlomoTrack.h"
+#include "FloatPropertyTrackEditor.h"
+
+
+class ISequencer;
+
+
+/**
+* A property track editor for slow motion control.
+*/
+class FSlomoTrackEditor
+	: public FFloatPropertyTrackEditor
+{
+public:
+
+	/**
+	 * Factory function to create an instance of this class (called by a sequencer).
+	 *
+	 * @param InSequencer The sequencer instance to be used by this tool.
+	 * @return The new instance of this class.
+	 */
+	static TSharedRef<FMovieSceneTrackEditor> CreateTrackEditor(TSharedRef<ISequencer> InSequencer);
+
+public:
+
+	/**
+	 * Creates and initializes a new instance.
+	 *
+	 * @param InSequencer The sequencer instance to be used by this tool.
+	 */
+	FSlomoTrackEditor(TSharedRef<ISequencer> InSequencer);	
+
+public:
+
+	// FMovieSceneTrackEditor interface
+
+//	virtual void AddKey(const FGuid& ObjectGuid, UObject* AdditionalAsset = nullptr) override;
+//	virtual void BuildObjectBindingTrackMenu(FMenuBuilder& MenuBuilder, const FGuid& ObjectBinding, const UClass* ObjectClass) override;
+//	virtual TSharedRef<ISequencerSection> MakeSectionInterface(UMovieSceneSection& SectionObject, UMovieSceneTrack* Track) override;
+	virtual bool SupportsType(TSubclassOf<UMovieSceneTrack> Type) const override;
+};
diff --git a/Engine/Source/Runtime/MovieSceneTracks/Private/Sections/MovieSceneSlomoSection.cpp b/Engine/Source/Runtime/MovieSceneTracks/Private/Sections/MovieSceneSlomoSection.cpp
new file mode 100644
--- /dev/null
+++ b/Engine/Source/Runtime/MovieSceneTracks/Private/Sections/MovieSceneSlomoSection.cpp
@@ -0,0 +1,4 @@
+// Copyright 1998-2015 Epic Games, Inc. All Rights Reserved.
+
+#include "MovieSceneTracksPrivatePCH.h"
+#include "MovieSceneSlomoSection.h"
diff --git a/Engine/Source/Runtime/MovieSceneTracks/Private/TrackInstances/MovieSceneSlomoTrackInstance.cpp b/Engine/Source/Runtime/MovieSceneTracks/Private/TrackInstances/MovieSceneSlomoTrackInstance.cpp
new file mode 100644
--- /dev/null
+++ b/Engine/Source/Runtime/MovieSceneTracks/Private/TrackInstances/MovieSceneSlomoTrackInstance.cpp
@@ -0,0 +1,5 @@
+// Copyright 1998-2015 Epic Games, Inc. All Rights Reserved.
+
+#include "MovieSceneTracksPrivatePCH.h"
+#include "MovieSceneSlomoTrack.h"
+#include "MovieSceneSlomoTrackInstance.h"
diff --git a/Engine/Source/Runtime/MovieSceneTracks/Private/Tracks/MovieSceneSlomoTrack.cpp b/Engine/Source/Runtime/MovieSceneTracks/Private/Tracks/MovieSceneSlomoTrack.cpp
new file mode 100644
--- /dev/null
+++ b/Engine/Source/Runtime/MovieSceneTracks/Private/Tracks/MovieSceneSlomoTrack.cpp
@@ -0,0 +1,7 @@
+// Copyright 1998-2015 Epic Games, Inc. All Rights Reserved.
+
+#include "MovieSceneTracksPrivatePCH.h"
+#include "MovieSceneSlomoSection.h"
+#include "MovieSceneSlomoTrack.h"
+#include "MovieSceneSlomoTrackInstance.h"
+
diff --git a/Engine/Source/Runtime/MovieSceneTracks/Public/Sections/MovieSceneSlomoSection.h b/Engine/Source/Runtime/MovieSceneTracks/Public/Sections/MovieSceneSlomoSection.h
new file mode 100644
--- /dev/null
+++ b/Engine/Source/Runtime/MovieSceneTracks/Public/Sections/MovieSceneSlomoSection.h
@@ -0,0 +1,18 @@
+// Copyright 1998-2015 Epic Games, Inc. All Rights Reserved.
+
+#pragma once
+
+#include "MovieSceneFloatSection.h"
+#include "MovieSceneSlomoSection.generated.h"
+
+
+/**
+ * A single floating point section
+ */
+UCLASS(MinimalAPI)
+class UMovieSceneSlomoSection
+	: public UMovieSceneFloatSection
+{
+	GENERATED_BODY()
+
+};
diff --git a/Engine/Source/Runtime/MovieSceneTracks/Public/TrackInstances/MovieSceneSlomoTrackInstance.h b/Engine/Source/Runtime/MovieSceneTracks/Public/TrackInstances/MovieSceneSlomoTrackInstance.h
new file mode 100644
--- /dev/null
+++ b/Engine/Source/Runtime/MovieSceneTracks/Public/TrackInstances/MovieSceneSlomoTrackInstance.h
@@ -0,0 +1,43 @@
+// Copyright 1998-2015 Epic Games, Inc. All Rights Reserved.
+
+#pragma once
+
+#include "IMovieSceneTrackInstance.h"
+
+
+class UMovieSceneSlomoTrack;
+
+
+/**
+ * Instance of a UMovieSceneSlomoTrack
+ */
+class FMovieSceneSlomoTrackInstance
+	: public IMovieSceneTrackInstance
+{
+public:
+
+	/**
+	 * Creates and initializes a new instance.
+	 *
+	 * @param InSlomoTrack The event track to create an instance for.
+	 */
+	FMovieSceneSlomoTrackInstance(UMovieSceneSlomoTrack& InSlomoTrack );
+
+	/** Virtual destructor. */
+	virtual ~FMovieSceneSlomoTrackInstance() { }
+
+public:
+
+	// IMovieSceneTrackInstance interface
+
+	virtual void ClearInstance(IMovieScenePlayer& Player) override { }
+	virtual void RefreshInstance(const TArray<UObject*>& RuntimeObjects, IMovieScenePlayer& Player) override { }
+	virtual void RestoreState(const TArray<UObject*>& RuntimeObjects) override { }
+	virtual void SaveState(const TArray<UObject*>& RuntimeObjects) override { }
+//	virtual void Update(float Position, float LastPosition, const TArray<UObject*>& RuntimeObjects, IMovieScenePlayer& Player) override;
+
+private:
+
+	/** Track that is being instanced */
+	UMovieSceneSlomoTrack* SlomoTrack;
+};
diff --git a/Engine/Source/Runtime/MovieSceneTracks/Public/Tracks/MovieSceneSlomoTrack.h b/Engine/Source/Runtime/MovieSceneTracks/Public/Tracks/MovieSceneSlomoTrack.h
new file mode 100644
--- /dev/null
+++ b/Engine/Source/Runtime/MovieSceneTracks/Public/Tracks/MovieSceneSlomoTrack.h
@@ -0,0 +1,17 @@
+// Copyright 1998-2015 Epic Games, Inc. All Rights Reserved.
+
+#pragma once
+
+#include "MovieSceneFloatTrack.h"
+#include "MovieSceneSlomoTrack.generated.h"
+
+
+/**
+ * Implements a movie scene track that controls a movie scene's playback speed.
+ */
+UCLASS(MinimalAPI)
+class UMovieSceneSlomoTrack
+	: public UMovieSceneFloatTrack
+{
+	GENERATED_BODY()
+};
