Sequencer: Fix shot section deletion.

diff --git a/Engine/Source/Editor/MovieSceneTools/Private/TrackEditors/ShotTrackEditor.cpp b/Engine/Source/Editor/MovieSceneTools/Private/TrackEditors/ShotTrackEditor.cpp
--- a/Engine/Source/Editor/MovieSceneTools/Private/TrackEditors/ShotTrackEditor.cpp
+++ b/Engine/Source/Editor/MovieSceneTools/Private/TrackEditors/ShotTrackEditor.cpp
@@ -222,11 +222,11 @@ FShotSection::~FShotSection()
 	if( ThumbnailPool.IsValid() )
 	{
 		ThumbnailPool.Pin()->RemoveThumbnailsNeedingRedraw(Thumbnails);
-
-		if(InternalViewportClient.IsValid())
-		{
-			InternalViewportClient->Viewport = nullptr;
-		}
+	}
+	
+	if(InternalViewportClient.IsValid())
+	{
+		InternalViewportClient->Viewport = nullptr;
 	}
 }
 
