Sequencer: Refactored sequencer to not be aware of animation sequence tracks

diff --git a/Engine/Source/Editor/MovieSceneTools/Private/TrackEditors/SkeletalAnimationTrackEditor.cpp b/Engine/Source/Editor/MovieSceneTools/Private/TrackEditors/SkeletalAnimationTrackEditor.cpp
--- a/Engine/Source/Editor/MovieSceneTools/Private/TrackEditors/SkeletalAnimationTrackEditor.cpp
+++ b/Engine/Source/Editor/MovieSceneTools/Private/TrackEditors/SkeletalAnimationTrackEditor.cpp
@@ -243,7 +243,7 @@ void FSkeletalAnimationTrackEditor::BuildAnimationSubMenu(FMenuBuilder& MenuBuil
 				FText::Format( NSLOCTEXT("Sequencer", "AddAnimSequence", "{AnimationName}"), Args ),
 				FText::Format( NSLOCTEXT("Sequencer", "AddAnimSequenceTooltip", "Adds a {AnimationName} animation to this skeletal mesh."), Args ),
 				FSlateIcon(),
-				FUIAction(FExecuteAction::CreateSP(ParentSequencer.Get(), &ISequencer::AddAnimation, ObjectBinding, AnimSequence))
+				FUIAction(FExecuteAction::CreateRaw(this, &FSkeletalAnimationTrackEditor::HandleAddAnimSequencerMenuEntryExecute, ObjectBinding, AnimSequence))
 				);
 		}
 	}
@@ -298,3 +298,9 @@ USkeleton* FSkeletalAnimationTrackEditor::AcquireSkeletonFromObjectGuid(const FG
 
 	return Skeleton;
 }
+
+
+void FSkeletalAnimationTrackEditor::HandleAddAnimSequencerMenuEntryExecute(FGuid ObjectGuid, UAnimSequence* AnimSequence)
+{
+	AddKey(ObjectGuid, AnimSequence);
+}
diff --git a/Engine/Source/Editor/MovieSceneTools/Private/TrackEditors/SkeletalAnimationTrackEditor.h b/Engine/Source/Editor/MovieSceneTools/Private/TrackEditors/SkeletalAnimationTrackEditor.h
--- a/Engine/Source/Editor/MovieSceneTools/Private/TrackEditors/SkeletalAnimationTrackEditor.h
+++ b/Engine/Source/Editor/MovieSceneTools/Private/TrackEditors/SkeletalAnimationTrackEditor.h
@@ -41,6 +41,9 @@ private:
 
 	/** Gets a skeleton from an object guid in the movie scene */
 	class USkeleton* AcquireSkeletonFromObjectGuid(const FGuid& Guid);
+
+	/** Callback for executing the "Add Animation Sequence" menu entry. */
+	void HandleAddAnimSequencerMenuEntryExecute(FGuid ObjectGuid, UAnimSequence* AnimSequence);
 };
 
 
diff --git a/Engine/Source/Editor/Sequencer/Private/Sequencer.cpp b/Engine/Source/Editor/Sequencer/Private/Sequencer.cpp
--- a/Engine/Source/Editor/Sequencer/Private/Sequencer.cpp
+++ b/Engine/Source/Editor/Sequencer/Private/Sequencer.cpp
@@ -375,15 +375,6 @@ void FSequencer::PopToMovieScene( TSharedRef<FMovieSceneSequenceInstance> SubMov
 	}
 }
 
-void FSequencer::AddAnimation(FGuid ObjectGuid, class UAnimSequence* AnimSequence)
-{
-	// @todo Sequencer: This should not be needed here.  the track editor could be the entry point
-	if (SkeletalAnimationTrackEditor.IsValid())
-	{
-		SkeletalAnimationTrackEditor.Pin()->AddKey(ObjectGuid, AnimSequence);
-	}
-}
-
 void FSequencer::DeleteSections(const TSet<TWeakObjectPtr<UMovieSceneSection> >& Sections)
 {
 	UMovieScene* MovieScene = GetFocusedMovieSceneSequence()->GetMovieScene();
diff --git a/Engine/Source/Editor/Sequencer/Private/Sequencer.h b/Engine/Source/Editor/Sequencer/Private/Sequencer.h
--- a/Engine/Source/Editor/Sequencer/Private/Sequencer.h
+++ b/Engine/Source/Editor/Sequencer/Private/Sequencer.h
@@ -68,7 +68,6 @@ public:
 	virtual TSharedRef<FMovieSceneSequenceInstance> GetFocusedMovieSceneSequenceInstance() const override;
 	virtual void FocusSubMovieScene( TSharedRef<FMovieSceneSequenceInstance> SubMovieSceneInstance ) override;
 	TSharedRef<FMovieSceneSequenceInstance> GetInstanceForSubMovieSceneSection( UMovieSceneSection& SubMovieSceneSection ) const override;
-	virtual void AddAnimation(FGuid ObjectGuid, class UAnimSequence* AnimSequence) override;
 	virtual bool GetAutoKeyEnabled() const override;
 	virtual void SetAutoKeyEnabled(bool bAutoKeyEnabled) override;
 	virtual bool GetKeyAllEnabled() const override;
diff --git a/Engine/Source/Editor/Sequencer/Public/ISequencer.h b/Engine/Source/Editor/Sequencer/Public/ISequencer.h
--- a/Engine/Source/Editor/Sequencer/Public/ISequencer.h
+++ b/Engine/Source/Editor/Sequencer/Public/ISequencer.h
@@ -110,9 +110,6 @@ public:
 	 * @see FilterToShotSections
 	 */
 	virtual void FilterToSelectedShotSections(bool bZoomToShotBounds = true) = 0;
-	
-	/** Adds a new animation to the specified skeletal mesh object guid */
-	virtual void AddAnimation(FGuid ObjectGuid, UAnimSequence* AnimSequence) = 0;
 
 	/**
 	 * Gets the current time of the time slider relative to the passed in movie scene
