Sequencer: Code and documentation cleanup pass

- introduced ISequencerTrack
- reduced dependency on FMovieSceneTrack
- forward declarations
- added missing virtual destructors, include files, comments
- NULL to nullptr

diff --git a/Engine/Source/Editor/MovieSceneTools/Private/TrackEditors/AttachTrackEditor.cpp b/Engine/Source/Editor/MovieSceneTools/Private/TrackEditors/AttachTrackEditor.cpp
--- a/Engine/Source/Editor/MovieSceneTools/Private/TrackEditors/AttachTrackEditor.cpp
+++ b/Engine/Source/Editor/MovieSceneTools/Private/TrackEditors/AttachTrackEditor.cpp
@@ -16,6 +16,7 @@
 #include "GameFramework/WorldSettings.h"
 #include "AttachTrackEditor.h"
 
+
 void GetAttachActors(TSet<AActor*>& SelectedActors, TSet<AActor*>& UnselectedActors, TArray<UObject*> InObjects)
 {
 	if (GEditor->GetEditorWorldContext().World())
@@ -51,17 +52,19 @@ void GetAttachActors(TSet<AActor*>& SelectedActors, TSet<AActor*>& UnselectedAct
 	});
 }
 
+
 /**
  * Class that draws an attach section in the sequencer
  */
-class F3DAttachSection : public ISequencerSection
+class F3DAttachSection
+	: public ISequencerSection
 {
 public:
+
 	F3DAttachSection( UMovieSceneSection& InSection, F3DAttachTrackEditor* InAttachTrackEditor )
 		: Section( InSection )
 		, AttachTrackEditor(InAttachTrackEditor)
-	{
-	}
+	{ }
 
 	/** ISequencerSection interface */
 	virtual UMovieSceneSection* GetSectionObject() override
@@ -153,6 +156,7 @@ public:
 	}
 
 private:
+
 	/** The section we are visualizing */
 	UMovieSceneSection& Section;
 
@@ -160,26 +164,31 @@ private:
 	F3DAttachTrackEditor* AttachTrackEditor;
 };
 
+
 F3DAttachTrackEditor::F3DAttachTrackEditor( TSharedRef<ISequencer> InSequencer )
 	: FMovieSceneTrackEditor( InSequencer ) 
 {
 }
 
+
 F3DAttachTrackEditor::~F3DAttachTrackEditor()
 {
 }
 
-TSharedRef<FMovieSceneTrackEditor> F3DAttachTrackEditor::CreateTrackEditor( TSharedRef<ISequencer> InSequencer )
+
+TSharedRef<ISequencerTrackEditor> F3DAttachTrackEditor::CreateTrackEditor( TSharedRef<ISequencer> InSequencer )
 {
 	return MakeShareable( new F3DAttachTrackEditor( InSequencer ) );
 }
 
+
 bool F3DAttachTrackEditor::SupportsType( TSubclassOf<UMovieSceneTrack> Type ) const
 {
 	// We support animatable transforms
 	return Type == UMovieScene3DAttachTrack::StaticClass();
 }
 
+
 TSharedRef<ISequencerSection> F3DAttachTrackEditor::MakeSectionInterface( UMovieSceneSection& SectionObject, UMovieSceneTrack* Track )
 {
 	check( SupportsType( SectionObject.GetOuter()->GetClass() ) );
@@ -189,11 +198,13 @@ TSharedRef<ISequencerSection> F3DAttachTrackEditor::MakeSectionInterface( UMovie
 	return NewSection;
 }
 
+
 void F3DAttachTrackEditor::AddKey( const FGuid& ObjectGuid, UObject* AdditionalAsset )
 {
 	AddAttach(ObjectGuid, AdditionalAsset);
 }
 
+
 void F3DAttachTrackEditor::BuildObjectBindingTrackMenu(FMenuBuilder& MenuBuilder, const FGuid& ObjectBinding, const UClass* ObjectClass)
 {
 	// build a menu with all the possible actors, add the selected actors first
@@ -227,6 +238,7 @@ void F3DAttachTrackEditor::BuildObjectBindingTrackMenu(FMenuBuilder& MenuBuilder
 	}
 }
 
+
 void F3DAttachTrackEditor::AddAttachSubMenu(FMenuBuilder& MenuBuilder, FGuid ObjectBinding)
 {
 	TArray<UObject*> OutObjects;
@@ -284,6 +296,7 @@ void F3DAttachTrackEditor::AddAttachSubMenu(FMenuBuilder& MenuBuilder, FGuid Obj
 	}
 }
 
+
 void F3DAttachTrackEditor::AddAttach(FGuid ObjectGuid, UObject* AdditionalAsset)
 {
 	TArray<UObject*> OutObjects;
@@ -293,6 +306,7 @@ void F3DAttachTrackEditor::AddAttach(FGuid ObjectGuid, UObject* AdditionalAsset)
 		FOnKeyProperty::CreateRaw( this, &F3DAttachTrackEditor::AddKeyInternal, OutObjects, AdditionalAsset) );
 }
 
+
 void F3DAttachTrackEditor::SetAttach(UMovieSceneSection* Section, AActor* Actor)
 {
 	const FScopedTransaction Transaction(NSLOCTEXT("Sequencer", "UndoSetAttach", "Set Attach"));
@@ -306,6 +320,7 @@ void F3DAttachTrackEditor::SetAttach(UMovieSceneSection* Section, AActor* Actor)
 	}
 }
 
+
 void F3DAttachTrackEditor::AddKeyInternal( float KeyTime, const TArray<UObject*> Objects, UObject* AdditionalAsset)
 {
 	AActor* Actor = nullptr;
@@ -322,39 +337,41 @@ void F3DAttachTrackEditor::AddKeyInternal( float KeyTime, const TArray<UObject*>
 		ActorId = FindOrCreateHandleToObject(Actor);
 	}
 
-	if (ActorId.IsValid())
+	if (!ActorId.IsValid())
 	{
-		for( int32 ObjectIndex = 0; ObjectIndex < Objects.Num(); ++ObjectIndex )
-		{
-			UObject* Object = Objects[ObjectIndex];
+		return;
+	}
 
-			FGuid ObjectHandle = FindOrCreateHandleToObject( Object );
-			if (ObjectHandle.IsValid())
-			{
-				MovieSceneHelpers::SetRuntimeObjectMobility(Object);
+	for( int32 ObjectIndex = 0; ObjectIndex < Objects.Num(); ++ObjectIndex )
+	{
+		UObject* Object = Objects[ObjectIndex];
 
-				UMovieSceneTrack* Track = GetTrackForObject( ObjectHandle, UMovieScene3DAttachTrack::StaticClass(), FName("Attach"));
+		FGuid ObjectHandle = FindOrCreateHandleToObject( Object );
+		if (ObjectHandle.IsValid())
+		{
+			MovieSceneHelpers::SetRuntimeObjectMobility(Object);
 
-				if (ensure(Track))
-				{
-					// Clamp to next attach section's start time or the end of the current sequencer view range
-					float AttachEndTime = GetSequencer()->GetViewRange().GetUpperBoundValue();
+			UMovieSceneTrack* Track = GetTrackForObject( ObjectHandle, UMovieScene3DAttachTrack::StaticClass(), FName("Attach"));
+
+			if (ensure(Track))
+			{
+				// Clamp to next attach section's start time or the end of the current sequencer view range
+				float AttachEndTime = GetSequencer()->GetViewRange().GetUpperBoundValue();
 	
-					for (int32 AttachSectionIndex = 0; AttachSectionIndex < Track->GetAllSections().Num(); ++AttachSectionIndex)
+				for (int32 AttachSectionIndex = 0; AttachSectionIndex < Track->GetAllSections().Num(); ++AttachSectionIndex)
+				{
+					float StartTime = Track->GetAllSections()[AttachSectionIndex]->GetStartTime();
+					float EndTime = Track->GetAllSections()[AttachSectionIndex]->GetEndTime();
+					if (KeyTime < StartTime)
 					{
-						float StartTime = Track->GetAllSections()[AttachSectionIndex]->GetStartTime();
-						float EndTime = Track->GetAllSections()[AttachSectionIndex]->GetEndTime();
-						if (KeyTime < StartTime)
+						if (AttachEndTime > StartTime)
 						{
-							if (AttachEndTime > StartTime)
-							{
-								AttachEndTime = StartTime;
-							}
+							AttachEndTime = StartTime;
 						}
 					}
-
-					Cast<UMovieScene3DAttachTrack>(Track)->AddConstraint( KeyTime, AttachEndTime, ActorId );
 				}
+
+				Cast<UMovieScene3DAttachTrack>(Track)->AddConstraint( KeyTime, AttachEndTime, ActorId );
 			}
 		}
 	}
diff --git a/Engine/Source/Editor/MovieSceneTools/Private/TrackEditors/AttachTrackEditor.h b/Engine/Source/Editor/MovieSceneTools/Private/TrackEditors/AttachTrackEditor.h
--- a/Engine/Source/Editor/MovieSceneTools/Private/TrackEditors/AttachTrackEditor.h
+++ b/Engine/Source/Editor/MovieSceneTools/Private/TrackEditors/AttachTrackEditor.h
@@ -4,19 +4,25 @@
 
 #include "MovieScene3DAttachSection.h"
 #include "MovieScene3DAttachTrack.h"
+
+
 /**
  * Tools for attaching an object to another object
  */
-class F3DAttachTrackEditor : public FMovieSceneTrackEditor
+class F3DAttachTrackEditor
+	: public FMovieSceneTrackEditor
 {
 public:
+
 	/**
 	 * Constructor
 	 *
-	 * @param InSequencer	The sequencer instance to be used by this tool
+	 * @param InSequencer The sequencer instance to be used by this tool
 	 */
 	F3DAttachTrackEditor( TSharedRef<ISequencer> InSequencer );
-	~F3DAttachTrackEditor();
+
+	/** Virtual destructor. */
+	virtual ~F3DAttachTrackEditor();
 
 	/**
 	 * Creates an instance of this class.  Called by a sequencer 
@@ -24,13 +30,9 @@ public:
 	 * @param OwningSequencer The sequencer instance to be used by this tool
 	 * @return The new instance of this class
 	 */
-	static TSharedRef<FMovieSceneTrackEditor> CreateTrackEditor( TSharedRef<ISequencer> OwningSequencer );
+	static TSharedRef<ISequencerTrackEditor> CreateTrackEditor( TSharedRef<ISequencer> OwningSequencer );
 
-	/** FMovieSceneTrackEditor Interface */
-	virtual bool SupportsType( TSubclassOf<UMovieSceneTrack> Type ) const override;
-	virtual TSharedRef<ISequencerSection> MakeSectionInterface( UMovieSceneSection& SectionObject, UMovieSceneTrack* Track ) override;
-	virtual void AddKey( const FGuid& ObjectGuid, UObject* AdditionalAsset = NULL ) override;
-	virtual void BuildObjectBindingTrackMenu(FMenuBuilder& MenuBuilder, const FGuid& ObjectBinding, const UClass* ObjectClass) override;
+public:
 
 	/** Add attach */
 	void AddAttach(FGuid ObjectGuid, UObject* AdditionalAsset);
@@ -38,6 +40,15 @@ public:
 	/** Set attach */
 	void SetAttach(UMovieSceneSection* Section, AActor* ActorWithSplineComponent);
 
+public:
+
+	// FMovieSceneTrackEditor interface
+
+	virtual bool SupportsType( TSubclassOf<UMovieSceneTrack> Type ) const override;
+	virtual TSharedRef<ISequencerSection> MakeSectionInterface( UMovieSceneSection& SectionObject, UMovieSceneTrack* Track ) override;
+	virtual void AddKey( const FGuid& ObjectGuid, UObject* AdditionalAsset = NULL ) override;
+	virtual void BuildObjectBindingTrackMenu(FMenuBuilder& MenuBuilder, const FGuid& ObjectBinding, const UClass* ObjectClass) override;
+
 private:
 
 	/** Add attach sub menu */
@@ -46,5 +57,3 @@ private:
 	/** Delegate for AnimatablePropertyChanged in AddKey */
 	void AddKeyInternal(float KeyTime, const TArray<UObject*> Objects, UObject* AdditionalAsset);
 };
-
-
diff --git a/Engine/Source/Editor/MovieSceneTools/Private/TrackEditors/AudioTrackEditor.cpp b/Engine/Source/Editor/MovieSceneTools/Private/TrackEditors/AudioTrackEditor.cpp
--- a/Engine/Source/Editor/MovieSceneTools/Private/TrackEditors/AudioTrackEditor.cpp
+++ b/Engine/Source/Editor/MovieSceneTools/Private/TrackEditors/AudioTrackEditor.cpp
@@ -39,7 +39,6 @@ namespace AnimatableAudioEditorConstants
 }
 
 
-
 /** These utility functions should go away once we start handling sound cues properly */
 USoundWave* DeriveSoundWave(UMovieSceneAudioSection* AudioSection)
 {
@@ -70,6 +69,7 @@ USoundWave* DeriveSoundWave(UMovieSceneAudioSection* AudioSection)
 	return SoundWave;
 }
 
+
 float DeriveUnloopedDuration(UMovieSceneAudioSection* AudioSection)
 {
 	USoundWave* SoundWave = DeriveSoundWave(AudioSection);
@@ -79,11 +79,12 @@ float DeriveUnloopedDuration(UMovieSceneAudioSection* AudioSection)
 }
 
 
-
 /**
  * The audio thumnail, which holds a texture which it can pass back to a viewport to render
  */
-class FAudioThumbnail : public ISlateViewport, public TSharedFromThis<FAudioThumbnail>
+class FAudioThumbnail
+	: public ISlateViewport
+	, public TSharedFromThis<FAudioThumbnail>
 {
 public:
 	FAudioThumbnail(UMovieSceneSection& InSection, TRange<float> DrawRange, int InTextureSize);
@@ -115,6 +116,7 @@ private:
 	int32 TextureSize;
 };
 
+
 FAudioThumbnail::FAudioThumbnail(UMovieSceneSection& InSection, TRange<float> DrawRange, int32 InTextureSize)
 	: Section(InSection)
 	, Texture(NULL)
@@ -138,6 +140,7 @@ FAudioThumbnail::FAudioThumbnail(UMovieSceneSection& InSection, TRange<float> Dr
 	}
 }
 
+
 FAudioThumbnail::~FAudioThumbnail()
 {
 	if (ShouldRender())
@@ -153,10 +156,12 @@ FAudioThumbnail::~FAudioThumbnail()
 	}
 }
 
+
 FIntPoint FAudioThumbnail::GetSize() const {return FIntPoint(TextureSize, AnimatableAudioEditorConstants::AudioTrackHeight);}
 FSlateShaderResource* FAudioThumbnail::GetViewportRenderTargetTexture() const {return Texture;}
 bool FAudioThumbnail::RequiresVsync() const {return false;}
 
+
 void FAudioThumbnail::GenerateWaveformPreview(TArray<uint8>& OutData, TRange<float> DrawRange)
 {
 	UMovieSceneAudioSection* AudioSection = Cast<UMovieSceneAudioSection>(&Section);
@@ -209,6 +214,7 @@ void FAudioThumbnail::GenerateWaveformPreview(TArray<uint8>& OutData, TRange<flo
 	}
 }
 
+
 void FAudioThumbnail::DrawWaveformLine(TArray<uint8>& OutData, int32 NumChannels, int16* LookupData, int32 LookupStartIndex, int32 LookupEndIndex, int32 LookupSize, int32 X, int32 XAxisHeight, int32 MaxAmplitude)
 {
 	int32 ModifiedLookupStartIndex = NumChannels == 2 ? (LookupStartIndex % 2 == 0 ? LookupStartIndex : LookupStartIndex - 1) : LookupStartIndex;
@@ -318,7 +324,6 @@ void FAudioThumbnail::DrawWaveformLine(TArray<uint8>& OutData, int32 NumChannels
 }
 
 
-
 FAudioSection::FAudioSection( UMovieSceneSection& InSection, bool bOnAMasterTrack )
 	: Section( InSection )
 	, StoredDrawRange(TRange<float>::Empty())
@@ -326,29 +331,35 @@ FAudioSection::FAudioSection( UMovieSceneSection& InSection, bool bOnAMasterTrac
 {
 }
 
+
 FAudioSection::~FAudioSection()
 {
 }
 
+
 UMovieSceneSection* FAudioSection::GetSectionObject()
 { 
 	return &Section;
 }
 
+
 FText FAudioSection::GetDisplayName() const
 {
 	return bIsOnAMasterTrack ? NSLOCTEXT("FAudioSection", "MasterAudioDisplayName", "Master Audio") :  NSLOCTEXT("FAudioSection", "AudioDisplayName", "Audio");
 }
 
+
 FText FAudioSection::GetSectionTitle() const
 {
 	return FText::FromString( Cast<UMovieSceneAudioSection>(&Section)->GetSound()->GetName() );
 }
 
+
 float FAudioSection::GetSectionHeight() const
 {
 	return (float)AnimatableAudioEditorConstants::AudioTrackHeight;
- }
+}
+
 
 int32 FAudioSection::OnPaintSection( const FGeometry& AllottedGeometry, const FSlateRect& SectionClippingRect, FSlateWindowElementList& OutDrawElements, int32 LayerId, bool bParentEnabled ) const
 {
@@ -384,6 +395,7 @@ int32 FAudioSection::OnPaintSection( const FGeometry& AllottedGeometry, const FS
 	return LayerId+1;
 }
 
+
 void FAudioSection::Tick( const FGeometry& AllottedGeometry, const FGeometry& ParentGeometry, const double InCurrentTime, const float InDeltaTime )
 {
 	UMovieSceneAudioSection* AudioSection = Cast<UMovieSceneAudioSection>(&Section);
@@ -420,6 +432,7 @@ void FAudioSection::Tick( const FGeometry& AllottedGeometry, const FGeometry& Pa
 	}
 }
 
+
 void FAudioSection::RegenerateWaveforms(TRange<float> DrawRange, int32 XOffset, int32 XSize)
 {
 	StoredDrawRange = DrawRange;
@@ -437,7 +450,6 @@ void FAudioSection::RegenerateWaveforms(TRange<float> DrawRange, int32 XOffset,
 }
 
 
-
 FAudioTrackEditor::FAudioTrackEditor( TSharedRef<ISequencer> InSequencer )
 	: FMovieSceneTrackEditor( InSequencer ) 
 {
@@ -447,16 +459,19 @@ FAudioTrackEditor::~FAudioTrackEditor()
 {
 }
 
-TSharedRef<FMovieSceneTrackEditor> FAudioTrackEditor::CreateTrackEditor( TSharedRef<ISequencer> InSequencer )
+
+TSharedRef<ISequencerTrackEditor> FAudioTrackEditor::CreateTrackEditor( TSharedRef<ISequencer> InSequencer )
 {
 	return MakeShareable( new FAudioTrackEditor( InSequencer ) );
 }
 
+
 bool FAudioTrackEditor::SupportsType( TSubclassOf<UMovieSceneTrack> Type ) const
 {
 	return Type == UMovieSceneAudioTrack::StaticClass();
 }
 
+
 TSharedRef<ISequencerSection> FAudioTrackEditor::MakeSectionInterface( UMovieSceneSection& SectionObject, UMovieSceneTrack* Track )
 {
 	check( SupportsType( SectionObject.GetOuter()->GetClass() ) );
@@ -467,6 +482,7 @@ TSharedRef<ISequencerSection> FAudioTrackEditor::MakeSectionInterface( UMovieSce
 	return NewSection;
 }
 
+
 bool FAudioTrackEditor::HandleAssetAdded(UObject* Asset, const FGuid& TargetObjectGuid)
 {
 	if (Asset->IsA<USoundBase>())
@@ -492,6 +508,7 @@ bool FAudioTrackEditor::HandleAssetAdded(UObject* Asset, const FGuid& TargetObje
 	return false;
 }
 
+
 void FAudioTrackEditor::AddNewMasterSound( float KeyTime, USoundBase* Sound )
 {
 	UMovieSceneTrack* Track = GetMasterTrack( UMovieSceneAudioTrack::StaticClass() );
@@ -500,6 +517,7 @@ void FAudioTrackEditor::AddNewMasterSound( float KeyTime, USoundBase* Sound )
 	AudioTrack->AddNewSound( Sound, KeyTime );
 }
 
+
 void FAudioTrackEditor::AddNewAttachedSound( float KeyTime, USoundBase* Sound, TArray<UObject*> ObjectsToAttachTo )
 {
 	for( int32 ObjectIndex = 0; ObjectIndex < ObjectsToAttachTo.Num(); ++ObjectIndex )
diff --git a/Engine/Source/Editor/MovieSceneTools/Private/TrackEditors/AudioTrackEditor.h b/Engine/Source/Editor/MovieSceneTools/Private/TrackEditors/AudioTrackEditor.h
--- a/Engine/Source/Editor/MovieSceneTools/Private/TrackEditors/AudioTrackEditor.h
+++ b/Engine/Source/Editor/MovieSceneTools/Private/TrackEditors/AudioTrackEditor.h
@@ -6,16 +6,20 @@
 /**
  * Tools for audio tracks
  */
-class FAudioTrackEditor : public FMovieSceneTrackEditor
+class FAudioTrackEditor
+	: public FMovieSceneTrackEditor
 {
 public:
+
 	/**
 	 * Constructor
 	 *
-	 * @param InSequencer	The sequencer instance to be used by this tool
+	 * @param InSequencer The sequencer instance to be used by this tool
 	 */
 	FAudioTrackEditor( TSharedRef<ISequencer> InSequencer );
-	~FAudioTrackEditor();
+
+	/** Virtual destructor. */
+	virtual ~FAudioTrackEditor();
 
 	/**
 	 * Creates an instance of this class.  Called by a sequencer 
@@ -23,30 +27,45 @@ public:
 	 * @param OwningSequencer The sequencer instance to be used by this tool
 	 * @return The new instance of this class
 	 */
-	static TSharedRef<FMovieSceneTrackEditor> CreateTrackEditor( TSharedRef<ISequencer> OwningSequencer );
+	static TSharedRef<ISequencerTrackEditor> CreateTrackEditor( TSharedRef<ISequencer> OwningSequencer );
+
+public:
+
+	// FMovieSceneTrackEditor interface
 
-	/** FMovieSceneTrackEditor Interface */
 	virtual bool SupportsType( TSubclassOf<UMovieSceneTrack> Type ) const override;
 	virtual TSharedRef<ISequencerSection> MakeSectionInterface( UMovieSceneSection& SectionObject, UMovieSceneTrack* Track ) override;
 	virtual bool HandleAssetAdded(UObject* Asset, const FGuid& TargetObjectGuid) override;
 	
 private:
+
 	/** Delegate for AnimatablePropertyChanged in HandleAssetAdded for master sounds */
 	void AddNewMasterSound( float KeyTime, class USoundBase* Sound );
+
 	/** Delegate for AnimatablePropertyChanged in HandleAssetAdded for attached sounds */
 	void AddNewAttachedSound( float KeyTime, class USoundBase* Sound, TArray<UObject*> ObjectsToAttachTo );
 };
 
 
-
-/** Class for audio sections, handles drawing of all waveform previews */
-class FAudioSection : public ISequencerSection, public TSharedFromThis<FAudioSection>
+/**
+ * Class for audio sections, handles drawing of all waveform previews.
+ */
+class FAudioSection
+	: public ISequencerSection
+	, public TSharedFromThis<FAudioSection>
 {
 public:
+
+	/** Constructor. */
 	FAudioSection( UMovieSceneSection& InSection, bool bOnAMasterTrack );
-	~FAudioSection();
 
-	/** ISequencerSection interface */
+	/** Virtual destructor. */
+	virtual ~FAudioSection();
+
+public:
+
+	// ISequencerSection interface
+
 	virtual UMovieSceneSection* GetSectionObject() override;
 	virtual FText GetDisplayName() const override;
 	virtual FText GetSectionTitle() const override;
@@ -56,21 +75,23 @@ public:
 	virtual void Tick( const FGeometry& AllottedGeometry, const FGeometry& ParentGeometry, const double InCurrentTime, const float InDeltaTime ) override;
 	
 private:
-	/* Re-creates the texture used to preview the waveform */
+
+	/* Re-creates the texture used to preview the waveform. */
 	void RegenerateWaveforms(TRange<float> DrawRange, int32 XOffset, int32 XSize);
 
 private:
-	/** The section we are visualizing */
+
+	/** The section we are visualizing. */
 	UMovieSceneSection& Section;
 	
-	/** The waveform thumbnail render object */
+	/** The waveform thumbnail render object. */
 	TSharedPtr<class FAudioThumbnail> WaveformThumbnail;
 
-	/** Stored data about the waveform to determine when it is invalidated */
+	/** Stored data about the waveform to determine when it is invalidated. */
 	TRange<float> StoredDrawRange;
 	int32 StoredXOffset;
 	int32 StoredXSize;
 
-	/** Whether this section is on a master audio track or an attached audio track */
+	/** Whether this section is on a master audio track or an attached audio track. */
 	bool bIsOnAMasterTrack;
 };
diff --git a/Engine/Source/Editor/MovieSceneTools/Private/TrackEditors/EventTrackEditor.cpp b/Engine/Source/Editor/MovieSceneTools/Private/TrackEditors/EventTrackEditor.cpp
--- a/Engine/Source/Editor/MovieSceneTools/Private/TrackEditors/EventTrackEditor.cpp
+++ b/Engine/Source/Editor/MovieSceneTools/Private/TrackEditors/EventTrackEditor.cpp
@@ -8,7 +8,7 @@
 /* FEventTrackEditor static functions
  *****************************************************************************/
 
-TSharedRef<FMovieSceneTrackEditor> FEventTrackEditor::CreateTrackEditor(TSharedRef<ISequencer> InSequencer)
+TSharedRef<ISequencerTrackEditor> FEventTrackEditor::CreateTrackEditor(TSharedRef<ISequencer> InSequencer)
 {
 	return MakeShareable(new FEventTrackEditor(InSequencer));
 }
diff --git a/Engine/Source/Editor/MovieSceneTools/Private/TrackEditors/EventTrackEditor.h b/Engine/Source/Editor/MovieSceneTools/Private/TrackEditors/EventTrackEditor.h
--- a/Engine/Source/Editor/MovieSceneTools/Private/TrackEditors/EventTrackEditor.h
+++ b/Engine/Source/Editor/MovieSceneTools/Private/TrackEditors/EventTrackEditor.h
@@ -23,7 +23,7 @@ public:
 	 * @param InSequencer The sequencer instance to be used by this tool.
 	 * @return The new instance of this class.
 	 */
-	static TSharedRef<FMovieSceneTrackEditor> CreateTrackEditor(TSharedRef<ISequencer> InSequencer);
+	static TSharedRef<ISequencerTrackEditor> CreateTrackEditor(TSharedRef<ISequencer> InSequencer);
 
 public:
 
diff --git a/Engine/Source/Editor/MovieSceneTools/Private/TrackEditors/MaterialTrackEditor.cpp b/Engine/Source/Editor/MovieSceneTools/Private/TrackEditors/MaterialTrackEditor.cpp
--- a/Engine/Source/Editor/MovieSceneTools/Private/TrackEditors/MaterialTrackEditor.cpp
+++ b/Engine/Source/Editor/MovieSceneTools/Private/TrackEditors/MaterialTrackEditor.cpp
@@ -6,13 +6,16 @@
 #include "MaterialParameterSection.h"
 #include "MovieSceneMaterialParameterSection.h"
 
+
 #define LOCTEXT_NAMESPACE "MaterialTrackEditor"
 
+
 FMaterialTrackEditor::FMaterialTrackEditor( TSharedRef<ISequencer> InSequencer )
 	: FMovieSceneTrackEditor( InSequencer )
 {
 }
 
+
 TSharedRef<ISequencerSection> FMaterialTrackEditor::MakeSectionInterface( UMovieSceneSection& SectionObject, UMovieSceneTrack* Track )
 {
 	UMovieSceneMaterialParameterSection* ParameterSection = Cast<UMovieSceneMaterialParameterSection>(&SectionObject);
@@ -20,6 +23,7 @@ TSharedRef<ISequencerSection> FMaterialTrackEditor::MakeSectionInterface( UMovie
 	return MakeShareable(new FMaterialParameterSection( *ParameterSection, ParameterSection->GetFName()));
 }
 
+
 TSharedPtr<SWidget> FMaterialTrackEditor::BuildOutlinerEditWidget( const FGuid& ObjectBinding, UMovieSceneTrack* Track )
 {
 	UMovieSceneMaterialTrack* MaterialTrack = Cast<UMovieSceneMaterialTrack>(Track);
@@ -52,6 +56,7 @@ TSharedPtr<SWidget> FMaterialTrackEditor::BuildOutlinerEditWidget( const FGuid&
 		];
 }
 
+
 struct FParameterNameAndAction
 {
 	FName ParameterName;
@@ -69,6 +74,7 @@ struct FParameterNameAndAction
 	}
 };
 
+
 TSharedRef<SWidget> FMaterialTrackEditor::OnGetAddParameterMenuContent( FGuid ObjectBinding, UMovieSceneMaterialTrack* MaterialTrack )
 {
 	FMenuBuilder AddParameterMenuBuilder( true, nullptr );
@@ -111,6 +117,7 @@ TSharedRef<SWidget> FMaterialTrackEditor::OnGetAddParameterMenuContent( FGuid Ob
 	return AddParameterMenuBuilder.MakeWidget();
 }
 
+
 UMaterial* FMaterialTrackEditor::GetMaterialForTrack( FGuid ObjectBinding, UMovieSceneMaterialTrack* MaterialTrack )
 {
 	UMaterialInterface* MaterialInterface = GetMaterialInterfaceForTrack( ObjectBinding, MaterialTrack );
@@ -133,6 +140,7 @@ UMaterial* FMaterialTrackEditor::GetMaterialForTrack( FGuid ObjectBinding, UMovi
 	return nullptr;
 }
 
+
 void FMaterialTrackEditor::AddScalarParameter( FGuid ObjectBinding, UMovieSceneMaterialTrack* MaterialTrack, FName ParameterName )
 {
 	UMovieSceneSequence* MovieSceneSequence = GetMovieSceneSequence();
@@ -154,6 +162,7 @@ void FMaterialTrackEditor::AddScalarParameter( FGuid ObjectBinding, UMovieSceneM
 	NotifyMovieSceneDataChanged();
 }
 
+
 void FMaterialTrackEditor::AddVectorParameter( FGuid ObjectBinding, UMovieSceneMaterialTrack* MaterialTrack, FName ParameterName )
 {
 	UMovieSceneSequence* MovieSceneSequence = GetMovieSceneSequence();
@@ -175,21 +184,25 @@ void FMaterialTrackEditor::AddVectorParameter( FGuid ObjectBinding, UMovieSceneM
 	NotifyMovieSceneDataChanged();
 }
 
+
 FComponentMaterialTrackEditor::FComponentMaterialTrackEditor( TSharedRef<ISequencer> InSequencer )
 	: FMaterialTrackEditor( InSequencer )
 {
 }
 
-TSharedRef<FMovieSceneTrackEditor> FComponentMaterialTrackEditor::CreateTrackEditor( TSharedRef<ISequencer> OwningSequencer )
+
+TSharedRef<ISequencerTrackEditor> FComponentMaterialTrackEditor::CreateTrackEditor( TSharedRef<ISequencer> OwningSequencer )
 {
 	return MakeShareable( new FComponentMaterialTrackEditor( OwningSequencer ) );
 }
 
+
 bool FComponentMaterialTrackEditor::SupportsType( TSubclassOf<UMovieSceneTrack> Type ) const
 {
 	return Type == UMovieSceneComponentMaterialTrack::StaticClass();
 }
 
+
 UMaterialInterface* FComponentMaterialTrackEditor::GetMaterialInterfaceForTrack( FGuid ObjectBinding, UMovieSceneMaterialTrack* MaterialTrack )
 {
 	UObject* ComponentObject = GetSequencer()->GetFocusedMovieSceneSequence()->FindObject( ObjectBinding );
@@ -202,4 +215,5 @@ UMaterialInterface* FComponentMaterialTrackEditor::GetMaterialInterfaceForTrack(
 	return nullptr;
 }
 
-#undef LOCTEXT_NAMESPACE
\ No newline at end of file
+
+#undef LOCTEXT_NAMESPACE
diff --git a/Engine/Source/Editor/MovieSceneTools/Private/TrackEditors/MaterialTrackEditor.h b/Engine/Source/Editor/MovieSceneTools/Private/TrackEditors/MaterialTrackEditor.h
--- a/Engine/Source/Editor/MovieSceneTools/Private/TrackEditors/MaterialTrackEditor.h
+++ b/Engine/Source/Editor/MovieSceneTools/Private/TrackEditors/MaterialTrackEditor.h
@@ -2,27 +2,37 @@
 
 #pragma once
 
+
 class UMovieSceneMaterialTrack;
 
+
 /**
  * Track editor for material parameters.
  */
-class FMaterialTrackEditor : public FMovieSceneTrackEditor
+class FMaterialTrackEditor
+	: public FMovieSceneTrackEditor
 {
 public:
+
+	/** Constructor. */
 	FMaterialTrackEditor( TSharedRef<ISequencer> InSequencer );
 
+	/** Virtual destructor. */
 	virtual ~FMaterialTrackEditor() { }
 
-	/* FMovieSceneTrackEditor interface */
+public:
+
+	// FMovieSceneTrackEditor interface
 	virtual TSharedRef<ISequencerSection> MakeSectionInterface( UMovieSceneSection& SectionObject, UMovieSceneTrack* Track ) override;
 	virtual TSharedPtr<SWidget> BuildOutlinerEditWidget( const FGuid& ObjectBinding, UMovieSceneTrack* Track ) override;
 
 protected:
+
 	/** Gets a material interface for a specific object binding and material track. */
 	virtual UMaterialInterface* GetMaterialInterfaceForTrack( FGuid ObjectBinding, UMovieSceneMaterialTrack* MaterialTrack ) = 0;
 
 private:
+
 	/** Provides the contents of the add parameter menu. */
 	TSharedRef<SWidget> OnGetAddParameterMenuContent( FGuid ObjectBinding, UMovieSceneMaterialTrack* MaterialTrack );
 
@@ -44,20 +54,27 @@ private:
 	void AddVectorParameter( FGuid ObjectBinding, UMovieSceneMaterialTrack* MaterialTrack, FName ParameterName );
 };
 
+
 /**
  * A specialized material track editor for component materials
  */
-class FComponentMaterialTrackEditor : public FMaterialTrackEditor
+class FComponentMaterialTrackEditor
+	: public FMaterialTrackEditor
 {
 public:
+
 	FComponentMaterialTrackEditor( TSharedRef<ISequencer> InSequencer );
 
-	static TSharedRef<FMovieSceneTrackEditor> CreateTrackEditor( TSharedRef<ISequencer> OwningSequencer );
+	static TSharedRef<ISequencerTrackEditor> CreateTrackEditor( TSharedRef<ISequencer> OwningSequencer );
+
+public:
 
-	/** FMovieSceneTrackEditor Interface */
+	// FMovieSceneTrackEditor interface
 	virtual bool SupportsType( TSubclassOf<UMovieSceneTrack> Type ) const override;
 
 protected:
-	/** FMaterialtrackEditor interface */
+
+	// FMaterialtrackEditor interface
+
 	virtual UMaterialInterface* GetMaterialInterfaceForTrack( FGuid ObjectBinding, UMovieSceneMaterialTrack* MaterialTrack ) override;
-};
\ No newline at end of file
+};
diff --git a/Engine/Source/Editor/MovieSceneTools/Private/TrackEditors/ParticleTrackEditor.cpp b/Engine/Source/Editor/MovieSceneTools/Private/TrackEditors/ParticleTrackEditor.cpp
--- a/Engine/Source/Editor/MovieSceneTools/Private/TrackEditors/ParticleTrackEditor.cpp
+++ b/Engine/Source/Editor/MovieSceneTools/Private/TrackEditors/ParticleTrackEditor.cpp
@@ -28,26 +28,29 @@ namespace AnimatableParticleEditorConstants
 }
 
 
-
 FParticleSection::FParticleSection( UMovieSceneSection& InSection )
 	: Section( InSection )
 {
 }
 
+
 FParticleSection::~FParticleSection()
 {
 }
 
+
 UMovieSceneSection* FParticleSection::GetSectionObject()
 { 
 	return &Section;
 }
 
+
 FText FParticleSection::GetDisplayName() const
 {
 	return NSLOCTEXT("FParticleSection", "Emitter", "Emitter");
 }
 
+
 FText FParticleSection::GetSectionTitle() const
 {
 	EParticleKey::Type KeyType = Cast<UMovieSceneParticleSection>(&Section)->GetKeyType();
@@ -55,11 +58,13 @@ FText FParticleSection::GetSectionTitle() const
 		KeyType == EParticleKey::Trigger ? NSLOCTEXT("FParticleSection", "Trigger", "Trigger") : NSLOCTEXT("FParticleSection", "None", "None");
 }
 
+
 float FParticleSection::GetSectionHeight() const
 {
 	return (float)AnimatableParticleEditorConstants::ParticleTrackHeight;
 }
 
+
 int32 FParticleSection::OnPaintSection( const FGeometry& AllottedGeometry, const FSlateRect& SectionClippingRect, FSlateWindowElementList& OutDrawElements, int32 LayerId, bool bParentEnabled ) const
 {
 	const ESlateDrawEffect::Type DrawEffects = bParentEnabled ? ESlateDrawEffect::None : ESlateDrawEffect::DisabledEffect;
@@ -100,32 +105,36 @@ int32 FParticleSection::OnPaintSection( const FGeometry& AllottedGeometry, const
 	return LayerId+1;
 }
 
+
 bool FParticleSection::SectionIsResizable() const
 {
 	return Cast<UMovieSceneParticleSection>(&Section)->GetKeyType() != EParticleKey::Trigger;
 }
 
 
-
 FParticleTrackEditor::FParticleTrackEditor( TSharedRef<ISequencer> InSequencer )
 	: FMovieSceneTrackEditor( InSequencer ) 
 {
 }
 
+
 FParticleTrackEditor::~FParticleTrackEditor()
 {
 }
 
-TSharedRef<FMovieSceneTrackEditor> FParticleTrackEditor::CreateTrackEditor( TSharedRef<ISequencer> InSequencer )
+
+TSharedRef<ISequencerTrackEditor> FParticleTrackEditor::CreateTrackEditor( TSharedRef<ISequencer> InSequencer )
 {
 	return MakeShareable( new FParticleTrackEditor( InSequencer ) );
 }
 
+
 bool FParticleTrackEditor::SupportsType( TSubclassOf<UMovieSceneTrack> Type ) const
 {
 	return Type == UMovieSceneParticleTrack::StaticClass();
 }
 
+
 TSharedRef<ISequencerSection> FParticleTrackEditor::MakeSectionInterface( UMovieSceneSection& SectionObject, UMovieSceneTrack* Track )
 {
 	check( SupportsType( SectionObject.GetOuter()->GetClass() ) );
@@ -135,6 +144,7 @@ TSharedRef<ISequencerSection> FParticleTrackEditor::MakeSectionInterface( UMovie
 	return NewSection;
 }
 
+
 void FParticleTrackEditor::BuildObjectBindingTrackMenu(FMenuBuilder& MenuBuilder, const FGuid& ObjectBinding, const UClass* ObjectClass)
 {
 	if (ObjectClass->IsChildOf(AEmitter::StaticClass()))
@@ -156,6 +166,7 @@ void FParticleTrackEditor::BuildObjectBindingTrackMenu(FMenuBuilder& MenuBuilder
 	}
 }
 
+
 void FParticleTrackEditor::AddParticleKey(const FGuid ObjectGuid, bool bTrigger)
 {
 	TArray<UObject*> OutObjects;
@@ -165,6 +176,7 @@ void FParticleTrackEditor::AddParticleKey(const FGuid ObjectGuid, bool bTrigger)
 		FOnKeyProperty::CreateRaw( this, &FParticleTrackEditor::AddKeyInternal, OutObjects, bTrigger) );
 }
 
+
 void FParticleTrackEditor::AddKeyInternal( float KeyTime, const TArray<UObject*> Objects, bool bTrigger)
 {
 	for( int32 ObjectIndex = 0; ObjectIndex < Objects.Num(); ++ObjectIndex )
diff --git a/Engine/Source/Editor/MovieSceneTools/Private/TrackEditors/ParticleTrackEditor.h b/Engine/Source/Editor/MovieSceneTools/Private/TrackEditors/ParticleTrackEditor.h
--- a/Engine/Source/Editor/MovieSceneTools/Private/TrackEditors/ParticleTrackEditor.h
+++ b/Engine/Source/Editor/MovieSceneTools/Private/TrackEditors/ParticleTrackEditor.h
@@ -6,47 +6,61 @@
 /**
  * Tools for particle tracks
  */
-class FParticleTrackEditor : public FMovieSceneTrackEditor
+class FParticleTrackEditor
+	: public FMovieSceneTrackEditor
 {
 public:
+
 	/**
-	 * Constructor
+	 * Constructor.
 	 *
-	 * @param InSequencer	The sequencer instance to be used by this tool
+	 * @param InSequencer The sequencer instance to be used by this tool.
 	 */
 	FParticleTrackEditor( TSharedRef<ISequencer> InSequencer );
 	~FParticleTrackEditor();
 
 	/**
-	 * Creates an instance of this class.  Called by a sequencer 
+	 * Creates an instance of this class.  Called by a sequencer.
 	 *
-	 * @param OwningSequencer The sequencer instance to be used by this tool
-	 * @return The new instance of this class
+	 * @param OwningSequencer The sequencer instance to be used by this tool.
+	 * @return The new instance of this class.
 	 */
-	static TSharedRef<FMovieSceneTrackEditor> CreateTrackEditor( TSharedRef<ISequencer> OwningSequencer );
+	static TSharedRef<ISequencerTrackEditor> CreateTrackEditor( TSharedRef<ISequencer> OwningSequencer );
+
+public:
+
+	void AddParticleKey(const FGuid ObjectGuid, bool bTrigger);
+
+public:
+
+	// FMovieSceneTrackEditor interface
 
-	/** FMovieSceneTrackEditor Interface */
 	virtual bool SupportsType( TSubclassOf<UMovieSceneTrack> Type ) const override;
 	virtual TSharedRef<ISequencerSection> MakeSectionInterface( UMovieSceneSection& SectionObject, UMovieSceneTrack* Track ) override;
 	virtual void BuildObjectBindingTrackMenu(FMenuBuilder& MenuBuilder, const FGuid& ObjectBinding, const UClass* ObjectClass) override;
 	
-	void AddParticleKey(const FGuid ObjectGuid, bool bTrigger);
-
 private:
-	/** Delegate for AnimatablePropertyChanged in AddKey */
+
+	/** Delegate for AnimatablePropertyChanged in AddKey. */
 	virtual void AddKeyInternal( float KeyTime, const TArray<UObject*> Objects, bool bTrigger);
 };
 
 
-
-/** Class for particle sections */
-class FParticleSection : public ISequencerSection, public TSharedFromThis<FParticleSection>
+/**
+ * Class for particle sections.
+ */
+class FParticleSection
+	: public ISequencerSection
+	, public TSharedFromThis<FParticleSection>
 {
 public:
+
 	FParticleSection( UMovieSceneSection& InSection );
 	~FParticleSection();
 
-	/** ISequencerSection interface */
+public:
+
+	// ISequencerSection interface
 	virtual UMovieSceneSection* GetSectionObject() override;
 	virtual FText GetDisplayName() const override;
 	virtual FText GetSectionTitle() const override;
@@ -56,6 +70,7 @@ public:
 	virtual bool SectionIsResizable() const override;
 
 private:
-	/** The section we are visualizing */
+
+	/** The section we are visualizing. */
 	UMovieSceneSection& Section;
 };
diff --git a/Engine/Source/Editor/MovieSceneTools/Private/TrackEditors/PathTrackEditor.cpp b/Engine/Source/Editor/MovieSceneTools/Private/TrackEditors/PathTrackEditor.cpp
--- a/Engine/Source/Editor/MovieSceneTools/Private/TrackEditors/PathTrackEditor.cpp
+++ b/Engine/Source/Editor/MovieSceneTools/Private/TrackEditors/PathTrackEditor.cpp
@@ -15,6 +15,7 @@
 #include "PathTrackEditor.h"
 #include "Components/SplineComponent.h"
 
+
 void GetActorsWithSplineComponents(TSet<AActor*>& SelectedActors, TSet<AActor*>& UnselectedActors)
 {
 	if (GEditor->GetEditorWorldContext().World())
@@ -39,10 +40,12 @@ void GetActorsWithSplineComponents(TSet<AActor*>& SelectedActors, TSet<AActor*>&
 	}
 }
 
+
 /**
  * Class that draws a path section in the sequencer
  */
-class F3DPathSection : public ISequencerSection
+class F3DPathSection
+	: public ISequencerSection
 {
 public:
 	F3DPathSection( UMovieSceneSection& InSection, F3DPathTrackEditor* InPathTrackEditor )
@@ -140,6 +143,7 @@ public:
 	}
 
 private:
+
 	/** The section we are visualizing */
 	UMovieSceneSection& Section;
 
@@ -147,26 +151,31 @@ private:
 	F3DPathTrackEditor* PathTrackEditor;
 };
 
+
 F3DPathTrackEditor::F3DPathTrackEditor( TSharedRef<ISequencer> InSequencer )
 	: FMovieSceneTrackEditor( InSequencer ) 
 {
 }
 
+
 F3DPathTrackEditor::~F3DPathTrackEditor()
 {
 }
 
-TSharedRef<FMovieSceneTrackEditor> F3DPathTrackEditor::CreateTrackEditor( TSharedRef<ISequencer> InSequencer )
+
+TSharedRef<ISequencerTrackEditor> F3DPathTrackEditor::CreateTrackEditor( TSharedRef<ISequencer> InSequencer )
 {
 	return MakeShareable( new F3DPathTrackEditor( InSequencer ) );
 }
 
+
 bool F3DPathTrackEditor::SupportsType( TSubclassOf<UMovieSceneTrack> Type ) const
 {
 	// We support animatable transforms
 	return Type == UMovieScene3DPathTrack::StaticClass();
 }
 
+
 TSharedRef<ISequencerSection> F3DPathTrackEditor::MakeSectionInterface( UMovieSceneSection& SectionObject, UMovieSceneTrack* Track )
 {
 	check( SupportsType( SectionObject.GetOuter()->GetClass() ) );
@@ -176,11 +185,13 @@ TSharedRef<ISequencerSection> F3DPathTrackEditor::MakeSectionInterface( UMovieSc
 	return NewSection;
 }
 
+
 void F3DPathTrackEditor::AddKey( const FGuid& ObjectGuid, UObject* AdditionalAsset )
 {
 	AddPath(ObjectGuid, AdditionalAsset);
 }
 
+
 void F3DPathTrackEditor::BuildObjectBindingTrackMenu(FMenuBuilder& MenuBuilder, const FGuid& ObjectBinding, const UClass* ObjectClass)
 {
 	// build a menu with all the actors that have spline components, add the selected actors first
@@ -211,6 +222,7 @@ void F3DPathTrackEditor::BuildObjectBindingTrackMenu(FMenuBuilder& MenuBuilder,
 	}
 }
 
+
 void F3DPathTrackEditor::AddPathSubMenu(FMenuBuilder& MenuBuilder, FGuid ObjectBinding)
 {
 	TSet<AActor*> SelectedActors;
@@ -265,6 +277,7 @@ void F3DPathTrackEditor::AddPathSubMenu(FMenuBuilder& MenuBuilder, FGuid ObjectB
 	}
 }
 
+
 void F3DPathTrackEditor::AddPath(FGuid ObjectGuid, UObject* AdditionalAsset)
 {
 	TArray<UObject*> OutObjects;
@@ -274,6 +287,7 @@ void F3DPathTrackEditor::AddPath(FGuid ObjectGuid, UObject* AdditionalAsset)
 		FOnKeyProperty::CreateRaw( this, &F3DPathTrackEditor::AddKeyInternal, OutObjects, AdditionalAsset) );
 }
 
+
 void F3DPathTrackEditor::SetPath(UMovieSceneSection* Section, AActor* ActorWithSplineComponent)
 {
 	const FScopedTransaction Transaction(NSLOCTEXT("Sequencer", "UndoSetPath", "Set Path"));
@@ -287,6 +301,7 @@ void F3DPathTrackEditor::SetPath(UMovieSceneSection* Section, AActor* ActorWithS
 	}
 }
 
+
 void F3DPathTrackEditor::AddKeyInternal( float KeyTime, const TArray<UObject*> Objects, UObject* AdditionalAsset)
 {
 	AActor* ActorWithSplineComponent = nullptr;
@@ -303,39 +318,41 @@ void F3DPathTrackEditor::AddKeyInternal( float KeyTime, const TArray<UObject*> O
 		SplineId = FindOrCreateHandleToObject(ActorWithSplineComponent);
 	}
 
-	if (SplineId.IsValid())
+	if (!SplineId.IsValid())
 	{
-		for( int32 ObjectIndex = 0; ObjectIndex < Objects.Num(); ++ObjectIndex )
-		{
-			UObject* Object = Objects[ObjectIndex];
+		return;
+	}
 
-			FGuid ObjectHandle = FindOrCreateHandleToObject( Object );
-			if (ObjectHandle.IsValid())
-			{
-				MovieSceneHelpers::SetRuntimeObjectMobility(Object);
+	for( int32 ObjectIndex = 0; ObjectIndex < Objects.Num(); ++ObjectIndex )
+	{
+		UObject* Object = Objects[ObjectIndex];
 
-				UMovieSceneTrack* Track = GetTrackForObject( ObjectHandle, UMovieScene3DPathTrack::StaticClass(), FName("Path"));
+		FGuid ObjectHandle = FindOrCreateHandleToObject( Object );
+		if (ObjectHandle.IsValid())
+		{
+			MovieSceneHelpers::SetRuntimeObjectMobility(Object);
 
-				if (ensure(Track))
-				{
-					// Clamp to next path section's start time or the end of the current sequencer view range
-					float PathEndTime = GetSequencer()->GetViewRange().GetUpperBoundValue();
+			UMovieSceneTrack* Track = GetTrackForObject( ObjectHandle, UMovieScene3DPathTrack::StaticClass(), FName("Path"));
+
+			if (ensure(Track))
+			{
+				// Clamp to next path section's start time or the end of the current sequencer view range
+				float PathEndTime = GetSequencer()->GetViewRange().GetUpperBoundValue();
 	
-					for (int32 PathSectionIndex = 0; PathSectionIndex < Track->GetAllSections().Num(); ++PathSectionIndex)
+				for (int32 PathSectionIndex = 0; PathSectionIndex < Track->GetAllSections().Num(); ++PathSectionIndex)
+				{
+					float StartTime = Track->GetAllSections()[PathSectionIndex]->GetStartTime();
+					float EndTime = Track->GetAllSections()[PathSectionIndex]->GetEndTime();
+					if (KeyTime < StartTime)
 					{
-						float StartTime = Track->GetAllSections()[PathSectionIndex]->GetStartTime();
-						float EndTime = Track->GetAllSections()[PathSectionIndex]->GetEndTime();
-						if (KeyTime < StartTime)
+						if (PathEndTime > StartTime)
 						{
-							if (PathEndTime > StartTime)
-							{
-								PathEndTime = StartTime;
-							}
+							PathEndTime = StartTime;
 						}
 					}
-
-					Cast<UMovieScene3DPathTrack>(Track)->AddConstraint( KeyTime, PathEndTime, SplineId );
 				}
+
+				Cast<UMovieScene3DPathTrack>(Track)->AddConstraint( KeyTime, PathEndTime, SplineId );
 			}
 		}
 	}
diff --git a/Engine/Source/Editor/MovieSceneTools/Private/TrackEditors/PathTrackEditor.h b/Engine/Source/Editor/MovieSceneTools/Private/TrackEditors/PathTrackEditor.h
--- a/Engine/Source/Editor/MovieSceneTools/Private/TrackEditors/PathTrackEditor.h
+++ b/Engine/Source/Editor/MovieSceneTools/Private/TrackEditors/PathTrackEditor.h
@@ -4,19 +4,25 @@
 
 #include "MovieScene3DPathSection.h"
 #include "MovieScene3DPathTrack.h"
+
+
 /**
  * Tools for animating a transform with a path curve
  */
-class F3DPathTrackEditor : public FMovieSceneTrackEditor
+class F3DPathTrackEditor
+	: public FMovieSceneTrackEditor
 {
 public:
+
 	/**
 	 * Constructor
 	 *
-	 * @param InSequencer	The sequencer instance to be used by this tool
+	 * @param InSequencer The sequencer instance to be used by this tool
 	 */
 	F3DPathTrackEditor( TSharedRef<ISequencer> InSequencer );
-	~F3DPathTrackEditor();
+
+	/** Virtual destructor. */
+	virtual ~F3DPathTrackEditor();
 
 	/**
 	 * Creates an instance of this class.  Called by a sequencer 
@@ -24,13 +30,9 @@ public:
 	 * @param OwningSequencer The sequencer instance to be used by this tool
 	 * @return The new instance of this class
 	 */
-	static TSharedRef<FMovieSceneTrackEditor> CreateTrackEditor( TSharedRef<ISequencer> OwningSequencer );
+	static TSharedRef<ISequencerTrackEditor> CreateTrackEditor( TSharedRef<ISequencer> OwningSequencer );
 
-	/** FMovieSceneTrackEditor Interface */
-	virtual bool SupportsType( TSubclassOf<UMovieSceneTrack> Type ) const override;
-	virtual TSharedRef<ISequencerSection> MakeSectionInterface( UMovieSceneSection& SectionObject, UMovieSceneTrack* Track ) override;
-	virtual void AddKey( const FGuid& ObjectGuid, UObject* AdditionalAsset = NULL ) override;
-	virtual void BuildObjectBindingTrackMenu(FMenuBuilder& MenuBuilder, const FGuid& ObjectBinding, const UClass* ObjectClass) override;
+public:
 
 	/** Add path */
 	void AddPath(FGuid ObjectGuid, UObject* AdditionalAsset);
@@ -38,6 +40,15 @@ public:
 	/** Set path */
 	void SetPath(UMovieSceneSection* Section, AActor* ActorWithSplineComponent);
 
+public:
+
+	// FMovieSceneTrackEditor interface
+
+	virtual bool SupportsType( TSubclassOf<UMovieSceneTrack> Type ) const override;
+	virtual TSharedRef<ISequencerSection> MakeSectionInterface( UMovieSceneSection& SectionObject, UMovieSceneTrack* Track ) override;
+	virtual void AddKey( const FGuid& ObjectGuid, UObject* AdditionalAsset = NULL ) override;
+	virtual void BuildObjectBindingTrackMenu(FMenuBuilder& MenuBuilder, const FGuid& ObjectBinding, const UClass* ObjectClass) override;
+
 private:
 
 	/** Add path sub menu */
diff --git a/Engine/Source/Editor/MovieSceneTools/Private/TrackEditors/PropertyTrackEditors/BoolPropertyTrackEditor.cpp b/Engine/Source/Editor/MovieSceneTools/Private/TrackEditors/PropertyTrackEditors/BoolPropertyTrackEditor.cpp
--- a/Engine/Source/Editor/MovieSceneTools/Private/TrackEditors/PropertyTrackEditors/BoolPropertyTrackEditor.cpp
+++ b/Engine/Source/Editor/MovieSceneTools/Private/TrackEditors/PropertyTrackEditors/BoolPropertyTrackEditor.cpp
@@ -5,26 +5,32 @@
 #include "BoolPropertyTrackEditor.h"
 #include "BoolPropertySection.h"
 
-TSharedRef<FMovieSceneTrackEditor> FBoolPropertyTrackEditor::CreateTrackEditor( TSharedRef<ISequencer> OwningSequencer )
+
+TSharedRef<ISequencerTrackEditor> FBoolPropertyTrackEditor::CreateTrackEditor( TSharedRef<ISequencer> OwningSequencer )
 {
 	return MakeShareable(new FBoolPropertyTrackEditor(OwningSequencer));
 }
 
+
 TSharedRef<ISequencerSection> FBoolPropertyTrackEditor::MakeSectionInterface( UMovieSceneSection& SectionObject, UMovieSceneTrack* Track )
 {
 	return MakeShareable(new FBoolPropertySection( SectionObject, Track->GetTrackName(), GetSequencer().Get() ));
 }
 
+
 bool FBoolPropertyTrackEditor::TryGenerateKeyFromPropertyChanged( const UMovieSceneTrack* InTrack, const FPropertyChangedParams& PropertyChangedParams, bool& OutKey )
 {
 	const UBoolProperty* BoolProperty = Cast<const UBoolProperty>(PropertyChangedParams.PropertyPath.Last());
+
 	if (BoolProperty && PropertyChangedParams.ObjectsThatChanged.Num() != 0)
 	{
 		void* CurrentObject = PropertyChangedParams.ObjectsThatChanged[0];
+
 		for (int32 i = 0; i < PropertyChangedParams.PropertyPath.Num(); ++i)
 		{
 			CurrentObject = PropertyChangedParams.PropertyPath[i]->ContainerPtrToValuePtr<void>(CurrentObject, 0);
 		}
+
 		OutKey = BoolProperty->GetPropertyValue(CurrentObject);
 	
 		if (InTrack)
@@ -39,4 +45,4 @@ bool FBoolPropertyTrackEditor::TryGenerateKeyFromPropertyChanged( const UMovieSc
 	}
 
 	return false;
-}
\ No newline at end of file
+}
diff --git a/Engine/Source/Editor/MovieSceneTools/Private/TrackEditors/PropertyTrackEditors/BoolPropertyTrackEditor.h b/Engine/Source/Editor/MovieSceneTools/Private/TrackEditors/PropertyTrackEditors/BoolPropertyTrackEditor.h
--- a/Engine/Source/Editor/MovieSceneTools/Private/TrackEditors/PropertyTrackEditors/BoolPropertyTrackEditor.h
+++ b/Engine/Source/Editor/MovieSceneTools/Private/TrackEditors/PropertyTrackEditors/BoolPropertyTrackEditor.h
@@ -4,20 +4,22 @@
 
 #include "MovieSceneBoolTrack.h"
 
+
 /**
 * A property track editor for bools.
 */
-class FBoolPropertyTrackEditor : public FPropertyTrackEditor<UMovieSceneBoolTrack, bool>
+class FBoolPropertyTrackEditor
+	: public FPropertyTrackEditor<UMovieSceneBoolTrack, bool>
 {
 public:
 	/**
 	 * Constructor
 	 *
-	 * @param InSequencer	The sequencer instance to be used by this tool
+	 * @param InSequencer The sequencer instance to be used by this tool
 	 */
 	FBoolPropertyTrackEditor( TSharedRef<ISequencer> InSequencer )
 		: FPropertyTrackEditor( InSequencer, NAME_BoolProperty )
-	{}
+	{ }
 
 	/**
 	 * Creates an instance of this class.  Called by a sequencer 
@@ -25,14 +27,16 @@ public:
 	 * @param OwningSequencer The sequencer instance to be used by this tool
 	 * @return The new instance of this class
 	 */
-	static TSharedRef<FMovieSceneTrackEditor> CreateTrackEditor( TSharedRef<ISequencer> OwningSequencer );
+	static TSharedRef<ISequencerTrackEditor> CreateTrackEditor( TSharedRef<ISequencer> OwningSequencer );
+
+public:
 
-	/** FMovieSceneTrackEditor Interface */
+	// FMovieSceneTrackEditor interface
 	virtual TSharedRef<ISequencerSection> MakeSectionInterface( UMovieSceneSection& SectionObject, UMovieSceneTrack* Track ) override;
 
 protected:
-	/** FPropertyTrackEditor Interface */
-	virtual bool TryGenerateKeyFromPropertyChanged( const UMovieSceneTrack* InTrack, const FPropertyChangedParams& PropertyChangedParams, bool& OutKey ) override;
-};
 
+	// FPropertyTrackEditor interface
 
+	virtual bool TryGenerateKeyFromPropertyChanged( const UMovieSceneTrack* InTrack, const FPropertyChangedParams& PropertyChangedParams, bool& OutKey ) override;
+};
diff --git a/Engine/Source/Editor/MovieSceneTools/Private/TrackEditors/PropertyTrackEditors/BytePropertyTrackEditor.cpp b/Engine/Source/Editor/MovieSceneTools/Private/TrackEditors/PropertyTrackEditors/BytePropertyTrackEditor.cpp
--- a/Engine/Source/Editor/MovieSceneTools/Private/TrackEditors/PropertyTrackEditors/BytePropertyTrackEditor.cpp
+++ b/Engine/Source/Editor/MovieSceneTools/Private/TrackEditors/PropertyTrackEditors/BytePropertyTrackEditor.cpp
@@ -6,22 +6,25 @@
 #include "MovieSceneByteTrack.h"
 #include "MovieSceneSequence.h"
 
-TSharedRef<FMovieSceneTrackEditor> FBytePropertyTrackEditor::CreateTrackEditor( TSharedRef<ISequencer> OwningSequencer )
+
+TSharedRef<ISequencerTrackEditor> FBytePropertyTrackEditor::CreateTrackEditor( TSharedRef<ISequencer> OwningSequencer )
 {
 	return MakeShareable(new FBytePropertyTrackEditor(OwningSequencer));
 }
 
+
 TSharedRef<ISequencerSection> FBytePropertyTrackEditor::MakeSectionInterface( UMovieSceneSection& SectionObject, UMovieSceneTrack* Track )
 {
 	return MakeShareable( new FBytePropertySection( SectionObject, Track->GetTrackName(), Cast<UMovieSceneByteTrack>( SectionObject.GetOuter() )->GetEnum() ) );
 }
 
+
 UEnum* GetEnumForByteTrack(TSharedPtr<ISequencer> Sequencer, const FGuid& OwnerObjectHandle, FName PropertyName, UMovieSceneByteTrack* ByteTrack)
 {
 	
 	UObject* RuntimeObject = Sequencer->GetFocusedMovieSceneSequence()->FindObject(OwnerObjectHandle);
-
 	TSet<UEnum*> PropertyEnums;
+
 	if (RuntimeObject != nullptr)
 	{
 		UProperty* Property = RuntimeObject->GetClass()->FindPropertyByName(PropertyName);
@@ -36,6 +39,7 @@ UEnum* GetEnumForByteTrack(TSharedPtr<ISequencer> Sequencer, const FGuid& OwnerO
 	}
 
 	UEnum* TrackEnum;
+
 	if (PropertyEnums.Num() == 1)
 	{
 		TrackEnum = PropertyEnums.Array()[0];
@@ -44,15 +48,17 @@ UEnum* GetEnumForByteTrack(TSharedPtr<ISequencer> Sequencer, const FGuid& OwnerO
 	{
 		TrackEnum = nullptr;
 	}
+
 	return TrackEnum;
 }
 
+
 UMovieSceneTrack* FBytePropertyTrackEditor::AddTrack(UMovieScene* FocusedMovieScene, const FGuid& ObjectHandle, TSubclassOf<class UMovieSceneTrack> TrackClass, FName UniqueTypeName)
 {
 	UMovieSceneTrack* NewTrack = FPropertyTrackEditor::AddTrack(FocusedMovieScene, ObjectHandle, TrackClass, UniqueTypeName);
-	
 	UMovieSceneByteTrack* ByteTrack = Cast<UMovieSceneByteTrack>(NewTrack);
 	UEnum* TrackEnum = GetEnumForByteTrack(GetSequencer(), ObjectHandle, UniqueTypeName, ByteTrack);
+
 	if (TrackEnum != nullptr)
 	{
 		ByteTrack->SetEnum(TrackEnum);
@@ -61,19 +67,24 @@ UMovieSceneTrack* FBytePropertyTrackEditor::AddTrack(UMovieScene* FocusedMovieSc
 	return NewTrack;
 }
 
+
 bool FBytePropertyTrackEditor::TryGenerateKeyFromPropertyChanged( const UMovieSceneTrack* InTrack, const FPropertyChangedParams& PropertyChangedParams, uint8& OutKey )
 {
 	OutKey = *PropertyChangedParams.GetPropertyValue<uint8>();
 
-	if (InTrack)
+	if (InTrack == nullptr)
 	{
-		const UMovieSceneByteTrack* ByteTrack = CastChecked<const UMovieSceneByteTrack>( InTrack );
-		if (ByteTrack)
-		{
-			float KeyTime =	GetTimeForKey(GetMovieSceneSequence());
-			return ByteTrack->CanKeyTrack(KeyTime, OutKey, PropertyChangedParams.KeyParams);
-		}
+		return false;
 	}
 
-	return false;
-}
\ No newline at end of file
+	const UMovieSceneByteTrack* ByteTrack = CastChecked<const UMovieSceneByteTrack>( InTrack );
+
+	if (ByteTrack == nullptr)
+	{
+		return false;
+	}
+
+	float KeyTime =	GetTimeForKey(GetMovieSceneSequence());
+
+	return ByteTrack->CanKeyTrack(KeyTime, OutKey, PropertyChangedParams.KeyParams);
+}
diff --git a/Engine/Source/Editor/MovieSceneTools/Private/TrackEditors/PropertyTrackEditors/BytePropertyTrackEditor.h b/Engine/Source/Editor/MovieSceneTools/Private/TrackEditors/PropertyTrackEditors/BytePropertyTrackEditor.h
--- a/Engine/Source/Editor/MovieSceneTools/Private/TrackEditors/PropertyTrackEditors/BytePropertyTrackEditor.h
+++ b/Engine/Source/Editor/MovieSceneTools/Private/TrackEditors/PropertyTrackEditors/BytePropertyTrackEditor.h
@@ -4,20 +4,23 @@
 
 #include "MovieSceneByteTrack.h"
 
+
 /**
 * A property track editor for byte and enums.
 */
-class FBytePropertyTrackEditor : public FPropertyTrackEditor<UMovieSceneByteTrack, uint8>
+class FBytePropertyTrackEditor
+	: public FPropertyTrackEditor<UMovieSceneByteTrack, uint8>
 {
 public:
+
 	/**
-	 * Constructor
+	 * Constructor.
 	 *
-	 * @param InSequencer	The sequencer instance to be used by this tool
+	 * @param InSequencer The sequencer instance to be used by this tool
 	 */
 	FBytePropertyTrackEditor( TSharedRef<ISequencer> InSequencer)
 		: FPropertyTrackEditor( InSequencer, NAME_ByteProperty )
-	{}
+	{ }
 
 	/**
 	 * Creates an instance of this class.  Called by a sequencer 
@@ -25,15 +28,18 @@ public:
 	 * @param OwningSequencer The sequencer instance to be used by this tool
 	 * @return The new instance of this class
 	 */
-	static TSharedRef<FMovieSceneTrackEditor> CreateTrackEditor( TSharedRef<ISequencer> OwningSequencer );
+	static TSharedRef<ISequencerTrackEditor> CreateTrackEditor( TSharedRef<ISequencer> OwningSequencer );
+
+public:
+
+	// FMovieSceneTrackEditor interface
 
-	/** FMovieSceneTrackEditor Interface */
 	virtual TSharedRef<ISequencerSection> MakeSectionInterface( UMovieSceneSection& SectionObject, UMovieSceneTrack* Track ) override;
 	virtual UMovieSceneTrack* AddTrack( UMovieScene* FocusedMovieScene, const FGuid& ObjectHandle, TSubclassOf<class UMovieSceneTrack> TrackClass, FName UniqueTypeName ) override;
 
 protected:
-	/** FPropertyTrackEditor Interface */
-	virtual bool TryGenerateKeyFromPropertyChanged( const UMovieSceneTrack* InTrack, const FPropertyChangedParams& PropertyChangedParams, uint8& OutKey ) override;
-};
 
+	// FPropertyTrackEditor interface
 
+	virtual bool TryGenerateKeyFromPropertyChanged( const UMovieSceneTrack* InTrack, const FPropertyChangedParams& PropertyChangedParams, uint8& OutKey ) override;
+};
diff --git a/Engine/Source/Editor/MovieSceneTools/Private/TrackEditors/PropertyTrackEditors/ColorPropertyTrackEditor.cpp b/Engine/Source/Editor/MovieSceneTools/Private/TrackEditors/PropertyTrackEditors/ColorPropertyTrackEditor.cpp
--- a/Engine/Source/Editor/MovieSceneTools/Private/TrackEditors/PropertyTrackEditors/ColorPropertyTrackEditor.cpp
+++ b/Engine/Source/Editor/MovieSceneTools/Private/TrackEditors/PropertyTrackEditors/ColorPropertyTrackEditor.cpp
@@ -5,20 +5,24 @@
 #include "ColorPropertyTrackEditor.h"
 #include "ColorPropertySection.h"
 
-TSharedRef<FMovieSceneTrackEditor> FColorPropertyTrackEditor::CreateTrackEditor( TSharedRef<ISequencer> InSequencer )
+
+TSharedRef<ISequencerTrackEditor> FColorPropertyTrackEditor::CreateTrackEditor( TSharedRef<ISequencer> InSequencer )
 {
 	return MakeShareable( new FColorPropertyTrackEditor( InSequencer ) );
 }
 
+
 TSharedRef<ISequencerSection> FColorPropertyTrackEditor::MakeSectionInterface( UMovieSceneSection& SectionObject, UMovieSceneTrack* Track )
 {
 	return MakeShareable(new FColorPropertySection( SectionObject, Track->GetTrackName(), GetSequencer().Get(), Track ));
 }
 
+
 bool FColorPropertyTrackEditor::TryGenerateKeyFromPropertyChanged(  const UMovieSceneTrack* InTrack, const FPropertyChangedParams& PropertyChangedParams, FColorKey& OutKey )
 {
 	bool bIsFColor = false, bIsFLinearColor = false, bIsSlateColor = false;
 	const UStructProperty* StructProp = Cast<const UStructProperty>(PropertyChangedParams.PropertyPath.Last());
+
 	if (StructProp && StructProp->Struct)
 	{
 		FName StructName = StructProp->Struct->GetFName();
@@ -36,6 +40,7 @@ bool FColorPropertyTrackEditor::TryGenerateKeyFromPropertyChanged(  const UMovie
 	FName PropertyName = PropertyChangedParams.PropertyPath.Last()->GetFName();
 
 	FLinearColor ColorValue;
+
 	if (bIsFColor)
 	{
 		ColorValue = PropertyChangedParams.GetPropertyValue<FColor>()->ReinterpretAsLinear();
diff --git a/Engine/Source/Editor/MovieSceneTools/Private/TrackEditors/PropertyTrackEditors/ColorPropertyTrackEditor.h b/Engine/Source/Editor/MovieSceneTools/Private/TrackEditors/PropertyTrackEditors/ColorPropertyTrackEditor.h
--- a/Engine/Source/Editor/MovieSceneTools/Private/TrackEditors/PropertyTrackEditors/ColorPropertyTrackEditor.h
+++ b/Engine/Source/Editor/MovieSceneTools/Private/TrackEditors/PropertyTrackEditors/ColorPropertyTrackEditor.h
@@ -4,20 +4,23 @@
 
 #include "MovieSceneColorTrack.h"
 
+
 /**
 * A property track editor for colors.
 */
-class FColorPropertyTrackEditor : public FPropertyTrackEditor<UMovieSceneColorTrack, FColorKey>
+class FColorPropertyTrackEditor
+	: public FPropertyTrackEditor<UMovieSceneColorTrack, FColorKey>
 {
 public:
+
 	/**
 	 * Constructor
 	 *
-	 * @param InSequencer	The sequencer instance to be used by this tool
+	 * @param InSequencer The sequencer instance to be used by this tool
 	 */
 	FColorPropertyTrackEditor( TSharedRef<ISequencer> InSequencer )
 		: FPropertyTrackEditor<UMovieSceneColorTrack, FColorKey>( InSequencer, NAME_Color, NAME_LinearColor, "SlateColor" )
-	{}
+	{ }
 
 	/**
 	 * Creates an instance of this class.  Called by a sequencer 
@@ -25,14 +28,16 @@ public:
 	 * @param OwningSequencer The sequencer instance to be used by this tool
 	 * @return The new instance of this class
 	 */
-	static TSharedRef<FMovieSceneTrackEditor> CreateTrackEditor( TSharedRef<ISequencer> OwningSequencer );
+	static TSharedRef<ISequencerTrackEditor> CreateTrackEditor( TSharedRef<ISequencer> OwningSequencer );
+
+public:
+
+	// FMovieSceneTrackEditor interface
 
-	/** FMovieSceneTrackEditor Interface */
 	virtual TSharedRef<ISequencerSection> MakeSectionInterface( UMovieSceneSection& SectionObject, UMovieSceneTrack* Track ) override;
 
 protected:
-	/** FPropertyTrackEditor Interface */
+
+	// FPropertyTrackEditor interface
 	virtual bool TryGenerateKeyFromPropertyChanged( const UMovieSceneTrack* InTrack, const FPropertyChangedParams& PropertyChangedParams, FColorKey& OutKey ) override;
 };
-
-
diff --git a/Engine/Source/Editor/MovieSceneTools/Private/TrackEditors/PropertyTrackEditors/FloatPropertyTrackEditor.cpp b/Engine/Source/Editor/MovieSceneTools/Private/TrackEditors/PropertyTrackEditors/FloatPropertyTrackEditor.cpp
--- a/Engine/Source/Editor/MovieSceneTools/Private/TrackEditors/PropertyTrackEditors/FloatPropertyTrackEditor.cpp
+++ b/Engine/Source/Editor/MovieSceneTools/Private/TrackEditors/PropertyTrackEditors/FloatPropertyTrackEditor.cpp
@@ -4,16 +4,19 @@
 #include "FloatPropertyTrackEditor.h"
 #include "FloatPropertySection.h"
 
-TSharedRef<FMovieSceneTrackEditor> FFloatPropertyTrackEditor::CreateTrackEditor( TSharedRef<ISequencer> OwningSequencer )
+
+TSharedRef<ISequencerTrackEditor> FFloatPropertyTrackEditor::CreateTrackEditor( TSharedRef<ISequencer> OwningSequencer )
 {
 	return MakeShareable(new FFloatPropertyTrackEditor(OwningSequencer));
 }
 
+
 TSharedRef<ISequencerSection> FFloatPropertyTrackEditor::MakeSectionInterface( UMovieSceneSection& SectionObject, UMovieSceneTrack* Track )
 {
 	return MakeShareable(new FFloatPropertySection( SectionObject, Track->GetTrackName() ));
 }
 
+
 bool FFloatPropertyTrackEditor::TryGenerateKeyFromPropertyChanged( const UMovieSceneTrack* InTrack, const FPropertyChangedParams& PropertyChangedParams, float& OutKey )
 {
 	OutKey = *PropertyChangedParams.GetPropertyValue<float>();
@@ -21,6 +24,7 @@ bool FFloatPropertyTrackEditor::TryGenerateKeyFromPropertyChanged( const UMovieS
 	if (InTrack)
 	{
 		const UMovieSceneFloatTrack* FloatTrack = CastChecked<const UMovieSceneFloatTrack>( InTrack );
+
 		if (FloatTrack)
 		{
 			float KeyTime =	GetTimeForKey(GetMovieSceneSequence());
@@ -29,4 +33,4 @@ bool FFloatPropertyTrackEditor::TryGenerateKeyFromPropertyChanged( const UMovieS
 	}
 
 	return false;
-}
\ No newline at end of file
+}
diff --git a/Engine/Source/Editor/MovieSceneTools/Private/TrackEditors/PropertyTrackEditors/FloatPropertyTrackEditor.h b/Engine/Source/Editor/MovieSceneTools/Private/TrackEditors/PropertyTrackEditors/FloatPropertyTrackEditor.h
--- a/Engine/Source/Editor/MovieSceneTools/Private/TrackEditors/PropertyTrackEditors/FloatPropertyTrackEditor.h
+++ b/Engine/Source/Editor/MovieSceneTools/Private/TrackEditors/PropertyTrackEditors/FloatPropertyTrackEditor.h
@@ -7,17 +7,19 @@
 /**
  * A property track editor for floats.
  */
-class FFloatPropertyTrackEditor : public FPropertyTrackEditor<UMovieSceneFloatTrack, float>
+class FFloatPropertyTrackEditor
+	: public FPropertyTrackEditor<UMovieSceneFloatTrack, float>
 {
 public:
+
 	/**
 	 * Constructor
 	 *
-	 * @param InSequencer	The sequencer instance to be used by this tool
+	 * @param InSequencer The sequencer instance to be used by this tool
 	 */
 	FFloatPropertyTrackEditor( TSharedRef<ISequencer> InSequencer )
 		: FPropertyTrackEditor( InSequencer, NAME_FloatProperty )
-	{}
+	{ }
 
 	/**
 	 * Creates an instance of this class.  Called by a sequencer 
@@ -25,14 +27,15 @@ public:
 	 * @param OwningSequencer The sequencer instance to be used by this tool
 	 * @return The new instance of this class
 	 */
-	static TSharedRef<FMovieSceneTrackEditor> CreateTrackEditor( TSharedRef<ISequencer> OwningSequencer );
+	static TSharedRef<ISequencerTrackEditor> CreateTrackEditor( TSharedRef<ISequencer> OwningSequencer );
+
+public:
 
-	/** FMovieSceneTrackEditor Interface */
+	// FMovieSceneTrackEditor interface
 	virtual TSharedRef<ISequencerSection> MakeSectionInterface( UMovieSceneSection& SectionObject, UMovieSceneTrack* Track ) override;
 
 protected:
-	/** FPropertyTrackEditor Interface */
+
+	// FPropertyTrackEditor interface
 	virtual bool TryGenerateKeyFromPropertyChanged( const UMovieSceneTrack* InTrack, const FPropertyChangedParams& PropertyChangedParams, float& OutKey ) override;
 };
-
-
diff --git a/Engine/Source/Editor/MovieSceneTools/Private/TrackEditors/PropertyTrackEditors/VectorPropertyTrackEditor.cpp b/Engine/Source/Editor/MovieSceneTools/Private/TrackEditors/PropertyTrackEditors/VectorPropertyTrackEditor.cpp
--- a/Engine/Source/Editor/MovieSceneTools/Private/TrackEditors/PropertyTrackEditors/VectorPropertyTrackEditor.cpp
+++ b/Engine/Source/Editor/MovieSceneTools/Private/TrackEditors/PropertyTrackEditors/VectorPropertyTrackEditor.cpp
@@ -5,20 +5,24 @@
 #include "VectorPropertyTrackEditor.h"
 #include "VectorPropertySection.h"
 
-TSharedRef<FMovieSceneTrackEditor> FVectorPropertyTrackEditor::CreateTrackEditor( TSharedRef<ISequencer> InSequencer )
+
+TSharedRef<ISequencerTrackEditor> FVectorPropertyTrackEditor::CreateTrackEditor( TSharedRef<ISequencer> InSequencer )
 {
 	return MakeShareable( new FVectorPropertyTrackEditor( InSequencer ) );
 }
 
+
 TSharedRef<ISequencerSection> FVectorPropertyTrackEditor::MakeSectionInterface( UMovieSceneSection& SectionObject, UMovieSceneTrack* Track )
 {
 	return MakeShareable(new FVectorPropertySection( SectionObject, Track->GetTrackName() ));
 }
 
+
 bool FVectorPropertyTrackEditor::TryGenerateKeyFromPropertyChanged( const UMovieSceneTrack* InTrack, const FPropertyChangedParams& PropertyChangedParams, FVectorKey& OutKey )
 {
 	bool bIsVector2D = false, bIsVector = false, bIsVector4 = false;
 	const UStructProperty* StructProp = Cast<const UStructProperty>( PropertyChangedParams.PropertyPath.Last() );
+
 	if ( StructProp && StructProp->Struct )
 	{
 		FName StructName = StructProp->Struct->GetFName();
@@ -27,6 +31,7 @@ bool FVectorPropertyTrackEditor::TryGenerateKeyFromPropertyChanged( const UMovie
 		bIsVector = StructName == NAME_Vector;
 		bIsVector4 = StructName == NAME_Vector4;
 	}
+
 	if ( !bIsVector2D && !bIsVector && !bIsVector4 )
 	{
 		return false;
diff --git a/Engine/Source/Editor/MovieSceneTools/Private/TrackEditors/PropertyTrackEditors/VectorPropertyTrackEditor.h b/Engine/Source/Editor/MovieSceneTools/Private/TrackEditors/PropertyTrackEditors/VectorPropertyTrackEditor.h
--- a/Engine/Source/Editor/MovieSceneTools/Private/TrackEditors/PropertyTrackEditors/VectorPropertyTrackEditor.h
+++ b/Engine/Source/Editor/MovieSceneTools/Private/TrackEditors/PropertyTrackEditors/VectorPropertyTrackEditor.h
@@ -4,20 +4,23 @@
 
 #include "MovieSceneVectorTrack.h"
 
+
 /**
  * A property track editor for vectors.
  */
-class FVectorPropertyTrackEditor : public FPropertyTrackEditor<UMovieSceneVectorTrack, FVectorKey>
+class FVectorPropertyTrackEditor
+	: public FPropertyTrackEditor<UMovieSceneVectorTrack, FVectorKey>
 {
 public:
+
 	/**
 	 * Constructor
 	 *
-	 * @param InSequencer	The sequencer instance to be used by this tool
+	 * @param InSequencer The sequencer instance to be used by this tool
 	 */
 	FVectorPropertyTrackEditor( TSharedRef<ISequencer> InSequencer )
 		: FPropertyTrackEditor<UMovieSceneVectorTrack, FVectorKey>( InSequencer, NAME_Vector, NAME_Vector4, NAME_Vector2D )
-	{}
+	{ }
 
 	/**
 	 * Creates an instance of this class.  Called by a sequencer 
@@ -25,14 +28,15 @@ public:
 	 * @param OwningSequencer The sequencer instance to be used by this tool
 	 * @return The new instance of this class
 	 */
-	static TSharedRef<FMovieSceneTrackEditor> CreateTrackEditor( TSharedRef<ISequencer> OwningSequencer );
+	static TSharedRef<ISequencerTrackEditor> CreateTrackEditor( TSharedRef<ISequencer> OwningSequencer );
+
+public:
 
-	/** FMovieSceneTrackEditor Interface */
+	// FMovieSceneTrackEditor interface
 	virtual TSharedRef<ISequencerSection> MakeSectionInterface( UMovieSceneSection& SectionObject, UMovieSceneTrack* Track ) override;
 
 protected:
-	/** FPropertyTrackEditor Interface */
+
+	// FPropertyTrackEditor interface
 	virtual bool TryGenerateKeyFromPropertyChanged( const UMovieSceneTrack* InTrack, const FPropertyChangedParams& PropertyChangedParams, FVectorKey& OutKey ) override;
 };
-
-
diff --git a/Engine/Source/Editor/MovieSceneTools/Private/TrackEditors/PropertyTrackEditors/VisibilityPropertyTrackEditor.cpp b/Engine/Source/Editor/MovieSceneTools/Private/TrackEditors/PropertyTrackEditors/VisibilityPropertyTrackEditor.cpp
--- a/Engine/Source/Editor/MovieSceneTools/Private/TrackEditors/PropertyTrackEditors/VisibilityPropertyTrackEditor.cpp
+++ b/Engine/Source/Editor/MovieSceneTools/Private/TrackEditors/PropertyTrackEditors/VisibilityPropertyTrackEditor.cpp
@@ -5,19 +5,23 @@
 #include "VisibilityPropertyTrackEditor.h"
 #include "VisibilityPropertySection.h"
 
-TSharedRef<FMovieSceneTrackEditor> FVisibilityPropertyTrackEditor::CreateTrackEditor( TSharedRef<ISequencer> OwningSequencer )
+
+TSharedRef<ISequencerTrackEditor> FVisibilityPropertyTrackEditor::CreateTrackEditor( TSharedRef<ISequencer> OwningSequencer )
 {
 	return MakeShareable(new FVisibilityPropertyTrackEditor(OwningSequencer));
 }
 
+
 TSharedRef<ISequencerSection> FVisibilityPropertyTrackEditor::MakeSectionInterface( UMovieSceneSection& SectionObject, UMovieSceneTrack* Track )
 {
 	return MakeShareable(new FVisibilityPropertySection( SectionObject, Track->GetTrackName(), GetSequencer().Get() ));
 }
 
+
 bool FVisibilityPropertyTrackEditor::TryGenerateKeyFromPropertyChanged( const UMovieSceneTrack* InTrack, const FPropertyChangedParams& PropertyChangedParams, bool& OutKey )
 {
 	const UBoolProperty* BoolProperty = Cast<const UBoolProperty>(PropertyChangedParams.PropertyPath.Last());
+
 	if (BoolProperty && PropertyChangedParams.ObjectsThatChanged.Num() != 0)
 	{
 		void* CurrentObject = PropertyChangedParams.ObjectsThatChanged[0];
@@ -39,4 +43,4 @@ bool FVisibilityPropertyTrackEditor::TryGenerateKeyFromPropertyChanged( const UM
 	}
 
 	return false;
-}
\ No newline at end of file
+}
diff --git a/Engine/Source/Editor/MovieSceneTools/Private/TrackEditors/PropertyTrackEditors/VisibilityPropertyTrackEditor.h b/Engine/Source/Editor/MovieSceneTools/Private/TrackEditors/PropertyTrackEditors/VisibilityPropertyTrackEditor.h
--- a/Engine/Source/Editor/MovieSceneTools/Private/TrackEditors/PropertyTrackEditors/VisibilityPropertyTrackEditor.h
+++ b/Engine/Source/Editor/MovieSceneTools/Private/TrackEditors/PropertyTrackEditors/VisibilityPropertyTrackEditor.h
@@ -4,12 +4,17 @@
 
 #include "MovieSceneVisibilityTrack.h"
 
+
 /**
- * A property track for bool properties which have been set to use a UMovieSceneVisibilityTrack through metadata.
+ * A property track for bool properties which have been set
+ * to use a UMovieSceneVisibilityTrack through metadata.
  */
-class FVisibilityPropertyTrackEditor : public FPropertyTrackEditor<UMovieSceneVisibilityTrack, bool>
+class FVisibilityPropertyTrackEditor
+	: public FPropertyTrackEditor<UMovieSceneVisibilityTrack, bool>
 {
 public:
+
+	/** Constructor. */
 	FVisibilityPropertyTrackEditor( TSharedRef<ISequencer> InSequencer)
 		// Don't supply any property type names to watch since the FBoolPropertyTrackEditor is already watching for bool property changes.
 		: FPropertyTrackEditor<UMovieSceneVisibilityTrack, bool>( InSequencer )
@@ -21,14 +26,15 @@ public:
 	 * @param OwningSequencer The sequencer instance to be used by this tool
 	 * @return The new instance of this class
 	 */
-	static TSharedRef<FMovieSceneTrackEditor> CreateTrackEditor( TSharedRef<ISequencer> OwningSequencer );
+	static TSharedRef<ISequencerTrackEditor> CreateTrackEditor( TSharedRef<ISequencer> OwningSequencer );
+
+public:
 
-	/** FMovieSceneTrackEditor Interface */
+	// FMovieSceneTrackEditor interface
 	virtual TSharedRef<ISequencerSection> MakeSectionInterface( UMovieSceneSection& SectionObject, UMovieSceneTrack* Track ) override;
 
 protected:
-	/** FPropertyTrackEditor Interface */
+
+	// FPropertyTrackEditor interface
 	virtual bool TryGenerateKeyFromPropertyChanged( const UMovieSceneTrack* InTrack, const FPropertyChangedParams& PropertyChangedParams, bool& OutKey ) override;
 };
-
-
diff --git a/Engine/Source/Editor/MovieSceneTools/Private/TrackEditors/ShotTrackEditor.cpp b/Engine/Source/Editor/MovieSceneTools/Private/TrackEditors/ShotTrackEditor.cpp
--- a/Engine/Source/Editor/MovieSceneTools/Private/TrackEditors/ShotTrackEditor.cpp
+++ b/Engine/Source/Editor/MovieSceneTools/Private/TrackEditors/ShotTrackEditor.cpp
@@ -515,7 +515,7 @@ FShotTrackEditor::~FShotTrackEditor()
 {
 }
 
-TSharedRef<FMovieSceneTrackEditor> FShotTrackEditor::CreateTrackEditor( TSharedRef<ISequencer> InSequencer )
+TSharedRef<ISequencerTrackEditor> FShotTrackEditor::CreateTrackEditor( TSharedRef<ISequencer> InSequencer )
 {
 	return MakeShareable( new FShotTrackEditor( InSequencer ) );
 }
diff --git a/Engine/Source/Editor/MovieSceneTools/Private/TrackEditors/ShotTrackEditor.h b/Engine/Source/Editor/MovieSceneTools/Private/TrackEditors/ShotTrackEditor.h
--- a/Engine/Source/Editor/MovieSceneTools/Private/TrackEditors/ShotTrackEditor.h
+++ b/Engine/Source/Editor/MovieSceneTools/Private/TrackEditors/ShotTrackEditor.h
@@ -23,7 +23,7 @@ public:
 	 * @param OwningSequencer The sequencer instance to be used by this tool
 	 * @return The new instance of this class
 	 */
-	static TSharedRef<FMovieSceneTrackEditor> CreateTrackEditor( TSharedRef<ISequencer> OwningSequencer );
+	static TSharedRef<ISequencerTrackEditor> CreateTrackEditor( TSharedRef<ISequencer> OwningSequencer );
 
 	/** FMovieSceneTrackEditor Interface */
 	virtual bool SupportsType( TSubclassOf<UMovieSceneTrack> Type ) const override;
diff --git a/Engine/Source/Editor/MovieSceneTools/Private/TrackEditors/SkeletalAnimationTrackEditor.cpp b/Engine/Source/Editor/MovieSceneTools/Private/TrackEditors/SkeletalAnimationTrackEditor.cpp
--- a/Engine/Source/Editor/MovieSceneTools/Private/TrackEditors/SkeletalAnimationTrackEditor.cpp
+++ b/Engine/Source/Editor/MovieSceneTools/Private/TrackEditors/SkeletalAnimationTrackEditor.cpp
@@ -21,7 +21,6 @@
 #include "Animation/SkeletalMeshActor.h"
 
 
-
 namespace SkeletalAnimationEditorConstants
 {
 	// @todo Sequencer Allow this to be customizable
@@ -29,35 +28,40 @@ namespace SkeletalAnimationEditorConstants
 }
 
 
-
 FSkeletalAnimationSection::FSkeletalAnimationSection( UMovieSceneSection& InSection )
 	: Section( InSection )
 {
 }
 
+
 FSkeletalAnimationSection::~FSkeletalAnimationSection()
 {
 }
 
+
 UMovieSceneSection* FSkeletalAnimationSection::GetSectionObject()
 { 
 	return &Section;
 }
 
+
 FText FSkeletalAnimationSection::GetDisplayName() const
 {
 	return NSLOCTEXT("FAnimationSection", "AnimationSection", "Animation");
 }
 
+
 FText FSkeletalAnimationSection::GetSectionTitle() const
 {
 	return FText::FromString( Cast<UMovieSceneSkeletalAnimationSection>(&Section)->GetAnimSequence()->GetName() );
 }
 
+
 float FSkeletalAnimationSection::GetSectionHeight() const
 {
 	return (float)SkeletalAnimationEditorConstants::AnimationTrackHeight;
- }
+}
+
 
 int32 FSkeletalAnimationSection::OnPaintSection( const FGeometry& AllottedGeometry, const FSlateRect& SectionClippingRect, FSlateWindowElementList& OutDrawElements, int32 LayerId, bool bParentEnabled ) const
 {
@@ -125,25 +129,30 @@ int32 FSkeletalAnimationSection::OnPaintSection( const FGeometry& AllottedGeomet
 	return LayerId+3;
 }
 
+
 FSkeletalAnimationTrackEditor::FSkeletalAnimationTrackEditor( TSharedRef<ISequencer> InSequencer )
 	: FMovieSceneTrackEditor( InSequencer ) 
 {
 }
 
+
 FSkeletalAnimationTrackEditor::~FSkeletalAnimationTrackEditor()
 {
 }
 
-TSharedRef<FMovieSceneTrackEditor> FSkeletalAnimationTrackEditor::CreateTrackEditor( TSharedRef<ISequencer> InSequencer )
+
+TSharedRef<ISequencerTrackEditor> FSkeletalAnimationTrackEditor::CreateTrackEditor( TSharedRef<ISequencer> InSequencer )
 {
 	return MakeShareable( new FSkeletalAnimationTrackEditor( InSequencer ) );
 }
 
+
 bool FSkeletalAnimationTrackEditor::SupportsType( TSubclassOf<UMovieSceneTrack> Type ) const
 {
 	return Type == UMovieSceneSkeletalAnimationTrack::StaticClass();
 }
 
+
 TSharedRef<ISequencerSection> FSkeletalAnimationTrackEditor::MakeSectionInterface( UMovieSceneSection& SectionObject, UMovieSceneTrack* Track )
 {
 	check( SupportsType( SectionObject.GetOuter()->GetClass() ) );
@@ -153,6 +162,7 @@ TSharedRef<ISequencerSection> FSkeletalAnimationTrackEditor::MakeSectionInterfac
 	return NewSection;
 }
 
+
 void FSkeletalAnimationTrackEditor::AddKey(const FGuid& ObjectGuid, UObject* AdditionalAsset)
 {
 	UAnimSequence* AnimSequence = Cast<UAnimSequence>(AdditionalAsset);
@@ -167,6 +177,7 @@ void FSkeletalAnimationTrackEditor::AddKey(const FGuid& ObjectGuid, UObject* Add
 	}
 }
 
+
 bool FSkeletalAnimationTrackEditor::HandleAssetAdded(UObject* Asset, const FGuid& TargetObjectGuid)
 {
 	if (Asset->IsA<UAnimSequence>())
@@ -192,6 +203,7 @@ bool FSkeletalAnimationTrackEditor::HandleAssetAdded(UObject* Asset, const FGuid
 	return false;
 }
 
+
 void FSkeletalAnimationTrackEditor::BuildObjectBindingTrackMenu(FMenuBuilder& MenuBuilder, const FGuid& ObjectBinding, const UClass* ObjectClass)
 {
 	if (ObjectClass->IsChildOf(ASkeletalMeshActor::StaticClass()))
@@ -219,6 +231,7 @@ void FSkeletalAnimationTrackEditor::BuildObjectBindingTrackMenu(FMenuBuilder& Me
 	}
 }
 
+
 void FSkeletalAnimationTrackEditor::BuildAnimationSubMenu(FMenuBuilder& MenuBuilder, FGuid ObjectBinding, USkeleton* Skeleton)
 {
 	const TSharedPtr<ISequencer> ParentSequencer = GetSequencer();
@@ -249,6 +262,7 @@ void FSkeletalAnimationTrackEditor::BuildAnimationSubMenu(FMenuBuilder& MenuBuil
 	}
 }
 
+
 void FSkeletalAnimationTrackEditor::AddKeyInternal( float KeyTime, const TArray<UObject*> Objects, class UAnimSequence* AnimSequence )
 {
 	for( int32 ObjectIndex = 0; ObjectIndex < Objects.Num(); ++ObjectIndex )
@@ -268,6 +282,7 @@ void FSkeletalAnimationTrackEditor::AddKeyInternal( float KeyTime, const TArray<
 	}
 }
 
+
 USkeleton* FSkeletalAnimationTrackEditor::AcquireSkeletonFromObjectGuid(const FGuid& Guid)
 {
 	TArray<UObject*> OutObjects;
diff --git a/Engine/Source/Editor/MovieSceneTools/Private/TrackEditors/SkeletalAnimationTrackEditor.h b/Engine/Source/Editor/MovieSceneTools/Private/TrackEditors/SkeletalAnimationTrackEditor.h
--- a/Engine/Source/Editor/MovieSceneTools/Private/TrackEditors/SkeletalAnimationTrackEditor.h
+++ b/Engine/Source/Editor/MovieSceneTools/Private/TrackEditors/SkeletalAnimationTrackEditor.h
@@ -12,10 +12,12 @@ public:
 	/**
 	 * Constructor
 	 *
-	 * @param InSequencer	The sequencer instance to be used by this tool
+	 * @param InSequencer The sequencer instance to be used by this tool
 	 */
 	FSkeletalAnimationTrackEditor( TSharedRef<ISequencer> InSequencer );
-	~FSkeletalAnimationTrackEditor();
+
+	/** Virtual destructor. */
+	virtual ~FSkeletalAnimationTrackEditor();
 
 	/**
 	 * Creates an instance of this class.  Called by a sequencer 
@@ -23,9 +25,12 @@ public:
 	 * @param OwningSequencer The sequencer instance to be used by this tool
 	 * @return The new instance of this class
 	 */
-	static TSharedRef<FMovieSceneTrackEditor> CreateTrackEditor( TSharedRef<ISequencer> OwningSequencer );
+	static TSharedRef<ISequencerTrackEditor> CreateTrackEditor( TSharedRef<ISequencer> OwningSequencer );
+
+public:
+
+	// FMovieSceneTrackEditor interface
 
-	/** FMovieSceneTrackEditor Interface */
 	virtual bool SupportsType( TSubclassOf<UMovieSceneTrack> Type ) const override;
 	virtual TSharedRef<ISequencerSection> MakeSectionInterface( UMovieSceneSection& SectionObject, UMovieSceneTrack* Track ) override;
 	virtual void AddKey(const FGuid& ObjectGuid, UObject* AdditionalAsset) override;
@@ -33,6 +38,7 @@ public:
 	virtual void BuildObjectBindingTrackMenu(FMenuBuilder& MenuBuilder, const FGuid& ObjectBinding, const UClass* ObjectClass) override;
 
 private:
+
 	/** Animation sub menu */
 	void BuildAnimationSubMenu(FMenuBuilder& MenuBuilder, FGuid ObjectBinding, USkeleton* Skeleton);
 
@@ -47,15 +53,23 @@ private:
 };
 
 
-
 /** Class for animation sections, handles drawing of all waveform previews */
-class FSkeletalAnimationSection : public ISequencerSection, public TSharedFromThis<FSkeletalAnimationSection>
+class FSkeletalAnimationSection
+	: public ISequencerSection
+	, public TSharedFromThis<FSkeletalAnimationSection>
 {
 public:
+
+	/** Constructor. */
 	FSkeletalAnimationSection( UMovieSceneSection& InSection );
-	~FSkeletalAnimationSection();
 
-	/** ISequencerSection interface */
+	/** Virtual destructor. */
+	virtual ~FSkeletalAnimationSection();
+
+public:
+
+	// ISequencerSection interface
+
 	virtual UMovieSceneSection* GetSectionObject() override;
 	virtual FText GetDisplayName() const override;
 	virtual FText GetSectionTitle() const override;
@@ -64,6 +78,7 @@ public:
 	virtual int32 OnPaintSection( const FGeometry& AllottedGeometry, const FSlateRect& SectionClippingRect, FSlateWindowElementList& OutDrawElements, int32 LayerId, bool bParentEnabled ) const override;
 
 private:
+
 	/** The section we are visualizing */
 	UMovieSceneSection& Section;
 };
diff --git a/Engine/Source/Editor/MovieSceneTools/Private/TrackEditors/SubMovieSceneTrackEditor.cpp b/Engine/Source/Editor/MovieSceneTools/Private/TrackEditors/SubMovieSceneTrackEditor.cpp
--- a/Engine/Source/Editor/MovieSceneTools/Private/TrackEditors/SubMovieSceneTrackEditor.cpp
+++ b/Engine/Source/Editor/MovieSceneTools/Private/TrackEditors/SubMovieSceneTrackEditor.cpp
@@ -6,10 +6,12 @@
 #include "SubMovieSceneSection.h"
 #include "MovieSceneSequence.h"
 
+
 /**
  * A generic implementation for displaying simple property sections
  */
-class FSubMovieSceneSection : public ISequencerSection
+class FSubMovieSceneSection
+	: public ISequencerSection
 {
 public:
 	FSubMovieSceneSection( TSharedPtr<ISequencer> InSequencer, UMovieSceneSection& InSectionObject, FName SectionName )
@@ -77,6 +79,7 @@ private:
 	TWeakPtr<ISequencer> Sequencer;
 };
 
+
 FSubMovieSceneTrackEditor::FSubMovieSceneTrackEditor( TSharedRef<ISequencer> InSequencer )
 	: FMovieSceneTrackEditor( InSequencer ) 
 {
@@ -84,22 +87,25 @@ FSubMovieSceneTrackEditor::FSubMovieSceneTrackEditor( TSharedRef<ISequencer> InS
 }
 
 
-TSharedRef<FMovieSceneTrackEditor> FSubMovieSceneTrackEditor::CreateTrackEditor( TSharedRef<ISequencer> InSequencer )
+TSharedRef<ISequencerTrackEditor> FSubMovieSceneTrackEditor::CreateTrackEditor( TSharedRef<ISequencer> InSequencer )
 {
 	return MakeShareable( new FSubMovieSceneTrackEditor( InSequencer ) );
 }
 
+
 bool FSubMovieSceneTrackEditor::SupportsType( TSubclassOf<UMovieSceneTrack> Type ) const
 {
 	// We support sub movie scenes
 	return Type == USubMovieSceneTrack::StaticClass();
 }
 
+
 TSharedRef<ISequencerSection> FSubMovieSceneTrackEditor::MakeSectionInterface( UMovieSceneSection& SectionObject, UMovieSceneTrack* Track )
 {
 	return MakeShareable( new FSubMovieSceneSection( GetSequencer(), SectionObject, Track->GetTrackName() ) );
 }
 
+
 bool FSubMovieSceneTrackEditor::HandleAssetAdded(UObject* Asset, const FGuid& TargetObjectGuid)
 {
 	UMovieSceneSequence* Animation = Cast<UMovieSceneSequence>( Asset );
@@ -112,5 +118,3 @@ bool FSubMovieSceneTrackEditor::HandleAssetAdded(UObject* Asset, const FGuid& Ta
 	}
 	return false;
 }
-
-
diff --git a/Engine/Source/Editor/MovieSceneTools/Private/TrackEditors/SubMovieSceneTrackEditor.h b/Engine/Source/Editor/MovieSceneTools/Private/TrackEditors/SubMovieSceneTrackEditor.h
--- a/Engine/Source/Editor/MovieSceneTools/Private/TrackEditors/SubMovieSceneTrackEditor.h
+++ b/Engine/Source/Editor/MovieSceneTools/Private/TrackEditors/SubMovieSceneTrackEditor.h
@@ -6,28 +6,34 @@
 /**
  * Tools for animatable property types such as floats ands vectors
  */
-class FSubMovieSceneTrackEditor : public FMovieSceneTrackEditor
+class FSubMovieSceneTrackEditor
+	: public FMovieSceneTrackEditor
 {
 public:
+
 	/**
 	 * Constructor
 	 *
-	 * @param InSequencer	The sequencer instance to be used by this tool
+	 * @param InSequencer The sequencer instance to be used by this tool
 	 */
 	FSubMovieSceneTrackEditor( TSharedRef<ISequencer> InSequencer );
 
+	/** Virtual destructor. */
+	~FSubMovieSceneTrackEditor() { }
+
 	/**
 	 * Creates an instance of this class.  Called by a sequencer 
 	 *
 	 * @param OwningSequencer The sequencer instance to be used by this tool
 	 * @return The new instance of this class
 	 */
-	static TSharedRef<FMovieSceneTrackEditor> CreateTrackEditor( TSharedRef<ISequencer> OwningSequencer );
+	static TSharedRef<ISequencerTrackEditor> CreateTrackEditor( TSharedRef<ISequencer> OwningSequencer );
+
+public:
+
+	// FMovieSceneTrackEditor interface
 
-	/** FMovieSceneTrackEditor Interface */
 	virtual bool SupportsType( TSubclassOf<UMovieSceneTrack> Type ) const override;
 	virtual TSharedRef<ISequencerSection> MakeSectionInterface( UMovieSceneSection& SectionObject, UMovieSceneTrack* Track ) override;
 	virtual bool HandleAssetAdded(UObject* Asset, const FGuid& TargetObjectGuid) override;
 };
-
-
diff --git a/Engine/Source/Editor/MovieSceneTools/Private/TrackEditors/TransformTrackEditor.cpp b/Engine/Source/Editor/MovieSceneTools/Private/TrackEditors/TransformTrackEditor.cpp
--- a/Engine/Source/Editor/MovieSceneTools/Private/TrackEditors/TransformTrackEditor.cpp
+++ b/Engine/Source/Editor/MovieSceneTools/Private/TrackEditors/TransformTrackEditor.cpp
@@ -15,20 +15,29 @@
 #include "MovieSceneTrackEditor.h"
 #include "TransformTrackEditor.h"
 
+
 #define LOCTEXT_NAMESPACE "MovieScene_TransformTrack"
 
+
 /**
  * Class that draws a transform section in the sequencer
  */
-class F3DTransformSection : public ISequencerSection
+class F3DTransformSection
+	: public ISequencerSection
 {
 public:
+
 	F3DTransformSection( UMovieSceneSection& InSection )
 		: Section( InSection )
-	{
-	}
+	{ }
+
+	/** Virtual destructor. */
+	~F3DTransformSection() { }
+
+public:
+
+	// ISequencerSection interface
 
-	/** ISequencerSection interface */
 	virtual UMovieSceneSection* GetSectionObject() override
 	{ 
 		return &Section;
@@ -83,22 +92,26 @@ public:
 	}
 
 private:
+
 	/** The section we are visualizing */
 	UMovieSceneSection& Section;
 };
 
-class F3DTransformTrackCommands : public TCommands<F3DTransformTrackCommands>
-{
 
+class F3DTransformTrackCommands
+	: public TCommands<F3DTransformTrackCommands>
+{
 public:
-	F3DTransformTrackCommands() : TCommands<F3DTransformTrackCommands>
+
+	F3DTransformTrackCommands()
+		: TCommands<F3DTransformTrackCommands>
 	(
 		"3DTransformTrack",
 		NSLOCTEXT("Contexts", "3DTransformTrack", "3DTransformTrack"),
 		NAME_None, // "MainFrame" // @todo Fix this crash
 		FEditorStyle::GetStyleSetName() // Icon Style Set
 	)
-	{}
+	{ }
 		
 	/** Sets a transform key at the current time for the selected actor */
 	TSharedPtr< FUICommandInfo > AddTransformKey;
@@ -118,6 +131,7 @@ public:
 	virtual void RegisterCommands() override;
 };
 
+
 void F3DTransformTrackCommands::RegisterCommands()
 {
 	UI_COMMAND( AddTransformKey, "Add Transform Key", "Add a transform key at the current time for the selected actor.", EUserInterfaceActionType::Button, FInputChord(EKeys::S) );
@@ -126,6 +140,7 @@ void F3DTransformTrackCommands::RegisterCommands()
 	UI_COMMAND( AddScaleKey, "Add Scale Key", "Add a scale key at the current time for the selected actor.", EUserInterfaceActionType::Button, FInputChord(EModifierKey::Shift, EKeys::R) );
 }
 
+
 F3DTransformTrackEditor::F3DTransformTrackEditor( TSharedRef<ISequencer> InSequencer )
 	: FMovieSceneTrackEditor( InSequencer ) 
 {
@@ -140,12 +155,14 @@ F3DTransformTrackEditor::F3DTransformTrackEditor( TSharedRef<ISequencer> InSeque
 	F3DTransformTrackCommands::Register();
 }
 
+
 F3DTransformTrackEditor::~F3DTransformTrackEditor()
 {
 	OnRelease();
 	F3DTransformTrackCommands::Unregister();
 }
 
+
 void F3DTransformTrackEditor::OnRelease()
 {
 	GEditor->OnBeginObjectMovement().RemoveAll( this );
@@ -155,17 +172,20 @@ void F3DTransformTrackEditor::OnRelease()
 	GEditor->OnEndCameraMovement().RemoveAll( this );
 }
 
-TSharedRef<FMovieSceneTrackEditor> F3DTransformTrackEditor::CreateTrackEditor( TSharedRef<ISequencer> InSequencer )
+
+TSharedRef<ISequencerTrackEditor> F3DTransformTrackEditor::CreateTrackEditor( TSharedRef<ISequencer> InSequencer )
 {
 	return MakeShareable( new F3DTransformTrackEditor( InSequencer ) );
 }
 
+
 bool F3DTransformTrackEditor::SupportsType( TSubclassOf<UMovieSceneTrack> Type ) const
 {
 	// We support animatable transforms
 	return Type == UMovieScene3DTransformTrack::StaticClass();
 }
 
+
 TSharedRef<ISequencerSection> F3DTransformTrackEditor::MakeSectionInterface( UMovieSceneSection& SectionObject, UMovieSceneTrack* Track )
 {
 	check( SupportsType( SectionObject.GetOuter()->GetClass() ) );
@@ -183,7 +203,7 @@ void F3DTransformTrackEditor::OnPreTransformChanged( UObject& InObject )
 
 	if( IsAllowedToAutoKey() )
 	{
-		USceneComponent* SceneComponent = NULL;
+		USceneComponent* SceneComponent = nullptr;
 		AActor* Actor = Cast<AActor>( &InObject );
 		if( Actor && Actor->GetRootComponent() )
 		{
@@ -206,6 +226,7 @@ void F3DTransformTrackEditor::OnPreTransformChanged( UObject& InObject )
 	}
 }
 
+
 /**
  * Temp struct used because delegates only accept 4 or less payloads
  * FTransformKey is immutable and would require heavy re-architecting to fit here
@@ -220,13 +241,14 @@ struct FTransformDataPair
 	FTransformData LastTransformData;
 };
 
+
 void F3DTransformTrackEditor::OnTransformChanged( UObject& InObject )
 {
 	const bool bCreateHandleIfMissing = false;
 
 	const TSharedPtr<ISequencer> Sequencer = GetSequencer();
 
-	USceneComponent* SceneComponentThatChanged = NULL;
+	USceneComponent* SceneComponentThatChanged = nullptr;
 
 	// The runtime binding
 	FGuid ObjectHandle;
@@ -281,11 +303,13 @@ void F3DTransformTrackEditor::OnTransformChanged( UObject& InObject )
 	}
 }
 
+
 void F3DTransformTrackEditor::AddKey(const FGuid& ObjectGuid, UObject* AdditionalAsset)
 {
 	AddKeyInternal(ObjectGuid, AdditionalAsset);
 }
 
+
 void F3DTransformTrackEditor::AddKeyInternal(const FGuid& ObjectGuid, UObject* AdditionalAsset, const bool bAddKeyEvenIfUnchanged, F3DTransformTrackKey::Type KeyType)
 {
 	TArray<UObject*> OutObjects;
@@ -294,7 +318,7 @@ void F3DTransformTrackEditor::AddKeyInternal(const FGuid& ObjectGuid, UObject* A
 	for ( UObject* Object : OutObjects )
 	{
 		FGuid ObjectHandle;
-		USceneComponent* SceneComponent = NULL;
+		USceneComponent* SceneComponent = nullptr;
 		AActor* Actor = Cast<AActor>( Object );
 		if( Actor && Actor->GetRootComponent() )
 		{
@@ -331,26 +355,31 @@ void F3DTransformTrackEditor::AddKeyInternal(const FGuid& ObjectGuid, UObject* A
 	}
 }
 
+
 void F3DTransformTrackEditor::AddTransformKey()
 {
 	AddTransformKeyInternal();
 }
 
+
 void F3DTransformTrackEditor::AddTranslationKey()
 {
 	AddTransformKeyInternal(F3DTransformTrackKey::Key_Translation);
 }
 
+
 void F3DTransformTrackEditor::AddRotationKey()
 {
 	AddTransformKeyInternal(F3DTransformTrackKey::Key_Rotation);
 }
 
+
 void F3DTransformTrackEditor::AddScaleKey()
 {
 	AddTransformKeyInternal(F3DTransformTrackKey::Key_Scale);
 }
 
+
 void F3DTransformTrackEditor::AddTransformKeyInternal(F3DTransformTrackKey::Type KeyType)
 {
 	// WASD hotkeys to fly the viewport can conflict with hotkeys for setting keyframes (ie. s). 
@@ -370,10 +399,11 @@ void F3DTransformTrackEditor::AddTransformKeyInternal(F3DTransformTrackKey::Type
 	for (TArray<UObject*>::TIterator It(SelectedActors); It; ++It)
 	{
 		FGuid ObjectGuid = GetSequencer()->GetHandleToObject(*It);
-		AddKeyInternal(ObjectGuid, NULL, true, KeyType);
+		AddKeyInternal(ObjectGuid, nullptr, true, KeyType);
 	}
 }
 
+
 void F3DTransformTrackEditor::BindCommands(TSharedRef<FUICommandList> SequencerCommandBindings)
 {
 	const F3DTransformTrackCommands& Commands = F3DTransformTrackCommands::Get();
@@ -395,6 +425,7 @@ void F3DTransformTrackEditor::BindCommands(TSharedRef<FUICommandList> SequencerC
 		FExecuteAction::CreateSP( this, &F3DTransformTrackEditor::AddScaleKey ) );
 }
 
+
 void F3DTransformTrackEditor::BuildObjectBindingEditButtons(TSharedPtr<SHorizontalBox> EditBox, const FGuid& ObjectGuid, const UClass* ObjectClass)
 {
 	TArray<UObject*> OutObjects;
@@ -438,6 +469,7 @@ void F3DTransformTrackEditor::BuildObjectBindingEditButtons(TSharedPtr<SHorizont
 	];
 }
 
+
 void F3DTransformTrackEditor::BuildObjectBindingTrackMenu(FMenuBuilder& MenuBuilder, const FGuid& ObjectBinding, const UClass* ObjectClass)
 {
 	if (ObjectClass->IsChildOf(AActor::StaticClass()))
@@ -452,6 +484,7 @@ void F3DTransformTrackEditor::BuildObjectBindingTrackMenu(FMenuBuilder& MenuBuil
 	}
 }
 
+
 void F3DTransformTrackEditor::AddTransform(FGuid ObjectBinding)
 {
 	FScopedTransaction DeleteSectionTransaction( NSLOCTEXT("Sequencer", "AddTransformTrack_Transaction", "Add Transform Track") );
@@ -469,10 +502,10 @@ void F3DTransformTrackEditor::AddTransform(FGuid ObjectBinding)
 	TransformTrack->SetAsShowable();
 	
 	AddKey(ObjectBinding, nullptr);
-
 	NotifyMovieSceneDataChanged();
 }
 
+
 bool F3DTransformTrackEditor::CanAddTransform(FGuid ObjectBinding) const
 {
 	FName Transform("Transform");
@@ -483,6 +516,7 @@ bool F3DTransformTrackEditor::CanAddTransform(FGuid ObjectBinding) const
 	return true;
 }
 
+
 FTransformKey F3DTransformTrackEditor::GetTransformKey(const UMovieScene3DTransformTrack* TransformTrack, float KeyTime, FTransformDataPair TransformPair, FKeyParams KeyParams) const
 {
 	if (!TransformPair.LastTransformData.IsValid())
@@ -501,9 +535,11 @@ FTransformKey F3DTransformTrackEditor::GetTransformKey(const UMovieScene3DTransf
 	TransformKey.KeyParams.bAddKeyEvenIfUnchanged = KeyParams.bAddKeyEvenIfUnchanged;
 	TransformKey.KeyParams.bCreateTrackIfMissing = KeyParams.bCreateTrackIfMissing;
 	TransformKey.KeyParams.bCreateHandleIfMissing = KeyParams.bCreateHandleIfMissing;
+
 	return TransformKey;
 }
 
+
 void F3DTransformTrackEditor::OnTransformChangedInternals(float KeyTime, UObject* InObject, FGuid ObjectHandle, FTransformDataPair TransformPair, FKeyParams KeyParams, F3DTransformTrackKey::Type KeyType)
 {
 	// Only unwind rotation if we're generating keys while recording (scene is actively playing back)
@@ -527,6 +563,7 @@ void F3DTransformTrackEditor::OnTransformChangedInternals(float KeyTime, UObject
 	}
 }
 
+
 bool F3DTransformTrackEditor::CanKeyPropertyInternal(float KeyTime, FGuid ObjectHandle, FTransformDataPair TransformPair, FKeyParams KeyParams) const
 {
 	FName Transform("Transform");
@@ -551,6 +588,7 @@ bool F3DTransformTrackEditor::CanKeyPropertyInternal(float KeyTime, FGuid Object
 	return false;
 }
 
+
 ECheckBoxState F3DTransformTrackEditor::IsCameraLocked(TWeakObjectPtr<ACameraActor> CameraActor) const
 {
 	for (int32 i = 0; i < GEditor->LevelViewportClients.Num(); ++i)
@@ -565,6 +603,7 @@ ECheckBoxState F3DTransformTrackEditor::IsCameraLocked(TWeakObjectPtr<ACameraAct
 	return ECheckBoxState::Unchecked;
 }
 
+
 void F3DTransformTrackEditor::OnLockCameraClicked(ECheckBoxState CheckBoxState, TWeakObjectPtr<ACameraActor> CameraActor)
 {
 	// If toggle is on, lock the active viewport to the camera
@@ -573,7 +612,7 @@ void F3DTransformTrackEditor::OnLockCameraClicked(ECheckBoxState CheckBoxState,
 		// Set the active viewport or any viewport if there is no active viewport
 		FViewport* ActiveViewport = GEditor->GetActiveViewport();
 
-		FLevelEditorViewportClient* LevelVC = NULL;
+		FLevelEditorViewportClient* LevelVC = nullptr;
 
 		for (int32 i = 0; i < GEditor->LevelViewportClients.Num(); ++i)
 		{		
@@ -589,7 +628,7 @@ void F3DTransformTrackEditor::OnLockCameraClicked(ECheckBoxState CheckBoxState,
 			}
 		}
 
-		if (LevelVC != NULL)
+		if (LevelVC != nullptr)
 		{
 			LevelVC->SetActorLock(CameraActor.Get());
 			LevelVC->bLockedCameraView = true;
@@ -614,6 +653,7 @@ void F3DTransformTrackEditor::OnLockCameraClicked(ECheckBoxState CheckBoxState,
 	}
 }
 
+
 FText F3DTransformTrackEditor::GetLockCameraToolTip(TWeakObjectPtr<ACameraActor> CameraActor) const
 {
 	return IsCameraLocked(CameraActor) == ECheckBoxState::Checked ?
@@ -621,4 +661,5 @@ FText F3DTransformTrackEditor::GetLockCameraToolTip(TWeakObjectPtr<ACameraActor>
 		FText::Format(LOCTEXT("LockCamera", "Lock {0} to Selected Viewport"), FText::FromString(CameraActor.Get()->GetActorLabel()));
 }
 
-#undef LOCTEXT_NAMESPACE
\ No newline at end of file
+
+#undef LOCTEXT_NAMESPACE
diff --git a/Engine/Source/Editor/MovieSceneTools/Private/TrackEditors/TransformTrackEditor.h b/Engine/Source/Editor/MovieSceneTools/Private/TrackEditors/TransformTrackEditor.h
--- a/Engine/Source/Editor/MovieSceneTools/Private/TrackEditors/TransformTrackEditor.h
+++ b/Engine/Source/Editor/MovieSceneTools/Private/TrackEditors/TransformTrackEditor.h
@@ -4,19 +4,24 @@
 
 #include "MovieScene3DTransformTrack.h"
 
+
 /**
  * Tools for animatable transforms
  */
-class F3DTransformTrackEditor : public FMovieSceneTrackEditor
+class F3DTransformTrackEditor
+	: public FMovieSceneTrackEditor
 {
 public:
+
 	/**
 	 * Constructor
 	 *
 	 * @param InSequencer	The sequencer instance to be used by this tool
 	 */
 	F3DTransformTrackEditor( TSharedRef<ISequencer> InSequencer );
-	~F3DTransformTrackEditor();
+
+	/** Virtual destructor. */
+	virtual ~F3DTransformTrackEditor();
 
 	/**
 	 * Creates an instance of this class.  Called by a sequencer 
@@ -24,9 +29,12 @@ public:
 	 * @param OwningSequencer The sequencer instance to be used by this tool
 	 * @return The new instance of this class
 	 */
-	static TSharedRef<FMovieSceneTrackEditor> CreateTrackEditor( TSharedRef<ISequencer> OwningSequencer );
+	static TSharedRef<ISequencerTrackEditor> CreateTrackEditor( TSharedRef<ISequencer> OwningSequencer );
+
+public:
+
+	// FMovieSceneTrackEditor interface
 
-	/** FMovieSceneTrackEditor Interface */
 	virtual void OnRelease() override;
 	virtual bool SupportsType( TSubclassOf<UMovieSceneTrack> Type ) const override;
 	virtual TSharedRef<ISequencerSection> MakeSectionInterface( UMovieSceneSection& SectionObject, UMovieSceneTrack* Track ) override;
@@ -49,14 +57,14 @@ private:
 	/**
 	 * Called before an actor or component transform changes
 	 *
-	 * @param Object	The object whose transform is about to change
+	 * @param Object The object whose transform is about to change
 	 */
 	void OnPreTransformChanged( UObject& InObject );
 
 	/**
 	 * Called when an actor or component transform changes
 	 *
-	 * @param Object	The object whose transform has changed
+	 * @param Object The object whose transform has changed
 	 */
 	void OnTransformChanged( UObject& InObject );
 
@@ -94,8 +102,7 @@ private:
 	void AddTransformKeyInternal(F3DTransformTrackKey::Type KeyType = F3DTransformTrackKey::Key_All);
 
 private:
+
 	/** Mapping of objects to their existing transform data (for comparing against new transform data) */
 	TMap< TWeakObjectPtr<UObject>, FTransformData > ObjectToExistingTransform;
 };
-
-
diff --git a/Engine/Source/Editor/NiagaraEditor/Private/NiagaraEffectEditor.h b/Engine/Source/Editor/NiagaraEditor/Private/NiagaraEffectEditor.h
--- a/Engine/Source/Editor/NiagaraEditor/Private/NiagaraEffectEditor.h
+++ b/Engine/Source/Editor/NiagaraEditor/Private/NiagaraEffectEditor.h
@@ -39,7 +39,7 @@ public:
 
 	virtual UNiagaraEffect *GetEffect() const	override { return Effect; }	
 
-	static TSharedRef<FMovieSceneTrackEditor> CreateTrackEditor(TSharedRef<ISequencer> InSequencer)
+	static TSharedRef<ISequencerTrackEditor> CreateTrackEditor(TSharedRef<ISequencer> InSequencer)
 	{
 		return MakeShareable(new FNiagaraTrackEditor(InSequencer) );
 	}
diff --git a/Engine/Source/Editor/Sequencer/Private/Sequencer.cpp b/Engine/Source/Editor/Sequencer/Private/Sequencer.cpp
--- a/Engine/Source/Editor/Sequencer/Private/Sequencer.cpp
+++ b/Engine/Source/Editor/Sequencer/Private/Sequencer.cpp
@@ -130,7 +130,7 @@ void FSequencer::InitSequencer(const FSequencerInitParams& InitParams, const TSh
 		{
 			check( TrackEditorDelegates[DelegateIndex].IsBound() );
 			// Tools may exist in other modules, call a delegate that will create one for us 
-			TSharedRef<FMovieSceneTrackEditor> TrackEditor = TrackEditorDelegates[DelegateIndex].Execute( SharedThis( this ) );
+			TSharedRef<ISequencerTrackEditor> TrackEditor = TrackEditorDelegates[DelegateIndex].Execute( SharedThis( this ) );
 			TrackEditors.Add( TrackEditor );
 		}
 
diff --git a/Engine/Source/Editor/Sequencer/Private/Sequencer.h b/Engine/Source/Editor/Sequencer/Private/Sequencer.h
--- a/Engine/Source/Editor/Sequencer/Private/Sequencer.h
+++ b/Engine/Source/Editor/Sequencer/Private/Sequencer.h
@@ -147,7 +147,10 @@ public:
 	/**
 	 * @return Movie scene tools used by the sequencer
 	 */
-	const TArray< TSharedPtr<FMovieSceneTrackEditor> >& GetTrackEditors() const { return TrackEditors; }
+	const TArray<TSharedPtr<ISequencerTrackEditor>>& GetTrackEditors() const
+	{
+		return TrackEditors;
+	}
 
 	/**
 	 * Attempts to add a new spawnable to the MovieScene for the specified asset or class
@@ -459,7 +462,7 @@ private:
 	TSharedRef<FUICommandList> SequencerCommandBindings;
 
 	/** List of tools we own */
-	TArray< TSharedPtr<FMovieSceneTrackEditor> > TrackEditors;
+	TArray<TSharedPtr<ISequencerTrackEditor>> TrackEditors;
 
 	/** Listener for object changes being made while this sequencer is open*/
 	TSharedPtr< class ISequencerObjectChangeListener > ObjectChangeListener;
diff --git a/Engine/Source/Editor/Sequencer/Private/SequencerDisplayNode.cpp b/Engine/Source/Editor/Sequencer/Private/SequencerDisplayNode.cpp
--- a/Engine/Source/Editor/Sequencer/Private/SequencerDisplayNode.cpp
+++ b/Engine/Source/Editor/Sequencer/Private/SequencerDisplayNode.cpp
@@ -249,7 +249,7 @@ TSharedRef<FSectionCategoryNode> FSequencerDisplayNode::AddCategoryNode( FName C
 	return CategoryNode.ToSharedRef();
 }
 
-TSharedRef<FTrackNode> FSequencerDisplayNode::AddSectionAreaNode( FName SectionName, UMovieSceneTrack& AssociatedTrack, FMovieSceneTrackEditor& AssociatedEditor )
+TSharedRef<FTrackNode> FSequencerDisplayNode::AddSectionAreaNode( FName SectionName, UMovieSceneTrack& AssociatedTrack, ISequencerTrackEditor& AssociatedEditor )
 {
 	TSharedPtr<FTrackNode> SectionNode;
 
@@ -498,7 +498,7 @@ void FSectionKeyAreaNode::AddKeyArea( TSharedRef< IKeyArea> KeyArea )
 	KeyAreas.Add( KeyArea );
 }
 
-FTrackNode::FTrackNode( FName NodeName, UMovieSceneTrack& InAssociatedType, FMovieSceneTrackEditor& InAssociatedEditor, TSharedPtr<FSequencerDisplayNode> InParentNode, FSequencerNodeTree& InParentTree )
+FTrackNode::FTrackNode( FName NodeName, UMovieSceneTrack& InAssociatedType, ISequencerTrackEditor& InAssociatedEditor, TSharedPtr<FSequencerDisplayNode> InParentNode, FSequencerNodeTree& InParentTree )
 	: FSequencerDisplayNode( NodeName, InParentNode, InParentTree )
 	, AssociatedType( &InAssociatedType )
 	, AssociatedEditor( InAssociatedEditor )
diff --git a/Engine/Source/Editor/Sequencer/Private/SequencerDisplayNode.h b/Engine/Source/Editor/Sequencer/Private/SequencerDisplayNode.h
--- a/Engine/Source/Editor/Sequencer/Private/SequencerDisplayNode.h
+++ b/Engine/Source/Editor/Sequencer/Private/SequencerDisplayNode.h
@@ -75,7 +75,7 @@ public:
 	 * @param AssociatedTrack	The track associated with sections in this node
 	 * @param AssociatedEditor	The track editor for the associated track
 	 */
-	TSharedRef<class FTrackNode> AddSectionAreaNode( FName SectionName, UMovieSceneTrack& AssociatedTrack, FMovieSceneTrackEditor& AssociatedEditor );
+	TSharedRef<class FTrackNode> AddSectionAreaNode( FName SectionName, UMovieSceneTrack& AssociatedTrack, ISequencerTrackEditor& AssociatedEditor );
 
 	/**
 	 * Adds a key area to this node
@@ -394,7 +394,7 @@ public:
 	 * @param InParentNode The parent of this node or NULL if this is a root node
 	 * @param InParentTree The tree this node is in
 	 */
-	FTrackNode( FName NodeName, UMovieSceneTrack& InAssociatedType, FMovieSceneTrackEditor& InAssociatedEditor, TSharedPtr<FSequencerDisplayNode> InParentNode, FSequencerNodeTree& InParentTree );
+	FTrackNode( FName NodeName, UMovieSceneTrack& InAssociatedType, ISequencerTrackEditor& InAssociatedEditor, TSharedPtr<FSequencerDisplayNode> InParentNode, FSequencerNodeTree& InParentTree );
 
 	/** FSequencerDisplayNodeInterface */
 	virtual ESequencerNode::Type GetType() const override { return ESequencerNode::Track; }
@@ -448,7 +448,7 @@ private:
 	/** The type associated with the sections in this node */
 	TWeakObjectPtr<UMovieSceneTrack> AssociatedType;
 	/** The track editor for the track associated with this node. */
-	FMovieSceneTrackEditor& AssociatedEditor;
+	ISequencerTrackEditor& AssociatedEditor;
 };
 
 /**
diff --git a/Engine/Source/Editor/Sequencer/Private/SequencerNodeTree.cpp b/Engine/Source/Editor/Sequencer/Private/SequencerNodeTree.cpp
--- a/Engine/Source/Editor/Sequencer/Private/SequencerNodeTree.cpp
+++ b/Engine/Source/Editor/Sequencer/Private/SequencerNodeTree.cpp
@@ -3,14 +3,15 @@
 #include "SequencerPrivatePCH.h"
 #include "MovieSceneSequence.h"
 #include "MovieSceneSection.h"
+#include "MovieSceneTrack.h"
 #include "SequencerNodeTree.h"
 #include "Sequencer.h"
 #include "ScopedTransaction.h"
 #include "MovieScene.h"
-#include "MovieSceneTrack.h"
 #include "MovieSceneTrackEditor.h"
 #include "SectionLayoutBuilder.h"
 #include "ISequencerSection.h"
+#include "ISequencerTrackEditor.h"
 
 
 void FSequencerNodeTree::Empty()
@@ -139,17 +140,17 @@ void FSequencerNodeTree::Update()
 	FilterNodes( FilterString );
 }
 
-TSharedRef<FMovieSceneTrackEditor> FSequencerNodeTree::FindOrAddTypeEditor( UMovieSceneTrack& InTrack )
+TSharedRef<ISequencerTrackEditor> FSequencerNodeTree::FindOrAddTypeEditor( UMovieSceneTrack& InTrack )
 {
-	TSharedPtr<FMovieSceneTrackEditor> Editor = EditorMap.FindRef( &InTrack );
+	TSharedPtr<ISequencerTrackEditor> Editor = EditorMap.FindRef( &InTrack );
 
 	if( !Editor.IsValid() )
 	{
-		const TArray< TSharedPtr<FMovieSceneTrackEditor> >& TrackEditors = Sequencer.GetTrackEditors();
+		const TArray< TSharedPtr<ISequencerTrackEditor> >& TrackEditors = Sequencer.GetTrackEditors();
 
 		// Get a tool for each track
 		// @todo Sequencer - Should probably only need to get this once and it shouldnt be done here. It depends on when movie scene tool modules are loaded
-		TSharedPtr<FMovieSceneTrackEditor> SupportedTool;
+		TSharedPtr<ISequencerTrackEditor> SupportedTool;
 		for( int32 EditorIndex = 0; EditorIndex < TrackEditors.Num(); ++EditorIndex )
 		{
 			if( TrackEditors[EditorIndex]->SupportsType( InTrack.GetClass() ) )
@@ -168,7 +169,7 @@ void FSequencerNodeTree::MakeSectionInterfaces( UMovieSceneTrack& Track, TShared
 {
 	const TArray<UMovieSceneSection*>& MovieSceneSections = Track.GetAllSections();
 
-	TSharedRef<FMovieSceneTrackEditor> Editor = FindOrAddTypeEditor( Track );
+	TSharedRef<ISequencerTrackEditor> Editor = FindOrAddTypeEditor( Track );
 
 	for (int32 SectionIndex = 0; SectionIndex < MovieSceneSections.Num(); ++SectionIndex )
 	{
diff --git a/Engine/Source/Editor/Sequencer/Private/SequencerNodeTree.h b/Engine/Source/Editor/Sequencer/Private/SequencerNodeTree.h
--- a/Engine/Source/Editor/Sequencer/Private/SequencerNodeTree.h
+++ b/Engine/Source/Editor/Sequencer/Private/SequencerNodeTree.h
@@ -2,12 +2,15 @@
 
 #pragma once
 
+
+class ISequencerTrackEditor;
 class UMovieSceneSection;
 class FSequencerDisplayNode;
 class FObjectBindingNode;
 class UMovieSceneTrack;
 struct FMovieSceneBinding;
 
+
 /**
  * Represents a tree of sequencer display nodes, used to populate the Sequencer UI with MovieScene data
  */
@@ -86,7 +89,7 @@ private:
 	 * @param Track	The type to find an editor for
 	 * @rerturn The editor for the type
 	 */
-	TSharedRef<FMovieSceneTrackEditor> FindOrAddTypeEditor( UMovieSceneTrack& Track );
+	TSharedRef<ISequencerTrackEditor> FindOrAddTypeEditor( UMovieSceneTrack& Track );
 
 	/**
 	 * Makes section interfaces for all sections in a track
@@ -102,7 +105,7 @@ private:
 	TSharedRef<FObjectBindingNode> AddObjectBinding( const FString& ObjectName, const FGuid& ObjectBinding, TMap<FGuid, const FMovieSceneBinding*>& GuidToBindingMap, TArray< TSharedRef<FSequencerDisplayNode> >& OutNodeList );
 private:
 	/** Tools for building movie scene section layouts.  One tool for each track */
-	TMap< class UMovieSceneTrack*, TSharedPtr<FMovieSceneTrackEditor> > EditorMap;
+	TMap< UMovieSceneTrack*, TSharedPtr<ISequencerTrackEditor> > EditorMap;
 	/** Root nodes */
 	TArray< TSharedRef<FSequencerDisplayNode> > RootNodes;
 	/** Mapping of object binding guids to their node (for fast lookup) */
diff --git a/Engine/Source/Editor/Sequencer/Public/ISequencerModule.h b/Engine/Source/Editor/Sequencer/Public/ISequencerModule.h
--- a/Engine/Source/Editor/Sequencer/Public/ISequencerModule.h
+++ b/Engine/Source/Editor/Sequencer/Public/ISequencerModule.h
@@ -9,6 +9,7 @@
 
 
 class ISequencerObjectBindingManager;
+class ISequencerTrackEditor;
 class IToolkitHost;
 class UActorAnimation;
 
@@ -20,7 +21,7 @@ namespace SequencerMenuExtensionPoints
 
 
 /** A delegate which will create an auto-key handler. */
-DECLARE_DELEGATE_RetVal_OneParam(TSharedRef<class FMovieSceneTrackEditor>, FOnCreateTrackEditor, TSharedRef<ISequencer>);
+DECLARE_DELEGATE_RetVal_OneParam(TSharedRef<ISequencerTrackEditor>, FOnCreateTrackEditor, TSharedRef<ISequencer>);
 
 /** A delegate that is executed when adding menu content. */
 DECLARE_DELEGATE_RetVal_OneParam(TSharedRef<SWidget>, FOnGetAddMenuContent, TSharedRef<ISequencer>);
diff --git a/Engine/Source/Editor/Sequencer/Public/MovieSceneTrackEditor.h b/Engine/Source/Editor/Sequencer/Public/MovieSceneTrackEditor.h
--- a/Engine/Source/Editor/Sequencer/Public/MovieSceneTrackEditor.h
+++ b/Engine/Source/Editor/Sequencer/Public/MovieSceneTrackEditor.h
@@ -4,32 +4,41 @@
 
 //@todo Sequencer - these have to be here for now because this class contains a small amount of implementation inside this header to avoid exporting this class
 #include "ISequencer.h"
+#include "ISequencerSection.h"
+#include "ISequencerTrackEditor.h"
 #include "MovieScene.h"
 #include "ScopedTransaction.h"
 #include "MovieSceneSequence.h"
 
+
 class ISequencerSection;
 
+
 DECLARE_DELEGATE_OneParam(FOnKeyProperty, float)
 DECLARE_DELEGATE_RetVal_OneParam(bool, FCanKeyProperty, float)
 
+
 /**
- * Base class for handling key and section drawing and manipulation for a UMovieSceneTrack class
- * @todo Sequencer - Interface needs cleanup
+ * Base class for handling key and section drawing and manipulation
+ * of a UMovieSceneTrack class.
+ *
+ * @todo Sequencer Interface needs cleanup
  */
-class FMovieSceneTrackEditor : public TSharedFromThis<FMovieSceneTrackEditor>
+class FMovieSceneTrackEditor
+	: public TSharedFromThis<FMovieSceneTrackEditor>
+	, public ISequencerTrackEditor
 {
 public:
+
 	/** Constructor */
-	FMovieSceneTrackEditor( TSharedRef<ISequencer> InSequencer )
-		: Sequencer( InSequencer )
-	{}
+	FMovieSceneTrackEditor(TSharedRef<ISequencer> InSequencer)
+		: Sequencer(InSequencer)
+	{ }
 
 	/** Destructor */
-	virtual ~FMovieSceneTrackEditor(){};
+	virtual ~FMovieSceneTrackEditor() { };
 
-	/** Called when the instance of this track editor is released */
-	virtual void OnRelease(){};
+public:
 
 	/** @return The current movie scene */
 	UMovieSceneSequence* GetMovieSceneSequence() const
@@ -49,20 +58,6 @@ public:
 		return Sequencer.Pin()->GetCurrentLocalTime( *InMovieSceneSequence );
 	}
 
-	/** Gets whether the tool can legally autokey */
-	virtual bool IsAllowedToAutoKey() const
-	{
-		// @todo sequencer livecapture: This turns on "auto key" for the purpose of capture keys for actor state
-		// during PIE sessions when record mode is active.
-		return Sequencer.Pin()->IsRecordingLive() || Sequencer.Pin()->GetAutoKeyEnabled();
-	}
-
-	/** Gets whether the tool can key all*/
-	virtual bool IsAllowedKeyAll() const
-	{
-		return Sequencer.Pin()->GetKeyAllEnabled();
-	}
-
 	void NotifyMovieSceneDataChanged()
 	{
 		TSharedPtr<ISequencer> SequencerPin = Sequencer.Pin();
@@ -135,11 +130,6 @@ public:
 		return Type;
 	}
 
-	virtual UMovieSceneTrack* AddTrack(UMovieScene* FocusedMovieScene, const FGuid& ObjectHandle, TSubclassOf<class UMovieSceneTrack> TrackClass, FName UniqueTypeName)
-	{
-		return FocusedMovieScene->AddTrack(TrackClass, ObjectHandle);
-	}
-
 	UMovieSceneTrack* GetMasterTrack( TSubclassOf<UMovieSceneTrack> TrackClass )
 	{
 		UMovieScene* MovieScene = GetSequencer()->GetFocusedMovieSceneSequence()->GetMovieScene();
@@ -155,89 +145,48 @@ public:
 
 
 	/** @return The sequencer bound to this handler */
-	const TSharedPtr<ISequencer> GetSequencer() const { return Sequencer.Pin(); }
-	
-	/**
-	 * Returns whether a track class is supported by this tool
-	 *
-	 * @param TrackClass	The track class that could be supported
-	 * @return true if the type is supported
-	 */
-	virtual bool SupportsType( TSubclassOf<class UMovieSceneTrack> TrackClass ) const = 0;
+	const TSharedPtr<ISequencer> GetSequencer() const
+	{
+		return Sequencer.Pin();
+	}
 
-	/**
-	 * Called to generate a section layout for a particular section
-	 *
-	 * @param SectionObject	The section to make UI for
-	 * @param SectionName	The name of the section
-	 */
-	virtual TSharedRef<ISequencerSection> MakeSectionInterface( class UMovieSceneSection& SectionObject, UMovieSceneTrack* Track ) = 0;
+public:
 
-	/**
-	 * Manually adds a key
+	// ISequencerTrackEditor interface
 
-	 * @param ObjectGuid		The Guid of the object that we are adding a key to
-	 * @param AdditionalAsset	An optional asset that can be added with the key
-	 */
-	virtual void AddKey( const FGuid& ObjectGuid, UObject* AdditionalAsset = NULL ) {}
+	virtual void AddKey( const FGuid& ObjectGuid, UObject* AdditionalAsset = nullptr ) override { }
 
-	/** Ticks this tool */
-	virtual void Tick(float DeltaTime) {}
-	
-	/**
-	 * Called when an asset is dropped into Sequencer. Can potentially consume the asset
-	 * so it doesn't get added as a spawnable
-	 *
-	 * @param Asset				The asset that is dropped in
-	 * @param TargetObjectGuid	The object guid this asset is dropped onto, if applicable
-	 * @return					True, if we want to consume this asset
-	 */
-	virtual bool HandleAssetAdded(UObject* Asset, const FGuid& TargetObjectGuid) {return false;}
-	
-	/**
-	 * Allows the track editors to bind commands
-	 *
-	 * @param SequencerCommandBindings	The command bindings to map to
-	*/
-	virtual void BindCommands(TSharedRef<FUICommandList> SequencerCommandBindings) {}
+	virtual UMovieSceneTrack* AddTrack(UMovieScene* FocusedMovieScene, const FGuid& ObjectHandle, TSubclassOf<class UMovieSceneTrack> TrackClass, FName UniqueTypeName) override
+	{
+		return FocusedMovieScene->AddTrack(TrackClass, ObjectHandle);
+	}
 
-	/**
-	 * Builds up the object binding context menu for the outliner
-	 *
-	 * @param MenuBuilder	The menu builder to change
-	 * @param ObjectBinding	The object binding this is for
-	 * @param ObjectClass	The class of the object this is for
-	 */
-	virtual void BuildObjectBindingContextMenu(FMenuBuilder& MenuBuilder, const FGuid& ObjectBinding, const UClass* ObjectClass) {}
+	virtual void BindCommands(TSharedRef<FUICommandList> SequencerCommandBindings) override { }
+	virtual void BuildObjectBindingContextMenu(FMenuBuilder& MenuBuilder, const FGuid& ObjectBinding, const UClass* ObjectClass) override { }
+	virtual void BuildObjectBindingEditButtons(TSharedPtr<SHorizontalBox> EditBox, const FGuid& ObjectBinding, const UClass* ObjectClass) override { }
+	virtual void BuildObjectBindingTrackMenu(FMenuBuilder& MenuBuilder, const FGuid& ObjectBinding, const UClass* ObjectClass) override { }
+	virtual TSharedPtr<SWidget> BuildOutlinerEditWidget(const FGuid& ObjectBinding, UMovieSceneTrack* Track) override  { return TSharedPtr<SWidget>(); }
+	virtual bool HandleAssetAdded(UObject* Asset, const FGuid& TargetObjectGuid) override { return false; }
 
-	/**
-	 * Builds up the object binding track menu for the outliner
-	 *
-	 * @param MenuBuilder	The menu builder to change
-	 * @param ObjectBinding	The object binding this is for
-	 * @param ObjectClass	The class of the object this is for
-	 */
-	virtual void BuildObjectBindingTrackMenu(FMenuBuilder& MenuBuilder, const FGuid& ObjectBinding, const UClass* ObjectClass) {}
+	virtual bool IsAllowedKeyAll() const
+	{
+		return Sequencer.Pin()->GetKeyAllEnabled();
+	}
 
-	/**
-	 * Builds up the object binding edit buttons for the outliner
-	 *
-	 * @param EditBox	    The edit box to add buttons to
-	 * @param ObjectBinding	The object binding this is for
-	 * @param ObjectClass	The class of the object this is for
-	 */
-	virtual void BuildObjectBindingEditButtons(TSharedPtr<SHorizontalBox> EditBox, const FGuid& ObjectBinding, const UClass* ObjectClass) {}
+	virtual bool IsAllowedToAutoKey() const
+	{
+		// @todo sequencer livecapture: This turns on "auto key" for the purpose of capture keys for actor state
+		// during PIE sessions when record mode is active.
+		return Sequencer.Pin()->IsRecordingLive() || Sequencer.Pin()->GetAutoKeyEnabled();
+	}
 
-	/**
-	 * Builds an edit widget for the outliner nodes which represent tracks which are edited by this editor.
-	 * @param ObjectBinding The object binding associated with the track being edited by this editor.
-	 * @param Track The track being edited by this editor.
-	 * @returns The the widget to display in the outliner, or an empty shared ptr if not widget is to be displayed.
-	 */
-	virtual TSharedPtr<SWidget> BuildOutlinerEditWidget(const FGuid& ObjectBinding, UMovieSceneTrack* Track) { return TSharedPtr<SWidget>(); }
+	virtual void OnRelease() override { };
+	virtual TSharedRef<ISequencerSection> MakeSectionInterface( class UMovieSceneSection& SectionObject, UMovieSceneTrack* Track ) = 0;
+	virtual bool SupportsType( TSubclassOf<class UMovieSceneTrack> TrackClass ) const = 0;
+	virtual void Tick(float DeltaTime) override { }
 
 private:
+
 	/** The sequencer bound to this handler.  Used to access movie scene and time info during auto-key */
 	TWeakPtr<ISequencer> Sequencer;
 };
-
diff --git a/Engine/Source/Editor/UMGEditor/Private/Animation/MarginTrackEditor.cpp b/Engine/Source/Editor/UMGEditor/Private/Animation/MarginTrackEditor.cpp
--- a/Engine/Source/Editor/UMGEditor/Private/Animation/MarginTrackEditor.cpp
+++ b/Engine/Source/Editor/UMGEditor/Private/Animation/MarginTrackEditor.cpp
@@ -8,7 +8,9 @@
 #include "ISectionLayoutBuilder.h"
 #include "MovieSceneToolHelpers.h"
 
-class FMarginPropertySection : public FPropertySection
+
+class FMarginPropertySection
+	: public FPropertySection
 {
 public:
 	FMarginPropertySection( UMovieSceneSection& InSectionObject, FName SectionName )
@@ -25,11 +27,13 @@ public:
 	}
 };
 
-TSharedRef<FMovieSceneTrackEditor> FMarginTrackEditor::CreateTrackEditor( TSharedRef<ISequencer> InSequencer )
+
+TSharedRef<ISequencerTrackEditor> FMarginTrackEditor::CreateTrackEditor( TSharedRef<ISequencer> InSequencer )
 {
 	return MakeShareable( new FMarginTrackEditor( InSequencer ) );
 }
 
+
 TSharedRef<ISequencerSection> FMarginTrackEditor::MakeSectionInterface( UMovieSceneSection& SectionObject, UMovieSceneTrack* Track )
 {
 	check( SupportsType( SectionObject.GetOuter()->GetClass() ) );
@@ -41,6 +45,7 @@ TSharedRef<ISequencerSection> FMarginTrackEditor::MakeSectionInterface( UMovieSc
 	return NewSection;
 }
 
+
 bool FMarginTrackEditor::TryGenerateKeyFromPropertyChanged( const UMovieSceneTrack* InTrack, const FPropertyChangedParams& PropertyChangedParams, FMarginKey& OutKey )
 {
 	OutKey.CurveName = PropertyChangedParams.StructPropertyNameToKey;
@@ -57,4 +62,4 @@ bool FMarginTrackEditor::TryGenerateKeyFromPropertyChanged( const UMovieSceneTra
 	}
 
 	return false;
-}
\ No newline at end of file
+}
diff --git a/Engine/Source/Editor/UMGEditor/Private/Animation/MarginTrackEditor.h b/Engine/Source/Editor/UMGEditor/Private/Animation/MarginTrackEditor.h
--- a/Engine/Source/Editor/UMGEditor/Private/Animation/MarginTrackEditor.h
+++ b/Engine/Source/Editor/UMGEditor/Private/Animation/MarginTrackEditor.h
@@ -5,34 +5,42 @@
 #include "PropertyTrackEditor.h"
 #include "MovieSceneMarginTrack.h"
 
+
 class UMovieSceneTrack;
 
-class FMarginTrackEditor : public FPropertyTrackEditor<UMovieSceneMarginTrack, FMarginKey>
+
+class FMarginTrackEditor
+	: public FPropertyTrackEditor<UMovieSceneMarginTrack, FMarginKey>
 {
 public:
+
 	/**
 	 * Constructor
 	 *
-	 * @param InSequencer	The sequencer instance to be used by this tool
+	 * @param InSequencer The sequencer instance to be used by this tool
 	 */
 	FMarginTrackEditor( TSharedRef<ISequencer> InSequencer )
 		: FPropertyTrackEditor<UMovieSceneMarginTrack, FMarginKey>( InSequencer, "Margin" )
 	{ }
 
+	/** Virtual destructor. */
+
 	/**
 	 * Creates an instance of this class.  Called by a sequencer 
 	 *
 	 * @param OwningSequencer The sequencer instance to be used by this tool
 	 * @return The new instance of this class
 	 */
-	static TSharedRef<FMovieSceneTrackEditor> CreateTrackEditor( TSharedRef<ISequencer> OwningSequencer );
+	static TSharedRef<ISequencerTrackEditor> CreateTrackEditor( TSharedRef<ISequencer> OwningSequencer );
+
+public:
+
+	// FMovieSceneTrackEditor interface
 
-	/** FMovieSceneTrackEditor Interface */
 	virtual TSharedRef<ISequencerSection> MakeSectionInterface( UMovieSceneSection& SectionObject, UMovieSceneTrack* Track ) override;
 
 protected:
+
 	/** FPropertyTrackEditor Interface */
 	virtual bool TryGenerateKeyFromPropertyChanged( const UMovieSceneTrack* InTrack, const FPropertyChangedParams& PropertyChangedParams, FMarginKey& OutKey ) override;
 };
-
-
diff --git a/Engine/Source/Editor/UMGEditor/Private/Animation/Sequencer2DTransformTrackEditor.cpp b/Engine/Source/Editor/UMGEditor/Private/Animation/Sequencer2DTransformTrackEditor.cpp
--- a/Engine/Source/Editor/UMGEditor/Private/Animation/Sequencer2DTransformTrackEditor.cpp
+++ b/Engine/Source/Editor/UMGEditor/Private/Animation/Sequencer2DTransformTrackEditor.cpp
@@ -8,9 +8,12 @@
 #include "ISectionLayoutBuilder.h"
 #include "MovieSceneToolHelpers.h"
 
-class F2DTransformSection : public FPropertySection
+
+class F2DTransformSection
+	: public FPropertySection
 {
 public:
+
 	F2DTransformSection( UMovieSceneSection& InSectionObject, FName SectionName )
 		: FPropertySection(InSectionObject, SectionName) {}
 
@@ -40,11 +43,13 @@ public:
 	}
 };
 
-TSharedRef<FMovieSceneTrackEditor> F2DTransformTrackEditor::CreateTrackEditor( TSharedRef<ISequencer> InSequencer )
+
+TSharedRef<ISequencerTrackEditor> F2DTransformTrackEditor::CreateTrackEditor( TSharedRef<ISequencer> InSequencer )
 {
 	return MakeShareable( new F2DTransformTrackEditor( InSequencer ) );
 }
 
+
 TSharedRef<ISequencerSection> F2DTransformTrackEditor::MakeSectionInterface( UMovieSceneSection& SectionObject, UMovieSceneTrack* Track )
 {
 	check( SupportsType( SectionObject.GetOuter()->GetClass() ) );
@@ -56,6 +61,7 @@ TSharedRef<ISequencerSection> F2DTransformTrackEditor::MakeSectionInterface( UMo
 	return NewSection;
 }
 
+
 bool F2DTransformTrackEditor::TryGenerateKeyFromPropertyChanged( const UMovieSceneTrack* InTrack, const FPropertyChangedParams& PropertyChangedParams, F2DTransformKey& OutKey )
 {
 	OutKey.CurveName = PropertyChangedParams.StructPropertyNameToKey;
@@ -72,4 +78,4 @@ bool F2DTransformTrackEditor::TryGenerateKeyFromPropertyChanged( const UMovieSce
 	}
 
 	return false;
-}
\ No newline at end of file
+}
diff --git a/Engine/Source/Editor/UMGEditor/Private/Animation/Sequencer2DTransformTrackEditor.h b/Engine/Source/Editor/UMGEditor/Private/Animation/Sequencer2DTransformTrackEditor.h
--- a/Engine/Source/Editor/UMGEditor/Private/Animation/Sequencer2DTransformTrackEditor.h
+++ b/Engine/Source/Editor/UMGEditor/Private/Animation/Sequencer2DTransformTrackEditor.h
@@ -5,15 +5,19 @@
 #include "PropertyTrackEditor.h"
 #include "MovieScene2DTransformTrack.h"
 
+
 class UMovieSceneTrack;
 
-class F2DTransformTrackEditor : public FPropertyTrackEditor<UMovieScene2DTransformTrack, F2DTransformKey>
+
+class F2DTransformTrackEditor
+	: public FPropertyTrackEditor<UMovieScene2DTransformTrack, F2DTransformKey>
 {
 public:
+
 	/**
 	 * Constructor
 	 *
-	 * @param InSequencer	The sequencer instance to be used by this tool
+	 * @param InSequencer The sequencer instance to be used by this tool
 	 */
 	F2DTransformTrackEditor( TSharedRef<ISequencer> InSequencer )
 		: FPropertyTrackEditor<UMovieScene2DTransformTrack, F2DTransformKey>( InSequencer, "WidgetTransform" )
@@ -25,14 +29,17 @@ public:
 	 * @param OwningSequencer The sequencer instance to be used by this tool
 	 * @return The new instance of this class
 	 */
-	static TSharedRef<FMovieSceneTrackEditor> CreateTrackEditor( TSharedRef<ISequencer> OwningSequencer );
+	static TSharedRef<ISequencerTrackEditor> CreateTrackEditor( TSharedRef<ISequencer> OwningSequencer );
+
+public:
+
+	// FMovieSceneTrackEditor interface
 
-	/** FMovieSceneTrackEditor Interface */
 	virtual TSharedRef<ISequencerSection> MakeSectionInterface( UMovieSceneSection& SectionObject, UMovieSceneTrack* Track ) override;
 
 protected:
-	/** FPropertyTrackEditor Interface */
-	virtual bool TryGenerateKeyFromPropertyChanged( const UMovieSceneTrack* InTrack, const FPropertyChangedParams& PropertyChangedParams, F2DTransformKey& OutKey ) override;
-};
 
+	// FPropertyTrackEditor interface
 
+	virtual bool TryGenerateKeyFromPropertyChanged( const UMovieSceneTrack* InTrack, const FPropertyChangedParams& PropertyChangedParams, F2DTransformKey& OutKey ) override;
+};
