Sequencer: Fixed some more comments, virtual destructors

diff --git a/Engine/Source/Editor/MovieSceneTools/Private/TrackEditors/AttachTrackEditor.h b/Engine/Source/Editor/MovieSceneTools/Private/TrackEditors/AttachTrackEditor.h
--- a/Engine/Source/Editor/MovieSceneTools/Private/TrackEditors/AttachTrackEditor.h
+++ b/Engine/Source/Editor/MovieSceneTools/Private/TrackEditors/AttachTrackEditor.h
@@ -42,12 +42,12 @@ public:
 
 public:
 
-	// FMovieSceneTrackEditor interface
+	// ISequencerTrackEditor interface
 
-	virtual bool SupportsType( TSubclassOf<UMovieSceneTrack> Type ) const override;
-	virtual TSharedRef<ISequencerSection> MakeSectionInterface( UMovieSceneSection& SectionObject, UMovieSceneTrack* Track ) override;
 	virtual void AddKey( const FGuid& ObjectGuid, UObject* AdditionalAsset = NULL ) override;
 	virtual void BuildObjectBindingTrackMenu(FMenuBuilder& MenuBuilder, const FGuid& ObjectBinding, const UClass* ObjectClass) override;
+	virtual TSharedRef<ISequencerSection> MakeSectionInterface( UMovieSceneSection& SectionObject, UMovieSceneTrack* Track ) override;
+	virtual bool SupportsType( TSubclassOf<UMovieSceneTrack> Type ) const override;
 
 private:
 
diff --git a/Engine/Source/Editor/MovieSceneTools/Private/TrackEditors/AudioTrackEditor.h b/Engine/Source/Editor/MovieSceneTools/Private/TrackEditors/AudioTrackEditor.h
--- a/Engine/Source/Editor/MovieSceneTools/Private/TrackEditors/AudioTrackEditor.h
+++ b/Engine/Source/Editor/MovieSceneTools/Private/TrackEditors/AudioTrackEditor.h
@@ -31,11 +31,11 @@ public:
 
 public:
 
-	// FMovieSceneTrackEditor interface
+	// ISequencerTrackEditor interface
 
-	virtual bool SupportsType( TSubclassOf<UMovieSceneTrack> Type ) const override;
-	virtual TSharedRef<ISequencerSection> MakeSectionInterface( UMovieSceneSection& SectionObject, UMovieSceneTrack* Track ) override;
 	virtual bool HandleAssetAdded(UObject* Asset, const FGuid& TargetObjectGuid) override;
+	virtual TSharedRef<ISequencerSection> MakeSectionInterface( UMovieSceneSection& SectionObject, UMovieSceneTrack* Track ) override;
+	virtual bool SupportsType( TSubclassOf<UMovieSceneTrack> Type ) const override;
 	
 private:
 
diff --git a/Engine/Source/Editor/MovieSceneTools/Private/TrackEditors/EventTrackEditor.h b/Engine/Source/Editor/MovieSceneTools/Private/TrackEditors/EventTrackEditor.h
--- a/Engine/Source/Editor/MovieSceneTools/Private/TrackEditors/EventTrackEditor.h
+++ b/Engine/Source/Editor/MovieSceneTools/Private/TrackEditors/EventTrackEditor.h
@@ -36,7 +36,7 @@ public:
 
 public:
 
-	// FMovieSceneTrackEditor interface
+	// ISequencerTrackEditor interface
 
 	virtual void AddKey(const FGuid& ObjectGuid, UObject* AdditionalAsset = nullptr) override;
 	virtual void BuildObjectBindingTrackMenu(FMenuBuilder& MenuBuilder, const FGuid& ObjectBinding, const UClass* ObjectClass) override;
diff --git a/Engine/Source/Editor/MovieSceneTools/Private/TrackEditors/MaterialTrackEditor.h b/Engine/Source/Editor/MovieSceneTools/Private/TrackEditors/MaterialTrackEditor.h
--- a/Engine/Source/Editor/MovieSceneTools/Private/TrackEditors/MaterialTrackEditor.h
+++ b/Engine/Source/Editor/MovieSceneTools/Private/TrackEditors/MaterialTrackEditor.h
@@ -22,9 +22,10 @@ public:
 
 public:
 
-	// FMovieSceneTrackEditor interface
-	virtual TSharedRef<ISequencerSection> MakeSectionInterface( UMovieSceneSection& SectionObject, UMovieSceneTrack* Track ) override;
+	// ISequencerTrackEditor interface
+
 	virtual TSharedPtr<SWidget> BuildOutlinerEditWidget( const FGuid& ObjectBinding, UMovieSceneTrack* Track ) override;
+	virtual TSharedRef<ISequencerSection> MakeSectionInterface( UMovieSceneSection& SectionObject, UMovieSceneTrack* Track ) override;
 
 protected:
 
@@ -69,7 +70,8 @@ public:
 
 public:
 
-	// FMovieSceneTrackEditor interface
+	// ISequencerTrackEditor interface
+
 	virtual bool SupportsType( TSubclassOf<UMovieSceneTrack> Type ) const override;
 
 protected:
diff --git a/Engine/Source/Editor/MovieSceneTools/Private/TrackEditors/ParticleTrackEditor.h b/Engine/Source/Editor/MovieSceneTools/Private/TrackEditors/ParticleTrackEditor.h
--- a/Engine/Source/Editor/MovieSceneTools/Private/TrackEditors/ParticleTrackEditor.h
+++ b/Engine/Source/Editor/MovieSceneTools/Private/TrackEditors/ParticleTrackEditor.h
@@ -33,11 +33,11 @@ public:
 
 public:
 
-	// FMovieSceneTrackEditor interface
+	// ISequencerTrackEditor interface
 
-	virtual bool SupportsType( TSubclassOf<UMovieSceneTrack> Type ) const override;
-	virtual TSharedRef<ISequencerSection> MakeSectionInterface( UMovieSceneSection& SectionObject, UMovieSceneTrack* Track ) override;
 	virtual void BuildObjectBindingTrackMenu(FMenuBuilder& MenuBuilder, const FGuid& ObjectBinding, const UClass* ObjectClass) override;
+	virtual TSharedRef<ISequencerSection> MakeSectionInterface( UMovieSceneSection& SectionObject, UMovieSceneTrack* Track ) override;
+	virtual bool SupportsType( TSubclassOf<UMovieSceneTrack> Type ) const override;
 	
 private:
 
diff --git a/Engine/Source/Editor/MovieSceneTools/Private/TrackEditors/PathTrackEditor.h b/Engine/Source/Editor/MovieSceneTools/Private/TrackEditors/PathTrackEditor.h
--- a/Engine/Source/Editor/MovieSceneTools/Private/TrackEditors/PathTrackEditor.h
+++ b/Engine/Source/Editor/MovieSceneTools/Private/TrackEditors/PathTrackEditor.h
@@ -42,12 +42,12 @@ public:
 
 public:
 
-	// FMovieSceneTrackEditor interface
+	// ISequencerTrackEditor interface
 
-	virtual bool SupportsType( TSubclassOf<UMovieSceneTrack> Type ) const override;
-	virtual TSharedRef<ISequencerSection> MakeSectionInterface( UMovieSceneSection& SectionObject, UMovieSceneTrack* Track ) override;
 	virtual void AddKey( const FGuid& ObjectGuid, UObject* AdditionalAsset = NULL ) override;
 	virtual void BuildObjectBindingTrackMenu(FMenuBuilder& MenuBuilder, const FGuid& ObjectBinding, const UClass* ObjectClass) override;
+	virtual TSharedRef<ISequencerSection> MakeSectionInterface( UMovieSceneSection& SectionObject, UMovieSceneTrack* Track ) override;
+	virtual bool SupportsType( TSubclassOf<UMovieSceneTrack> Type ) const override;
 
 private:
 
diff --git a/Engine/Source/Editor/MovieSceneTools/Private/TrackEditors/PropertyTrackEditors/BoolPropertyTrackEditor.h b/Engine/Source/Editor/MovieSceneTools/Private/TrackEditors/PropertyTrackEditors/BoolPropertyTrackEditor.h
--- a/Engine/Source/Editor/MovieSceneTools/Private/TrackEditors/PropertyTrackEditors/BoolPropertyTrackEditor.h
+++ b/Engine/Source/Editor/MovieSceneTools/Private/TrackEditors/PropertyTrackEditors/BoolPropertyTrackEditor.h
@@ -31,7 +31,8 @@ public:
 
 public:
 
-	// FMovieSceneTrackEditor interface
+	// ISequencerTrackEditor interface
+
 	virtual TSharedRef<ISequencerSection> MakeSectionInterface( UMovieSceneSection& SectionObject, UMovieSceneTrack* Track ) override;
 
 protected:
diff --git a/Engine/Source/Editor/MovieSceneTools/Private/TrackEditors/PropertyTrackEditors/BytePropertyTrackEditor.h b/Engine/Source/Editor/MovieSceneTools/Private/TrackEditors/PropertyTrackEditors/BytePropertyTrackEditor.h
--- a/Engine/Source/Editor/MovieSceneTools/Private/TrackEditors/PropertyTrackEditors/BytePropertyTrackEditor.h
+++ b/Engine/Source/Editor/MovieSceneTools/Private/TrackEditors/PropertyTrackEditors/BytePropertyTrackEditor.h
@@ -32,10 +32,10 @@ public:
 
 public:
 
-	// FMovieSceneTrackEditor interface
+	// ISequencerTrackEditor interface
 
-	virtual TSharedRef<ISequencerSection> MakeSectionInterface( UMovieSceneSection& SectionObject, UMovieSceneTrack* Track ) override;
 	virtual UMovieSceneTrack* AddTrack( UMovieScene* FocusedMovieScene, const FGuid& ObjectHandle, TSubclassOf<class UMovieSceneTrack> TrackClass, FName UniqueTypeName ) override;
+	virtual TSharedRef<ISequencerSection> MakeSectionInterface( UMovieSceneSection& SectionObject, UMovieSceneTrack* Track ) override;
 
 protected:
 
diff --git a/Engine/Source/Editor/MovieSceneTools/Private/TrackEditors/PropertyTrackEditors/ColorPropertyTrackEditor.h b/Engine/Source/Editor/MovieSceneTools/Private/TrackEditors/PropertyTrackEditors/ColorPropertyTrackEditor.h
--- a/Engine/Source/Editor/MovieSceneTools/Private/TrackEditors/PropertyTrackEditors/ColorPropertyTrackEditor.h
+++ b/Engine/Source/Editor/MovieSceneTools/Private/TrackEditors/PropertyTrackEditors/ColorPropertyTrackEditor.h
@@ -32,7 +32,7 @@ public:
 
 public:
 
-	// FMovieSceneTrackEditor interface
+	// ISequencerTrackEditor interface
 
 	virtual TSharedRef<ISequencerSection> MakeSectionInterface( UMovieSceneSection& SectionObject, UMovieSceneTrack* Track ) override;
 
diff --git a/Engine/Source/Editor/MovieSceneTools/Private/TrackEditors/PropertyTrackEditors/FloatPropertyTrackEditor.h b/Engine/Source/Editor/MovieSceneTools/Private/TrackEditors/PropertyTrackEditors/FloatPropertyTrackEditor.h
--- a/Engine/Source/Editor/MovieSceneTools/Private/TrackEditors/PropertyTrackEditors/FloatPropertyTrackEditor.h
+++ b/Engine/Source/Editor/MovieSceneTools/Private/TrackEditors/PropertyTrackEditors/FloatPropertyTrackEditor.h
@@ -31,11 +31,13 @@ public:
 
 public:
 
-	// FMovieSceneTrackEditor interface
+	// ISequencerTrackEditor interface
+
 	virtual TSharedRef<ISequencerSection> MakeSectionInterface( UMovieSceneSection& SectionObject, UMovieSceneTrack* Track ) override;
 
 protected:
 
 	// FPropertyTrackEditor interface
+
 	virtual bool TryGenerateKeyFromPropertyChanged( const UMovieSceneTrack* InTrack, const FPropertyChangedParams& PropertyChangedParams, float& OutKey ) override;
 };
diff --git a/Engine/Source/Editor/MovieSceneTools/Private/TrackEditors/PropertyTrackEditors/VectorPropertyTrackEditor.h b/Engine/Source/Editor/MovieSceneTools/Private/TrackEditors/PropertyTrackEditors/VectorPropertyTrackEditor.h
--- a/Engine/Source/Editor/MovieSceneTools/Private/TrackEditors/PropertyTrackEditors/VectorPropertyTrackEditor.h
+++ b/Engine/Source/Editor/MovieSceneTools/Private/TrackEditors/PropertyTrackEditors/VectorPropertyTrackEditor.h
@@ -32,11 +32,13 @@ public:
 
 public:
 
-	// FMovieSceneTrackEditor interface
+	// ISequencerTrackEditor interface
+
 	virtual TSharedRef<ISequencerSection> MakeSectionInterface( UMovieSceneSection& SectionObject, UMovieSceneTrack* Track ) override;
 
 protected:
 
 	// FPropertyTrackEditor interface
+
 	virtual bool TryGenerateKeyFromPropertyChanged( const UMovieSceneTrack* InTrack, const FPropertyChangedParams& PropertyChangedParams, FVectorKey& OutKey ) override;
 };
diff --git a/Engine/Source/Editor/MovieSceneTools/Private/TrackEditors/PropertyTrackEditors/VisibilityPropertyTrackEditor.h b/Engine/Source/Editor/MovieSceneTools/Private/TrackEditors/PropertyTrackEditors/VisibilityPropertyTrackEditor.h
--- a/Engine/Source/Editor/MovieSceneTools/Private/TrackEditors/PropertyTrackEditors/VisibilityPropertyTrackEditor.h
+++ b/Engine/Source/Editor/MovieSceneTools/Private/TrackEditors/PropertyTrackEditors/VisibilityPropertyTrackEditor.h
@@ -30,11 +30,13 @@ public:
 
 public:
 
-	// FMovieSceneTrackEditor interface
+	// ISequencerTrackEditor interface
+
 	virtual TSharedRef<ISequencerSection> MakeSectionInterface( UMovieSceneSection& SectionObject, UMovieSceneTrack* Track ) override;
 
 protected:
 
 	// FPropertyTrackEditor interface
+
 	virtual bool TryGenerateKeyFromPropertyChanged( const UMovieSceneTrack* InTrack, const FPropertyChangedParams& PropertyChangedParams, bool& OutKey ) override;
 };
diff --git a/Engine/Source/Editor/MovieSceneTools/Private/TrackEditors/ShotTrackEditor.h b/Engine/Source/Editor/MovieSceneTools/Private/TrackEditors/ShotTrackEditor.h
--- a/Engine/Source/Editor/MovieSceneTools/Private/TrackEditors/ShotTrackEditor.h
+++ b/Engine/Source/Editor/MovieSceneTools/Private/TrackEditors/ShotTrackEditor.h
@@ -9,13 +9,16 @@
 class FShotTrackEditor : public FMovieSceneTrackEditor
 {
 public:
+
 	/**
 	 * Constructor
 	 *
-	 * @param InSequencer	The sequencer instance to be used by this tool
+	 * @param InSequencer The sequencer instance to be used by this tool
 	 */
 	FShotTrackEditor( TSharedRef<ISequencer> InSequencer );
-	~FShotTrackEditor();
+
+	/** Virtual destructor. */
+	virtual ~FShotTrackEditor();
 
 	/**
 	 * Creates an instance of this class.  Called by a sequencer 
@@ -25,7 +28,10 @@ public:
 	 */
 	static TSharedRef<ISequencerTrackEditor> CreateTrackEditor( TSharedRef<ISequencer> OwningSequencer );
 
-	/** FMovieSceneTrackEditor Interface */
+public:
+
+	// ISequencerTrackEditor interface
+
 	virtual bool SupportsType( TSubclassOf<UMovieSceneTrack> Type ) const override;
 	virtual TSharedRef<ISequencerSection> MakeSectionInterface( UMovieSceneSection& SectionObject, UMovieSceneTrack* Track ) override;
 	virtual void AddKey(const FGuid& ObjectGuid, UObject* AdditionalAsset = NULL) override;
@@ -33,6 +39,7 @@ public:
 	virtual void BuildObjectBindingTrackMenu(FMenuBuilder& MenuBuilder, const FGuid& ObjectBinding, const UClass* ObjectClass) override;
 
 private:
+
 	/** Delegate for AnimatablePropertyChanged in AddKey */
 	void AddKeyInternal(float AutoKeyTime, const FGuid ObjectGuid);
 
@@ -48,6 +55,7 @@ private:
 	void HandleAddShotMenuEntryExecute(FGuid CameraGuid);
 
 private:
+
 	/** The Thumbnail pool which draws all the viewport thumbnails for the director track */
 	TSharedPtr<class FShotThumbnailPool> ThumbnailPool;
 };
diff --git a/Engine/Source/Editor/MovieSceneTools/Private/TrackEditors/SkeletalAnimationTrackEditor.h b/Engine/Source/Editor/MovieSceneTools/Private/TrackEditors/SkeletalAnimationTrackEditor.h
--- a/Engine/Source/Editor/MovieSceneTools/Private/TrackEditors/SkeletalAnimationTrackEditor.h
+++ b/Engine/Source/Editor/MovieSceneTools/Private/TrackEditors/SkeletalAnimationTrackEditor.h
@@ -29,13 +29,13 @@ public:
 
 public:
 
-	// FMovieSceneTrackEditor interface
+	// ISequencerTrackEditor interface
 
-	virtual bool SupportsType( TSubclassOf<UMovieSceneTrack> Type ) const override;
-	virtual TSharedRef<ISequencerSection> MakeSectionInterface( UMovieSceneSection& SectionObject, UMovieSceneTrack* Track ) override;
 	virtual void AddKey(const FGuid& ObjectGuid, UObject* AdditionalAsset) override;
-	virtual bool HandleAssetAdded(UObject* Asset, const FGuid& TargetObjectGuid) override;
 	virtual void BuildObjectBindingTrackMenu(FMenuBuilder& MenuBuilder, const FGuid& ObjectBinding, const UClass* ObjectClass) override;
+	virtual bool HandleAssetAdded(UObject* Asset, const FGuid& TargetObjectGuid) override;
+	virtual TSharedRef<ISequencerSection> MakeSectionInterface( UMovieSceneSection& SectionObject, UMovieSceneTrack* Track ) override;
+	virtual bool SupportsType( TSubclassOf<UMovieSceneTrack> Type ) const override;
 
 private:
 
diff --git a/Engine/Source/Editor/MovieSceneTools/Private/TrackEditors/SlomoTrackEditor.h b/Engine/Source/Editor/MovieSceneTools/Private/TrackEditors/SlomoTrackEditor.h
--- a/Engine/Source/Editor/MovieSceneTools/Private/TrackEditors/SlomoTrackEditor.h
+++ b/Engine/Source/Editor/MovieSceneTools/Private/TrackEditors/SlomoTrackEditor.h
@@ -36,7 +36,7 @@ public:
 
 public:
 
-	// FMovieSceneTrackEditor interface
+	// ISequencerTrackEditor interface
 
 //	virtual void AddKey(const FGuid& ObjectGuid, UObject* AdditionalAsset = nullptr) override;
 //	virtual void BuildObjectBindingTrackMenu(FMenuBuilder& MenuBuilder, const FGuid& ObjectBinding, const UClass* ObjectClass) override;
diff --git a/Engine/Source/Editor/MovieSceneTools/Private/TrackEditors/SubMovieSceneTrackEditor.h b/Engine/Source/Editor/MovieSceneTools/Private/TrackEditors/SubMovieSceneTrackEditor.h
--- a/Engine/Source/Editor/MovieSceneTools/Private/TrackEditors/SubMovieSceneTrackEditor.h
+++ b/Engine/Source/Editor/MovieSceneTools/Private/TrackEditors/SubMovieSceneTrackEditor.h
@@ -19,7 +19,7 @@ public:
 	FSubMovieSceneTrackEditor( TSharedRef<ISequencer> InSequencer );
 
 	/** Virtual destructor. */
-	~FSubMovieSceneTrackEditor() { }
+	virtual ~FSubMovieSceneTrackEditor() { }
 
 	/**
 	 * Creates an instance of this class.  Called by a sequencer 
@@ -31,9 +31,9 @@ public:
 
 public:
 
-	// FMovieSceneTrackEditor interface
+	// ISequencerTrackEditor interface
 
-	virtual bool SupportsType( TSubclassOf<UMovieSceneTrack> Type ) const override;
 	virtual TSharedRef<ISequencerSection> MakeSectionInterface( UMovieSceneSection& SectionObject, UMovieSceneTrack* Track ) override;
 	virtual bool HandleAssetAdded(UObject* Asset, const FGuid& TargetObjectGuid) override;
+	virtual bool SupportsType( TSubclassOf<UMovieSceneTrack> Type ) const override;
 };
diff --git a/Engine/Source/Editor/MovieSceneTools/Private/TrackEditors/TransformTrackEditor.h b/Engine/Source/Editor/MovieSceneTools/Private/TrackEditors/TransformTrackEditor.h
--- a/Engine/Source/Editor/MovieSceneTools/Private/TrackEditors/TransformTrackEditor.h
+++ b/Engine/Source/Editor/MovieSceneTools/Private/TrackEditors/TransformTrackEditor.h
@@ -33,15 +33,15 @@ public:
 
 public:
 
-	// FMovieSceneTrackEditor interface
+	// ISequencerTrackEditor interface
 
-	virtual void OnRelease() override;
-	virtual bool SupportsType( TSubclassOf<UMovieSceneTrack> Type ) const override;
-	virtual TSharedRef<ISequencerSection> MakeSectionInterface( UMovieSceneSection& SectionObject, UMovieSceneTrack* Track ) override;
 	virtual void AddKey( const FGuid& ObjectGuid, UObject* AdditionalAsset = NULL ) override;
 	virtual void BindCommands(TSharedRef<FUICommandList> SequencerCommandBindings) override;
 	virtual void BuildObjectBindingEditButtons(TSharedPtr<SHorizontalBox> EditBox, const FGuid& ObjectBinding, const UClass* ObjectClass) override;
 	virtual void BuildObjectBindingTrackMenu(FMenuBuilder& MenuBuilder, const FGuid& ObjectBinding, const UClass* ObjectClass) override;
+	virtual TSharedRef<ISequencerSection> MakeSectionInterface( UMovieSceneSection& SectionObject, UMovieSceneTrack* Track ) override;
+	virtual void OnRelease() override;
+	virtual bool SupportsType( TSubclassOf<UMovieSceneTrack> Type ) const override;
 
 private:
 
diff --git a/Engine/Source/Editor/MovieSceneTools/Public/PropertyTrackEditor.h b/Engine/Source/Editor/MovieSceneTools/Public/PropertyTrackEditor.h
--- a/Engine/Source/Editor/MovieSceneTools/Public/PropertyTrackEditor.h
+++ b/Engine/Source/Editor/MovieSceneTools/Public/PropertyTrackEditor.h
@@ -94,7 +94,10 @@ public:
 		}
 	}
 
-	/** FMovieSceneTrackEditor Interface */
+public:
+
+	// ISequencerTrackEditor interface
+
 	virtual bool SupportsType( TSubclassOf<UMovieSceneTrack> Type ) const override
 	{
 		return Type == TrackType::StaticClass();
diff --git a/Engine/Source/Editor/UMGEditor/Private/Animation/MarginTrackEditor.h b/Engine/Source/Editor/UMGEditor/Private/Animation/MarginTrackEditor.h
--- a/Engine/Source/Editor/UMGEditor/Private/Animation/MarginTrackEditor.h
+++ b/Engine/Source/Editor/UMGEditor/Private/Animation/MarginTrackEditor.h
@@ -35,7 +35,7 @@ public:
 
 public:
 
-	// FMovieSceneTrackEditor interface
+	// ISequencerTrackEditor interface
 
 	virtual TSharedRef<ISequencerSection> MakeSectionInterface( UMovieSceneSection& SectionObject, UMovieSceneTrack* Track ) override;
 
diff --git a/Engine/Source/Editor/UMGEditor/Private/Animation/Sequencer2DTransformTrackEditor.h b/Engine/Source/Editor/UMGEditor/Private/Animation/Sequencer2DTransformTrackEditor.h
--- a/Engine/Source/Editor/UMGEditor/Private/Animation/Sequencer2DTransformTrackEditor.h
+++ b/Engine/Source/Editor/UMGEditor/Private/Animation/Sequencer2DTransformTrackEditor.h
@@ -33,7 +33,7 @@ public:
 
 public:
 
-	// FMovieSceneTrackEditor interface
+	// ISequencerTrackEditor interface
 
 	virtual TSharedRef<ISequencerSection> MakeSectionInterface( UMovieSceneSection& SectionObject, UMovieSceneTrack* Track ) override;
 
