Sequencer: Added mechanism for track editors to register themselves with the "Add Track" menu; hooked in event, shot and slomo track (not functional yet).

diff --git a/Engine/Source/Editor/MovieSceneTools/Private/MovieSceneToolsModule.cpp b/Engine/Source/Editor/MovieSceneTools/Private/MovieSceneToolsModule.cpp
--- a/Engine/Source/Editor/MovieSceneTools/Private/MovieSceneToolsModule.cpp
+++ b/Engine/Source/Editor/MovieSceneTools/Private/MovieSceneToolsModule.cpp
@@ -19,6 +19,7 @@
 
 #include "TransformTrackEditor.h"
 #include "ShotTrackEditor.h"
+#include "SlomoTrackEditor.h"
 #include "SubMovieSceneTrackEditor.h"
 #include "AudioTrackEditor.h"
 #include "SkeletalAnimationTrackEditor.h"
@@ -59,6 +60,7 @@ public:
 		ParticleTrackCreateEditorHandle = SequencerModule.RegisterTrackEditor_Handle( FOnCreateTrackEditor::CreateStatic( &FParticleTrackEditor::CreateTrackEditor ) );
 		PathTrackCreateEditorHandle = SequencerModule.RegisterTrackEditor_Handle( FOnCreateTrackEditor::CreateStatic( &F3DPathTrackEditor::CreateTrackEditor ) );
 		ShotTrackCreateEditorHandle = SequencerModule.RegisterTrackEditor_Handle( FOnCreateTrackEditor::CreateStatic( &FShotTrackEditor::CreateTrackEditor ) );
+		SlomoTrackCreateEditorHandle = SequencerModule.RegisterTrackEditor_Handle( FOnCreateTrackEditor::CreateStatic( &FSlomoTrackEditor::CreateTrackEditor ) );
 		SubMovieSceneTrackCreateEditorHandle = SequencerModule.RegisterTrackEditor_Handle( FOnCreateTrackEditor::CreateStatic( &FSubMovieSceneTrackEditor::CreateTrackEditor ) );
 		TransformTrackCreateEditorHandle = SequencerModule.RegisterTrackEditor_Handle( FOnCreateTrackEditor::CreateStatic( &F3DTransformTrackEditor::CreateTrackEditor ) );
 		ComponentMaterialTrackCreatEditorHandle = SequencerModule.RegisterTrackEditor_Handle( FOnCreateTrackEditor::CreateStatic( &FComponentMaterialTrackEditor::CreateTrackEditor ) );
@@ -89,6 +91,7 @@ public:
 		SequencerModule.UnRegisterTrackEditor_Handle( ParticleTrackCreateEditorHandle );
 		SequencerModule.UnRegisterTrackEditor_Handle( PathTrackCreateEditorHandle );
 		SequencerModule.UnRegisterTrackEditor_Handle( ShotTrackCreateEditorHandle );
+		SequencerModule.UnRegisterTrackEditor_Handle( SlomoTrackCreateEditorHandle );
 		SequencerModule.UnRegisterTrackEditor_Handle( SubMovieSceneTrackCreateEditorHandle );
 		SequencerModule.UnRegisterTrackEditor_Handle( TransformTrackCreateEditorHandle );
 		SequencerModule.UnRegisterTrackEditor_Handle( ComponentMaterialTrackCreatEditorHandle );
@@ -111,6 +114,7 @@ private:
 	FDelegateHandle ParticleTrackCreateEditorHandle;
 	FDelegateHandle PathTrackCreateEditorHandle;
 	FDelegateHandle ShotTrackCreateEditorHandle;
+	FDelegateHandle SlomoTrackCreateEditorHandle;
 	FDelegateHandle SubMovieSceneTrackCreateEditorHandle;
 	FDelegateHandle TransformTrackCreateEditorHandle;
 	FDelegateHandle ComponentMaterialTrackCreatEditorHandle;
diff --git a/Engine/Source/Editor/MovieSceneTools/Private/TrackEditors/EventTrackEditor.cpp b/Engine/Source/Editor/MovieSceneTools/Private/TrackEditors/EventTrackEditor.cpp
--- a/Engine/Source/Editor/MovieSceneTools/Private/TrackEditors/EventTrackEditor.cpp
+++ b/Engine/Source/Editor/MovieSceneTools/Private/TrackEditors/EventTrackEditor.cpp
@@ -5,6 +5,9 @@
 #include "EventTrackSection.h"
 
 
+#define LOCTEXT_NAMESPACE "FEventTrackEditor"
+
+
 /* FEventTrackEditor static functions
  *****************************************************************************/
 
@@ -22,7 +25,7 @@ FEventTrackEditor::FEventTrackEditor(TSharedRef<ISequencer> InSequencer)
 { }
 
 
-/* FMovieSceneTrackEditor interface
+/* ISequencerTrackEditor interface
  *****************************************************************************/
 
 void FEventTrackEditor::AddKey(const FGuid& ObjectGuid, UObject* AdditionalAsset)
@@ -31,6 +34,19 @@ void FEventTrackEditor::AddKey(const FGuid& ObjectGuid, UObject* AdditionalAsset
 }
 
 
+void FEventTrackEditor::BuildAddTrackMenu(FMenuBuilder& MenuBuilder)
+{
+	MenuBuilder.AddMenuEntry(
+		LOCTEXT("AddEventTrack", "Add Event Track"),
+		LOCTEXT("AddEventTooltip", "Adds a new event track that can trigger events on the timeline."),
+		FSlateIcon(FEditorStyle::GetStyleSetName(), "ActorAnimationEditor.Tracks.Event"),
+		FUIAction(
+			FExecuteAction::CreateRaw(this, &FEventTrackEditor::HandleAddEventTrackMenuEntryExecute)
+		)
+	);
+}
+
+
 void FEventTrackEditor::BuildObjectBindingTrackMenu(FMenuBuilder& MenuBuilder, const FGuid& ObjectBinding, const UClass* ObjectClass)
 {
 	if (!ObjectClass->IsChildOf(AActor::StaticClass()))
@@ -38,7 +54,6 @@ void FEventTrackEditor::BuildObjectBindingTrackMenu(FMenuBuilder& MenuBuilder, c
 		return;
 	}
 
-
 }
 
 
@@ -52,3 +67,15 @@ bool FEventTrackEditor::SupportsType(TSubclassOf<UMovieSceneTrack> Type) const
 {
 	return (Type == UMovieSceneEventTrack::StaticClass());
 }
+
+
+/* FEventTrackEditor callbacks
+ *****************************************************************************/
+
+void FEventTrackEditor::HandleAddEventTrackMenuEntryExecute()
+{
+
+}
+
+
+#undef LOCTEXT_NAMESPACE
diff --git a/Engine/Source/Editor/MovieSceneTools/Private/TrackEditors/EventTrackEditor.h b/Engine/Source/Editor/MovieSceneTools/Private/TrackEditors/EventTrackEditor.h
--- a/Engine/Source/Editor/MovieSceneTools/Private/TrackEditors/EventTrackEditor.h
+++ b/Engine/Source/Editor/MovieSceneTools/Private/TrackEditors/EventTrackEditor.h
@@ -39,7 +39,13 @@ public:
 	// ISequencerTrackEditor interface
 
 	virtual void AddKey(const FGuid& ObjectGuid, UObject* AdditionalAsset = nullptr) override;
+	virtual void BuildAddTrackMenu(FMenuBuilder& MenuBuilder) override;
 	virtual void BuildObjectBindingTrackMenu(FMenuBuilder& MenuBuilder, const FGuid& ObjectBinding, const UClass* ObjectClass) override;
 	virtual TSharedRef<ISequencerSection> MakeSectionInterface(UMovieSceneSection& SectionObject, UMovieSceneTrack* Track) override;
 	virtual bool SupportsType( TSubclassOf<UMovieSceneTrack> Type ) const override;
+
+private:
+
+	/** Callback for executing the "Add Event Track" menu entry. */
+	void HandleAddEventTrackMenuEntryExecute();
 };
diff --git a/Engine/Source/Editor/MovieSceneTools/Private/TrackEditors/ShotTrackEditor.cpp b/Engine/Source/Editor/MovieSceneTools/Private/TrackEditors/ShotTrackEditor.cpp
--- a/Engine/Source/Editor/MovieSceneTools/Private/TrackEditors/ShotTrackEditor.cpp
+++ b/Engine/Source/Editor/MovieSceneTools/Private/TrackEditors/ShotTrackEditor.cpp
@@ -22,6 +22,9 @@
 #include "SInlineEditableTextBlock.h"
 
 
+#define LOCTEXT_NAMESPACE "FEventTrackEditor"
+
+
 namespace ShotTrackConstants
 {
 	// @todo Sequencer Perhaps allow this to be customizable
@@ -34,6 +37,7 @@ namespace ShotTrackConstants
 
 }
 
+
 /////////////////////////////////////////////////
 // FShotThumbnailPool
 
@@ -244,7 +248,7 @@ TSharedRef<SWidget> FShotSection::GenerateSectionWidget()
 		.Padding( FMargin( 15.0f, 7.0f, 0.0f, 0.0f ) )
 		[
 			SNew( SInlineEditableTextBlock )
-			.ToolTipText( NSLOCTEXT("FShotTrackEditor", "RenameShot", "The name of this shot.  Click or hit F2 to rename") )
+			.ToolTipText(LOCTEXT("RenameShot", "The name of this shot.  Click or hit F2 to rename"))
 			.Text( this, &FShotSection::GetShotName )
 			.ShadowOffset( FVector2D(1,1) )
 			.OnTextCommitted(this, &FShotSection::OnRenameShot )
@@ -554,6 +558,19 @@ void FShotTrackEditor::Tick(float DeltaTime)
 	}
 }
 
+void FShotTrackEditor::BuildAddTrackMenu(FMenuBuilder& MenuBuilder)
+{
+	MenuBuilder.AddMenuEntry(
+		LOCTEXT("AddShotTrack", "Add Shot Track"),
+		LOCTEXT("AddShotTooltip", "Adds a shot track, as well as a new shot at the current scrubber location if a camera is selected."),
+		FSlateIcon(FEditorStyle::GetStyleSetName(), "ActorAnimationEditor.Tracks.Shot"),
+		FUIAction(
+			FExecuteAction::CreateRaw(this, &FShotTrackEditor::HandleAddShotTrackMenuEntryExecute),
+			FCanExecuteAction::CreateRaw(this, &FShotTrackEditor::HandleAddShotTrackMenuEntryCanExecute)
+		)
+	);
+}
+
 void FShotTrackEditor::BuildObjectBindingTrackMenu(FMenuBuilder& MenuBuilder, const FGuid& ObjectBinding, const UClass* ObjectClass)
 {
 	if (ObjectClass->IsChildOf(ACameraActor::StaticClass()))
@@ -561,8 +578,8 @@ void FShotTrackEditor::BuildObjectBindingTrackMenu(FMenuBuilder& MenuBuilder, co
 		const TSharedPtr<ISequencer> ParentSequencer = GetSequencer();
 
 		MenuBuilder.AddMenuEntry(
-			NSLOCTEXT("FShotTrackEditor", "AddShot", "Add New Shot"),
-			NSLOCTEXT("FShotTrackEditor", "AddShotTooltip", "Adds a new shot using this camera at the scrubber location."),
+			LOCTEXT("AddShot", "Add New Shot"),
+			LOCTEXT("AddShotTooltip", "Adds a new shot using this camera at the scrubber location."),
 			FSlateIcon(),
 			FUIAction(FExecuteAction::CreateRaw(this, &FShotTrackEditor::HandleAddShotMenuEntryExecute, ObjectBinding))
 			);
@@ -710,7 +727,44 @@ UFactory* FShotTrackEditor::GetAssetFactoryForNewShot( UClass* SequenceClass )
 }
 
 
+UMovieScene* FShotTrackEditor::GetFocusedMovieScene() const
+{
+	UMovieSceneSequence* FocusedSequence = GetSequencer()->GetFocusedMovieSceneSequence();
+
+	return FocusedSequence->GetMovieScene();
+}
+
+
+bool FShotTrackEditor::HandleAddShotTrackMenuEntryCanExecute() const
+{
+	UMovieScene* FocusedMovieScene = GetFocusedMovieScene();
+
+	return ((FocusedMovieScene != nullptr) && (FocusedMovieScene->GetShotTrack() == nullptr));
+}
+
+
+void FShotTrackEditor::HandleAddShotTrackMenuEntryExecute()
+{
+	UMovieScene* FocusedMovieScene = GetFocusedMovieScene();
+
+	if (FocusedMovieScene == nullptr)
+	{
+		return;
+	}
+
+	UMovieSceneTrack* ShotTrack = FocusedMovieScene->GetShotTrack();
+	
+	if (ShotTrack == nullptr)
+	{
+		ShotTrack = FocusedMovieScene->AddShotTrack(UMovieSceneShotTrack::StaticClass());
+	}
+}
+
+
 void FShotTrackEditor::HandleAddShotMenuEntryExecute(FGuid CameraGuid)
 {
 	AddKey(CameraGuid);
 }
+
+
+#undef LOCTEXT_NAMESPACE
diff --git a/Engine/Source/Editor/MovieSceneTools/Private/TrackEditors/ShotTrackEditor.h b/Engine/Source/Editor/MovieSceneTools/Private/TrackEditors/ShotTrackEditor.h
--- a/Engine/Source/Editor/MovieSceneTools/Private/TrackEditors/ShotTrackEditor.h
+++ b/Engine/Source/Editor/MovieSceneTools/Private/TrackEditors/ShotTrackEditor.h
@@ -6,7 +6,8 @@
 /**
  * Tools for director tracks
  */
-class FShotTrackEditor : public FMovieSceneTrackEditor
+class FShotTrackEditor
+	: public FMovieSceneTrackEditor
 {
 public:
 
@@ -32,11 +33,21 @@ public:
 
 	// ISequencerTrackEditor interface
 
-	virtual bool SupportsType( TSubclassOf<UMovieSceneTrack> Type ) const override;
-	virtual TSharedRef<ISequencerSection> MakeSectionInterface( UMovieSceneSection& SectionObject, UMovieSceneTrack* Track ) override;
 	virtual void AddKey(const FGuid& ObjectGuid, UObject* AdditionalAsset = NULL) override;
-	virtual void Tick(float DeltaTime) override;
+	virtual void BuildAddTrackMenu(FMenuBuilder& MenuBuilder) override;
 	virtual void BuildObjectBindingTrackMenu(FMenuBuilder& MenuBuilder, const FGuid& ObjectBinding, const UClass* ObjectClass) override;
+	virtual TSharedRef<ISequencerSection> MakeSectionInterface( UMovieSceneSection& SectionObject, UMovieSceneTrack* Track ) override;
+	virtual bool SupportsType( TSubclassOf<UMovieSceneTrack> Type ) const override;
+	virtual void Tick(float DeltaTime) override;
+
+protected:
+
+	/**
+	 * Gets the movie scene of the currently focused sequence.
+	 *
+	 * @return The focused movie scene, or nullptr if no scene is focused.
+	 */
+	UMovieScene* GetFocusedMovieScene() const;
 
 private:
 
@@ -51,6 +62,12 @@ private:
 
 	UFactory* GetAssetFactoryForNewShot( UClass* SequenceClass );
 
+	/** Callback for determining whether the "Add Shot Track" menu entry can execute. */
+	bool HandleAddShotTrackMenuEntryCanExecute() const;
+
+	/** Callback for executing the "Add Shot Track" menu entry. */
+	void HandleAddShotTrackMenuEntryExecute();
+
 	/** Callback for executing the "Add Shot" menu entry. */
 	void HandleAddShotMenuEntryExecute(FGuid CameraGuid);
 
diff --git a/Engine/Source/Editor/MovieSceneTools/Private/TrackEditors/SlomoTrackEditor.cpp b/Engine/Source/Editor/MovieSceneTools/Private/TrackEditors/SlomoTrackEditor.cpp
--- a/Engine/Source/Editor/MovieSceneTools/Private/TrackEditors/SlomoTrackEditor.cpp
+++ b/Engine/Source/Editor/MovieSceneTools/Private/TrackEditors/SlomoTrackEditor.cpp
@@ -4,10 +4,13 @@
 #include "SlomoTrackEditor.h"
 
 
+#define LOCTEXT_NAMESPACE "FSlomoTrackEditor"
+
+
 /* FSlomoTrackEditor static functions
  *****************************************************************************/
 
-TSharedRef<FMovieSceneTrackEditor> FSlomoTrackEditor::CreateTrackEditor(TSharedRef<ISequencer> InSequencer)
+TSharedRef<ISequencerTrackEditor> FSlomoTrackEditor::CreateTrackEditor(TSharedRef<ISequencer> InSequencer)
 {
 	return MakeShareable(new FSlomoTrackEditor(InSequencer));
 }
@@ -21,10 +24,35 @@ FSlomoTrackEditor::FSlomoTrackEditor(TSharedRef<ISequencer> InSequencer)
 { }
 
 
-/* FMovieSceneTrackEditor interface
+/* ISequencerTrackEditor interface
  *****************************************************************************/
 
+void FSlomoTrackEditor::BuildAddTrackMenu(FMenuBuilder& MenuBuilder)
+{
+	MenuBuilder.AddMenuEntry(
+		LOCTEXT("AddSlomoTrack", "Add Slomo Track"),
+		LOCTEXT("AddSlomoTooltip", "Adds a new slomo track that controls the playback speed of the sequence."),
+		FSlateIcon(FEditorStyle::GetStyleSetName(), "ActorAnimationEditor.Tracks.Slomo"),
+		FUIAction(
+			FExecuteAction::CreateRaw(this, &FSlomoTrackEditor::HandleAddSlomoTrackMenuEntryExecute)
+		)
+	);
+}
+
+
 bool FSlomoTrackEditor::SupportsType(TSubclassOf<UMovieSceneTrack> Type) const
 {
 	return (Type == UMovieSceneSlomoTrack::StaticClass());
 }
+
+
+/* FSlomoTrackEditor callbacks
+ *****************************************************************************/
+
+void FSlomoTrackEditor::HandleAddSlomoTrackMenuEntryExecute()
+{
+
+}
+
+
+#undef LOCTEXT_NAMESPACE
diff --git a/Engine/Source/Editor/MovieSceneTools/Private/TrackEditors/SlomoTrackEditor.h b/Engine/Source/Editor/MovieSceneTools/Private/TrackEditors/SlomoTrackEditor.h
--- a/Engine/Source/Editor/MovieSceneTools/Private/TrackEditors/SlomoTrackEditor.h
+++ b/Engine/Source/Editor/MovieSceneTools/Private/TrackEditors/SlomoTrackEditor.h
@@ -23,7 +23,7 @@ public:
 	 * @param InSequencer The sequencer instance to be used by this tool.
 	 * @return The new instance of this class.
 	 */
-	static TSharedRef<FMovieSceneTrackEditor> CreateTrackEditor(TSharedRef<ISequencer> InSequencer);
+	static TSharedRef<ISequencerTrackEditor> CreateTrackEditor(TSharedRef<ISequencer> InSequencer);
 
 public:
 
@@ -39,7 +39,13 @@ public:
 	// ISequencerTrackEditor interface
 
 //	virtual void AddKey(const FGuid& ObjectGuid, UObject* AdditionalAsset = nullptr) override;
+	virtual void BuildAddTrackMenu(FMenuBuilder& MenuBuilder) override;
 //	virtual void BuildObjectBindingTrackMenu(FMenuBuilder& MenuBuilder, const FGuid& ObjectBinding, const UClass* ObjectClass) override;
 //	virtual TSharedRef<ISequencerSection> MakeSectionInterface(UMovieSceneSection& SectionObject, UMovieSceneTrack* Track) override;
 	virtual bool SupportsType(TSubclassOf<UMovieSceneTrack> Type) const override;
+
+private:
+
+	/** Callback for executing the "Add Event Track" menu entry. */
+	void HandleAddSlomoTrackMenuEntryExecute();
 };
diff --git a/Engine/Source/Editor/Sequencer/Private/SSequencer.cpp b/Engine/Source/Editor/Sequencer/Private/SSequencer.cpp
--- a/Engine/Source/Editor/Sequencer/Private/SSequencer.cpp
+++ b/Engine/Source/Editor/Sequencer/Private/SSequencer.cpp
@@ -565,49 +565,47 @@ TSharedRef<SWidget> SSequencer::MakeToolBar()
 	
 	ToolBarBuilder.BeginSection("Base Commands");
 	{
-		if (OnGetAddMenuContent.IsBound())
-		{
-			ToolBarBuilder.AddWidget(
-				SNew(SComboButton)
-				.OnGetMenuContent(this, &SSequencer::MakeAddMenu)
-				.ButtonStyle(FEditorStyle::Get(), "FlatButton.Success")
-				.ContentPadding(FMargin(2.0f, 1.0f))
-				.HasDownArrow(false)
-				.ButtonContent()
-				[
-					SNew(SHorizontalBox)
+		// Add Track menu
+		ToolBarBuilder.AddWidget(
+			SNew(SComboButton)
+			.OnGetMenuContent(this, &SSequencer::MakeAddMenu)
+			.ButtonStyle(FEditorStyle::Get(), "FlatButton.Success")
+			.ContentPadding(FMargin(2.0f, 1.0f))
+			.HasDownArrow(false)
+			.ButtonContent()
+			[
+				SNew(SHorizontalBox)
 
-					+ SHorizontalBox::Slot()
-					.VAlign(VAlign_Center)
-					.AutoWidth()
-					[
-						SNew(STextBlock)
-						.TextStyle(FEditorStyle::Get(), "NormalText.Important")
-						.Font(FEditorStyle::Get().GetFontStyle("FontAwesome.10"))
-						.Text(FText::FromString(FString(TEXT("\xf067"))) /*fa-plus*/)
-					]
+				+ SHorizontalBox::Slot()
+				.VAlign(VAlign_Center)
+				.AutoWidth()
+				[
+					SNew(STextBlock)
+					.TextStyle(FEditorStyle::Get(), "NormalText.Important")
+					.Font(FEditorStyle::Get().GetFontStyle("FontAwesome.10"))
+					.Text(FText::FromString(FString(TEXT("\xf067"))) /*fa-plus*/)
+				]
 
-					+ SHorizontalBox::Slot()
-					.AutoWidth()
-					.Padding(4, 0, 0, 0)
-					[
-						SNew(STextBlock)
-						.TextStyle(FEditorStyle::Get(), "NormalText.Important")
-						.Text(LOCTEXT("AddButton", "Add"))
-					]
+				+ SHorizontalBox::Slot()
+				.AutoWidth()
+				.Padding(4, 0, 0, 0)
+				[
+					SNew(STextBlock)
+					.TextStyle(FEditorStyle::Get(), "NormalText.Important")
+					.Text(LOCTEXT("AddButton", "Add"))
+				]
 
-					+ SHorizontalBox::Slot()
-					.VAlign(VAlign_Center)
-					.AutoWidth()
-					.Padding(4, 0, 0, 0)
-					[
-						SNew(STextBlock)
-						.TextStyle(FEditorStyle::Get(), "NormalText.Important")
-						.Font(FEditorStyle::Get().GetFontStyle("FontAwesome.10"))
-						.Text(FText::FromString(FString(TEXT("\xf0d7"))) /*fa-caret-down*/)
-					]
-				]);
-		}
+				+ SHorizontalBox::Slot()
+				.VAlign(VAlign_Center)
+				.AutoWidth()
+				.Padding(4, 0, 0, 0)
+				[
+					SNew(STextBlock)
+					.TextStyle(FEditorStyle::Get(), "NormalText.Important")
+					.Font(FEditorStyle::Get().GetFontStyle("FontAwesome.10"))
+					.Text(FText::FromString(FString(TEXT("\xf0d7"))) /*fa-caret-down*/)
+				]
+			]);
 
 		// General 
 		if( Sequencer.Pin()->IsLevelEditorSequencer() )
@@ -756,7 +754,21 @@ TSharedRef<SWidget> SSequencer::MakeToolBar()
 
 TSharedRef<SWidget> SSequencer::MakeAddMenu()
 {
-	return OnGetAddMenuContent.Execute(Sequencer.Pin().ToSharedRef());
+	FMenuBuilder MenuBuilder(true, nullptr);
+	{
+		// let toolkits populate the menu
+		OnGetAddMenuContent.ExecuteIfBound(MenuBuilder, Sequencer.Pin().ToSharedRef());
+
+		// let track editors populate the menu
+		TSharedPtr<FSequencer> PinnedSequencer = Sequencer.Pin();
+
+		if (PinnedSequencer.IsValid())
+		{
+			PinnedSequencer->BuildAddTrackMenu(MenuBuilder);
+		}
+	}
+
+	return MenuBuilder.MakeWidget();
 }
 
 TSharedRef<SWidget> SSequencer::MakeGeneralMenu()
diff --git a/Engine/Source/Editor/Sequencer/Private/Sequencer.cpp b/Engine/Source/Editor/Sequencer/Private/Sequencer.cpp
--- a/Engine/Source/Editor/Sequencer/Private/Sequencer.cpp
+++ b/Engine/Source/Editor/Sequencer/Private/Sequencer.cpp
@@ -2287,6 +2287,14 @@ void FSequencer::BindSequencerCommands()
 	}
 }
 
+void FSequencer::BuildAddTrackMenu(class FMenuBuilder& MenuBuilder)
+{
+	for (int32 i = 0; i < TrackEditors.Num(); ++i)
+	{
+		TrackEditors[i]->BuildAddTrackMenu(MenuBuilder);
+	}
+}
+
 void FSequencer::BuildObjectBindingContextMenu(FMenuBuilder& MenuBuilder, const FGuid& ObjectBinding, const UClass* ObjectClass, FObjectBindingNode* ObjectBindingNode)
 {
 	if (IsLevelEditorSequencer())
diff --git a/Engine/Source/Editor/Sequencer/Private/Sequencer.h b/Engine/Source/Editor/Sequencer/Private/Sequencer.h
--- a/Engine/Source/Editor/Sequencer/Private/Sequencer.h
+++ b/Engine/Source/Editor/Sequencer/Private/Sequencer.h
@@ -201,6 +201,13 @@ public:
 	/** Gets the command bindings for the sequencer */
 	TSharedPtr<FUICommandList> GetCommandBindings() { return SequencerCommandBindings; }
 
+	/**
+	 * Builds up the sequencer's "Add Track" menu.
+	 *
+	 * @param MenuBuilder The menu builder to add things to.
+	 */
+	void BuildAddTrackMenu(class FMenuBuilder& MenuBuilder);
+
 	/**
 	 * Builds up the context menu for object binding nodes in the outliner
 	 * 
diff --git a/Engine/Source/Editor/Sequencer/Public/ISequencerModule.h b/Engine/Source/Editor/Sequencer/Public/ISequencerModule.h
--- a/Engine/Source/Editor/Sequencer/Public/ISequencerModule.h
+++ b/Engine/Source/Editor/Sequencer/Public/ISequencerModule.h
@@ -24,7 +24,7 @@ namespace SequencerMenuExtensionPoints
 DECLARE_DELEGATE_RetVal_OneParam(TSharedRef<ISequencerTrackEditor>, FOnCreateTrackEditor, TSharedRef<ISequencer>);
 
 /** A delegate that is executed when adding menu content. */
-DECLARE_DELEGATE_RetVal_OneParam(TSharedRef<SWidget>, FOnGetAddMenuContent, TSharedRef<ISequencer>);
+DECLARE_DELEGATE_TwoParams(FOnGetAddMenuContent, FMenuBuilder& /*MenuBuilder*/, TSharedRef<ISequencer>);
 
 
 /**
diff --git a/Engine/Source/Editor/Sequencer/Public/ISequencerTrackEditor.h b/Engine/Source/Editor/Sequencer/Public/ISequencerTrackEditor.h
--- a/Engine/Source/Editor/Sequencer/Public/ISequencerTrackEditor.h
+++ b/Engine/Source/Editor/Sequencer/Public/ISequencerTrackEditor.h
@@ -21,7 +21,7 @@ public:
 	 * @param ObjectGuid The Guid of the object that we are adding a key to.
 	 * @param AdditionalAsset An optional asset that can be added with the key.
 	 */
-	virtual void AddKey(const FGuid& ObjectGuid, UObject* AdditionalAsset = NULL) = 0;
+	virtual void AddKey(const FGuid& ObjectGuid, UObject* AdditionalAsset = nullptr) = 0;
 	
 	/**
 	 * Add a new track to the sequence.
@@ -35,6 +35,13 @@ public:
 	*/
 	virtual void BindCommands(TSharedRef<FUICommandList> SequencerCommandBindings) = 0;
 
+	/**
+	 * Builds up the sequencer's "Add Track" menu.
+	 *
+	 * @param MenuBuilder The menu builder to change.
+	 */
+	virtual void BuildAddTrackMenu(FMenuBuilder& MenuBuilder) = 0;
+
 	/**
 	 * Builds up the object binding context menu for the outliner.
 	 *
diff --git a/Engine/Source/Editor/Sequencer/Public/MovieSceneTrackEditor.h b/Engine/Source/Editor/Sequencer/Public/MovieSceneTrackEditor.h
--- a/Engine/Source/Editor/Sequencer/Public/MovieSceneTrackEditor.h
+++ b/Engine/Source/Editor/Sequencer/Public/MovieSceneTrackEditor.h
@@ -162,6 +162,7 @@ public:
 	}
 
 	virtual void BindCommands(TSharedRef<FUICommandList> SequencerCommandBindings) override { }
+	virtual void BuildAddTrackMenu(FMenuBuilder& MenuBuilder) override { }
 	virtual void BuildObjectBindingContextMenu(FMenuBuilder& MenuBuilder, const FGuid& ObjectBinding, const UClass* ObjectClass) override { }
 	virtual void BuildObjectBindingEditButtons(TSharedPtr<SHorizontalBox> EditBox, const FGuid& ObjectBinding, const UClass* ObjectClass) override { }
 	virtual void BuildObjectBindingTrackMenu(FMenuBuilder& MenuBuilder, const FGuid& ObjectBinding, const UClass* ObjectClass) override { }
diff --git a/Engine/Source/Editor/UMGEditor/Private/WidgetBlueprintEditor.cpp b/Engine/Source/Editor/UMGEditor/Private/WidgetBlueprintEditor.cpp
--- a/Engine/Source/Editor/UMGEditor/Private/WidgetBlueprintEditor.cpp
+++ b/Engine/Source/Editor/UMGEditor/Private/WidgetBlueprintEditor.cpp
@@ -794,7 +794,7 @@ void GetBindableObjects(UWidget* RootWidget, TArray<FObjectAndDisplayName>& Bind
 	}
 }
 
-TSharedRef<SWidget> FWidgetBlueprintEditor::OnGetAnimationAddMenuContent(TSharedRef<ISequencer> InSequencer)
+void FWidgetBlueprintEditor::OnGetAnimationAddMenuContent(FMenuBuilder& MenuBuilder, TSharedRef<ISequencer> InSequencer)
 {
 	TArray<FObjectAndDisplayName> BindableObjects;
 	{
@@ -802,8 +802,6 @@ TSharedRef<SWidget> FWidgetBlueprintEditor::OnGetAnimationAddMenuContent(TShared
 		BindableObjects.Sort();
 	}
 
-	FMenuBuilder AddMenuBuilder(true, nullptr);
-
 	if (CurrentAnimation.IsValid())
 	{
 		for (FObjectAndDisplayName& BindableObject : BindableObjects)
@@ -812,12 +810,10 @@ TSharedRef<SWidget> FWidgetBlueprintEditor::OnGetAnimationAddMenuContent(TShared
 			if (BoundObjectGuid.IsValid() == false)
 			{
 				FUIAction AddMenuAction(FExecuteAction::CreateSP(this, &FWidgetBlueprintEditor::AddObjectToAnimation, BindableObject.Object));
-				AddMenuBuilder.AddMenuEntry(BindableObject.DisplayName, FText(), FSlateIcon(), AddMenuAction);
+				MenuBuilder.AddMenuEntry(BindableObject.DisplayName, FText(), FSlateIcon(), AddMenuAction);
 			}
 		}
 	}
-
-	return AddMenuBuilder.MakeWidget();
 }
 
 void FWidgetBlueprintEditor::AddObjectToAnimation(UObject* ObjectToAnimate)
diff --git a/Engine/Source/Editor/UMGEditor/Private/WidgetBlueprintEditor.h b/Engine/Source/Editor/UMGEditor/Private/WidgetBlueprintEditor.h
--- a/Engine/Source/Editor/UMGEditor/Private/WidgetBlueprintEditor.h
+++ b/Engine/Source/Editor/UMGEditor/Private/WidgetBlueprintEditor.h
@@ -176,7 +176,7 @@ private:
 	void UpdatePreview(UBlueprint* InBlueprint, bool bInForceFullUpdate);
 
 	/** Populates the sequencer add menu. */
-	TSharedRef<SWidget> OnGetAnimationAddMenuContent(TSharedRef<ISequencer> Sequencer);
+	void OnGetAnimationAddMenuContent(FMenuBuilder& MenuBuilder, TSharedRef<ISequencer> Sequencer);
 
 	/** Adds the supplied UObject to the current animation. */
 	void AddObjectToAnimation(UObject* ObjectToAnimate);
