Sequencer: Quick hack to remove certain track editors from UMG

diff --git a/Engine/Source/Editor/MovieSceneTools/Private/TrackEditors/EventTrackEditor.cpp b/Engine/Source/Editor/MovieSceneTools/Private/TrackEditors/EventTrackEditor.cpp
--- a/Engine/Source/Editor/MovieSceneTools/Private/TrackEditors/EventTrackEditor.cpp
+++ b/Engine/Source/Editor/MovieSceneTools/Private/TrackEditors/EventTrackEditor.cpp
@@ -36,6 +36,13 @@ void FEventTrackEditor::AddKey(const FGuid& ObjectGuid, UObject* AdditionalAsset
 
 void FEventTrackEditor::BuildAddTrackMenu(FMenuBuilder& MenuBuilder)
 {
+	UMovieSceneSequence* RootMovieSceneSequence = GetSequencer()->GetRootMovieSceneSequence();
+
+	if ((RootMovieSceneSequence == nullptr) || (RootMovieSceneSequence->GetClass()->GetName() != TEXT("ActorAnimationInstance")))
+	{
+		return;
+	}
+
 	MenuBuilder.AddMenuEntry(
 		LOCTEXT("AddEventTrack", "Add Event Track"),
 		LOCTEXT("AddEventTooltip", "Adds a new event track that can trigger events on the timeline."),
diff --git a/Engine/Source/Editor/MovieSceneTools/Private/TrackEditors/ShotTrackEditor.cpp b/Engine/Source/Editor/MovieSceneTools/Private/TrackEditors/ShotTrackEditor.cpp
--- a/Engine/Source/Editor/MovieSceneTools/Private/TrackEditors/ShotTrackEditor.cpp
+++ b/Engine/Source/Editor/MovieSceneTools/Private/TrackEditors/ShotTrackEditor.cpp
@@ -560,6 +560,13 @@ void FShotTrackEditor::Tick(float DeltaTime)
 
 void FShotTrackEditor::BuildAddTrackMenu(FMenuBuilder& MenuBuilder)
 {
+	UMovieSceneSequence* RootMovieSceneSequence = GetSequencer()->GetRootMovieSceneSequence();
+
+	if ((RootMovieSceneSequence == nullptr) || (RootMovieSceneSequence->GetClass()->GetName() != TEXT("ActorAnimationInstance")))
+	{
+		return;
+	}
+
 	MenuBuilder.AddMenuEntry(
 		LOCTEXT("AddShotTrack", "Add Shot Track"),
 		LOCTEXT("AddShotTooltip", "Adds a shot track, as well as a new shot at the current scrubber location if a camera is selected."),
diff --git a/Engine/Source/Editor/MovieSceneTools/Private/TrackEditors/SlomoTrackEditor.cpp b/Engine/Source/Editor/MovieSceneTools/Private/TrackEditors/SlomoTrackEditor.cpp
--- a/Engine/Source/Editor/MovieSceneTools/Private/TrackEditors/SlomoTrackEditor.cpp
+++ b/Engine/Source/Editor/MovieSceneTools/Private/TrackEditors/SlomoTrackEditor.cpp
@@ -31,6 +31,13 @@ FSlomoTrackEditor::FSlomoTrackEditor(TSharedRef<ISequencer> InSequencer)
 
 void FSlomoTrackEditor::BuildAddTrackMenu(FMenuBuilder& MenuBuilder)
 {
+	UMovieSceneSequence* RootMovieSceneSequence = GetSequencer()->GetRootMovieSceneSequence();
+
+	if ((RootMovieSceneSequence == nullptr) || (RootMovieSceneSequence->GetClass()->GetName() != TEXT("ActorAnimationInstance")))
+	{
+		return;
+	}
+
 	MenuBuilder.AddMenuEntry(
 		LOCTEXT("AddSlomoTrack", "Add Slomo Track"),
 		LOCTEXT("AddSlomoTooltip", "Adds a new slomo track that controls the playback speed of the sequence."),
