Sequencer: Made event and slomo tracks visible in sequencer

diff --git a/Engine/Source/Editor/MovieSceneTools/Private/TrackEditors/EventTrackEditor.cpp b/Engine/Source/Editor/MovieSceneTools/Private/TrackEditors/EventTrackEditor.cpp
--- a/Engine/Source/Editor/MovieSceneTools/Private/TrackEditors/EventTrackEditor.cpp
+++ b/Engine/Source/Editor/MovieSceneTools/Private/TrackEditors/EventTrackEditor.cpp
@@ -3,6 +3,7 @@
 #include "MovieSceneToolsPrivatePCH.h"
 #include "EventTrackEditor.h"
 #include "EventTrackSection.h"
+#include "MovieSceneEventTrack.h"
 
 
 #define LOCTEXT_NAMESPACE "FEventTrackEditor"
@@ -81,7 +82,14 @@ bool FEventTrackEditor::SupportsType(TSubclassOf<UMovieSceneTrack> Type) const
 
 void FEventTrackEditor::HandleAddEventTrackMenuEntryExecute()
 {
+	UMovieSceneSequence* FocusedSequence = GetSequencer()->GetFocusedMovieSceneSequence();
+	UMovieScene* MovieScene = FocusedSequence->GetMovieScene();
 
+	if (MovieScene != nullptr)
+	{
+		MovieScene->AddMasterTrack(UMovieSceneEventTrack::StaticClass(), FName("Events"));
+		GetSequencer()->NotifyMovieSceneDataChanged();
+	}
 }
 
 
diff --git a/Engine/Source/Editor/MovieSceneTools/Private/TrackEditors/SlomoTrackEditor.cpp b/Engine/Source/Editor/MovieSceneTools/Private/TrackEditors/SlomoTrackEditor.cpp
--- a/Engine/Source/Editor/MovieSceneTools/Private/TrackEditors/SlomoTrackEditor.cpp
+++ b/Engine/Source/Editor/MovieSceneTools/Private/TrackEditors/SlomoTrackEditor.cpp
@@ -70,7 +70,8 @@ void FSlomoTrackEditor::HandleAddSlomoTrackMenuEntryExecute()
 
 	if (MovieScene != nullptr)
 	{
-		SlomoTrack = MovieScene->AddMasterTrack(UMovieSceneSlomoTrack::StaticClass());
+		SlomoTrack = MovieScene->AddMasterTrack(UMovieSceneSlomoTrack::StaticClass(), FName("Play Rate"));
+		GetSequencer()->NotifyMovieSceneDataChanged();
 	}
 }
 
