Sequencer: Proper names for event and slomo tracks

diff --git a/Engine/Source/Editor/MovieSceneTools/Private/TrackEditors/EventTrackEditor.cpp b/Engine/Source/Editor/MovieSceneTools/Private/TrackEditors/EventTrackEditor.cpp
--- a/Engine/Source/Editor/MovieSceneTools/Private/TrackEditors/EventTrackEditor.cpp
+++ b/Engine/Source/Editor/MovieSceneTools/Private/TrackEditors/EventTrackEditor.cpp
@@ -87,7 +87,7 @@ void FEventTrackEditor::HandleAddEventTrackMenuEntryExecute()
 
 	if (MovieScene != nullptr)
 	{
-		MovieScene->AddMasterTrack(UMovieSceneEventTrack::StaticClass(), FName("Events"));
+		MovieScene->AddMasterTrack(UMovieSceneEventTrack::StaticClass());
 		GetSequencer()->NotifyMovieSceneDataChanged();
 	}
 }
diff --git a/Engine/Source/Editor/MovieSceneTools/Private/TrackEditors/SlomoTrackEditor.cpp b/Engine/Source/Editor/MovieSceneTools/Private/TrackEditors/SlomoTrackEditor.cpp
--- a/Engine/Source/Editor/MovieSceneTools/Private/TrackEditors/SlomoTrackEditor.cpp
+++ b/Engine/Source/Editor/MovieSceneTools/Private/TrackEditors/SlomoTrackEditor.cpp
@@ -70,7 +70,7 @@ void FSlomoTrackEditor::HandleAddSlomoTrackMenuEntryExecute()
 
 	if (MovieScene != nullptr)
 	{
-		SlomoTrack = MovieScene->AddMasterTrack(UMovieSceneSlomoTrack::StaticClass(), FName("Play Rate"));
+		SlomoTrack = MovieScene->AddMasterTrack(UMovieSceneSlomoTrack::StaticClass());
 		GetSequencer()->NotifyMovieSceneDataChanged();
 	}
 }
diff --git a/Engine/Source/Runtime/MovieSceneTracks/Private/Tracks/MovieSceneSlomoTrack.cpp b/Engine/Source/Runtime/MovieSceneTracks/Private/Tracks/MovieSceneSlomoTrack.cpp
--- a/Engine/Source/Runtime/MovieSceneTracks/Private/Tracks/MovieSceneSlomoTrack.cpp
+++ b/Engine/Source/Runtime/MovieSceneTracks/Private/Tracks/MovieSceneSlomoTrack.cpp
@@ -4,3 +4,12 @@
 #include "MovieSceneSlomoSection.h"
 #include "MovieSceneSlomoTrack.h"
 #include "MovieSceneSlomoTrackInstance.h"
+
+
+/* UMovieSceneEventTrack interface
+ *****************************************************************************/
+
+FName UMovieSceneSlomoTrack::GetTrackName() const
+{
+	return FName("Play Rate");
+}
diff --git a/Engine/Source/Runtime/MovieSceneTracks/Public/Tracks/MovieSceneEventTrack.h b/Engine/Source/Runtime/MovieSceneTracks/Public/Tracks/MovieSceneEventTrack.h
--- a/Engine/Source/Runtime/MovieSceneTracks/Public/Tracks/MovieSceneEventTrack.h
+++ b/Engine/Source/Runtime/MovieSceneTracks/Public/Tracks/MovieSceneEventTrack.h
@@ -15,6 +15,13 @@ class UMovieSceneEventTrack
 {
 	GENERATED_BODY()
 
+public:
+
+	/** Default constructor. */
+	UMovieSceneEventTrack()
+		: TrackName("Events")
+	{ }
+
 public:
 
 	/**
diff --git a/Engine/Source/Runtime/MovieSceneTracks/Public/Tracks/MovieSceneSlomoTrack.h b/Engine/Source/Runtime/MovieSceneTracks/Public/Tracks/MovieSceneSlomoTrack.h
--- a/Engine/Source/Runtime/MovieSceneTracks/Public/Tracks/MovieSceneSlomoTrack.h
+++ b/Engine/Source/Runtime/MovieSceneTracks/Public/Tracks/MovieSceneSlomoTrack.h
@@ -14,4 +14,10 @@ class UMovieSceneSlomoTrack
 	: public UMovieSceneFloatTrack
 {
 	GENERATED_BODY()
+
+public:
+
+	// UMovieSceneTrack interface
+
+	virtual FName GetTrackName() const override;
 };
