Sequencer: Renamed "Add Slomo Track" to "Add Play Rate Track"

diff --git a/Engine/Source/Editor/MovieSceneTools/Private/TrackEditors/SlomoTrackEditor.cpp b/Engine/Source/Editor/MovieSceneTools/Private/TrackEditors/SlomoTrackEditor.cpp
--- a/Engine/Source/Editor/MovieSceneTools/Private/TrackEditors/SlomoTrackEditor.cpp
+++ b/Engine/Source/Editor/MovieSceneTools/Private/TrackEditors/SlomoTrackEditor.cpp
@@ -39,8 +39,8 @@ void FSlomoTrackEditor::BuildAddTrackMenu(FMenuBuilder& MenuBuilder)
 	}
 
 	MenuBuilder.AddMenuEntry(
-		LOCTEXT("AddSlomoTrack", "Add Slomo Track"),
-		LOCTEXT("AddSlomoTooltip", "Adds a new slomo track that controls the playback speed of the sequence."),
+		LOCTEXT("AddPlayRateTrack", "Add Play Rate Track"),
+		LOCTEXT("AddPlayRateTrackTooltip", "Adds a new track that controls the playback rate of the sequence."),
 		FSlateIcon(FEditorStyle::GetStyleSetName(), "ActorAnimationEditor.Tracks.Slomo"),
 		FUIAction(
 			FExecuteAction::CreateRaw(this, &FSlomoTrackEditor::HandleAddSlomoTrackMenuEntryExecute)
