Sequencer: Implemented and integrated key area for event tracks

diff --git a/Engine/Source/Editor/MovieSceneTools/Private/MovieSceneToolHelpers.cpp b/Engine/Source/Editor/MovieSceneTools/Private/MovieSceneToolHelpers.cpp
--- a/Engine/Source/Editor/MovieSceneTools/Private/MovieSceneToolHelpers.cpp
+++ b/Engine/Source/Editor/MovieSceneTools/Private/MovieSceneToolHelpers.cpp
@@ -16,6 +16,7 @@
 #include "SEnumCurveKeyEditor.h"
 #include "SBoolCurveKeyEditor.h"
 
+
 TArray<FKeyHandle> FFloatCurveKeyArea::AddKeyUnique(float Time, EMovieSceneKeyInterpolation InKeyInterpolation, float TimeToCopyFrom)
 {
 	TArray<FKeyHandle> AddedKeyHandles;
@@ -89,6 +90,7 @@ TArray<FKeyHandle> FFloatCurveKeyArea::AddKeyUnique(float Time, EMovieSceneKeyIn
 	return AddedKeyHandles;
 }
 
+
 TSharedRef<SWidget> FFloatCurveKeyArea::CreateKeyEditor(ISequencer* Sequencer)
 {
 	return SNew(SFloatCurveKeyEditor)
@@ -97,6 +99,32 @@ TSharedRef<SWidget> FFloatCurveKeyArea::CreateKeyEditor(ISequencer* Sequencer)
 		.Curve(Curve);
 };
 
+
+TArray<FKeyHandle> FNameCurveKeyArea::AddKeyUnique(float Time, EMovieSceneKeyInterpolation InKeyInterpolation, float TimeToCopyFrom)
+{
+	TArray<FKeyHandle> AddedKeyHandles;
+	FKeyHandle CurrentKey = Curve.FindKey(Time);
+
+	if (Curve.IsKeyHandleValid(CurrentKey) == false)
+	{
+		if (OwningSection->GetStartTime() > Time)
+		{
+			OwningSection->SetStartTime(Time);
+		}
+
+		if (OwningSection->GetEndTime() < Time)
+		{
+			OwningSection->SetEndTime(Time);
+		}
+
+		Curve.AddKey(Time, NAME_None, CurrentKey);
+		AddedKeyHandles.Add(CurrentKey);
+	}
+
+	return AddedKeyHandles;
+}
+
+
 TArray<FKeyHandle> FIntegralKeyArea::AddKeyUnique(float Time, EMovieSceneKeyInterpolation InKeyInterpolation, float TimeToCopyFrom)
 {
 	TArray<FKeyHandle> AddedKeyHandles;
@@ -126,6 +154,7 @@ TArray<FKeyHandle> FIntegralKeyArea::AddKeyUnique(float Time, EMovieSceneKeyInte
 	return AddedKeyHandles;
 }
 
+
 TSharedRef<SWidget> FIntegralKeyArea::CreateKeyEditor(ISequencer* Sequencer)
 {
 	return SNew(SIntegralCurveKeyEditor)
@@ -134,6 +163,7 @@ TSharedRef<SWidget> FIntegralKeyArea::CreateKeyEditor(ISequencer* Sequencer)
 		.Curve(&Curve);
 };
 
+
 TSharedRef<SWidget> FEnumKeyArea::CreateKeyEditor(ISequencer* Sequencer)
 {
 	return SNew(SEnumCurveKeyEditor)
@@ -143,6 +173,7 @@ TSharedRef<SWidget> FEnumKeyArea::CreateKeyEditor(ISequencer* Sequencer)
 		.Enum(Enum);
 };
 
+
 TSharedRef<SWidget> FBoolKeyArea::CreateKeyEditor(ISequencer* Sequencer)
 {
 	return SNew(SBoolCurveKeyEditor)
@@ -151,6 +182,7 @@ TSharedRef<SWidget> FBoolKeyArea::CreateKeyEditor(ISequencer* Sequencer)
 		.Curve(&Curve);
 };
 
+
 void MovieSceneToolHelpers::TrimSection(const TSet<TWeakObjectPtr<UMovieSceneSection>>& Sections, float Time, bool bTrimLeft)
 {
 	for (auto Section : Sections)
@@ -171,6 +203,7 @@ void MovieSceneToolHelpers::TrimSection(const TSet<TWeakObjectPtr<UMovieSceneSec
 	}
 }
 
+
 void MovieSceneToolHelpers::SplitSection(const TSet<TWeakObjectPtr<UMovieSceneSection>>& Sections, float Time)
 {
 	for (auto Section : Sections)
@@ -194,4 +227,4 @@ void MovieSceneToolHelpers::SplitSection(const TSet<TWeakObjectPtr<UMovieSceneSe
 			NewSection->SetEndTime(SectionEndTime);
 		}
 	}
-}
\ No newline at end of file
+}
diff --git a/Engine/Source/Editor/MovieSceneTools/Private/Sections/EventTrackSection.cpp b/Engine/Source/Editor/MovieSceneTools/Private/Sections/EventTrackSection.cpp
--- a/Engine/Source/Editor/MovieSceneTools/Private/Sections/EventTrackSection.cpp
+++ b/Engine/Source/Editor/MovieSceneTools/Private/Sections/EventTrackSection.cpp
@@ -12,7 +12,7 @@
  *****************************************************************************/
 
 FEventTrackSection::FEventTrackSection(UMovieSceneSection& InSection, TSharedPtr<ISequencer> InSequencer)
-	: Section(&InSection)
+	: Section(Cast<UMovieSceneEventSection>(&InSection))
 	, Sequencer(InSequencer)
 { }
 
@@ -46,7 +46,7 @@ FText FEventTrackSection::GetSectionTitle() const
 
 void FEventTrackSection::GenerateSectionLayout(class ISectionLayoutBuilder& LayoutBuilder) const
 {
-
+	LayoutBuilder.SetSectionAsKeyArea(MakeShareable(new FNameCurveKeyArea(Section->GetEventCurve(), Section)));
 }
 
 
diff --git a/Engine/Source/Editor/MovieSceneTools/Private/Sections/EventTrackSection.h b/Engine/Source/Editor/MovieSceneTools/Private/Sections/EventTrackSection.h
--- a/Engine/Source/Editor/MovieSceneTools/Private/Sections/EventTrackSection.h
+++ b/Engine/Source/Editor/MovieSceneTools/Private/Sections/EventTrackSection.h
@@ -3,7 +3,7 @@
 #pragma once
 
 
-class UMovieSceneSection;
+class UMovieSceneEventSection;
 
 
 /**
@@ -44,7 +44,7 @@ public:
 private:
 
 	/** The section being visualized. */
-	UMovieSceneSection* Section;
+	UMovieSceneEventSection* Section;
 
 	/** The parent sequencer we are a part of. */
 	TWeakPtr<ISequencer> Sequencer;
diff --git a/Engine/Source/Editor/MovieSceneTools/Public/MovieSceneToolHelpers.h b/Engine/Source/Editor/MovieSceneTools/Public/MovieSceneToolHelpers.h
--- a/Engine/Source/Editor/MovieSceneTools/Public/MovieSceneToolHelpers.h
+++ b/Engine/Source/Editor/MovieSceneTools/Public/MovieSceneToolHelpers.h
@@ -6,6 +6,7 @@
 #include "ISequencerSection.h"
 #include "IKeyArea.h"
 
+
 /**
  * A key area for float keys
  */
@@ -131,7 +132,6 @@ private:
 };
 
 
-
 /** A key area for integral keys */
 class MOVIESCENETOOLS_API FIntegralKeyArea : public IKeyArea
 {
@@ -219,6 +219,7 @@ protected:
 	UMovieSceneSection* OwningSection;
 };
 
+
 /** A key area for displaying and editing intragral curves representing enums. */
 class FEnumKeyArea : public FIntegralKeyArea
 {
@@ -240,6 +241,7 @@ private:
 	UEnum* Enum;
 };
 
+
 /** A key area for displaying and editing integral curves representing bools. */
 class FBoolKeyArea : public FIntegralKeyArea
 {
@@ -256,6 +258,124 @@ public:
 	virtual TSharedRef<SWidget> CreateKeyEditor(ISequencer* Sequencer) override;
 };
 
+
+/**
+ * A key area for FName curves.
+ */
+class MOVIESCENETOOLS_API FNameCurveKeyArea
+	: public IKeyArea
+{
+public:
+
+	/**
+	 * Create and initialize a new instance.
+	 *
+	 * @param InCurve The curve to assign to this key area.
+	 * @param InOwningSection The section that owns this key area.
+	 */
+	FNameCurveKeyArea(FNameCurve& InCurve, UMovieSceneSection* InOwningSection)
+		: Curve(InCurve)
+		, OwningSection(InOwningSection)
+	{ }
+
+public:
+
+	// IKeyArea interface
+
+	virtual TArray<FKeyHandle> GetUnsortedKeyHandles() const override
+	{
+		TArray<FKeyHandle> OutKeyHandles;
+
+		for (auto It(Curve.GetKeyHandleIterator()); It; ++It)
+		{
+			OutKeyHandles.Add(It.Key());
+		}
+
+		return OutKeyHandles;
+	}
+
+	virtual void SetKeyTime(FKeyHandle KeyHandle, float NewKeyTime) const override
+	{
+		Curve.SetKeyTime(KeyHandle, NewKeyTime);
+	}
+
+	virtual float GetKeyTime(FKeyHandle KeyHandle) const override
+	{
+		return Curve.GetKeyTime(KeyHandle);
+	}
+
+	virtual FKeyHandle MoveKey( FKeyHandle KeyHandle, float DeltaPosition ) override
+	{
+		return Curve.SetKeyTime(KeyHandle, Curve.GetKeyTime(KeyHandle) + DeltaPosition);
+	}
+	
+	virtual void DeleteKey(FKeyHandle KeyHandle) override
+	{
+		Curve.DeleteKey(KeyHandle);
+	}
+
+	virtual void SetKeyInterpMode(FKeyHandle KeyHandle, ERichCurveInterpMode InterpMode) override
+	{
+		// do nothing
+	}
+
+	virtual ERichCurveInterpMode GetKeyInterpMode(FKeyHandle Keyhandle) const override
+	{
+		return RCIM_None;
+	}
+
+	virtual void SetKeyTangentMode(FKeyHandle KeyHandle, ERichCurveTangentMode TangentMode) override
+	{
+		// do nothing
+	}
+
+	virtual ERichCurveTangentMode GetKeyTangentMode(FKeyHandle KeyHandle) const override
+	{
+		return RCTM_None;
+	}
+
+	virtual void SetExtrapolationMode(ERichCurveExtrapolation ExtrapMode, bool bPreInfinity) override
+	{
+		// do nothing
+	}
+
+	virtual ERichCurveExtrapolation GetExtrapolationMode(bool bPreInfinity) const override
+	{
+		return RCCE_None;
+	}
+
+	virtual TArray<FKeyHandle> AddKeyUnique(float Time, EMovieSceneKeyInterpolation InKeyInterpolation, float TimeToCopyFrom = FLT_MAX) override;
+
+	virtual FRichCurve* GetRichCurve() override
+	{
+		return nullptr;
+	}
+
+	virtual UMovieSceneSection* GetOwningSection() override
+	{
+		return OwningSection.Get();
+	}
+
+	virtual bool CanCreateKeyEditor() override
+	{
+		return false;
+	}
+
+	virtual TSharedRef<SWidget> CreateKeyEditor(ISequencer* Sequencer) override
+	{
+		return SNullWidget::NullWidget;
+	}
+
+protected:
+
+	/** The curve managed by this area. */
+	FNameCurve& Curve;
+
+	/** The section that owns this area. */
+	TWeakObjectPtr<UMovieSceneSection> OwningSection;
+};
+
+
 class MOVIESCENETOOLS_API MovieSceneToolHelpers
 {
 public:
@@ -276,5 +396,3 @@ public:
 	*/
 	static void SplitSection(const TSet<TWeakObjectPtr<UMovieSceneSection>>& Sections, float Time);
 };
-
-
diff --git a/Engine/Source/Editor/Sequencer/Private/SSequencer.cpp b/Engine/Source/Editor/Sequencer/Private/SSequencer.cpp
--- a/Engine/Source/Editor/Sequencer/Private/SSequencer.cpp
+++ b/Engine/Source/Editor/Sequencer/Private/SSequencer.cpp
@@ -259,10 +259,8 @@ void SSequencer::Construct( const FArguments& InArgs, TSharedRef< class FSequenc
 		DetailsView->SetVisibility(TAttribute<EVisibility>::Create(TAttribute<EVisibility>::FGetter::CreateSP(this, &SSequencer::HandleDetailsViewVisibility)));
 	}
 
-	const int32 				Column0	= 0,	Column1	= 1;
-	const int32 Row0	= 0,
-				Row1	= 1,
-				Row2	= 2;
+	const int32 Column0 = 0, Column1 = 1;
+	const int32 Row0 = 0, Row1 = 1, Row2 = 2;
 
 	ChildSlot
 	[
@@ -447,25 +445,25 @@ void SSequencer::Construct( const FArguments& InArgs, TSharedRef< class FSequenc
 					]
 
 				+ SOverlay::Slot()
-				[
-					SNew( SSequencerSplitterOverlay )
-					.Style(FEditorStyle::Get(), "Sequencer.AnimationOutliner.Splitter")
-					.Visibility(EVisibility::SelfHitTestInvisible)
-
-					+ SSplitter::Slot()
-					.Value( FillCoefficient_0 )
-					.OnSlotResized( SSplitter::FOnSlotResized::CreateSP(this, &SSequencer::OnColumnFillCoefficientChanged, 0) )
 					[
-						SNew(SSpacer)
-					]
+						SNew( SSequencerSplitterOverlay )
+							.Style(FEditorStyle::Get(), "Sequencer.AnimationOutliner.Splitter")
+							.Visibility(EVisibility::SelfHitTestInvisible)
 
-					+ SSplitter::Slot()
-					.Value( FillCoefficient_1 )
-					.OnSlotResized( SSplitter::FOnSlotResized::CreateSP(this, &SSequencer::OnColumnFillCoefficientChanged, 1) )
-					[
-						SNew(SSpacer)
+						+ SSplitter::Slot()
+							.Value( FillCoefficient_0 )
+							.OnSlotResized( SSplitter::FOnSlotResized::CreateSP(this, &SSequencer::OnColumnFillCoefficientChanged, 0) )
+							[
+								SNew(SSpacer)
+							]
+
+						+ SSplitter::Slot()
+							.Value( FillCoefficient_1 )
+							.OnSlotResized( SSplitter::FOnSlotResized::CreateSP(this, &SSequencer::OnColumnFillCoefficientChanged, 1) )
+							[
+								SNew(SSpacer)
+							]
 					]
-				]
 			]
 		]
 	];
@@ -538,7 +536,17 @@ EVisibility SSequencer::HandleDetailsViewVisibility() const
 
 void SSequencer::HandleKeySelectionChanged()
 {
+	TSharedPtr<FStructOnScope> Struct;
+	const TArray<FSelectedKey> SelectedKeys = Sequencer.Pin()->GetSelection().GetSelectedKeys().Array();
 
+	if (SelectedKeys.Num() > 0)
+	{
+		// @todo sequencer: highlight selected keys in details view
+	}
+	else
+	{
+		// @todo sequencer: remove key highlights in details view
+	}
 }
 
 
diff --git a/Engine/Source/Runtime/MovieSceneTracks/Private/Sections/MovieSceneEventSection.cpp b/Engine/Source/Runtime/MovieSceneTracks/Private/Sections/MovieSceneEventSection.cpp
--- a/Engine/Source/Runtime/MovieSceneTracks/Private/Sections/MovieSceneEventSection.cpp
+++ b/Engine/Source/Runtime/MovieSceneTracks/Private/Sections/MovieSceneEventSection.cpp
@@ -18,17 +18,8 @@ UMovieSceneEventSection::UMovieSceneEventSection()
 
 void UMovieSceneEventSection::AddKey(float Time, const FName& EventName, FKeyParams KeyParams)
 {
-	for (auto& Key : Keys)
-	{
-		if (Key.Time == Time)
-		{
-			Key.EventNames.AddUnique(EventName);
-
-			return;
-		}
-	}
-
-	Keys.HeapPush(FMovieSceneEventSectionKey(EventName, Time));
+	Modify();
+	Events.UpdateOrAddKey(Time, EventName);
 }
 
 
@@ -39,26 +30,19 @@ void UMovieSceneEventSection::DilateSection(float DilationFactor, float Origin,
 {
 	Super::DilateSection(DilationFactor, Origin, KeyHandles);
 
-	for (auto& Key : Keys)
-	{
-		Key.Time = Key.Time * DilationFactor - Origin;
-	}
+	Events.ScaleCurve(Origin, DilationFactor, KeyHandles);
 }
 
 
 void UMovieSceneEventSection::GetKeyHandles(TSet<FKeyHandle>& KeyHandles) const
 {
-	 // @todo gmp: implement event tracks
-}
-
-
-void UMovieSceneEventSection::GetSnapTimes(TArray<float>& OutSnapTimes, bool bGetSectionBorders) const
-{
-	Super::GetSnapTimes(OutSnapTimes, bGetSectionBorders);
-
-	for (auto& Key : Keys)
+	for (auto It(Events.GetKeyHandleIterator()); It; ++It)
 	{
-		OutSnapTimes.Add(Key.Time);
+		float Time = Events.GetKeyTime(It.Key());
+		if (IsTimeWithinSection(Time))
+		{
+			KeyHandles.Add(It.Key());
+		}
 	}
 }
 
@@ -67,8 +51,5 @@ void UMovieSceneEventSection::MoveSection(float DeltaPosition, TSet<FKeyHandle>&
 {
 	Super::MoveSection(DeltaPosition, KeyHandles);
 
-	for (auto& Key : Keys)
-	{
-		Key.Time += DeltaPosition;
-	}
+	Events.ShiftCurve(DeltaPosition, KeyHandles);
 }
diff --git a/Engine/Source/Runtime/MovieSceneTracks/Public/Sections/MovieSceneEventSection.h b/Engine/Source/Runtime/MovieSceneTracks/Public/Sections/MovieSceneEventSection.h
--- a/Engine/Source/Runtime/MovieSceneTracks/Public/Sections/MovieSceneEventSection.h
+++ b/Engine/Source/Runtime/MovieSceneTracks/Public/Sections/MovieSceneEventSection.h
@@ -64,18 +64,25 @@ public:
 	 */
 	void AddKey(float Time, const FName& EventName, FKeyParams KeyParams);
 
+	/**
+	 * @return The float curve on this section
+	 */
+	FNameCurve& GetEventCurve()
+	{
+		return Events;
+	}
+
 public:
 
 	// UMovieSceneSection interface
 
 	virtual void DilateSection(float DilationFactor, float Origin, TSet<FKeyHandle>& KeyHandles) override;
 	virtual void GetKeyHandles(TSet<FKeyHandle>& KeyHandles) const override;
-	virtual void GetSnapTimes(TArray<float>& OutSnapTimes, bool bGetSectionBorders) const override;
 	virtual void MoveSection(float DeltaPosition, TSet<FKeyHandle>& KeyHandles) override;
 
 private:
 
 	/** The section's keys. */
-	UPROPERTY(EditAnywhere, Category="Audio")
-	TArray<FMovieSceneEventSectionKey> Keys;
+	UPROPERTY(EditAnywhere, Category="Events")
+	FNameCurve Events;
 };
