Sequencer: Fix undo for split/trim section.

diff --git a/Engine/Source/Editor/MovieSceneTools/Private/MovieSceneToolHelpers.cpp b/Engine/Source/Editor/MovieSceneTools/Private/MovieSceneToolHelpers.cpp
--- a/Engine/Source/Editor/MovieSceneTools/Private/MovieSceneToolHelpers.cpp
+++ b/Engine/Source/Editor/MovieSceneTools/Private/MovieSceneToolHelpers.cpp
@@ -189,6 +189,8 @@ void MovieSceneToolHelpers::TrimSection(const TSet<TWeakObjectPtr<UMovieSceneSec
 	{
 		if (Section.IsValid() && Section->IsTimeWithinSection(Time))
 		{
+			Section->SetFlags( RF_Transactional );
+
 			Section->Modify();
 
 			if (bTrimLeft)
@@ -210,6 +212,8 @@ void MovieSceneToolHelpers::SplitSection(const TSet<TWeakObjectPtr<UMovieSceneSe
 	{
 		if (Section.IsValid() && Section->IsTimeWithinSection(Time))
 		{
+			Section->SetFlags( RF_Transactional );
+
 			Section->Modify();
 
 			float SectionEndTime = Section->GetEndTime();
