Sequencer: Change skeletal animation clip menu into an asset picker for animation clips.

diff --git a/Engine/Source/Editor/MovieSceneTools/MovieSceneTools.Build.cs b/Engine/Source/Editor/MovieSceneTools/MovieSceneTools.Build.cs
--- a/Engine/Source/Editor/MovieSceneTools/MovieSceneTools.Build.cs
+++ b/Engine/Source/Editor/MovieSceneTools/MovieSceneTools.Build.cs
@@ -32,6 +32,7 @@ public class MovieSceneTools : ModuleRules
 				"MovieScene",
 				"MovieSceneTracks",
 				"BlueprintGraph",
+                "ContentBrowser",
 				"Slate",
 				"SlateCore",
                 "EditorStyle",
diff --git a/Engine/Source/Editor/MovieSceneTools/Private/TrackEditors/SkeletalAnimationTrackEditor.cpp b/Engine/Source/Editor/MovieSceneTools/Private/TrackEditors/SkeletalAnimationTrackEditor.cpp
--- a/Engine/Source/Editor/MovieSceneTools/Private/TrackEditors/SkeletalAnimationTrackEditor.cpp
+++ b/Engine/Source/Editor/MovieSceneTools/Private/TrackEditors/SkeletalAnimationTrackEditor.cpp
@@ -19,6 +19,7 @@
 #include "CommonMovieSceneTools.h"
 #include "AssetRegistryModule.h"
 #include "Animation/SkeletalMeshActor.h"
+#include "ContentBrowserModule.h"
 
 
 namespace SkeletalAnimationEditorConstants
@@ -231,38 +232,39 @@ void FSkeletalAnimationTrackEditor::BuildObjectBindingTrackMenu(FMenuBuilder& Me
 	}
 }
 
-
 void FSkeletalAnimationTrackEditor::BuildAnimationSubMenu(FMenuBuilder& MenuBuilder, FGuid ObjectBinding, USkeleton* Skeleton)
 {
-	const TSharedPtr<ISequencer> ParentSequencer = GetSequencer();
-
-	// Load the asset registry module
-	FAssetRegistryModule& AssetRegistryModule = FModuleManager::LoadModuleChecked<FAssetRegistryModule>(TEXT("AssetRegistry"));
-
-	// Collect a full list of assets with the specified class
-	TArray<FAssetData> AssetDataList;
-	AssetRegistryModule.Get().GetAssetsByClass(UAnimSequence::StaticClass()->GetFName(), AssetDataList);
+	FAssetPickerConfig AssetPickerConfig;
+	AssetPickerConfig.OnAssetSelected = FOnAssetSelected::CreateRaw( this, &FSkeletalAnimationTrackEditor::OnAnimationAssetSelected, ObjectBinding);
+	AssetPickerConfig.bAllowNullSelection = false;
+	AssetPickerConfig.InitialAssetViewType = EAssetViewType::List;
+
+	// Filter config
+	AssetPickerConfig.Filter.ClassNames.Add(UAnimSequence::StaticClass()->GetFName());
+	AssetPickerConfig.Filter.TagsAndValues.Add(TEXT("Skeleton"), FAssetData(Skeleton).GetExportTextName());
+
+	FContentBrowserModule& ContentBrowserModule = FModuleManager::Get().LoadModuleChecked<FContentBrowserModule>(TEXT("ContentBrowser"));
+
+	TSharedPtr<SBox> MenuEntry = SNew(SBox)
+	.WidthOverride(300.0f)
+	.HeightOverride(300.f)
+	[
+		ContentBrowserModule.Get().CreateAssetPicker(AssetPickerConfig)
+	];
+	MenuBuilder.AddWidget(MenuEntry.ToSharedRef(), FText::GetEmpty(), true);
+}
 
-	for( int32 AssetIndex = 0; AssetIndex < AssetDataList.Num(); ++AssetIndex )
+void FSkeletalAnimationTrackEditor::OnAnimationAssetSelected(const FAssetData& AssetData, FGuid ObjectBinding)
+{
+	UObject* SelectedObject = AssetData.GetAsset();
+	if (SelectedObject && SelectedObject->IsA(UAnimSequence::StaticClass()))
 	{
-		const FAssetData& PossibleAnimSequence = AssetDataList[AssetIndex];
-		if( FAssetData(Skeleton).GetExportTextName() == PossibleAnimSequence.TagsAndValues.FindRef("Skeleton") )
-		{
-			FFormatNamedArguments Args;
-			Args.Add( TEXT("AnimationName"), FText::FromName( PossibleAnimSequence.AssetName ) );
-
-			UAnimSequence* AnimSequence = CastChecked<UAnimSequence>(AssetDataList[AssetIndex].GetAsset());
-			MenuBuilder.AddMenuEntry(
-				FText::Format( NSLOCTEXT("Sequencer", "AddAnimSequence", "{AnimationName}"), Args ),
-				FText::Format( NSLOCTEXT("Sequencer", "AddAnimSequenceTooltip", "Adds a {AnimationName} animation to this skeletal mesh."), Args ),
-				FSlateIcon(),
-				FUIAction(FExecuteAction::CreateRaw(this, &FSkeletalAnimationTrackEditor::HandleAddAnimSequencerMenuEntryExecute, ObjectBinding, AnimSequence))
-				);
-		}
+		UAnimSequence* AnimSequence = CastChecked<UAnimSequence>(AssetData.GetAsset());
+
+		AddKey(ObjectBinding, AnimSequence);
 	}
 }
 
-
 void FSkeletalAnimationTrackEditor::AddKeyInternal( float KeyTime, const TArray<UObject*> Objects, class UAnimSequence* AnimSequence )
 {
 	for( int32 ObjectIndex = 0; ObjectIndex < Objects.Num(); ++ObjectIndex )
@@ -313,9 +315,3 @@ USkeleton* FSkeletalAnimationTrackEditor::AcquireSkeletonFromObjectGuid(const FG
 
 	return Skeleton;
 }
-
-
-void FSkeletalAnimationTrackEditor::HandleAddAnimSequencerMenuEntryExecute(FGuid ObjectGuid, UAnimSequence* AnimSequence)
-{
-	AddKey(ObjectGuid, AnimSequence);
-}
diff --git a/Engine/Source/Editor/MovieSceneTools/Private/TrackEditors/SkeletalAnimationTrackEditor.h b/Engine/Source/Editor/MovieSceneTools/Private/TrackEditors/SkeletalAnimationTrackEditor.h
--- a/Engine/Source/Editor/MovieSceneTools/Private/TrackEditors/SkeletalAnimationTrackEditor.h
+++ b/Engine/Source/Editor/MovieSceneTools/Private/TrackEditors/SkeletalAnimationTrackEditor.h
@@ -42,14 +42,14 @@ private:
 	/** Animation sub menu */
 	void BuildAnimationSubMenu(FMenuBuilder& MenuBuilder, FGuid ObjectBinding, USkeleton* Skeleton);
 
+	/** Animation asset selected */
+	void OnAnimationAssetSelected(const FAssetData& AssetData, FGuid ObjectBinding);
+
 	/** Delegate for AnimatablePropertyChanged in AddKey */
 	void AddKeyInternal(float KeyTime, const TArray<UObject*> Objects, class UAnimSequence* AnimSequence);
 
 	/** Gets a skeleton from an object guid in the movie scene */
 	class USkeleton* AcquireSkeletonFromObjectGuid(const FGuid& Guid);
-
-	/** Callback for executing the "Add Animation Sequence" menu entry. */
-	void HandleAddAnimSequencerMenuEntryExecute(FGuid ObjectGuid, UAnimSequence* AnimSequence);
 };
 
 
