Sequencer: Remove some unnecessary "Add" and "New"  text.

diff --git a/Engine/Plugins/MovieScene/ActorAnimationEditor/Source/ActorAnimationEditor/Private/ActorAnimationEditorModule.cpp b/Engine/Plugins/MovieScene/ActorAnimationEditor/Source/ActorAnimationEditor/Private/ActorAnimationEditorModule.cpp
--- a/Engine/Plugins/MovieScene/ActorAnimationEditor/Source/ActorAnimationEditor/Private/ActorAnimationEditorModule.cpp
+++ b/Engine/Plugins/MovieScene/ActorAnimationEditor/Source/ActorAnimationEditor/Private/ActorAnimationEditorModule.cpp
@@ -20,7 +20,7 @@ public:
 	/** Initialize commands */
 	virtual void RegisterCommands() override
 	{
-		UI_COMMAND(CreateNewActorAnimationInLevel, "Add New Actor Animation", "Create a new actor animation asset, and place an instance of it in this level", EUserInterfaceActionType::Button, FInputChord());
+		UI_COMMAND(CreateNewActorAnimationInLevel, "Add Actor Animation", "Create a new actor animation asset, and place an instance of it in this level", EUserInterfaceActionType::Button, FInputChord());
 	}
 };
 
diff --git a/Engine/Plugins/MovieScene/ActorAnimationEditor/Source/ActorAnimationEditor/Private/ActorAnimationEditorToolkit.cpp b/Engine/Plugins/MovieScene/ActorAnimationEditor/Source/ActorAnimationEditor/Private/ActorAnimationEditorToolkit.cpp
--- a/Engine/Plugins/MovieScene/ActorAnimationEditor/Source/ActorAnimationEditor/Private/ActorAnimationEditorToolkit.cpp
+++ b/Engine/Plugins/MovieScene/ActorAnimationEditor/Source/ActorAnimationEditor/Private/ActorAnimationEditorToolkit.cpp
@@ -101,7 +101,7 @@ void FActorAnimationEditorToolkit::Initialize( const EToolkitMode::Type Mode, co
 			FMenuExtensionDelegate::CreateLambda([=](FMenuBuilder& MenuBuilder){
 
 				MenuBuilder.AddSubMenu(
-					LOCTEXT("AddActor_Label", "Add Actor To Sequencer"),
+					LOCTEXT("AddActor_Label", "Actor To Sequencer"),
 					LOCTEXT("AddActor_ToolTip", "Allow sequencer to possess an actor that already exists in the current level"),
 					FNewMenuDelegate::CreateRaw(this, &FActorAnimationEditorToolkit::AddPosessActorMenuExtensions),
 					false /*bInOpenSubMenuOnClick*/,
diff --git a/Engine/Source/Editor/MovieSceneTools/Private/TrackEditors/AttachTrackEditor.cpp b/Engine/Source/Editor/MovieSceneTools/Private/TrackEditors/AttachTrackEditor.cpp
--- a/Engine/Source/Editor/MovieSceneTools/Private/TrackEditors/AttachTrackEditor.cpp
+++ b/Engine/Source/Editor/MovieSceneTools/Private/TrackEditors/AttachTrackEditor.cpp
@@ -220,7 +220,7 @@ void F3DAttachTrackEditor::BuildObjectBindingTrackMenu(FMenuBuilder& MenuBuilder
 		if (SelectedActors.Num()+UnselectedActors.Num() > 0)
 		{
 			MenuBuilder.AddSubMenu(
-				NSLOCTEXT("Sequencer", "AddAttach", "Add Attach"), NSLOCTEXT("Sequencer", "AddAttachTooltip", "Adds an attach track."),
+				NSLOCTEXT("Sequencer", "AddAttach", "Attach"), NSLOCTEXT("Sequencer", "AddAttachTooltip", "Adds an attach track."),
 				FNewMenuDelegate::CreateRaw(this, &F3DAttachTrackEditor::AddAttachSubMenu, ObjectBinding));
 		}
 		else
@@ -228,7 +228,7 @@ void F3DAttachTrackEditor::BuildObjectBindingTrackMenu(FMenuBuilder& MenuBuilder
 			AActor* Actor = nullptr;
 			FFormatNamedArguments Args;
 			MenuBuilder.AddMenuEntry(
-				FText::Format( NSLOCTEXT("Sequencer", "AddAttach", "Add Attach"), Args),
+				FText::Format( NSLOCTEXT("Sequencer", "AddAttach", "Attach"), Args),
 				FText::Format( NSLOCTEXT("Sequencer", "AddAttachTooltip", "Adds an attach track."), Args ),
 				FSlateIcon(),
 				FUIAction(FExecuteAction::CreateSP(this, &F3DAttachTrackEditor::AddAttach, ObjectBinding, (UObject*)Actor),
diff --git a/Engine/Source/Editor/MovieSceneTools/Private/TrackEditors/EventTrackEditor.cpp b/Engine/Source/Editor/MovieSceneTools/Private/TrackEditors/EventTrackEditor.cpp
--- a/Engine/Source/Editor/MovieSceneTools/Private/TrackEditors/EventTrackEditor.cpp
+++ b/Engine/Source/Editor/MovieSceneTools/Private/TrackEditors/EventTrackEditor.cpp
@@ -45,7 +45,7 @@ void FEventTrackEditor::BuildAddTrackMenu(FMenuBuilder& MenuBuilder)
 	}
 
 	MenuBuilder.AddMenuEntry(
-		LOCTEXT("AddEventTrack", "Add Event Track"),
+		LOCTEXT("AddEventTrack", "Event Track"),
 		LOCTEXT("AddEventTooltip", "Adds a new event track that can trigger events on the timeline."),
 		FSlateIcon(FEditorStyle::GetStyleSetName(), "Sequencer.Tracks.Event"),
 		FUIAction(
diff --git a/Engine/Source/Editor/MovieSceneTools/Private/TrackEditors/PathTrackEditor.cpp b/Engine/Source/Editor/MovieSceneTools/Private/TrackEditors/PathTrackEditor.cpp
--- a/Engine/Source/Editor/MovieSceneTools/Private/TrackEditors/PathTrackEditor.cpp
+++ b/Engine/Source/Editor/MovieSceneTools/Private/TrackEditors/PathTrackEditor.cpp
@@ -204,7 +204,7 @@ void F3DPathTrackEditor::BuildObjectBindingTrackMenu(FMenuBuilder& MenuBuilder,
 		if (SelectedActors.Num()+UnselectedActors.Num() > 0)
 		{
 			MenuBuilder.AddSubMenu(
-				NSLOCTEXT("Sequencer", "AddPath", "Add Path"), NSLOCTEXT("Sequencer", "AddPathTooltip", "Adds a path track."),
+				NSLOCTEXT("Sequencer", "AddPath", "Path"), NSLOCTEXT("Sequencer", "AddPathTooltip", "Adds a path track."),
 				FNewMenuDelegate::CreateRaw(this, &F3DPathTrackEditor::AddPathSubMenu, ObjectBinding));
 		}
 		else
@@ -212,7 +212,7 @@ void F3DPathTrackEditor::BuildObjectBindingTrackMenu(FMenuBuilder& MenuBuilder,
 			AActor* ActorWithSplineComponent = nullptr;
 			FFormatNamedArguments Args;
 			MenuBuilder.AddMenuEntry(
-				FText::Format( NSLOCTEXT("Sequencer", "AddPath", "Add Path"), Args),
+				FText::Format( NSLOCTEXT("Sequencer", "AddPath", "Path"), Args),
 				FText::Format( NSLOCTEXT("Sequencer", "AddPathTooltip", "Adds a path track."), Args ),
 				FSlateIcon(),
 				FUIAction(FExecuteAction::CreateSP(this, &F3DPathTrackEditor::AddPath, ObjectBinding, (UObject*)ActorWithSplineComponent),
diff --git a/Engine/Source/Editor/MovieSceneTools/Private/TrackEditors/ShotTrackEditor.cpp b/Engine/Source/Editor/MovieSceneTools/Private/TrackEditors/ShotTrackEditor.cpp
--- a/Engine/Source/Editor/MovieSceneTools/Private/TrackEditors/ShotTrackEditor.cpp
+++ b/Engine/Source/Editor/MovieSceneTools/Private/TrackEditors/ShotTrackEditor.cpp
@@ -568,7 +568,7 @@ void FShotTrackEditor::BuildAddTrackMenu(FMenuBuilder& MenuBuilder)
 	}
 
 	MenuBuilder.AddMenuEntry(
-		LOCTEXT("AddShotTrack", "Add Shot Track"),
+		LOCTEXT("AddShotTrack", "Shot Track"),
 		LOCTEXT("AddShotTooltip", "Adds a shot track, as well as a new shot at the current scrubber location if a camera is selected."),
 		FSlateIcon(FEditorStyle::GetStyleSetName(), "Sequencer.Tracks.Shot"),
 		FUIAction(
@@ -585,7 +585,7 @@ void FShotTrackEditor::BuildObjectBindingTrackMenu(FMenuBuilder& MenuBuilder, co
 		const TSharedPtr<ISequencer> ParentSequencer = GetSequencer();
 
 		MenuBuilder.AddMenuEntry(
-			LOCTEXT("AddShot", "Add New Shot"),
+			LOCTEXT("AddShot", "New Shot"),
 			LOCTEXT("AddShotTooltip", "Adds a new shot using this camera at the scrubber location."),
 			FSlateIcon(),
 			FUIAction(FExecuteAction::CreateRaw(this, &FShotTrackEditor::HandleAddShotMenuEntryExecute, ObjectBinding))
diff --git a/Engine/Source/Editor/MovieSceneTools/Private/TrackEditors/SkeletalAnimationTrackEditor.cpp b/Engine/Source/Editor/MovieSceneTools/Private/TrackEditors/SkeletalAnimationTrackEditor.cpp
--- a/Engine/Source/Editor/MovieSceneTools/Private/TrackEditors/SkeletalAnimationTrackEditor.cpp
+++ b/Engine/Source/Editor/MovieSceneTools/Private/TrackEditors/SkeletalAnimationTrackEditor.cpp
@@ -225,7 +225,7 @@ void FSkeletalAnimationTrackEditor::BuildObjectBindingTrackMenu(FMenuBuilder& Me
 			if (AssetDataList.Num())
 			{
 				MenuBuilder.AddSubMenu(
-					NSLOCTEXT("Sequencer", "AddAnimation", "Add Animation"), NSLOCTEXT("Sequencer", "AddAnimationTooltip", "Adds an animation track."),
+					NSLOCTEXT("Sequencer", "AddAnimation", "Animation"), NSLOCTEXT("Sequencer", "AddAnimationTooltip", "Adds an animation track."),
 					FNewMenuDelegate::CreateRaw(this, &FSkeletalAnimationTrackEditor::BuildAnimationSubMenu, ObjectBinding, Skeleton));
 			}
 		}
diff --git a/Engine/Source/Editor/MovieSceneTools/Private/TrackEditors/SlomoTrackEditor.cpp b/Engine/Source/Editor/MovieSceneTools/Private/TrackEditors/SlomoTrackEditor.cpp
--- a/Engine/Source/Editor/MovieSceneTools/Private/TrackEditors/SlomoTrackEditor.cpp
+++ b/Engine/Source/Editor/MovieSceneTools/Private/TrackEditors/SlomoTrackEditor.cpp
@@ -39,7 +39,7 @@ void FSlomoTrackEditor::BuildAddTrackMenu(FMenuBuilder& MenuBuilder)
 	}
 
 	MenuBuilder.AddMenuEntry(
-		LOCTEXT("AddPlayRateTrack", "Add Play Rate Track"),
+		LOCTEXT("AddPlayRateTrack", "Play Rate Track"),
 		LOCTEXT("AddPlayRateTrackTooltip", "Adds a new track that controls the playback rate of the sequence."),
 		FSlateIcon(FEditorStyle::GetStyleSetName(), "Sequencer.Tracks.Slomo"),
 		FUIAction(
diff --git a/Engine/Source/Editor/MovieSceneTools/Private/TrackEditors/TransformTrackEditor.cpp b/Engine/Source/Editor/MovieSceneTools/Private/TrackEditors/TransformTrackEditor.cpp
--- a/Engine/Source/Editor/MovieSceneTools/Private/TrackEditors/TransformTrackEditor.cpp
+++ b/Engine/Source/Editor/MovieSceneTools/Private/TrackEditors/TransformTrackEditor.cpp
@@ -476,7 +476,7 @@ void F3DTransformTrackEditor::BuildObjectBindingTrackMenu(FMenuBuilder& MenuBuil
 	{
 		FFormatNamedArguments Args;
 		MenuBuilder.AddMenuEntry(
-			FText::Format( NSLOCTEXT("Sequencer", "AddTransform", "Add Transform"), Args),
+			FText::Format( NSLOCTEXT("Sequencer", "AddTransform", "Transform"), Args),
 			FText::Format( NSLOCTEXT("Sequencer", "AddPTransformTooltip", "Adds a transform track."), Args ),
 			FSlateIcon(),
 			FUIAction(FExecuteAction::CreateSP(this, &F3DTransformTrackEditor::AddTransform, ObjectBinding),
