Sequencer: Cleanup pass.

diff --git a/Engine/Source/Editor/MovieSceneTools/Private/TrackEditors/SlomoTrackEditor.h b/Engine/Source/Editor/MovieSceneTools/Private/TrackEditors/SlomoTrackEditor.h
--- a/Engine/Source/Editor/MovieSceneTools/Private/TrackEditors/SlomoTrackEditor.h
+++ b/Engine/Source/Editor/MovieSceneTools/Private/TrackEditors/SlomoTrackEditor.h
@@ -38,15 +38,12 @@ public:
 
 	// ISequencerTrackEditor interface
 
-//	virtual void AddKey(const FGuid& ObjectGuid, UObject* AdditionalAsset = nullptr) override;
 	virtual void BuildAddTrackMenu(FMenuBuilder& MenuBuilder) override;
-//	virtual void BuildObjectBindingTrackMenu(FMenuBuilder& MenuBuilder, const FGuid& ObjectBinding, const UClass* ObjectClass) override;
-//	virtual TSharedRef<ISequencerSection> MakeSectionInterface(UMovieSceneSection& SectionObject, UMovieSceneTrack* Track) override;
 	virtual bool SupportsType(TSubclassOf<UMovieSceneTrack> Type) const override;
 
 private:
 
-	/** Callback for executing the "Add Event Track" menu entry. */
+	/** Callback for executing the "Add Slomo Track" menu entry. */
 	void HandleAddSlomoTrackMenuEntryExecute();
 
 private:
