Sequencer - Refactor the particle track.
+ Use keys for activating effects instead of sections.
+ Draw the length of non looping effects.
+ Fix bugs in the enum key editor which were causing it to assign and display the wrong values somtimes.
+ Add support for overriding the key icon based on the key handle.

diff --git a/Engine/Source/Editor/EditorStyle/Private/SlateEditorStyle.cpp b/Engine/Source/Editor/EditorStyle/Private/SlateEditorStyle.cpp
--- a/Engine/Source/Editor/EditorStyle/Private/SlateEditorStyle.cpp
+++ b/Engine/Source/Editor/EditorStyle/Private/SlateEditorStyle.cpp
@@ -1590,6 +1590,8 @@ void FSlateEditorStyle::FStyle::SetupGeneralStyles()
 		Set( "Sequencer.Save", new IMAGE_BRUSH( "Icons/SequencerIcons/icon_Sequencer_Save_24x", Icon40x40 ) );
 		Set( "Sequencer.Save.Small", new IMAGE_BRUSH( "Icons/SequencerIcons/icon_Sequencer_Save_24x", Icon20x20 ) );
 		Set( "Sequencer.Key", new IMAGE_BRUSH( "Common/Key", FVector2D(7.0f,12.0f) ) );
+		Set( "Sequencer.KeyLeft", new IMAGE_BRUSH( "Common/KeyLeft", FVector2D( 7.0f, 12.0f ) ) );
+		Set( "Sequencer.KeyRight", new IMAGE_BRUSH( "Common/KeyRight", FVector2D( 7.0f, 12.0f ) ) );
 		Set( "Sequencer.PartialKey", new IMAGE_BRUSH( "Sequencer/PartialKey", FVector2D(11.f, 11.f) ) );
 		Set( "Sequencer.GenericSection.Background", new BOX_BRUSH( "Sequencer/GenericSectionBackground", FMargin(5.f/16.f) ) );
 		Set( "Sequencer.DefaultSectionGripLeft", new BOX_BRUSH( "Sequencer/SectionGripLeft", FMargin(5.f/16.f) ) );
diff --git a/Engine/Source/Editor/MovieSceneTools/Private/CurveKeyEditors/SEnumCurveKeyEditor.cpp b/Engine/Source/Editor/MovieSceneTools/Private/CurveKeyEditors/SEnumCurveKeyEditor.cpp
--- a/Engine/Source/Editor/MovieSceneTools/Private/CurveKeyEditors/SEnumCurveKeyEditor.cpp
+++ b/Engine/Source/Editor/MovieSceneTools/Private/CurveKeyEditors/SEnumCurveKeyEditor.cpp
@@ -10,24 +10,24 @@ void SEnumCurveKeyEditor::Construct(const FArguments& InArgs)
 	Sequencer = InArgs._Sequencer;
 	OwningSection = InArgs._OwningSection;
 	Curve = InArgs._Curve;
-	for (int32 i = 0; i < InArgs._Enum->NumEnums(); i++)
+	Enum = InArgs._Enum;
+
+	for (int32 i = 0; i < InArgs._Enum->NumEnums() - 1; i++)
 	{
-		FString EnumDisplayValue = InArgs._Enum->GetDisplayNameText(i).ToString();
-		if (EnumDisplayValue.Len() > 0)
+		if (Enum->HasMetaData( TEXT("Hidden"), i ) == false)
 		{
-			EnumDisplayValue = InArgs._Enum->GetEnumName(i);
+			VisibleEnumNameIndices.Add(MakeShareable(new int32(i)));
 		}
-		EnumValues.Add(MakeShareable(new FString(EnumDisplayValue)));
 	}
-	float CurrentTime = Sequencer->GetCurrentLocalTime(*Sequencer->GetFocusedMovieSceneSequence());
+
 	ChildSlot
 	[
-		SNew(SComboBox<TSharedPtr<FString>>)
+		SAssignNew(EnumValuesCombo, SComboBox<TSharedPtr<int32>>)
 		.ButtonStyle(FEditorStyle::Get(), "FlatButton.Light")
-		.OptionsSource(&EnumValues)
-		.InitiallySelectedItem(EnumValues[Curve->Evaluate(CurrentTime)])
+		.OptionsSource(&VisibleEnumNameIndices)
 		.OnGenerateWidget(this, &SEnumCurveKeyEditor::OnGenerateWidget)
 		.OnSelectionChanged(this, &SEnumCurveKeyEditor::OnComboSelectionChanged)
+		.OnComboBoxOpening(this, &SEnumCurveKeyEditor::OnComboMenuOpening)
 		.ContentPadding(FMargin(2, 0))
 		[
 			SNew(STextBlock)
@@ -35,52 +35,77 @@ void SEnumCurveKeyEditor::Construct(const FArguments& InArgs)
 			.Text(this, &SEnumCurveKeyEditor::GetCurrentValue)
 		]
 	];
+	bUpdatingSelectionInternally = false;
 }
 
 FText SEnumCurveKeyEditor::GetCurrentValue() const
 {
 	float CurrentTime = Sequencer->GetCurrentLocalTime(*Sequencer->GetFocusedMovieSceneSequence());
-	return FText::FromString(*EnumValues[Curve->Evaluate(CurrentTime)]);
+	int32 CurrentNameIndex = Enum->GetIndexByValue(Curve->Evaluate(CurrentTime));
+	return Enum->GetDisplayNameText(CurrentNameIndex);
 }
 
-TSharedRef<SWidget> SEnumCurveKeyEditor::OnGenerateWidget(TSharedPtr<FString> InItem)
+TSharedRef<SWidget> SEnumCurveKeyEditor::OnGenerateWidget(TSharedPtr<int32> InItem)
 {
 	return SNew(STextBlock)
-		.Text(FText::FromString(*InItem));
+		.Text(Enum->GetDisplayNameText(*InItem));
 }
 
-void SEnumCurveKeyEditor::OnComboSelectionChanged(TSharedPtr<FString> InSelectedItem, ESelectInfo::Type SelectInfo)
+void SEnumCurveKeyEditor::OnComboSelectionChanged(TSharedPtr<int32> InSelectedItem, ESelectInfo::Type SelectInfo)
 {
-	FScopedTransaction Transaction(LOCTEXT("SetEnumKey", "Set Enum Key Value"));
-	OwningSection->SetFlags(RF_Transactional);
-	OwningSection->Modify();
+	if (bUpdatingSelectionInternally == false)
+	{
+		FScopedTransaction Transaction(LOCTEXT("SetEnumKey", "Set Enum Key Value"));
+		OwningSection->SetFlags(RF_Transactional);
+		OwningSection->Modify();
 
-	float CurrentTime = Sequencer->GetCurrentLocalTime(*Sequencer->GetFocusedMovieSceneSequence());
+		float CurrentTime = Sequencer->GetCurrentLocalTime(*Sequencer->GetFocusedMovieSceneSequence());
 
-	bool bKeyWillBeAdded = Curve->IsKeyHandleValid(Curve->FindKey(CurrentTime)) == false;
-	if (bKeyWillBeAdded)
-	{
-		if (OwningSection->GetStartTime() > CurrentTime)
+		bool bKeyWillBeAdded = Curve->IsKeyHandleValid(Curve->FindKey(CurrentTime)) == false;
+		if (bKeyWillBeAdded)
 		{
-			OwningSection->SetStartTime(CurrentTime);
+			if (OwningSection->GetStartTime() > CurrentTime)
+			{
+				OwningSection->SetStartTime(CurrentTime);
+			}
+			if (OwningSection->GetEndTime() < CurrentTime)
+			{
+				OwningSection->SetEndTime(CurrentTime);
+			}
 		}
-		if (OwningSection->GetEndTime() < CurrentTime)
+
+		int32 SelectedValue = Enum->GetValueByIndex(*InSelectedItem);
+		if (Curve->GetNumKeys() == 0)
 		{
-			OwningSection->SetEndTime(CurrentTime);
+			Curve->SetDefaultValue(SelectedValue);
 		}
+		else
+		{
+			Curve->UpdateOrAddKey(CurrentTime, SelectedValue);
+		}
+		Sequencer->UpdateRuntimeInstances();
 	}
+}
 
-	int32 SelectedIndex;
-	EnumValues.Find(InSelectedItem, SelectedIndex);
-	if (Curve->GetNumKeys() == 0)
+void SEnumCurveKeyEditor::OnComboMenuOpening()
+{
+	float CurrentTime = Sequencer->GetCurrentLocalTime( *Sequencer->GetFocusedMovieSceneSequence() );
+	int32 CurrentNameIndex = Enum->GetIndexByValue( Curve->Evaluate( CurrentTime ) );
+	TSharedPtr<int32> FoundNameIndexItem;
+	for ( int32 i = 0; i < VisibleEnumNameIndices.Num(); i++ )
 	{
-		Curve->SetDefaultValue(SelectedIndex);
+		if ( *VisibleEnumNameIndices[i] == CurrentNameIndex )
+		{
+			FoundNameIndexItem = VisibleEnumNameIndices[i];
+			break;
+		}
 	}
-	else
+	if ( FoundNameIndexItem.IsValid() )
 	{
-		Curve->UpdateOrAddKey(CurrentTime, SelectedIndex);
+		bUpdatingSelectionInternally = true;
+		EnumValuesCombo->SetSelectedItem(FoundNameIndexItem);
+		bUpdatingSelectionInternally = false;
 	}
-	Sequencer->UpdateRuntimeInstances();
 }
 
 #undef LOCTEXT_NAMESPACE
\ No newline at end of file
diff --git a/Engine/Source/Editor/MovieSceneTools/Private/CurveKeyEditors/SEnumCurveKeyEditor.h b/Engine/Source/Editor/MovieSceneTools/Private/CurveKeyEditors/SEnumCurveKeyEditor.h
--- a/Engine/Source/Editor/MovieSceneTools/Private/CurveKeyEditors/SEnumCurveKeyEditor.h
+++ b/Engine/Source/Editor/MovieSceneTools/Private/CurveKeyEditors/SEnumCurveKeyEditor.h
@@ -12,7 +12,7 @@ public:
 		SLATE_ARGUMENT(ISequencer*, Sequencer)
 		SLATE_ARGUMENT(UMovieSceneSection*, OwningSection)
 		SLATE_ARGUMENT(FIntegralCurve*, Curve)
-		SLATE_ARGUMENT(UEnum*, Enum)
+		SLATE_ARGUMENT(const UEnum*, Enum)
 	SLATE_END_ARGS();
 
 	void Construct(const FArguments& InArgs);
@@ -20,12 +20,17 @@ public:
 private:
 	FText GetCurrentValue() const;
 
-	TSharedRef<SWidget> OnGenerateWidget(TSharedPtr<FString> InItem);
+	TSharedRef<SWidget> OnGenerateWidget(TSharedPtr<int32> InItem);
 
-	void OnComboSelectionChanged(TSharedPtr<FString> InSelectedItem, ESelectInfo::Type SelectInfo);
+	void OnComboSelectionChanged(TSharedPtr<int32> InSelectedItem, ESelectInfo::Type SelectInfo);
+
+	void OnComboMenuOpening();
 
 	ISequencer* Sequencer;
 	UMovieSceneSection* OwningSection;
 	FIntegralCurve* Curve;
-	TArray<TSharedPtr<FString>> EnumValues;
+	const UEnum* Enum;
+	TArray<TSharedPtr<int32>> VisibleEnumNameIndices;
+	TSharedPtr<SComboBox<TSharedPtr<int32>>> EnumValuesCombo;
+	bool bUpdatingSelectionInternally;
 };
\ No newline at end of file
diff --git a/Engine/Source/Editor/MovieSceneTools/Private/TrackEditors/ParticleTrackEditor.cpp b/Engine/Source/Editor/MovieSceneTools/Private/TrackEditors/ParticleTrackEditor.cpp
--- a/Engine/Source/Editor/MovieSceneTools/Private/TrackEditors/ParticleTrackEditor.cpp
+++ b/Engine/Source/Editor/MovieSceneTools/Private/TrackEditors/ParticleTrackEditor.cpp
@@ -19,6 +19,9 @@
 #include "CommonMovieSceneTools.h"
 #include "AssetRegistryModule.h"
 #include "Particles/Emitter.h"
+#include "Particles/ParticleEmitter.h"
+#include "Particles/ParticleLODLevel.h"
+#include "Particles/ParticleModuleRequired.h"
 
 
 namespace AnimatableParticleEditorConstants
@@ -28,9 +31,15 @@ namespace AnimatableParticleEditorConstants
 }
 
 
-FParticleSection::FParticleSection( UMovieSceneSection& InSection )
+FParticleSection::FParticleSection( UMovieSceneSection& InSection, TSharedRef<ISequencer>InOwningSequencer )
 	: Section( InSection )
+	, OwningSequencer( InOwningSequencer )
 {
+	ParticleKeyEnum = FindObject<UEnum>( ANY_PACKAGE, TEXT( "EParticleKey" ) );
+	checkf( ParticleKeyEnum != nullptr, TEXT( "FParticleSection could not find the EParticleKey UEnum by name." ) )
+
+	LeftKeyBrush = FEditorStyle::GetBrush( "Sequencer.KeyLeft" );
+	RightKeyBrush = FEditorStyle::GetBrush( "Sequencer.KeyRight" );
 }
 
 
@@ -40,75 +49,166 @@ FParticleSection::~FParticleSection()
 
 
 UMovieSceneSection* FParticleSection::GetSectionObject()
-{ 
+{
 	return &Section;
 }
 
 
 FText FParticleSection::GetDisplayName() const
 {
-	return NSLOCTEXT("FParticleSection", "Emitter", "Emitter");
+	return NSLOCTEXT( "FParticleSection", "Emitter", "Emitter" );
 }
 
-
-FText FParticleSection::GetSectionTitle() const
-{
-	EParticleKey::Type KeyType = Cast<UMovieSceneParticleSection>(&Section)->GetKeyType();
-	return KeyType == EParticleKey::Toggle ? NSLOCTEXT("FParticleSection", "Toggle", "Toggle") :
-		KeyType == EParticleKey::Trigger ? NSLOCTEXT("FParticleSection", "Trigger", "Trigger") : NSLOCTEXT("FParticleSection", "None", "None");
-}
-
-
 float FParticleSection::GetSectionHeight() const
 {
 	return (float)AnimatableParticleEditorConstants::ParticleTrackHeight;
 }
 
+void FParticleSection::GenerateSectionLayout( class ISectionLayoutBuilder& LayoutBuilder ) const
+{
+	UMovieSceneParticleSection* ParticleSection = Cast<UMovieSceneParticleSection>( &Section );
+	LayoutBuilder.SetSectionAsKeyArea( MakeShareable( new FEnumKeyArea( ParticleSection->GetParticleCurve(), ParticleSection, ParticleKeyEnum ) ) );
+}
 
 int32 FParticleSection::OnPaintSection( const FGeometry& AllottedGeometry, const FSlateRect& SectionClippingRect, FSlateWindowElementList& OutDrawElements, int32 LayerId, bool bParentEnabled ) const
 {
 	const ESlateDrawEffect::Type DrawEffects = bParentEnabled ? ESlateDrawEffect::None : ESlateDrawEffect::DisabledEffect;
+	UMovieSceneParticleSection* AnimSection = Cast<UMovieSceneParticleSection>( &Section );
+	FTimeToPixel TimeToPixelConverter( AllottedGeometry, OwningSequencer->GetViewRange() );
+
+	// @todo Sequencer - These values should be cached and then refreshed only when the particle system changes.
+	bool bIsLooping = false;
+	float LastEmitterEndTime = 0;
+	UMovieSceneParticleSection* ParticleSection = Cast<UMovieSceneParticleSection>( &Section );
+	if ( ParticleSection != nullptr )
+	{
+		UMovieSceneParticleTrack* ParentTrack = Cast<UMovieSceneParticleTrack>( ParticleSection->GetOuter() );
+		if ( ParentTrack != nullptr )
+		{
+			FGuid ObjectHandle;
+			for ( const FMovieSceneBinding& Binding : OwningSequencer->GetFocusedMovieSceneSequence()->GetMovieScene()->GetBindings() )
+			{
+				if ( Binding.GetTracks().Contains( ParentTrack ) )
+				{
+					ObjectHandle = Binding.GetObjectGuid();
+					break;
+				}
+			}
 
-	UMovieSceneParticleSection* AnimSection = Cast<UMovieSceneParticleSection>(&Section);
-	
-	FTimeToPixel TimeToPixelConverter( AllottedGeometry, TRange<float>( Section.GetStartTime(), Section.GetEndTime() ) );
-	
-	EParticleKey::Type KeyType = Cast<UMovieSceneParticleSection>(&Section)->GetKeyType();
+			if ( ObjectHandle.IsValid() )
+			{
+				AEmitter* ParticleSystemActor = Cast<AEmitter>( OwningSequencer->GetFocusedMovieSceneSequence()->FindObject( ObjectHandle ) );
+				if ( ParticleSystemActor != nullptr )
+				{
+					for ( UParticleEmitter* Emitter : ParticleSystemActor->GetParticleSystemComponent()->Template->Emitters )
+					{
+						UParticleModuleRequired* RequiredModule = Emitter->GetLODLevel( 0 )->RequiredModule;
+						bIsLooping |= RequiredModule->EmitterLoops == 0;
+						LastEmitterEndTime = FMath::Max( LastEmitterEndTime, RequiredModule->EmitterDelay + RequiredModule->EmitterDuration );
+					}
+				}
+			}
+		}
+	}
 
-	// @todo Sequencer Particle Triggers should be sections with 0 duration, but a custom graphic which
-	// allows them to be dragged around
-	if (KeyType == EParticleKey::Trigger)
+	// @todo Sequencer - This should only draw the visible ranges.
+	TArray<TRange<float>> DrawRanges;
+	TOptional<float> CurrentRangeStart;
+	for ( auto KeyIterator = ParticleSection->GetParticleCurve().GetKeyIterator(); KeyIterator; ++KeyIterator )
 	{
-		FSlateDrawElement::MakeBox(
-			OutDrawElements,
-			LayerId,
-			AllottedGeometry.ToPaintGeometry(),
-			FEditorStyle::GetBrush("Sequencer.GenericSection.Background"),
-			SectionClippingRect,
-			DrawEffects,
-			FLinearColor(0.6f, 0.4f, 0.3f, 1.f)
-		);
+		FIntegralKey Key = *KeyIterator;
+		if ( (EParticleKey::Type)Key.Value == EParticleKey::Active )
+		{
+			if ( CurrentRangeStart.IsSet() == false )
+			{
+				CurrentRangeStart = Key.Time;
+			}
+			else
+			{
+				if ( bIsLooping == false )
+				{
+					if ( Key.Time > CurrentRangeStart.GetValue() + LastEmitterEndTime )
+					{
+						DrawRanges.Add( TRange<float>( CurrentRangeStart.GetValue(), CurrentRangeStart.GetValue() + LastEmitterEndTime ) );
+					}
+					else
+					{
+						DrawRanges.Add( TRange<float>( CurrentRangeStart.GetValue(), Key.Time ) );
+					}
+					CurrentRangeStart = Key.Time;
+				}
+			}
+		}
+		if ( (EParticleKey::Type)Key.Value == EParticleKey::Inactive )
+		{
+			if ( CurrentRangeStart.IsSet() )
+			{
+				if (bIsLooping)
+				{
+					DrawRanges.Add( TRange<float>( CurrentRangeStart.GetValue(), Key.Time ) );
+				}
+				else
+				{
+					if ( Key.Time > CurrentRangeStart.GetValue() + LastEmitterEndTime )
+					{
+						DrawRanges.Add( TRange<float>( CurrentRangeStart.GetValue(), CurrentRangeStart.GetValue() + LastEmitterEndTime ) );
+					}
+					else
+					{
+						DrawRanges.Add( TRange<float>( CurrentRangeStart.GetValue(), Key.Time ) );
+					}
+				}
+				CurrentRangeStart.Reset();
+			}
+		}
+	}
+	if ( CurrentRangeStart.IsSet() )
+	{
+		if ( bIsLooping )
+		{
+			DrawRanges.Add( TRange<float>( CurrentRangeStart.GetValue(), OwningSequencer->GetViewRange().GetUpperBoundValue() ) );
+		}
+		else
+		{
+			DrawRanges.Add( TRange<float>( CurrentRangeStart.GetValue(), CurrentRangeStart.GetValue() + LastEmitterEndTime ) );
+		}
 	}
-	else
+
+	for ( const TRange<float>& DrawRange : DrawRanges )
 	{
+		float XOffset = TimeToPixelConverter.TimeToPixel(DrawRange.GetLowerBoundValue());
+		float XSize = TimeToPixelConverter.TimeToPixel(DrawRange.GetUpperBoundValue()) - XOffset;
 		FSlateDrawElement::MakeBox(
 			OutDrawElements,
 			LayerId,
-			AllottedGeometry.ToPaintGeometry(),
-			FEditorStyle::GetBrush("Sequencer.GenericSection.Background"),
+			AllottedGeometry.ToPaintGeometry( FVector2D( XOffset, (AllottedGeometry.GetLocalSize().Y - SequencerSectionConstants::KeySize.Y) / 2 ), FVector2D( XSize, SequencerSectionConstants::KeySize.Y ) ),
+			FEditorStyle::GetBrush( "Sequencer.GenericSection.Background" ),
 			SectionClippingRect,
 			DrawEffects,
-			FLinearColor(0.8f, 0.4f, 0.3f, 1.f)
-		);
+			FLinearColor( 0.0f, 0.0f, 1.f, 1.f )
+			);
 	}
 
 	return LayerId+1;
 }
 
 
-bool FParticleSection::SectionIsResizable() const
+const FSlateBrush* FParticleSection::GetKeyBrush( FKeyHandle KeyHandle ) const
 {
-	return Cast<UMovieSceneParticleSection>(&Section)->GetKeyType() != EParticleKey::Trigger;
+	UMovieSceneParticleSection* ParticleSection = Cast<UMovieSceneParticleSection>( &Section );
+	if ( ParticleSection != nullptr )
+	{
+		FIntegralKey ParticleKey = ParticleSection->GetParticleCurve().GetKey(KeyHandle);
+		if ( (EParticleKey::Type)ParticleKey.Value == EParticleKey::Active )
+		{
+			return LeftKeyBrush;
+		}
+		else if ( (EParticleKey::Type)ParticleKey.Value == EParticleKey::Inactive )
+		{
+			return RightKeyBrush;
+		}
+	}
+	return nullptr;
 }
 
 
@@ -139,7 +239,8 @@ TSharedRef<ISequencerSection> FParticleTrackEditor::MakeSectionInterface( UMovie
 {
 	check( SupportsType( SectionObject.GetOuter()->GetClass() ) );
 	
-	TSharedRef<ISequencerSection> NewSection( new FParticleSection(SectionObject) );
+	const TSharedPtr<ISequencer> OwningSequencer = GetSequencer();
+	TSharedRef<ISequencerSection> NewSection( new FParticleSection( SectionObject, OwningSequencer.ToSharedRef() ) );
 
 	return NewSection;
 }
@@ -152,32 +253,24 @@ void FParticleTrackEditor::BuildObjectBindingTrackMenu(FMenuBuilder& MenuBuilder
 		const TSharedPtr<ISequencer> ParentSequencer = GetSequencer();
 
 		MenuBuilder.AddMenuEntry(
-			NSLOCTEXT("Sequencer", "TriggerParticles", "Trigger"),
-			NSLOCTEXT("Sequencer", "TriggerParticlesTooltip", "Adds a trigger event for this particle emitter at the current time."),
-			FSlateIcon(),
-			FUIAction(FExecuteAction::CreateSP(this, &FParticleTrackEditor::AddParticleKey, ObjectBinding, true))
-			);
-		MenuBuilder.AddMenuEntry(
-			NSLOCTEXT("Sequencer", "ToggleParticles", "Toggle"),
-			NSLOCTEXT("Sequencer", "ToggleParticlesTooltip", "Adds a toggle section for this particle emitter at the current time."),
+			NSLOCTEXT("Sequencer", "AddParticleTrack", "Add Particle Track"),
+			NSLOCTEXT("Sequencer", "TriggerParticlesTooltip", "Adds a track for controlling particle emitter state."),
 			FSlateIcon(),
-			FUIAction(FExecuteAction::CreateSP(this, &FParticleTrackEditor::AddParticleKey, ObjectBinding, false))
+			FUIAction(FExecuteAction::CreateSP(this, &FParticleTrackEditor::AddParticleKey, ObjectBinding))
 			);
 	}
 }
 
-
-void FParticleTrackEditor::AddParticleKey(const FGuid ObjectGuid, bool bTrigger)
+void FParticleTrackEditor::AddParticleKey( const FGuid ObjectGuid )
 {
 	TArray<UObject*> OutObjects;
 	GetSequencer()->GetRuntimeObjects( GetSequencer()->GetFocusedMovieSceneSequenceInstance(), ObjectGuid, OutObjects );
 
 	AnimatablePropertyChanged( UMovieSceneParticleTrack::StaticClass(), 
-		FOnKeyProperty::CreateRaw( this, &FParticleTrackEditor::AddKeyInternal, OutObjects, bTrigger) );
+		FOnKeyProperty::CreateRaw( this, &FParticleTrackEditor::AddKeyInternal, OutObjects ) );
 }
 
-
-void FParticleTrackEditor::AddKeyInternal( float KeyTime, const TArray<UObject*> Objects, bool bTrigger)
+void FParticleTrackEditor::AddKeyInternal( float KeyTime, const TArray<UObject*> Objects )
 {
 	for( int32 ObjectIndex = 0; ObjectIndex < Objects.Num(); ++ObjectIndex )
 	{
@@ -190,7 +283,7 @@ void FParticleTrackEditor::AddKeyInternal( float KeyTime, const TArray<UObject*>
 
 			if (ensure(Track))
 			{
-				Cast<UMovieSceneParticleTrack>(Track)->AddNewParticleSystem( KeyTime, bTrigger );
+				Cast<UMovieSceneParticleTrack>(Track)->AddNewKey( KeyTime );
 			}
 		}
 	}
diff --git a/Engine/Source/Editor/MovieSceneTools/Private/TrackEditors/ParticleTrackEditor.h b/Engine/Source/Editor/MovieSceneTools/Private/TrackEditors/ParticleTrackEditor.h
--- a/Engine/Source/Editor/MovieSceneTools/Private/TrackEditors/ParticleTrackEditor.h
+++ b/Engine/Source/Editor/MovieSceneTools/Private/TrackEditors/ParticleTrackEditor.h
@@ -39,10 +39,11 @@ public:
 	virtual TSharedRef<ISequencerSection> MakeSectionInterface( UMovieSceneSection& SectionObject, UMovieSceneTrack* Track ) override;
 	virtual bool SupportsType( TSubclassOf<UMovieSceneTrack> Type ) const override;
 	
+	void AddParticleKey(const FGuid ObjectGuid);
 private:
 
 	/** Delegate for AnimatablePropertyChanged in AddKey. */
-	virtual void AddKeyInternal( float KeyTime, const TArray<UObject*> Objects, bool bTrigger);
+	virtual void AddKeyInternal( float KeyTime, const TArray<UObject*> Objects);
 };
 
 
@@ -55,7 +56,7 @@ class FParticleSection
 {
 public:
 
-	FParticleSection( UMovieSceneSection& InSection );
+	FParticleSection( UMovieSceneSection& InSection, TSharedRef<ISequencer> InOwningSequencer );
 	~FParticleSection();
 
 public:
@@ -63,14 +64,24 @@ public:
 	// ISequencerSection interface
 	virtual UMovieSceneSection* GetSectionObject() override;
 	virtual FText GetDisplayName() const override;
-	virtual FText GetSectionTitle() const override;
+	virtual FText GetSectionTitle() const override { return FText::GetEmpty(); }
 	virtual float GetSectionHeight() const override;
-	virtual void GenerateSectionLayout( class ISectionLayoutBuilder& LayoutBuilder ) const override {}
+	virtual void GenerateSectionLayout( class ISectionLayoutBuilder& LayoutBuilder ) const override;
 	virtual int32 OnPaintSection( const FGeometry& AllottedGeometry, const FSlateRect& SectionClippingRect, FSlateWindowElementList& OutDrawElements, int32 LayerId, bool bParentEnabled ) const override;
-	virtual bool SectionIsResizable() const override;
+	virtual const FSlateBrush* GetKeyBrush( FKeyHandle KeyHandle ) const override;
+	virtual bool SectionIsResizable() const override { return false; }
 
 private:
 
 	/** The section we are visualizing. */
 	UMovieSceneSection& Section;
+
+	/** The sequencer that owns this section */
+	TSharedRef<ISequencer> OwningSequencer;
+
+	/** The UEnum for the EParticleKey enum */
+	const UEnum* ParticleKeyEnum;
+
+	const FSlateBrush* LeftKeyBrush;
+	const FSlateBrush* RightKeyBrush;
 };
diff --git a/Engine/Source/Editor/MovieSceneTools/Public/MovieSceneToolHelpers.h b/Engine/Source/Editor/MovieSceneTools/Public/MovieSceneToolHelpers.h
--- a/Engine/Source/Editor/MovieSceneTools/Public/MovieSceneToolHelpers.h
+++ b/Engine/Source/Editor/MovieSceneTools/Public/MovieSceneToolHelpers.h
@@ -224,7 +224,7 @@ protected:
 class FEnumKeyArea : public FIntegralKeyArea
 {
 public:
-	FEnumKeyArea(FIntegralCurve& InCurve, UMovieSceneSection* InOwningSection, UEnum* InEnum)
+	FEnumKeyArea(FIntegralCurve& InCurve, UMovieSceneSection* InOwningSection, const UEnum* InEnum)
 		: FIntegralKeyArea(InCurve, InOwningSection)
 		, Enum(InEnum)
 	{}
@@ -238,7 +238,7 @@ public:
 
 private:
 	/** The enum which provides available integral values for this key area. */
-	UEnum* Enum;
+	const UEnum* Enum;
 };
 
 
diff --git a/Engine/Source/Editor/Sequencer/Private/SSection.cpp b/Engine/Source/Editor/Sequencer/Private/SSection.cpp
--- a/Engine/Source/Editor/Sequencer/Private/SSection.cpp
+++ b/Engine/Source/Editor/Sequencer/Private/SSection.cpp
@@ -304,7 +304,7 @@ void SSection::PaintKeys( const FGeometry& AllottedGeometry, const FSlateRect& M
 
 	const FSlateBrush* BackgroundBrush = FEditorStyle::GetBrush(BackgroundBrushName);
 
-	const FSlateBrush* KeyBrush = FEditorStyle::GetBrush(KeyBrushName);
+	const FSlateBrush* DefaultKeyBrush = FEditorStyle::GetBrush(KeyBrushName);
 
 	static const FName SelectionColorName("SelectionColor");
 	static const FName SelectionInactiveColorName("SelectionColorInactive");
@@ -355,7 +355,11 @@ void SSection::PaintKeys( const FGeometry& AllottedGeometry, const FSlateRect& M
 			// Omit keys which would not be visible
 			if( SectionObject.IsTimeWithinSection( KeyTime ) )
 			{
-				const FSlateBrush* BrushToUse = KeyBrush;
+				const FSlateBrush* BrushToUse = SectionInterface->GetKeyBrush(KeyHandle);
+				if(BrushToUse == nullptr)
+				{
+					BrushToUse = DefaultKeyBrush;
+				}
 				FLinearColor KeyColor( 1.0f, 1.0f, 1.0f, 1.0f );
 				FLinearColor KeyTint(1.f, 1.f, 1.f, 1.f);
 
diff --git a/Engine/Source/Editor/Sequencer/Public/ISequencerSection.h b/Engine/Source/Editor/Sequencer/Public/ISequencerSection.h
--- a/Engine/Source/Editor/Sequencer/Public/ISequencerSection.h
+++ b/Engine/Source/Editor/Sequencer/Public/ISequencerSection.h
@@ -57,6 +57,13 @@ public:
 	 */
 	virtual int32 OnPaintSection( const FGeometry& AllottedGeometry, const FSlateRect& SectionClippingRect, FSlateWindowElementList& OutDrawElements, int32 LayerId, bool bParentEnabled ) const = 0;
 
+	/** Allows a section to override the brush to use for a key by handle.
+	 *
+	 * @param KeyHandle the handle of the key to get a brush for.
+	 * @return A const pointer to a slate brush if the brush should be overridden, otherwise null.
+	 */
+	virtual const FSlateBrush* GetKeyBrush(FKeyHandle KeyHandle) const { return nullptr; }
+
 	/**
 	 * Called when the section is double clicked
 	 *
diff --git a/Engine/Source/Runtime/Engine/Classes/Curves/CurveBase.h b/Engine/Source/Runtime/Engine/Classes/Curves/CurveBase.h
--- a/Engine/Source/Runtime/Engine/Classes/Curves/CurveBase.h
+++ b/Engine/Source/Runtime/Engine/Classes/Curves/CurveBase.h
@@ -844,9 +844,11 @@ public:
 
 	/** Default constructor. */
 	FIntegralCurve() 
-		: FIndexedCurve()
-		, DefaultValue(MAX_int32)
-	{ }
+	: FIndexedCurve()
+	, DefaultValue(MAX_int32)
+	, bUseDefaultValueBeforeFirstKey(false)
+	{
+	}
 
 	/** Virtual destructor. */
 	virtual ~FIntegralCurve() { }
@@ -889,6 +891,12 @@ public:
 	/** Get the default value for the curve */
 	int32 GetDefaultValue() const { return DefaultValue; }
 
+	/** Sets whether or not the default value should be used for evaluation for time values before the first key. */
+	void SetUseDefaultValueBeforeFirstKey(bool InbUseDefaultValueBeforeFirstKey) { bUseDefaultValueBeforeFirstKey = InbUseDefaultValueBeforeFirstKey; }
+	
+	/** Gets whether or not the default value should be used for evaluation for time values before the first key. */
+	bool GetUseDefaultValueBeforeFirstKey() const { return bUseDefaultValueBeforeFirstKey; }
+
 	/** Shifts all keys forwards or backwards in time by an even amount, preserving order */
 	void ShiftCurve(float DeltaTime);
 	void ShiftCurve(float DeltaTime, TSet<FKeyHandle>& KeyHandles);
@@ -903,6 +911,9 @@ public:
 
 	FKeyHandle FindKey(float KeyTime) const;
 
+	/** Gets the handle for the last key which is at or before the time requested.  If there are no keys at or before the requested time, an invalid handle is returned. */
+	FKeyHandle FindKeyBeforeOrAt(float KeyTime) const;
+
 private:
 
 	/** The keys, ordered by time */
@@ -912,4 +923,7 @@ private:
 	/** Default value */
 	UPROPERTY(EditAnywhere, Category="Curve")
 	int32 DefaultValue;
+
+	UPROPERTY()
+	bool bUseDefaultValueBeforeFirstKey;
 };
diff --git a/Engine/Source/Runtime/Engine/Private/CurveBase.cpp b/Engine/Source/Runtime/Engine/Private/CurveBase.cpp
--- a/Engine/Source/Runtime/Engine/Private/CurveBase.cpp
+++ b/Engine/Source/Runtime/Engine/Private/CurveBase.cpp
@@ -1646,9 +1646,9 @@ int32 FIntegralCurve::Evaluate(float Time, int32 InDefaultValue) const
 	// If the default value hasn't been initialized, use the incoming default value
 	int32 ReturnVal = DefaultValue == MAX_int32 ? InDefaultValue : DefaultValue;
 
-	if( Keys.Num() == 0 )
+	if( Keys.Num() == 0 || (bUseDefaultValueBeforeFirstKey && Time < Keys[0].Time) )
 	{
-		// If no keys in curve, return the Default value.
+		// If no keys in curve, or bUseDefaultValueBeforeFirstKey is set and the time is before the first key, return the Default value.
 	}
 	else if( Keys.Num() < 2 || Time < Keys[0].Time )
 	{
@@ -1872,3 +1872,45 @@ FKeyHandle FIntegralCurve::FindKey(float KeyTime) const
 
 	return FKeyHandle();
 }
+
+FKeyHandle FIntegralCurve::FindKeyBeforeOrAt( float KeyTime ) const
+{
+	// If there are no keys or the time is before the first key return an invalid handle.
+	if ( Keys.Num() == 0 || KeyTime < Keys[0].Time )
+	{
+		return FKeyHandle();
+	}
+
+	// If the time is after or at the last key return the last key.
+	if ( KeyTime >= Keys[Keys.Num() - 1].Time )
+	{
+		return GetKeyHandle(Keys.Num() - 1);
+	}
+
+	// Otherwise binary search to find the handle of the nearest key at or before the time.
+	int32 Start = 0;
+	int32 End = Keys.Num() - 1;
+	int32 FoundIndex = -1;
+	while ( FoundIndex < 0 )
+	{
+		int32 TestPos = (Start + End) / 2;
+		float TestKeyTime = Keys[TestPos].Time;
+		float NextTestKeyTime = Keys[TestPos + 1].Time;
+		if ( TestKeyTime <= KeyTime )
+		{
+			if( NextTestKeyTime > KeyTime )
+			{
+				FoundIndex = TestPos;
+			}
+			else
+			{
+				Start = TestPos + 1;
+			}
+		}
+		else
+		{
+			End = TestPos;
+		}
+	}
+	return GetKeyHandle(FoundIndex);
+}
diff --git a/Engine/Source/Runtime/MovieScene/Public/MovieSceneSection.h b/Engine/Source/Runtime/MovieScene/Public/MovieSceneSection.h
--- a/Engine/Source/Runtime/MovieScene/Public/MovieSceneSection.h
+++ b/Engine/Source/Runtime/MovieScene/Public/MovieSceneSection.h
@@ -58,8 +58,12 @@ public:
 	 */
 	TRange<float> GetRange() const 
 	{
-		// Uses an inclusive range so that sections that start and end on the same value aren't considered empty ranges.
-		return TRange<float>( TRange<float>::BoundsType::Inclusive( StartTime ), TRange<float>::BoundsType::Inclusive( EndTime ) );
+		// Use the single value constructor for zero sized ranges because it creates a range that is inclusive on both upper and lower
+		// bounds which isn't considered "empty".  Use the standard constructor for non-zero sized ranges so that they work well when
+		// calculating overlap with other non-zero sized ranges.
+		return StartTime == EndTime 
+			? TRange<float>(StartTime) 
+			: TRange<float>(StartTime, EndTime);
 	}
 	
 	/**
diff --git a/Engine/Source/Runtime/MovieSceneTracks/Private/Sections/MovieSceneParticleSection.cpp b/Engine/Source/Runtime/MovieSceneTracks/Private/Sections/MovieSceneParticleSection.cpp
--- a/Engine/Source/Runtime/MovieSceneTracks/Private/Sections/MovieSceneParticleSection.cpp
+++ b/Engine/Source/Runtime/MovieSceneTracks/Private/Sections/MovieSceneParticleSection.cpp
@@ -7,5 +7,43 @@
 UMovieSceneParticleSection::UMovieSceneParticleSection( const FObjectInitializer& ObjectInitializer )
 	: Super( ObjectInitializer )
 {
-	KeyType = EParticleKey::Toggle;
+	ParticleKeys.SetDefaultValue((int32)EParticleKey::Inactive);
+	ParticleKeys.SetUseDefaultValueBeforeFirstKey(true);
+	SetIsInfinite(true);
+}
+
+void UMovieSceneParticleSection::AddKey( float Time, EParticleKey::Type KeyType )
+{
+	ParticleKeys.AddKey( Time, (int32)KeyType );
+}
+
+FIntegralCurve& UMovieSceneParticleSection::GetParticleCurve()
+{
+	return ParticleKeys;
+}
+
+void UMovieSceneParticleSection::MoveSection( float DeltaPosition, TSet<FKeyHandle>& KeyHandles )
+{
+	Super::MoveSection( DeltaPosition, KeyHandles );
+
+	ParticleKeys.ShiftCurve( DeltaPosition, KeyHandles );
+}
+
+void UMovieSceneParticleSection::DilateSection( float DilationFactor, float Origin, TSet<FKeyHandle>& KeyHandles )
+{
+	Super::DilateSection( DilationFactor, Origin, KeyHandles );
+
+	ParticleKeys.ScaleCurve( Origin, DilationFactor, KeyHandles );
+}
+
+void UMovieSceneParticleSection::GetKeyHandles( TSet<FKeyHandle>& KeyHandles ) const
+{
+	for ( auto It( ParticleKeys.GetKeyHandleIterator() ); It; ++It )
+	{
+		float Time = ParticleKeys.GetKeyTime( It.Key() );
+		if ( IsTimeWithinSection( Time ) )
+		{
+			KeyHandles.Add( It.Key() );
+		}
+	}
 }
diff --git a/Engine/Source/Runtime/MovieSceneTracks/Private/TrackInstances/MovieSceneParticleTrackInstance.cpp b/Engine/Source/Runtime/MovieSceneTracks/Private/TrackInstances/MovieSceneParticleTrackInstance.cpp
--- a/Engine/Source/Runtime/MovieSceneTracks/Private/TrackInstances/MovieSceneParticleTrackInstance.cpp
+++ b/Engine/Source/Runtime/MovieSceneTracks/Private/TrackInstances/MovieSceneParticleTrackInstance.cpp
@@ -20,51 +20,45 @@ void FMovieSceneParticleTrackInstance::Update( float Position, float LastPositio
 
 	if (Position > LastPosition && Player.GetPlaybackStatus() == EMovieScenePlayerStatus::Playing)
 	{
-		bool bTrigger = false, bOn = false, bOff = false;
-
+		
 		const TArray<UMovieSceneSection*> Sections = ParticleTrack->GetAllParticleSections();
+		EParticleKey::Type ParticleKey = EParticleKey::Active;
+		bool bKeyFound = false;
 		for (int32 i = 0; i < Sections.Num(); ++i)
 		{
-			UMovieSceneParticleSection* Section = Cast<UMovieSceneParticleSection>(Sections[i]);
+			UMovieSceneParticleSection* Section = Cast<UMovieSceneParticleSection>( Sections[i] );
 			if (Section->IsActive())
 			{
-				if (Section->GetKeyType() == EParticleKey::Trigger)
-				{
-					if (Position > Section->GetStartTime() && LastPosition < Section->GetStartTime())
-					{
-						bTrigger = true;
-					}
-				}
-				else if (Section->GetKeyType() == EParticleKey::Toggle)
+				FIntegralCurve& ParticleKeyCurve = Section->GetParticleCurve();
+				FKeyHandle PreviousHandle = ParticleKeyCurve.FindKeyBeforeOrAt( Position );
+				if ( ParticleKeyCurve.IsKeyHandleValid( PreviousHandle ) )
 				{
-					if (Position >= Section->GetStartTime() && Position <= Section->GetEndTime())
-					{
-						bOn = true;
-					}
-					else if (Position >= Section->GetEndTime() && LastPosition < Section->GetEndTime())
+					FIntegralKey& PreviousKey = ParticleKeyCurve.GetKey( PreviousHandle );
+					if ( PreviousKey.Time > LastPosition )
 					{
-						bOff = true;
+						ParticleKey = (EParticleKey::Type)PreviousKey.Value;
+						bKeyFound = true;
 					}
 				}
 			}
 		}
 		
-		if (bTrigger || bOn || bOff)
+		if ( bKeyFound )
 		{
 			for (int32 i = 0; i < RuntimeObjects.Num(); ++i)
 			{
 				AEmitter* Emitter = Cast<AEmitter>(RuntimeObjects[i]);
 				if (Emitter)
 				{
-					if (bTrigger)
-					{
-						Emitter->ToggleActive();
-					}
-					else if (bOn)
+					if ( ParticleKey == EParticleKey::Active )
 					{
+						if ( Emitter->IsActive() )
+						{
+							Emitter->Deactivate();
+						}
 						Emitter->Activate();
 					}
-					else if (bOff)
+					else if( ParticleKey == EParticleKey::Inactive )
 					{
 						Emitter->Deactivate();
 					}
diff --git a/Engine/Source/Runtime/MovieSceneTracks/Private/Tracks/MovieSceneParticleTrack.cpp b/Engine/Source/Runtime/MovieSceneTracks/Private/Tracks/MovieSceneParticleTrack.cpp
--- a/Engine/Source/Runtime/MovieSceneTracks/Private/Tracks/MovieSceneParticleTrack.cpp
+++ b/Engine/Source/Runtime/MovieSceneTracks/Private/Tracks/MovieSceneParticleTrack.cpp
@@ -73,18 +73,28 @@ TRange<float> UMovieSceneParticleTrack::GetSectionBoundaries() const
 	return TRange<float>::Hull(Bounds);
 }
 
-
-void UMovieSceneParticleTrack::AddNewParticleSystem(float KeyTime, bool bTrigger)
+void UMovieSceneParticleTrack::AddNewKey( float KeyTime )
 {
-	EParticleKey::Type KeyType = bTrigger ? EParticleKey::Trigger : EParticleKey::Toggle;
-	// @todo Instead of a 0.1 second event, this should be 0 seconds, requires handling 0 size sections
-	float Duration = KeyType == EParticleKey::Trigger ? 0.1f : 1.f;
-
-	UMovieSceneParticleSection* NewSection = NewObject<UMovieSceneParticleSection>(this);
-	NewSection->InitialPlacement(ParticleSections, KeyTime, KeyTime + Duration, SupportsMultipleRows());
-	NewSection->SetKeyType(KeyType);
-
-	ParticleSections.Add(NewSection);
+	UMovieSceneParticleSection* NearestSection = Cast<UMovieSceneParticleSection>( MovieSceneHelpers::FindNearestSectionAtTime( ParticleSections, KeyTime ) );
+	if ( NearestSection == nullptr )
+	{
+		NearestSection = NewObject<UMovieSceneParticleSection>( this );
+		NearestSection->SetStartTime( KeyTime );
+		NearestSection->SetEndTime( KeyTime );
+		ParticleSections.Add(NearestSection);
+	}
+	else
+	{
+		if ( NearestSection->GetStartTime() > KeyTime )
+		{
+			NearestSection->SetStartTime( KeyTime );
+		}
+		if ( NearestSection->GetEndTime() < KeyTime )
+		{
+			NearestSection->SetEndTime( KeyTime );
+		}
+	}
+	NearestSection->AddKey( KeyTime, EParticleKey::Active );
 }
 
 
diff --git a/Engine/Source/Runtime/MovieSceneTracks/Public/Sections/MovieSceneParticleSection.h b/Engine/Source/Runtime/MovieSceneTracks/Public/Sections/MovieSceneParticleSection.h
--- a/Engine/Source/Runtime/MovieSceneTracks/Public/Sections/MovieSceneParticleSection.h
+++ b/Engine/Source/Runtime/MovieSceneTracks/Public/Sections/MovieSceneParticleSection.h
@@ -7,15 +7,15 @@
 
 
 /**
- * The kind of key this particle section is.
- */
+* Defines the types of particle keys.
+*/
 UENUM()
 namespace EParticleKey
 {
 	enum Type
 	{
-		Toggle,
-		Trigger
+		Active = 0,
+		Inactive = 1,
 	};
 }
 
@@ -29,17 +29,19 @@ class UMovieSceneParticleSection
 {
 	GENERATED_UCLASS_BODY()
 
-public:
+	void AddKey(float Time, EParticleKey::Type KeyType);
 
-	/** Sets this section's key type*/
-	void SetKeyType(EParticleKey::Type InKeyType) {KeyType = InKeyType;}
+	MOVIESCENETRACKS_API FIntegralCurve& GetParticleCurve();
 
-	/** Gets the key type for this section */
-	EParticleKey::Type GetKeyType() {return KeyType;}
-	
-private:
+	/**
+	* UMovieSceneSection interface
+	*/
+	virtual void MoveSection( float DeltaPosition, TSet<FKeyHandle>& KeyHandles ) override;
+	virtual void DilateSection( float DilationFactor, float Origin, TSet<FKeyHandle>& KeyHandles ) override;
+	virtual void GetKeyHandles( TSet<FKeyHandle>& KeyHandles ) const override;
 
-	/** The way this particle key will operate */
+private:
+	/** Curve containing the particle keys. */
 	UPROPERTY(EditAnywhere, Category="Particles")
-	TEnumAsByte<EParticleKey::Type> KeyType;
+	FIntegralCurve ParticleKeys;
 };
diff --git a/Engine/Source/Runtime/MovieSceneTracks/Public/Tracks/MovieSceneParticleTrack.h b/Engine/Source/Runtime/MovieSceneTracks/Public/Tracks/MovieSceneParticleTrack.h
--- a/Engine/Source/Runtime/MovieSceneTracks/Public/Tracks/MovieSceneParticleTrack.h
+++ b/Engine/Source/Runtime/MovieSceneTracks/Public/Tracks/MovieSceneParticleTrack.h
@@ -6,7 +6,6 @@
 #include "MovieSceneTrack.h"
 #include "MovieSceneParticleTrack.generated.h"
 
-
 /**
  * Handles triggering of particle emitters
  */
@@ -25,10 +24,9 @@ public:
 	virtual void RemoveSection( UMovieSceneSection* Section ) override;
 	virtual bool IsEmpty() const override;
 	virtual TRange<float> GetSectionBoundaries() const override;
-	virtual bool SupportsMultipleRows() const override { return false; }
 	virtual const TArray<UMovieSceneSection*>& GetAllSections() const override;
 	
-	virtual void AddNewParticleSystem(float KeyTime, bool bTrigger);
+	virtual void AddNewKey( float KeyTime );
 
 	virtual TArray<UMovieSceneSection*> GetAllParticleSections() const {return ParticleSections;}
 
