Sequencer: Refactored ISequencerTrackEditor::MakeSectionInterface to be safer; should also fix color track crash.

diff --git a/Engine/Source/Editor/MovieSceneTools/Private/Sections/ColorPropertySection.cpp b/Engine/Source/Editor/MovieSceneTools/Private/Sections/ColorPropertySection.cpp
--- a/Engine/Source/Editor/MovieSceneTools/Private/Sections/ColorPropertySection.cpp
+++ b/Engine/Source/Editor/MovieSceneTools/Private/Sections/ColorPropertySection.cpp
@@ -6,6 +6,7 @@
 #include "ColorPropertySection.h"
 #include "MovieSceneSequence.h"
 
+
 void FColorPropertySection::GenerateSectionLayout( class ISectionLayoutBuilder& LayoutBuilder ) const
 {
 	UMovieSceneColorSection* ColorSection = Cast<UMovieSceneColorSection>( &SectionObject );
@@ -16,6 +17,7 @@ void FColorPropertySection::GenerateSectionLayout( class ISectionLayoutBuilder&
 	LayoutBuilder.AddKeyArea( "A", NSLOCTEXT( "FColorPropertySection", "OpacityArea", "Opacity" ), MakeShareable( new FFloatCurveKeyArea( &ColorSection->GetAlphaCurve(), ColorSection ) ) );
 }
 
+
 int32 FColorPropertySection::OnPaintSection( const FGeometry& AllottedGeometry, const FSlateRect& SectionClippingRect, FSlateWindowElementList& OutDrawElements, int32 LayerId, bool bParentEnabled ) const
 {
 	const ESlateDrawEffect::Type DrawEffects = bParentEnabled ? ESlateDrawEffect::None : ESlateDrawEffect::DisabledEffect;
@@ -74,50 +76,8 @@ int32 FColorPropertySection::OnPaintSection( const FGeometry& AllottedGeometry,
 void FColorPropertySection::ConsolidateColorCurves( TArray< TKeyValuePair<float, FLinearColor> >& OutColorKeys, const UMovieSceneColorSection* Section ) const
 {
 	// Get the default color of the first instance
-	FLinearColor DefaultColor( 0.0f, 0.0f, 0.0f, 0.0f );
-
-	static const FName SlateColorName( "SlateColor" );
-
-	UMovieSceneSequence* Sequence = Sequencer->GetFocusedMovieSceneSequence();
-
-	const TArray<FMovieSceneBinding>& MovieSceneBindings = Sequence->GetMovieScene()->GetBindings();
-
-	bool bFoundColor = false;
-	for ( int32 BindingIndex = 0; BindingIndex < MovieSceneBindings.Num() && !bFoundColor; ++BindingIndex )
-	{
-		const FMovieSceneBinding& MovieSceneBinding = MovieSceneBindings[BindingIndex];
-
-		for ( int32 TrackIndex = 0; TrackIndex < MovieSceneBinding.GetTracks().Num() && !bFoundColor; ++TrackIndex )
-		{
-			if ( MovieSceneBinding.GetTracks()[TrackIndex] == Track )
-			{
-				UObject* RuntimeObject = Sequence->FindObject( MovieSceneBinding.GetObjectGuid() );
-
-				if ( RuntimeObject != nullptr )
-				{
-					UProperty* Property = RuntimeObject->GetClass()->FindPropertyByName( CastChecked<UMovieSceneColorTrack>( Track )->GetPropertyName() );
-					UStructProperty* ColorStructProp = Cast<UStructProperty>( Property );
-					if ( ColorStructProp && ColorStructProp->Struct )
-					{
-						if ( ColorStructProp->Struct->GetFName() == SlateColorName )
-						{
-							DefaultColor = (*Property->ContainerPtrToValuePtr<FSlateColor>( RuntimeObject )).GetSpecifiedColor();
-						}
-						else if ( ColorStructProp->Struct->GetFName() == NAME_LinearColor )
-						{
-							DefaultColor = *Property->ContainerPtrToValuePtr<FLinearColor>( RuntimeObject );
-						}
-						else
-						{
-							DefaultColor = Property->ContainerPtrToValuePtr<FColor>( RuntimeObject )->ReinterpretAsLinear();
-						}
-						bFoundColor = true;
-						break;
-					}
-				}
-			}
-		}
-	}
+	static const FName SlateColorName("SlateColor");
+	FLinearColor DefaultColor = FindSlateColor(SlateColorName);
 
 	// @todo Sequencer Optimize - This could all get cached, instead of recalculating everything every OnPaint
 
@@ -180,4 +140,51 @@ void FColorPropertySection::ConsolidateColorCurves( TArray< TKeyValuePair<float,
 	{
 		OutColorKeys.Add( TKeyValuePair<float, FLinearColor>( TimesWithKeys[i], Section->Eval( TimesWithKeys[i], DefaultColor ) ) );
 	}
-}
\ No newline at end of file
+}
+
+
+FLinearColor FColorPropertySection::FindSlateColor(const FName& ColorName) const
+{
+	const UMovieSceneSequence* FocusedSequence = Sequencer->GetFocusedMovieSceneSequence();
+	const TArray<FMovieSceneBinding>& FocusedBindings = FocusedSequence->GetMovieScene()->GetBindings();
+
+	for (const FMovieSceneBinding& Binding : FocusedBindings)
+	{
+		for (const UMovieSceneTrack* BindingTrack : Binding.GetTracks())
+		{
+			if (BindingTrack != &Track)
+			{
+				continue;
+			}
+
+			UObject* RuntimeObject = FocusedSequence->FindObject(Binding.GetObjectGuid());
+
+			if (RuntimeObject == nullptr)
+			{
+				continue;
+			}
+
+			UProperty* Property = RuntimeObject->GetClass()->FindPropertyByName(Track.GetPropertyName());
+			UStructProperty* ColorStructProp = Cast<UStructProperty>(Property);
+			
+			if ((ColorStructProp == nullptr) || (ColorStructProp->Struct == nullptr))
+			{
+				continue;
+			}
+
+			if (ColorStructProp->Struct->GetFName() == ColorName)
+			{
+				return (*Property->ContainerPtrToValuePtr<FSlateColor>(RuntimeObject)).GetSpecifiedColor();
+			}
+
+			if (ColorStructProp->Struct->GetFName() == NAME_LinearColor)
+			{
+				return *Property->ContainerPtrToValuePtr<FLinearColor>(RuntimeObject);
+			}
+
+			return Property->ContainerPtrToValuePtr<FColor>(RuntimeObject)->ReinterpretAsLinear();
+		}
+	}
+
+	return FLinearColor(0.0f, 0.0f, 0.0f, 0.0f);
+}
diff --git a/Engine/Source/Editor/MovieSceneTools/Private/Sections/ColorPropertySection.h b/Engine/Source/Editor/MovieSceneTools/Private/Sections/ColorPropertySection.h
--- a/Engine/Source/Editor/MovieSceneTools/Private/Sections/ColorPropertySection.h
+++ b/Engine/Source/Editor/MovieSceneTools/Private/Sections/ColorPropertySection.h
@@ -7,21 +7,45 @@ class UMovieSceneColorSection;
 /**
 * A color section implementation
 */
-class FColorPropertySection : public FPropertySection
+class FColorPropertySection
+	: public FPropertySection
 {
 public:
-	FColorPropertySection( UMovieSceneSection& InSectionObject, FName SectionName, ISequencer* InSequencer, UMovieSceneTrack* InTrack )
-		: FPropertySection( InSectionObject, SectionName )
-		, Sequencer( InSequencer )
-		, Track( InTrack ) {}
 
-	virtual void GenerateSectionLayout( class ISectionLayoutBuilder& LayoutBuilder ) const override;
+	/**
+	 * Create and initialize a new instance.
+	 *
+	 * @param InSectionObject
+	 * @param SectionName The name of the section.
+	 * @param InSequencer The sequencer that manages the section.
+	 * @param InTrack The track that owns the section.
+	 */
+	FColorPropertySection(UMovieSceneSection& InSectionObject, ISequencer* InSequencer, UMovieSceneTrack& InTrack)
+		: FPropertySection(InSectionObject, InTrack.GetTrackName())
+		, Sequencer(InSequencer)
+		, Track(*Cast<UMovieSceneColorTrack>(&InTrack))
+	{ }
 
-	virtual int32 OnPaintSection( const FGeometry& AllottedGeometry, const FSlateRect& SectionClippingRect, FSlateWindowElementList& OutDrawElements, int32 LayerId, bool bParentEnabled ) const override;
+public:
+
+	// FPropertySection interface
+
+	virtual void GenerateSectionLayout(class ISectionLayoutBuilder& LayoutBuilder) const override;
+	virtual int32 OnPaintSection(const FGeometry& AllottedGeometry, const FSlateRect& SectionClippingRect, FSlateWindowElementList& OutDrawElements, int32 LayerId, bool bParentEnabled) const override;
+
+protected:
+
+	/** Consolidate color curves for all track sections. */
+	void ConsolidateColorCurves(TArray< TKeyValuePair<float, FLinearColor> >& OutColorKeys, const UMovieSceneColorSection* Section) const;
+	
+	/** Find the Slate color of the specified name in the track. */
+	FLinearColor FindSlateColor(const FName& ColorName) const;
 
 private:
-	void ConsolidateColorCurves( TArray< TKeyValuePair<float, FLinearColor> >& OutColorKeys, const UMovieSceneColorSection* Section ) const;
 
+	/** The sequencer that manages the section. */
 	ISequencer* Sequencer;
-	UMovieSceneTrack* Track;
+
+	/** The track that owns the section. */
+	UMovieSceneColorTrack& Track;
 };
\ No newline at end of file
diff --git a/Engine/Source/Editor/MovieSceneTools/Private/TrackEditors/AttachTrackEditor.cpp b/Engine/Source/Editor/MovieSceneTools/Private/TrackEditors/AttachTrackEditor.cpp
--- a/Engine/Source/Editor/MovieSceneTools/Private/TrackEditors/AttachTrackEditor.cpp
+++ b/Engine/Source/Editor/MovieSceneTools/Private/TrackEditors/AttachTrackEditor.cpp
@@ -189,13 +189,11 @@ bool F3DAttachTrackEditor::SupportsType( TSubclassOf<UMovieSceneTrack> Type ) co
 }
 
 
-TSharedRef<ISequencerSection> F3DAttachTrackEditor::MakeSectionInterface( UMovieSceneSection& SectionObject, UMovieSceneTrack* Track )
+TSharedRef<ISequencerSection> F3DAttachTrackEditor::MakeSectionInterface( UMovieSceneSection& SectionObject, UMovieSceneTrack& Track )
 {
 	check( SupportsType( SectionObject.GetOuter()->GetClass() ) );
 
-	TSharedRef<ISequencerSection> NewSection( new F3DAttachSection( SectionObject, this ) );
-
-	return NewSection;
+	return MakeShareable( new F3DAttachSection( SectionObject, this ) );
 }
 
 
diff --git a/Engine/Source/Editor/MovieSceneTools/Private/TrackEditors/AttachTrackEditor.h b/Engine/Source/Editor/MovieSceneTools/Private/TrackEditors/AttachTrackEditor.h
--- a/Engine/Source/Editor/MovieSceneTools/Private/TrackEditors/AttachTrackEditor.h
+++ b/Engine/Source/Editor/MovieSceneTools/Private/TrackEditors/AttachTrackEditor.h
@@ -46,7 +46,7 @@ public:
 
 	virtual void AddKey( const FGuid& ObjectGuid, UObject* AdditionalAsset = NULL ) override;
 	virtual void BuildObjectBindingTrackMenu(FMenuBuilder& MenuBuilder, const FGuid& ObjectBinding, const UClass* ObjectClass) override;
-	virtual TSharedRef<ISequencerSection> MakeSectionInterface( UMovieSceneSection& SectionObject, UMovieSceneTrack* Track ) override;
+	virtual TSharedRef<ISequencerSection> MakeSectionInterface( UMovieSceneSection& SectionObject, UMovieSceneTrack& Track ) override;
 	virtual bool SupportsType( TSubclassOf<UMovieSceneTrack> Type ) const override;
 
 private:
diff --git a/Engine/Source/Editor/MovieSceneTools/Private/TrackEditors/AudioTrackEditor.cpp b/Engine/Source/Editor/MovieSceneTools/Private/TrackEditors/AudioTrackEditor.cpp
--- a/Engine/Source/Editor/MovieSceneTools/Private/TrackEditors/AudioTrackEditor.cpp
+++ b/Engine/Source/Editor/MovieSceneTools/Private/TrackEditors/AudioTrackEditor.cpp
@@ -472,14 +472,12 @@ bool FAudioTrackEditor::SupportsType( TSubclassOf<UMovieSceneTrack> Type ) const
 }
 
 
-TSharedRef<ISequencerSection> FAudioTrackEditor::MakeSectionInterface( UMovieSceneSection& SectionObject, UMovieSceneTrack* Track )
+TSharedRef<ISequencerSection> FAudioTrackEditor::MakeSectionInterface( UMovieSceneSection& SectionObject, UMovieSceneTrack& Track )
 {
 	check( SupportsType( SectionObject.GetOuter()->GetClass() ) );
 	
-	bool bIsAMasterTrack = Cast<UMovieScene>(Track->GetOuter())->IsAMasterTrack(Track);
-	TSharedRef<ISequencerSection> NewSection( new FAudioSection(SectionObject, bIsAMasterTrack) );
-
-	return NewSection;
+	bool bIsAMasterTrack = Cast<UMovieScene>(Track.GetOuter())->IsAMasterTrack(Track);
+	return MakeShareable( new FAudioSection(SectionObject, bIsAMasterTrack) );
 }
 
 
diff --git a/Engine/Source/Editor/MovieSceneTools/Private/TrackEditors/AudioTrackEditor.h b/Engine/Source/Editor/MovieSceneTools/Private/TrackEditors/AudioTrackEditor.h
--- a/Engine/Source/Editor/MovieSceneTools/Private/TrackEditors/AudioTrackEditor.h
+++ b/Engine/Source/Editor/MovieSceneTools/Private/TrackEditors/AudioTrackEditor.h
@@ -34,7 +34,7 @@ public:
 	// ISequencerTrackEditor interface
 
 	virtual bool HandleAssetAdded(UObject* Asset, const FGuid& TargetObjectGuid) override;
-	virtual TSharedRef<ISequencerSection> MakeSectionInterface( UMovieSceneSection& SectionObject, UMovieSceneTrack* Track ) override;
+	virtual TSharedRef<ISequencerSection> MakeSectionInterface( UMovieSceneSection& SectionObject, UMovieSceneTrack& Track ) override;
 	virtual bool SupportsType( TSubclassOf<UMovieSceneTrack> Type ) const override;
 	
 private:
diff --git a/Engine/Source/Editor/MovieSceneTools/Private/TrackEditors/EventTrackEditor.cpp b/Engine/Source/Editor/MovieSceneTools/Private/TrackEditors/EventTrackEditor.cpp
--- a/Engine/Source/Editor/MovieSceneTools/Private/TrackEditors/EventTrackEditor.cpp
+++ b/Engine/Source/Editor/MovieSceneTools/Private/TrackEditors/EventTrackEditor.cpp
@@ -65,7 +65,7 @@ void FEventTrackEditor::BuildObjectBindingTrackMenu(FMenuBuilder& MenuBuilder, c
 }
 
 
-TSharedRef<ISequencerSection> FEventTrackEditor::MakeSectionInterface(UMovieSceneSection& SectionObject, UMovieSceneTrack* Track)
+TSharedRef<ISequencerSection> FEventTrackEditor::MakeSectionInterface(UMovieSceneSection& SectionObject, UMovieSceneTrack& Track)
 {
 	return MakeShareable(new FEventTrackSection(SectionObject/*, Track->GetTrackName()*/, GetSequencer()));
 }
diff --git a/Engine/Source/Editor/MovieSceneTools/Private/TrackEditors/EventTrackEditor.h b/Engine/Source/Editor/MovieSceneTools/Private/TrackEditors/EventTrackEditor.h
--- a/Engine/Source/Editor/MovieSceneTools/Private/TrackEditors/EventTrackEditor.h
+++ b/Engine/Source/Editor/MovieSceneTools/Private/TrackEditors/EventTrackEditor.h
@@ -41,7 +41,7 @@ public:
 	virtual void AddKey(const FGuid& ObjectGuid, UObject* AdditionalAsset = nullptr) override;
 	virtual void BuildAddTrackMenu(FMenuBuilder& MenuBuilder) override;
 	virtual void BuildObjectBindingTrackMenu(FMenuBuilder& MenuBuilder, const FGuid& ObjectBinding, const UClass* ObjectClass) override;
-	virtual TSharedRef<ISequencerSection> MakeSectionInterface(UMovieSceneSection& SectionObject, UMovieSceneTrack* Track) override;
+	virtual TSharedRef<ISequencerSection> MakeSectionInterface(UMovieSceneSection& SectionObject, UMovieSceneTrack& Track) override;
 	virtual bool SupportsType( TSubclassOf<UMovieSceneTrack> Type ) const override;
 
 private:
diff --git a/Engine/Source/Editor/MovieSceneTools/Private/TrackEditors/MaterialTrackEditor.cpp b/Engine/Source/Editor/MovieSceneTools/Private/TrackEditors/MaterialTrackEditor.cpp
--- a/Engine/Source/Editor/MovieSceneTools/Private/TrackEditors/MaterialTrackEditor.cpp
+++ b/Engine/Source/Editor/MovieSceneTools/Private/TrackEditors/MaterialTrackEditor.cpp
@@ -16,10 +16,11 @@ FMaterialTrackEditor::FMaterialTrackEditor( TSharedRef<ISequencer> InSequencer )
 }
 
 
-TSharedRef<ISequencerSection> FMaterialTrackEditor::MakeSectionInterface( UMovieSceneSection& SectionObject, UMovieSceneTrack* Track )
+TSharedRef<ISequencerSection> FMaterialTrackEditor::MakeSectionInterface( UMovieSceneSection& SectionObject, UMovieSceneTrack& Track )
 {
 	UMovieSceneMaterialParameterSection* ParameterSection = Cast<UMovieSceneMaterialParameterSection>(&SectionObject);
 	checkf( ParameterSection != nullptr, TEXT("Unsupported section type.") );
+
 	return MakeShareable(new FMaterialParameterSection( *ParameterSection, ParameterSection->GetFName()));
 }
 
diff --git a/Engine/Source/Editor/MovieSceneTools/Private/TrackEditors/MaterialTrackEditor.h b/Engine/Source/Editor/MovieSceneTools/Private/TrackEditors/MaterialTrackEditor.h
--- a/Engine/Source/Editor/MovieSceneTools/Private/TrackEditors/MaterialTrackEditor.h
+++ b/Engine/Source/Editor/MovieSceneTools/Private/TrackEditors/MaterialTrackEditor.h
@@ -25,7 +25,7 @@ public:
 	// ISequencerTrackEditor interface
 
 	virtual TSharedPtr<SWidget> BuildOutlinerEditWidget( const FGuid& ObjectBinding, UMovieSceneTrack* Track ) override;
-	virtual TSharedRef<ISequencerSection> MakeSectionInterface( UMovieSceneSection& SectionObject, UMovieSceneTrack* Track ) override;
+	virtual TSharedRef<ISequencerSection> MakeSectionInterface( UMovieSceneSection& SectionObject, UMovieSceneTrack& Track ) override;
 
 protected:
 
diff --git a/Engine/Source/Editor/MovieSceneTools/Private/TrackEditors/ParticleTrackEditor.cpp b/Engine/Source/Editor/MovieSceneTools/Private/TrackEditors/ParticleTrackEditor.cpp
--- a/Engine/Source/Editor/MovieSceneTools/Private/TrackEditors/ParticleTrackEditor.cpp
+++ b/Engine/Source/Editor/MovieSceneTools/Private/TrackEditors/ParticleTrackEditor.cpp
@@ -235,14 +235,12 @@ bool FParticleTrackEditor::SupportsType( TSubclassOf<UMovieSceneTrack> Type ) co
 }
 
 
-TSharedRef<ISequencerSection> FParticleTrackEditor::MakeSectionInterface( UMovieSceneSection& SectionObject, UMovieSceneTrack* Track )
+TSharedRef<ISequencerSection> FParticleTrackEditor::MakeSectionInterface( UMovieSceneSection& SectionObject, UMovieSceneTrack& Track )
 {
 	check( SupportsType( SectionObject.GetOuter()->GetClass() ) );
-	
-	const TSharedPtr<ISequencer> OwningSequencer = GetSequencer();
-	TSharedRef<ISequencerSection> NewSection( new FParticleSection( SectionObject, OwningSequencer.ToSharedRef() ) );
 
-	return NewSection;
+	const TSharedPtr<ISequencer> OwningSequencer = GetSequencer();
+	return MakeShareable( new FParticleSection( SectionObject, OwningSequencer.ToSharedRef() ) );
 }
 
 
diff --git a/Engine/Source/Editor/MovieSceneTools/Private/TrackEditors/ParticleTrackEditor.h b/Engine/Source/Editor/MovieSceneTools/Private/TrackEditors/ParticleTrackEditor.h
--- a/Engine/Source/Editor/MovieSceneTools/Private/TrackEditors/ParticleTrackEditor.h
+++ b/Engine/Source/Editor/MovieSceneTools/Private/TrackEditors/ParticleTrackEditor.h
@@ -36,7 +36,7 @@ public:
 	// ISequencerTrackEditor interface
 
 	virtual void BuildObjectBindingTrackMenu(FMenuBuilder& MenuBuilder, const FGuid& ObjectBinding, const UClass* ObjectClass) override;
-	virtual TSharedRef<ISequencerSection> MakeSectionInterface( UMovieSceneSection& SectionObject, UMovieSceneTrack* Track ) override;
+	virtual TSharedRef<ISequencerSection> MakeSectionInterface( UMovieSceneSection& SectionObject, UMovieSceneTrack& Track ) override;
 	virtual bool SupportsType( TSubclassOf<UMovieSceneTrack> Type ) const override;
 	
 	void AddParticleKey(const FGuid ObjectGuid);
diff --git a/Engine/Source/Editor/MovieSceneTools/Private/TrackEditors/PathTrackEditor.cpp b/Engine/Source/Editor/MovieSceneTools/Private/TrackEditors/PathTrackEditor.cpp
--- a/Engine/Source/Editor/MovieSceneTools/Private/TrackEditors/PathTrackEditor.cpp
+++ b/Engine/Source/Editor/MovieSceneTools/Private/TrackEditors/PathTrackEditor.cpp
@@ -176,13 +176,11 @@ bool F3DPathTrackEditor::SupportsType( TSubclassOf<UMovieSceneTrack> Type ) cons
 }
 
 
-TSharedRef<ISequencerSection> F3DPathTrackEditor::MakeSectionInterface( UMovieSceneSection& SectionObject, UMovieSceneTrack* Track )
+TSharedRef<ISequencerSection> F3DPathTrackEditor::MakeSectionInterface( UMovieSceneSection& SectionObject, UMovieSceneTrack& Track )
 {
 	check( SupportsType( SectionObject.GetOuter()->GetClass() ) );
 
-	TSharedRef<ISequencerSection> NewSection( new F3DPathSection( SectionObject, this ) );
-
-	return NewSection;
+	return MakeShareable( new F3DPathSection( SectionObject, this ) );
 }
 
 
diff --git a/Engine/Source/Editor/MovieSceneTools/Private/TrackEditors/PathTrackEditor.h b/Engine/Source/Editor/MovieSceneTools/Private/TrackEditors/PathTrackEditor.h
--- a/Engine/Source/Editor/MovieSceneTools/Private/TrackEditors/PathTrackEditor.h
+++ b/Engine/Source/Editor/MovieSceneTools/Private/TrackEditors/PathTrackEditor.h
@@ -46,7 +46,7 @@ public:
 
 	virtual void AddKey( const FGuid& ObjectGuid, UObject* AdditionalAsset = NULL ) override;
 	virtual void BuildObjectBindingTrackMenu(FMenuBuilder& MenuBuilder, const FGuid& ObjectBinding, const UClass* ObjectClass) override;
-	virtual TSharedRef<ISequencerSection> MakeSectionInterface( UMovieSceneSection& SectionObject, UMovieSceneTrack* Track ) override;
+	virtual TSharedRef<ISequencerSection> MakeSectionInterface( UMovieSceneSection& SectionObject, UMovieSceneTrack& Track ) override;
 	virtual bool SupportsType( TSubclassOf<UMovieSceneTrack> Type ) const override;
 
 private:
diff --git a/Engine/Source/Editor/MovieSceneTools/Private/TrackEditors/PropertyTrackEditors/BoolPropertyTrackEditor.cpp b/Engine/Source/Editor/MovieSceneTools/Private/TrackEditors/PropertyTrackEditors/BoolPropertyTrackEditor.cpp
--- a/Engine/Source/Editor/MovieSceneTools/Private/TrackEditors/PropertyTrackEditors/BoolPropertyTrackEditor.cpp
+++ b/Engine/Source/Editor/MovieSceneTools/Private/TrackEditors/PropertyTrackEditors/BoolPropertyTrackEditor.cpp
@@ -12,9 +12,9 @@ TSharedRef<ISequencerTrackEditor> FBoolPropertyTrackEditor::CreateTrackEditor( T
 }
 
 
-TSharedRef<ISequencerSection> FBoolPropertyTrackEditor::MakeSectionInterface( UMovieSceneSection& SectionObject, UMovieSceneTrack* Track )
+TSharedRef<ISequencerSection> FBoolPropertyTrackEditor::MakeSectionInterface( UMovieSceneSection& SectionObject, UMovieSceneTrack& Track )
 {
-	return MakeShareable(new FBoolPropertySection( SectionObject, Track->GetTrackName(), GetSequencer().Get() ));
+	return MakeShareable(new FBoolPropertySection( SectionObject, Track.GetTrackName(), GetSequencer().Get() ));
 }
 
 
diff --git a/Engine/Source/Editor/MovieSceneTools/Private/TrackEditors/PropertyTrackEditors/BoolPropertyTrackEditor.h b/Engine/Source/Editor/MovieSceneTools/Private/TrackEditors/PropertyTrackEditors/BoolPropertyTrackEditor.h
--- a/Engine/Source/Editor/MovieSceneTools/Private/TrackEditors/PropertyTrackEditors/BoolPropertyTrackEditor.h
+++ b/Engine/Source/Editor/MovieSceneTools/Private/TrackEditors/PropertyTrackEditors/BoolPropertyTrackEditor.h
@@ -33,7 +33,7 @@ public:
 
 	// ISequencerTrackEditor interface
 
-	virtual TSharedRef<ISequencerSection> MakeSectionInterface( UMovieSceneSection& SectionObject, UMovieSceneTrack* Track ) override;
+	virtual TSharedRef<ISequencerSection> MakeSectionInterface( UMovieSceneSection& SectionObject, UMovieSceneTrack& Track ) override;
 
 protected:
 
diff --git a/Engine/Source/Editor/MovieSceneTools/Private/TrackEditors/PropertyTrackEditors/BytePropertyTrackEditor.cpp b/Engine/Source/Editor/MovieSceneTools/Private/TrackEditors/PropertyTrackEditors/BytePropertyTrackEditor.cpp
--- a/Engine/Source/Editor/MovieSceneTools/Private/TrackEditors/PropertyTrackEditors/BytePropertyTrackEditor.cpp
+++ b/Engine/Source/Editor/MovieSceneTools/Private/TrackEditors/PropertyTrackEditors/BytePropertyTrackEditor.cpp
@@ -13,9 +13,9 @@ TSharedRef<ISequencerTrackEditor> FBytePropertyTrackEditor::CreateTrackEditor( T
 }
 
 
-TSharedRef<ISequencerSection> FBytePropertyTrackEditor::MakeSectionInterface( UMovieSceneSection& SectionObject, UMovieSceneTrack* Track )
+TSharedRef<ISequencerSection> FBytePropertyTrackEditor::MakeSectionInterface( UMovieSceneSection& SectionObject, UMovieSceneTrack& Track )
 {
-	return MakeShareable( new FBytePropertySection( SectionObject, Track->GetTrackName(), Cast<UMovieSceneByteTrack>( SectionObject.GetOuter() )->GetEnum() ) );
+	return MakeShareable( new FBytePropertySection( SectionObject, Track.GetTrackName(), Cast<UMovieSceneByteTrack>( SectionObject.GetOuter() )->GetEnum() ) );
 }
 
 
diff --git a/Engine/Source/Editor/MovieSceneTools/Private/TrackEditors/PropertyTrackEditors/BytePropertyTrackEditor.h b/Engine/Source/Editor/MovieSceneTools/Private/TrackEditors/PropertyTrackEditors/BytePropertyTrackEditor.h
--- a/Engine/Source/Editor/MovieSceneTools/Private/TrackEditors/PropertyTrackEditors/BytePropertyTrackEditor.h
+++ b/Engine/Source/Editor/MovieSceneTools/Private/TrackEditors/PropertyTrackEditors/BytePropertyTrackEditor.h
@@ -35,7 +35,7 @@ public:
 	// ISequencerTrackEditor interface
 
 	virtual UMovieSceneTrack* AddTrack( UMovieScene* FocusedMovieScene, const FGuid& ObjectHandle, TSubclassOf<class UMovieSceneTrack> TrackClass, FName UniqueTypeName ) override;
-	virtual TSharedRef<ISequencerSection> MakeSectionInterface( UMovieSceneSection& SectionObject, UMovieSceneTrack* Track ) override;
+	virtual TSharedRef<ISequencerSection> MakeSectionInterface( UMovieSceneSection& SectionObject, UMovieSceneTrack& Track ) override;
 
 protected:
 
diff --git a/Engine/Source/Editor/MovieSceneTools/Private/TrackEditors/PropertyTrackEditors/ColorPropertyTrackEditor.cpp b/Engine/Source/Editor/MovieSceneTools/Private/TrackEditors/PropertyTrackEditors/ColorPropertyTrackEditor.cpp
--- a/Engine/Source/Editor/MovieSceneTools/Private/TrackEditors/PropertyTrackEditors/ColorPropertyTrackEditor.cpp
+++ b/Engine/Source/Editor/MovieSceneTools/Private/TrackEditors/PropertyTrackEditors/ColorPropertyTrackEditor.cpp
@@ -6,15 +6,15 @@
 #include "ColorPropertySection.h"
 
 
-TSharedRef<ISequencerTrackEditor> FColorPropertyTrackEditor::CreateTrackEditor( TSharedRef<ISequencer> InSequencer )
+TSharedRef<ISequencerTrackEditor> FColorPropertyTrackEditor::CreateTrackEditor(TSharedRef<ISequencer> InSequencer)
 {
-	return MakeShareable( new FColorPropertyTrackEditor( InSequencer ) );
+	return MakeShareable(new FColorPropertyTrackEditor(InSequencer));
 }
 
 
-TSharedRef<ISequencerSection> FColorPropertyTrackEditor::MakeSectionInterface( UMovieSceneSection& SectionObject, UMovieSceneTrack* Track )
+TSharedRef<ISequencerSection> FColorPropertyTrackEditor::MakeSectionInterface( UMovieSceneSection& SectionObject, UMovieSceneTrack& Track )
 {
-	return MakeShareable(new FColorPropertySection( SectionObject, Track->GetTrackName(), GetSequencer().Get(), Track ));
+	return MakeShareable(new FColorPropertySection(SectionObject, GetSequencer().Get(), Track));
 }
 
 
diff --git a/Engine/Source/Editor/MovieSceneTools/Private/TrackEditors/PropertyTrackEditors/ColorPropertyTrackEditor.h b/Engine/Source/Editor/MovieSceneTools/Private/TrackEditors/PropertyTrackEditors/ColorPropertyTrackEditor.h
--- a/Engine/Source/Editor/MovieSceneTools/Private/TrackEditors/PropertyTrackEditors/ColorPropertyTrackEditor.h
+++ b/Engine/Source/Editor/MovieSceneTools/Private/TrackEditors/PropertyTrackEditors/ColorPropertyTrackEditor.h
@@ -34,7 +34,7 @@ public:
 
 	// ISequencerTrackEditor interface
 
-	virtual TSharedRef<ISequencerSection> MakeSectionInterface( UMovieSceneSection& SectionObject, UMovieSceneTrack* Track ) override;
+	virtual TSharedRef<ISequencerSection> MakeSectionInterface( UMovieSceneSection& SectionObject, UMovieSceneTrack& Track ) override;
 
 protected:
 
diff --git a/Engine/Source/Editor/MovieSceneTools/Private/TrackEditors/PropertyTrackEditors/FloatPropertyTrackEditor.cpp b/Engine/Source/Editor/MovieSceneTools/Private/TrackEditors/PropertyTrackEditors/FloatPropertyTrackEditor.cpp
--- a/Engine/Source/Editor/MovieSceneTools/Private/TrackEditors/PropertyTrackEditors/FloatPropertyTrackEditor.cpp
+++ b/Engine/Source/Editor/MovieSceneTools/Private/TrackEditors/PropertyTrackEditors/FloatPropertyTrackEditor.cpp
@@ -11,9 +11,9 @@ TSharedRef<ISequencerTrackEditor> FFloatPropertyTrackEditor::CreateTrackEditor(
 }
 
 
-TSharedRef<ISequencerSection> FFloatPropertyTrackEditor::MakeSectionInterface( UMovieSceneSection& SectionObject, UMovieSceneTrack* Track )
+TSharedRef<ISequencerSection> FFloatPropertyTrackEditor::MakeSectionInterface( UMovieSceneSection& SectionObject, UMovieSceneTrack& Track )
 {
-	return MakeShareable(new FFloatPropertySection( SectionObject, Track->GetTrackName() ));
+	return MakeShareable(new FFloatPropertySection( SectionObject, Track.GetTrackName() ));
 }
 
 
diff --git a/Engine/Source/Editor/MovieSceneTools/Private/TrackEditors/PropertyTrackEditors/FloatPropertyTrackEditor.h b/Engine/Source/Editor/MovieSceneTools/Private/TrackEditors/PropertyTrackEditors/FloatPropertyTrackEditor.h
--- a/Engine/Source/Editor/MovieSceneTools/Private/TrackEditors/PropertyTrackEditors/FloatPropertyTrackEditor.h
+++ b/Engine/Source/Editor/MovieSceneTools/Private/TrackEditors/PropertyTrackEditors/FloatPropertyTrackEditor.h
@@ -33,7 +33,7 @@ public:
 
 	// ISequencerTrackEditor interface
 
-	virtual TSharedRef<ISequencerSection> MakeSectionInterface( UMovieSceneSection& SectionObject, UMovieSceneTrack* Track ) override;
+	virtual TSharedRef<ISequencerSection> MakeSectionInterface( UMovieSceneSection& SectionObject, UMovieSceneTrack& Track ) override;
 
 protected:
 
diff --git a/Engine/Source/Editor/MovieSceneTools/Private/TrackEditors/PropertyTrackEditors/VectorPropertyTrackEditor.cpp b/Engine/Source/Editor/MovieSceneTools/Private/TrackEditors/PropertyTrackEditors/VectorPropertyTrackEditor.cpp
--- a/Engine/Source/Editor/MovieSceneTools/Private/TrackEditors/PropertyTrackEditors/VectorPropertyTrackEditor.cpp
+++ b/Engine/Source/Editor/MovieSceneTools/Private/TrackEditors/PropertyTrackEditors/VectorPropertyTrackEditor.cpp
@@ -12,9 +12,9 @@ TSharedRef<ISequencerTrackEditor> FVectorPropertyTrackEditor::CreateTrackEditor(
 }
 
 
-TSharedRef<ISequencerSection> FVectorPropertyTrackEditor::MakeSectionInterface( UMovieSceneSection& SectionObject, UMovieSceneTrack* Track )
+TSharedRef<ISequencerSection> FVectorPropertyTrackEditor::MakeSectionInterface( UMovieSceneSection& SectionObject, UMovieSceneTrack& Track )
 {
-	return MakeShareable(new FVectorPropertySection( SectionObject, Track->GetTrackName() ));
+	return MakeShareable(new FVectorPropertySection( SectionObject, Track.GetTrackName() ));
 }
 
 
diff --git a/Engine/Source/Editor/MovieSceneTools/Private/TrackEditors/PropertyTrackEditors/VectorPropertyTrackEditor.h b/Engine/Source/Editor/MovieSceneTools/Private/TrackEditors/PropertyTrackEditors/VectorPropertyTrackEditor.h
--- a/Engine/Source/Editor/MovieSceneTools/Private/TrackEditors/PropertyTrackEditors/VectorPropertyTrackEditor.h
+++ b/Engine/Source/Editor/MovieSceneTools/Private/TrackEditors/PropertyTrackEditors/VectorPropertyTrackEditor.h
@@ -34,7 +34,7 @@ public:
 
 	// ISequencerTrackEditor interface
 
-	virtual TSharedRef<ISequencerSection> MakeSectionInterface( UMovieSceneSection& SectionObject, UMovieSceneTrack* Track ) override;
+	virtual TSharedRef<ISequencerSection> MakeSectionInterface( UMovieSceneSection& SectionObject, UMovieSceneTrack& Track ) override;
 
 protected:
 
diff --git a/Engine/Source/Editor/MovieSceneTools/Private/TrackEditors/PropertyTrackEditors/VisibilityPropertyTrackEditor.cpp b/Engine/Source/Editor/MovieSceneTools/Private/TrackEditors/PropertyTrackEditors/VisibilityPropertyTrackEditor.cpp
--- a/Engine/Source/Editor/MovieSceneTools/Private/TrackEditors/PropertyTrackEditors/VisibilityPropertyTrackEditor.cpp
+++ b/Engine/Source/Editor/MovieSceneTools/Private/TrackEditors/PropertyTrackEditors/VisibilityPropertyTrackEditor.cpp
@@ -12,9 +12,9 @@ TSharedRef<ISequencerTrackEditor> FVisibilityPropertyTrackEditor::CreateTrackEdi
 }
 
 
-TSharedRef<ISequencerSection> FVisibilityPropertyTrackEditor::MakeSectionInterface( UMovieSceneSection& SectionObject, UMovieSceneTrack* Track )
+TSharedRef<ISequencerSection> FVisibilityPropertyTrackEditor::MakeSectionInterface( UMovieSceneSection& SectionObject, UMovieSceneTrack& Track )
 {
-	return MakeShareable(new FVisibilityPropertySection( SectionObject, Track->GetTrackName(), GetSequencer().Get() ));
+	return MakeShareable(new FVisibilityPropertySection( SectionObject, Track.GetTrackName(), GetSequencer().Get() ));
 }
 
 
diff --git a/Engine/Source/Editor/MovieSceneTools/Private/TrackEditors/PropertyTrackEditors/VisibilityPropertyTrackEditor.h b/Engine/Source/Editor/MovieSceneTools/Private/TrackEditors/PropertyTrackEditors/VisibilityPropertyTrackEditor.h
--- a/Engine/Source/Editor/MovieSceneTools/Private/TrackEditors/PropertyTrackEditors/VisibilityPropertyTrackEditor.h
+++ b/Engine/Source/Editor/MovieSceneTools/Private/TrackEditors/PropertyTrackEditors/VisibilityPropertyTrackEditor.h
@@ -32,7 +32,7 @@ public:
 
 	// ISequencerTrackEditor interface
 
-	virtual TSharedRef<ISequencerSection> MakeSectionInterface( UMovieSceneSection& SectionObject, UMovieSceneTrack* Track ) override;
+	virtual TSharedRef<ISequencerSection> MakeSectionInterface( UMovieSceneSection& SectionObject, UMovieSceneTrack& Track ) override;
 
 protected:
 
diff --git a/Engine/Source/Editor/MovieSceneTools/Private/TrackEditors/ShotTrackEditor.cpp b/Engine/Source/Editor/MovieSceneTools/Private/TrackEditors/ShotTrackEditor.cpp
--- a/Engine/Source/Editor/MovieSceneTools/Private/TrackEditors/ShotTrackEditor.cpp
+++ b/Engine/Source/Editor/MovieSceneTools/Private/TrackEditors/ShotTrackEditor.cpp
@@ -529,13 +529,11 @@ bool FShotTrackEditor::SupportsType( TSubclassOf<UMovieSceneTrack> Type ) const
 	return Type == UMovieSceneShotTrack::StaticClass();
 }
 
-TSharedRef<ISequencerSection> FShotTrackEditor::MakeSectionInterface( UMovieSceneSection& SectionObject, UMovieSceneTrack* Track )
+TSharedRef<ISequencerSection> FShotTrackEditor::MakeSectionInterface( UMovieSceneSection& SectionObject, UMovieSceneTrack& Track )
 {
 	check( SupportsType( SectionObject.GetOuter()->GetClass() ) );
 
-	TSharedRef<ISequencerSection> NewSection( new FShotSection( GetSequencer(), ThumbnailPool, SectionObject ) );
-
-	return NewSection;
+	return MakeShareable( new FShotSection( GetSequencer(), ThumbnailPool, SectionObject ) );
 }
 
 void FShotTrackEditor::AddKey(const FGuid& ObjectGuid, UObject* AdditionalAsset)
diff --git a/Engine/Source/Editor/MovieSceneTools/Private/TrackEditors/ShotTrackEditor.h b/Engine/Source/Editor/MovieSceneTools/Private/TrackEditors/ShotTrackEditor.h
--- a/Engine/Source/Editor/MovieSceneTools/Private/TrackEditors/ShotTrackEditor.h
+++ b/Engine/Source/Editor/MovieSceneTools/Private/TrackEditors/ShotTrackEditor.h
@@ -36,7 +36,7 @@ public:
 	virtual void AddKey(const FGuid& ObjectGuid, UObject* AdditionalAsset = NULL) override;
 	virtual void BuildAddTrackMenu(FMenuBuilder& MenuBuilder) override;
 	virtual void BuildObjectBindingTrackMenu(FMenuBuilder& MenuBuilder, const FGuid& ObjectBinding, const UClass* ObjectClass) override;
-	virtual TSharedRef<ISequencerSection> MakeSectionInterface( UMovieSceneSection& SectionObject, UMovieSceneTrack* Track ) override;
+	virtual TSharedRef<ISequencerSection> MakeSectionInterface( UMovieSceneSection& SectionObject, UMovieSceneTrack& Track ) override;
 	virtual bool SupportsType( TSubclassOf<UMovieSceneTrack> Type ) const override;
 	virtual void Tick(float DeltaTime) override;
 
diff --git a/Engine/Source/Editor/MovieSceneTools/Private/TrackEditors/SkeletalAnimationTrackEditor.cpp b/Engine/Source/Editor/MovieSceneTools/Private/TrackEditors/SkeletalAnimationTrackEditor.cpp
--- a/Engine/Source/Editor/MovieSceneTools/Private/TrackEditors/SkeletalAnimationTrackEditor.cpp
+++ b/Engine/Source/Editor/MovieSceneTools/Private/TrackEditors/SkeletalAnimationTrackEditor.cpp
@@ -154,13 +154,11 @@ bool FSkeletalAnimationTrackEditor::SupportsType( TSubclassOf<UMovieSceneTrack>
 }
 
 
-TSharedRef<ISequencerSection> FSkeletalAnimationTrackEditor::MakeSectionInterface( UMovieSceneSection& SectionObject, UMovieSceneTrack* Track )
+TSharedRef<ISequencerSection> FSkeletalAnimationTrackEditor::MakeSectionInterface( UMovieSceneSection& SectionObject, UMovieSceneTrack& Track )
 {
 	check( SupportsType( SectionObject.GetOuter()->GetClass() ) );
 	
-	TSharedRef<ISequencerSection> NewSection( new FSkeletalAnimationSection(SectionObject) );
-
-	return NewSection;
+	return MakeShareable( new FSkeletalAnimationSection(SectionObject) );
 }
 
 
diff --git a/Engine/Source/Editor/MovieSceneTools/Private/TrackEditors/SkeletalAnimationTrackEditor.h b/Engine/Source/Editor/MovieSceneTools/Private/TrackEditors/SkeletalAnimationTrackEditor.h
--- a/Engine/Source/Editor/MovieSceneTools/Private/TrackEditors/SkeletalAnimationTrackEditor.h
+++ b/Engine/Source/Editor/MovieSceneTools/Private/TrackEditors/SkeletalAnimationTrackEditor.h
@@ -34,7 +34,7 @@ public:
 	virtual void AddKey(const FGuid& ObjectGuid, UObject* AdditionalAsset) override;
 	virtual void BuildObjectBindingTrackMenu(FMenuBuilder& MenuBuilder, const FGuid& ObjectBinding, const UClass* ObjectClass) override;
 	virtual bool HandleAssetAdded(UObject* Asset, const FGuid& TargetObjectGuid) override;
-	virtual TSharedRef<ISequencerSection> MakeSectionInterface( UMovieSceneSection& SectionObject, UMovieSceneTrack* Track ) override;
+	virtual TSharedRef<ISequencerSection> MakeSectionInterface( UMovieSceneSection& SectionObject, UMovieSceneTrack& Track ) override;
 	virtual bool SupportsType( TSubclassOf<UMovieSceneTrack> Type ) const override;
 
 private:
diff --git a/Engine/Source/Editor/MovieSceneTools/Private/TrackEditors/SubMovieSceneTrackEditor.cpp b/Engine/Source/Editor/MovieSceneTools/Private/TrackEditors/SubMovieSceneTrackEditor.cpp
--- a/Engine/Source/Editor/MovieSceneTools/Private/TrackEditors/SubMovieSceneTrackEditor.cpp
+++ b/Engine/Source/Editor/MovieSceneTools/Private/TrackEditors/SubMovieSceneTrackEditor.cpp
@@ -100,9 +100,9 @@ bool FSubMovieSceneTrackEditor::SupportsType( TSubclassOf<UMovieSceneTrack> Type
 }
 
 
-TSharedRef<ISequencerSection> FSubMovieSceneTrackEditor::MakeSectionInterface( UMovieSceneSection& SectionObject, UMovieSceneTrack* Track )
+TSharedRef<ISequencerSection> FSubMovieSceneTrackEditor::MakeSectionInterface( UMovieSceneSection& SectionObject, UMovieSceneTrack& Track )
 {
-	return MakeShareable( new FSubMovieSceneSection( GetSequencer(), SectionObject, Track->GetTrackName() ) );
+	return MakeShareable( new FSubMovieSceneSection( GetSequencer(), SectionObject, Track.GetTrackName() ) );
 }
 
 
diff --git a/Engine/Source/Editor/MovieSceneTools/Private/TrackEditors/SubMovieSceneTrackEditor.h b/Engine/Source/Editor/MovieSceneTools/Private/TrackEditors/SubMovieSceneTrackEditor.h
--- a/Engine/Source/Editor/MovieSceneTools/Private/TrackEditors/SubMovieSceneTrackEditor.h
+++ b/Engine/Source/Editor/MovieSceneTools/Private/TrackEditors/SubMovieSceneTrackEditor.h
@@ -33,7 +33,7 @@ public:
 
 	// ISequencerTrackEditor interface
 
-	virtual TSharedRef<ISequencerSection> MakeSectionInterface( UMovieSceneSection& SectionObject, UMovieSceneTrack* Track ) override;
+	virtual TSharedRef<ISequencerSection> MakeSectionInterface( UMovieSceneSection& SectionObject, UMovieSceneTrack& Track ) override;
 	virtual bool HandleAssetAdded(UObject* Asset, const FGuid& TargetObjectGuid) override;
 	virtual bool SupportsType( TSubclassOf<UMovieSceneTrack> Type ) const override;
 };
diff --git a/Engine/Source/Editor/MovieSceneTools/Private/TrackEditors/TransformTrackEditor.cpp b/Engine/Source/Editor/MovieSceneTools/Private/TrackEditors/TransformTrackEditor.cpp
--- a/Engine/Source/Editor/MovieSceneTools/Private/TrackEditors/TransformTrackEditor.cpp
+++ b/Engine/Source/Editor/MovieSceneTools/Private/TrackEditors/TransformTrackEditor.cpp
@@ -186,13 +186,11 @@ bool F3DTransformTrackEditor::SupportsType( TSubclassOf<UMovieSceneTrack> Type )
 }
 
 
-TSharedRef<ISequencerSection> F3DTransformTrackEditor::MakeSectionInterface( UMovieSceneSection& SectionObject, UMovieSceneTrack* Track )
+TSharedRef<ISequencerSection> F3DTransformTrackEditor::MakeSectionInterface( UMovieSceneSection& SectionObject, UMovieSceneTrack& Track )
 {
 	check( SupportsType( SectionObject.GetOuter()->GetClass() ) );
 
-	TSharedRef<ISequencerSection> NewSection( new F3DTransformSection( SectionObject ) );
-
-	return NewSection;
+	return MakeShareable( new F3DTransformSection( SectionObject ) );
 }
 
 
diff --git a/Engine/Source/Editor/MovieSceneTools/Private/TrackEditors/TransformTrackEditor.h b/Engine/Source/Editor/MovieSceneTools/Private/TrackEditors/TransformTrackEditor.h
--- a/Engine/Source/Editor/MovieSceneTools/Private/TrackEditors/TransformTrackEditor.h
+++ b/Engine/Source/Editor/MovieSceneTools/Private/TrackEditors/TransformTrackEditor.h
@@ -39,7 +39,7 @@ public:
 	virtual void BindCommands(TSharedRef<FUICommandList> SequencerCommandBindings) override;
 	virtual void BuildObjectBindingEditButtons(TSharedPtr<SHorizontalBox> EditBox, const FGuid& ObjectBinding, const UClass* ObjectClass) override;
 	virtual void BuildObjectBindingTrackMenu(FMenuBuilder& MenuBuilder, const FGuid& ObjectBinding, const UClass* ObjectClass) override;
-	virtual TSharedRef<ISequencerSection> MakeSectionInterface( UMovieSceneSection& SectionObject, UMovieSceneTrack* Track ) override;
+	virtual TSharedRef<ISequencerSection> MakeSectionInterface( UMovieSceneSection& SectionObject, UMovieSceneTrack& Track ) override;
 	virtual void OnRelease() override;
 	virtual bool SupportsType( TSubclassOf<UMovieSceneTrack> Type ) const override;
 
diff --git a/Engine/Source/Editor/NiagaraEditor/Private/NiagaraSequencer.h b/Engine/Source/Editor/NiagaraEditor/Private/NiagaraSequencer.h
--- a/Engine/Source/Editor/NiagaraEditor/Private/NiagaraSequencer.h
+++ b/Engine/Source/Editor/NiagaraEditor/Private/NiagaraSequencer.h
@@ -209,7 +209,7 @@ public:
 		return false;
 	}
 
-	virtual TSharedRef<ISequencerSection> MakeSectionInterface(class UMovieSceneSection& SectionObject, UMovieSceneTrack* Track)
+	virtual TSharedRef<ISequencerSection> MakeSectionInterface(class UMovieSceneSection& SectionObject, UMovieSceneTrack& Track)
 	{
 		INiagaraSequencerSection *Sec = new INiagaraSequencerSection(SectionObject);
 		return MakeShareable(Sec);
diff --git a/Engine/Source/Editor/Sequencer/Private/Sequencer.cpp b/Engine/Source/Editor/Sequencer/Private/Sequencer.cpp
--- a/Engine/Source/Editor/Sequencer/Private/Sequencer.cpp
+++ b/Engine/Source/Editor/Sequencer/Private/Sequencer.cpp
@@ -1555,20 +1555,23 @@ void FSequencer::OnRequestNodeDeleted( TSharedRef<const FSequencerDisplayNode> N
 		TSharedRef<const FTrackNode> SectionAreaNode = StaticCastSharedRef<const FTrackNode>( NodeToBeDeleted );
 		UMovieSceneTrack* Track = SectionAreaNode->GetTrack();
 
-		if( OwnerMovieScene->IsAMasterTrack( Track ) )
+		if (Track != nullptr)
 		{
-			OwnerMovieScene->RemoveMasterTrack( Track );
-		}
-		else if( OwnerMovieScene->GetShotTrack() == Track )
-		{
-			OwnerMovieScene->RemoveShotTrack();
-		}
-		else
-		{
-			OwnerMovieScene->RemoveTrack( Track );
-		}
+			if( OwnerMovieScene->IsAMasterTrack( *Track ) )
+			{
+				OwnerMovieScene->RemoveMasterTrack( Track );
+			}
+			else if( OwnerMovieScene->GetShotTrack() == Track )
+			{
+				OwnerMovieScene->RemoveShotTrack();
+			}
+			else
+			{
+				OwnerMovieScene->RemoveTrack( Track );
+			}
 		
-		bAnythingRemoved = true;
+			bAnythingRemoved = true;
+		}
 	}
 	else if ( NodeToBeDeleted->GetType() == ESequencerNode::Category )
 	{
diff --git a/Engine/Source/Editor/Sequencer/Private/SequencerNodeTree.cpp b/Engine/Source/Editor/Sequencer/Private/SequencerNodeTree.cpp
--- a/Engine/Source/Editor/Sequencer/Private/SequencerNodeTree.cpp
+++ b/Engine/Source/Editor/Sequencer/Private/SequencerNodeTree.cpp
@@ -170,7 +170,7 @@ void FSequencerNodeTree::MakeSectionInterfaces( UMovieSceneTrack& Track, TShared
 	for (int32 SectionIndex = 0; SectionIndex < MovieSceneSections.Num(); ++SectionIndex )
 	{
 		UMovieSceneSection* SectionObject = MovieSceneSections[SectionIndex];
-		TSharedRef<ISequencerSection> Section = Editor->MakeSectionInterface( *SectionObject, &Track );
+		TSharedRef<ISequencerSection> Section = Editor->MakeSectionInterface( *SectionObject, Track );
 
 		// Ask the section to generate it's inner layout
 		FSectionLayoutBuilder Builder( SectionAreaNode );
diff --git a/Engine/Source/Editor/Sequencer/Public/ISequencerTrackEditor.h b/Engine/Source/Editor/Sequencer/Public/ISequencerTrackEditor.h
--- a/Engine/Source/Editor/Sequencer/Public/ISequencerTrackEditor.h
+++ b/Engine/Source/Editor/Sequencer/Public/ISequencerTrackEditor.h
@@ -88,9 +88,9 @@ public:
 	 * Called to generate a section layout for a particular section.
 	 *
 	 * @param SectionObject The section to make UI for.
-	 * @param SectionName The name of the section.
+	 * @param Track The track that owns the section.
 	 */
-	virtual TSharedRef<ISequencerSection> MakeSectionInterface(UMovieSceneSection& SectionObject, UMovieSceneTrack* Track) = 0;
+	virtual TSharedRef<ISequencerSection> MakeSectionInterface(UMovieSceneSection& SectionObject, UMovieSceneTrack& Track) = 0;
 
 	/** Called when the instance of this track editor is released */
 	virtual void OnRelease() = 0;
diff --git a/Engine/Source/Editor/Sequencer/Public/MovieSceneTrackEditor.h b/Engine/Source/Editor/Sequencer/Public/MovieSceneTrackEditor.h
--- a/Engine/Source/Editor/Sequencer/Public/MovieSceneTrackEditor.h
+++ b/Engine/Source/Editor/Sequencer/Public/MovieSceneTrackEditor.h
@@ -181,7 +181,7 @@ public:
 	}
 
 	virtual void OnRelease() override { };
-	virtual TSharedRef<ISequencerSection> MakeSectionInterface( class UMovieSceneSection& SectionObject, UMovieSceneTrack* Track ) = 0;
+	virtual TSharedRef<ISequencerSection> MakeSectionInterface( class UMovieSceneSection& SectionObject, UMovieSceneTrack& Track ) = 0;
 	virtual bool SupportsType( TSubclassOf<class UMovieSceneTrack> TrackClass ) const = 0;
 	virtual void Tick(float DeltaTime) override { }
 
diff --git a/Engine/Source/Editor/UMGEditor/Private/Animation/MarginTrackEditor.cpp b/Engine/Source/Editor/UMGEditor/Private/Animation/MarginTrackEditor.cpp
--- a/Engine/Source/Editor/UMGEditor/Private/Animation/MarginTrackEditor.cpp
+++ b/Engine/Source/Editor/UMGEditor/Private/Animation/MarginTrackEditor.cpp
@@ -34,15 +34,12 @@ TSharedRef<ISequencerTrackEditor> FMarginTrackEditor::CreateTrackEditor( TShared
 }
 
 
-TSharedRef<ISequencerSection> FMarginTrackEditor::MakeSectionInterface( UMovieSceneSection& SectionObject, UMovieSceneTrack* Track )
+TSharedRef<ISequencerSection> FMarginTrackEditor::MakeSectionInterface( UMovieSceneSection& SectionObject, UMovieSceneTrack& Track )
 {
 	check( SupportsType( SectionObject.GetOuter()->GetClass() ) );
 
 	UClass* SectionClass = SectionObject.GetOuter()->GetClass();
-
-	TSharedRef<ISequencerSection> NewSection = MakeShareable( new FMarginPropertySection( SectionObject, Track->GetTrackName() ) );
-
-	return NewSection;
+	return MakeShareable( new FMarginPropertySection( SectionObject, Track.GetTrackName() ) );
 }
 
 
diff --git a/Engine/Source/Editor/UMGEditor/Private/Animation/MarginTrackEditor.h b/Engine/Source/Editor/UMGEditor/Private/Animation/MarginTrackEditor.h
--- a/Engine/Source/Editor/UMGEditor/Private/Animation/MarginTrackEditor.h
+++ b/Engine/Source/Editor/UMGEditor/Private/Animation/MarginTrackEditor.h
@@ -37,7 +37,7 @@ public:
 
 	// ISequencerTrackEditor interface
 
-	virtual TSharedRef<ISequencerSection> MakeSectionInterface( UMovieSceneSection& SectionObject, UMovieSceneTrack* Track ) override;
+	virtual TSharedRef<ISequencerSection> MakeSectionInterface( UMovieSceneSection& SectionObject, UMovieSceneTrack& Track ) override;
 
 protected:
 
diff --git a/Engine/Source/Editor/UMGEditor/Private/Animation/Sequencer2DTransformTrackEditor.cpp b/Engine/Source/Editor/UMGEditor/Private/Animation/Sequencer2DTransformTrackEditor.cpp
--- a/Engine/Source/Editor/UMGEditor/Private/Animation/Sequencer2DTransformTrackEditor.cpp
+++ b/Engine/Source/Editor/UMGEditor/Private/Animation/Sequencer2DTransformTrackEditor.cpp
@@ -50,15 +50,12 @@ TSharedRef<ISequencerTrackEditor> F2DTransformTrackEditor::CreateTrackEditor( TS
 }
 
 
-TSharedRef<ISequencerSection> F2DTransformTrackEditor::MakeSectionInterface( UMovieSceneSection& SectionObject, UMovieSceneTrack* Track )
+TSharedRef<ISequencerSection> F2DTransformTrackEditor::MakeSectionInterface( UMovieSceneSection& SectionObject, UMovieSceneTrack& Track )
 {
 	check( SupportsType( SectionObject.GetOuter()->GetClass() ) );
 
 	UClass* SectionClass = SectionObject.GetOuter()->GetClass();
-
-	TSharedRef<ISequencerSection> NewSection = MakeShareable( new F2DTransformSection( SectionObject, Track->GetTrackName() ) );
-
-	return NewSection;
+	return MakeShareable( new F2DTransformSection( SectionObject, Track.GetTrackName() ) );
 }
 
 
diff --git a/Engine/Source/Editor/UMGEditor/Private/Animation/Sequencer2DTransformTrackEditor.h b/Engine/Source/Editor/UMGEditor/Private/Animation/Sequencer2DTransformTrackEditor.h
--- a/Engine/Source/Editor/UMGEditor/Private/Animation/Sequencer2DTransformTrackEditor.h
+++ b/Engine/Source/Editor/UMGEditor/Private/Animation/Sequencer2DTransformTrackEditor.h
@@ -35,7 +35,7 @@ public:
 
 	// ISequencerTrackEditor interface
 
-	virtual TSharedRef<ISequencerSection> MakeSectionInterface( UMovieSceneSection& SectionObject, UMovieSceneTrack* Track ) override;
+	virtual TSharedRef<ISequencerSection> MakeSectionInterface( UMovieSceneSection& SectionObject, UMovieSceneTrack& Track ) override;
 
 protected:
 
diff --git a/Engine/Source/Runtime/MovieScene/Private/MovieScene.cpp b/Engine/Source/Runtime/MovieScene/Private/MovieScene.cpp
--- a/Engine/Source/Runtime/MovieScene/Private/MovieScene.cpp
+++ b/Engine/Source/Runtime/MovieScene/Private/MovieScene.cpp
@@ -406,11 +406,11 @@ bool UMovieScene::RemoveMasterTrack( UMovieSceneTrack* Track )
 }
 
 
-bool UMovieScene::IsAMasterTrack(const UMovieSceneTrack* Track) const
+bool UMovieScene::IsAMasterTrack(const UMovieSceneTrack& Track) const
 {
-	for (int32 TrackIndex = 0; TrackIndex < MasterTracks.Num(); ++TrackIndex)
+	for ( const UMovieSceneTrack* MasterTrack : MasterTracks)
 	{
-		if (Track == MasterTracks[TrackIndex])
+		if (&Track == MasterTrack)
 		{
 			return true;
 		}
diff --git a/Engine/Source/Runtime/MovieScene/Public/MovieScene.h b/Engine/Source/Runtime/MovieScene/Public/MovieScene.h
--- a/Engine/Source/Runtime/MovieScene/Public/MovieScene.h
+++ b/Engine/Source/Runtime/MovieScene/Public/MovieScene.h
@@ -233,7 +233,7 @@ public:
 	 *
 	 * @return true if the track is a master track, false otherwise.
 	 */
-	bool IsAMasterTrack(const UMovieSceneTrack* Track) const;
+	bool IsAMasterTrack(const UMovieSceneTrack& Track) const;
 
 	/**
 	 * Get all master tracks.
diff --git a/Engine/Source/Runtime/MovieSceneTracks/Private/Tracks/MovieSceneAudioTrack.cpp b/Engine/Source/Runtime/MovieSceneTracks/Private/Tracks/MovieSceneAudioTrack.cpp
--- a/Engine/Source/Runtime/MovieSceneTracks/Private/Tracks/MovieSceneAudioTrack.cpp
+++ b/Engine/Source/Runtime/MovieSceneTracks/Private/Tracks/MovieSceneAudioTrack.cpp
@@ -134,7 +134,7 @@ void UMovieSceneAudioTrack::AddNewSound(USoundBase* Sound, float Time)
 
 bool UMovieSceneAudioTrack::IsAMasterTrack() const
 {
-	return Cast<UMovieScene>(GetOuter())->IsAMasterTrack(this);
+	return Cast<UMovieScene>(GetOuter())->IsAMasterTrack(*this);
 }
 
 
