Sequencer: Fix slomo, fade and event tracks so that they can be removed and added.

diff --git a/Engine/Source/Editor/MovieSceneTools/Private/TrackEditors/EventTrackEditor.cpp b/Engine/Source/Editor/MovieSceneTools/Private/TrackEditors/EventTrackEditor.cpp
--- a/Engine/Source/Editor/MovieSceneTools/Private/TrackEditors/EventTrackEditor.cpp
+++ b/Engine/Source/Editor/MovieSceneTools/Private/TrackEditors/EventTrackEditor.cpp
@@ -84,14 +84,27 @@ void FEventTrackEditor::HandleAddEventTrackMenuEntryExecute()
 {
 	UMovieSceneSequence* FocusedSequence = GetSequencer()->GetFocusedMovieSceneSequence();
 	UMovieScene* MovieScene = FocusedSequence->GetMovieScene();
-
-	if (MovieScene != nullptr)
+	if (MovieScene == nullptr)
 	{
-		UMovieSceneTrack* NewTrack = MovieScene->AddMasterTrack(UMovieSceneEventTrack::StaticClass());
-		NewTrack->AddSection(NewTrack->CreateNewSection());
+		return;
+	}
 
-		GetSequencer()->NotifyMovieSceneDataChanged();
+	UMovieSceneTrack* EventTrack = MovieScene->FindMasterTrack( UMovieSceneEventTrack::StaticClass() );
+	if (EventTrack != nullptr)
+	{
+		return;
 	}
+
+	const FScopedTransaction Transaction(NSLOCTEXT("Sequencer", "AddEventTrack_Transaction", "Add Event Track"));
+
+	MovieScene->Modify();
+
+	EventTrack = GetMasterTrack( UMovieSceneEventTrack::StaticClass() );
+	ensure(EventTrack);
+
+	EventTrack->AddSection(EventTrack->CreateNewSection());
+
+	GetSequencer()->NotifyMovieSceneDataChanged();
 }
 
 
diff --git a/Engine/Source/Editor/MovieSceneTools/Private/TrackEditors/FadeTrackEditor.cpp b/Engine/Source/Editor/MovieSceneTools/Private/TrackEditors/FadeTrackEditor.cpp
--- a/Engine/Source/Editor/MovieSceneTools/Private/TrackEditors/FadeTrackEditor.cpp
+++ b/Engine/Source/Editor/MovieSceneTools/Private/TrackEditors/FadeTrackEditor.cpp
@@ -60,21 +60,29 @@ bool FFadeTrackEditor::SupportsType(TSubclassOf<UMovieSceneTrack> Type) const
 
 void FFadeTrackEditor::HandleAddFadeTrackMenuEntryExecute()
 {
-	if (FadeTrack.IsValid())
+	UMovieSceneSequence* FocusedSequence = GetSequencer()->GetFocusedMovieSceneSequence();
+	UMovieScene* MovieScene = FocusedSequence->GetMovieScene();
+	if (MovieScene == nullptr)
 	{
 		return;
 	}
 
-	UMovieSceneSequence* FocusedSequence = GetSequencer()->GetFocusedMovieSceneSequence();
-	UMovieScene* MovieScene = FocusedSequence->GetMovieScene();
-
-	if (MovieScene != nullptr)
+	UMovieSceneTrack* FadeTrack = MovieScene->FindMasterTrack( UMovieSceneFadeTrack::StaticClass() );
+	if (FadeTrack != nullptr)
 	{
-		FadeTrack = MovieScene->AddMasterTrack(UMovieSceneFadeTrack::StaticClass());
-		FadeTrack->AddSection(FadeTrack->CreateNewSection());
-
-		GetSequencer()->NotifyMovieSceneDataChanged();
+		return;
 	}
+
+	const FScopedTransaction Transaction(NSLOCTEXT("Sequencer", "AddFadeTrack_Transaction", "Add Fade Track"));
+
+	MovieScene->Modify();
+		
+	FadeTrack = GetMasterTrack( UMovieSceneFadeTrack::StaticClass() );
+	ensure(FadeTrack);
+
+	FadeTrack->AddSection(FadeTrack->CreateNewSection());
+
+	GetSequencer()->NotifyMovieSceneDataChanged();
 }
 
 
diff --git a/Engine/Source/Editor/MovieSceneTools/Private/TrackEditors/FadeTrackEditor.h b/Engine/Source/Editor/MovieSceneTools/Private/TrackEditors/FadeTrackEditor.h
--- a/Engine/Source/Editor/MovieSceneTools/Private/TrackEditors/FadeTrackEditor.h
+++ b/Engine/Source/Editor/MovieSceneTools/Private/TrackEditors/FadeTrackEditor.h
@@ -45,9 +45,4 @@ private:
 
 	/** Callback for executing the "Add Fade Track" menu entry. */
 	void HandleAddFadeTrackMenuEntryExecute();
-
-private:
-
-	/** Pointer to the sequence's fade track, if it exists. */
-	TWeakObjectPtr<UMovieSceneTrack> FadeTrack;
 };
diff --git a/Engine/Source/Editor/MovieSceneTools/Private/TrackEditors/SlomoTrackEditor.cpp b/Engine/Source/Editor/MovieSceneTools/Private/TrackEditors/SlomoTrackEditor.cpp
--- a/Engine/Source/Editor/MovieSceneTools/Private/TrackEditors/SlomoTrackEditor.cpp
+++ b/Engine/Source/Editor/MovieSceneTools/Private/TrackEditors/SlomoTrackEditor.cpp
@@ -60,21 +60,29 @@ bool FSlomoTrackEditor::SupportsType(TSubclassOf<UMovieSceneTrack> Type) const
 
 void FSlomoTrackEditor::HandleAddSlomoTrackMenuEntryExecute()
 {
-	if (SlomoTrack.IsValid())
+	UMovieSceneSequence* FocusedSequence = GetSequencer()->GetFocusedMovieSceneSequence();
+	UMovieScene* MovieScene = FocusedSequence->GetMovieScene();
+	if (MovieScene == nullptr)
 	{
 		return;
 	}
 
-	UMovieSceneSequence* FocusedSequence = GetSequencer()->GetFocusedMovieSceneSequence();
-	UMovieScene* MovieScene = FocusedSequence->GetMovieScene();
-
-	if (MovieScene != nullptr)
+	UMovieSceneTrack* SlomoTrack = MovieScene->FindMasterTrack( UMovieSceneSlomoTrack::StaticClass() );
+	if (SlomoTrack != nullptr)
 	{
-		SlomoTrack = MovieScene->AddMasterTrack(UMovieSceneSlomoTrack::StaticClass());
-		SlomoTrack->AddSection(SlomoTrack->CreateNewSection());
-
-		GetSequencer()->NotifyMovieSceneDataChanged();
+		return;
 	}
+
+	const FScopedTransaction Transaction(NSLOCTEXT("Sequencer", "AddSlomoTrack_Transaction", "Add Play Rate Track"));
+
+	MovieScene->Modify();
+		
+	SlomoTrack = GetMasterTrack( UMovieSceneSlomoTrack::StaticClass() );
+	ensure(SlomoTrack);
+
+	SlomoTrack->AddSection(SlomoTrack->CreateNewSection());
+
+	GetSequencer()->NotifyMovieSceneDataChanged();
 }
 
 
diff --git a/Engine/Source/Editor/MovieSceneTools/Private/TrackEditors/SlomoTrackEditor.h b/Engine/Source/Editor/MovieSceneTools/Private/TrackEditors/SlomoTrackEditor.h
--- a/Engine/Source/Editor/MovieSceneTools/Private/TrackEditors/SlomoTrackEditor.h
+++ b/Engine/Source/Editor/MovieSceneTools/Private/TrackEditors/SlomoTrackEditor.h
@@ -45,9 +45,4 @@ private:
 
 	/** Callback for executing the "Add Slomo Track" menu entry. */
 	void HandleAddSlomoTrackMenuEntryExecute();
-
-private:
-
-	/** Pointer to the sequence's slomo track, if it exists. */
-	TWeakObjectPtr<UMovieSceneTrack> SlomoTrack;
 };
