Clang compile error fixes, initialization order warnings, shadow variable warnings

diff --git a/Engine/Source/Developer/MaterialUtilities/Private/MaterialUtilities.cpp b/Engine/Source/Developer/MaterialUtilities/Private/MaterialUtilities.cpp
--- a/Engine/Source/Developer/MaterialUtilities/Private/MaterialUtilities.cpp
+++ b/Engine/Source/Developer/MaterialUtilities/Private/MaterialUtilities.cpp
@@ -1681,14 +1681,9 @@ void FMaterialUtilities::FullyLoadMaterialStatic(UMaterialInterface* Material)
 	// Now load all used textures.
 	TArray<UTexture*> Textures;
 	Material->GetUsedTextures(Textures, EMaterialQualityLevel::Num, true, ERHIFeatureLevel::SM5, true);
-	for (int32 i = 0; i < Textures.Num(); ++i)
+	for (UTexture* Texture : Textures)
 	{
-		UTexture* Texture = Textures[i];
-		FLinkerLoad* Linker = Material->GetLinker();
-		if (Linker)
-		{
-			Linker->Preload(Texture);
-		}
+		Linker->Preload(Texture);
 	}
 }
 
diff --git a/Engine/Source/Developer/MeshUtilities/Private/MeshUtilities.cpp b/Engine/Source/Developer/MeshUtilities/Private/MeshUtilities.cpp
--- a/Engine/Source/Developer/MeshUtilities/Private/MeshUtilities.cpp
+++ b/Engine/Source/Developer/MeshUtilities/Private/MeshUtilities.cpp
@@ -21,12 +21,17 @@
 #include "HierarchicalLODUtilities.h"
 #include "MeshBoneReduction.h"
 #include "MeshMergeData.h"
+#include "Editor/EditorPerProjectUserSettings.h"
 
 //@todo - implement required vector intrinsics for other implementations
 #if PLATFORM_ENABLE_VECTORINTRINSICS
 #include "kDOP.h"
 #endif
 
+#if WITH_EDITOR
+#include "Editor.h"
+#endif
+
 /*------------------------------------------------------------------------------
 	MeshUtilities module.
 ------------------------------------------------------------------------------*/
@@ -6665,7 +6670,7 @@ void FMeshUtilities::MergeStaticMeshComponents(const TArray<UStaticMeshComponent
 	MainTask.EnterProgressFrame(20);
 
 	FRawMeshExt MergedMesh;
-	memset(&MergedMesh, NULL, sizeof(MergedMesh));
+	FMemory::Memset(&MergedMesh, INDEX_NONE, sizeof(MergedMesh));
 	
 	if (InSettings.bMergeMaterials)
 	{
diff --git a/Engine/Source/Developer/MeshUtilities/Public/MeshUtilities.h b/Engine/Source/Developer/MeshUtilities/Public/MeshUtilities.h
--- a/Engine/Source/Developer/MeshUtilities/Public/MeshUtilities.h
+++ b/Engine/Source/Developer/MeshUtilities/Public/MeshUtilities.h
@@ -3,10 +3,11 @@
 #pragma once
 
 #include "ModuleInterface.h"
-
+#include "SkeletalMeshTypes.h"
 
 typedef FIntPoint FMeshIdAndLOD;
 struct FFlattenMaterial;
+struct FReferenceSkeleton;
 
 /**
  * Mesh reduction interface.
diff --git a/Engine/Source/Developer/ShaderCompilerCommon/Private/HlslExpressionParser.inl b/Engine/Source/Developer/ShaderCompilerCommon/Private/HlslExpressionParser.inl
--- a/Engine/Source/Developer/ShaderCompilerCommon/Private/HlslExpressionParser.inl
+++ b/Engine/Source/Developer/ShaderCompilerCommon/Private/HlslExpressionParser.inl
@@ -27,8 +27,8 @@ namespace CrossCompiler
 
 		FSymbolScope(FLinearAllocator* InAllocator, FSymbolScope* InParent) :
 			Parent(InParent),
-			Children(InAllocator),
-			Name(nullptr)
+			Name(nullptr),
+			Children(InAllocator)
 		{
 		}
 		~FSymbolScope() {}
@@ -422,8 +422,8 @@ namespace CrossCompiler
 					}
 					TypeString += OuterNamespace->String;
 					TypeString += TEXT("::");
-					auto* Token1 = Scanner.PeekToken(1);
-					if (Token1 && Token1->Token == EHlslToken::ColonColon)
+					auto* Token3 = Scanner.PeekToken(1);
+					if (Token3 && Token3->Token == EHlslToken::ColonColon)
 					{
 						continue;
 					}
diff --git a/Engine/Source/Developer/ShaderCompilerCommon/Private/HlslUtils.cpp b/Engine/Source/Developer/ShaderCompilerCommon/Private/HlslUtils.cpp
--- a/Engine/Source/Developer/ShaderCompilerCommon/Private/HlslUtils.cpp
+++ b/Engine/Source/Developer/ShaderCompilerCommon/Private/HlslUtils.cpp
@@ -93,8 +93,8 @@ struct FRemoveUnusedOutputs
 	FRemoveUnusedOutputs(const TArray<FString>& InSystemOutputs, const TArray<FString>& InUsedOutputs, CrossCompiler::FLinearAllocator* InAllocator = nullptr) :
 		SystemOutputs(InSystemOutputs),
 		UsedOutputs(InUsedOutputs),
-		Allocator(InAllocator),
-		bSuccess(false)
+		bSuccess(false),
+		Allocator(InAllocator)
 	{
 	}
 
diff --git a/Engine/Source/Editor/MergeActors/Private/MeshProxyTool/MeshProxyTool.cpp b/Engine/Source/Editor/MergeActors/Private/MeshProxyTool/MeshProxyTool.cpp
--- a/Engine/Source/Editor/MergeActors/Private/MeshProxyTool/MeshProxyTool.cpp
+++ b/Engine/Source/Editor/MergeActors/Private/MeshProxyTool/MeshProxyTool.cpp
@@ -95,22 +95,22 @@ bool FMeshProxyTool::RunMerge(const FString& PackageName)
 		
 		FCreateProxyDelegate ProxyDelegate;
 		ProxyDelegate.BindLambda(
-			[](const FGuid Guid, TArray<UObject*>& AssetsToSync)
+			[](const FGuid Guid, TArray<UObject*>& InAssetsToSync)
 		{
 			//Update the asset registry that a new static mash and material has been created
-			if (AssetsToSync.Num())
+			if (InAssetsToSync.Num())
 			{
 				FAssetRegistryModule& AssetRegistry = FModuleManager::Get().LoadModuleChecked<FAssetRegistryModule>("AssetRegistry");
-				int32 AssetCount = AssetsToSync.Num();
+				int32 AssetCount = InAssetsToSync.Num();
 				for (int32 AssetIndex = 0; AssetIndex < AssetCount; AssetIndex++)
 				{
-					AssetRegistry.AssetCreated(AssetsToSync[AssetIndex]);
-					GEditor->BroadcastObjectReimported(AssetsToSync[AssetIndex]);
+					AssetRegistry.AssetCreated(InAssetsToSync[AssetIndex]);
+					GEditor->BroadcastObjectReimported(InAssetsToSync[AssetIndex]);
 				}
 
 				//Also notify the content browser that the new assets exists
 				FContentBrowserModule& ContentBrowserModule = FModuleManager::Get().LoadModuleChecked<FContentBrowserModule>("ContentBrowser");
-				ContentBrowserModule.Get().SyncBrowserToAssets(AssetsToSync, true);
+				ContentBrowserModule.Get().SyncBrowserToAssets(InAssetsToSync, true);
 			}
 		});		
 
diff --git a/Engine/Source/Editor/MovieSceneTools/Public/PropertyTrackEditor.h b/Engine/Source/Editor/MovieSceneTools/Public/PropertyTrackEditor.h
--- a/Engine/Source/Editor/MovieSceneTools/Public/PropertyTrackEditor.h
+++ b/Engine/Source/Editor/MovieSceneTools/Public/PropertyTrackEditor.h
@@ -88,7 +88,7 @@ public:
 
 	~FPropertyTrackEditor()
 	{
-		TSharedPtr<ISequencer> Sequencer = GetSequencer();
+		TSharedPtr<ISequencer> Sequencer = this->GetSequencer();
 		if ( Sequencer.IsValid() )
 		{
 			ISequencerObjectChangeListener& ObjectChangeListener = Sequencer->GetObjectChangeListener();
@@ -110,10 +110,11 @@ public:
 
 	virtual void AddKey( const FGuid& ObjectGuid, UObject* AdditionalAsset = NULL ) override
 	{
-		ISequencerObjectChangeListener& ObjectChangeListener = GetSequencer()->GetObjectChangeListener();
+		TSharedPtr<ISequencer> Sequencer = this->GetSequencer();
+		ISequencerObjectChangeListener& ObjectChangeListener = Sequencer->GetObjectChangeListener();
 
 		TArray<UObject*> OutObjects;
-		GetSequencer()->GetRuntimeObjects( GetSequencer()->GetFocusedMovieSceneSequenceInstance(), ObjectGuid, OutObjects );
+		Sequencer->GetRuntimeObjects( Sequencer->GetFocusedMovieSceneSequenceInstance(), ObjectGuid, OutObjects );
 		for ( int32 i = 0; i < OutObjects.Num(); ++i )
 		{
 			ObjectChangeListener.TriggerAllPropertiesChanged( OutObjects[i] );
@@ -124,7 +125,7 @@ public:
 	{
 		TSharedRef<FPropertySection> PropertySection = MakePropertySectionInterface( SectionObject, Track );
 		OnSetIntermediateValueFromPropertyChange.AddSP( PropertySection, &FPropertySection::SetIntermediateValueForTrack );
-		GetSequencer()->OnGlobalTimeChanged().AddSP( PropertySection, &FPropertySection::ClearIntermediateValue );
+		this->GetSequencer()->OnGlobalTimeChanged().AddSP( PropertySection, &FPropertySection::ClearIntermediateValue );
 		return PropertySection;
 	}
 
@@ -159,7 +160,7 @@ private:
 	/** Adds a callback for property changes for the supplied property type name. */
 	void AddWatchedPropertyType( FName WatchedPropertyTypeName )
 	{
-		GetSequencer()->GetObjectChangeListener().GetOnAnimatablePropertyChanged( WatchedPropertyTypeName ).AddRaw( this, &FPropertyTrackEditor<TrackType, SectionType, KeyDataType>::OnAnimatedPropertyChanged );
+		this->GetSequencer()->GetObjectChangeListener().GetOnAnimatablePropertyChanged( WatchedPropertyTypeName ).AddRaw( this, &FPropertyTrackEditor<TrackType, SectionType, KeyDataType>::OnAnimatedPropertyChanged );
 		WatchedPropertyTypeNames.Add( WatchedPropertyTypeName );
 	}
 
@@ -188,7 +189,7 @@ private:
 	*/
 	virtual void OnAnimatedPropertyChanged( const FPropertyChangedParams& PropertyChangedParams )
 	{
-		AnimatablePropertyChanged( FOnKeyProperty::CreateRaw( this, &FPropertyTrackEditor::OnKeyProperty, PropertyChangedParams ) );
+		this->AnimatablePropertyChanged( FOnKeyProperty::CreateRaw( this, &FPropertyTrackEditor::OnKeyProperty, PropertyChangedParams ) );
 	}
 
 	/** Adds a key based on a property change. */
@@ -198,19 +199,19 @@ private:
 		GenerateKeysFromPropertyChanged( PropertyChangedParams, KeysForPropertyChange );
 
 		TSubclassOf<UMovieSceneTrack> CustomizedClass = GetCustomizedTrackClass( PropertyChangedParams.PropertyPath.Last() );
-		TSubclassOf<UMovieSceneTrack> TrackClass = CustomizedClass != nullptr ? CustomizedClass : TrackType::StaticClass();
+		TSubclassOf<UMovieSceneTrack> TrackClass = (*CustomizedClass != nullptr ? *CustomizedClass : TrackType::StaticClass());
 
 		// If the track class has been customized for this property then it's possible this track editor doesn't support it, 
 		// also check for track editors which should only be used for customization.
-		if ( SupportsType( TrackClass ) && ( ForCustomizedUseOnly() == false || CustomizedClass != nullptr) )
+		if ( SupportsType( TrackClass ) && ( ForCustomizedUseOnly() == false || *CustomizedClass != nullptr) )
 		{
-			FOnInitializeNewTrack OnInitializeNewTrack;
+			typename FKeyframeTrackEditor<TrackType, SectionType, KeyDataType>::FOnInitializeNewTrack OnInitializeNewTrack;
 			OnInitializeNewTrack.BindLambda( [&](TrackType* NewTrack) { InitializeNewTrack( NewTrack, PropertyChangedParams ); } );
 
-			FOnSetIntermediateValue OnSetIntermediateValue;
+			typename FKeyframeTrackEditor<TrackType, SectionType, KeyDataType>::FOnSetIntermediateValue OnSetIntermediateValue;
 			OnSetIntermediateValue.BindSP( this, &FPropertyTrackEditor::SetIntermediateValueFromPropertyChange, PropertyChangedParams );
 
-			return AddKeysToObjects( PropertyChangedParams.ObjectsThatChanged, KeyTime, KeysForPropertyChange, PropertyChangedParams.KeyParams,
+			return this->AddKeysToObjects( PropertyChangedParams.ObjectsThatChanged, KeyTime, KeysForPropertyChange, PropertyChangedParams.KeyParams,
 				TrackClass, PropertyChangedParams.PropertyPath.Last()->GetFName(), OnInitializeNewTrack, OnSetIntermediateValue );
 		}
 		else
diff --git a/Engine/Source/Editor/Sequencer/Private/TimeSliderController.cpp b/Engine/Source/Editor/Sequencer/Private/TimeSliderController.cpp
--- a/Engine/Source/Editor/Sequencer/Private/TimeSliderController.cpp
+++ b/Engine/Source/Editor/Sequencer/Private/TimeSliderController.cpp
@@ -584,7 +584,6 @@ FReply FSequencerTimeSliderController::OnMouseMove( TSharedRef<SWidget> WidgetOw
 		else if (MouseEvent.IsMouseButtonDown( EKeys::LeftMouseButton ))
 		{
 			TRange<float> LocalViewRange = TimeSliderArgs.ViewRange.Get();
-			FScrubRangeToScreen RangeToScreen(LocalViewRange, MyGeometry.Size);
 
 			DistanceDragged += FMath::Abs( MouseEvent.GetCursorDelta().X );
 
@@ -592,6 +591,7 @@ FReply FSequencerTimeSliderController::OnMouseMove( TSharedRef<SWidget> WidgetOw
 			{
 				if ( DistanceDragged > FSlateApplication::Get().GetDragTriggerDistance() )
 				{
+					FScrubRangeToScreen RangeToScreen(LocalViewRange, MyGeometry.Size);
 					const float ScrubPosition = TimeSliderArgs.ScrubPosition.Get();
 
 					TRange<float> PlaybackRange = TimeSliderArgs.PlaybackRange.Get();
@@ -742,20 +742,20 @@ int32 FSequencerTimeSliderController::OnPaintSectionView( const FGeometry& Allot
 		static FLinearColor TickColor(0.f, 0.f, 0.f, 0.3f);
 		
 		// Draw major tick lines in the section area
-		FDrawTickArgs Args;
-		Args.AllottedGeometry = AllottedGeometry;
-		Args.bMirrorLabels = false;
-		Args.bOnlyDrawMajorTicks = true;
-		Args.TickColor = TickColor;
-		Args.ClippingRect = MyClippingRect;
-		Args.DrawEffects = DrawEffects;
+		FDrawTickArgs DrawTickArgs;
+		DrawTickArgs.AllottedGeometry = AllottedGeometry;
+		DrawTickArgs.bMirrorLabels = false;
+		DrawTickArgs.bOnlyDrawMajorTicks = true;
+		DrawTickArgs.TickColor = TickColor;
+		DrawTickArgs.ClippingRect = MyClippingRect;
+		DrawTickArgs.DrawEffects = DrawEffects;
 		// Draw major ticks under sections
-		Args.StartLayer = LayerId-1;
+		DrawTickArgs.StartLayer = LayerId-1;
 		// Draw the tick the entire height of the section area
-		Args.TickOffset = 0.0f;
-		Args.MajorTickHeight = AllottedGeometry.Size.Y;
+		DrawTickArgs.TickOffset = 0.0f;
+		DrawTickArgs.MajorTickHeight = AllottedGeometry.Size.Y;
 
-		DrawTicks( OutDrawElements, RangeToScreen, Args );
+		DrawTicks( OutDrawElements, RangeToScreen, DrawTickArgs );
 	}
 
 	if( Args.bDisplayScrubPosition )
diff --git a/Engine/Source/Runtime/Engine/Public/Slate/SceneViewport.h b/Engine/Source/Runtime/Engine/Public/Slate/SceneViewport.h
--- a/Engine/Source/Runtime/Engine/Public/Slate/SceneViewport.h
+++ b/Engine/Source/Runtime/Engine/Public/Slate/SceneViewport.h
@@ -9,6 +9,8 @@
 /** Called in FSceneViewport::ResizeFrame after ResizeViewport*/
 DECLARE_DELEGATE_OneParam( FOnSceneViewportResize, FVector2D );
 
+class SViewport;
+
 /**
  * A viewport for use with Slate SViewport widgets.
  */
diff --git a/Engine/Source/Runtime/JsonUtilities/Public/JsonObjectConverter.h b/Engine/Source/Runtime/JsonUtilities/Public/JsonObjectConverter.h
--- a/Engine/Source/Runtime/JsonUtilities/Public/JsonObjectConverter.h
+++ b/Engine/Source/Runtime/JsonUtilities/Public/JsonObjectConverter.h
@@ -2,6 +2,8 @@
 
 #pragma once
 
+#include "Json.h"
+
 /** Class that handles converting Json objects to and from UStructs */
 class JSONUTILITIES_API FJsonObjectConverter
 {
diff --git a/Engine/Source/Runtime/MovieSceneCapture/Public/MovieSceneCapture.h b/Engine/Source/Runtime/MovieSceneCapture/Public/MovieSceneCapture.h
--- a/Engine/Source/Runtime/MovieSceneCapture/Public/MovieSceneCapture.h
+++ b/Engine/Source/Runtime/MovieSceneCapture/Public/MovieSceneCapture.h
@@ -7,6 +7,7 @@
 #include "MovieSceneCaptureHandle.h"
 #include "MovieScene.h"
 #include "AVIWriter.h"
+#include "Tickable.h"
 #include "MovieSceneCapture.generated.h"
 
 /** Interface that defines when to capture or drop frames */
diff --git a/Engine/Source/Runtime/MovieSceneCapture/Public/MovieSceneCaptureSettings.h b/Engine/Source/Runtime/MovieSceneCapture/Public/MovieSceneCaptureSettings.h
--- a/Engine/Source/Runtime/MovieSceneCapture/Public/MovieSceneCaptureSettings.h
+++ b/Engine/Source/Runtime/MovieSceneCapture/Public/MovieSceneCaptureSettings.h
@@ -2,6 +2,7 @@
 
 #pragma once
 
+#include "Engine/EngineTypes.h"
 #include "MovieSceneCaptureSettings.generated.h"
 
 UENUM()
diff --git a/Engine/Source/Runtime/Slate/Private/Widgets/Docking/SDockTab.cpp b/Engine/Source/Runtime/Slate/Private/Widgets/Docking/SDockTab.cpp
--- a/Engine/Source/Runtime/Slate/Private/Widgets/Docking/SDockTab.cpp
+++ b/Engine/Source/Runtime/Slate/Private/Widgets/Docking/SDockTab.cpp
@@ -751,3 +751,8 @@ FVector2D SDockTab::GetAnimatedScale() const
 	static FVector2D FullyClosed = FVector2D(1.0f, 0.0f);
 	return FMath::Lerp(FullyClosed, FullyOpen, SpawnAnimCurve.GetLerp());
 }
+
+void SDockTab::UpdateActivationTime()
+{
+	LastActivationTime = FSlateApplication::Get().GetCurrentTime();
+}
diff --git a/Engine/Source/Runtime/Slate/Public/Framework/Docking/TabManager.h b/Engine/Source/Runtime/Slate/Public/Framework/Docking/TabManager.h
--- a/Engine/Source/Runtime/Slate/Public/Framework/Docking/TabManager.h
+++ b/Engine/Source/Runtime/Slate/Public/Framework/Docking/TabManager.h
@@ -5,7 +5,8 @@
 #include "SlateIcon.h"
 #include "WorkspaceItem.h"
 
-
+class FMenuBuilder;
+class FMultiBox;
 class FTabManager;
 class SDockingArea;
 class SDockTab;
@@ -13,7 +14,6 @@ class SDockingTabStack;
 class SWindow;
 struct FTabMatcher;
 
-
 DECLARE_MULTICAST_DELEGATE_TwoParams(
 	FOnActiveTabChanged,
 	/** Previously active tab */
diff --git a/Engine/Source/Runtime/Slate/Public/Widgets/Docking/SDockTab.h b/Engine/Source/Runtime/Slate/Public/Widgets/Docking/SDockTab.h
--- a/Engine/Source/Runtime/Slate/Public/Widgets/Docking/SDockTab.h
+++ b/Engine/Source/Runtime/Slate/Public/Widgets/Docking/SDockTab.h
@@ -2,9 +2,14 @@
 
 #pragma once
 
+#include "SBorder.h"
+#include "TabManager.h"
 
 class SDockingTabWell;
 class SDockingTabStack;
+class SImage;
+class STextBlock;
+class SToolTip;
 
 /** How will this tab be used. */
 enum ETabRole : uint8
@@ -244,10 +249,7 @@ public:
 	bool HasSiblingTab(const FTabId& SiblingTabId, const bool TreatIndexNoneAsWildcard = true) const;
 
 	/** Updates the 'last activated' time to the current time */
-	void UpdateActivationTime()
-	{
-		LastActivationTime = FSlateApplication::Get().GetCurrentTime();
-	}
+	void UpdateActivationTime();
 
 	/** Returns the time this tab was last activated */
 	double GetLastActivationTime()
