Mergine changes from //depot/UE4 to //UE4/Main

diff --git a/Engine/Plugins/MovieScene/LevelSequenceEditor/Source/LevelSequenceEditor/Private/LevelSequenceEditorToolkit.cpp b/Engine/Plugins/MovieScene/LevelSequenceEditor/Source/LevelSequenceEditor/Private/LevelSequenceEditorToolkit.cpp
--- a/Engine/Plugins/MovieScene/LevelSequenceEditor/Source/LevelSequenceEditor/Private/LevelSequenceEditorToolkit.cpp
+++ b/Engine/Plugins/MovieScene/LevelSequenceEditor/Source/LevelSequenceEditor/Private/LevelSequenceEditorToolkit.cpp
@@ -263,7 +263,7 @@ void FLevelSequenceEditorToolkit::AddActorsToSequencer(AActor*const* InActors, i
 		// Add this object if it hasn't already been possessed
 		if (!ObjectGuid.IsValid() || !FocussedMovieScene->FindPossessable(ObjectGuid))
 		{
-			ObjectGuid = FocussedMovieScene->AddPossessable(ThisActor->GetName(), ThisActor->GetClass());
+			ObjectGuid = FocussedMovieScene->AddPossessable(ThisActor->GetActorLabel(), ThisActor->GetClass());
 			FocussedSequence->BindPossessableObject(ObjectGuid, *ThisActor);
 		}
 
diff --git a/Engine/Plugins/MovieScene/LevelSequenceEditor/Source/LevelSequenceEditor/Private/LevelSequenceProjectSettings.cpp b/Engine/Plugins/MovieScene/LevelSequenceEditor/Source/LevelSequenceEditor/Private/LevelSequenceProjectSettings.cpp
--- a/Engine/Plugins/MovieScene/LevelSequenceEditor/Source/LevelSequenceEditor/Private/LevelSequenceProjectSettings.cpp
+++ b/Engine/Plugins/MovieScene/LevelSequenceEditor/Source/LevelSequenceEditor/Private/LevelSequenceProjectSettings.cpp
@@ -3,8 +3,8 @@
 #include "LevelSequenceEditorPCH.h"
 #include "LevelSequenceProjectSettings.h"
 
+
 ULevelSequenceProjectSettings::ULevelSequenceProjectSettings()
-	: DefaultStartTime(1.f/30)	// Default to frame 1 of a 30fps movie
+	: DefaultStartTime(0.f)
 	, DefaultDuration(5.f)
-{
-}
+{ }
diff --git a/Engine/Plugins/Runtime/PacketHandlers/CompressionComponents/Oodle/Oodle.uplugin b/Engine/Plugins/Runtime/PacketHandlers/CompressionComponents/Oodle/Oodle.uplugin
--- a/Engine/Plugins/Runtime/PacketHandlers/CompressionComponents/Oodle/Oodle.uplugin
+++ b/Engine/Plugins/Runtime/PacketHandlers/CompressionComponents/Oodle/Oodle.uplugin
@@ -9,7 +9,7 @@
 	"EngineVersion" : "4.4.0",
 	"Description" : "Realtime network compression using Oodle technology",
 	"Category" : "Networking",
-	"EnabledByDefault" : false,
+	"EnabledByDefault" : true,
     "IsBetaVersion": false,
 	
 	"Modules" :
diff --git a/Engine/Plugins/Runtime/PacketHandlers/CompressionComponents/Oodle/Source/OodleHandlerComponent/Private/OodleHandlerComponent.cpp b/Engine/Plugins/Runtime/PacketHandlers/CompressionComponents/Oodle/Source/OodleHandlerComponent/Private/OodleHandlerComponent.cpp
--- a/Engine/Plugins/Runtime/PacketHandlers/CompressionComponents/Oodle/Source/OodleHandlerComponent/Private/OodleHandlerComponent.cpp
+++ b/Engine/Plugins/Runtime/PacketHandlers/CompressionComponents/Oodle/Source/OodleHandlerComponent/Private/OodleHandlerComponent.cpp
@@ -892,7 +892,7 @@ void FOodleComponentModuleInterface::StartupModule()
 	GOodleSaveDir = FPaths::ConvertRelativePathToFull(FPaths::Combine(*FPaths::GameSavedDir(), TEXT("Oodle")));
 
 
-	// Load the Oodle library
+	// Load the Oodle library (NOTE: Path mirrored in Oodle.Build.cs)
 	FString OodleBinaryPath = FPaths::EngineDir() / TEXT("Binaries/ThirdParty/NotForLicensees/Oodle/");
 	FString OodleBinaryFile;
 
diff --git a/Engine/Shaders/BasePassPixelShader.usf b/Engine/Shaders/BasePassPixelShader.usf
--- a/Engine/Shaders/BasePassPixelShader.usf
+++ b/Engine/Shaders/BasePassPixelShader.usf
@@ -934,6 +934,14 @@ void Main(
 			GBuffer.ShadingModelID = SHADINGMODELID_CLOTH;
 			GBuffer.CustomData.rgb = SubsurfaceColor;
 			GBuffer.CustomData.a = saturate( GetMaterialCustomData0(MaterialParameters) );	// Cloth
+		#elif MATERIAL_SHADINGMODEL_EYE
+			GBuffer.ShadingModelID = SHADINGMODELID_EYE;
+			#if NUM_MATERIAL_OUTPUTS_GETTANGENTOUTPUT > 0
+				float3 Tangent = GetTangentOutput0(MaterialParameters);
+				GBuffer.CustomData.xy = UnitVectorToOctahedron( normalize(Tangent) ) * 0.5 + 0.5;
+			#endif
+			GBuffer.CustomData.z = saturate( GetMaterialCustomData0(MaterialParameters) );	// Iris Mask
+			GBuffer.CustomData.w = saturate( GetMaterialCustomData1(MaterialParameters) );	// Iris Distance
 		#else
 			// missing shading model, compiler should report ShadingModelID is not set
 		#endif
diff --git a/Engine/Shaders/CircleDOFCommon.usf b/Engine/Shaders/CircleDOFCommon.usf
--- a/Engine/Shaders/CircleDOFCommon.usf
+++ b/Engine/Shaders/CircleDOFCommon.usf
@@ -5,12 +5,13 @@
 =============================================================================*/
 
 
-// {radius, depth blur amp, depth blur radius, ViewPortScale}
-float4 CircleDofParams;
 
-// CircleDOF: Compute circle of confusion size in pixels.
+// CircleDOF: Compute circle of confusion size in half res pixels.
 float DepthToCoc(float SceneDepth) 
 {
+	// {radius, depth blur amp, depth blur radius, ViewPortScale}
+	float4 CircleDofParams = View.CircleDOFParams;
+
 	// We intentionally don't do the DepthOfFieldFocalRegion as it breaks realism.
 
 	float Focus = View.DepthOfFieldFocalDistance;
diff --git a/Engine/Shaders/CommonViewUniformBuffer.usf b/Engine/Shaders/CommonViewUniformBuffer.usf
--- a/Engine/Shaders/CommonViewUniformBuffer.usf
+++ b/Engine/Shaders/CommonViewUniformBuffer.usf
@@ -17,14 +17,15 @@ float2 GetPrevTanHalfFieldOfView()
 	return float2(View.PrevClipToView[0][0], View.PrevClipToView[1][1]);
 }
 
-/*
-* @return 1 / tan(View.FieldOfViewWideAngles * .5)
-*/
+// might be used by Custom material expressions (still best to wrap the custom node in a material function)
+// @return 1 / tan(View.FieldOfViewWideAngles * .5)
 float2 GetCotanHalfFieldOfView()
 {
 	return float2(View.ViewToClip[0][0], View.ViewToClip[1][1]);
 }
 
+// might be used by Custom material expressions (still best to wrap the custom node in a material function)
+// @return previous 1 / tan(View.FieldOfViewWideAngles * .5)
 float2 GetPrevCotanHalfFieldOfView()
 {
 	return float2(View.PrevViewToClip[0][0], View.PrevViewToClip[1][1]);
diff --git a/Engine/Shaders/DeferredShadingCommon.usf b/Engine/Shaders/DeferredShadingCommon.usf
--- a/Engine/Shaders/DeferredShadingCommon.usf
+++ b/Engine/Shaders/DeferredShadingCommon.usf
@@ -213,6 +213,8 @@ float Decode71(float Scalar, out uint Mask)
 #define SHADINGMODELID_TWOSIDED_FOLIAGE		6
 #define SHADINGMODELID_HAIR					7
 #define SHADINGMODELID_CLOTH				8
+#define SHADINGMODELID_EYE					9
+#define SHADINGMODELID_NUM					10
 #define SHADINGMODELID_MASK					0xF					
 
 // The flags are defined so that 0 value has no effect!
diff --git a/Engine/Shaders/Definitions.usf b/Engine/Shaders/Definitions.usf
--- a/Engine/Shaders/Definitions.usf
+++ b/Engine/Shaders/Definitions.usf
@@ -188,6 +188,10 @@
 #define USE_CLOTH										0
 #endif
 
+#ifndef USE_EYE
+#define USE_EYE										0
+#endif
+
 #if SM5_PROFILE || COMPILER_SUPPORTS_ATTRIBUTES
 /** Avoids flow control constructs. */
 #define UNROLL											[unroll]
diff --git a/Engine/Shaders/MaterialTemplate.usf b/Engine/Shaders/MaterialTemplate.usf
--- a/Engine/Shaders/MaterialTemplate.usf
+++ b/Engine/Shaders/MaterialTemplate.usf
@@ -9,6 +9,7 @@
 
 // for MaterialExpressionDepthOfFieldFunction
 #include "DepthOfFieldCommon.usf"
+#include "CircleDOFCommon.usf"
 #include "GlobalDistanceFieldShared.usf"
 
 #if USES_SPEEDTREE
@@ -621,6 +622,11 @@ float MaterialExpressionDepthOfFieldFunction(float SceneDepth, int FunctionValue
 	{
 		return CalcUnfocusedPercentCustomBound(SceneDepth, 0, 1);
 	}
+	else if(FunctionValueIndex == 3) // TDOF_CircleOfConfusionRadius
+	{
+		// * 2 to compensate for half res
+		return DepthToCoc(SceneDepth) * 2.0f;
+	}
 	return 0;
 }
 
@@ -1055,7 +1061,8 @@ float3 GetShadingModelColor(FGBufferData GBuffer)
 	else if (GBuffer.ShadingModelID == SHADINGMODELID_SUBSURFACE_PROFILE) return float3(0.2f, 0.6f, 0.5f); // Cyan
 	else if (GBuffer.ShadingModelID == SHADINGMODELID_TWOSIDED_FOLIAGE) return float3(0.2f, 0.2f, 0.8f); // Blue
 	else if (GBuffer.ShadingModelID == SHADINGMODELID_HAIR) return float3(0.6f, 0.1f, 0.5f);
-	else if (GBuffer.ShadingModelID == SHADINGMODELID_CLOTH) return float3(0.18f, 0.18f, 0.18f); 
+	else if (GBuffer.ShadingModelID == SHADINGMODELID_CLOTH) return float3(0.7f, 1.0f, 1.0f); 
+	else if (GBuffer.ShadingModelID == SHADINGMODELID_EYE) return float3(0.3f, 1.0f, 1.0f); 
 	else return float3(1.0f, 1.0f, 1.0f); // White
 #else
 	switch(GBuffer.ShadingModelID)
@@ -1068,7 +1075,8 @@ float3 GetShadingModelColor(FGBufferData GBuffer)
 		case SHADINGMODELID_SUBSURFACE_PROFILE: return float3(0.2f, 0.6f, 0.5f); // Cyan
 		case SHADINGMODELID_TWOSIDED_FOLIAGE: return float3(0.2f, 0.2f, 0.8f); // Cyan
 		case SHADINGMODELID_HAIR: return float3(0.6f, 0.1f, 0.5f);
-		case SHADINGMODELID_CLOTH: return float3(0.18f, 0.18f, 0.18f); // Brown
+		case SHADINGMODELID_CLOTH: return float3(0.7f, 1.0f, 1.0f);
+		case SHADINGMODELID_EYE: return float3(0.3f, 1.0f, 1.0f);
 		default: return float3(1.0f, 1.0f, 1.0f); // White
 	}
 #endif
diff --git a/Engine/Shaders/ParticleGPUSpriteVertexFactory.usf b/Engine/Shaders/ParticleGPUSpriteVertexFactory.usf
--- a/Engine/Shaders/ParticleGPUSpriteVertexFactory.usf
+++ b/Engine/Shaders/ParticleGPUSpriteVertexFactory.usf
@@ -364,6 +364,7 @@ FMaterialPixelParameters GetMaterialPixelParameters( FVertexFactoryInterpolantsV
 {
 	// GetMaterialPixelParameters is responsible for fully initializing the result
 	FMaterialPixelParameters Result = MakeInitializedMaterialPixelParameters();
+	Result.Particle.MacroUV = EmitterDynamicUniforms.MacroUVParameters;
 
 #if NUM_MATERIAL_TEXCOORDS
 	UNROLL
@@ -449,6 +450,7 @@ FMaterialVertexParameters GetMaterialVertexParameters(FVertexFactoryInput Input,
 	Result.WorldPosition = WorldPosition;
 	Result.VertexColor = 1;
 	Result.TangentToWorld = mul(TangentToLocal, GetLocalToWorld3x3()); 
+	Result.Particle.MacroUV = EmitterDynamicUniforms.MacroUVParameters;
 
 #if USE_PARTICLE_TIME
 	Result.Particle.RelativeTime = Intermediates.RelativeTime;
diff --git a/Engine/Shaders/PostProcessAmbient.usf b/Engine/Shaders/PostProcessAmbient.usf
--- a/Engine/Shaders/PostProcessAmbient.usf
+++ b/Engine/Shaders/PostProcessAmbient.usf
@@ -315,6 +315,13 @@ float3 ImageBasedLightingMIS( FGBufferData GBuffer, float3 V, float3 N, uint2 Ra
 					Shading = ClothShading( GBuffer, GBuffer.Roughness, 1, L, V, N, 1 );
 				}
 				else
+#endif
+#if USE_EYE
+				BRANCH if( GBuffer.ShadingModelID == SHADINGMODELID_EYE )
+				{
+					Shading = EyeShading( GBuffer, GBuffer.Roughness, 1, L, V, N, 1 );
+				}
+				else
 #endif
 				{
 					Shading = StandardShading( GBuffer, GBuffer.Roughness, L, V, N );
diff --git a/Engine/Shaders/PostProcessCircleDOF.usf b/Engine/Shaders/PostProcessCircleDOF.usf
--- a/Engine/Shaders/PostProcessCircleDOF.usf
+++ b/Engine/Shaders/PostProcessCircleDOF.usf
@@ -222,7 +222,7 @@ void Circle4Samples(float4 UVAndScreenPos : TEXCOORD0, out float4 OutColor0 : SV
 
 	// Fixed maximum search size (in terms of Circle of Confusion radius).
 	// Higher than 8 is too noizy for 4 samples.
-	float Coc = 8.0 * CircleDofParams.w;
+	float Coc = 8.0 * View.CircleDOFParams.w;
 
 	float LocalRand = RandomOffset.x + InRand;
 
diff --git a/Engine/Shaders/ReflectionEnvironmentComputeShaders.usf b/Engine/Shaders/ReflectionEnvironmentComputeShaders.usf
--- a/Engine/Shaders/ReflectionEnvironmentComputeShaders.usf
+++ b/Engine/Shaders/ReflectionEnvironmentComputeShaders.usf
@@ -587,8 +587,8 @@ void ReflectionEnvironmentTiledDeferredMain(
 	uint ThreadIndex = GroupThreadId.y * THREADGROUP_SIZEX + GroupThreadId.x;
 	
 	uint2 PixelPos = DispatchThreadId.xy + ViewDimensions.xy;
-	float2 UV = (float2(DispatchThreadId.xy) + .5f) / (ViewDimensions.zw - ViewDimensions.xy);
-	float2 ScreenPosition = float2(2.0f, -2.0f) * UV + float2(-1.0f, 1.0f);
+	float2 ViewportUV = (float2(DispatchThreadId.xy) + .5f) / (ViewDimensions.zw - ViewDimensions.xy);
+	float2 ScreenPosition = float2(2.0f, -2.0f) * ViewportUV + float2(-1.0f, 1.0f);
 	float SceneDepth = CalcSceneDepth(PixelPos);
 
 	float MinTileZ;
diff --git a/Engine/Shaders/ShaderVersion.usf b/Engine/Shaders/ShaderVersion.usf
--- a/Engine/Shaders/ShaderVersion.usf
+++ b/Engine/Shaders/ShaderVersion.usf
@@ -2,4 +2,4 @@
 // Each time the console command is executed it generates a new GUID. As this file is included
 // in common.usf (which should be included in any shader) it allows to invalidate the shader DDC.
 // 
-// GUID = 3022cb7a3314413686da187d920fd485
\ No newline at end of file
+// GUID = 8FE3C62E-D12D4E3B94F997A3F1C66BDF
\ No newline at end of file
diff --git a/Engine/Shaders/ShadingModels.usf b/Engine/Shaders/ShadingModels.usf
--- a/Engine/Shaders/ShadingModels.usf
+++ b/Engine/Shaders/ShadingModels.usf
@@ -664,6 +664,68 @@ float3 SubsurfaceShadingTwoSided( float3 SubsurfaceColor, float3 L, float3 V, ha
 	return NoL * GGX * SubsurfaceColor;
 }
 
+float3 EyeShading( FGBufferData GBuffer, float3 LobeRoughness, float3 LobeEnergy, float3 L, float3 V, half3 N )
+{
+	float3 H = normalize(V + L);
+	float VoH = saturate( dot(V, H) );
+
+	float3 F = F_Schlick( GBuffer.SpecularColor, VoH );
+
+	//
+	// Cornea Reflectance - Use geo normal
+	//
+	float CNoL = saturate( dot(N, L) );
+	float CNoV = max( dot(N, V), 1e-5 );	// constant to prevent NaN
+	float CNoH = saturate( dot(N, H) );
+	
+	// Generalized microfacet specular
+	float D = D_GGX( LobeRoughness[1], CNoH ) * LobeEnergy[1];
+	float Vis = Vis_SmithJointApprox( LobeRoughness[1], CNoV, CNoL );
+
+	// Cornea Specular - Covers the whole surface
+	float3 CorneaSpecular;
+	CorneaSpecular = D * Vis * F;
+
+	return CorneaSpecular;
+}
+
+float3 EyeSubsurfaceShading(  FGBufferData GBuffer, float3 L, float3 V, half3 N )
+{
+	// Need to make sure we are getting the correct Diffuse color intensity
+	float3 SubsurfaceColor = GBuffer.DiffuseColor * PI;
+
+	float  IrisDistance = GBuffer.CustomData.w;
+	float  IrisMask = GBuffer.CustomData.z;
+
+	float3 IrisNormal;
+	IrisNormal = OctahedronToUnitVector( GBuffer.CustomData.xy * 2 - 1 );
+
+	// Blend in the negative intersection normal to create some concavity
+	// Not great as it ties the concavity to the convexity of the cornea surface
+	// No good justification for that. On the other hand, if we're just looking to
+	// introduce some concavity, this does the job.
+	IrisNormal = normalize(lerp(IrisNormal, -N, IrisMask*IrisDistance));
+
+	// Use IrisNormal only within the Iris
+	IrisNormal = lerp(N, IrisNormal, IrisMask);
+
+	// http://blog.stevemcauley.com/2011/12/03/energy-conserving-wrapped-diffuse/
+	float Wrap = 0.5;
+	float INoL = saturate( ( dot(IrisNormal, L) + Wrap ) / Square( 1 + Wrap ) );
+
+	// GGX scatter distribution
+	float VoL = saturate( dot(V, -L) );
+	float a = 0.6;
+	float a2 = a * a;
+	float d = ( VoL * a2 - VoL ) * VoL + 1;	// 2 mad
+	float GGX = (a2 / PI) / (d * d);		// 2 mul, 1 rcp
+	
+	float3 rawResult = INoL * GGX * SubsurfaceColor;
+
+	// Covers whole object
+	return rawResult;
+}
+
 Texture2D		PreIntegratedBRDF;
 SamplerState	PreIntegratedBRDFSampler;
 
@@ -693,6 +755,8 @@ float3 SurfaceShading( FGBufferData GBuffer, float3 LobeRoughness, float3 LobeEn
 			return ClearCoatShading( GBuffer, LobeRoughness, LobeEnergy, L, V, N );
 		case SHADINGMODELID_CLOTH:
 			return ClothShading( GBuffer, LobeRoughness, LobeEnergy, L, V, N );
+		case SHADINGMODELID_EYE:
+			return EyeShading( GBuffer, LobeRoughness, LobeEnergy, L, V, N );
 		default:
 			return 0;
 	}
@@ -712,6 +776,8 @@ float3 SubsurfaceShading( FGBufferData GBuffer, float3 L, float3 V, half3 N, flo
 			return SubsurfaceShadingTwoSided( SubsurfaceColor, L, V, N );
 		case SHADINGMODELID_HAIR:
 			return HairShading( GBuffer, L, V, N, Shadow, Random );
+		case SHADINGMODELID_EYE:
+			return EyeSubsurfaceShading( GBuffer, L, V, N );
 		default:
 			return 0;
 	}
diff --git a/Engine/Source/Developer/LogVisualizer/Private/TimeSliderController.cpp b/Engine/Source/Developer/LogVisualizer/Private/TimeSliderController.cpp
--- a/Engine/Source/Developer/LogVisualizer/Private/TimeSliderController.cpp
+++ b/Engine/Source/Developer/LogVisualizer/Private/TimeSliderController.cpp
@@ -422,7 +422,7 @@ FReply FVisualLoggerTimeSliderController::OnMouseMove( TSharedRef<SWidget> Widge
 					NewViewOutputMax = LocalClampMax;
 				}
 
-				TimeSliderArgs.OnViewRangeChanged.ExecuteIfBound(TRange<float>(NewViewOutputMin, NewViewOutputMax), EViewRangeInterpolation::Immediate, false);
+				TimeSliderArgs.OnViewRangeChanged.ExecuteIfBound(TRange<float>(NewViewOutputMin, NewViewOutputMax), EViewRangeInterpolation::Immediate);
 				if (Scrollbar.IsValid())
 				{
 					float InOffsetFraction = (NewViewOutputMin - LocalClampMin) / (LocalClampMax - LocalClampMin);
@@ -600,7 +600,7 @@ FReply FVisualLoggerTimeSliderController::OnMouseWheel( TSharedRef<SWidget> Widg
 					NewViewOutputMax = LocalClampMax;
 				}
 
-				TimeSliderArgs.OnViewRangeChanged.ExecuteIfBound(TRange<float>(NewViewOutputMin, NewViewOutputMax), EViewRangeInterpolation::Immediate, false);
+				TimeSliderArgs.OnViewRangeChanged.ExecuteIfBound(TRange<float>(NewViewOutputMin, NewViewOutputMax), EViewRangeInterpolation::Immediate);
 				if (Scrollbar.IsValid())
 				{
 					float InOffsetFraction = (NewViewOutputMin - LocalClampMin) / (LocalClampMax - LocalClampMin);
diff --git a/Engine/Source/Developer/MaterialUtilities/Private/MaterialUtilities.cpp b/Engine/Source/Developer/MaterialUtilities/Private/MaterialUtilities.cpp
--- a/Engine/Source/Developer/MaterialUtilities/Private/MaterialUtilities.cpp
+++ b/Engine/Source/Developer/MaterialUtilities/Private/MaterialUtilities.cpp
@@ -239,11 +239,12 @@ public:
 		{
 			TArray<FShaderType*> ShaderTypes;
 			TArray<FVertexFactoryType*> VFTypes;
-			GetDependentShaderAndVFTypes(GMaxRHIShaderPlatform, ShaderTypes, VFTypes);
+			TArray<const FShaderPipelineType*> ShaderPipelineTypes;
+			GetDependentShaderAndVFTypes(GMaxRHIShaderPlatform, ShaderTypes, ShaderPipelineTypes, VFTypes);
 
 			// Overwrite the shader map Id's dependencies with ones that came from the FMaterial actually being compiled (this)
 			// This is necessary as we change FMaterial attributes like GetShadingModel(), which factor into the ShouldCache functions that determine dependent shader types
-			ResourceId.SetShaderDependencies(ShaderTypes, VFTypes);
+			ResourceId.SetShaderDependencies(ShaderTypes, ShaderPipelineTypes, VFTypes);
 		}
 
 		// Override with a special usage so we won't re-use the shader map used by the material for rendering
diff --git a/Engine/Source/Developer/ShaderCompilerCommon/Private/HlslAST.cpp b/Engine/Source/Developer/ShaderCompilerCommon/Private/HlslAST.cpp
--- a/Engine/Source/Developer/ShaderCompilerCommon/Private/HlslAST.cpp
+++ b/Engine/Source/Developer/ShaderCompilerCommon/Private/HlslAST.cpp
@@ -351,6 +351,17 @@ namespace CrossCompiler
 			}
 		}
 
+		bool FUnaryExpression::GetConstantIntValue(int32& OutValue) const
+		{
+			if (IsConstant())
+			{
+				OutValue = (int32)GetUintConstantValue();
+				return true;
+			}
+
+			return false;
+		}
+
 		FBinaryExpression::FBinaryExpression(FLinearAllocator* InAllocator, EOperators InOperator, FExpression* E0, FExpression* E1, const FSourceInfo& InInfo) :
 			FExpression(InAllocator, InOperator, E0, E1, nullptr, InInfo)
 		{
@@ -389,6 +400,42 @@ namespace CrossCompiler
 			}
 		}
 
+		bool FBinaryExpression::GetConstantIntValue(int32& OutValue) const
+		{
+			int32 LHS = 0;
+			int32 RHS = 0;
+			if (!SubExpressions[0]->GetConstantIntValue(LHS) || !SubExpressions[1]->GetConstantIntValue(RHS))
+			{
+				return false;
+			}
+
+			switch (Operator)
+			{
+			default:
+				return false;
+
+			case EOperators::LogicOr:	OutValue = LHS || RHS; break;
+			case EOperators::LogicAnd:	OutValue = LHS && RHS; break;
+			case EOperators::BitOr:		OutValue = LHS | RHS; break;
+			case EOperators::BitXor:	OutValue = LHS ^ RHS; break;
+			case EOperators::BitAnd:	OutValue = LHS ^ RHS; break;
+			case EOperators::Equal:		OutValue = LHS == RHS; break;
+			case EOperators::NEqual:	OutValue = LHS != RHS; break;
+			case EOperators::Less:		OutValue = LHS < RHS; break;
+			case EOperators::Greater:	OutValue = LHS > RHS; break;
+			case EOperators::LEqual:	OutValue = LHS <= RHS; break;
+			case EOperators::GEqual:	OutValue = LHS >= RHS; break;
+			case EOperators::LShift:	OutValue = LHS << RHS; break;
+			case EOperators::RShift:	OutValue = LHS >> RHS; break;
+			case EOperators::Add:		OutValue = LHS + RHS; break;
+			case EOperators::Sub:		OutValue = LHS - RHS; break;
+			case EOperators::Mul:		OutValue = LHS * RHS; break;
+			case EOperators::Div:		OutValue = LHS / RHS; break;
+			case EOperators::Mod:		OutValue = LHS % RHS; break;
+			}
+			return true;
+		}
+
 		FExpressionStatement::FExpressionStatement(FLinearAllocator* InAllocator, FExpression* InExpr, const FSourceInfo& InInfo) :
 			FNode(InAllocator, InInfo),
 			Expression(InExpr)
diff --git a/Engine/Source/Developer/ShaderCompilerCommon/Private/HlslAST.h b/Engine/Source/Developer/ShaderCompilerCommon/Private/HlslAST.h
--- a/Engine/Source/Developer/ShaderCompilerCommon/Private/HlslAST.h
+++ b/Engine/Source/Developer/ShaderCompilerCommon/Private/HlslAST.h
@@ -108,6 +108,9 @@ namespace CrossCompiler
 			virtual FParameterDeclarator* AsParameterDeclarator() { return nullptr; }
 			virtual FUnaryExpression* AsUnaryExpression() { return nullptr; }
 
+			// Returns true if the expression can be evaluated to a constant int
+			virtual bool GetConstantIntValue(int32& OutValue) const { return false; }
+
 			/* Callers of this ralloc-based new need not call delete. It's
 			* easier to just ralloc_free 'ctx' (or any of its ancestors). */
 			static void* operator new(SIZE_T Size, FLinearAllocator* Allocator)
@@ -335,6 +338,7 @@ namespace CrossCompiler
 
 			virtual void Write(FASTWriter& Writer) const override;
 			virtual FUnaryExpression* AsUnaryExpression() override { return this; }
+			virtual bool GetConstantIntValue(int32& OutValue) const override;
 
 			bool IsConstant() const
 			{
@@ -377,6 +381,7 @@ namespace CrossCompiler
 			FBinaryExpression(FLinearAllocator* InAllocator, EOperators InOperator, FExpression* E0, FExpression* E1, const FSourceInfo& InInfo);
 
 			virtual void Write(FASTWriter& Writer) const override;
+			virtual bool GetConstantIntValue(int32& OutValue) const override;
 		};
 
 		struct FFunctionExpression : public FExpression
diff --git a/Engine/Source/Developer/ShaderCompilerCommon/Private/HlslUtils.cpp b/Engine/Source/Developer/ShaderCompilerCommon/Private/HlslUtils.cpp
--- a/Engine/Source/Developer/ShaderCompilerCommon/Private/HlslUtils.cpp
+++ b/Engine/Source/Developer/ShaderCompilerCommon/Private/HlslUtils.cpp
@@ -598,14 +598,13 @@ struct FRemoveUnusedOutputs
 
 			for (int32 Index = 0; Index < A->ArraySize.Num(); ++Index)
 			{
-				FUnaryExpression* UnaryA = A->ArraySize[Index]->AsUnaryExpression();
-				if (!UnaryA || !UnaryA->IsConstant())
+				int32 DimA = 0;
+				if (!A->ArraySize[Index]->GetConstantIntValue(DimA))
 				{
 					Errors.Add(FString::Printf(TEXT("Array %s is not a compile-time constant expression!"), A->Identifier));
 					return false;
 				}
 
-				uint32 DimA = UnaryA->GetUintConstantValue();
 				OutLength = DimA;
 			}
 		}
diff --git a/Engine/Source/Developer/Windows/ShaderFormatD3D/Private/D3D11ShaderCompiler.cpp b/Engine/Source/Developer/Windows/ShaderFormatD3D/Private/D3D11ShaderCompiler.cpp
--- a/Engine/Source/Developer/Windows/ShaderFormatD3D/Private/D3D11ShaderCompiler.cpp
+++ b/Engine/Source/Developer/Windows/ShaderFormatD3D/Private/D3D11ShaderCompiler.cpp
@@ -328,6 +328,61 @@ void CompileD3D11Shader(const FShaderCompilerInput& Input,FShaderCompilerOutput&
 		return;
 	}
 
+	FString EntryPointName = Input.EntryPointName;
+
+	Output.bFailedRemovingUnused = false;
+	if (Input.Target.Frequency == SF_Vertex && Input.bCompilingForShaderPipeline)
+	{
+		static TArray<FString> VertexSystemOutputs;
+		if (VertexSystemOutputs.Num() == 0)
+		{
+			VertexSystemOutputs.Add(TEXT("SV_POSITION"));
+			VertexSystemOutputs.Add(TEXT("SV_Position"));
+			VertexSystemOutputs.Add(TEXT("SV_ClipDistance"));
+			VertexSystemOutputs.Add(TEXT("SV_ClipDistance0"));
+			VertexSystemOutputs.Add(TEXT("SV_ClipDistance1"));
+			VertexSystemOutputs.Add(TEXT("SV_ClipDistance2"));
+			VertexSystemOutputs.Add(TEXT("SV_ClipDistance3"));
+			VertexSystemOutputs.Add(TEXT("SV_ClipDistance4"));
+			VertexSystemOutputs.Add(TEXT("SV_ClipDistance5"));
+			VertexSystemOutputs.Add(TEXT("SV_CullDistance0"));
+			VertexSystemOutputs.Add(TEXT("SV_CullDistance1"));
+			VertexSystemOutputs.Add(TEXT("SV_CullDistance2"));
+			VertexSystemOutputs.Add(TEXT("SV_CullDistance3"));
+			VertexSystemOutputs.Add(TEXT("SV_CullDistance4"));
+			VertexSystemOutputs.Add(TEXT("SV_CullDistance5"));
+			VertexSystemOutputs.Add(TEXT("SV_Coverage"));
+			VertexSystemOutputs.Add(TEXT("SV_Depth"));
+			VertexSystemOutputs.Add(TEXT("SV_DomainLocation"));
+			VertexSystemOutputs.Add(TEXT("SV_IsFrontFace"));
+			VertexSystemOutputs.Add(TEXT("SV_OutputControlPointID"));
+			VertexSystemOutputs.Add(TEXT("SV_RenderTargetArrayIndex"));
+			VertexSystemOutputs.Add(TEXT("SV_SampleIndex"));
+			VertexSystemOutputs.Add(TEXT("SV_ViewportArrayIndex"));
+			VertexSystemOutputs.Add(TEXT("SV_TessFactor"));
+			VertexSystemOutputs.Add(TEXT("SV_InstanceID"));
+			VertexSystemOutputs.Add(TEXT("SV_PrimitiveID"));
+			VertexSystemOutputs.Add(TEXT("SV_VertexID"));
+		}
+
+		TArray<FString> Errors;
+		if (RemoveUnusedOutputs(PreprocessedShaderSource, VertexSystemOutputs, Input.UsedOutputs, Input.EntryPointName, Errors))
+		{
+			EntryPointName += FString(TEXT("__OPTIMIZED"));
+		}
+		else
+		{
+			UE_LOG(LogD3D11ShaderCompiler, Warning, TEXT("Failed to Remove unused outputs [%s]!"), *Input.DumpDebugInfoPath);
+			for (int32 Index = 0; Index < Errors.Num(); ++Index)
+			{
+				FShaderCompilerError NewError;
+				NewError.StrippedErrorMessage = Errors[Index];
+				Output.Errors.Add(NewError);
+			}
+			Output.bFailedRemovingUnused = true;
+		}
+	}
+
 	if (!RemoveUniformBuffersFromSource(PreprocessedShaderSource))
 	{
 		return;
@@ -395,7 +450,7 @@ void CompileD3D11Shader(const FShaderCompilerInput& Input,FShaderCompilerOutput&
 				TCHAR_TO_ANSI(*Input.SourceFilename),
 				/*pDefines=*/ NULL,
 				/*pInclude=*/ NULL,
-				TCHAR_TO_ANSI(*Input.EntryPointName),
+				TCHAR_TO_ANSI(*EntryPointName),
 				TCHAR_TO_ANSI(ShaderProfile),
 				CompileFlags,
 				0,
@@ -456,15 +511,16 @@ void CompileD3D11Shader(const FShaderCompilerInput& Input,FShaderCompilerOutput&
 				delete FileWriter;
 			}
 
-			const FString BatchFileContents = D3D11CreateShaderCompileCommandLine((Input.SourceFilename + TEXT(".usf")), *Input.EntryPointName, ShaderProfile, CompileFlags);
+			const FString BatchFileContents = D3D11CreateShaderCompileCommandLine((Input.SourceFilename + TEXT(".usf")), *EntryPointName, ShaderProfile, CompileFlags);
 			FFileHelper::SaveStringToFile(BatchFileContents, *(Input.DumpDebugInfoPath / TEXT("CompileD3D.bat")));
 
-			const FString BatchFileContents2 = CreateAMDCodeXLCommandLine((Input.SourceFilename + TEXT(".usf")), *Input.EntryPointName, ShaderProfile, CompileFlags);
+			const FString BatchFileContents2 = CreateAMDCodeXLCommandLine((Input.SourceFilename + TEXT(".usf")), *EntryPointName, ShaderProfile, CompileFlags);
 			FFileHelper::SaveStringToFile(BatchFileContents2, *(Input.DumpDebugInfoPath / TEXT("CompileAMD.bat")));
 		}
 
 		int32 NumInterpolants = 0;
 		TIndirectArray<FString> InterpolantNames;
+		TArray<FString> ShaderInputs;
 		if (SUCCEEDED(Result))
 		{
 			Output.bSucceeded = true;
@@ -486,6 +542,7 @@ void CompileD3D11Shader(const FShaderCompilerInput& Input,FShaderCompilerOutput&
 			uint32 NumCBs = 0;
 			uint32 NumUAVs = 0;
 			TArray<FString> UniformBufferNames;
+			TArray<FString> ShaderOutputs;
 
 			TBitArray<> UsedUniformBufferSlots;
 			UsedUniformBufferSlots.Init(false,32);
@@ -499,7 +556,21 @@ void CompileD3D11Shader(const FShaderCompilerInput& Input,FShaderCompilerOutput&
 					if (ParamDesc.SystemValueType == D3D_NAME_UNDEFINED && ParamDesc.Mask != 0)
 					{
 						++NumInterpolants;
-						new(InterpolantNames) FString(FString::Printf(TEXT("%s%d"),ANSI_TO_TCHAR(ParamDesc.SemanticName),ParamDesc.SemanticIndex));
+						new(InterpolantNames) FString(FString::Printf(TEXT("%s%d"), ANSI_TO_TCHAR(ParamDesc.SemanticName),ParamDesc.SemanticIndex));
+						ShaderOutputs.Add(*InterpolantNames.Last());
+					}
+				}
+			}
+			else if (Input.Target.Frequency == SF_Pixel)
+			{
+				for (uint32 Index = 0; Index < ShaderDesc.InputParameters; ++Index)
+				{
+					D3D11_SIGNATURE_PARAMETER_DESC ParamDesc;
+					Reflector->GetInputParameterDesc(Index, &ParamDesc);
+					if (ParamDesc.SystemValueType == D3D_NAME_UNDEFINED && ParamDesc.ReadWriteMask != 0)
+					{
+						FString UsedInput = FString::Printf(TEXT("%s%d"), ANSI_TO_TCHAR(ParamDesc.SemanticName), ParamDesc.SemanticIndex);
+						ShaderInputs.Add(UsedInput);
 					}
 				}
 			}
@@ -716,6 +787,12 @@ void CompileD3D11Shader(const FShaderCompilerInput& Input,FShaderCompilerOutput&
 				BuildResourceTableTokenStream(GenericSRT.UnorderedAccessViewMap, GenericSRT.MaxBoundResourceTable, SRT.UnorderedAccessViewMap);
 			}
 
+			if (Input.Target.Frequency == SF_Pixel && Input.bCompilingForShaderPipeline)
+			{
+				Output.bSupportsQueryingUsedAttributes = true;
+				Output.UsedAttributes = ShaderInputs;
+			}
+
 			FMemoryWriter Ar(Output.ShaderCode.GetWriteAccess(), true);
 			Ar << SRT;
 			Ar.Serialize(CompressedData->GetBufferPointer(), CompressedData->GetBufferSize());
diff --git a/Engine/Source/Editor/AnimGraph/Classes/AnimGraphNode_AnimDynamics.h b/Engine/Source/Editor/AnimGraph/Classes/AnimGraphNode_AnimDynamics.h
--- a/Engine/Source/Editor/AnimGraph/Classes/AnimGraphNode_AnimDynamics.h
+++ b/Engine/Source/Editor/AnimGraph/Classes/AnimGraphNode_AnimDynamics.h
@@ -83,4 +83,8 @@ protected:
 
 	void DrawLinearLimits(FPrimitiveDrawInterface* PDI, FTransform ShapeTransform, const FAnimNode_AnimDynamics& NodeToVisualise) const;
 	void DrawAngularLimits(FPrimitiveDrawInterface* PDI, FTransform JointTransform, const FAnimNode_AnimDynamics& NodeToVisualize) const;
+
+	// UAnimGraphNode_SkeletalControlBase protected interface
+	virtual const FAnimNode_SkeletalControlBase* GetNode() const override { return &Node; }
+	// End of UAnimGraphNode_SkeletalControlBase protected interface
 };
\ No newline at end of file
diff --git a/Engine/Source/Editor/AnimGraph/Classes/AnimGraphNode_ApplyAdditive.h b/Engine/Source/Editor/AnimGraph/Classes/AnimGraphNode_ApplyAdditive.h
--- a/Engine/Source/Editor/AnimGraph/Classes/AnimGraphNode_ApplyAdditive.h
+++ b/Engine/Source/Editor/AnimGraph/Classes/AnimGraphNode_ApplyAdditive.h
@@ -21,5 +21,6 @@ class UAnimGraphNode_ApplyAdditive : public UAnimGraphNode_Base
 
 	// UAnimGraphNode_Base interface
 	virtual FString GetNodeCategory() const override;
+	virtual void ValidateAnimNodeDuringCompilation(class USkeleton* ForSkeleton, class FCompilerResultsLog& MessageLog) override;
 	// End of UAnimGraphNode_Base interface
 };
diff --git a/Engine/Source/Editor/AnimGraph/Classes/AnimGraphNode_ApplyMeshSpaceAdditive.h b/Engine/Source/Editor/AnimGraph/Classes/AnimGraphNode_ApplyMeshSpaceAdditive.h
--- a/Engine/Source/Editor/AnimGraph/Classes/AnimGraphNode_ApplyMeshSpaceAdditive.h
+++ b/Engine/Source/Editor/AnimGraph/Classes/AnimGraphNode_ApplyMeshSpaceAdditive.h
@@ -21,5 +21,6 @@ class UAnimGraphNode_ApplyMeshSpaceAdditive : public UAnimGraphNode_Base
 
 	// UAnimGraphNode_Base interface
 	virtual FString GetNodeCategory() const override;
+	virtual void ValidateAnimNodeDuringCompilation(class USkeleton* ForSkeleton, class FCompilerResultsLog& MessageLog) override;
 	// End of UAnimGraphNode_Base interface
 };
diff --git a/Engine/Source/Editor/AnimGraph/Classes/AnimGraphNode_BoneDrivenController.h b/Engine/Source/Editor/AnimGraph/Classes/AnimGraphNode_BoneDrivenController.h
--- a/Engine/Source/Editor/AnimGraph/Classes/AnimGraphNode_BoneDrivenController.h
+++ b/Engine/Source/Editor/AnimGraph/Classes/AnimGraphNode_BoneDrivenController.h
@@ -44,6 +44,7 @@ protected:
 
 	// UAnimGraphNode_SkeletalControlBase protected interface
 	virtual FText GetControllerDescription() const override;
+	virtual const FAnimNode_SkeletalControlBase* GetNode() const override { return &Node; }
 	// End of UAnimGraphNode_SkeletalControlBase protected interface
 
 	// Should non-curve mapping values be shown (multiplier, range)?
diff --git a/Engine/Source/Editor/AnimGraph/Classes/AnimGraphNode_CopyBone.h b/Engine/Source/Editor/AnimGraph/Classes/AnimGraphNode_CopyBone.h
--- a/Engine/Source/Editor/AnimGraph/Classes/AnimGraphNode_CopyBone.h
+++ b/Engine/Source/Editor/AnimGraph/Classes/AnimGraphNode_CopyBone.h
@@ -23,6 +23,7 @@ public:
 protected:
 	// UAnimGraphNode_SkeletalControlBase interface
 	virtual FText GetControllerDescription() const override;
+	virtual const FAnimNode_SkeletalControlBase* GetNode() const override { return &Node; }
 	// End of UAnimGraphNode_SkeletalControlBase interface
 
 private:
diff --git a/Engine/Source/Editor/AnimGraph/Classes/AnimGraphNode_SkeletalControlBase.h b/Engine/Source/Editor/AnimGraph/Classes/AnimGraphNode_SkeletalControlBase.h
--- a/Engine/Source/Editor/AnimGraph/Classes/AnimGraphNode_SkeletalControlBase.h
+++ b/Engine/Source/Editor/AnimGraph/Classes/AnimGraphNode_SkeletalControlBase.h
@@ -27,6 +27,7 @@ public:
 	// UAnimGraphNode_Base interface
 	virtual FString GetNodeCategory() const override;
 	virtual void CreateOutputPins() override;
+	virtual void ValidateAnimNodePostCompile(FCompilerResultsLog& MessageLog, UAnimBlueprintGeneratedClass* CompiledClass, int32 CompiledNodeIndex) override;
 	// End of UAnimGraphNode_Base interface
 
 	// Draw function for supporting visualization
@@ -130,4 +131,6 @@ protected:
 		}
 		CompileNodeValue = InValue;
 	}
+
+	virtual const FAnimNode_SkeletalControlBase* GetNode() const PURE_VIRTUAL(UAnimGraphNode_SkeletalControlBase::GetNode, return nullptr;);
 };
diff --git a/Engine/Source/Editor/AnimGraph/Classes/AnimGraphNode_TwoBoneIK.h b/Engine/Source/Editor/AnimGraph/Classes/AnimGraphNode_TwoBoneIK.h
--- a/Engine/Source/Editor/AnimGraph/Classes/AnimGraphNode_TwoBoneIK.h
+++ b/Engine/Source/Editor/AnimGraph/Classes/AnimGraphNode_TwoBoneIK.h
@@ -47,6 +47,8 @@ public:
 	ANIMGRAPH_API virtual void DoTranslation(const USkeletalMeshComponent* SkelComp, FVector& Drag, FAnimNode_Base* InOutAnimNode) override;
 	ANIMGRAPH_API virtual void	CopyNodeDataTo(FAnimNode_Base* AnimNode) override;
 	ANIMGRAPH_API virtual void	DeselectActor(USkeletalMeshComponent* SkelComp) override;
+
+	virtual const FAnimNode_SkeletalControlBase* GetNode() const override { return &Node; }
 	// End of UAnimGraphNode_SkeletalControlBase interface
 
 	enum BoneSelectModeType
diff --git a/Engine/Source/Editor/AnimGraph/Private/AnimGraphNode_ApplyAdditive.cpp b/Engine/Source/Editor/AnimGraph/Private/AnimGraphNode_ApplyAdditive.cpp
--- a/Engine/Source/Editor/AnimGraph/Private/AnimGraphNode_ApplyAdditive.cpp
+++ b/Engine/Source/Editor/AnimGraph/Private/AnimGraphNode_ApplyAdditive.cpp
@@ -5,6 +5,8 @@
 #include "GraphEditorActions.h"
 #include "ScopedTransaction.h"
 #include "AnimGraphNode_ApplyAdditive.h"
+#include "Animation/AnimationSettings.h"
+#include "CompilerResultsLog.h"
 
 /////////////////////////////////////////////////////
 // UAnimGraphNode_ApplyAdditive
@@ -37,4 +39,14 @@ FString UAnimGraphNode_ApplyAdditive::GetNodeCategory() const
 	//@TODO: TEXT("Apply additive to normal pose"), TEXT("Apply additive pose"));
 }
 
+void UAnimGraphNode_ApplyAdditive::ValidateAnimNodeDuringCompilation(class USkeleton* ForSkeleton, class FCompilerResultsLog& MessageLog)
+{
+	if (UAnimationSettings::Get()->bEnablePerformanceLog)
+	{
+		if (Node.LODThreshold < 0)
+		{
+			MessageLog.Warning(TEXT("@@ contains no LOD Threshold."), this);
+		}
+	}
+}
 #undef LOCTEXT_NAMESPACE
diff --git a/Engine/Source/Editor/AnimGraph/Private/AnimGraphNode_ApplyMeshSpaceAdditive.cpp b/Engine/Source/Editor/AnimGraph/Private/AnimGraphNode_ApplyMeshSpaceAdditive.cpp
--- a/Engine/Source/Editor/AnimGraph/Private/AnimGraphNode_ApplyMeshSpaceAdditive.cpp
+++ b/Engine/Source/Editor/AnimGraph/Private/AnimGraphNode_ApplyMeshSpaceAdditive.cpp
@@ -4,6 +4,8 @@
 
 #include "GraphEditorActions.h"
 #include "ScopedTransaction.h"
+#include "Animation/AnimationSettings.h"
+#include "CompilerResultsLog.h"
 #include "AnimGraphNode_ApplyMeshSpaceAdditive.h"
 
 /////////////////////////////////////////////////////
@@ -37,4 +39,14 @@ FString UAnimGraphNode_ApplyMeshSpaceAdditive::GetNodeCategory() const
 	//@TODO: TEXT("Apply additive to normal pose"), TEXT("Apply additive pose"));
 }
 
+void UAnimGraphNode_ApplyMeshSpaceAdditive::ValidateAnimNodeDuringCompilation(class USkeleton* ForSkeleton, class FCompilerResultsLog& MessageLog)
+{
+	if (UAnimationSettings::Get()->bEnablePerformanceLog)
+	{
+		if (Node.LODThreshold < 0)
+		{
+			MessageLog.Warning(TEXT("@@ contains no LOD Threshold."), this);
+		}
+	}
+}
 #undef LOCTEXT_NAMESPACE
diff --git a/Engine/Source/Editor/AnimGraph/Private/AnimGraphNode_RotationOffsetBlendSpace.cpp b/Engine/Source/Editor/AnimGraph/Private/AnimGraphNode_RotationOffsetBlendSpace.cpp
--- a/Engine/Source/Editor/AnimGraph/Private/AnimGraphNode_RotationOffsetBlendSpace.cpp
+++ b/Engine/Source/Editor/AnimGraph/Private/AnimGraphNode_RotationOffsetBlendSpace.cpp
@@ -6,6 +6,7 @@
 #include "CompilerResultsLog.h"
 #include "BlueprintNodeSpawner.h"
 #include "BlueprintActionDatabaseRegistrar.h"
+#include "Animation/AnimationSettings.h"
 
 /////////////////////////////////////////////////////
 // UAnimGraphNode_RotationOffsetBlendSpace
@@ -139,6 +140,14 @@ void UAnimGraphNode_RotationOffsetBlendSpace::ValidateAnimNodeDuringCompilation(
 			MessageLog.Error(TEXT("@@ references blendspace that uses different skeleton @@"), this, BlendSpaceSkeleton);
 		}
 	}
+
+	if (UAnimationSettings::Get()->bEnablePerformanceLog)
+	{
+		if (Node.LODThreshold < 0)
+		{
+			MessageLog.Warning(TEXT("@@ contains no LOD Threshold."), this);
+		}
+	}
 }
 
 void UAnimGraphNode_RotationOffsetBlendSpace::GetContextMenuActions(const FGraphNodeContextMenuBuilder& Context) const
@@ -166,4 +175,5 @@ void UAnimGraphNode_RotationOffsetBlendSpace::ReplaceReferredAnimations(const TM
 {
 	HandleAnimReferenceReplacement(Node.BlendSpace, ComplexAnimsMap, AnimSequenceMap);
 }
+
 #undef LOCTEXT_NAMESPACE
\ No newline at end of file
diff --git a/Engine/Source/Editor/AnimGraph/Private/AnimGraphNode_SkeletalControlBase.cpp b/Engine/Source/Editor/AnimGraph/Private/AnimGraphNode_SkeletalControlBase.cpp
--- a/Engine/Source/Editor/AnimGraph/Private/AnimGraphNode_SkeletalControlBase.cpp
+++ b/Engine/Source/Editor/AnimGraph/Private/AnimGraphNode_SkeletalControlBase.cpp
@@ -3,7 +3,9 @@
 #include "AnimGraphPrivatePCH.h"
 #include "AnimGraphNode_SkeletalControlBase.h"
 #include "AnimationGraphSchema.h"
-
+#include "Animation/AnimationSettings.h"
+#include "CompilerResultsLog.h"
+#include "AnimNode_SkeletalControlBase.h"
 /////////////////////////////////////////////////////
 // UAnimGraphNode_SkeletalControlBase
 
@@ -348,4 +350,15 @@ bool UAnimGraphNode_SkeletalControlBase::IsPinShown(const FString& PinName) cons
 	return false;
 }
 
+void UAnimGraphNode_SkeletalControlBase::ValidateAnimNodePostCompile(FCompilerResultsLog& MessageLog, UAnimBlueprintGeneratedClass* CompiledClass, int32 CompiledNodeIndex)
+{
+	if (UAnimationSettings::Get()->bEnablePerformanceLog)
+	{
+		const FAnimNode_SkeletalControlBase* Node = GetNode();
+		if (Node && Node->LODThreshold < 0)
+		{
+			MessageLog.Warning(TEXT("@@ contains no LOD Threshold."), this);
+		}
+	}
+}
 #undef LOCTEXT_NAMESPACE
\ No newline at end of file
diff --git a/Engine/Source/Editor/EditorStyle/Private/SlateEditorStyle.cpp b/Engine/Source/Editor/EditorStyle/Private/SlateEditorStyle.cpp
--- a/Engine/Source/Editor/EditorStyle/Private/SlateEditorStyle.cpp
+++ b/Engine/Source/Editor/EditorStyle/Private/SlateEditorStyle.cpp
@@ -1634,6 +1634,7 @@ void FSlateEditorStyle::FStyle::SetupGeneralStyles()
 		Set( "Sequencer.Time.Small", new IMAGE_BRUSH( "Icons/SequencerIcons/icon_Sequencer_Time_24x", Icon20x20 ) );
 		Set( "Sequencer.Value", new IMAGE_BRUSH( "Icons/SequencerIcons/icon_Sequencer_Value_24x", Icon24x24 ) );
 		Set( "Sequencer.Value.Small", new IMAGE_BRUSH( "Icons/SequencerIcons/icon_Sequencer_Value_24x", Icon20x20 ) );
+		Set( "Sequencer.Tracks.Audio", new IMAGE_BRUSH("Icons/AssetIcons/SoundCue_16x", Icon16x16));
 		Set( "Sequencer.Tracks.Event", new IMAGE_BRUSH("Icons/SequencerIcons/icon_Sequencer_Event_16x", Icon16x16));
 		Set( "Sequencer.Tracks.Fade", new IMAGE_BRUSH("Icons/SequencerIcons/icon_Sequencer_Fade_16x", Icon16x16));
 		Set( "Sequencer.Tracks.Shot", new IMAGE_BRUSH("Icons/SequencerIcons/icon_Sequencer_Shot_16x", Icon16x16));
diff --git a/Engine/Source/Editor/MovieSceneTools/MovieSceneTools.Build.cs b/Engine/Source/Editor/MovieSceneTools/MovieSceneTools.Build.cs
--- a/Engine/Source/Editor/MovieSceneTools/MovieSceneTools.Build.cs
+++ b/Engine/Source/Editor/MovieSceneTools/MovieSceneTools.Build.cs
@@ -30,6 +30,7 @@ public class MovieSceneTools : ModuleRules
 
 		PrivateDependencyModuleNames.AddRange(
 			new string[] {
+                "ActorPickerMode",
 				"MovieScene",
 				"MovieSceneTracks",
 				"BlueprintGraph",
@@ -46,6 +47,7 @@ public class MovieSceneTools : ModuleRules
 			new string[] {
 				"AssetTools",
 				"Sequencer",
+				"SceneOutliner",
 				"PropertyEditor"
 			}
 		);
@@ -53,8 +55,11 @@ public class MovieSceneTools : ModuleRules
 		DynamicallyLoadedModuleNames.AddRange(
 			new string[] {
 				"AssetTools",
+				"SceneOutliner",
 				"PropertyEditor"
 			}
 		);
+
+        bFasterWithoutUnity = true;
 	}
 }
diff --git a/Engine/Source/Editor/MovieSceneTools/Private/Sections/BoolPropertySection.h b/Engine/Source/Editor/MovieSceneTools/Private/Sections/BoolPropertySection.h
--- a/Engine/Source/Editor/MovieSceneTools/Private/Sections/BoolPropertySection.h
+++ b/Engine/Source/Editor/MovieSceneTools/Private/Sections/BoolPropertySection.h
@@ -8,7 +8,7 @@
 class FBoolPropertySection : public FPropertySection
 {
 public:
-	FBoolPropertySection( UMovieSceneSection& InSectionObject, FName SectionName, ISequencer* InSequencer )
+	FBoolPropertySection( UMovieSceneSection& InSectionObject, const FText& SectionName, ISequencer* InSequencer )
 		: FPropertySection( InSectionObject, SectionName )
 		, Sequencer( InSequencer ) {}
 
diff --git a/Engine/Source/Editor/MovieSceneTools/Private/Sections/BytePropertySection.h b/Engine/Source/Editor/MovieSceneTools/Private/Sections/BytePropertySection.h
--- a/Engine/Source/Editor/MovieSceneTools/Private/Sections/BytePropertySection.h
+++ b/Engine/Source/Editor/MovieSceneTools/Private/Sections/BytePropertySection.h
@@ -2,22 +2,29 @@
 
 #pragma once
 
+
 /**
-* An implementation of byte property sections
-*/
-class FBytePropertySection : public FPropertySection
+ * An implementation of byte property sections
+ */
+class FBytePropertySection
+	: public FPropertySection
 {
 public:
-	FBytePropertySection( UMovieSceneSection& InSectionObject, FName InSectionName, UEnum* InEnum )
-		: FPropertySection( InSectionObject, InSectionName )
-		, Enum( InEnum ) {}
+	FBytePropertySection(UMovieSceneSection& InSectionObject, const FText& InSectionName, UEnum* InEnum)
+		: FPropertySection(InSectionObject, InSectionName)
+		, Enum( InEnum )
+	{ }
+
+public:
 
 	// FPropertySection interface
-	virtual void GenerateSectionLayout( class ISectionLayoutBuilder& LayoutBuilder ) const override;
-	virtual void SetIntermediateValue( FPropertyChangedParams PropertyChangedParams ) override;
+
+	virtual void GenerateSectionLayout(class ISectionLayoutBuilder& LayoutBuilder) const override;
+	virtual void SetIntermediateValue(FPropertyChangedParams PropertyChangedParams) override;
 	virtual void ClearIntermediateValue() override;
 
 private:
+
 	mutable TSharedPtr<FByteKeyArea> KeyArea;
 	UEnum* Enum;
 };
\ No newline at end of file
diff --git a/Engine/Source/Editor/MovieSceneTools/Private/Sections/ColorPropertySection.h b/Engine/Source/Editor/MovieSceneTools/Private/Sections/ColorPropertySection.h
--- a/Engine/Source/Editor/MovieSceneTools/Private/Sections/ColorPropertySection.h
+++ b/Engine/Source/Editor/MovieSceneTools/Private/Sections/ColorPropertySection.h
@@ -21,7 +21,7 @@ public:
 	 * @param InTrack The track that owns the section.
 	 */
 	FColorPropertySection(UMovieSceneSection& InSectionObject, ISequencer* InSequencer, UMovieSceneTrack& InTrack)
-		: FPropertySection(InSectionObject, InTrack.GetTrackName())
+		: FPropertySection(InSectionObject, InTrack.GetDisplayName())
 		, Sequencer(InSequencer)
 		, Track(*Cast<UMovieSceneColorTrack>(&InTrack))
 	{ }
diff --git a/Engine/Source/Editor/MovieSceneTools/Private/Sections/FloatPropertySection.h b/Engine/Source/Editor/MovieSceneTools/Private/Sections/FloatPropertySection.h
--- a/Engine/Source/Editor/MovieSceneTools/Private/Sections/FloatPropertySection.h
+++ b/Engine/Source/Editor/MovieSceneTools/Private/Sections/FloatPropertySection.h
@@ -2,21 +2,28 @@
 
 #pragma once
 
+
 /**
-* An implementation of float property sections
-*/
-class FFloatPropertySection : public FPropertySection
+ * An implementation of float property sections
+ */
+class FFloatPropertySection
+	: public FPropertySection
 {
 public:
-	FFloatPropertySection( UMovieSceneSection& InSectionObject, FName SectionName )
-		: FPropertySection( InSectionObject, SectionName ) {}
 
-	// FPropertyInterface 
+	FFloatPropertySection(UMovieSceneSection& InSectionObject, const FText& SectionName)
+		: FPropertySection(InSectionObject, SectionName)
+	{ }
+
+public:
 
-	virtual void GenerateSectionLayout( class ISectionLayoutBuilder& LayoutBuilder ) const override;
-	virtual void SetIntermediateValue( FPropertyChangedParams PropertyChangedParams ) override;
+	// FPropertySection interface 
+
+	virtual void GenerateSectionLayout(class ISectionLayoutBuilder& LayoutBuilder) const override;
+	virtual void SetIntermediateValue(FPropertyChangedParams PropertyChangedParams) override;
 	virtual void ClearIntermediateValue() override;
 
 private:
+
 	mutable TSharedPtr<FFloatCurveKeyArea> KeyArea;
 };
\ No newline at end of file
diff --git a/Engine/Source/Editor/MovieSceneTools/Private/Sections/MaterialParameterSection.h b/Engine/Source/Editor/MovieSceneTools/Private/Sections/MaterialParameterSection.h
--- a/Engine/Source/Editor/MovieSceneTools/Private/Sections/MaterialParameterSection.h
+++ b/Engine/Source/Editor/MovieSceneTools/Private/Sections/MaterialParameterSection.h
@@ -2,21 +2,30 @@
 
 #pragma once
 
+
 /**
-* A movie scene section for material parameters.
-*/
-class FMaterialParameterSection : public FPropertySection
+ * A movie scene section for material parameters.
+ */
+class FMaterialParameterSection
+	: public FPropertySection
 {
 public:
-	FMaterialParameterSection( UMovieSceneSection& InSectionObject, FName SectionName )
-		: FPropertySection( InSectionObject, SectionName ) {}
 
-	// ISequencerSection interface.
-	virtual void GenerateSectionLayout( class ISectionLayoutBuilder& LayoutBuilder ) const override;
-	virtual bool RequestDeleteCategory( const TArray<FName>& CategoryNamePath ) override;
-	virtual bool RequestDeleteKeyArea( const TArray<FName>& KeyAreaNamePath ) override;
+	FMaterialParameterSection(UMovieSceneSection& InSectionObject, const FText& SectionName)
+		: FPropertySection(InSectionObject, SectionName)
+	{ }
+
+public:
 
-	// FPropertySection interface.
+	// FPropertySection interface
 	virtual void SetIntermediateValue(FPropertyChangedParams PropertyChangedParams) override { }
 	virtual void ClearIntermediateValue() override { }
+
+public:
+
+	// ISequencerSection interface
+
+	virtual void GenerateSectionLayout(class ISectionLayoutBuilder& LayoutBuilder) const override;
+	virtual bool RequestDeleteCategory(const TArray<FName>& CategoryNamePath) override;
+	virtual bool RequestDeleteKeyArea(const TArray<FName>& KeyAreaNamePath) override;
 };
\ No newline at end of file
diff --git a/Engine/Source/Editor/MovieSceneTools/Private/Sections/VectorPropertySection.h b/Engine/Source/Editor/MovieSceneTools/Private/Sections/VectorPropertySection.h
--- a/Engine/Source/Editor/MovieSceneTools/Private/Sections/VectorPropertySection.h
+++ b/Engine/Source/Editor/MovieSceneTools/Private/Sections/VectorPropertySection.h
@@ -2,14 +2,19 @@
 
 #pragma once
 
+
 /**
-* An implementation of vector property sections
-*/
+ * An implementation of vector property sections.
+ */
 class FVectorPropertySection : public FPropertySection
 {
 public:
-	FVectorPropertySection( UMovieSceneSection& InSectionObject, FName SectionName )
-		: FPropertySection( InSectionObject, SectionName ) {}
+
+	FVectorPropertySection(UMovieSceneSection& InSectionObject, const FText& SectionName)
+		: FPropertySection( InSectionObject, SectionName )
+	{ }
+
+public:
 
 	// FPropertySectionInterface
 
@@ -18,8 +23,9 @@ public:
 	virtual void ClearIntermediateValue() override;
 
 private:
+
 	mutable TSharedPtr<FFloatCurveKeyArea> XKeyArea;
 	mutable TSharedPtr<FFloatCurveKeyArea> YKeyArea;
 	mutable TSharedPtr<FFloatCurveKeyArea> ZKeyArea;
 	mutable TSharedPtr<FFloatCurveKeyArea> WKeyArea;
-};
\ No newline at end of file
+};
diff --git a/Engine/Source/Editor/MovieSceneTools/Private/Sections/VisibilityPropertySection.h b/Engine/Source/Editor/MovieSceneTools/Private/Sections/VisibilityPropertySection.h
--- a/Engine/Source/Editor/MovieSceneTools/Private/Sections/VisibilityPropertySection.h
+++ b/Engine/Source/Editor/MovieSceneTools/Private/Sections/VisibilityPropertySection.h
@@ -10,7 +10,7 @@
 class FVisibilityPropertySection : public FBoolPropertySection
 {
 public:
-	FVisibilityPropertySection( UMovieSceneSection& InSectionObject, FName SectionName, ISequencer* InSequencer )
+	FVisibilityPropertySection( UMovieSceneSection& InSectionObject, const FText& SectionName, ISequencer* InSequencer )
 		: FBoolPropertySection( InSectionObject, SectionName, InSequencer )
 	{
 		DisplayName = FText::FromString( TEXT( "Visible" ) );
diff --git a/Engine/Source/Editor/MovieSceneTools/Private/TrackEditors/ActorPickerTrackEditor.cpp b/Engine/Source/Editor/MovieSceneTools/Private/TrackEditors/ActorPickerTrackEditor.cpp
new file mode 100644
--- /dev/null
+++ b/Engine/Source/Editor/MovieSceneTools/Private/TrackEditors/ActorPickerTrackEditor.cpp
@@ -0,0 +1,147 @@
+// Copyright 1998-2015 Epic Games, Inc. All Rights Reserved.
+
+#include "MovieSceneToolsPrivatePCH.h"
+#include "ActorPickerTrackEditor.h"
+#include "ActorPickerMode.h"
+#include "Editor/SceneOutliner/Public/SceneOutliner.h"
+#include "Editor/SceneOutliner/Private/SSocketChooser.h"
+#include "Editor/LevelEditor/Public/LevelEditor.h"
+#include "SceneOutlinerPublicTypes.h"
+
+#define LOCTEXT_NAMESPACE "FActorPickerTrackEditor"
+
+FActorPickerTrackEditor::FActorPickerTrackEditor(TSharedRef<ISequencer> InSequencer) 
+	: FMovieSceneTrackEditor( InSequencer ) 
+
+{
+}
+
+void FActorPickerTrackEditor::PickActorInteractive(FGuid ObjectBinding, UMovieSceneSection* Section)
+{
+	if(GUnrealEd->GetSelectedActorCount())
+	{
+		FActorPickerModeModule& ActorPickerMode = FModuleManager::Get().GetModuleChecked<FActorPickerModeModule>("ActorPickerMode");
+
+		ActorPickerMode.BeginActorPickingMode(
+			FOnGetAllowedClasses(), 
+			FOnShouldFilterActor::CreateSP(this, &FActorPickerTrackEditor::IsActorPickable), 
+			FOnActorSelected::CreateSP(this, &FActorPickerTrackEditor::ActorPicked, ObjectBinding, Section)
+			);
+	}
+}
+
+void FActorPickerTrackEditor::ShowActorSubMenu(FMenuBuilder& MenuBuilder, FGuid ObjectBinding, UMovieSceneSection* Section)
+{
+	struct Local
+	{
+		static FReply OnInteractiveActorPickerClicked(FActorPickerTrackEditor* ActorPickerTrackEditor, FGuid ObjectBinding, UMovieSceneSection* Section)
+		{
+			FSlateApplication::Get().DismissAllMenus();
+			ActorPickerTrackEditor->PickActorInteractive(ObjectBinding, Section);
+			return FReply::Handled();
+		}
+	};
+
+	using namespace SceneOutliner;
+
+	SceneOutliner::FInitializationOptions InitOptions;
+	{
+		InitOptions.Mode = ESceneOutlinerMode::ActorPicker;			
+		InitOptions.bShowHeaderRow = false;
+		InitOptions.bFocusSearchBoxWhenOpened = true;
+		InitOptions.bShowTransient = true;
+		InitOptions.bShowCreateNewFolder = false;
+		// Only want the actor label column
+		InitOptions.ColumnMap.Add(FBuiltInColumnTypes::Label(), FColumnInfo(EColumnVisibility::Visible, 0));
+
+		// Only display Actors that we can attach too
+		InitOptions.Filters->AddFilterPredicate( SceneOutliner::FActorFilterPredicate::CreateSP(this, &FActorPickerTrackEditor::IsActorPickable) );
+	}		
+
+	// Actor selector to allow the user to choose a parent actor
+	FSceneOutlinerModule& SceneOutlinerModule = FModuleManager::LoadModuleChecked<FSceneOutlinerModule>( "SceneOutliner" );
+
+	TSharedRef< SWidget > MenuWidget = 
+		SNew(SHorizontalBox)
+
+		+SHorizontalBox::Slot()
+		.AutoWidth()
+		[
+			SNew(SVerticalBox)
+			+SVerticalBox::Slot()
+			.MaxHeight(400.0f)
+			[
+				SceneOutlinerModule.CreateSceneOutliner(
+					InitOptions,
+					FOnActorPicked::CreateSP(this, &FActorPickerTrackEditor::ActorPicked, ObjectBinding, Section )
+					)
+			]
+		]
+	
+		+SHorizontalBox::Slot()
+		.VAlign(VAlign_Top)
+		.AutoWidth()
+		[
+			SNew(SVerticalBox)
+
+			+SVerticalBox::Slot()
+			.AutoHeight()
+			.Padding(4.0f, 0.0f, 0.0f, 0.0f)
+			[
+				SNew(SButton)
+				.ToolTipText( LOCTEXT( "PickButtonLabel", "Pick a parent actor to attach to") )
+				.ButtonStyle(FEditorStyle::Get(), "HoverHintOnly")
+				.OnClicked(FOnClicked::CreateStatic(&Local::OnInteractiveActorPickerClicked, this, ObjectBinding, Section))
+				.ContentPadding(4.0f)
+				.ForegroundColor(FSlateColor::UseForeground())
+				.IsFocusable(false)
+				[
+					SNew(SImage)
+					.Image(FEditorStyle::GetBrush("PropertyWindow.Button_PickActorInteractive"))
+					.ColorAndOpacity(FSlateColor::UseForeground())
+				]
+			]
+		];
+
+	MenuBuilder.AddWidget(MenuWidget, FText::GetEmpty(), false);
+}
+
+void FActorPickerTrackEditor::ActorPicked(AActor* ParentActor, FGuid ObjectGuid, UMovieSceneSection* Section)
+{
+	USceneComponent* ComponentWithSockets = NULL;
+	if (ParentActor != NULL)
+	{
+		if (USceneComponent* RootComponent = Cast<USceneComponent>(ParentActor->GetRootComponent()))
+		{
+			if (RootComponent->HasAnySockets())
+			{
+				ComponentWithSockets = RootComponent;
+			}
+		}
+	}
+
+	// Show socket chooser if we have sockets to select
+	if (ComponentWithSockets != NULL)
+	{
+		FLevelEditorModule& LevelEditorModule = FModuleManager::GetModuleChecked<FLevelEditorModule>( "LevelEditor");
+		TSharedPtr< ILevelEditor > LevelEditor = LevelEditorModule.GetFirstLevelEditor();
+
+		// Create as context menu
+		FSlateApplication::Get().PushMenu(
+			LevelEditor.ToSharedRef(),
+			FWidgetPath(),
+			SNew(SSocketChooserPopup)
+			.SceneComponent( ComponentWithSockets )
+			.OnSocketChosen( this, &FActorPickerTrackEditor::ActorSocketPicked, ParentActor, ObjectGuid, Section ),
+			FSlateApplication::Get().GetCursorPos(),
+			FPopupTransitionEffect( FPopupTransitionEffect::ContextMenu )
+			);
+	}
+	else
+	{
+		FSlateApplication::Get().DismissAllMenus();
+		ActorSocketPicked( NAME_None, ParentActor, ObjectGuid, Section );
+	}
+}
+
+#undef LOCTEXT_NAMESPACE
diff --git a/Engine/Source/Editor/MovieSceneTools/Private/TrackEditors/ActorPickerTrackEditor.h b/Engine/Source/Editor/MovieSceneTools/Private/TrackEditors/ActorPickerTrackEditor.h
new file mode 100644
--- /dev/null
+++ b/Engine/Source/Editor/MovieSceneTools/Private/TrackEditors/ActorPickerTrackEditor.h
@@ -0,0 +1,31 @@
+// Copyright 1998-2015 Epic Games, Inc. All Rights Reserved.
+
+#pragma once
+
+/**
+ * Track editor actor picker
+ */
+class FActorPickerTrackEditor : public FMovieSceneTrackEditor
+{
+public:
+
+	FActorPickerTrackEditor(TSharedRef<ISequencer> InSequencer);
+
+	/** Is this actor pickable? */
+	virtual bool IsActorPickable( const AActor* const ParentActor ) { return false; }
+
+	/** Actor socked was picked */
+	virtual void ActorSocketPicked(const FName SocketName, AActor* ParentActor, FGuid ObjectBinding, UMovieSceneSection* Section) {}
+
+	/** Show a sub menu of the pickable actors */
+	void ShowActorSubMenu(FMenuBuilder& MenuBuilder, FGuid ObjectBinding, UMovieSceneSection* Section);
+
+	/** Actor was picked */
+	void ActorPicked(AActor* ParentActor, FGuid ObjectBinding, UMovieSceneSection* Section);
+
+private: 
+
+	/** Interactively pick an actor from the viewport */
+	void PickActorInteractive(FGuid ObjectBinding, UMovieSceneSection* Section);
+
+};
\ No newline at end of file
diff --git a/Engine/Source/Editor/MovieSceneTools/Private/TrackEditors/AttachTrackEditor.cpp b/Engine/Source/Editor/MovieSceneTools/Private/TrackEditors/AttachTrackEditor.cpp
--- a/Engine/Source/Editor/MovieSceneTools/Private/TrackEditors/AttachTrackEditor.cpp
+++ b/Engine/Source/Editor/MovieSceneTools/Private/TrackEditors/AttachTrackEditor.cpp
@@ -15,43 +15,10 @@
 #include "ActorEditorUtils.h"
 #include "GameFramework/WorldSettings.h"
 #include "AttachTrackEditor.h"
+#include "ActorPickerTrackEditor.h"
 
 
-void GetAttachActors(TSet<AActor*>& SelectedActors, TSet<AActor*>& UnselectedActors, TArray<UObject*> InObjects)
-{
-	if (GEditor->GetEditorWorldContext().World())
-	{
-		for( FActorIterator	ActorIt(GEditor->GetEditorWorldContext().World()); ActorIt; ++ActorIt )
-		{
-			AActor* Actor = *ActorIt;
-
-			if (Actor->IsListedInSceneOutliner() &&
-				!FActorEditorUtils::IsABuilderBrush(Actor) &&
-				!Actor->IsA( AWorldSettings::StaticClass() ) &&
-				!Actor->IsPendingKill() &&
-				!InObjects.Contains(Actor))
-			{			
-				if (GEditor->GetSelectedActors()->IsSelected(Actor))
-				{
-					SelectedActors.Add(Actor);
-				}
-				else
-				{
-					UnselectedActors.Add(Actor);
-				}
-			}
-		}
-	}
-	SelectedActors.Sort([](const AActor& A, const AActor& B)
-	{
-		return A.GetActorLabel() < B.GetActorLabel();
-	});
-	UnselectedActors.Sort([](const AActor& A, const AActor& B)
-	{
-		return A.GetActorLabel() < B.GetActorLabel();
-	});
-}
-
+#define LOCTEXT_NAMESPACE "F3DAttachTrackEditor"
 
 /**
  * Class that draws an attach section in the sequencer
@@ -74,7 +41,7 @@ public:
 
 	virtual FText GetDisplayName() const override
 	{ 
-		return NSLOCTEXT("FAttachSection", "DisplayName", "Attach");
+		return LOCTEXT("SectionDisplayName", "Attach");
 	}
 	
 	virtual FText GetSectionTitle() const override { return FText::GetEmpty(); }
@@ -100,59 +67,11 @@ public:
 	
 	virtual void BuildSectionContextMenu(FMenuBuilder& MenuBuilder) override
 	{
-		TArray<UObject*> OutObjects;
-		//@todo
-
-		TSet<AActor*> SelectedActors;
-		TSet<AActor*> UnselectedActors;
-		GetAttachActors(SelectedActors, UnselectedActors, OutObjects);	
-		
-		if (SelectedActors.Num()+UnselectedActors.Num() == 0)
-		{
-			return;
-		}
-
-		if (SelectedActors.Num())
-		{
-			MenuBuilder.BeginSection(NAME_None, FText::FromString("Selected Actors"));
-
-			for (TSet<AActor*>::TIterator It(SelectedActors); It; ++It)
-			{
-				AActor* Actor = *It;
-				FFormatNamedArguments Args;
-				Args.Add( TEXT("AttachName"), FText::FromString( Actor->GetActorLabel() ) );
-
-				MenuBuilder.AddMenuEntry(
-					FText::Format( NSLOCTEXT("Sequencer", "SetAttach", "Set {AttachName}"), Args ),
-					FText::Format( NSLOCTEXT("Sequencer", "SetAttachTooltip", "Set attach track driven by {AttachName}."), Args ),
-					FSlateIcon(),
-					FUIAction(FExecuteAction::CreateSP(AttachTrackEditor, &F3DAttachTrackEditor::SetAttach, &Section, Actor))
-				);
-			}
-			
-			MenuBuilder.EndSection();
-		}
-
-		if (UnselectedActors.Num())
-		{
-			MenuBuilder.BeginSection(NAME_None, FText::FromString("Actors"));
-
-			for (TSet<AActor*>::TIterator It(UnselectedActors); It; ++It)
-			{
-				AActor* Actor = *It;
-				FFormatNamedArguments Args;
-				Args.Add( TEXT("AttachName"), FText::FromString( Actor->GetActorLabel() ) );
-
-				MenuBuilder.AddMenuEntry(
-					FText::Format( NSLOCTEXT("Sequencer", "SetAttach", "Set {AttachName}"), Args ),
-					FText::Format( NSLOCTEXT("Sequencer", "SetAttachTooltip", "Set attach track driven by {AttachName}."), Args ),
-					FSlateIcon(),
-					FUIAction(FExecuteAction::CreateSP(AttachTrackEditor, &F3DAttachTrackEditor::SetAttach, &Section, Actor))
-				);
-			}
+		FGuid DummyObjectBinding;
 
-			MenuBuilder.EndSection();
-		}
+		MenuBuilder.AddSubMenu(
+			LOCTEXT("SetAttach", "Attach"), LOCTEXT("SetAttachTooltip", "Set attach"),
+			FNewMenuDelegate::CreateRaw(AttachTrackEditor, &FActorPickerTrackEditor::ShowActorSubMenu, DummyObjectBinding, &Section));
 	}
 
 private:
@@ -166,7 +85,7 @@ private:
 
 
 F3DAttachTrackEditor::F3DAttachTrackEditor( TSharedRef<ISequencer> InSequencer )
-	: FMovieSceneTrackEditor( InSequencer ) 
+	: FActorPickerTrackEditor( InSequencer ) 
 {
 }
 
@@ -199,148 +118,80 @@ TSharedRef<ISequencerSection> F3DAttachTrackEditor::MakeSectionInterface( UMovie
 
 void F3DAttachTrackEditor::AddKey( const FGuid& ObjectGuid, UObject* AdditionalAsset )
 {
-	AddAttach(ObjectGuid, AdditionalAsset);
+	AActor* ParentActor = Cast<AActor>(AdditionalAsset);
+	if (ParentActor != nullptr)
+	{
+		UMovieSceneSection* DummySection = nullptr;
+		ActorPicked(ParentActor, ObjectGuid, DummySection);
+	}
 }
 
-
 void F3DAttachTrackEditor::BuildObjectBindingTrackMenu(FMenuBuilder& MenuBuilder, const FGuid& ObjectBinding, const UClass* ObjectClass)
 {
-	// build a menu with all the possible actors, add the selected actors first
 	if (ObjectClass->IsChildOf(AActor::StaticClass()))
 	{
-		TArray<UObject*> OutObjects;
-		GetSequencer()->GetRuntimeObjects( GetSequencer()->GetFocusedMovieSceneSequenceInstance(), ObjectBinding, OutObjects);
+		UMovieSceneSection* DummySection = nullptr;
 
-		TSet<AActor*> SelectedActors;
-		TSet<AActor*> UnselectedActors;
-		GetAttachActors(SelectedActors, UnselectedActors, OutObjects);
-
-		if (SelectedActors.Num()+UnselectedActors.Num() > 0)
-		{
-			MenuBuilder.AddSubMenu(
-				NSLOCTEXT("Sequencer", "AddAttach", "Attach"), NSLOCTEXT("Sequencer", "AddAttachTooltip", "Adds an attach track."),
-				FNewMenuDelegate::CreateRaw(this, &F3DAttachTrackEditor::AddAttachSubMenu, ObjectBinding));
-		}
-		else
-		{
-			AActor* Actor = nullptr;
-			FFormatNamedArguments Args;
-			MenuBuilder.AddMenuEntry(
-				FText::Format( NSLOCTEXT("Sequencer", "AddAttach", "Attach"), Args),
-				FText::Format( NSLOCTEXT("Sequencer", "AddAttachTooltip", "Adds an attach track."), Args ),
-				FSlateIcon(),
-				FUIAction(FExecuteAction::CreateSP(this, &F3DAttachTrackEditor::AddAttach, ObjectBinding, (UObject*)Actor),
-						  FCanExecuteAction::CreateLambda( []{ return false; } ))
-			);
-		}
+		MenuBuilder.AddSubMenu(
+			LOCTEXT("AddAttach", "Attach"), LOCTEXT("AddAttachTooltip", "Adds an attach track."),
+			FNewMenuDelegate::CreateRaw(this, &FActorPickerTrackEditor::ShowActorSubMenu, ObjectBinding, DummySection));
 	}
 }
 
 
-void F3DAttachTrackEditor::AddAttachSubMenu(FMenuBuilder& MenuBuilder, FGuid ObjectBinding)
+bool F3DAttachTrackEditor::IsActorPickable(const AActor* const ParentActor)
 {
-	TArray<UObject*> OutObjects;
-	GetSequencer()->GetRuntimeObjects( GetSequencer()->GetFocusedMovieSceneSequenceInstance(), ObjectBinding, OutObjects);
-
-	TSet<AActor*> SelectedActors;
-	TSet<AActor*> UnselectedActors;
-	GetAttachActors(SelectedActors, UnselectedActors, OutObjects);	
-
-	if (SelectedActors.Num()+UnselectedActors.Num() == 0)
-	{
-		return;
+	if (ParentActor->IsListedInSceneOutliner() &&
+		!FActorEditorUtils::IsABuilderBrush(ParentActor) &&
+		!ParentActor->IsA( AWorldSettings::StaticClass() ) &&
+		!ParentActor->IsPendingKill())
+	{			
+		return true;
 	}
+	return false;
+}
 
-	if (SelectedActors.Num())
+
+void F3DAttachTrackEditor::ActorSocketPicked(const FName SocketName, AActor* ParentActor, FGuid ObjectGuid, UMovieSceneSection* Section)
+{
+	if (ObjectGuid.IsValid())
 	{
-		MenuBuilder.BeginSection(NAME_None, FText::FromString("Selected Actors"));
+		TArray<UObject*> OutObjects;
+		GetSequencer()->GetRuntimeObjects( GetSequencer()->GetFocusedMovieSceneSequenceInstance(), ObjectGuid, OutObjects);
 
-		for (TSet<AActor*>::TIterator It(SelectedActors); It; ++It)
-		{
-			AActor* Actor = *It;
-			FFormatNamedArguments Args;
-			Args.Add( TEXT("AttachName"), FText::FromString( Actor->GetActorLabel() ) );
-
-			MenuBuilder.AddMenuEntry(
-				FText::Format( NSLOCTEXT("Sequencer", "AddAttach", "{AttachName}"), Args ),
-				FText::Format( NSLOCTEXT("Sequencer", "AddAttachTooltip", "Add attach track driven by {AttachName}."), Args ),
-				FSlateIcon(),
-				FUIAction(FExecuteAction::CreateSP(this, &F3DAttachTrackEditor::AddAttach, ObjectBinding, (UObject*)Actor))
-			);
-		}
-			
-		MenuBuilder.EndSection();
+		AnimatablePropertyChanged( FOnKeyProperty::CreateRaw( this, &F3DAttachTrackEditor::AddKeyInternal, OutObjects, SocketName, ParentActor) );
 	}
-
-	if (UnselectedActors.Num())
+	else if (Section != nullptr)
 	{
-		MenuBuilder.BeginSection(NAME_None, FText::FromString("Actors"));
+		const FScopedTransaction Transaction(LOCTEXT("UndoSetAttach", "Set Attach"));
+
+		UMovieScene3DAttachSection* AttachSection = (UMovieScene3DAttachSection*)(Section);
+		FGuid ActorId = FindOrCreateHandleToObject(ParentActor).Handle;
 
-		for (TSet<AActor*>::TIterator It(UnselectedActors); It; ++It)
+		if (ActorId.IsValid())
 		{
-			AActor* Actor = *It;
-			FFormatNamedArguments Args;
-			Args.Add( TEXT("AttachName"), FText::FromString( Actor->GetActorLabel() ) );
-
-			MenuBuilder.AddMenuEntry(
-				FText::Format( NSLOCTEXT("Sequencer", "AddAttach", "{AttachName}"), Args ),
-				FText::Format( NSLOCTEXT("Sequencer", "AddAttachTooltip", "Add attach track driven by {AttachName}."), Args ),
-				FSlateIcon(),
-				FUIAction(FExecuteAction::CreateSP(this, &F3DAttachTrackEditor::AddAttach, ObjectBinding, (UObject*)Actor))
-			);
+			AttachSection->SetConstraintId(ActorId);
+			AttachSection->AttachSocketName = SocketName;
 		}
-
-		MenuBuilder.EndSection();
 	}
 }
 
-
-void F3DAttachTrackEditor::AddAttach(FGuid ObjectGuid, UObject* AdditionalAsset)
-{
-	TArray<UObject*> OutObjects;
-	GetSequencer()->GetRuntimeObjects( GetSequencer()->GetFocusedMovieSceneSequenceInstance(), ObjectGuid, OutObjects);
-
-	AnimatablePropertyChanged( FOnKeyProperty::CreateRaw( this, &F3DAttachTrackEditor::AddKeyInternal, OutObjects, AdditionalAsset) );
-}
-
-
-void F3DAttachTrackEditor::SetAttach(UMovieSceneSection* Section, AActor* Actor)
-{
-	const FScopedTransaction Transaction(NSLOCTEXT("Sequencer", "UndoSetAttach", "Set Attach"));
-
-	UMovieScene3DAttachSection* AttachSection = (UMovieScene3DAttachSection*)(Section);
-	FGuid ActorId = FindOrCreateHandleToObject(Actor).Handle;
-
-	if (ActorId.IsValid())
-	{
-		AttachSection->SetConstraintId(ActorId);
-	}
-}
-
-
-bool F3DAttachTrackEditor::AddKeyInternal( float KeyTime, const TArray<UObject*> Objects, UObject* AdditionalAsset)
+bool F3DAttachTrackEditor::AddKeyInternal( float KeyTime, const TArray<UObject*> Objects, const FName SocketName, AActor* ParentActor)
 {
 	bool bHandleCreated = false;
 	bool bTrackCreated = false;
 	bool bTrackModified = false;
 
-	AActor* Actor = nullptr;
-	
-	if (AdditionalAsset != nullptr)
-	{
-		Actor = Cast<AActor>(AdditionalAsset);
-	}
-
-	FGuid ActorId;
+	FGuid ParentActorId;
 
-	if (Actor != nullptr)
+	if (ParentActor != nullptr)
 	{
-		FFindOrCreateHandleResult HandleResult = FindOrCreateHandleToObject(Actor);
-		ActorId = HandleResult.Handle;
+		FFindOrCreateHandleResult HandleResult = FindOrCreateHandleToObject(ParentActor);
+		ParentActorId = HandleResult.Handle;
 		bHandleCreated |= HandleResult.bWasCreated;
 	}
 
-	if (!ActorId.IsValid())
+	if (!ParentActorId.IsValid())
 	{
 		return false;
 	}
@@ -354,16 +205,15 @@ bool F3DAttachTrackEditor::AddKeyInternal( float KeyTime, const TArray<UObject*>
 		bHandleCreated |= HandleResult.bWasCreated;
 		if (ObjectHandle.IsValid())
 		{
-			MovieSceneHelpers::SetRuntimeObjectMobility(Object);
-
-			FFindOrCreateTrackResult TrackResult = FindOrCreateTrackForObject( ObjectHandle, UMovieScene3DAttachTrack::StaticClass(), FName("Attach"));
+			FFindOrCreateTrackResult TrackResult = FindOrCreateTrackForObject( ObjectHandle, UMovieScene3DAttachTrack::StaticClass());
 			UMovieSceneTrack* Track = TrackResult.Track;
 			bTrackCreated |= TrackResult.bWasCreated;
+
 			if (ensure(Track))
 			{
 				// Clamp to next attach section's start time or the end of the current sequencer view range
 				float AttachEndTime = GetSequencer()->GetViewRange().GetUpperBoundValue();
-	
+
 				for (int32 AttachSectionIndex = 0; AttachSectionIndex < Track->GetAllSections().Num(); ++AttachSectionIndex)
 				{
 					float StartTime = Track->GetAllSections()[AttachSectionIndex]->GetStartTime();
@@ -377,7 +227,7 @@ bool F3DAttachTrackEditor::AddKeyInternal( float KeyTime, const TArray<UObject*>
 					}
 				}
 
-				Cast<UMovieScene3DAttachTrack>(Track)->AddConstraint( KeyTime, AttachEndTime, ActorId );
+				Cast<UMovieScene3DAttachTrack>(Track)->AddConstraint( KeyTime, AttachEndTime, SocketName, ParentActorId );
 				bTrackModified = true;
 			}
 		}
@@ -385,3 +235,6 @@ bool F3DAttachTrackEditor::AddKeyInternal( float KeyTime, const TArray<UObject*>
 
 	return bHandleCreated || bTrackCreated || bTrackModified;
 }
+
+
+#undef LOCTEXT_NAMESPACE
diff --git a/Engine/Source/Editor/MovieSceneTools/Private/TrackEditors/AttachTrackEditor.h b/Engine/Source/Editor/MovieSceneTools/Private/TrackEditors/AttachTrackEditor.h
--- a/Engine/Source/Editor/MovieSceneTools/Private/TrackEditors/AttachTrackEditor.h
+++ b/Engine/Source/Editor/MovieSceneTools/Private/TrackEditors/AttachTrackEditor.h
@@ -4,13 +4,14 @@
 
 #include "MovieScene3DAttachSection.h"
 #include "MovieScene3DAttachTrack.h"
+#include "ActorPickerTrackEditor.h"
 
 
 /**
  * Tools for attaching an object to another object
  */
 class F3DAttachTrackEditor
-	: public FMovieSceneTrackEditor
+	: public FActorPickerTrackEditor
 {
 public:
 
@@ -32,14 +33,6 @@ public:
 	 */
 	static TSharedRef<ISequencerTrackEditor> CreateTrackEditor( TSharedRef<ISequencer> OwningSequencer );
 
-public:
-
-	/** Add attach */
-	void AddAttach(FGuid ObjectGuid, UObject* AdditionalAsset);
-
-	/** Set attach */
-	void SetAttach(UMovieSceneSection* Section, AActor* ActorWithSplineComponent);
-
 public:
 
 	// ISequencerTrackEditor interface
@@ -49,11 +42,12 @@ public:
 	virtual TSharedRef<ISequencerSection> MakeSectionInterface( UMovieSceneSection& SectionObject, UMovieSceneTrack& Track ) override;
 	virtual bool SupportsType( TSubclassOf<UMovieSceneTrack> Type ) const override;
 
-private:
+	// FTrackEditorActorPicker
+	virtual bool IsActorPickable( const AActor* const ParentActor ) override;
+	virtual void ActorSocketPicked(const FName SocketName, AActor* ParentActor, FGuid ObjectBinding, UMovieSceneSection* Section) override;
 
-	/** Add attach sub menu */
-	void AddAttachSubMenu(FMenuBuilder& MenuBuilder, FGuid ObjectBinding);
+private:
 
 	/** Delegate for AnimatablePropertyChanged in AddKey */
-	bool AddKeyInternal(float KeyTime, const TArray<UObject*> Objects, UObject* AdditionalAsset);
+	bool AddKeyInternal(float KeyTime, const TArray<UObject*> Objects, const FName SocketName, AActor* ParentActor);
 };
diff --git a/Engine/Source/Editor/MovieSceneTools/Private/TrackEditors/AudioTrackEditor.cpp b/Engine/Source/Editor/MovieSceneTools/Private/TrackEditors/AudioTrackEditor.cpp
--- a/Engine/Source/Editor/MovieSceneTools/Private/TrackEditors/AudioTrackEditor.cpp
+++ b/Engine/Source/Editor/MovieSceneTools/Private/TrackEditors/AudioTrackEditor.cpp
@@ -29,6 +29,9 @@
 #include "Sound/SoundCue.h"
 
 
+#define LOCTEXT_NAMESPACE "FAudioTrackEditor"
+
+
 namespace AnimatableAudioEditorConstants
 {
 	// @todo Sequencer Allow this to be customizable
@@ -466,6 +469,19 @@ TSharedRef<ISequencerTrackEditor> FAudioTrackEditor::CreateTrackEditor( TSharedR
 }
 
 
+void FAudioTrackEditor::BuildAddTrackMenu(FMenuBuilder& MenuBuilder)
+{
+	MenuBuilder.AddMenuEntry(
+		LOCTEXT("AddTrack", "Audio Track"),
+		LOCTEXT("AddTooltip", "Adds a new master audio track that can play sounds."),
+		FSlateIcon(FEditorStyle::GetStyleSetName(), "Sequencer.Tracks.Audio"),
+		FUIAction(
+			FExecuteAction::CreateRaw(this, &FAudioTrackEditor::HandleAddAudioTrackMenuEntryExecute)
+		)
+	);
+}
+
+
 bool FAudioTrackEditor::SupportsType( TSubclassOf<UMovieSceneTrack> Type ) const
 {
 	return Type == UMovieSceneAudioTrack::StaticClass();
@@ -511,6 +527,8 @@ bool FAudioTrackEditor::AddNewMasterSound( float KeyTime, USoundBase* Sound )
 
 	auto AudioTrack = Cast<UMovieSceneAudioTrack>(Track);
 	AudioTrack->AddNewSound( Sound, KeyTime );
+	AudioTrack->SetDisplayName(LOCTEXT("AudioTrackName", "Audio"));
+
 	return true;
 }
 
@@ -528,14 +546,18 @@ bool FAudioTrackEditor::AddNewAttachedSound( float KeyTime, USoundBase* Sound, T
 		FFindOrCreateHandleResult HandleResult = FindOrCreateHandleToObject( Object );
 		FGuid ObjectHandle = HandleResult.Handle;
 		bHandleCreated |= HandleResult.bWasCreated;
+
 		if (ObjectHandle.IsValid())
 		{
-			FFindOrCreateTrackResult TrackResult = FindOrCreateTrackForObject( ObjectHandle, UMovieSceneAudioTrack::StaticClass(), AudioTrackConstants::UniqueTrackName );
+			FFindOrCreateTrackResult TrackResult = FindOrCreateTrackForObject(ObjectHandle, UMovieSceneAudioTrack::StaticClass());
 			UMovieSceneTrack* Track = TrackResult.Track;
 			bTrackCreated |= TrackResult.bWasCreated;
+
 			if (ensure(Track))
 			{
-				Cast<UMovieSceneAudioTrack>(Track)->AddNewSound( Sound, KeyTime );
+				auto AudioTrack = Cast<UMovieSceneAudioTrack>(Track);
+				AudioTrack->AddNewSound(Sound, KeyTime);
+				AudioTrack->SetDisplayName(LOCTEXT("AudioTrackName", "Audio"));
 				bTrackModified = true;
 			}
 		}
@@ -543,3 +565,30 @@ bool FAudioTrackEditor::AddNewAttachedSound( float KeyTime, USoundBase* Sound, T
 
 	return bHandleCreated || bTrackCreated || bTrackModified;
 }
+
+
+/* FAudioTrackEditor callbacks
+ *****************************************************************************/
+
+void FAudioTrackEditor::HandleAddAudioTrackMenuEntryExecute()
+{
+	UMovieScene* FocusedMovieScene = GetFocusedMovieScene();
+
+	if (FocusedMovieScene == nullptr)
+	{
+		return;
+	}
+
+	const FScopedTransaction Transaction(NSLOCTEXT("Sequencer", "AddAudioTrack_Transaction", "Add Audio Track"));
+	FocusedMovieScene->Modify();
+	
+	auto NewTrack = FocusedMovieScene->AddMasterTrack<UMovieSceneAudioTrack>();
+	ensure(NewTrack);
+
+	NewTrack->SetDisplayName(LOCTEXT("AudioTrackName", "Audio"));
+
+	GetSequencer()->NotifyMovieSceneDataChanged();
+}
+
+
+#undef LOCTEXT_NAMESPACE
diff --git a/Engine/Source/Editor/MovieSceneTools/Private/TrackEditors/AudioTrackEditor.h b/Engine/Source/Editor/MovieSceneTools/Private/TrackEditors/AudioTrackEditor.h
--- a/Engine/Source/Editor/MovieSceneTools/Private/TrackEditors/AudioTrackEditor.h
+++ b/Engine/Source/Editor/MovieSceneTools/Private/TrackEditors/AudioTrackEditor.h
@@ -16,7 +16,7 @@ public:
 	 *
 	 * @param InSequencer The sequencer instance to be used by this tool
 	 */
-	FAudioTrackEditor( TSharedRef<ISequencer> InSequencer );
+	FAudioTrackEditor(TSharedRef<ISequencer> InSequencer);
 
 	/** Virtual destructor. */
 	virtual ~FAudioTrackEditor();
@@ -27,23 +27,29 @@ public:
 	 * @param OwningSequencer The sequencer instance to be used by this tool
 	 * @return The new instance of this class
 	 */
-	static TSharedRef<ISequencerTrackEditor> CreateTrackEditor( TSharedRef<ISequencer> OwningSequencer );
+	static TSharedRef<ISequencerTrackEditor> CreateTrackEditor(TSharedRef<ISequencer> OwningSequencer);
 
 public:
 
 	// ISequencerTrackEditor interface
 
+	virtual void BuildAddTrackMenu(FMenuBuilder& MenuBuilder) override;
 	virtual bool HandleAssetAdded(UObject* Asset, const FGuid& TargetObjectGuid) override;
-	virtual TSharedRef<ISequencerSection> MakeSectionInterface( UMovieSceneSection& SectionObject, UMovieSceneTrack& Track ) override;
-	virtual bool SupportsType( TSubclassOf<UMovieSceneTrack> Type ) const override;
+	virtual TSharedRef<ISequencerSection> MakeSectionInterface(UMovieSceneSection& SectionObject, UMovieSceneTrack& Track) override;
+	virtual bool SupportsType(TSubclassOf<UMovieSceneTrack> Type) const override;
 	
-private:
+protected:
 
 	/** Delegate for AnimatablePropertyChanged in HandleAssetAdded for master sounds */
-	bool AddNewMasterSound( float KeyTime, class USoundBase* Sound );
+	bool AddNewMasterSound(float KeyTime, class USoundBase* Sound);
 
 	/** Delegate for AnimatablePropertyChanged in HandleAssetAdded for attached sounds */
-	bool AddNewAttachedSound( float KeyTime, class USoundBase* Sound, TArray<UObject*> ObjectsToAttachTo );
+	bool AddNewAttachedSound(float KeyTime, class USoundBase* Sound, TArray<UObject*> ObjectsToAttachTo);
+
+private:
+
+	/** Callback for executing the "Add Event Track" menu entry. */
+	void HandleAddAudioTrackMenuEntryExecute();
 };
 
 
@@ -57,7 +63,7 @@ class FAudioSection
 public:
 
 	/** Constructor. */
-	FAudioSection( UMovieSceneSection& InSection, bool bOnAMasterTrack );
+	FAudioSection(UMovieSceneSection& InSection, bool bOnAMasterTrack);
 
 	/** Virtual destructor. */
 	virtual ~FAudioSection();
@@ -70,9 +76,9 @@ public:
 	virtual FText GetDisplayName() const override;
 	virtual FText GetSectionTitle() const override;
 	virtual float GetSectionHeight() const override;
-	virtual void GenerateSectionLayout( class ISectionLayoutBuilder& LayoutBuilder ) const override {}
-	virtual int32 OnPaintSection( const FGeometry& AllottedGeometry, const FSlateRect& SectionClippingRect, FSlateWindowElementList& OutDrawElements, int32 LayerId, bool bParentEnabled ) const override;
-	virtual void Tick( const FGeometry& AllottedGeometry, const FGeometry& ParentGeometry, const double InCurrentTime, const float InDeltaTime ) override;
+	virtual void GenerateSectionLayout( class ISectionLayoutBuilder& LayoutBuilder) const override { }
+	virtual int32 OnPaintSection(const FGeometry& AllottedGeometry, const FSlateRect& SectionClippingRect, FSlateWindowElementList& OutDrawElements, int32 LayerId, bool bParentEnabled) const override;
+	virtual void Tick(const FGeometry& AllottedGeometry, const FGeometry& ParentGeometry, const double InCurrentTime, const float InDeltaTime) override;
 	
 private:
 
diff --git a/Engine/Source/Editor/MovieSceneTools/Private/TrackEditors/EventTrackEditor.cpp b/Engine/Source/Editor/MovieSceneTools/Private/TrackEditors/EventTrackEditor.cpp
--- a/Engine/Source/Editor/MovieSceneTools/Private/TrackEditors/EventTrackEditor.cpp
+++ b/Engine/Source/Editor/MovieSceneTools/Private/TrackEditors/EventTrackEditor.cpp
@@ -67,7 +67,7 @@ void FEventTrackEditor::BuildObjectBindingTrackMenu(FMenuBuilder& MenuBuilder, c
 
 TSharedRef<ISequencerSection> FEventTrackEditor::MakeSectionInterface(UMovieSceneSection& SectionObject, UMovieSceneTrack& Track)
 {
-	return MakeShareable(new FEventTrackSection(SectionObject/*, Track->GetTrackName()*/, GetSequencer()));
+	return MakeShareable(new FEventTrackSection(SectionObject, GetSequencer()));
 }
 
 
@@ -92,13 +92,15 @@ void FEventTrackEditor::HandleAddEventTrackMenuEntryExecute()
 	const FScopedTransaction Transaction(NSLOCTEXT("Sequencer", "AddEventTrack_Transaction", "Add Event Track"));
 	FocusedMovieScene->Modify();
 	
-	UMovieSceneTrack* NewTrack = FocusedMovieScene->AddMasterTrack<UMovieSceneEventTrack>();
+	auto NewTrack = FocusedMovieScene->AddMasterTrack<UMovieSceneEventTrack>();
 	ensure(NewTrack);
 
 	UMovieSceneSection* NewSection = NewTrack->CreateNewSection();
 	ensure(NewSection);
 
 	NewTrack->AddSection(*NewSection);
+	NewTrack->SetDisplayName(LOCTEXT("TrackName", "Events"));
+
 	GetSequencer()->NotifyMovieSceneDataChanged();
 }
 
diff --git a/Engine/Source/Editor/MovieSceneTools/Private/TrackEditors/FadeTrackEditor.cpp b/Engine/Source/Editor/MovieSceneTools/Private/TrackEditors/FadeTrackEditor.cpp
--- a/Engine/Source/Editor/MovieSceneTools/Private/TrackEditors/FadeTrackEditor.cpp
+++ b/Engine/Source/Editor/MovieSceneTools/Private/TrackEditors/FadeTrackEditor.cpp
@@ -74,7 +74,7 @@ void FFadeTrackEditor::HandleAddFadeTrackMenuEntryExecute()
 		return;
 	}
 
-	const FScopedTransaction Transaction(NSLOCTEXT("Sequencer", "AddFadeTrack_Transaction", "Add Fade Track"));
+	const FScopedTransaction Transaction(LOCTEXT("AddFadeTrack_Transaction", "Add Fade Track"));
 
 	MovieScene->Modify();
 
@@ -85,6 +85,7 @@ void FFadeTrackEditor::HandleAddFadeTrackMenuEntryExecute()
 	ensure(NewSection);
 
 	FadeTrack->AddSection(*NewSection);
+
 	GetSequencer()->NotifyMovieSceneDataChanged();
 }
 
diff --git a/Engine/Source/Editor/MovieSceneTools/Private/TrackEditors/MaterialTrackEditor.cpp b/Engine/Source/Editor/MovieSceneTools/Private/TrackEditors/MaterialTrackEditor.cpp
--- a/Engine/Source/Editor/MovieSceneTools/Private/TrackEditors/MaterialTrackEditor.cpp
+++ b/Engine/Source/Editor/MovieSceneTools/Private/TrackEditors/MaterialTrackEditor.cpp
@@ -21,7 +21,7 @@ TSharedRef<ISequencerSection> FMaterialTrackEditor::MakeSectionInterface( UMovie
 	UMovieSceneMaterialParameterSection* ParameterSection = Cast<UMovieSceneMaterialParameterSection>(&SectionObject);
 	checkf( ParameterSection != nullptr, TEXT("Unsupported section type.") );
 
-	return MakeShareable(new FMaterialParameterSection( *ParameterSection, ParameterSection->GetFName()));
+	return MakeShareable(new FMaterialParameterSection( *ParameterSection, FText::FromName(ParameterSection->GetFName())));
 }
 
 
diff --git a/Engine/Source/Editor/MovieSceneTools/Private/TrackEditors/ParticleTrackEditor.cpp b/Engine/Source/Editor/MovieSceneTools/Private/TrackEditors/ParticleTrackEditor.cpp
--- a/Engine/Source/Editor/MovieSceneTools/Private/TrackEditors/ParticleTrackEditor.cpp
+++ b/Engine/Source/Editor/MovieSceneTools/Private/TrackEditors/ParticleTrackEditor.cpp
@@ -31,6 +31,9 @@ namespace AnimatableParticleEditorConstants
 }
 
 
+#define LOCTEXT_NAMESPACE "FParticleTrackEditor"
+
+
 FParticleSection::FParticleSection( UMovieSceneSection& InSection, TSharedRef<ISequencer>InOwningSequencer )
 	: Section( InSection )
 	, OwningSequencer( InOwningSequencer )
@@ -56,7 +59,7 @@ UMovieSceneSection* FParticleSection::GetSectionObject()
 
 FText FParticleSection::GetDisplayName() const
 {
-	return NSLOCTEXT( "FParticleSection", "Emitter", "Emitter" );
+	return LOCTEXT("Emitter", "Emitter");
 }
 
 float FParticleSection::GetSectionHeight() const
@@ -217,13 +220,7 @@ const FSlateBrush* FParticleSection::GetKeyBrush( FKeyHandle KeyHandle ) const
 
 FParticleTrackEditor::FParticleTrackEditor( TSharedRef<ISequencer> InSequencer )
 	: FMovieSceneTrackEditor( InSequencer ) 
-{
-}
-
-
-FParticleTrackEditor::~FParticleTrackEditor()
-{
-}
+{ }
 
 
 TSharedRef<ISequencerTrackEditor> FParticleTrackEditor::CreateTrackEditor( TSharedRef<ISequencer> InSequencer )
@@ -254,22 +251,24 @@ void FParticleTrackEditor::BuildObjectBindingTrackMenu(FMenuBuilder& MenuBuilder
 		const TSharedPtr<ISequencer> ParentSequencer = GetSequencer();
 
 		MenuBuilder.AddMenuEntry(
-			NSLOCTEXT("Sequencer", "AddParticleTrack", "Particle Track"),
-			NSLOCTEXT("Sequencer", "TriggerParticlesTooltip", "Adds a track for controlling particle emitter state."),
+			LOCTEXT("AddParticleTrack", "Particle Track"),
+			LOCTEXT("TriggerParticlesTooltip", "Adds a track for controlling particle emitter state."),
 			FSlateIcon(),
 			FUIAction(FExecuteAction::CreateSP(this, &FParticleTrackEditor::AddParticleKey, ObjectBinding))
-			);
+		);
 	}
 }
 
+
 void FParticleTrackEditor::AddParticleKey( const FGuid ObjectGuid )
 {
 	TArray<UObject*> OutObjects;
-	GetSequencer()->GetRuntimeObjects( GetSequencer()->GetFocusedMovieSceneSequenceInstance(), ObjectGuid, OutObjects );
 
+	GetSequencer()->GetRuntimeObjects( GetSequencer()->GetFocusedMovieSceneSequenceInstance(), ObjectGuid, OutObjects );
 	AnimatablePropertyChanged( FOnKeyProperty::CreateRaw( this, &FParticleTrackEditor::AddKeyInternal, OutObjects ) );
 }
 
+
 bool FParticleTrackEditor::AddKeyInternal( float KeyTime, const TArray<UObject*> Objects )
 {
 	bool bHandleCreated = false;
@@ -282,17 +281,24 @@ bool FParticleTrackEditor::AddKeyInternal( float KeyTime, const TArray<UObject*>
 		FFindOrCreateHandleResult HandleResult = FindOrCreateHandleToObject( Object );
 		FGuid ObjectHandle = HandleResult.Handle;
 		bHandleCreated |= HandleResult.bWasCreated;
+
 		if (ObjectHandle.IsValid())
 		{
-			FFindOrCreateTrackResult TrackResult = FindOrCreateTrackForObject( ObjectHandle, UMovieSceneParticleTrack::StaticClass(), FName( "ParticleSystem" ) );
+			FFindOrCreateTrackResult TrackResult = FindOrCreateTrackForObject(ObjectHandle, UMovieSceneParticleTrack::StaticClass());
 			UMovieSceneTrack* Track = TrackResult.Track;
 			bTrackCreated |= TrackResult.bWasCreated;
-			if ( bTrackCreated && ensure(Track))
+
+			if (bTrackCreated && ensure(Track))
 			{
-				Cast<UMovieSceneParticleTrack>(Track)->AddNewSection( KeyTime );
+				UMovieSceneParticleTrack* ParticleTrack = Cast<UMovieSceneParticleTrack>(Track);
+				ParticleTrack->AddNewSection(KeyTime);
+				ParticleTrack->SetDisplayName(LOCTEXT("TrackName", "Particle System"));
 			}
 		}
 	}
 
 	return bHandleCreated || bTrackCreated;
 }
+
+
+#undef LOCTEXT_NAMESPACE
diff --git a/Engine/Source/Editor/MovieSceneTools/Private/TrackEditors/ParticleTrackEditor.h b/Engine/Source/Editor/MovieSceneTools/Private/TrackEditors/ParticleTrackEditor.h
--- a/Engine/Source/Editor/MovieSceneTools/Private/TrackEditors/ParticleTrackEditor.h
+++ b/Engine/Source/Editor/MovieSceneTools/Private/TrackEditors/ParticleTrackEditor.h
@@ -17,7 +17,9 @@ public:
 	 * @param InSequencer The sequencer instance to be used by this tool.
 	 */
 	FParticleTrackEditor( TSharedRef<ISequencer> InSequencer );
-	~FParticleTrackEditor();
+
+	/** Virtual destructor. */
+	virtual ~FParticleTrackEditor() { }
 
 	/**
 	 * Creates an instance of this class.  Called by a sequencer.
diff --git a/Engine/Source/Editor/MovieSceneTools/Private/TrackEditors/PathTrackEditor.cpp b/Engine/Source/Editor/MovieSceneTools/Private/TrackEditors/PathTrackEditor.cpp
--- a/Engine/Source/Editor/MovieSceneTools/Private/TrackEditors/PathTrackEditor.cpp
+++ b/Engine/Source/Editor/MovieSceneTools/Private/TrackEditors/PathTrackEditor.cpp
@@ -13,33 +13,12 @@
 #include "MovieSceneToolHelpers.h"
 #include "MovieSceneTrackEditor.h"
 #include "PathTrackEditor.h"
+#include "ActorEditorUtils.h"
 #include "Components/SplineComponent.h"
+#include "ActorPickerTrackEditor.h"
 
 
-void GetActorsWithSplineComponents(TSet<AActor*>& SelectedActors, TSet<AActor*>& UnselectedActors)
-{
-	if (GEditor->GetEditorWorldContext().World())
-	{
-		for( FActorIterator	ActorIt(GEditor->GetEditorWorldContext().World()); ActorIt; ++ActorIt )
-		{
-			AActor* Actor = *ActorIt;
-			TArray<USplineComponent*> SplineComponents;
-			Actor->GetComponents(SplineComponents);
-			if (SplineComponents.Num())
-			{
-				if (GEditor->GetSelectedActors()->IsSelected(Actor))
-				{
-					SelectedActors.Add(Actor);
-				}
-				else
-				{
-					UnselectedActors.Add(Actor);
-				}
-			}
-		}
-	}
-}
-
+#define LOCTEXT_NAMESPACE "FPathTrackEditor"
 
 /**
  * Class that draws a path section in the sequencer
@@ -51,8 +30,7 @@ public:
 	F3DPathSection( UMovieSceneSection& InSection, F3DPathTrackEditor* InPathTrackEditor )
 		: Section( InSection )
 		, PathTrackEditor(InPathTrackEditor)
-	{
-	}
+	{ }
 
 	/** ISequencerSection interface */
 	virtual UMovieSceneSection* GetSectionObject() override
@@ -62,7 +40,7 @@ public:
 
 	virtual FText GetDisplayName() const override
 	{ 
-		return NSLOCTEXT("FPathSection", "DisplayName", "Path");
+		return LOCTEXT("DisplayName", "Path");
 	}
 	
 	virtual FText GetSectionTitle() const override { return FText::GetEmpty(); }
@@ -71,7 +49,7 @@ public:
 	{
 		UMovieScene3DPathSection* PathSection = Cast<UMovieScene3DPathSection>( &Section );
 
-		LayoutBuilder.AddKeyArea("Timing", NSLOCTEXT("FPathSection", "TimingArea", "Timing"), MakeShareable( new FFloatCurveKeyArea ( &PathSection->GetTimingCurve( ), PathSection ) ) );
+		LayoutBuilder.AddKeyArea("Timing", LOCTEXT("TimingArea", "Timing"), MakeShareable( new FFloatCurveKeyArea ( &PathSection->GetTimingCurve( ), PathSection ) ) );
 	}
 
 	virtual int32 OnPaintSection( const FGeometry& AllottedGeometry, const FSlateRect& SectionClippingRect, FSlateWindowElementList& OutDrawElements, int32 LayerId, bool bParentEnabled ) const override 
@@ -90,56 +68,11 @@ public:
 	
 	virtual void BuildSectionContextMenu(FMenuBuilder& MenuBuilder) override
 	{
-		TSet<AActor*> SelectedActors;
-		TSet<AActor*> UnselectedActors;
-		GetActorsWithSplineComponents(SelectedActors, UnselectedActors);	
-
-		if (SelectedActors.Num()+UnselectedActors.Num() == 0)
-		{
-			return;
-		}
+		FGuid DummyObjectBinding;
 
-		if (SelectedActors.Num())
-		{
-			MenuBuilder.BeginSection(NAME_None, FText::FromString("Selected Splines"));
-
-			for (TSet<AActor*>::TIterator It(SelectedActors); It; ++It)
-			{
-				AActor* ActorWithSplineComponent = *It;
-				FFormatNamedArguments Args;
-				Args.Add( TEXT("PathName"), FText::FromString( ActorWithSplineComponent->GetActorLabel() ) );
-
-				MenuBuilder.AddMenuEntry(
-					FText::Format( NSLOCTEXT("Sequencer", "SetPath", "Set {PathName}"), Args ),
-					FText::Format( NSLOCTEXT("Sequencer", "SetPathTooltip", "Set path track driven by {PathName}."), Args ),
-					FSlateIcon(),
-					FUIAction(FExecuteAction::CreateSP(PathTrackEditor, &F3DPathTrackEditor::SetPath, &Section, ActorWithSplineComponent))
-				);
-			}
-			
-			MenuBuilder.EndSection();
-		}
-
-		if (UnselectedActors.Num())
-		{
-			MenuBuilder.BeginSection(NAME_None, FText::FromString("Splines"));
-
-			for (TSet<AActor*>::TIterator It(UnselectedActors); It; ++It)
-			{
-				AActor* ActorWithSplineComponent = *It;
-				FFormatNamedArguments Args;
-				Args.Add( TEXT("PathName"), FText::FromString( ActorWithSplineComponent->GetActorLabel() ) );
-
-				MenuBuilder.AddMenuEntry(
-					FText::Format( NSLOCTEXT("Sequencer", "SetPath", "Set {PathName}"), Args ),
-					FText::Format( NSLOCTEXT("Sequencer", "SetPathTooltip", "Set path track driven by {PathName}."), Args ),
-					FSlateIcon(),
-					FUIAction(FExecuteAction::CreateSP(PathTrackEditor, &F3DPathTrackEditor::SetPath, &Section, ActorWithSplineComponent))
-				);
-			}
-
-			MenuBuilder.EndSection();
-		}
+		MenuBuilder.AddSubMenu(
+			LOCTEXT("SetPath", "Path"), LOCTEXT("SetPathTooltip", "Set path"),
+			FNewMenuDelegate::CreateRaw(PathTrackEditor, &FActorPickerTrackEditor::ShowActorSubMenu, DummyObjectBinding, &Section));
 	}
 
 private:
@@ -153,13 +86,8 @@ private:
 
 
 F3DPathTrackEditor::F3DPathTrackEditor( TSharedRef<ISequencer> InSequencer )
-	: FMovieSceneTrackEditor( InSequencer ) 
-{
-}
-
-
-F3DPathTrackEditor::~F3DPathTrackEditor()
-{
+	: FActorPickerTrackEditor( InSequencer ) 
+{ 
 }
 
 
@@ -186,132 +114,76 @@ TSharedRef<ISequencerSection> F3DPathTrackEditor::MakeSectionInterface( UMovieSc
 
 void F3DPathTrackEditor::AddKey( const FGuid& ObjectGuid, UObject* AdditionalAsset )
 {
-	AddPath(ObjectGuid, AdditionalAsset);
+	AActor* ParentActor = Cast<AActor>(AdditionalAsset);
+	if (ParentActor != nullptr)
+	{
+		UMovieSceneSection* DummySection = nullptr;
+		ActorPicked(ParentActor, ObjectGuid, DummySection);
+	}
 }
 
 
 void F3DPathTrackEditor::BuildObjectBindingTrackMenu(FMenuBuilder& MenuBuilder, const FGuid& ObjectBinding, const UClass* ObjectClass)
 {
-	// build a menu with all the actors that have spline components, add the selected actors first
 	if (ObjectClass->IsChildOf(AActor::StaticClass()))
 	{
-		TSet<AActor*> SelectedActors;
-		TSet<AActor*> UnselectedActors;
-		GetActorsWithSplineComponents(SelectedActors, UnselectedActors);
+		UMovieSceneSection* DummySection = nullptr;
 
-		if (SelectedActors.Num()+UnselectedActors.Num() > 0)
-		{
-			MenuBuilder.AddSubMenu(
-				NSLOCTEXT("Sequencer", "AddPath", "Path"), NSLOCTEXT("Sequencer", "AddPathTooltip", "Adds a path track."),
-				FNewMenuDelegate::CreateRaw(this, &F3DPathTrackEditor::AddPathSubMenu, ObjectBinding));
-		}
-		else
-		{
-			AActor* ActorWithSplineComponent = nullptr;
-			FFormatNamedArguments Args;
-			MenuBuilder.AddMenuEntry(
-				FText::Format( NSLOCTEXT("Sequencer", "AddPath", "Path"), Args),
-				FText::Format( NSLOCTEXT("Sequencer", "AddPathTooltip", "Adds a path track."), Args ),
-				FSlateIcon(),
-				FUIAction(FExecuteAction::CreateSP(this, &F3DPathTrackEditor::AddPath, ObjectBinding, (UObject*)ActorWithSplineComponent),
-						  FCanExecuteAction::CreateLambda( []{ return false; } ))
-			);
-		}
+		MenuBuilder.AddSubMenu(
+			LOCTEXT("AddPath", "Path"), LOCTEXT("AddPathTooltip", "Adds a path track."),
+			FNewMenuDelegate::CreateRaw(this, &FActorPickerTrackEditor::ShowActorSubMenu, ObjectBinding, DummySection));
 	}
 }
 
-
-void F3DPathTrackEditor::AddPathSubMenu(FMenuBuilder& MenuBuilder, FGuid ObjectBinding)
+bool F3DPathTrackEditor::IsActorPickable(const AActor* const ParentActor)
 {
-	TSet<AActor*> SelectedActors;
-	TSet<AActor*> UnselectedActors;
-	GetActorsWithSplineComponents(SelectedActors, UnselectedActors);	
-
-	if (SelectedActors.Num()+UnselectedActors.Num() == 0)
-	{
-		return;
-	}
-
-	if (SelectedActors.Num())
-	{
-		MenuBuilder.BeginSection(NAME_None, FText::FromString("Selected Splines"));
-
-		for (TSet<AActor*>::TIterator It(SelectedActors); It; ++It)
-		{
-			AActor* ActorWithSplineComponent = *It;
-			FFormatNamedArguments Args;
-			Args.Add( TEXT("PathName"), FText::FromString( ActorWithSplineComponent->GetActorLabel() ) );
-
-			MenuBuilder.AddMenuEntry(
-				FText::Format( NSLOCTEXT("Sequencer", "AddPath", "{PathName}"), Args ),
-				FText::Format( NSLOCTEXT("Sequencer", "AddPathTooltip", "Add path track driven by {PathName}."), Args ),
-				FSlateIcon(),
-				FUIAction(FExecuteAction::CreateSP(this, &F3DPathTrackEditor::AddPath, ObjectBinding, (UObject*)ActorWithSplineComponent))
-			);
-		}
-			
-		MenuBuilder.EndSection();
-	}
-
-	if (UnselectedActors.Num())
-	{
-		MenuBuilder.BeginSection(NAME_None, FText::FromString("Splines"));
-
-		for (TSet<AActor*>::TIterator It(UnselectedActors); It; ++It)
+	if (ParentActor->IsListedInSceneOutliner() &&
+		!FActorEditorUtils::IsABuilderBrush(ParentActor) &&
+		!ParentActor->IsA( AWorldSettings::StaticClass() ) &&
+		!ParentActor->IsPendingKill())
+	{			
+		TArray<USplineComponent*> SplineComponents;
+		ParentActor->GetComponents(SplineComponents);
+		if (SplineComponents.Num())
 		{
-			AActor* ActorWithSplineComponent = *It;
-			FFormatNamedArguments Args;
-			Args.Add( TEXT("PathName"), FText::FromString( ActorWithSplineComponent->GetActorLabel() ) );
-
-			MenuBuilder.AddMenuEntry(
-				FText::Format( NSLOCTEXT("Sequencer", "AddPath", "{PathName}"), Args ),
-				FText::Format( NSLOCTEXT("Sequencer", "AddPathTooltip", "Add path track driven by {PathName}."), Args ),
-				FSlateIcon(),
-				FUIAction(FExecuteAction::CreateSP(this, &F3DPathTrackEditor::AddPath, ObjectBinding, (UObject*)ActorWithSplineComponent))
-			);
+			return true;
 		}
-
-		MenuBuilder.EndSection();
 	}
+	return false;
 }
 
 
-void F3DPathTrackEditor::AddPath(FGuid ObjectGuid, UObject* AdditionalAsset)
+void F3DPathTrackEditor::ActorSocketPicked(const FName SocketName, AActor* ParentActor, FGuid ObjectGuid, UMovieSceneSection* Section)
 {
-	TArray<UObject*> OutObjects;
-	GetSequencer()->GetRuntimeObjects( GetSequencer()->GetFocusedMovieSceneSequenceInstance(), ObjectGuid, OutObjects);
-
-	AnimatablePropertyChanged( FOnKeyProperty::CreateRaw( this, &F3DPathTrackEditor::AddKeyInternal, OutObjects, AdditionalAsset) );
-}
-
+	if (ObjectGuid.IsValid())
+	{
+		TArray<UObject*> OutObjects;
+		GetSequencer()->GetRuntimeObjects( GetSequencer()->GetFocusedMovieSceneSequenceInstance(), ObjectGuid, OutObjects);
 
-void F3DPathTrackEditor::SetPath(UMovieSceneSection* Section, AActor* ActorWithSplineComponent)
-{
-	const FScopedTransaction Transaction(NSLOCTEXT("Sequencer", "UndoSetPath", "Set Path"));
+		AnimatablePropertyChanged( FOnKeyProperty::CreateRaw( this, &F3DPathTrackEditor::AddKeyInternal, OutObjects, ParentActor) );
+	}
+	else if (Section != nullptr)
+	{
+		const FScopedTransaction Transaction(LOCTEXT("UndoSetPath", "Set Path"));
 
-	UMovieScene3DPathSection* PathSection = (UMovieScene3DPathSection*)(Section);
-	FGuid SplineId = FindOrCreateHandleToObject(ActorWithSplineComponent).Handle;
+		UMovieScene3DPathSection* PathSection = (UMovieScene3DPathSection*)(Section);
+		FGuid SplineId = FindOrCreateHandleToObject(ParentActor).Handle;
 
-	if (SplineId.IsValid())
-	{
-		PathSection->SetConstraintId(SplineId);
+		if (SplineId.IsValid())
+		{
+			PathSection->SetConstraintId(SplineId);
+		}
 	}
 }
 
-
-bool F3DPathTrackEditor::AddKeyInternal( float KeyTime, const TArray<UObject*> Objects, UObject* AdditionalAsset)
+bool F3DPathTrackEditor::AddKeyInternal( float KeyTime, const TArray<UObject*> Objects, AActor* ParentActor)
 {
 	bool bHandleCreated = false;
 	bool bTrackCreated = false;
 	bool bTrackModified = false;
 
-	AActor* ActorWithSplineComponent = nullptr;
+	AActor* ActorWithSplineComponent = ParentActor;
 	
-	if (AdditionalAsset != nullptr)
-	{
-		ActorWithSplineComponent = Cast<AActor>(AdditionalAsset);
-	}
-
 	FGuid SplineId;
 
 	if (ActorWithSplineComponent != nullptr)
@@ -335,9 +207,7 @@ bool F3DPathTrackEditor::AddKeyInternal( float KeyTime, const TArray<UObject*> O
 		bHandleCreated |= HandleResult.bWasCreated;
 		if (ObjectHandle.IsValid())
 		{
-			MovieSceneHelpers::SetRuntimeObjectMobility(Object);
-
-			FFindOrCreateTrackResult TrackResult = FindOrCreateTrackForObject( ObjectHandle, UMovieScene3DPathTrack::StaticClass(), FName( "Path" ) );
+			FFindOrCreateTrackResult TrackResult = FindOrCreateTrackForObject(ObjectHandle, UMovieScene3DPathTrack::StaticClass());
 			UMovieSceneTrack* Track = TrackResult.Track;
 			bTrackCreated |= TrackResult.bWasCreated;
 
@@ -359,7 +229,7 @@ bool F3DPathTrackEditor::AddKeyInternal( float KeyTime, const TArray<UObject*> O
 					}
 				}
 
-				Cast<UMovieScene3DPathTrack>(Track)->AddConstraint( KeyTime, PathEndTime, SplineId );
+				Cast<UMovieScene3DPathTrack>(Track)->AddConstraint( KeyTime, PathEndTime, NAME_None, SplineId );
 				bTrackModified = true;
 			}
 		}
@@ -367,3 +237,6 @@ bool F3DPathTrackEditor::AddKeyInternal( float KeyTime, const TArray<UObject*> O
 
 	return bHandleCreated || bTrackCreated || bTrackModified;
 }
+
+
+#undef LOCTEXT_NAMESPACE
