Fix warning compiling shot track editor, API changed in the merge

diff --git a/Engine/Source/Editor/MovieSceneTools/Private/TrackEditors/ShotTrackEditor/ShotTrackEditor.cpp b/Engine/Source/Editor/MovieSceneTools/Private/TrackEditors/ShotTrackEditor/ShotTrackEditor.cpp
--- a/Engine/Source/Editor/MovieSceneTools/Private/TrackEditors/ShotTrackEditor/ShotTrackEditor.cpp
+++ b/Engine/Source/Editor/MovieSceneTools/Private/TrackEditors/ShotTrackEditor/ShotTrackEditor.cpp
@@ -426,7 +426,8 @@ void FShotTrackEditor::OnLockCameraClicked(ECheckBoxState CheckBoxState)
 	}
 	else
 	{
-		GetSequencer()->UpdateCameraCut(nullptr, false);
+		// TODO: MaxC, is this right?
+		GetSequencer()->UpdateCameraCut(nullptr, nullptr);
 		GetSequencer()->SetPerspectiveViewportCameraCutEnabled(false);
 	}
 
