Fix misc clang warnings so Marc doesn't have to

diff --git a/Engine/Source/Editor/MovieSceneTools/Private/TrackEditors/ActorPickerTrackEditor.cpp b/Engine/Source/Editor/MovieSceneTools/Private/TrackEditors/ActorPickerTrackEditor.cpp
--- a/Engine/Source/Editor/MovieSceneTools/Private/TrackEditors/ActorPickerTrackEditor.cpp
+++ b/Engine/Source/Editor/MovieSceneTools/Private/TrackEditors/ActorPickerTrackEditor.cpp
@@ -34,10 +34,10 @@ void FActorPickerTrackEditor::ShowActorSubMenu(FMenuBuilder& MenuBuilder, FGuid
 {
 	struct Local
 	{
-		static FReply OnInteractiveActorPickerClicked(FActorPickerTrackEditor* ActorPickerTrackEditor, FGuid ObjectBinding, UMovieSceneSection* Section)
+		static FReply OnInteractiveActorPickerClicked(FActorPickerTrackEditor* ActorPickerTrackEditor, FGuid InObjectBinding, UMovieSceneSection* InSection)
 		{
 			FSlateApplication::Get().DismissAllMenus();
-			ActorPickerTrackEditor->PickActorInteractive(ObjectBinding, Section);
+			ActorPickerTrackEditor->PickActorInteractive(InObjectBinding, InSection);
 			return FReply::Handled();
 		}
 	};
diff --git a/Engine/Source/Editor/Sequencer/Private/DisplayNodes/SequencerSectionCategoryNode.cpp b/Engine/Source/Editor/Sequencer/Private/DisplayNodes/SequencerSectionCategoryNode.cpp
--- a/Engine/Source/Editor/Sequencer/Private/DisplayNodes/SequencerSectionCategoryNode.cpp
+++ b/Engine/Source/Editor/Sequencer/Private/DisplayNodes/SequencerSectionCategoryNode.cpp
@@ -43,7 +43,7 @@ ESequencerNode::Type FSequencerSectionCategoryNode::GetType() const
 }
 
 
-void FSequencerSectionCategoryNode::SetDisplayName(const FText& DisplayName)
+void FSequencerSectionCategoryNode::SetDisplayName(const FText& InDisplayName)
 {
 	check(false);
 }
diff --git a/Engine/Source/Editor/Sequencer/Private/DisplayNodes/SequencerSectionKeyAreaNode.cpp b/Engine/Source/Editor/Sequencer/Private/DisplayNodes/SequencerSectionKeyAreaNode.cpp
--- a/Engine/Source/Editor/Sequencer/Private/DisplayNodes/SequencerSectionKeyAreaNode.cpp
+++ b/Engine/Source/Editor/Sequencer/Private/DisplayNodes/SequencerSectionKeyAreaNode.cpp
@@ -65,7 +65,7 @@ ESequencerNode::Type FSequencerSectionKeyAreaNode::GetType() const
 }
 
 
-void FSequencerSectionKeyAreaNode::SetDisplayName(const FText& DisplayName)
+void FSequencerSectionKeyAreaNode::SetDisplayName(const FText& InDisplayName)
 {
 	check(false);
 }
diff --git a/Engine/Source/Runtime/RHI/Public/GPUProfiler.h b/Engine/Source/Runtime/RHI/Public/GPUProfiler.h
--- a/Engine/Source/Runtime/RHI/Public/GPUProfiler.h
+++ b/Engine/Source/Runtime/RHI/Public/GPUProfiler.h
@@ -241,8 +241,8 @@ public:
 	}
 
 	FWindowedGPUTimer(uint32 BufferSize, uint32 FramesBehind, FRHICommandListImmediate& RHICmdList)
-		:QueriesFinished(0)
-		, QueriesFailed(0)
+		: QueriesFailed(0)
+		, QueriesFinished(0)
 	{
 		StartQueries.AddZeroed(BufferSize);
 		EndQueries.AddZeroed(BufferSize);
