Include MovieSceneSection.h to avoid Mac non-unity compile errors about member access into incomplete type

diff --git a/Engine/Source/Editor/MovieSceneTools/Public/MovieSceneToolHelpers.h b/Engine/Source/Editor/MovieSceneTools/Public/MovieSceneToolHelpers.h
--- a/Engine/Source/Editor/MovieSceneTools/Public/MovieSceneToolHelpers.h
+++ b/Engine/Source/Editor/MovieSceneTools/Public/MovieSceneToolHelpers.h
@@ -5,7 +5,7 @@
 #include "ISequencer.h"
 #include "ISequencerSection.h"
 #include "IKeyArea.h"
-
+#include "MovieSceneSection.h"
 
 /**
  * A key area for float keys
