Copying //UE4/Orion-Staging to Dev-Main (Originates from //Orion/Main @ CL-2759376)

diff --git a/Engine/Plugins/Runtime/PacketHandlers/CompressionComponents/Oodle/Source/OodleHandlerComponent/Private/OodleHandlerComponent.cpp b/Engine/Plugins/Runtime/PacketHandlers/CompressionComponents/Oodle/Source/OodleHandlerComponent/Private/OodleHandlerComponent.cpp
--- a/Engine/Plugins/Runtime/PacketHandlers/CompressionComponents/Oodle/Source/OodleHandlerComponent/Private/OodleHandlerComponent.cpp
+++ b/Engine/Plugins/Runtime/PacketHandlers/CompressionComponents/Oodle/Source/OodleHandlerComponent/Private/OodleHandlerComponent.cpp
@@ -6,6 +6,8 @@
 DEFINE_LOG_CATEGORY(OodleHandlerComponentLog);
 
 
+// @todo #JohnB: Deprecated the 'training' mode from this handler eventually, as it is misnamed/confusing
+
 // @todo #JohnB: Right now, there is no caching of loaded dictionary files, so this is very inefficient with multiple clients
 
 #if HAS_OODLE_SDK
@@ -166,7 +168,7 @@ void OodleHandlerComponent::InitFirstRunConfig()
 		GConfig->SetBool(OODLE_INI_SECTION, TEXT("bUseDictionaryIfPresent"), false, GEngineIni);
 #endif
 
-		GConfig->SetString(OODLE_INI_SECTION, TEXT("Mode"), TEXT("Training"), GEngineIni);
+		GConfig->SetString(OODLE_INI_SECTION, TEXT("Mode"), TEXT("Capturing"), GEngineIni);
 		GConfig->SetString(OODLE_INI_SECTION, TEXT("PacketLogFile"), TEXT("PacketDump"), GEngineIni);
 
 #if EXAMPLE_CAPTURE_FORMAT
@@ -203,9 +205,9 @@ void OodleHandlerComponent::Initialize()
 	FString ReadMode;
 	GConfig->GetString(OODLE_INI_SECTION, TEXT("Mode"), ReadMode, GEngineIni);
 
-	if(ReadMode == TEXT("Training"))
+	if (ReadMode == TEXT("Training") || ReadMode == TEXT("Capturing"))
 	{
-		Mode = EOodleHandlerMode::Training;
+		Mode = EOodleHandlerMode::Capturing;
 	}
 	else if(ReadMode == TEXT("Release"))
 	{
@@ -220,16 +222,17 @@ void OodleHandlerComponent::Initialize()
 	if (bEnableOodle)
 	{
 #if !UE_BUILD_SHIPPING
-		bool bForceTraining = FParse::Param(FCommandLine::Get(), TEXT("OodleTraining"));
+		bool bForceCapturing = FParse::Param(FCommandLine::Get(), TEXT("OodleTraining")) ||
+								FParse::Param(FCommandLine::Get(), TEXT("OodleCapturing"));
 
-		if (bForceTraining && Mode != EOodleHandlerMode::Training)
+		if (bForceCapturing && Mode != EOodleHandlerMode::Capturing)
 		{
-			UE_LOG(OodleHandlerComponentLog, Log, TEXT("Switching Oodle to Training mode."));
+			UE_LOG(OodleHandlerComponentLog, Log, TEXT("Switching Oodle to Capturing mode."));
 
-			Mode = EOodleHandlerMode::Training;
+			Mode = EOodleHandlerMode::Capturing;
 		}
 		// Forces release mode, if configured to search for (and finds) fallback dictionaries
-		else if (!bForceTraining && bUseDictionaryIfPresent && Mode != EOodleHandlerMode::Release)
+		else if (!bForceCapturing && bUseDictionaryIfPresent && Mode != EOodleHandlerMode::Release)
 		{
 			FString ServerDicPath;
 			FString ClientDicPath;
@@ -247,7 +250,7 @@ void OodleHandlerComponent::Initialize()
 
 		switch(Mode)
 		{
-		case EOodleHandlerMode::Training:
+		case EOodleHandlerMode::Capturing:
 			{
 				if (Handler->Mode == Handler::Mode::Server)
 				{
@@ -678,8 +681,13 @@ bool OodleHandlerComponent::FindFallbackDictionaries(FString& OutServerDictionar
 	bSuccess = bSuccess && FileMan.FileExists(*DefaultClientDicPath);
 
 
+	if (bSuccess)
+	{
+		OutServerDictionary = DefaultServerDicPath;
+		OutClientDictionary = DefaultClientDicPath;
+	}
 	// If either of the default dictionaries do not exist, do a more speculative search
-	if (!bSuccess)
+	else
 	{
 		TArray<FString> DictionaryList;
 		
@@ -755,7 +763,7 @@ void OodleHandlerComponent::Incoming(FBitReader& Packet)
 	{
 		switch(Mode)
 		{
-		case EOodleHandlerMode::Training:
+		case EOodleHandlerMode::Capturing:
 			{
 				if(Handler->Mode == Handler::Mode::Server)
 				{
@@ -812,7 +820,7 @@ void OodleHandlerComponent::Outgoing(FBitWriter& Packet)
 	{
 		switch(Mode)
 		{
-		case EOodleHandlerMode::Training:
+		case EOodleHandlerMode::Capturing:
 			{
 				if(Handler->Mode == Handler::Mode::Server)
 				{
diff --git a/Engine/Plugins/Runtime/PacketHandlers/CompressionComponents/Oodle/Source/OodleHandlerComponent/Public/OodleHandlerComponent.h b/Engine/Plugins/Runtime/PacketHandlers/CompressionComponents/Oodle/Source/OodleHandlerComponent/Public/OodleHandlerComponent.h
--- a/Engine/Plugins/Runtime/PacketHandlers/CompressionComponents/Oodle/Source/OodleHandlerComponent/Public/OodleHandlerComponent.h
+++ b/Engine/Plugins/Runtime/PacketHandlers/CompressionComponents/Oodle/Source/OodleHandlerComponent/Public/OodleHandlerComponent.h
@@ -39,7 +39,7 @@ extern FString GOodleSaveDir;
  */
 enum EOodleHandlerMode
 {
-	Training,	// Stores packet captures for the server
+	Capturing,	// Stores packet captures for the server
 	Release		// Compresses packet data, based on the dictionary file
 };
 
@@ -164,7 +164,7 @@ protected:
 	* Modes of the component
 	* Specify this in .ini by
 	* [OodleHandlerComponent]
-	* Mode=Training
+	* Mode=Capturing
 	*/
 	EOodleHandlerMode Mode;
 
diff --git a/Engine/Shaders/DeferredShadingCommon.usf b/Engine/Shaders/DeferredShadingCommon.usf
--- a/Engine/Shaders/DeferredShadingCommon.usf
+++ b/Engine/Shaders/DeferredShadingCommon.usf
@@ -246,7 +246,8 @@ bool IsSubsurfaceModel(int ShadingModel)
 		|| ShadingModel == SHADINGMODELID_PREINTEGRATED_SKIN 
 		|| ShadingModel == SHADINGMODELID_SUBSURFACE_PROFILE
 		|| ShadingModel == SHADINGMODELID_TWOSIDED_FOLIAGE
-		|| ShadingModel == SHADINGMODELID_HAIR;
+		|| ShadingModel == SHADINGMODELID_HAIR
+		|| ShadingModel == SHADINGMODELID_EYE;
 }
 
 // all values that are output by the forward rendering pass
diff --git a/Engine/Shaders/LandscapeVertexFactory.usf b/Engine/Shaders/LandscapeVertexFactory.usf
--- a/Engine/Shaders/LandscapeVertexFactory.usf
+++ b/Engine/Shaders/LandscapeVertexFactory.usf
@@ -148,20 +148,20 @@ struct FVertexFactoryIntermediates
 
 float3 GetLocalPosition(FVertexFactoryIntermediates Intermediates)
 {
-	return Intermediates.LocalPosition+float3(Intermediates.InputPosition.zw * LandscapeParameters.SubsectionOffsetParams.ww,0);
+	return INVARIANT(Intermediates.LocalPosition+float3(Intermediates.InputPosition.zw * LandscapeParameters.SubsectionOffsetParams.ww,0));
 }
 
 float4 VertexFactoryGetWorldPosition(FVertexFactoryInput Input, FVertexFactoryIntermediates Intermediates)
 {
-	return TransformLocalToTranslatedWorld(GetLocalPosition(Intermediates));
+	return INVARIANT(TransformLocalToTranslatedWorld(GetLocalPosition(Intermediates)));
 }
 
 float4 VertexFactoryGetPreviousWorldPosition(FVertexFactoryInput Input, FVertexFactoryIntermediates Intermediates)
 {
 	// Note we still use LocalToWorld. Landscape's transform never changes.
 	float3 LocalPosition = GetLocalPosition(Intermediates);
-	float3 RotatedPosition = Primitive.LocalToWorld[0].xyz * LocalPosition.xxx + Primitive.LocalToWorld[1].xyz * LocalPosition.yyy + Primitive.LocalToWorld[2].xyz * LocalPosition.zzz;
-	return float4(RotatedPosition + (Primitive.LocalToWorld[3].xyz + View.PrevPreViewTranslation.xyz),1);
+	float3 RotatedPosition = INVARIANT(Primitive.LocalToWorld[0].xyz * LocalPosition.xxx + Primitive.LocalToWorld[1].xyz * LocalPosition.yyy + Primitive.LocalToWorld[2].xyz * LocalPosition.zzz);
+	return INVARIANT(float4(RotatedPosition + (Primitive.LocalToWorld[3].xyz + View.PrevPreViewTranslation.xyz),1));
 }
 
 /** Calculate the texture coordinates generated by Landscape */
diff --git a/Engine/Shaders/MaterialTemplate.usf b/Engine/Shaders/MaterialTemplate.usf
--- a/Engine/Shaders/MaterialTemplate.usf
+++ b/Engine/Shaders/MaterialTemplate.usf
@@ -39,6 +39,12 @@
     #endif
 #endif
 
+// scale factor for screen position, for cases where we may render certain things at lower resolutions (e.g. separate translucency)
+// any pass wanting to render at reduced resolution needs to simply set this constant to scale scene texture UVs to make scene texture accesses work
+#if PIXELSHADER && TRANSLUCENCY_RENDERING
+float2 ScreenPositionUVScale;
+#endif
+
 //Platforms that don't run the editor shouldn't need editor features in the shaders.
 #define PLATFORM_SUPPORTS_EDITOR_SHADERS (!PS4_PROFILE && !XBOXONE_PROFILE && !ES31_AEP_PROFILE)
 
@@ -1117,6 +1123,10 @@ float4 SceneTextureLookup(float2 UV, int SceneTextureIndex, bool bFiltered)
 	return float4(0.0f, 0.0f, 0.0f, 0.0f);
 #endif
 
+#if PIXELSHADER && TRANSLUCENCY_RENDERING
+	UV *= ScreenPositionUVScale;
+#endif
+
 	FScreenSpaceData ScreenSpaceData = GetScreenSpaceData(UV, false);
 
 #if PS4_PROFILE
diff --git a/Engine/Shaders/ParticleGPUSpriteVertexFactory.usf b/Engine/Shaders/ParticleGPUSpriteVertexFactory.usf
--- a/Engine/Shaders/ParticleGPUSpriteVertexFactory.usf
+++ b/Engine/Shaders/ParticleGPUSpriteVertexFactory.usf
@@ -118,6 +118,9 @@ struct FVertexFactoryIntermediates
 
 	/** The texture coordinate at this vertex. */
 	float4 TexCoord;
+#if NUM_MATERIAL_TEXCOORDS >= 3 //Unflipped UVs in texcoords 3 and 4
+	float4 TexCoord_Unflipped;
+#endif
 	/** The sprite tangent in world space (+V). */
 	float3 TangentUp;
 	/** The sprite tangent in world space (+U). */
@@ -263,9 +266,15 @@ FVertexFactoryIntermediates GetVertexFactoryIntermediates( FVertexFactoryInput I
 	float SubImageBV = floor( SubImageB * EmitterUniforms.SubImageSize.z );
 
 	// Compute the final texture coordinates for both subimages.
-	float4 TexCoord;
+	float4 TexCoord; 
 	TexCoord.xy = (float2( SubImageAH, SubImageAV ) + FlippedUV) * EmitterUniforms.SubImageSize.zw;
-	TexCoord.zw = (float2( SubImageBH, SubImageBV ) + FlippedUV) * EmitterUniforms.SubImageSize.zw;
+	TexCoord.zw = (float2(SubImageBH, SubImageBV) + FlippedUV) * EmitterUniforms.SubImageSize.zw;
+
+#if NUM_MATERIAL_TEXCOORDS >= 3 //Unflipped UVs in texcoords 3 and 4
+	float4 TexCoord_Unflipped;
+	TexCoord_Unflipped.xy = (float2(SubImageAH, SubImageAV) + Input.TexCoord.xy) * EmitterUniforms.SubImageSize.zw;
+	TexCoord_Unflipped.zw = (float2(SubImageBH, SubImageBV) + Input.TexCoord.xy) * EmitterUniforms.SubImageSize.zw;
+#endif
 
 	// Current rotation of the sprite. Map [0,1] to one full rotation.
 	float RotationRate = AttributeSample.w * EmitterUniforms.RotationRateScale;
@@ -312,6 +321,9 @@ FVertexFactoryIntermediates GetVertexFactoryIntermediates( FVertexFactoryInput I
 #endif
 
 	Intermediates.TexCoord = TexCoord;
+#if NUM_MATERIAL_TEXCOORDS >= 3 //Unflipped UVs in texcoords 3 and 4
+	Intermediates.TexCoord_Unflipped = TexCoord_Unflipped; 
+#endif
 	Intermediates.TangentUp = TangentUp;
 	Intermediates.TangentRight = TangentRight;
 	Intermediates.Color = Color;
@@ -477,12 +489,18 @@ FMaterialVertexParameters GetMaterialVertexParameters(FVertexFactoryInput Input,
 #endif
 
 	Result.Particle.Color = Intermediates.Color;
-
+	
 #if NUM_MATERIAL_TEXCOORDS_VERTEX
 #if NUM_MATERIAL_TEXCOORDS_VERTEX >= 1
 		Result.TexCoords[0] = Intermediates.TexCoord.xy;
 #if NUM_MATERIAL_TEXCOORDS_VERTEX >= 2
 		Result.TexCoords[1] = Intermediates.TexCoord.zw;
+#if NUM_MATERIAL_TEXCOORDS_VERTEX >= 3
+		Result.TexCoords[2] = Intermediates.TexCoord_Unflipped.xy;
+#if NUM_MATERIAL_TEXCOORDS_VERTEX >= 4
+		Result.TexCoords[3] = Intermediates.TexCoord_Unflipped.zw;
+#endif	// >= 4
+#endif	// >= 3
 #endif	// >= 2
 #endif	// >= 1
 #endif	//NUM_MATERIAL_TEXCOORDS
diff --git a/Engine/Shaders/ParticleSpriteVertexFactory.usf b/Engine/Shaders/ParticleSpriteVertexFactory.usf
--- a/Engine/Shaders/ParticleSpriteVertexFactory.usf
+++ b/Engine/Shaders/ParticleSpriteVertexFactory.usf
@@ -98,6 +98,10 @@ struct FVertexFactoryIntermediates
 
 	/** The texture coordinate at this vertex. */
 	float4 TexCoord;
+#if NUM_MATERIAL_TEXCOORDS >= 3 //Unflipped UVs are placed in UVs 2-3
+	/** A second UV set. Always non-UV flipped. Allows use of UV flipping with Normal maps etc. */
+	float4 TexCoord_Unflipped;
+#endif
 	/** The sprite tangent in world space (+V). */
 	float3 TangentUp;
 	/** The sprite tangent in world space (+U). */
@@ -263,6 +267,12 @@ FMaterialVertexParameters GetMaterialVertexParameters(FVertexFactoryInput Input,
 		Result.TexCoords[0] = Intermediates.TexCoord.xy;
 	#if NUM_MATERIAL_TEXCOORDS_VERTEX >= 2
 		Result.TexCoords[1] = Intermediates.TexCoord.zw;
+	#if NUM_MATERIAL_TEXCOORDS_VERTEX >= 3
+		Result.TexCoords[2] = Intermediates.TexCoord_Unflipped.xy;
+	#if NUM_MATERIAL_TEXCOORDS_VERTEX >= 4
+		Result.TexCoords[3] = Intermediates.TexCoord_Unflipped.zw;
+	#endif	// >= 4
+	#endif	// >= 3
 	#endif	// >= 2
 	#endif	// >= 1
 #endif	//NUM_MATERIAL_TEXCOORDS_VERTEX
@@ -419,7 +429,11 @@ FVertexFactoryIntermediates GetVertexFactoryIntermediates(FVertexFactoryInput In
 	float SubImageAV = floor( SubImageA * SpriteVF.SubImageSize.z );
 	float SubImageBV = floor( SubImageB * SpriteVF.SubImageSize.z );
 	Intermediates.TexCoord.xy = (float2( SubImageAH, SubImageAV ) + FlippedUV) * SpriteVF.SubImageSize.zw ;
-	Intermediates.TexCoord.zw = (float2( SubImageBH, SubImageBV ) + FlippedUV) * SpriteVF.SubImageSize.zw ;
+	Intermediates.TexCoord.zw = (float2(SubImageBH, SubImageBV) + FlippedUV) * SpriteVF.SubImageSize.zw;
+#if NUM_MATERIAL_TEXCOORDS >= 3
+	Intermediates.TexCoord_Unflipped.xy = (float2(SubImageAH, SubImageAV) + Input.TexCoord.xy) * SpriteVF.SubImageSize.zw;
+	Intermediates.TexCoord_Unflipped.zw = (float2(SubImageBH, SubImageBV) + Input.TexCoord.xy) * SpriteVF.SubImageSize.zw;
+#endif
 	Intermediates.SubImageLerp = SubImageLerp;
 
 	Intermediates.Color = Input.Color;
diff --git a/Engine/Shaders/ShaderVersion.usf b/Engine/Shaders/ShaderVersion.usf
--- a/Engine/Shaders/ShaderVersion.usf
+++ b/Engine/Shaders/ShaderVersion.usf
@@ -2,4 +2,4 @@
 // Each time the console command is executed it generates a new GUID. As this file is included
 // in common.usf (which should be included in any shader) it allows to invalidate the shader DDC.
 // 
-// GUID = 8FE3C62E-D12D4E3B94F997A3F1C66BDF
\ No newline at end of file
+// GUID = 82D4113E49F563B43A4FDDB583CF26CB
\ No newline at end of file
diff --git a/Engine/Shaders/ShadingModels.usf b/Engine/Shaders/ShadingModels.usf
--- a/Engine/Shaders/ShadingModels.usf
+++ b/Engine/Shaders/ShadingModels.usf
@@ -691,8 +691,13 @@ float3 EyeShading( FGBufferData GBuffer, float3 LobeRoughness, float3 LobeEnergy
 
 float3 EyeSubsurfaceShading(  FGBufferData GBuffer, float3 L, float3 V, half3 N )
 {
-	// Need to make sure we are getting the correct Diffuse color intensity
-	float3 SubsurfaceColor = GBuffer.DiffuseColor * PI;
+	float3 H = normalize(V + L);
+	float VoH = saturate( dot(V, H) );
+
+	// F_Schlick
+	float F0 = GBuffer.SpecularColor;
+	float Fc = Pow5( 1 - VoH );
+	float F = Fc + (1 - Fc) * F0;
 
 	float  IrisDistance = GBuffer.CustomData.w;
 	float  IrisMask = GBuffer.CustomData.z;
@@ -708,19 +713,13 @@ float3 EyeSubsurfaceShading(  FGBufferData GBuffer, float3 L, float3 V, half3 N
 
 	// Use IrisNormal only within the Iris
 	IrisNormal = lerp(N, IrisNormal, IrisMask);
+	IrisNormal = normalize( IrisNormal );
 
 	// http://blog.stevemcauley.com/2011/12/03/energy-conserving-wrapped-diffuse/
-	float Wrap = 0.5;
+	float Wrap = 0.1;
 	float INoL = saturate( ( dot(IrisNormal, L) + Wrap ) / Square( 1 + Wrap ) );
-
-	// GGX scatter distribution
-	float VoL = saturate( dot(V, -L) );
-	float a = 0.6;
-	float a2 = a * a;
-	float d = ( VoL * a2 - VoL ) * VoL + 1;	// 2 mad
-	float GGX = (a2 / PI) / (d * d);		// 2 mul, 1 rcp
 	
-	float3 rawResult = INoL * GGX * SubsurfaceColor;
+	float3 rawResult = (1 - F) * INoL * GBuffer.DiffuseColor / PI;
 
 	// Covers whole object
 	return rawResult;
diff --git a/Engine/Shaders/ShadowProjectionPixelShader.usf b/Engine/Shaders/ShadowProjectionPixelShader.usf
--- a/Engine/Shaders/ShadowProjectionPixelShader.usf
+++ b/Engine/Shaders/ShadowProjectionPixelShader.usf
@@ -107,7 +107,7 @@ void Main(
 		BRANCH
 		if (IsSubsurfaceModel(InGBufferData.ShadingModelID))
 		{
-			float Opacity = InGBufferData.ShadingModelID == SHADINGMODELID_HAIR ? 1 : InGBufferData.CustomData.a;
+			float Opacity = InGBufferData.ShadingModelID == SHADINGMODELID_HAIR || InGBufferData.ShadingModelID == SHADINGMODELID_EYE ? 1 : InGBufferData.CustomData.a;
 			// Derive density from a heuristic using opacity, tweaked for useful falloff ranges and to give a linear depth falloff with opacity
 			float Density = -.05f * log(1 - min(Opacity, .999f));
 			//@todo - Use a larger kernel filter based on material opacity
diff --git a/Engine/Source/Developer/LogVisualizer/Private/SSequencerSectionOverlay.cpp b/Engine/Source/Developer/LogVisualizer/Private/SSequencerSectionOverlay.cpp
--- a/Engine/Source/Developer/LogVisualizer/Private/SSequencerSectionOverlay.cpp
+++ b/Engine/Source/Developer/LogVisualizer/Private/SSequencerSectionOverlay.cpp
@@ -13,24 +13,24 @@ int32 SVisualLoggerSectionOverlay::OnPaint( const FPaintArgs& Args, const FGeome
 
 FReply SVisualLoggerSectionOverlay::OnMouseButtonDown( const FGeometry& MyGeometry, const FPointerEvent& MouseEvent )
 {
-	return TimeSliderController->OnMouseButtonDown(SharedThis(this), MyGeometry, MouseEvent);
+	return TimeSliderController->OnMouseButtonDown(*this, MyGeometry, MouseEvent);
 }
 
 FReply SVisualLoggerSectionOverlay::OnMouseButtonUp( const FGeometry& MyGeometry, const FPointerEvent& MouseEvent )
 {
-	return TimeSliderController->OnMouseButtonUp( SharedThis(this),  MyGeometry, MouseEvent );
+	return TimeSliderController->OnMouseButtonUp( *this,  MyGeometry, MouseEvent );
 }
 
 FReply SVisualLoggerSectionOverlay::OnMouseMove( const FGeometry& MyGeometry, const FPointerEvent& MouseEvent )
 {
-	return TimeSliderController->OnMouseMove( SharedThis(this), MyGeometry, MouseEvent );
+	return TimeSliderController->OnMouseMove( *this, MyGeometry, MouseEvent );
 }
 
 FReply SVisualLoggerSectionOverlay::OnMouseWheel( const FGeometry& MyGeometry, const FPointerEvent& MouseEvent )
 {
 	if (MouseEvent.IsLeftShiftDown() || MouseEvent.IsLeftControlDown())
 	{
-		return TimeSliderController->OnMouseWheel(SharedThis(this), MyGeometry, MouseEvent);
+		return TimeSliderController->OnMouseWheel(*this, MyGeometry, MouseEvent);
 	}
 	return FReply::Unhandled();
 }
diff --git a/Engine/Source/Developer/LogVisualizer/Private/STimeSlider.cpp b/Engine/Source/Developer/LogVisualizer/Private/STimeSlider.cpp
--- a/Engine/Source/Developer/LogVisualizer/Private/STimeSlider.cpp
+++ b/Engine/Source/Developer/LogVisualizer/Private/STimeSlider.cpp
@@ -21,17 +21,17 @@ int32 STimeSlider::OnPaint( const FPaintArgs& Args, const FGeometry& AllottedGeo
 
 FReply STimeSlider::OnMouseButtonDown( const FGeometry& MyGeometry, const FPointerEvent& MouseEvent )
 {
-	return TimeSliderController->OnMouseButtonDown( SharedThis(this), MyGeometry, MouseEvent );
+	return TimeSliderController->OnMouseButtonDown( *this, MyGeometry, MouseEvent );
 }
 
 FReply STimeSlider::OnMouseButtonUp( const FGeometry& MyGeometry, const FPointerEvent& MouseEvent )
 {
-	return TimeSliderController->OnMouseButtonUp( SharedThis(this),  MyGeometry, MouseEvent );
+	return TimeSliderController->OnMouseButtonUp( *this,  MyGeometry, MouseEvent );
 }
 
 FReply STimeSlider::OnMouseMove( const FGeometry& MyGeometry, const FPointerEvent& MouseEvent )
 {
-	return TimeSliderController->OnMouseMove( SharedThis(this), MyGeometry, MouseEvent );
+	return TimeSliderController->OnMouseMove( *this, MyGeometry, MouseEvent );
 }
 
 FVector2D STimeSlider::ComputeDesiredSize( float ) const
@@ -41,7 +41,7 @@ FVector2D STimeSlider::ComputeDesiredSize( float ) const
 
 FReply STimeSlider::OnMouseWheel( const FGeometry& MyGeometry, const FPointerEvent& MouseEvent )
 {
-	return TimeSliderController->OnMouseWheel( SharedThis(this), MyGeometry, MouseEvent );
+	return TimeSliderController->OnMouseWheel( *this, MyGeometry, MouseEvent );
 }
 
 #undef LOCTEXT_NAMESPACE
diff --git a/Engine/Source/Developer/LogVisualizer/Private/STimelineBar.cpp b/Engine/Source/Developer/LogVisualizer/Private/STimelineBar.cpp
--- a/Engine/Source/Developer/LogVisualizer/Private/STimelineBar.cpp
+++ b/Engine/Source/Developer/LogVisualizer/Private/STimelineBar.cpp
@@ -16,7 +16,7 @@
 FReply STimelineBar::OnMouseButtonDown(const FGeometry& MyGeometry, const FPointerEvent& MouseEvent)
 {
 	TimelineOwner.Pin()->OnMouseButtonDown(MyGeometry, MouseEvent);
-	FReply Replay = TimeSliderController->OnMouseButtonDown(SharedThis(this), MyGeometry, MouseEvent);
+	FReply Replay = TimeSliderController->OnMouseButtonDown(*this, MyGeometry, MouseEvent);
 	if (Replay.IsEventHandled())
 	{
 		FName RowName = TimelineOwner.Pin()->GetName();
@@ -36,7 +36,7 @@ FReply STimelineBar::OnMouseButtonUp(const FGeometry& MyGeometry, const FPointer
 {
 	TimelineOwner.Pin()->OnMouseButtonUp(MyGeometry, MouseEvent);
 
-	FReply  Replay = TimeSliderController->OnMouseButtonUp(SharedThis(this), MyGeometry, MouseEvent);
+	FReply  Replay = TimeSliderController->OnMouseButtonUp(*this, MyGeometry, MouseEvent);
 	if (Replay.IsEventHandled())
 	{
 		FName RowName = TimelineOwner.Pin()->GetName();
@@ -54,7 +54,7 @@ FReply STimelineBar::OnMouseButtonUp(const FGeometry& MyGeometry, const FPointer
 
 FReply STimelineBar::OnMouseMove(const FGeometry& MyGeometry, const FPointerEvent& MouseEvent)
 {
-	return TimeSliderController->OnMouseMove(SharedThis(this), MyGeometry, MouseEvent);
+	return TimeSliderController->OnMouseMove(*this, MyGeometry, MouseEvent);
 }
 
 FReply STimelineBar::OnMouseButtonDoubleClick(const FGeometry& MyGeometry, const FPointerEvent& MouseEvent)
diff --git a/Engine/Source/Developer/LogVisualizer/Private/STimelinesContainer.cpp b/Engine/Source/Developer/LogVisualizer/Private/STimelinesContainer.cpp
--- a/Engine/Source/Developer/LogVisualizer/Private/STimelinesContainer.cpp
+++ b/Engine/Source/Developer/LogVisualizer/Private/STimelinesContainer.cpp
@@ -129,7 +129,7 @@ FReply STimelinesContainer::OnMouseButtonDown(const FGeometry& MyGeometry, const
 {
 	if (MouseEvent.GetEffectingButton() != EKeys::LeftMouseButton)
 	{
-		return TimeSliderController->OnMouseButtonDown(SharedThis(this), MyGeometry, MouseEvent);
+		return TimeSliderController->OnMouseButtonDown(*this, MyGeometry, MouseEvent);
 	}
 	return FReply::Unhandled();
 }
@@ -138,7 +138,7 @@ FReply STimelinesContainer::OnMouseButtonUp(const FGeometry& MyGeometry, const F
 {
 	if (MouseEvent.GetEffectingButton() != EKeys::LeftMouseButton)
 	{
-		return TimeSliderController->OnMouseButtonUp(SharedThis(this), MyGeometry, MouseEvent);
+		return TimeSliderController->OnMouseButtonUp(*this, MyGeometry, MouseEvent);
 	}
 	return FReply::Unhandled();
 }
@@ -147,7 +147,7 @@ FReply STimelinesContainer::OnMouseMove(const FGeometry& MyGeometry, const FPoin
 {
 	if (MouseEvent.GetEffectingButton() != EKeys::LeftMouseButton)
 	{
-		return TimeSliderController->OnMouseMove(SharedThis(this), MyGeometry, MouseEvent);
+		return TimeSliderController->OnMouseMove(*this, MyGeometry, MouseEvent);
 	}
 	return FReply::Unhandled();
 }
@@ -156,7 +156,7 @@ FReply STimelinesContainer::OnMouseWheel(const FGeometry& MyGeometry, const FPoi
 {
 	if (MouseEvent.IsLeftControlDown() || MouseEvent.IsLeftShiftDown())
 	{
-		return TimeSliderController->OnMouseWheel(SharedThis(this), MyGeometry, MouseEvent);
+		return TimeSliderController->OnMouseWheel(*this, MyGeometry, MouseEvent);
 	}
 	return FReply::Unhandled();
 }
diff --git a/Engine/Source/Developer/LogVisualizer/Private/SVisualLoggerView.cpp b/Engine/Source/Developer/LogVisualizer/Private/SVisualLoggerView.cpp
--- a/Engine/Source/Developer/LogVisualizer/Private/SVisualLoggerView.cpp
+++ b/Engine/Source/Developer/LogVisualizer/Private/SVisualLoggerView.cpp
@@ -36,7 +36,7 @@ FReply SInputCatcherOverlay::OnMouseButtonDown(const FGeometry& MyGeometry, cons
 {
 	if (MouseEvent.GetEffectingButton() != EKeys::LeftMouseButton)
 	{
-		return TimeSliderController->OnMouseButtonDown(SharedThis(this), MyGeometry, MouseEvent);
+		return TimeSliderController->OnMouseButtonDown(*this, MyGeometry, MouseEvent);
 	}
 	return FReply::Unhandled();
 }
@@ -45,21 +45,21 @@ FReply SInputCatcherOverlay::OnMouseButtonUp(const FGeometry& MyGeometry, const
 {
 	if (MouseEvent.GetEffectingButton() != EKeys::LeftMouseButton)
 	{
-		return TimeSliderController->OnMouseButtonUp(SharedThis(this), MyGeometry, MouseEvent);
+		return TimeSliderController->OnMouseButtonUp(*this, MyGeometry, MouseEvent);
 	}
 	return FReply::Unhandled();
 }
 
 FReply SInputCatcherOverlay::OnMouseMove(const FGeometry& MyGeometry, const FPointerEvent& MouseEvent)
 {
-	return TimeSliderController->OnMouseMove(SharedThis(this), MyGeometry, MouseEvent);
+	return TimeSliderController->OnMouseMove(*this, MyGeometry, MouseEvent);
 }
 
 FReply SInputCatcherOverlay::OnMouseWheel(const FGeometry& MyGeometry, const FPointerEvent& MouseEvent)
 {
 	if (MouseEvent.IsLeftShiftDown() || MouseEvent.IsLeftControlDown())
 	{
-		return TimeSliderController->OnMouseWheel(SharedThis(this), MyGeometry, MouseEvent);
+		return TimeSliderController->OnMouseWheel(*this, MyGeometry, MouseEvent);
 	}
 	return FReply::Unhandled();
 }
diff --git a/Engine/Source/Developer/LogVisualizer/Private/TimeSliderController.cpp b/Engine/Source/Developer/LogVisualizer/Private/TimeSliderController.cpp
--- a/Engine/Source/Developer/LogVisualizer/Private/TimeSliderController.cpp
+++ b/Engine/Source/Developer/LogVisualizer/Private/TimeSliderController.cpp
@@ -310,7 +310,7 @@ int32 FVisualLoggerTimeSliderController::OnPaintTimeSlider( bool bMirrorLabels,
 	return LayerId;
 }
 
-FReply FVisualLoggerTimeSliderController::OnMouseButtonDown( TSharedRef<SWidget> WidgetOwner, const FGeometry& MyGeometry, const FPointerEvent& MouseEvent )
+FReply FVisualLoggerTimeSliderController::OnMouseButtonDown( SWidget& WidgetOwner, const FGeometry& MyGeometry, const FPointerEvent& MouseEvent )
 {
 	bool bHandleLeftMouseButton = MouseEvent.GetEffectingButton() == EKeys::LeftMouseButton;
 	bool bHandleRightMouseButton = MouseEvent.GetEffectingButton() == EKeys::RightMouseButton && TimeSliderArgs.AllowZoom;
@@ -325,20 +325,20 @@ FReply FVisualLoggerTimeSliderController::OnMouseButtonDown( TSharedRef<SWidget>
 		float NewValue = RangeToScreen.LocalXToInput(CursorPos.X);
 
 		CommitScrubPosition(NewValue, /*bIsScrubbing=*/false);
-		return FReply::Handled().CaptureMouse( WidgetOwner ).PreventThrottling();
+		return FReply::Handled().CaptureMouse( WidgetOwner.AsShared() ).PreventThrottling();
 	}
 	else if ( bHandleRightMouseButton )
 	{
-		return FReply::Handled().CaptureMouse(WidgetOwner);
+		return FReply::Handled().CaptureMouse(WidgetOwner.AsShared());
 	}
 
 	return FReply::Unhandled();
 }
 
-FReply FVisualLoggerTimeSliderController::OnMouseButtonUp( TSharedRef<SWidget> WidgetOwner, const FGeometry& MyGeometry, const FPointerEvent& MouseEvent )
+FReply FVisualLoggerTimeSliderController::OnMouseButtonUp( SWidget& WidgetOwner, const FGeometry& MyGeometry, const FPointerEvent& MouseEvent )
 {
-	bool bHandleLeftMouseButton = MouseEvent.GetEffectingButton() == EKeys::LeftMouseButton && WidgetOwner->HasMouseCapture();
-	bool bHandleRightMouseButton = MouseEvent.GetEffectingButton() == EKeys::RightMouseButton && WidgetOwner->HasMouseCapture() && TimeSliderArgs.AllowZoom ;
+	bool bHandleLeftMouseButton = MouseEvent.GetEffectingButton() == EKeys::LeftMouseButton && WidgetOwner.HasMouseCapture();
+	bool bHandleRightMouseButton = MouseEvent.GetEffectingButton() == EKeys::RightMouseButton && WidgetOwner.HasMouseCapture() && TimeSliderArgs.AllowZoom ;
 	
 	if ( bHandleRightMouseButton )
 	{
@@ -349,7 +349,7 @@ FReply FVisualLoggerTimeSliderController::OnMouseButtonUp( TSharedRef<SWidget> W
 		}
 		
 		bPanning = false;
-		FReply::Handled().CaptureMouse(WidgetOwner).UseHighPrecisionMouseMovement(WidgetOwner);
+		FReply::Handled().CaptureMouse(WidgetOwner.AsShared()).UseHighPrecisionMouseMovement(WidgetOwner.AsShared());
 
 		return FReply::Handled().ReleaseMouseCapture();
 	}
@@ -376,9 +376,9 @@ FReply FVisualLoggerTimeSliderController::OnMouseButtonUp( TSharedRef<SWidget> W
 	return FReply::Unhandled();
 }
 
-FReply FVisualLoggerTimeSliderController::OnMouseMove( TSharedRef<SWidget> WidgetOwner, const FGeometry& MyGeometry, const FPointerEvent& MouseEvent )
+FReply FVisualLoggerTimeSliderController::OnMouseMove( SWidget& WidgetOwner, const FGeometry& MyGeometry, const FPointerEvent& MouseEvent )
 {
-	if ( WidgetOwner->HasMouseCapture() )
+	if ( WidgetOwner.HasMouseCapture() )
 	{
 		if (MouseEvent.IsMouseButtonDown(EKeys::RightMouseButton))
 		{
@@ -387,7 +387,7 @@ FReply FVisualLoggerTimeSliderController::OnMouseMove( TSharedRef<SWidget> Widge
 				DistanceDragged += FMath::Abs( MouseEvent.GetCursorDelta().X );
 				if ( DistanceDragged > FSlateApplication::Get().GetDragTriggerDistance() )
 				{
-					FReply::Handled().CaptureMouse(WidgetOwner).UseHighPrecisionMouseMovement(WidgetOwner);
+					FReply::Handled().CaptureMouse(WidgetOwner.AsShared()).UseHighPrecisionMouseMovement(WidgetOwner.AsShared());
 					SoftwareCursorPosition = MyGeometry.AbsoluteToLocal(MouseEvent.GetLastScreenSpacePosition());
 					bPanning = true;
 				}
@@ -559,7 +559,7 @@ void FVisualLoggerTimeSliderController::SetClampRange(float MinValue, float MaxV
 	SetTimeRange(ZoomDelta >= 1 ? MinValue : LocalViewRangeMin, ZoomDelta >= 1 ? MaxValue : LocalViewRangeMax);
 }
 
-FReply FVisualLoggerTimeSliderController::OnMouseWheel( TSharedRef<SWidget> WidgetOwner, const FGeometry& MyGeometry, const FPointerEvent& MouseEvent )
+FReply FVisualLoggerTimeSliderController::OnMouseWheel( SWidget& WidgetOwner, const FGeometry& MyGeometry, const FPointerEvent& MouseEvent )
 {
 	FReply ReturnValue = FReply::Unhandled();;
 
diff --git a/Engine/Source/Developer/LogVisualizer/Private/TimeSliderController.h b/Engine/Source/Developer/LogVisualizer/Private/TimeSliderController.h
--- a/Engine/Source/Developer/LogVisualizer/Private/TimeSliderController.h
+++ b/Engine/Source/Developer/LogVisualizer/Private/TimeSliderController.h
@@ -28,10 +28,10 @@ public:
 
 	/** ITimeSliderController Interface */
 	virtual int32 OnPaintTimeSlider( bool bMirrorLabels, const FGeometry& AllottedGeometry, const FSlateRect& MyClippingRect, FSlateWindowElementList& OutDrawElements, int32 LayerId, const FWidgetStyle& InWidgetStyle, bool bParentEnabled ) const override;
-	virtual FReply OnMouseButtonDown( TSharedRef<SWidget> WidgetOwner, const FGeometry& MyGeometry, const FPointerEvent& MouseEvent ) override;
-	virtual FReply OnMouseButtonUp( TSharedRef<SWidget> WidgetOwner, const FGeometry& MyGeometry, const FPointerEvent& MouseEvent ) override;
-	virtual FReply OnMouseMove( TSharedRef<SWidget> WidgetOwner, const FGeometry& MyGeometry, const FPointerEvent& MouseEvent ) override;
-	virtual FReply OnMouseWheel( TSharedRef<SWidget> WidgetOwner, const FGeometry& MyGeometry, const FPointerEvent& MouseEvent ) override;
+	virtual FReply OnMouseButtonDown( SWidget& WidgetOwner, const FGeometry& MyGeometry, const FPointerEvent& MouseEvent ) override;
+	virtual FReply OnMouseButtonUp( SWidget& WidgetOwner, const FGeometry& MyGeometry, const FPointerEvent& MouseEvent ) override;
+	virtual FReply OnMouseMove( SWidget& WidgetOwner, const FGeometry& MyGeometry, const FPointerEvent& MouseEvent ) override;
+	virtual FReply OnMouseWheel( SWidget& WidgetOwner, const FGeometry& MyGeometry, const FPointerEvent& MouseEvent ) override;
 	virtual FCursorReply OnCursorQuery( TSharedRef<const SWidget> WidgetOwner, const FGeometry& MyGeometry, const FPointerEvent& CursorEvent) const override { return FCursorReply::Unhandled(); }
 
 	/**
diff --git a/Engine/Source/Developer/MeshUtilities/Private/MeshUtilities.cpp b/Engine/Source/Developer/MeshUtilities/Private/MeshUtilities.cpp
--- a/Engine/Source/Developer/MeshUtilities/Private/MeshUtilities.cpp
+++ b/Engine/Source/Developer/MeshUtilities/Private/MeshUtilities.cpp
@@ -6670,7 +6670,7 @@ void FMeshUtilities::MergeStaticMeshComponents(const TArray<UStaticMeshComponent
 	MainTask.EnterProgressFrame(20);
 
 	FRawMeshExt MergedMesh;
-	FMemory::Memset(&MergedMesh, INDEX_NONE, sizeof(MergedMesh));
+	FMemory::Memset(&MergedMesh, 0, sizeof(MergedMesh));
 	
 	if (InSettings.bMergeMaterials)
 	{
diff --git a/Engine/Source/Editor/AnimGraph/Classes/AnimStateTransitionNode.h b/Engine/Source/Editor/AnimGraph/Classes/AnimStateTransitionNode.h
--- a/Engine/Source/Editor/AnimGraph/Classes/AnimStateTransitionNode.h
+++ b/Engine/Source/Editor/AnimGraph/Classes/AnimStateTransitionNode.h
@@ -44,7 +44,7 @@ class UAnimStateTransitionNode : public UAnimStateNodeBase
 	UPROPERTY(EditAnywhere, Category=Transition)
 	UBlendProfile* BlendProfile;
 
-	// Try setting the rule automatically based on the player node remaining time and the CrossfadeDuration, ignoring the internal time
+	// Try setting the rule automatically based on most relevant player node's remaining time and the CrossfadeDuration of the transition, ignoring the internal time
 	UPROPERTY(EditAnywhere, Category=Transition)
 	bool bAutomaticRuleBasedOnSequencePlayerInState;
 
diff --git a/Engine/Source/Editor/AnimGraph/Private/AnimStateTransitionNode.cpp b/Engine/Source/Editor/AnimGraph/Private/AnimStateTransitionNode.cpp
--- a/Engine/Source/Editor/AnimGraph/Private/AnimStateTransitionNode.cpp
+++ b/Engine/Source/Editor/AnimGraph/Private/AnimStateTransitionNode.cpp
@@ -506,7 +506,7 @@ void UAnimStateTransitionNode::ValidateNodeDuringCompilation(class FCompilerResu
 				}
 			}
 
-			if(!bHasNative)
+			if (!bHasNative && !bAutomaticRuleBasedOnSequencePlayerInState)
 			{
 				MessageLog.Warning(TEXT("@@ will never be taken, please connect something to @@"), this, BoolResultPin);
 			}
diff --git a/Engine/Source/Editor/EditorStyle/Private/SlateEditorStyle.cpp b/Engine/Source/Editor/EditorStyle/Private/SlateEditorStyle.cpp
--- a/Engine/Source/Editor/EditorStyle/Private/SlateEditorStyle.cpp
+++ b/Engine/Source/Editor/EditorStyle/Private/SlateEditorStyle.cpp
@@ -1595,7 +1595,10 @@ void FSlateEditorStyle::FStyle::SetupGeneralStyles()
 		Set( "Sequencer.NotificationImage_AddedPlayMovieSceneEvent", new IMAGE_BRUSH( "Old/Checkbox_checked", Icon16x16 ) );
 		Set( "Sequencer.Save", new IMAGE_BRUSH( "Icons/SequencerIcons/icon_Sequencer_Save_24x", Icon40x40 ) );
 		Set( "Sequencer.Save.Small", new IMAGE_BRUSH( "Icons/SequencerIcons/icon_Sequencer_Save_24x", Icon20x20 ) );
-		Set( "Sequencer.Key", new IMAGE_BRUSH( "Common/Key", FVector2D(7.0f,12.0f) ) );
+		Set( "Sequencer.KeyCircle", new IMAGE_BRUSH( "Sequencer/KeyCircle", Icon12x12 ) );
+		Set( "Sequencer.KeyDiamond", new IMAGE_BRUSH( "Sequencer/KeyDiamond", Icon12x12 ) );
+		Set( "Sequencer.KeySquare", new IMAGE_BRUSH( "Sequencer/KeySquare", Icon12x12 ) );
+		Set( "Sequencer.KeyTriangle", new IMAGE_BRUSH( "Sequencer/KeyTriangle", Icon12x12 ) );
 		Set( "Sequencer.KeyLeft", new IMAGE_BRUSH( "Common/KeyLeft", FVector2D( 7.0f, 12.0f ) ) );
 		Set( "Sequencer.KeyRight", new IMAGE_BRUSH( "Common/KeyRight", FVector2D( 7.0f, 12.0f ) ) );
 		Set( "Sequencer.PartialKey", new IMAGE_BRUSH( "Sequencer/PartialKey", FVector2D(11.f, 11.f) ) );
diff --git a/Engine/Source/Editor/HierarchicalLODOutliner/Private/HLODOutliner.cpp b/Engine/Source/Editor/HierarchicalLODOutliner/Private/HLODOutliner.cpp
--- a/Engine/Source/Editor/HierarchicalLODOutliner/Private/HLODOutliner.cpp
+++ b/Engine/Source/Editor/HierarchicalLODOutliner/Private/HLODOutliner.cpp
@@ -869,23 +869,27 @@ namespace HLODOutliner
 	{
 		if (CurrentWorld)
 		{
-			auto Level = CurrentWorld->GetCurrentLevel();
-			for (auto Actor : Level->Actors)
+			// Loop over all (streaming-)levels in the world
+			for (ULevel* Level : CurrentWorld->GetLevels())
 			{
-				ALODActor* LODActor = Cast<ALODActor>(Actor);
-				if (LODActor)
+				// For each LOD actor in the world update the drawing/transition distance
+				for (AActor* Actor : Level->Actors)
 				{
-					if (LODActor->LODLevel == LODLevelIndex + 1)
+					ALODActor* LODActor = Cast<ALODActor>(Actor);
+					if (LODActor)
 					{
-						if (!LODActor->IsDirty() && LODActor->GetStaticMeshComponent())
+						if (LODActor->LODLevel == LODLevelIndex + 1)
 						{
-							// At the moment this assumes a fixed field of view of 90 degrees (horizontal and vertical axi)
-							static const float FOVRad = 90.0f * (float)PI / 360.0f;
-							static const FMatrix ProjectionMatrix = FPerspectiveMatrix(FOVRad, 1920, 1080, 0.01f);						
-							FBoxSphereBounds Bounds = LODActor->GetStaticMeshComponent()->CalcBounds(FTransform());
-							LODActor->LODDrawDistance = FHierarchicalLODUtilities::CalculateDrawDistanceFromScreenSize(Bounds.SphereRadius, LODLevelTransitionScreenSizes[LODLevelIndex], ProjectionMatrix);
-							LODActor->UpdateSubActorLODParents();
-						}						
+							if (!LODActor->IsDirty() && LODActor->GetStaticMeshComponent())
+							{
+								// At the moment this assumes a fixed field of view of 90 degrees (horizontal and vertical axi)
+								static const float FOVRad = 90.0f * (float)PI / 360.0f;
+								static const FMatrix ProjectionMatrix = FPerspectiveMatrix(FOVRad, 1920, 1080, 0.01f);
+								FBoxSphereBounds Bounds = LODActor->GetStaticMeshComponent()->CalcBounds(FTransform());
+								LODActor->LODDrawDistance = FHierarchicalLODUtilities::CalculateDrawDistanceFromScreenSize(Bounds.SphereRadius, LODLevelTransitionScreenSizes[LODLevelIndex], ProjectionMatrix);
+								LODActor->UpdateSubActorLODParents();
+							}
+						}
 					}
 				}
 			}
diff --git a/Engine/Source/Editor/Kismet/Private/FindInBlueprintManager.cpp b/Engine/Source/Editor/Kismet/Private/FindInBlueprintManager.cpp
--- a/Engine/Source/Editor/Kismet/Private/FindInBlueprintManager.cpp
+++ b/Engine/Source/Editor/Kismet/Private/FindInBlueprintManager.cpp
@@ -671,7 +671,7 @@ namespace BlueprintSearchMetaDataHelpers
 			if(BlueprintSearchMetaDataHelpers::CheckIfJsonValueIsSearchable(JsonValue))
 			{
 				TSharedRef< FJsonValue > JsonValueAsSharedRef = JsonValue.ToSharedRef();
-				FJsonSerializer::Serialize(JsonValue, FFindInBlueprintSearchTags::FiB_DefaultValue.ToString(), StaticCastSharedRef<SearchMetaDataWriterParentClass>(InWriter) );
+				FJsonSerializer::Serialize(JsonValue, FFindInBlueprintSearchTags::FiB_DefaultValue.ToString(), StaticCastSharedRef<SearchMetaDataWriterParentClass>(InWriter), false );
 			}
 		}
 
@@ -749,14 +749,14 @@ namespace BlueprintSearchMetaDataHelpers
 				{
 					// Shallow conversion of property to string
 					auto JsonValue = FJsonObjectConverter::UPropertyToJsonValue(InProperty, InValue, 0, 0);
-					FJsonSerializer::Serialize(JsonValue, InProperty->GetName(), StaticCastSharedRef<SearchMetaDataWriterParentClass>(InWriter));
+					FJsonSerializer::Serialize(JsonValue, InProperty->GetName(), StaticCastSharedRef<SearchMetaDataWriterParentClass>(InWriter), false);
 				}
 			}
 		}
 		else
 		{
 			auto JsonValue = FJsonObjectConverter::UPropertyToJsonValue(InProperty, InValue, 0, 0);
-			FJsonSerializer::Serialize(JsonValue, InProperty->GetName(), StaticCastSharedRef<SearchMetaDataWriterParentClass>(InWriter));
+			FJsonSerializer::Serialize(JsonValue, InProperty->GetName(), StaticCastSharedRef<SearchMetaDataWriterParentClass>(InWriter), false);
 		}
 	}
 
diff --git a/Engine/Source/Editor/KismetCompiler/Private/AnimBlueprintCompiler.cpp b/Engine/Source/Editor/KismetCompiler/Private/AnimBlueprintCompiler.cpp
--- a/Engine/Source/Editor/KismetCompiler/Private/AnimBlueprintCompiler.cpp
+++ b/Engine/Source/Editor/KismetCompiler/Private/AnimBlueprintCompiler.cpp
@@ -989,26 +989,7 @@ void FAnimBlueprintCompiler::ProcessStateMachine(UAnimGraphNode_StateMachineBase
 			}
 
 			// Handle automatic time remaining rules
-			Rule.StateSequencePlayerToQueryIndex = INDEX_NONE;
-			if (TransitionNode->bAutomaticRuleBasedOnSequencePlayerInState)
-			{
-				if (UObject* SourceSequencePlayer = SimplePlayerStatesMap.FindRef(BakedState.StateRootNodeIndex))
-				{
-					//@TODO: This should be easier than it is... (the nodes are cloned so the pointers don't natrually match)
-					for (const auto& Pair : AllocatedAnimNodeIndices)
-					{
-						if (MessageLog.FindSourceObject(Pair.Key) == SourceSequencePlayer)
-						{
-							Rule.StateSequencePlayerToQueryIndex = Pair.Value;
-						}
-					}
-				}
-
-				if (Rule.StateSequencePlayerToQueryIndex == INDEX_NONE)
-				{
-					MessageLog.Error(*LOCTEXT("CantAutomaticallyDefineTransitionRule", "@@ must contain a single sequence player in order to use bAutomaticRuleBasedOnSequencePlayerInState on @@").ToString(), StateNode, TransitionNode);
-				}
-			}
+			Rule.bAutomaticRemainingTimeRule = TransitionNode->bAutomaticRuleBasedOnSequencePlayerInState;
 
 			// Handle custom transition graphs
 			Rule.CustomResultNodeIndex = INDEX_NONE;
diff --git a/Engine/Source/Editor/MovieSceneCaptureDialog/Private/MovieSceneCaptureDialogModule.cpp b/Engine/Source/Editor/MovieSceneCaptureDialog/Private/MovieSceneCaptureDialogModule.cpp
--- a/Engine/Source/Editor/MovieSceneCaptureDialog/Private/MovieSceneCaptureDialogModule.cpp
+++ b/Engine/Source/Editor/MovieSceneCaptureDialog/Private/MovieSceneCaptureDialogModule.cpp
@@ -13,6 +13,7 @@
 #include "EditorStyle.h"
 #include "Editor.h"
 #include "PropertyEditing.h"
+#include "FileHelpers.h"
 
 #include "ISessionServicesModule.h"
 #include "ISessionInstanceInfo.h"
@@ -342,12 +343,39 @@ class FMovieSceneCaptureDialogModule : public IMovieSceneCaptureDialogModule
 
 	FText OnStartCapture(UMovieSceneCapture* CaptureObject)
 	{
-		TArray<FString> SavedMapNames;
-		GEditor->SaveWorldForPlay(SavedMapNames);
+		FString MapNameToLoad;
+		if( CaptureObject->Settings.bCreateTemporaryCopiesOfLevels )
+		{
+			TArray<FString> SavedMapNames;
+			GEditor->SaveWorldForPlay(SavedMapNames);
+
+			if (SavedMapNames.Num() == 0)
+			{
+				return LOCTEXT("CouldNotSaveMap", "Could not save map for movie capture.");
+			}
 
-		if (SavedMapNames.Num() == 0)
+			MapNameToLoad = SavedMapNames[ 0 ];
+		}
+		else
 		{
-			return LOCTEXT("CouldNotSaveMap", "Could not save map for movie capture.");
+			// Prompt the user to save their changes so that they'll be in the movie, since we're not saving temporary copies of the level.
+			bool bPromptUserToSave = true;
+			bool bSaveMapPackages = true;
+			bool bSaveContentPackages = true;
+			if( !FEditorFileUtils::SaveDirtyPackages( bPromptUserToSave, bSaveMapPackages, bSaveContentPackages ) )
+			{
+				return LOCTEXT( "UserCancelled", "Capturing was cancelled from the save dialog." );
+			}
+
+			const FString WorldPackageName = GWorld->GetOutermost()->GetName();
+			MapNameToLoad = WorldPackageName;
+		}
+
+		// Allow the game mode to be overridden
+		if( CaptureObject->Settings.GameModeOverride != nullptr )
+		{
+			const FString GameModeName = CaptureObject->Settings.GameModeOverride->GetPathName();
+			MapNameToLoad += FString::Printf( TEXT( "?game=%s" ), *GameModeName );
 		}
 
 		if (InProgressCaptureNotification.IsValid())
@@ -384,7 +412,16 @@ class FMovieSceneCaptureDialogModule : public IMovieSceneCaptureDialogModule
 			return LOCTEXT("UnableToSaveCaptureManifest", "Unable to save capture manifest");
 		}
 
-		FString EditorCommandLine = FString::Printf(TEXT("%s -MovieSceneCaptureManifest=\"%s\" -game"), *SavedMapNames[0], *Filename);
+		FString EditorCommandLine = FString::Printf(TEXT("%s -MovieSceneCaptureManifest=\"%s\" -game"), *MapNameToLoad, *Filename);
+
+		if( CaptureObject->Settings.bCreateTemporaryCopiesOfLevels )
+		{
+			// the PIEVIACONSOLE parameter tells UGameEngine to add the auto-save dir to the paths array and repopulate the package file cache
+			// this is needed in order to support streaming levels as the streaming level packages will be loaded only when needed (thus
+			// their package names need to be findable by the package file caching system)
+			// (we add to EditorCommandLine because the URL is ignored by WindowsTools)
+			EditorCommandLine.Append( TEXT( " -PIEVIACONSOLE" ) );
+		}
 
 		// Spit out any additional, user-supplied command line args
 		if (!CaptureObject->AdditionalCommandLineArguments.IsEmpty())
diff --git a/Engine/Source/Editor/MovieSceneTools/Private/CurveKeyAreas/BoolKeyArea.cpp b/Engine/Source/Editor/MovieSceneTools/Private/CurveKeyAreas/BoolKeyArea.cpp
new file mode 100644
--- /dev/null
+++ b/Engine/Source/Editor/MovieSceneTools/Private/CurveKeyAreas/BoolKeyArea.cpp
@@ -0,0 +1,26 @@
+// Copyright 1998-2015 Epic Games, Inc. All Rights Reserved.
+
+#include "MovieSceneToolsPrivatePCH.h"
+#include "BoolKeyArea.h"
+#include "SBoolCurveKeyEditor.h"
+
+
+/* IKeyArea interface
+ *****************************************************************************/
+
+bool FBoolKeyArea::CanCreateKeyEditor()
+{
+	return true;
+}
+
+
+TSharedRef<SWidget> FBoolKeyArea::CreateKeyEditor(ISequencer* Sequencer)
+{
+	return SNew(SBoolCurveKeyEditor)
+		.Sequencer(Sequencer)
+		.OwningSection(OwningSection)
+		.Curve(&Curve)
+		.IntermediateValue_Lambda([this] {
+			return IntermediateValue;
+		});
+};
diff --git a/Engine/Source/Editor/MovieSceneTools/Private/CurveKeyAreas/ByteKeyArea.cpp b/Engine/Source/Editor/MovieSceneTools/Private/CurveKeyAreas/ByteKeyArea.cpp
new file mode 100644
--- /dev/null
+++ b/Engine/Source/Editor/MovieSceneTools/Private/CurveKeyAreas/ByteKeyArea.cpp
@@ -0,0 +1,26 @@
+// Copyright 1998-2015 Epic Games, Inc. All Rights Reserved.
+
+#include "MovieSceneToolsPrivatePCH.h"
+#include "ByteKeyArea.h"
+#include "SIntegralCurveKeyEditor.h"
+
+
+/* IKeyArea interface
+ *****************************************************************************/
+
+bool FByteKeyArea::CanCreateKeyEditor()
+{
+	return true;
+}
+
+
+TSharedRef<SWidget> FByteKeyArea::CreateKeyEditor(ISequencer* Sequencer)
+{
+	return SNew(SIntegralCurveKeyEditor)
+		.Sequencer(Sequencer)
+		.OwningSection(OwningSection)
+		.Curve(&Curve)
+		.IntermediateValue_Lambda( [this] {
+			return IntermediateValue;
+		});
+};
diff --git a/Engine/Source/Editor/MovieSceneTools/Private/CurveKeyAreas/EnumKeyArea.cpp b/Engine/Source/Editor/MovieSceneTools/Private/CurveKeyAreas/EnumKeyArea.cpp
new file mode 100644
--- /dev/null
+++ b/Engine/Source/Editor/MovieSceneTools/Private/CurveKeyAreas/EnumKeyArea.cpp
@@ -0,0 +1,27 @@
+// Copyright 1998-2015 Epic Games, Inc. All Rights Reserved.
+
+#include "MovieSceneToolsPrivatePCH.h"
+#include "EnumKeyArea.h"
+#include "SEnumCurveKeyEditor.h"
+
+
+/* IKeyArea interface
+ *****************************************************************************/
+
+bool FEnumKeyArea::CanCreateKeyEditor()
+{
+	return true;
+}
+
+
+TSharedRef<SWidget> FEnumKeyArea::CreateKeyEditor(ISequencer* Sequencer)
+{
+	return SNew(SEnumCurveKeyEditor)
+		.Sequencer(Sequencer)
+		.OwningSection(OwningSection)
+		.Curve(&Curve)
+		.Enum(Enum)
+		.IntermediateValue_Lambda([this] {
+			return IntermediateValue;
+		});
+};
diff --git a/Engine/Source/Editor/MovieSceneTools/Private/CurveKeyAreas/FloatCurveKeyArea.cpp b/Engine/Source/Editor/MovieSceneTools/Private/CurveKeyAreas/FloatCurveKeyArea.cpp
new file mode 100644
--- /dev/null
+++ b/Engine/Source/Editor/MovieSceneTools/Private/CurveKeyAreas/FloatCurveKeyArea.cpp
@@ -0,0 +1,196 @@
+// Copyright 1998-2015 Epic Games, Inc. All Rights Reserved.
+
+#include "MovieSceneToolsPrivatePCH.h"
+#include "FloatCurveKeyArea.h"
+#include "SFloatCurveKeyEditor.h"
+
+
+/* IKeyArea interface
+ *****************************************************************************/
+
+TArray<FKeyHandle> FFloatCurveKeyArea::AddKeyUnique(float Time, EMovieSceneKeyInterpolation InKeyInterpolation, float TimeToCopyFrom)
+{
+	TArray<FKeyHandle> AddedKeyHandles;
+	FKeyHandle CurrentKeyHandle = Curve->FindKey(Time);
+
+	if (Curve->IsKeyHandleValid(CurrentKeyHandle) == false)
+	{
+		if (OwningSection->GetStartTime() > Time)
+		{
+			OwningSection->SetStartTime(Time);
+		}
+
+		if (OwningSection->GetEndTime() < Time)
+		{
+			OwningSection->SetEndTime(Time);
+		}
+
+		float Value = Curve->Eval(Time);
+
+		if (TimeToCopyFrom != FLT_MAX)
+		{
+			Value = Curve->Eval(TimeToCopyFrom);
+		}
+		else if ( IntermediateValue.IsSet() )
+		{
+			Value = IntermediateValue.GetValue();
+		}
+
+		Curve->AddKey(Time, Value, false, CurrentKeyHandle);
+		AddedKeyHandles.Add(CurrentKeyHandle);
+		
+		MovieSceneHelpers::SetKeyInterpolation(*Curve, CurrentKeyHandle, InKeyInterpolation);		
+		
+		// Copy the properties from the key if it exists
+		FKeyHandle KeyHandleToCopy = Curve->FindKey(TimeToCopyFrom);
+
+		if (Curve->IsKeyHandleValid(KeyHandleToCopy))
+		{
+			FRichCurveKey& CurrentKey = Curve->GetKey(CurrentKeyHandle);
+			FRichCurveKey& KeyToCopy = Curve->GetKey(KeyHandleToCopy);
+			CurrentKey.InterpMode = KeyToCopy.InterpMode;
+			CurrentKey.TangentMode = KeyToCopy.TangentMode;
+			CurrentKey.TangentWeightMode = KeyToCopy.TangentWeightMode;
+			CurrentKey.ArriveTangent = KeyToCopy.ArriveTangent;
+			CurrentKey.LeaveTangent = KeyToCopy.LeaveTangent;
+			CurrentKey.ArriveTangentWeight = KeyToCopy.ArriveTangentWeight;
+			CurrentKey.LeaveTangentWeight = KeyToCopy.LeaveTangentWeight;
+		}
+	}
+
+	return AddedKeyHandles;
+}
+
+
+bool FFloatCurveKeyArea::CanCreateKeyEditor()
+{
+	return true;
+}
+
+
+TSharedRef<SWidget> FFloatCurveKeyArea::CreateKeyEditor(ISequencer* Sequencer)
+{
+	return SNew(SFloatCurveKeyEditor)
+		.Sequencer(Sequencer)
+		.OwningSection(OwningSection)
+		.Curve(Curve)
+		.IntermediateValue_Lambda([this] {
+			return IntermediateValue;
+		});
+};
+
+
+void FFloatCurveKeyArea::DeleteKey(FKeyHandle KeyHandle)
+{
+	if (Curve->IsKeyHandleValid(KeyHandle))
+	{
+		Curve->DeleteKey(KeyHandle);
+	}
+}
+
+
+ERichCurveExtrapolation FFloatCurveKeyArea::GetExtrapolationMode(bool bPreInfinity) const
+{
+	if (bPreInfinity)
+	{
+		return Curve->PreInfinityExtrap;
+	}
+
+	return Curve->PostInfinityExtrap;
+}
+
+
+ERichCurveInterpMode FFloatCurveKeyArea::GetKeyInterpMode(FKeyHandle KeyHandle) const
+{
+	if (Curve->IsKeyHandleValid(KeyHandle))
+	{
+		return Curve->GetKeyInterpMode(KeyHandle);
+	}
+
+	return RCIM_None;
+}
+
+
+ERichCurveTangentMode FFloatCurveKeyArea::GetKeyTangentMode(FKeyHandle KeyHandle) const
+{
+	if (Curve->IsKeyHandleValid(KeyHandle))
+	{
+		return Curve->GetKeyTangentMode(KeyHandle);
+	}
+
+	return RCTM_None;
+}
+
+
+float FFloatCurveKeyArea::GetKeyTime(FKeyHandle KeyHandle) const
+{
+	return Curve->GetKeyTime(KeyHandle);
+}
+
+
+UMovieSceneSection* FFloatCurveKeyArea::GetOwningSection()
+{
+	return OwningSection;
+}
+
+
+FRichCurve* FFloatCurveKeyArea::GetRichCurve()
+{
+	return Curve;
+}
+
+
+TArray<FKeyHandle> FFloatCurveKeyArea::GetUnsortedKeyHandles() const
+{
+	TArray<FKeyHandle> OutKeyHandles;
+
+	for (auto It(Curve->GetKeyHandleIterator()); It; ++It)
+	{
+		OutKeyHandles.Add(It.Key());
+	}
+
+	return OutKeyHandles;
+}
+
+
+FKeyHandle FFloatCurveKeyArea::MoveKey(FKeyHandle KeyHandle, float DeltaPosition)
+{
+	return Curve->SetKeyTime(KeyHandle, Curve->GetKeyTime(KeyHandle) + DeltaPosition);
+}
+
+
+void FFloatCurveKeyArea::SetExtrapolationMode(ERichCurveExtrapolation ExtrapMode, bool bPreInfinity)
+{
+	if (bPreInfinity)
+	{
+		Curve->PreInfinityExtrap = ExtrapMode;
+	}
+	else
+	{
+		Curve->PostInfinityExtrap = ExtrapMode;
+	}
+}
+
+
+void FFloatCurveKeyArea::SetKeyInterpMode(FKeyHandle KeyHandle, ERichCurveInterpMode InterpMode)
+{
+	if (Curve->IsKeyHandleValid(KeyHandle))
+	{
+		Curve->SetKeyInterpMode(KeyHandle, InterpMode);
+	}
+}
+
+
+void FFloatCurveKeyArea::SetKeyTangentMode(FKeyHandle KeyHandle, ERichCurveTangentMode TangentMode)
+{
+	if (Curve->IsKeyHandleValid(KeyHandle))
+	{
+		Curve->SetKeyTangentMode(KeyHandle, TangentMode);
+	}
+}
+
+
+void FFloatCurveKeyArea::SetKeyTime(FKeyHandle KeyHandle, float NewKeyTime) const
+{
+	Curve->SetKeyTime(KeyHandle, NewKeyTime);
+}
diff --git a/Engine/Source/Editor/MovieSceneTools/Private/CurveKeyAreas/IntegralKeyArea.cpp b/Engine/Source/Editor/MovieSceneTools/Private/CurveKeyAreas/IntegralKeyArea.cpp
new file mode 100644
--- /dev/null
+++ b/Engine/Source/Editor/MovieSceneTools/Private/CurveKeyAreas/IntegralKeyArea.cpp
@@ -0,0 +1,117 @@
+// Copyright 1998-2015 Epic Games, Inc. All Rights Reserved.
+
+#include "MovieSceneToolsPrivatePCH.h"
+#include "IntegralKeyArea.h"
+
+
+/* IKeyArea interface
+ *****************************************************************************/
+
+TArray<FKeyHandle> FIntegralCurveKeyAreaBase::AddKeyUnique(float Time, EMovieSceneKeyInterpolation InKeyInterpolation, float TimeToCopyFrom)
+{
+	TArray<FKeyHandle> AddedKeyHandles;
+	FKeyHandle CurrentKey = Curve.FindKey(Time);
+
+	if (Curve.IsKeyHandleValid(CurrentKey) == false)
+	{
+		if (OwningSection->GetStartTime() > Time)
+		{
+			OwningSection->SetStartTime(Time);
+		}
+
+		if (OwningSection->GetEndTime() < Time)
+		{
+			OwningSection->SetEndTime(Time);
+		}
+
+		EvaluateAndAddKey(Time, TimeToCopyFrom, CurrentKey);
+		AddedKeyHandles.Add(CurrentKey);
+	}
+
+	return AddedKeyHandles;
+}
+
+
+void FIntegralCurveKeyAreaBase::DeleteKey(FKeyHandle KeyHandle)
+{
+	Curve.DeleteKey(KeyHandle);
+}
+
+
+ERichCurveExtrapolation FIntegralCurveKeyAreaBase::GetExtrapolationMode(bool bPreInfinity) const
+{
+	return RCCE_None;
+}
+
+
+ERichCurveTangentMode FIntegralCurveKeyAreaBase::GetKeyTangentMode(FKeyHandle KeyHandle) const
+{
+	return RCTM_None;
+}
+
+
+ERichCurveInterpMode FIntegralCurveKeyAreaBase::GetKeyInterpMode(FKeyHandle KeyHandle) const
+{
+	return RCIM_None;
+}
+
+
+UMovieSceneSection* FIntegralCurveKeyAreaBase::GetOwningSection()
+{
+	return OwningSection;
+}
+
+
+float FIntegralCurveKeyAreaBase::GetKeyTime(FKeyHandle KeyHandle) const
+{
+	return Curve.GetKeyTime(KeyHandle);
+}
+
+
+FRichCurve* FIntegralCurveKeyAreaBase::GetRichCurve()
+{
+	return nullptr;
+};
+
+
+TArray<FKeyHandle> FIntegralCurveKeyAreaBase::GetUnsortedKeyHandles() const
+{
+	TArray<FKeyHandle> OutKeyHandles;
+
+	for (auto It(Curve.GetKeyHandleIterator()); It; ++It)
+	{
+		OutKeyHandles.Add(It.Key());
+	}
+
+	return OutKeyHandles;
+}
+
+
+FKeyHandle FIntegralCurveKeyAreaBase::MoveKey(FKeyHandle KeyHandle, float DeltaPosition)
+{
+	return Curve.SetKeyTime(KeyHandle, Curve.GetKeyTime(KeyHandle) + DeltaPosition);
+}
+
+
+void FIntegralCurveKeyAreaBase::SetExtrapolationMode(ERichCurveExtrapolation ExtrapMode, bool bPreInfinity)
+{
+	// do nothing
+}
+
+
+void FIntegralCurveKeyAreaBase::SetKeyInterpMode(FKeyHandle KeyHandle, ERichCurveInterpMode InterpMode)
+{
+	// do nothing
+}
+
+
+void FIntegralCurveKeyAreaBase::SetKeyTangentMode(FKeyHandle KeyHandle, ERichCurveTangentMode TangentMode)
+{
+	// do nothing
+}
+
+
+void FIntegralCurveKeyAreaBase::SetKeyTime(FKeyHandle KeyHandle, float NewKeyTime) const
+{
+	Curve.SetKeyTime(KeyHandle, NewKeyTime);
+}
diff --git a/Engine/Source/Editor/MovieSceneTools/Private/CurveKeyAreas/NameCurveKeyArea.cpp b/Engine/Source/Editor/MovieSceneTools/Private/CurveKeyAreas/NameCurveKeyArea.cpp
new file mode 100644
--- /dev/null
+++ b/Engine/Source/Editor/MovieSceneTools/Private/CurveKeyAreas/NameCurveKeyArea.cpp
@@ -0,0 +1,130 @@
+// Copyright 1998-2015 Epic Games, Inc. All Rights Reserved.
+
+#include "MovieSceneToolsPrivatePCH.h"
+#include "NameCurveKeyArea.h"
+
+
+/* IKeyArea interface
+ *****************************************************************************/
+
+TArray<FKeyHandle> FNameCurveKeyArea::AddKeyUnique(float Time, EMovieSceneKeyInterpolation InKeyInterpolation, float TimeToCopyFrom)
+{
+	TArray<FKeyHandle> AddedKeyHandles;
+	FKeyHandle CurrentKey = Curve.FindKey(Time);
+
+	if (Curve.IsKeyHandleValid(CurrentKey) == false)
+	{
+		if (OwningSection->GetStartTime() > Time)
+		{
+			OwningSection->SetStartTime(Time);
+		}
+
+		if (OwningSection->GetEndTime() < Time)
+		{
+			OwningSection->SetEndTime(Time);
+		}
+
+		Curve.AddKey(Time, NAME_None, CurrentKey);
+		AddedKeyHandles.Add(CurrentKey);
+	}
+
+	return AddedKeyHandles;
+}
+
+
+bool FNameCurveKeyArea::CanCreateKeyEditor()
+{
+	return false;
+}
+
+
+TSharedRef<SWidget> FNameCurveKeyArea::CreateKeyEditor(ISequencer* Sequencer)
+{
+	return SNullWidget::NullWidget;
+}
+
+
+void FNameCurveKeyArea::DeleteKey(FKeyHandle KeyHandle)
+{
+	Curve.DeleteKey(KeyHandle);
+}
+
+
+ERichCurveExtrapolation FNameCurveKeyArea::GetExtrapolationMode(bool bPreInfinity) const
+{
+	return RCCE_None;
+}
+
+
+ERichCurveInterpMode FNameCurveKeyArea::GetKeyInterpMode(FKeyHandle Keyhandle) const
+{
+	return RCIM_None;
+}
+
+
+ERichCurveTangentMode FNameCurveKeyArea::GetKeyTangentMode(FKeyHandle KeyHandle) const
+{
+	return RCTM_None;
+}
+
+
+float FNameCurveKeyArea::GetKeyTime(FKeyHandle KeyHandle) const
+{
+	return Curve.GetKeyTime(KeyHandle);
+}
+
+
+UMovieSceneSection* FNameCurveKeyArea::GetOwningSection()
+{
+	return OwningSection.Get();
+}
+
+
+
+FRichCurve* FNameCurveKeyArea::GetRichCurve()
+{
+	return nullptr;
+}
+
+
+TArray<FKeyHandle> FNameCurveKeyArea::GetUnsortedKeyHandles() const
+{
+	TArray<FKeyHandle> OutKeyHandles;
+
+	for (auto It(Curve.GetKeyHandleIterator()); It; ++It)
+	{
+		OutKeyHandles.Add(It.Key());
+	}
+
+	return OutKeyHandles;
+}
+
+
+FKeyHandle FNameCurveKeyArea::MoveKey( FKeyHandle KeyHandle, float DeltaPosition )
+{
+	return Curve.SetKeyTime(KeyHandle, Curve.GetKeyTime(KeyHandle) + DeltaPosition);
+}
+	
+
+void FNameCurveKeyArea::SetExtrapolationMode(ERichCurveExtrapolation ExtrapMode, bool bPreInfinity)
+{
+	// do nothing
+}
+
+
+void FNameCurveKeyArea::SetKeyInterpMode(FKeyHandle KeyHandle, ERichCurveInterpMode InterpMode)
+{
+	// do nothing
+}
+
+
+void FNameCurveKeyArea::SetKeyTangentMode(FKeyHandle KeyHandle, ERichCurveTangentMode TangentMode)
+{
+	// do nothing
+}
+
+
+void FNameCurveKeyArea::SetKeyTime(FKeyHandle KeyHandle, float NewKeyTime) const
+{
+	Curve.SetKeyTime(KeyHandle, NewKeyTime);
+}
diff --git a/Engine/Source/Editor/MovieSceneTools/Private/CurveKeyEditors/SFloatCurveKeyEditor.cpp b/Engine/Source/Editor/MovieSceneTools/Private/CurveKeyEditors/SFloatCurveKeyEditor.cpp
--- a/Engine/Source/Editor/MovieSceneTools/Private/CurveKeyEditors/SFloatCurveKeyEditor.cpp
+++ b/Engine/Source/Editor/MovieSceneTools/Private/CurveKeyEditors/SFloatCurveKeyEditor.cpp
@@ -75,6 +75,8 @@ void SFloatCurveKeyEditor::OnValueChanged(float Value)
 		else
 		{
 			Curve->AddKey(CurrentTime, Value, false, CurrentKeyHandle);
+
+			MovieSceneHelpers::SetKeyInterpolation(*Curve, CurrentKeyHandle, Sequencer->GetKeyInterpolation());
 		}
 
 		if (OwningSection->GetStartTime() > CurrentTime)
diff --git a/Engine/Source/Editor/MovieSceneTools/Private/MatineeImportTools.cpp b/Engine/Source/Editor/MovieSceneTools/Private/MatineeImportTools.cpp
new file mode 100644
--- /dev/null
+++ b/Engine/Source/Editor/MovieSceneTools/Private/MatineeImportTools.cpp
@@ -0,0 +1,262 @@
+// Copyright 1998-2015 Epic Games, Inc. All Rights Reserved.
+
+#include "MovieSceneToolsPrivatePCH.h"
+#include "MatineeImportTools.h"
+
+#include "Matinee/InterpTrackFloatBase.h"
+#include "Matinee/InterpTrackMove.h"
+#include "Matinee/InterpTrackAnimControl.h"
+
+#include "MovieSceneFloatTrack.h"
+#include "MovieScene3dTransformTrack.h"
+#include "MovieSceneParticleTrack.h"
+#include "MovieSceneSkeletalAnimationTrack.h"
+
+#include "MovieSceneFloatSection.h"
+#include "MovieScene3dTransformSection.h"
+#include "MovieSceneSkeletalAnimationSection.h"
+
+
+ERichCurveInterpMode FMatineeImportTools::MatineeInterpolationToRichCurveInterpolation( EInterpCurveMode CurveMode )
+{
+	switch ( CurveMode )
+	{
+	case CIM_Constant:
+		return ERichCurveInterpMode::RCIM_Constant;
+	case CIM_CurveAuto:
+	case CIM_CurveAutoClamped:
+	case CIM_CurveBreak:
+	case CIM_CurveUser:
+		return ERichCurveInterpMode::RCIM_Cubic;
+	case CIM_Linear:
+		return ERichCurveInterpMode::RCIM_Linear;
+	default:
+		return ERichCurveInterpMode::RCIM_None;
+	}
+}
+
+
+ERichCurveTangentMode FMatineeImportTools::MatineeInterpolationToRichCurveTangent( EInterpCurveMode CurveMode )
+{
+	switch ( CurveMode )
+	{
+	case CIM_CurveBreak:
+		return ERichCurveTangentMode::RCTM_Break;
+	case CIM_CurveUser:
+		return ERichCurveTangentMode::RCTM_User;
+	default:
+		return ERichCurveTangentMode::RCTM_Auto;
+	}
+}
+
+
+bool FMatineeImportTools::TryConvertMatineeToggleToOutParticleKey( ETrackToggleAction ToggleAction, EParticleKey::Type& OutParticleKey )
+{
+	switch ( ToggleAction )
+	{
+	case ETrackToggleAction::ETTA_On:
+		OutParticleKey = EParticleKey::Activate;
+		return true;
+	case ETrackToggleAction::ETTA_Off:
+		OutParticleKey = EParticleKey::Deactivate;
+		return true;
+	case ETrackToggleAction::ETTA_Trigger:
+		OutParticleKey = EParticleKey::Trigger;
+		return true;
+	}
+	return false;
+}
+
+
+void FMatineeImportTools::SetOrAddKey( FRichCurve& Curve, float Time, float Value, float ArriveTangent, float LeaveTangent, EInterpCurveMode MatineeInterpMode )
+{
+	FKeyHandle KeyHandle = Curve.AddKey( Time, Value, false, Curve.FindKey( Time ) );
+	FRichCurveKey& Key = Curve.GetKey( KeyHandle );
+	Key.ArriveTangent = ArriveTangent;
+	Key.LeaveTangent = LeaveTangent;
+	Key.InterpMode = MatineeInterpolationToRichCurveInterpolation( MatineeInterpMode );
+	Key.TangentMode = MatineeInterpolationToRichCurveTangent( MatineeInterpMode );
+}
+
+
+void FMatineeImportTools::CopyInterpFloatTrack( TSharedRef<ISequencer> Sequencer, UInterpTrackFloatBase* MatineeFloatTrack, UMovieSceneFloatTrack* FloatTrack )
+{
+	const FScopedTransaction Transaction( NSLOCTEXT( "Sequencer", "PasteMatineeFloatTrack", "Paste Matinee float track" ) );
+	bool bSectionCreated = false;
+
+	FloatTrack->Modify();
+
+	float KeyTime = MatineeFloatTrack->GetKeyframeTime( 0 );
+	UMovieSceneFloatSection* Section = Cast<UMovieSceneFloatSection>( MovieSceneHelpers::FindSectionAtTime( FloatTrack->GetAllSections(), KeyTime ) );
+	if ( Section == nullptr )
+	{
+		Section = Cast<UMovieSceneFloatSection>( FloatTrack->CreateNewSection() );
+		FloatTrack->AddSection( *Section );
+		bSectionCreated = true;
+	}
+	Section->Modify();
+
+	float SectionMin = Section->GetStartTime();
+	float SectionMax = Section->GetEndTime();
+
+	FRichCurve& FloatCurve = Section->GetFloatCurve();
+	for ( const auto& Point : MatineeFloatTrack->FloatTrack.Points )
+	{
+		FMatineeImportTools::SetOrAddKey( FloatCurve, Point.InVal, Point.OutVal, Point.ArriveTangent, Point.LeaveTangent, Point.InterpMode );
+		SectionMin = FMath::Min( SectionMin, Point.InVal );
+		SectionMax = FMath::Max( SectionMax, Point.InVal );
+	}
+
+	Section->SetStartTime( SectionMin );
+	Section->SetEndTime( SectionMax );
+
+	if ( bSectionCreated )
+	{
+		Sequencer->NotifyMovieSceneDataChanged();
+	}
+}
+
+
+void FMatineeImportTools::CopyInterpMoveTrack( TSharedRef<ISequencer> Sequencer, UInterpTrackMove* MoveTrack, UMovieScene3DTransformTrack* TransformTrack )
+{
+	const FScopedTransaction Transaction( NSLOCTEXT( "Sequencer", "PasteMatineeMoveTrack", "Paste Matinee move track" ) );
+	bool bSectionCreated = false;
+
+	TransformTrack->Modify();
+
+	float KeyTime = MoveTrack->GetKeyframeTime( 0 );
+	UMovieScene3DTransformSection* Section = Cast<UMovieScene3DTransformSection>( MovieSceneHelpers::FindSectionAtTime( TransformTrack->GetAllSections(), KeyTime ) );
+	if ( Section == nullptr )
+	{
+		Section = Cast<UMovieScene3DTransformSection>( TransformTrack->CreateNewSection() );
+		Section->GetScaleCurve(EAxis::X).SetDefaultValue(1);
+		Section->GetScaleCurve(EAxis::Y).SetDefaultValue(1);
+		Section->GetScaleCurve(EAxis::Z).SetDefaultValue(1);
+		TransformTrack->AddSection( *Section );
+		bSectionCreated = true;
+	}
+	Section->Modify();
+
+	float SectionMin = Section->GetStartTime();
+	float SectionMax = Section->GetEndTime();
+
+	FRichCurve& TranslationXCurve = Section->GetTranslationCurve( EAxis::X );
+	FRichCurve& TranslationYCurve = Section->GetTranslationCurve( EAxis::Y );
+	FRichCurve& TranslationZCurve = Section->GetTranslationCurve( EAxis::Z );
+
+	for ( const auto& Point : MoveTrack->PosTrack.Points )
+	{
+		FMatineeImportTools::SetOrAddKey( TranslationXCurve, Point.InVal, Point.OutVal.X, Point.ArriveTangent.X, Point.LeaveTangent.X, Point.InterpMode );
+		FMatineeImportTools::SetOrAddKey( TranslationYCurve, Point.InVal, Point.OutVal.Y, Point.ArriveTangent.Y, Point.LeaveTangent.Y, Point.InterpMode );
+		FMatineeImportTools::SetOrAddKey( TranslationZCurve, Point.InVal, Point.OutVal.Z, Point.ArriveTangent.Z, Point.LeaveTangent.Z, Point.InterpMode );
+		SectionMin = FMath::Min( SectionMin, Point.InVal );
+		SectionMax = FMath::Max( SectionMax, Point.InVal );
+	}
+
+	FRichCurve& RotationXCurve = Section->GetRotationCurve( EAxis::X );
+	FRichCurve& RotationYCurve = Section->GetRotationCurve( EAxis::Y );
+	FRichCurve& RotationZCurve = Section->GetRotationCurve( EAxis::Z );
+
+	for ( const auto& Point : MoveTrack->EulerTrack.Points )
+	{
+		FMatineeImportTools::SetOrAddKey( RotationXCurve, Point.InVal, Point.OutVal.X, Point.ArriveTangent.X, Point.LeaveTangent.X, Point.InterpMode );
+		FMatineeImportTools::SetOrAddKey( RotationYCurve, Point.InVal, Point.OutVal.Y, Point.ArriveTangent.Y, Point.LeaveTangent.Y, Point.InterpMode );
+		FMatineeImportTools::SetOrAddKey( RotationZCurve, Point.InVal, Point.OutVal.Z, Point.ArriveTangent.Z, Point.LeaveTangent.Z, Point.InterpMode );
+		SectionMin = FMath::Min( SectionMin, Point.InVal );
+		SectionMax = FMath::Max( SectionMax, Point.InVal );
+	}
+
+	Section->SetStartTime( SectionMin );
+	Section->SetEndTime( SectionMax );
+
+	if ( bSectionCreated )
+	{
+		Sequencer->NotifyMovieSceneDataChanged();
+	}
+}
+
+
+void FMatineeImportTools::CopyInterpParticleTrack( TSharedRef<ISequencer> Sequencer, UInterpTrackToggle* MatineeToggleTrack, UMovieSceneParticleTrack* ParticleTrack )
+{
+	const FScopedTransaction Transaction( NSLOCTEXT( "Sequencer", "PasteMatineeParticleTrack", "Paste Matinee particle track" ) );
+	bool bSectionCreated = false;
+
+	ParticleTrack->Modify();
+
+	float KeyTime = MatineeToggleTrack->GetKeyframeTime( 0 );
+	UMovieSceneParticleSection* Section = Cast<UMovieSceneParticleSection>( MovieSceneHelpers::FindSectionAtTime( ParticleTrack->GetAllSections(), KeyTime ) );
+	if ( Section == nullptr )
+	{
+		Section = Cast<UMovieSceneParticleSection>( ParticleTrack->CreateNewSection() );
+		ParticleTrack->AddSection( *Section );
+		bSectionCreated = true;
+	}
+	Section->Modify();
+
+	float SectionMin = Section->GetStartTime();
+	float SectionMax = Section->GetEndTime();
+
+	FIntegralCurve& ParticleCurve = Section->GetParticleCurve();
+	for ( const auto& Key : MatineeToggleTrack->ToggleTrack )
+	{
+		EParticleKey::Type ParticleKey;
+		if ( TryConvertMatineeToggleToOutParticleKey( Key.ToggleAction, ParticleKey ) )
+		{
+			ParticleCurve.AddKey( Key.Time, (int32)ParticleKey, ParticleCurve.FindKey( Key.Time ) );
+		}
+		SectionMin = FMath::Min( SectionMin, Key.Time );
+		SectionMax = FMath::Max( SectionMax, Key.Time );
+	}
+
+	Section->SetStartTime( SectionMin );
+	Section->SetEndTime( SectionMax );
+
+	if ( bSectionCreated )
+	{
+		Sequencer->NotifyMovieSceneDataChanged();
+	}
+}
+
+
+void FMatineeImportTools::CopyInterpAnimControlTrack( TSharedRef<ISequencer> Sequencer, UInterpTrackAnimControl* MatineeAnimControlTrack, UMovieSceneSkeletalAnimationTrack* SkeletalAnimationTrack )
+{
+	// @todo - Sequencer - Add support for slot names once they are implemented.
+	const FScopedTransaction Transaction( NSLOCTEXT( "Sequencer", "PasteMatineeAnimTrack", "Paste Matinee anim track" ) );
+
+	SkeletalAnimationTrack->Modify();
+	SkeletalAnimationTrack->RemoveAllAnimationData();
+
+	for (int32 i = 0; i < MatineeAnimControlTrack->AnimSeqs.Num(); i++)
+	{
+		const auto& AnimSeq = MatineeAnimControlTrack->AnimSeqs[i];
+
+		float EndTime;
+		if( AnimSeq.bLooping )
+		{
+			if( i < MatineeAnimControlTrack->AnimSeqs.Num() - 1 )
+			{
+				EndTime = MatineeAnimControlTrack->AnimSeqs[i + 1].StartTime;
+			}
+			else
+			{
+				EndTime = Sequencer->GetFocusedMovieSceneSequence()->GetMovieScene()->GetPlaybackRange().GetUpperBoundValue();
+			}
+		}
+		else
+		{
+			EndTime = AnimSeq.StartTime + ( ( ( AnimSeq.AnimSeq->SequenceLength - AnimSeq.AnimEndOffset ) - AnimSeq.AnimStartOffset ) / AnimSeq.AnimPlayRate );
+		}
+
+		UMovieSceneSkeletalAnimationSection* NewSection = Cast<UMovieSceneSkeletalAnimationSection>( SkeletalAnimationTrack->CreateNewSection() );
+		NewSection->SetStartTime( AnimSeq.StartTime );
+		NewSection->SetEndTime( EndTime );
+		NewSection->SetStartOffset( AnimSeq.AnimStartOffset );
+		NewSection->SetEndOffset( AnimSeq.AnimEndOffset );
+		NewSection->SetPlayRate( AnimSeq.AnimPlayRate );
+		NewSection->SetAnimSequence( AnimSeq.AnimSeq );
+
+		SkeletalAnimationTrack->AddSection( *NewSection );
+	}
+
+	Sequencer->NotifyMovieSceneDataChanged();
+}
diff --git a/Engine/Source/Editor/MovieSceneTools/Private/MatineeImportTools.h b/Engine/Source/Editor/MovieSceneTools/Private/MatineeImportTools.h
new file mode 100644
--- /dev/null
+++ b/Engine/Source/Editor/MovieSceneTools/Private/MatineeImportTools.h
@@ -0,0 +1,44 @@
+// Copyright 1998-2015 Epic Games, Inc. All Rights Reserved.
+#pragma once
+
+#include "Matinee/InterpTrackToggle.h"
+#include "MovieSceneParticleSection.h"
+
+class UInterpTrackFloatBase;
+class UInterpTrackMove;
+class UInterpTrackToggle;
+class UInterpTrackAnimControl;
+
+class UMovieSceneFloatTrack;
+class UMovieScene3DTransformTrack;
+class UMovieSceneParticleTrack;
+class UMovieSceneSkeletalAnimationTrack;
+
+class FMatineeImportTools
+{
+public:
+
+	/** Converts a matinee interpolation mode to its equivalent rich curve interpolation mode. */
+	static ERichCurveInterpMode MatineeInterpolationToRichCurveInterpolation( EInterpCurveMode CurveMode );
+
+	/** Converts a matinee interpolation mode to its equivalent rich curve tangent mode. */
+	static ERichCurveTangentMode MatineeInterpolationToRichCurveTangent( EInterpCurveMode CurveMode );
+
+	/** Tries to convert a matinee toggle to a particle key. */
+	static bool TryConvertMatineeToggleToOutParticleKey( ETrackToggleAction ToggleAction, EParticleKey::Type& OutParticleKey );
+
+	/** Adds a key to a rich curve based on matinee curve key data. */
+	static void SetOrAddKey( FRichCurve& Curve, float Time, float Value, float ArriveTangent, float LeaveTangent, EInterpCurveMode MatineeInterpMode );
+
+	/** Copies keys from a matinee float track to a sequencer float track. */
+	static void CopyInterpFloatTrack( TSharedRef<ISequencer> Sequencer, UInterpTrackFloatBase* MatineeFloatTrack, UMovieSceneFloatTrack* FloatTrack );
+
+	/** Copies keys from a matinee move track to a sequencer transform track. */
+	static void CopyInterpMoveTrack( TSharedRef<ISequencer> Sequencer, UInterpTrackMove* MoveTrack, UMovieScene3DTransformTrack* TransformTrack );
+
+	/** Copies keys from a matinee toggle track to a sequencer particle track. */
+	static void CopyInterpParticleTrack( TSharedRef<ISequencer> Sequencer, UInterpTrackToggle* MatineeToggleTrack, UMovieSceneParticleTrack* ParticleTrack );
+
+	/** Copies keys from a matinee anim control track to a sequencer skeletal animation track. */
+	static void CopyInterpAnimControlTrack( TSharedRef<ISequencer> Sequencer, UInterpTrackAnimControl* MatineeAnimControlTrack, UMovieSceneSkeletalAnimationTrack* SkeletalAnimationTrack );
+};
diff --git a/Engine/Source/Editor/MovieSceneTools/Private/MovieSceneToolHelpers.cpp b/Engine/Source/Editor/MovieSceneTools/Private/MovieSceneToolHelpers.cpp
--- a/Engine/Source/Editor/MovieSceneTools/Private/MovieSceneToolHelpers.cpp
+++ b/Engine/Source/Editor/MovieSceneTools/Private/MovieSceneToolHelpers.cpp
@@ -1,165 +1,11 @@
 // Copyright 1998-2015 Epic Games, Inc. All Rights Reserved.
 
 #include "MovieSceneToolsPrivatePCH.h"
-#include "PropertyEditing.h"
-#include "MovieScene.h"
-#include "MovieSceneTrack.h"
-#include "ScopedTransaction.h"
 #include "MovieSceneSection.h"
-#include "ISequencerObjectChangeListener.h"
-#include "ISequencerSection.h"
-#include "ISectionLayoutBuilder.h"
-#include "IKeyArea.h"
-#include "MovieSceneToolHelpers.h"
-#include "SFloatCurveKeyEditor.h"
-#include "SIntegralCurveKeyEditor.h"
-#include "SEnumCurveKeyEditor.h"
-#include "SBoolCurveKeyEditor.h"
 
 
-TArray<FKeyHandle> FFloatCurveKeyArea::AddKeyUnique(float Time, EMovieSceneKeyInterpolation InKeyInterpolation, float TimeToCopyFrom)
-{
-	TArray<FKeyHandle> AddedKeyHandles;
-
-	FKeyHandle CurrentKeyHandle = Curve->FindKey(Time);
-	if (Curve->IsKeyHandleValid(CurrentKeyHandle) == false)
-	{
-		if (OwningSection->GetStartTime() > Time)
-		{
-			OwningSection->SetStartTime(Time);
-		}
-		if (OwningSection->GetEndTime() < Time)
-		{
-			OwningSection->SetEndTime(Time);
-		}
-
-		float Value = Curve->Eval(Time);
-		if (TimeToCopyFrom != FLT_MAX)
-		{
-			Value = Curve->Eval(TimeToCopyFrom);
-		}
-		else if ( IntermediateValue.IsSet() )
-		{
-			Value = IntermediateValue.GetValue();
-		}
-
-		Curve->AddKey(Time, Value, false, CurrentKeyHandle);
-		AddedKeyHandles.Add(CurrentKeyHandle);
-		
-		// Set the default key interpolation
-		switch (InKeyInterpolation)
-		{
-			case EMovieSceneKeyInterpolation::Auto:
-				Curve->SetKeyInterpMode(CurrentKeyHandle, RCIM_Cubic);
-				Curve->SetKeyTangentMode(CurrentKeyHandle, RCTM_Auto);
-				break;
-			case EMovieSceneKeyInterpolation::User:
-				Curve->SetKeyInterpMode(CurrentKeyHandle, RCIM_Cubic);
-				Curve->SetKeyTangentMode(CurrentKeyHandle, RCTM_User);
-				break;
-			case EMovieSceneKeyInterpolation::Break:
-				Curve->SetKeyInterpMode(CurrentKeyHandle, RCIM_Cubic);
-				Curve->SetKeyTangentMode(CurrentKeyHandle, RCTM_Break);
-				break;
-			case EMovieSceneKeyInterpolation::Linear:
-				Curve->SetKeyInterpMode(CurrentKeyHandle, RCIM_Linear);
-				Curve->SetKeyTangentMode(CurrentKeyHandle, RCTM_Auto);
-				break;
-			case EMovieSceneKeyInterpolation::Constant:
-				Curve->SetKeyInterpMode(CurrentKeyHandle, RCIM_Constant);
-				Curve->SetKeyTangentMode(CurrentKeyHandle, RCTM_Auto);
-				break;
-			default:
-				Curve->SetKeyInterpMode(CurrentKeyHandle, RCIM_Cubic);
-				Curve->SetKeyTangentMode(CurrentKeyHandle, RCTM_Auto);
-				break;
-		}
-		
-		// Copy the properties from the key if it exists
-		FKeyHandle KeyHandleToCopy = Curve->FindKey(TimeToCopyFrom);
-		if (Curve->IsKeyHandleValid(KeyHandleToCopy))
-		{
-			FRichCurveKey& CurrentKey = Curve->GetKey(CurrentKeyHandle);
-			FRichCurveKey& KeyToCopy = Curve->GetKey(KeyHandleToCopy);
-			CurrentKey.InterpMode = KeyToCopy.InterpMode;
-			CurrentKey.TangentMode = KeyToCopy.TangentMode;
-			CurrentKey.TangentWeightMode = KeyToCopy.TangentWeightMode;
-			CurrentKey.ArriveTangent = KeyToCopy.ArriveTangent;
-			CurrentKey.LeaveTangent = KeyToCopy.LeaveTangent;
-			CurrentKey.ArriveTangentWeight = KeyToCopy.ArriveTangentWeight;
-			CurrentKey.LeaveTangentWeight = KeyToCopy.LeaveTangentWeight;
-		}
-	}
-
-	return AddedKeyHandles;
-}
-
-
-TSharedRef<SWidget> FFloatCurveKeyArea::CreateKeyEditor(ISequencer* Sequencer)
-{
-	return SNew(SFloatCurveKeyEditor)
-		.Sequencer(Sequencer)
-		.OwningSection(OwningSection)
-		.Curve(Curve)
-		.IntermediateValue_Lambda( [this] { return IntermediateValue; } );
-};
-
-
-TArray<FKeyHandle> FNameCurveKeyArea::AddKeyUnique(float Time, EMovieSceneKeyInterpolation InKeyInterpolation, float TimeToCopyFrom)
-{
-	TArray<FKeyHandle> AddedKeyHandles;
-	FKeyHandle CurrentKey = Curve.FindKey(Time);
-
-	if (Curve.IsKeyHandleValid(CurrentKey) == false)
-	{
-		if (OwningSection->GetStartTime() > Time)
-		{
-			OwningSection->SetStartTime(Time);
-		}
-
-		if (OwningSection->GetEndTime() < Time)
-		{
-			OwningSection->SetEndTime(Time);
-		}
-
-		Curve.AddKey(Time, NAME_None, CurrentKey);
-		AddedKeyHandles.Add(CurrentKey);
-	}
-
-	return AddedKeyHandles;
-}
-
-
-TSharedRef<SWidget> FByteKeyArea::CreateKeyEditor(ISequencer* Sequencer)
-{
-	return SNew(SIntegralCurveKeyEditor)
-		.Sequencer(Sequencer)
-		.OwningSection(OwningSection)
-		.Curve(&Curve)
-		.IntermediateValue_Lambda( [this] { return IntermediateValue; } );
-};
-
-
-TSharedRef<SWidget> FEnumKeyArea::CreateKeyEditor(ISequencer* Sequencer)
-{
-	return SNew(SEnumCurveKeyEditor)
-		.Sequencer(Sequencer)
-		.OwningSection(OwningSection)
-		.Curve(&Curve)
-		.Enum(Enum)
-		.IntermediateValue_Lambda( [this] { return IntermediateValue; } );
-};
-
-
-TSharedRef<SWidget> FBoolKeyArea::CreateKeyEditor(ISequencer* Sequencer)
-{
-	return SNew(SBoolCurveKeyEditor)
-		.Sequencer(Sequencer)
-		.OwningSection(OwningSection)
-		.Curve(&Curve)
-		.IntermediateValue_Lambda( [this] { return IntermediateValue; } );
-};
-
+/* MovieSceneToolHelpers
+ *****************************************************************************/
 
 void MovieSceneToolHelpers::TrimSection(const TSet<TWeakObjectPtr<UMovieSceneSection>>& Sections, float Time, bool bTrimLeft)
 {
diff --git a/Engine/Source/Editor/MovieSceneTools/Private/MovieSceneToolsPrivatePCH.h b/Engine/Source/Editor/MovieSceneTools/Private/MovieSceneToolsPrivatePCH.h
--- a/Engine/Source/Editor/MovieSceneTools/Private/MovieSceneToolsPrivatePCH.h
+++ b/Engine/Source/Editor/MovieSceneTools/Private/MovieSceneToolsPrivatePCH.h
@@ -17,6 +17,7 @@
 #include "PropertyTrackEditor.h"
 #include "CommonMovieSceneTools.h"
 #include "MovieSceneToolHelpers.h"
+#include "MovieSceneCommonHelpers.h"
 
 
 #include "ShotSequencerSection.h"
diff --git a/Engine/Source/Editor/MovieSceneTools/Private/Sections/BoolPropertySection.h b/Engine/Source/Editor/MovieSceneTools/Private/Sections/BoolPropertySection.h
--- a/Engine/Source/Editor/MovieSceneTools/Private/Sections/BoolPropertySection.h
+++ b/Engine/Source/Editor/MovieSceneTools/Private/Sections/BoolPropertySection.h
@@ -2,17 +2,25 @@
 
 #pragma once
 
+#include "BoolKeyArea.h"
+
+
 /**
 * An implementation of bool property sections
 */
-class FBoolPropertySection : public FPropertySection
+class FBoolPropertySection
+	: public FPropertySection
 {
 public:
+
 	FBoolPropertySection( UMovieSceneSection& InSectionObject, const FText& SectionName, ISequencer* InSequencer )
 		: FPropertySection( InSectionObject, SectionName )
 		, Sequencer( InSequencer ) {}
 
+public:
+
 	// FPropertySection interface
+
 	virtual void GenerateSectionLayout( class ISectionLayoutBuilder& LayoutBuilder ) const override;
 	virtual int32 OnPaintSection( const FGeometry& AllottedGeometry, const FSlateRect& SectionClippingRect, FSlateWindowElementList& OutDrawElements, int32 LayerId, bool bParentEnabled ) const override;
 	virtual void SetIntermediateValue( FPropertyChangedParams PropertyChangedParams ) override;
@@ -25,4 +33,4 @@ protected:
 private:
 
 	ISequencer* Sequencer;
-};
\ No newline at end of file
+};
diff --git a/Engine/Source/Editor/MovieSceneTools/Private/Sections/BytePropertySection.cpp b/Engine/Source/Editor/MovieSceneTools/Private/Sections/BytePropertySection.cpp
--- a/Engine/Source/Editor/MovieSceneTools/Private/Sections/BytePropertySection.cpp
+++ b/Engine/Source/Editor/MovieSceneTools/Private/Sections/BytePropertySection.cpp
@@ -4,6 +4,8 @@
 #include "MovieSceneByteTrack.h"
 #include "MovieSceneByteSection.h"
 #include "BytePropertySection.h"
+#include "EnumKeyArea.h"
+
 
 void FBytePropertySection::GenerateSectionLayout( class ISectionLayoutBuilder& LayoutBuilder ) const
 {
@@ -20,11 +22,13 @@ void FBytePropertySection::GenerateSectionLayout( class ISectionLayoutBuilder& L
 	}
 }
 
+
 void FBytePropertySection::SetIntermediateValue( FPropertyChangedParams PropertyChangedParams )
 {
 	KeyArea->SetIntermediateValue( PropertyChangedParams.GetPropertyValue<uint8>() );
 }
 
+
 void FBytePropertySection::ClearIntermediateValue()
 {
 	KeyArea->ClearIntermediateValue();
diff --git a/Engine/Source/Editor/MovieSceneTools/Private/Sections/BytePropertySection.h b/Engine/Source/Editor/MovieSceneTools/Private/Sections/BytePropertySection.h
--- a/Engine/Source/Editor/MovieSceneTools/Private/Sections/BytePropertySection.h
+++ b/Engine/Source/Editor/MovieSceneTools/Private/Sections/BytePropertySection.h
@@ -2,6 +2,8 @@
 
 #pragma once
 
+#include "ByteKeyArea.h"
+
 
 /**
  * An implementation of byte property sections
@@ -10,6 +12,7 @@ class FBytePropertySection
 	: public FPropertySection
 {
 public:
+
 	FBytePropertySection(UMovieSceneSection& InSectionObject, const FText& InSectionName, UEnum* InEnum)
 		: FPropertySection(InSectionObject, InSectionName)
 		, Enum( InEnum )
@@ -27,4 +30,4 @@ private:
 
 	mutable TSharedPtr<FByteKeyArea> KeyArea;
 	UEnum* Enum;
-};
\ No newline at end of file
+};
diff --git a/Engine/Source/Editor/MovieSceneTools/Private/Sections/ColorPropertySection.cpp b/Engine/Source/Editor/MovieSceneTools/Private/Sections/ColorPropertySection.cpp
--- a/Engine/Source/Editor/MovieSceneTools/Private/Sections/ColorPropertySection.cpp
+++ b/Engine/Source/Editor/MovieSceneTools/Private/Sections/ColorPropertySection.cpp
@@ -5,6 +5,7 @@
 #include "MovieSceneColorSection.h"
 #include "ColorPropertySection.h"
 #include "MovieSceneSequence.h"
+#include "FloatCurveKeyArea.h"
 
 
 void FColorPropertySection::GenerateSectionLayout( class ISectionLayoutBuilder& LayoutBuilder ) const
diff --git a/Engine/Source/Editor/MovieSceneTools/Private/Sections/ColorPropertySection.h b/Engine/Source/Editor/MovieSceneTools/Private/Sections/ColorPropertySection.h
--- a/Engine/Source/Editor/MovieSceneTools/Private/Sections/ColorPropertySection.h
+++ b/Engine/Source/Editor/MovieSceneTools/Private/Sections/ColorPropertySection.h
@@ -2,8 +2,12 @@
 
 #pragma once
 
+#include "FloatCurveKeyArea.h"
+
+
 class UMovieSceneColorSection;
 
+
 /**
 * A color section implementation
 */
@@ -44,6 +48,7 @@ protected:
 	FLinearColor FindSlateColor(const FName& ColorName) const;
 
 private:
+
 	mutable TSharedPtr<FFloatCurveKeyArea> RedKeyArea;
 	mutable TSharedPtr<FFloatCurveKeyArea> GreenKeyArea;
 	mutable TSharedPtr<FFloatCurveKeyArea> BlueKeyArea;
@@ -54,4 +59,4 @@ private:
 
 	/** The track that owns the section. */
 	UMovieSceneColorTrack& Track;
-};
\ No newline at end of file
+};
diff --git a/Engine/Source/Editor/MovieSceneTools/Private/Sections/EventTrackSection.cpp b/Engine/Source/Editor/MovieSceneTools/Private/Sections/EventTrackSection.cpp
--- a/Engine/Source/Editor/MovieSceneTools/Private/Sections/EventTrackSection.cpp
+++ b/Engine/Source/Editor/MovieSceneTools/Private/Sections/EventTrackSection.cpp
@@ -3,6 +3,7 @@
 #include "MovieSceneToolsPrivatePCH.h"
 #include "EventTrackSection.h"
 #include "MovieSceneEventSection.h"
+#include "NameCurveKeyArea.h"
 
 
 #define LOCTEXT_NAMESPACE "FEventTrackSection"
diff --git a/Engine/Source/Editor/MovieSceneTools/Private/Sections/FloatPropertySection.h b/Engine/Source/Editor/MovieSceneTools/Private/Sections/FloatPropertySection.h
--- a/Engine/Source/Editor/MovieSceneTools/Private/Sections/FloatPropertySection.h
+++ b/Engine/Source/Editor/MovieSceneTools/Private/Sections/FloatPropertySection.h
@@ -2,6 +2,8 @@
 
 #pragma once
 
+#include "FloatCurveKeyArea.h"
+
 
 /**
  * An implementation of float property sections
@@ -26,4 +28,4 @@ public:
 private:
 
 	mutable TSharedPtr<FFloatCurveKeyArea> KeyArea;
-};
\ No newline at end of file
+};
diff --git a/Engine/Source/Editor/MovieSceneTools/Private/Sections/MaterialParameterSection.cpp b/Engine/Source/Editor/MovieSceneTools/Private/Sections/MaterialParameterSection.cpp
--- a/Engine/Source/Editor/MovieSceneTools/Private/Sections/MaterialParameterSection.cpp
+++ b/Engine/Source/Editor/MovieSceneTools/Private/Sections/MaterialParameterSection.cpp
@@ -3,11 +3,14 @@
 #include "MovieSceneToolsPrivatePCH.h"
 #include "MaterialParameterSection.h"
 #include "MovieSceneMaterialParameterSection.h"
+#include "FloatCurveKeyArea.h"
+
 
 #define LOCTEXT_NAMESPACE "MaterialParameterSection"
 
 DECLARE_DELEGATE(FLayoutParameterDelegate);
 
+
 struct FIndexAndLayoutParameterDelegate
 {
 	int32 Index;
@@ -19,6 +22,7 @@ struct FIndexAndLayoutParameterDelegate
 	}
 };
 
+
 void LayoutScalarParameter( ISectionLayoutBuilder* LayoutBuilder, FScalarParameterNameAndCurve* ScalarParameterNameAndCurve, UMovieSceneMaterialParameterSection* ParameterSection )
 {
 	LayoutBuilder->AddKeyArea(
@@ -27,6 +31,7 @@ void LayoutScalarParameter( ISectionLayoutBuilder* LayoutBuilder, FScalarParamet
 		MakeShareable( new FFloatCurveKeyArea( &ScalarParameterNameAndCurve->ParameterCurve, ParameterSection ) ) );
 }
 
+
 void LayoutVectorParameter( ISectionLayoutBuilder* LayoutBuilder, FVectorParameterNameAndCurves* VectorParameterNameAndCurves, UMovieSceneMaterialParameterSection* ParameterSection )
 {
 	LayoutBuilder->PushCategory( VectorParameterNameAndCurves->ParameterName, FText::FromName( VectorParameterNameAndCurves->ParameterName ) );
@@ -43,6 +48,7 @@ void LayoutVectorParameter( ISectionLayoutBuilder* LayoutBuilder, FVectorParamet
 	LayoutBuilder->PopCategory();
 }
 
+
 void FMaterialParameterSection::GenerateSectionLayout( class ISectionLayoutBuilder& LayoutBuilder ) const
 {
 	// TODO - This seems a bit convoluted here to have to use delegates to get the layout to happen in the right order based on the indices.
@@ -77,6 +83,7 @@ void FMaterialParameterSection::GenerateSectionLayout( class ISectionLayoutBuild
 	}
 }
 
+
 bool FMaterialParameterSection::RequestDeleteCategory( const TArray<FName>& CategoryNamePath )
 {
 	const FScopedTransaction Transaction( LOCTEXT( "DeleteVectorParameter", "Delete vector parameter" ) );
@@ -85,6 +92,7 @@ bool FMaterialParameterSection::RequestDeleteCategory( const TArray<FName>& Cate
 	return ParameterSection->RemoveVectorParameter( CategoryNamePath[0] );
 }
 
+
 bool FMaterialParameterSection::RequestDeleteKeyArea( const TArray<FName>& KeyAreaNamePath )
 {
 	// Only handle paths with a single name, in all other cases the user is deleting a component of a vector parameter.
@@ -98,4 +106,5 @@ bool FMaterialParameterSection::RequestDeleteKeyArea( const TArray<FName>& KeyAr
 	return false;
 }
 
-#undef LOCTEXT_NAMESPACE
\ No newline at end of file
+
+#undef LOCTEXT_NAMESPACE
diff --git a/Engine/Source/Editor/MovieSceneTools/Private/Sections/VectorPropertySection.h b/Engine/Source/Editor/MovieSceneTools/Private/Sections/VectorPropertySection.h
--- a/Engine/Source/Editor/MovieSceneTools/Private/Sections/VectorPropertySection.h
+++ b/Engine/Source/Editor/MovieSceneTools/Private/Sections/VectorPropertySection.h
@@ -2,6 +2,8 @@
 
 #pragma once
 
+#include "FloatCurveKeyArea.h"
+
 
 /**
  * An implementation of vector property sections.
diff --git a/Engine/Source/Editor/MovieSceneTools/Private/TrackEditors/ActorPickerTrackEditor.cpp b/Engine/Source/Editor/MovieSceneTools/Private/TrackEditors/ActorPickerTrackEditor.cpp
--- a/Engine/Source/Editor/MovieSceneTools/Private/TrackEditors/ActorPickerTrackEditor.cpp
+++ b/Engine/Source/Editor/MovieSceneTools/Private/TrackEditors/ActorPickerTrackEditor.cpp
@@ -34,10 +34,10 @@ void FActorPickerTrackEditor::ShowActorSubMenu(FMenuBuilder& MenuBuilder, FGuid
 {
 	struct Local
 	{
-		static FReply OnInteractiveActorPickerClicked(FActorPickerTrackEditor* ActorPickerTrackEditor, FGuid InObjectBinding, UMovieSceneSection* InSection)
+		static FReply OnInteractiveActorPickerClicked(FActorPickerTrackEditor* ActorPickerTrackEditor, FGuid TheObjectBinding, UMovieSceneSection* TheSection)
 		{
 			FSlateApplication::Get().DismissAllMenus();
-			ActorPickerTrackEditor->PickActorInteractive(InObjectBinding, InSection);
+			ActorPickerTrackEditor->PickActorInteractive(TheObjectBinding, TheSection);
 			return FReply::Handled();
 		}
 	};
diff --git a/Engine/Source/Editor/MovieSceneTools/Private/TrackEditors/AttachTrackEditor.cpp b/Engine/Source/Editor/MovieSceneTools/Private/TrackEditors/AttachTrackEditor.cpp
--- a/Engine/Source/Editor/MovieSceneTools/Private/TrackEditors/AttachTrackEditor.cpp
+++ b/Engine/Source/Editor/MovieSceneTools/Private/TrackEditors/AttachTrackEditor.cpp
@@ -10,7 +10,6 @@
 #include "ISectionLayoutBuilder.h"
 #include "IKeyArea.h"
 #include "MovieSceneCommonHelpers.h"
-#include "MovieSceneToolHelpers.h"
 #include "MovieSceneTrackEditor.h"
 #include "ActorEditorUtils.h"
 #include "GameFramework/WorldSettings.h"
diff --git a/Engine/Source/Editor/MovieSceneTools/Private/TrackEditors/AudioTrackEditor.cpp b/Engine/Source/Editor/MovieSceneTools/Private/TrackEditors/AudioTrackEditor.cpp
--- a/Engine/Source/Editor/MovieSceneTools/Private/TrackEditors/AudioTrackEditor.cpp
+++ b/Engine/Source/Editor/MovieSceneTools/Private/TrackEditors/AudioTrackEditor.cpp
@@ -12,7 +12,6 @@
 #include "ISequencerObjectChangeListener.h"
 #include "ISectionLayoutBuilder.h"
 #include "IKeyArea.h"
-#include "MovieSceneToolHelpers.h"
 #include "MovieSceneTrackEditor.h"
 #include "AudioTrackEditor.h"
 #include "MovieSceneAudioSection.h"
diff --git a/Engine/Source/Editor/MovieSceneTools/Private/TrackEditors/ParticleTrackEditor.cpp b/Engine/Source/Editor/MovieSceneTools/Private/TrackEditors/ParticleTrackEditor.cpp
--- a/Engine/Source/Editor/MovieSceneTools/Private/TrackEditors/ParticleTrackEditor.cpp
+++ b/Engine/Source/Editor/MovieSceneTools/Private/TrackEditors/ParticleTrackEditor.cpp
@@ -12,7 +12,6 @@
 #include "ISequencerObjectChangeListener.h"
 #include "ISectionLayoutBuilder.h"
 #include "IKeyArea.h"
-#include "MovieSceneToolHelpers.h"
 #include "MovieSceneTrackEditor.h"
 #include "ParticleTrackEditor.h"
 #include "MovieSceneParticleSection.h"
@@ -22,6 +21,9 @@
 #include "Particles/ParticleEmitter.h"
 #include "Particles/ParticleLODLevel.h"
 #include "Particles/ParticleModuleRequired.h"
+#include "EnumKeyArea.h"
+#include "MatineeImportTools.h"
+#include "Matinee/InterpTrackToggle.h"
 
 
 namespace AnimatableParticleEditorConstants
@@ -103,9 +105,10 @@ int32 FParticleSection::OnPaintSection( const FGeometry& AllottedGeometry, const
 				AEmitter* ParticleSystemActor = Cast<AEmitter>( OwningSequencer->GetFocusedMovieSceneSequence()->FindObject( ObjectHandle ) );
 				if ( ParticleSystemActor != nullptr )
 				{
-					if (ParticleSystemActor->GetParticleSystemComponent())
+					UParticleSystemComponent* ParticleSystemComponent = ParticleSystemActor->GetParticleSystemComponent();
+					if ( ParticleSystemComponent != nullptr && ParticleSystemComponent->Template != nullptr )
 					{
-						for ( UParticleEmitter* Emitter : ParticleSystemActor->GetParticleSystemComponent()->Template->Emitters )
+						for ( UParticleEmitter* Emitter : ParticleSystemComponent->Template->Emitters )
 						{
 							UParticleModuleRequired* RequiredModule = Emitter->GetLODLevel( 0 )->RequiredModule;
 							bIsLooping |= RequiredModule->EmitterLoops == 0;
@@ -301,4 +304,26 @@ bool FParticleTrackEditor::AddKeyInternal( float KeyTime, const TArray<UObject*>
 }
 
 
+void FParticleTrackEditor::BuildTrackContextMenu( FMenuBuilder& MenuBuilder, UMovieSceneTrack* Track )
+{
+	UInterpTrackToggle* MatineeToggleTrack = nullptr;
+	for ( UObject* CopyPasteObject : GUnrealEd->MatineeCopyPasteBuffer )
+	{
+		MatineeToggleTrack = Cast<UInterpTrackToggle>( CopyPasteObject );
+		if ( MatineeToggleTrack != nullptr )
+		{
+			break;
+		}
+	}
+	UMovieSceneParticleTrack* ParticleTrack = Cast<UMovieSceneParticleTrack>( Track );
+	MenuBuilder.AddMenuEntry(
+		NSLOCTEXT( "Sequencer", "PasteMatineeToggleTrack", "Paste Matinee Particle Track" ),
+		NSLOCTEXT( "Sequencer", "PasteMatineeToggleTrackTooltip", "Pastes keys from a Matinee particle track into this track." ),
+		FSlateIcon(),
+		FUIAction(
+			FExecuteAction::CreateStatic( &FMatineeImportTools::CopyInterpParticleTrack, GetSequencer().ToSharedRef(), MatineeToggleTrack, ParticleTrack ),
+			FCanExecuteAction::CreateLambda( [=]()->bool { return MatineeToggleTrack != nullptr && MatineeToggleTrack->ToggleTrack.Num() > 0 && ParticleTrack != nullptr; } ) ) );
+}
+
+
 #undef LOCTEXT_NAMESPACE
diff --git a/Engine/Source/Editor/MovieSceneTools/Private/TrackEditors/ParticleTrackEditor.h b/Engine/Source/Editor/MovieSceneTools/Private/TrackEditors/ParticleTrackEditor.h
--- a/Engine/Source/Editor/MovieSceneTools/Private/TrackEditors/ParticleTrackEditor.h
+++ b/Engine/Source/Editor/MovieSceneTools/Private/TrackEditors/ParticleTrackEditor.h
@@ -40,6 +40,7 @@ public:
 	virtual void BuildObjectBindingTrackMenu(FMenuBuilder& MenuBuilder, const FGuid& ObjectBinding, const UClass* ObjectClass) override;
 	virtual TSharedRef<ISequencerSection> MakeSectionInterface( UMovieSceneSection& SectionObject, UMovieSceneTrack& Track ) override;
 	virtual bool SupportsType( TSubclassOf<UMovieSceneTrack> Type ) const override;
+	virtual void BuildTrackContextMenu( FMenuBuilder& MenuBuilder, UMovieSceneTrack* Track ) override;
 	
 	void AddParticleKey(const FGuid ObjectGuid);
 private:
diff --git a/Engine/Source/Editor/MovieSceneTools/Private/TrackEditors/PathTrackEditor.cpp b/Engine/Source/Editor/MovieSceneTools/Private/TrackEditors/PathTrackEditor.cpp
--- a/Engine/Source/Editor/MovieSceneTools/Private/TrackEditors/PathTrackEditor.cpp
+++ b/Engine/Source/Editor/MovieSceneTools/Private/TrackEditors/PathTrackEditor.cpp
@@ -8,14 +8,13 @@
 #include "MovieScene3DPathSection.h"
 #include "ISequencerSection.h"
 #include "ISectionLayoutBuilder.h"
-#include "IKeyArea.h"
 #include "MovieSceneCommonHelpers.h"
-#include "MovieSceneToolHelpers.h"
 #include "MovieSceneTrackEditor.h"
 #include "PathTrackEditor.h"
 #include "ActorEditorUtils.h"
 #include "Components/SplineComponent.h"
 #include "ActorPickerTrackEditor.h"
+#include "FloatCurveKeyArea.h"
 
 
 #define LOCTEXT_NAMESPACE "FPathTrackEditor"
diff --git a/Engine/Source/Editor/MovieSceneTools/Private/TrackEditors/PropertyTrackEditors/FloatPropertyTrackEditor.cpp b/Engine/Source/Editor/MovieSceneTools/Private/TrackEditors/PropertyTrackEditors/FloatPropertyTrackEditor.cpp
--- a/Engine/Source/Editor/MovieSceneTools/Private/TrackEditors/PropertyTrackEditors/FloatPropertyTrackEditor.cpp
+++ b/Engine/Source/Editor/MovieSceneTools/Private/TrackEditors/PropertyTrackEditors/FloatPropertyTrackEditor.cpp
@@ -3,6 +3,8 @@
 #include "MovieSceneToolsPrivatePCH.h"
 #include "FloatPropertyTrackEditor.h"
 #include "FloatPropertySection.h"
+#include "MatineeImportTools.h"
+#include "Matinee/InterpTrackFloatBase.h"
 
 
 TSharedRef<ISequencerTrackEditor> FFloatPropertyTrackEditor::CreateTrackEditor( TSharedRef<ISequencer> OwningSequencer )
@@ -21,3 +23,25 @@ void FFloatPropertyTrackEditor::GenerateKeysFromPropertyChanged( const FProperty
 {
 	GeneratedKeys.Add( PropertyChangedParams.GetPropertyValue<float>() );
 }
+
+
+void FFloatPropertyTrackEditor::BuildTrackContextMenu( FMenuBuilder& MenuBuilder, UMovieSceneTrack* Track )
+{
+	UInterpTrackFloatBase* MatineeFloatTrack = nullptr;
+	for ( UObject* CopyPasteObject : GUnrealEd->MatineeCopyPasteBuffer )
+	{
+		MatineeFloatTrack = Cast<UInterpTrackFloatBase>( CopyPasteObject );
+		if ( MatineeFloatTrack != nullptr )
+		{
+			break;
+		}
+	}
+	UMovieSceneFloatTrack* FloatTrack = Cast<UMovieSceneFloatTrack>( Track );
+	MenuBuilder.AddMenuEntry(
+		NSLOCTEXT( "Sequencer", "PasteMatineeFloatTrack", "Paste Matinee Float Track" ),
+		NSLOCTEXT( "Sequencer", "PasteMatineeFloatTrackTooltip", "Pastes keys from a Matinee float track into this track." ),
+		FSlateIcon(),
+		FUIAction(
+			FExecuteAction::CreateStatic( &FMatineeImportTools::CopyInterpFloatTrack, GetSequencer().ToSharedRef(), MatineeFloatTrack, FloatTrack ),
+			FCanExecuteAction::CreateLambda( [=]()->bool { return MatineeFloatTrack != nullptr && MatineeFloatTrack->GetNumKeys() > 0 && FloatTrack != nullptr; } ) ) );
+}
diff --git a/Engine/Source/Editor/MovieSceneTools/Private/TrackEditors/PropertyTrackEditors/FloatPropertyTrackEditor.h b/Engine/Source/Editor/MovieSceneTools/Private/TrackEditors/PropertyTrackEditors/FloatPropertyTrackEditor.h
--- a/Engine/Source/Editor/MovieSceneTools/Private/TrackEditors/PropertyTrackEditors/FloatPropertyTrackEditor.h
+++ b/Engine/Source/Editor/MovieSceneTools/Private/TrackEditors/PropertyTrackEditors/FloatPropertyTrackEditor.h
@@ -30,6 +30,9 @@ public:
 	 */
 	static TSharedRef<ISequencerTrackEditor> CreateTrackEditor( TSharedRef<ISequencer> OwningSequencer );
 
+	// FMovieSceneTrackEditor interface.
+	virtual void BuildTrackContextMenu( FMenuBuilder& MenuBuilder, UMovieSceneTrack* Track ) override;
+
 protected:
 
 	// FPropertyTrackEditor interface
diff --git a/Engine/Source/Editor/MovieSceneTools/Private/TrackEditors/ShotTrackEditor/ShotTrackEditor.cpp b/Engine/Source/Editor/MovieSceneTools/Private/TrackEditors/ShotTrackEditor/ShotTrackEditor.cpp
--- a/Engine/Source/Editor/MovieSceneTools/Private/TrackEditors/ShotTrackEditor/ShotTrackEditor.cpp
+++ b/Engine/Source/Editor/MovieSceneTools/Private/TrackEditors/ShotTrackEditor/ShotTrackEditor.cpp
@@ -10,7 +10,6 @@
 #include "ScopedTransaction.h"
 #include "ISequencerObjectChangeListener.h"
 #include "IKeyArea.h"
-#include "MovieSceneToolHelpers.h"
 #include "MovieSceneTrackEditor.h"
 #include "ShotTrackEditor.h"
 #include "CommonMovieSceneTools.h"
@@ -120,6 +119,7 @@ TSharedPtr<SWidget> FShotTrackEditor::BuildOutlinerEditWidget(const FGuid& Objec
 		.Padding(4, 0, 0, 0)
 		[
 			SNew(SCheckBox)
+				.IsFocusable(false)
 				.IsChecked(this, &FShotTrackEditor::IsCameraLocked)
 				.OnCheckStateChanged(this, &FShotTrackEditor::OnLockCameraClicked)
 				.ToolTipText(this, &FShotTrackEditor::GetLockCameraToolTip)
@@ -426,7 +426,6 @@ void FShotTrackEditor::OnLockCameraClicked(ECheckBoxState CheckBoxState)
 	}
 	else
 	{
-		// TODO: MaxC, is this right?
 		GetSequencer()->UpdateCameraCut(nullptr, nullptr);
 		GetSequencer()->SetPerspectiveViewportCameraCutEnabled(false);
 	}
diff --git a/Engine/Source/Editor/MovieSceneTools/Private/TrackEditors/SkeletalAnimationTrackEditor.cpp b/Engine/Source/Editor/MovieSceneTools/Private/TrackEditors/SkeletalAnimationTrackEditor.cpp
--- a/Engine/Source/Editor/MovieSceneTools/Private/TrackEditors/SkeletalAnimationTrackEditor.cpp
+++ b/Engine/Source/Editor/MovieSceneTools/Private/TrackEditors/SkeletalAnimationTrackEditor.cpp
@@ -12,7 +12,6 @@
 #include "ISequencerObjectChangeListener.h"
 #include "ISectionLayoutBuilder.h"
 #include "IKeyArea.h"
-#include "MovieSceneToolHelpers.h"
 #include "MovieSceneTrackEditor.h"
 #include "SkeletalAnimationTrackEditor.h"
 #include "MovieSceneSkeletalAnimationSection.h"
@@ -20,6 +19,8 @@
 #include "AssetRegistryModule.h"
 #include "Animation/SkeletalMeshActor.h"
 #include "ContentBrowserModule.h"
+#include "MatineeImportTools.h"
+#include "Matinee/InterpTrackAnimControl.h"
 
 
 namespace SkeletalAnimationEditorConstants
@@ -307,4 +308,26 @@ USkeleton* FSkeletalAnimationTrackEditor::AcquireSkeletonFromObjectGuid(const FG
 }
 
 
+void FSkeletalAnimationTrackEditor::BuildTrackContextMenu( FMenuBuilder& MenuBuilder, UMovieSceneTrack* Track )
+{
+	UInterpTrackAnimControl* MatineeAnimControlTrack = nullptr;
+	for ( UObject* CopyPasteObject : GUnrealEd->MatineeCopyPasteBuffer )
+	{
+		MatineeAnimControlTrack = Cast<UInterpTrackAnimControl>( CopyPasteObject );
+		if ( MatineeAnimControlTrack != nullptr )
+		{
+			break;
+		}
+	}
+	UMovieSceneSkeletalAnimationTrack* SkeletalAnimationTrack = Cast<UMovieSceneSkeletalAnimationTrack>( Track );
+	MenuBuilder.AddMenuEntry(
+		NSLOCTEXT( "Sequencer", "PasteMatineeAnimControlTrack", "Paste Matinee SkeletalAnimation Track" ),
+		NSLOCTEXT( "Sequencer", "PasteMatineeAnimControlTrackTooltip", "Pastes keys from a Matinee float track into this track." ),
+		FSlateIcon(),
+		FUIAction(
+			FExecuteAction::CreateStatic( &FMatineeImportTools::CopyInterpAnimControlTrack, GetSequencer().ToSharedRef(), MatineeAnimControlTrack, SkeletalAnimationTrack ),
+			FCanExecuteAction::CreateLambda( [=]()->bool { return MatineeAnimControlTrack != nullptr && MatineeAnimControlTrack->AnimSeqs.Num() > 0 && SkeletalAnimationTrack != nullptr; } ) ) );
+}
+
+
 #undef LOCTEXT_NAMESPACE
