Fix compilation on case-sensitive fs (UE-23281).

diff --git a/Engine/Source/Editor/MovieSceneTools/Private/MatineeImportTools.cpp b/Engine/Source/Editor/MovieSceneTools/Private/MatineeImportTools.cpp
--- a/Engine/Source/Editor/MovieSceneTools/Private/MatineeImportTools.cpp
+++ b/Engine/Source/Editor/MovieSceneTools/Private/MatineeImportTools.cpp
@@ -8,12 +8,12 @@
 #include "Matinee/InterpTrackAnimControl.h"
 
 #include "MovieSceneFloatTrack.h"
-#include "MovieScene3dTransformTrack.h"
+#include "MovieScene3DTransformTrack.h"
 #include "MovieSceneParticleTrack.h"
 #include "MovieSceneSkeletalAnimationTrack.h"
 
 #include "MovieSceneFloatSection.h"
-#include "MovieScene3dTransformSection.h"
+#include "MovieScene3DTransformSection.h"
 #include "MovieSceneSkeletalAnimationSection.h"
 
 
