Copying //UE4/Dev-Sequencer to Dev-Main (//UE4/Dev-Main)

diff --git a/Engine/Source/Editor/MovieSceneTools/Private/TrackEditorThumbnail/TrackEditorThumbnailPool.cpp b/Engine/Source/Editor/MovieSceneTools/Private/TrackEditorThumbnail/TrackEditorThumbnailPool.cpp
new file mode 100644
--- /dev/null
+++ b/Engine/Source/Editor/MovieSceneTools/Private/TrackEditorThumbnail/TrackEditorThumbnailPool.cpp
@@ -0,0 +1,79 @@
+// Copyright 1998-2016 Epic Games, Inc. All Rights Reserved.
+
+#include "MovieSceneToolsPrivatePCH.h"
+
+
+/* FShotSequencerSection structors
+ *****************************************************************************/
+
+FTrackEditorThumbnailPool::FTrackEditorThumbnailPool(TSharedPtr<ISequencer> InSequencer, uint32 InMaxThumbnailsToDrawAtATime)
+	: Sequencer(InSequencer)
+	, ThumbnailsNeedingDraw()
+	, MaxThumbnailsToDrawAtATime(InMaxThumbnailsToDrawAtATime)
+{ }
+
+
+/* FShotSequencerSection interface
+ *****************************************************************************/
+
+void FTrackEditorThumbnailPool::AddThumbnailsNeedingRedraw(const TArray<TSharedPtr<FTrackEditorThumbnail>>& InThumbnails)
+{
+	ThumbnailsNeedingDraw.Append(InThumbnails);
+}
+
+
+void FTrackEditorThumbnailPool::DrawThumbnails()
+{
+	if (ThumbnailsNeedingDraw.Num() == 0)
+	{
+		return;
+	}
+
+	// save the global time
+	float SavedTime = Sequencer.Pin()->GetGlobalTime();
+	uint32 ThumbnailsDrawn = 0;
+
+	for (int32 ThumbnailsIndex = 0; ThumbnailsDrawn < MaxThumbnailsToDrawAtATime && ThumbnailsIndex < ThumbnailsNeedingDraw.Num(); ++ThumbnailsIndex)
+	{
+		TSharedPtr<FTrackEditorThumbnail> Thumbnail = ThumbnailsNeedingDraw[ThumbnailsIndex];
+			
+		if (Thumbnail->IsVisible())
+		{
+			
+			bool bIsEnabled = Sequencer.Pin()->IsPerspectiveViewportCameraCutEnabled();
+
+			Sequencer.Pin()->SetPerspectiveViewportCameraCutEnabled(false);
+			
+			Thumbnail->DrawThumbnail();
+			
+			Sequencer.Pin()->SetPerspectiveViewportCameraCutEnabled(bIsEnabled);
+			++ThumbnailsDrawn;
+
+			ThumbnailsNeedingDraw.Remove(Thumbnail);
+		}
+		else if (!Thumbnail->IsValid())
+		{
+			ensure(0);
+
+			ThumbnailsNeedingDraw.Remove(Thumbnail);
+		}
+	}
+		
+	if (ThumbnailsDrawn > 0)
+	{
+		// Ensure all buffers are updated
+		FlushRenderingCommands();
+
+		// restore the global time
+		Sequencer.Pin()->SetGlobalTime(SavedTime);
+	}
+}
+
+
+void FTrackEditorThumbnailPool::RemoveThumbnailsNeedingRedraw(const TArray< TSharedPtr<FTrackEditorThumbnail> >& InThumbnails)
+{
+	for (int32 i = 0; i < InThumbnails.Num(); ++i)
+	{
+		ThumbnailsNeedingDraw.Remove(InThumbnails[i]);
+	}
+}
diff --git a/Engine/Source/Editor/MovieSceneTools/Private/TrackEditorThumbnail/TrackEditorThumbnailPool.h b/Engine/Source/Editor/MovieSceneTools/Private/TrackEditorThumbnail/TrackEditorThumbnailPool.h
new file mode 100644
--- /dev/null
+++ b/Engine/Source/Editor/MovieSceneTools/Private/TrackEditorThumbnail/TrackEditorThumbnailPool.h
@@ -0,0 +1,40 @@
+// Copyright 1998-2016 Epic Games, Inc. All Rights Reserved.
+
+#pragma once
+
+
+class FTrackEditorThumbnail;
+
+
+/**
+ * Track Editor Thumbnail pool, which keeps a list of thumbnails that
+ * need to be drawn and draws them incrementally.
+ */
+class FTrackEditorThumbnailPool
+{
+public:
+
+	FTrackEditorThumbnailPool(TSharedPtr<ISequencer> InSequencer, uint32 InMaxThumbnailsToDrawAtATime = 1);
+
+public:
+
+	/** Requests that the passed in thumbnails need to be drawn */
+	void AddThumbnailsNeedingRedraw(const TArray<TSharedPtr<FTrackEditorThumbnail>>& InThumbnails);
+
+	/** Draws a small number of thumbnails that are enqueued for drawing */
+	void DrawThumbnails();
+
+	/** Informs the pool that the thumbnails passed in no longer need to be drawn */
+	void RemoveThumbnailsNeedingRedraw(const TArray< TSharedPtr<FTrackEditorThumbnail>>& InThumbnails);
+
+private:
+
+	/** Parent sequencer we're drawing thumbnails for */
+	TWeakPtr<ISequencer> Sequencer;
+
+	/** Thumbnails enqueued for drawing */
+	TArray< TSharedPtr<FTrackEditorThumbnail> > ThumbnailsNeedingDraw;
+
+	/** How many thumbnails we are allowed to draw in a single DrawThumbnails call */
+	uint32 MaxThumbnailsToDrawAtATime;
+};
diff --git a/Engine/Source/Editor/MovieSceneTools/Private/TrackEditors/AudioTrackEditor.cpp b/Engine/Source/Editor/MovieSceneTools/Private/TrackEditors/AudioTrackEditor.cpp
--- a/Engine/Source/Editor/MovieSceneTools/Private/TrackEditors/AudioTrackEditor.cpp
+++ b/Engine/Source/Editor/MovieSceneTools/Private/TrackEditors/AudioTrackEditor.cpp
@@ -21,7 +21,6 @@
 #include "Runtime/Engine/Public/Slate/SlateTextures.h"
 #include "ObjectTools.h"
 #include "Runtime/Engine/Public/Slate/SceneViewport.h"
-#include "Runtime/MovieSceneTracks/Public/Sections/MovieSceneShotSection.h"
 #include "Runtime/Engine/Public/AudioDecompress.h"
 #include "Sound/SoundBase.h"
 #include "Sound/SoundWave.h"
@@ -490,6 +489,10 @@ bool FAudioTrackEditor::SupportsType( TSubclassOf<UMovieSceneTrack> Type ) const
 	return Type == UMovieSceneAudioTrack::StaticClass();
 }
 
+const FSlateBrush* FAudioTrackEditor::GetIconBrush() const
+{
+	return FEditorStyle::GetBrush("Sequencer.Tracks.Audio");
+}
 
 TSharedRef<ISequencerSection> FAudioTrackEditor::MakeSectionInterface( UMovieSceneSection& SectionObject, UMovieSceneTrack& Track )
 {
diff --git a/Engine/Source/Editor/MovieSceneTools/Private/TrackEditors/AudioTrackEditor.h b/Engine/Source/Editor/MovieSceneTools/Private/TrackEditors/AudioTrackEditor.h
--- a/Engine/Source/Editor/MovieSceneTools/Private/TrackEditors/AudioTrackEditor.h
+++ b/Engine/Source/Editor/MovieSceneTools/Private/TrackEditors/AudioTrackEditor.h
@@ -37,6 +37,7 @@ public:
 	virtual bool HandleAssetAdded(UObject* Asset, const FGuid& TargetObjectGuid) override;
 	virtual TSharedRef<ISequencerSection> MakeSectionInterface(UMovieSceneSection& SectionObject, UMovieSceneTrack& Track) override;
 	virtual bool SupportsType(TSubclassOf<UMovieSceneTrack> Type) const override;
+	virtual const FSlateBrush* GetIconBrush() const override;
 	
 protected:
 
diff --git a/Engine/Source/Editor/MovieSceneTools/Private/TrackEditors/CameraCutTrackEditor.cpp b/Engine/Source/Editor/MovieSceneTools/Private/TrackEditors/CameraCutTrackEditor.cpp
new file mode 100644
--- /dev/null
+++ b/Engine/Source/Editor/MovieSceneTools/Private/TrackEditors/CameraCutTrackEditor.cpp
@@ -0,0 +1,321 @@
+// Copyright 1998-2016 Epic Games, Inc. All Rights Reserved.
+
+#include "MovieSceneToolsPrivatePCH.h"
+#include "MovieScene.h"
+#include "MovieSceneSection.h"
+#include "ISequencerSection.h"
+#include "MovieSceneTrack.h"
+#include "MovieSceneCameraCutTrack.h"
+#include "MovieSceneCameraCutSection.h"
+#include "ScopedTransaction.h"
+#include "ISequencerObjectChangeListener.h"
+#include "IKeyArea.h"
+#include "MovieSceneTrackEditor.h"
+#include "CameraCutTrackEditor.h"
+#include "CommonMovieSceneTools.h"
+#include "Camera/CameraActor.h"
+#include "AssetToolsModule.h"
+#include "CameraCutSection.h"
+#include "SequencerUtilities.h"
+
+#include "ActorEditorUtils.h"
+#include "Editor/SceneOutliner/Public/SceneOutliner.h"
+#include "Editor/LevelEditor/Public/LevelEditor.h"
+#include "SceneOutlinerPublicTypes.h"
+
+#define LOCTEXT_NAMESPACE "FCameraCutTrackEditor"
+
+
+/* FCameraCutTrackEditor structors
+ *****************************************************************************/
+
+FCameraCutTrackEditor::FCameraCutTrackEditor(TSharedRef<ISequencer> InSequencer)
+	: FMovieSceneTrackEditor(InSequencer) 
+{
+	ThumbnailPool = MakeShareable(new FTrackEditorThumbnailPool(InSequencer));
+}
+
+
+TSharedRef<ISequencerTrackEditor> FCameraCutTrackEditor::CreateTrackEditor(TSharedRef<ISequencer> InSequencer)
+{
+	return MakeShareable(new FCameraCutTrackEditor(InSequencer));
+}
+
+
+/* ISequencerTrackEditor interface
+ *****************************************************************************/
+
+void FCameraCutTrackEditor::BuildAddTrackMenu(FMenuBuilder& MenuBuilder)
+{
+	UMovieSceneSequence* RootMovieSceneSequence = GetSequencer()->GetRootMovieSceneSequence();
+
+	if ((RootMovieSceneSequence == nullptr) || (RootMovieSceneSequence->GetClass()->GetName() != TEXT("LevelSequence")))
+	{
+		return;
+	}
+
+	MenuBuilder.AddMenuEntry(
+		LOCTEXT("AddCameraCutTrack", "Camera Cut Track"),
+		LOCTEXT("AddCameraCutTooltip", "Adds a camera cut track, as well as a new camera cut at the current scrubber location if a camera is selected."),
+		FSlateIcon(FEditorStyle::GetStyleSetName(), "Sequencer.Tracks.CameraCut"),
+		FUIAction(
+			FExecuteAction::CreateRaw(this, &FCameraCutTrackEditor::HandleAddCameraCutTrackMenuEntryExecute),
+			FCanExecuteAction::CreateRaw(this, &FCameraCutTrackEditor::HandleAddCameraCutTrackMenuEntryCanExecute)
+		)
+	);
+}
+
+TSharedPtr<SWidget> FCameraCutTrackEditor::BuildOutlinerEditWidget(const FGuid& ObjectBinding, UMovieSceneTrack* Track, const FBuildEditWidgetParams& Params)
+{
+	// Create a container edit box
+	return SNew(SHorizontalBox)
+
+	// Add the camera combo box
+	+ SHorizontalBox::Slot()
+	.AutoWidth()
+	.VAlign(VAlign_Center)
+	[
+		FSequencerUtilities::MakeAddButton(LOCTEXT("CameraCutText", "Camera"), FOnGetContent::CreateSP(this, &FCameraCutTrackEditor::HandleAddCameraCutComboButtonGetMenuContent), Params.NodeIsHovered)
+	]
+
+	+ SHorizontalBox::Slot()
+	.VAlign(VAlign_Center)
+	.HAlign(HAlign_Right)
+	.AutoWidth()
+	.Padding(4, 0, 0, 0)
+	[
+		SNew(SCheckBox)
+		.IsFocusable(false)
+		.IsChecked(this, &FCameraCutTrackEditor::IsCameraLocked)
+		.OnCheckStateChanged(this, &FCameraCutTrackEditor::OnLockCameraClicked)
+		.ToolTipText(this, &FCameraCutTrackEditor::GetLockCameraToolTip)
+		.ForegroundColor(FLinearColor::White)
+		.CheckedImage(FEditorStyle::GetBrush("Sequencer.LockCamera"))
+		.CheckedHoveredImage(FEditorStyle::GetBrush("Sequencer.LockCamera"))
+		.CheckedPressedImage(FEditorStyle::GetBrush("Sequencer.LockCamera"))
+		.UncheckedImage(FEditorStyle::GetBrush("Sequencer.UnlockCamera"))
+		.UncheckedHoveredImage(FEditorStyle::GetBrush("Sequencer.UnlockCamera"))
+		.UncheckedPressedImage(FEditorStyle::GetBrush("Sequencer.UnlockCamera"))
+	];
+}
+
+
+TSharedRef<ISequencerSection> FCameraCutTrackEditor::MakeSectionInterface(UMovieSceneSection& SectionObject, UMovieSceneTrack& Track)
+{
+	check(SupportsType(SectionObject.GetOuter()->GetClass()));
+
+	return MakeShareable(new FCameraCutSection(GetSequencer(), ThumbnailPool, SectionObject));
+}
+
+
+bool FCameraCutTrackEditor::SupportsType(TSubclassOf<UMovieSceneTrack> Type) const
+{
+	return (Type == UMovieSceneCameraCutTrack::StaticClass());
+}
+
+
+void FCameraCutTrackEditor::Tick(float DeltaTime)
+{
+	if (FSlateThrottleManager::Get().IsAllowingExpensiveTasks())
+	{
+		ThumbnailPool->DrawThumbnails();
+	}
+}
+
+const FSlateBrush* FCameraCutTrackEditor::GetIconBrush() const
+{
+	return FEditorStyle::GetBrush("Sequencer.Tracks.CameraCut");
+}
+
+
+/* FCameraCutTrackEditor implementation
+ *****************************************************************************/
+
+bool FCameraCutTrackEditor::AddKeyInternal( float KeyTime, const FGuid ObjectGuid )
+{
+	UMovieSceneCameraCutTrack* CameraCutTrack = FindOrCreateCameraCutTrack();
+	const TArray<UMovieSceneSection*>& AllSections = CameraCutTrack->GetAllSections();
+
+	CameraCutTrack->AddNewCameraCut(ObjectGuid, KeyTime);
+
+	return true;
+}
+
+
+UMovieSceneCameraCutTrack* FCameraCutTrackEditor::FindOrCreateCameraCutTrack()
+{
+	UMovieScene* FocusedMovieScene = GetFocusedMovieScene();
+	UMovieSceneTrack* CameraCutTrack = FocusedMovieScene->GetCameraCutTrack();
+
+	if (CameraCutTrack == nullptr)
+	{
+		CameraCutTrack = FocusedMovieScene->AddCameraCutTrack(UMovieSceneCameraCutTrack::StaticClass());
+	}
+
+	return CastChecked<UMovieSceneCameraCutTrack>(CameraCutTrack);
+}
+
+
+UFactory* FCameraCutTrackEditor::GetAssetFactoryForNewCameraCut( UClass* SequenceClass )
+{
+	static TWeakObjectPtr<UFactory> CameraCutFactory;
+
+	if( !CameraCutFactory.IsValid() || CameraCutFactory->SupportedClass != SequenceClass )
+	{
+		TArray<UFactory*> Factories;
+		for (TObjectIterator<UClass> It; It; ++It)
+		{
+			UClass* Class = *It;
+			if (Class->IsChildOf(UFactory::StaticClass()) && !Class->HasAnyClassFlags(CLASS_Abstract))
+			{
+				UFactory* Factory = Class->GetDefaultObject<UFactory>();
+				if (Factory->CanCreateNew())
+				{
+					UClass* SupportedClass = Factory->GetSupportedClass();
+					if ( SupportedClass == SequenceClass )
+					{
+						CameraCutFactory = Factory;
+					}
+				}
+			}
+		}
+	}
+
+	return CameraCutFactory.Get();
+}
+
+
+/* FCameraCutTrackEditor callbacks
+ *****************************************************************************/
+
+bool FCameraCutTrackEditor::HandleAddCameraCutTrackMenuEntryCanExecute() const
+{
+	UMovieScene* FocusedMovieScene = GetFocusedMovieScene();
+
+	return ((FocusedMovieScene != nullptr) && (FocusedMovieScene->GetCameraCutTrack() == nullptr));
+}
+
+
+void FCameraCutTrackEditor::HandleAddCameraCutTrackMenuEntryExecute()
+{
+	FindOrCreateCameraCutTrack();
+	GetSequencer()->NotifyMovieSceneDataChanged();
+}
+
+bool FCameraCutTrackEditor::IsCameraPickable(const AActor* const PickableActor)
+{
+	if (PickableActor->IsListedInSceneOutliner() &&
+		!FActorEditorUtils::IsABuilderBrush(PickableActor) &&
+		!PickableActor->IsA( AWorldSettings::StaticClass() ) &&
+		!PickableActor->IsPendingKill() &&
+		PickableActor->IsA( ACameraActor::StaticClass()) )
+	{			
+		return true;
+	}
+	return false;
+}
+
+TSharedRef<SWidget> FCameraCutTrackEditor::HandleAddCameraCutComboButtonGetMenuContent()
+{
+	FMenuBuilder MenuBuilder(true, nullptr);
+
+	using namespace SceneOutliner;
+
+	SceneOutliner::FInitializationOptions InitOptions;
+	{
+		InitOptions.Mode = ESceneOutlinerMode::ActorPicker;
+		InitOptions.bShowHeaderRow = false;
+		InitOptions.bFocusSearchBoxWhenOpened = true;
+		InitOptions.bShowTransient = true;
+		InitOptions.bShowCreateNewFolder = false;
+		// Only want the actor label column
+		InitOptions.ColumnMap.Add(FBuiltInColumnTypes::Label(), FColumnInfo(EColumnVisibility::Visible, 0));
+
+		// Only display Actors that we can attach too
+		InitOptions.Filters->AddFilterPredicate( SceneOutliner::FActorFilterPredicate::CreateSP(this, &FCameraCutTrackEditor::IsCameraPickable) );
+	}		
+
+	// Actor selector to allow the user to choose a parent actor
+	FSceneOutlinerModule& SceneOutlinerModule = FModuleManager::LoadModuleChecked<FSceneOutlinerModule>( "SceneOutliner" );
+
+	TSharedRef< SWidget > MenuWidget = 
+		SNew(SHorizontalBox)
+
+		+SHorizontalBox::Slot()
+		.AutoWidth()
+		[
+			SNew(SVerticalBox)
+			+SVerticalBox::Slot()
+			.MaxHeight(400.0f)
+			[
+				SceneOutlinerModule.CreateSceneOutliner(
+					InitOptions,
+					FOnActorPicked::CreateSP(this, &FCameraCutTrackEditor::HandleAddCameraCutComboButtonMenuEntryExecute )
+					)
+			]
+		];
+
+	MenuBuilder.AddWidget(MenuWidget, FText::GetEmpty(), false);
+	return MenuBuilder.MakeWidget();
+}
+
+
+void FCameraCutTrackEditor::HandleAddCameraCutComboButtonMenuEntryExecute(AActor* Camera)
+{
+	FGuid ObjectGuid = FindOrCreateHandleToObject(Camera).Handle;
+
+	if (ObjectGuid.IsValid())
+	{
+		AnimatablePropertyChanged(FOnKeyProperty::CreateRaw(this, &FCameraCutTrackEditor::AddKeyInternal, ObjectGuid));
+	}
+}
+
+ECheckBoxState FCameraCutTrackEditor::IsCameraLocked() const
+{
+	if (GetSequencer()->IsPerspectiveViewportCameraCutEnabled())
+	{
+		return ECheckBoxState::Checked;
+	}
+	else
+	{
+		return ECheckBoxState::Unchecked;
+	}
+}
+
+
+void FCameraCutTrackEditor::OnLockCameraClicked(ECheckBoxState CheckBoxState)
+{
+	if (CheckBoxState == ECheckBoxState::Checked)
+	{
+		for (int32 i = 0; i < GEditor->LevelViewportClients.Num(); ++i)
+		{		
+			FLevelEditorViewportClient* LevelVC = GEditor->LevelViewportClients[i];
+			if (LevelVC && LevelVC->IsPerspective() && LevelVC->AllowsCinematicPreview() && LevelVC->GetViewMode() != VMI_Unknown)
+			{
+				LevelVC->SetActorLock(nullptr);
+				LevelVC->bLockedCameraView = false;
+				LevelVC->UpdateViewForLockedActor();
+				LevelVC->Invalidate();
+			}
+		}
+		GetSequencer()->SetPerspectiveViewportCameraCutEnabled(true);
+	}
+	else
+	{
+		// TODO: MaxC, is this right?
+		GetSequencer()->UpdateCameraCut(nullptr, nullptr);
+		GetSequencer()->SetPerspectiveViewportCameraCutEnabled(false);
+	}
+
+	GetSequencer()->UpdateRuntimeInstances();
+}
+
+
+FText FCameraCutTrackEditor::GetLockCameraToolTip() const
+{
+	return IsCameraLocked() == ECheckBoxState::Checked ?
+		LOCTEXT("UnlockCamera", "Unlock Viewport from Camera Cuts") :
+		LOCTEXT("LockCamera", "Lock Viewport to Camera Cuts");
+}
+
+#undef LOCTEXT_NAMESPACE
diff --git a/Engine/Source/Editor/MovieSceneTools/Private/TrackEditors/CinematicShotTrackEditor.cpp b/Engine/Source/Editor/MovieSceneTools/Private/TrackEditors/CinematicShotTrackEditor.cpp
new file mode 100644
--- /dev/null
+++ b/Engine/Source/Editor/MovieSceneTools/Private/TrackEditors/CinematicShotTrackEditor.cpp
@@ -0,0 +1,646 @@
+// Copyright 1998-2016 Epic Games, Inc. All Rights Reserved.
+
+#include "MovieSceneToolsPrivatePCH.h"
+#include "MovieScene.h"
+#include "MovieSceneSection.h"
+#include "ISequencerSection.h"
+#include "MovieSceneTrack.h"
+#include "MovieSceneCinematicShotTrack.h"
+#include "MovieSceneCinematicShotSection.h"
+#include "ScopedTransaction.h"
+#include "ISequencerObjectChangeListener.h"
+#include "IKeyArea.h"
+#include "MovieSceneTrackEditor.h"
+#include "MovieSceneToolHelpers.h"
+#include "CinematicShotTrackEditor.h"
+#include "CommonMovieSceneTools.h"
+#include "Camera/CameraActor.h"
+#include "CinematicShotSection.h"
+#include "SequencerUtilities.h"
+#include "ContentBrowserModule.h"
+#include "MovieSceneToolHelpers.h"
+#include "LevelSequence.h"
+#include "AssetToolsModule.h"
+#include "AssetRegistryModule.h"
+
+#define LOCTEXT_NAMESPACE "FCinematicShotTrackEditor"
+
+
+/* FCinematicShotTrackEditor structors
+ *****************************************************************************/
+
+FCinematicShotTrackEditor::FCinematicShotTrackEditor(TSharedRef<ISequencer> InSequencer)
+	: FMovieSceneTrackEditor(InSequencer) 
+{
+	ThumbnailPool = MakeShareable(new FTrackEditorThumbnailPool(InSequencer));
+}
+
+
+TSharedRef<ISequencerTrackEditor> FCinematicShotTrackEditor::CreateTrackEditor(TSharedRef<ISequencer> InSequencer)
+{
+	return MakeShareable(new FCinematicShotTrackEditor(InSequencer));
+}
+
+void FCinematicShotTrackEditor::OnInitialize()
+{
+	OnCameraCutHandle = GetSequencer()->OnCameraCut().AddSP(this, &FCinematicShotTrackEditor::OnUpdateCameraCut);
+}
+	 
+void FCinematicShotTrackEditor::OnRelease()
+{
+	if (OnCameraCutHandle.IsValid() && GetSequencer().IsValid())
+	{
+		GetSequencer()->OnCameraCut().Remove(OnCameraCutHandle);
+	}
+}
+
+
+/* ISequencerTrackEditor interface
+ *****************************************************************************/
+
+void FCinematicShotTrackEditor::BuildAddTrackMenu(FMenuBuilder& MenuBuilder)
+{
+	UMovieSceneSequence* RootMovieSceneSequence = GetSequencer()->GetRootMovieSceneSequence();
+
+	if ((RootMovieSceneSequence == nullptr) || (RootMovieSceneSequence->GetClass()->GetName() != TEXT("LevelSequence")))
+	{
+		return;
+	}
+
+	MenuBuilder.AddMenuEntry(
+		LOCTEXT("AddCinematicShotTrack", "Shot Track"),
+		LOCTEXT("AddCinematicShotTooltip", "Adds a shot track."),
+		FSlateIcon(FEditorStyle::GetStyleSetName(), "Sequencer.Tracks.CinematicShot"),
+		FUIAction(
+			FExecuteAction::CreateRaw(this, &FCinematicShotTrackEditor::HandleAddCinematicShotTrackMenuEntryExecute),
+			FCanExecuteAction::CreateRaw(this, &FCinematicShotTrackEditor::HandleAddCinematicShotTrackMenuEntryCanExecute)
+		)
+	);
+}
+
+TSharedPtr<SWidget> FCinematicShotTrackEditor::BuildOutlinerEditWidget(const FGuid& ObjectBinding, UMovieSceneTrack* Track, const FBuildEditWidgetParams& Params)
+{
+	// Create a container edit box
+	return SNew(SHorizontalBox)
+
+	// Add the camera combo box
+	+ SHorizontalBox::Slot()
+	.AutoWidth()
+	.VAlign(VAlign_Center)
+	[
+		FSequencerUtilities::MakeAddButton(LOCTEXT("CinematicShotText", "Shot"), FOnGetContent::CreateSP(this, &FCinematicShotTrackEditor::HandleAddCinematicShotComboButtonGetMenuContent), Params.NodeIsHovered)
+	]
+
+	+ SHorizontalBox::Slot()
+	.VAlign(VAlign_Center)
+	.HAlign(HAlign_Right)
+	.AutoWidth()
+	.Padding(4, 0, 0, 0)
+	[
+		SNew(SCheckBox)
+		.IsFocusable(false)
+		.IsChecked(this, &FCinematicShotTrackEditor::AreShotsLocked)
+		.OnCheckStateChanged(this, &FCinematicShotTrackEditor::OnLockShotsClicked)
+		.ToolTipText(this, &FCinematicShotTrackEditor::GetLockShotsToolTip)
+		.ForegroundColor(FLinearColor::White)
+		.CheckedImage(FEditorStyle::GetBrush("Sequencer.LockCamera"))
+		.CheckedHoveredImage(FEditorStyle::GetBrush("Sequencer.LockCamera"))
+		.CheckedPressedImage(FEditorStyle::GetBrush("Sequencer.LockCamera"))
+		.UncheckedImage(FEditorStyle::GetBrush("Sequencer.UnlockCamera"))
+		.UncheckedHoveredImage(FEditorStyle::GetBrush("Sequencer.UnlockCamera"))
+		.UncheckedPressedImage(FEditorStyle::GetBrush("Sequencer.UnlockCamera"))
+	];
+}
+
+
+TSharedRef<ISequencerSection> FCinematicShotTrackEditor::MakeSectionInterface(UMovieSceneSection& SectionObject, UMovieSceneTrack& Track)
+{
+	check(SupportsType(SectionObject.GetOuter()->GetClass()));
+
+	return MakeShareable(new FCinematicShotSection(GetSequencer(), ThumbnailPool, SectionObject, SharedThis(this)));
+}
+
+bool FCinematicShotTrackEditor::HandleAssetAdded(UObject* Asset, const FGuid& TargetObjectGuid)
+{
+	UMovieSceneSequence* Sequence = Cast<UMovieSceneSequence>(Asset);
+
+	if (Sequence == nullptr)
+	{
+		return false;
+	}
+
+
+	//@todo If there's already a subscenes track, allow that track to handle this asset
+	UMovieScene* FocusedMovieScene = GetFocusedMovieScene();
+
+	if (FocusedMovieScene != nullptr && FocusedMovieScene->FindMasterTrack<UMovieSceneSubTrack>() != nullptr)
+	{
+		return false;
+	}
+
+	if (CanAddSubSequence(*Sequence))
+	{
+		AnimatablePropertyChanged(FOnKeyProperty::CreateRaw(this, &FCinematicShotTrackEditor::HandleSequenceAdded, Sequence));
+
+		return true;
+	}
+
+	return false;
+}
+
+bool FCinematicShotTrackEditor::SupportsType(TSubclassOf<UMovieSceneTrack> Type) const
+{
+	return (Type == UMovieSceneCinematicShotTrack::StaticClass());
+}
+
+
+void FCinematicShotTrackEditor::Tick(float DeltaTime)
+{
+	if (FSlateThrottleManager::Get().IsAllowingExpensiveTasks())
+	{
+		ThumbnailPool->DrawThumbnails();
+	}
+}
+
+const FSlateBrush* FCinematicShotTrackEditor::GetIconBrush() const
+{
+	return FEditorStyle::GetBrush("Sequencer.Tracks.CinematicShot");
+}
+
+bool IsPackageNameUnique(const TArray<FAssetData>& ObjectList, FString& NewPackageName)
+{
+	for (auto AssetObject : ObjectList)
+	{
+		if (AssetObject.PackageName.ToString() == NewPackageName)
+		{
+			return false;
+		}
+	}
+	return true;
+}
+
+UMovieSceneSubSection* FCinematicShotTrackEditor::CreateShotInternal(FString& NewShotName, float NewShotStartTime, UMovieSceneCinematicShotSection* ShotToDuplicate, const bool& bInsertShot)
+{
+	const UMovieSceneToolsProjectSettings* ProjectSettings = GetDefault<UMovieSceneToolsProjectSettings>();
+
+	FAssetRegistryModule& AssetRegistryModule = FModuleManager::LoadModuleChecked<FAssetRegistryModule>(TEXT("AssetRegistry"));
+	TArray<FAssetData> ObjectList;
+	AssetRegistryModule.Get().GetAssetsByClass(ULevelSequence::StaticClass()->GetFName(), ObjectList);
+
+	UObject* SequenceAsset = GetSequencer()->GetFocusedMovieSceneSequence()->GetMovieScene()->GetOuter();
+	UPackage* SequencePackage = SequenceAsset->GetOutermost();
+
+	FString SequencePackageName = SequencePackage->GetName(); // ie. /Game/cine/max/master
+	int LastSlashPos = SequencePackageName.Find(TEXT("/"), ESearchCase::IgnoreCase, ESearchDir::FromEnd);
+	FString SequencePath = SequencePackageName.Left(LastSlashPos);
+
+	FString NewShotPrefix;
+	uint32 NewShotNumber = INDEX_NONE;
+	uint32 NewTakeNumber = INDEX_NONE;
+	MovieSceneToolHelpers::ParseShotName(NewShotName, NewShotPrefix, NewShotNumber, NewTakeNumber);
+
+	FString NewShotDirectory = MovieSceneToolHelpers::ComposeShotName(NewShotPrefix, NewShotNumber, INDEX_NONE);
+	FString NewShotPath = SequencePath;
+	NewShotPath /= NewShotDirectory; // put this in the shot directory, ie. /Game/cine/max/sht0010
+
+	// Make sure this shot path is unique
+	FString NewPackageName = NewShotPath;
+	NewPackageName /= NewShotName; // ie. /Game/cine/max/sht0010/sht0010_001
+	if (!IsPackageNameUnique(ObjectList, NewPackageName))
+	{
+		while (1)
+		{
+			NewShotNumber += ProjectSettings->ShotIncrement;
+			NewShotName = MovieSceneToolHelpers::ComposeShotName(NewShotPrefix, NewShotNumber, NewTakeNumber);
+			NewShotDirectory = MovieSceneToolHelpers::ComposeShotName(NewShotPrefix, NewShotNumber, INDEX_NONE);
+			NewShotPath = SequencePath;
+			NewShotPath /= NewShotDirectory;
+
+			NewPackageName = NewShotPath;
+			NewPackageName /= NewShotName;
+			if (IsPackageNameUnique(ObjectList, NewPackageName))
+			{
+				break;
+			}
+		}
+	}
+
+	// Create a new level sequence asset with the appropriate name
+	IAssetTools& AssetTools = FModuleManager::GetModuleChecked<FAssetToolsModule>("AssetTools").Get();
+
+	UObject* NewAsset = nullptr;
+	for (TObjectIterator<UClass> It ; It ; ++It)
+	{
+		UClass* CurrentClass = *It;
+		if (CurrentClass->IsChildOf(UFactory::StaticClass()) && !(CurrentClass->HasAnyClassFlags(CLASS_Abstract)))
+		{
+			UFactory* Factory = Cast<UFactory>(CurrentClass->GetDefaultObject());
+			if (Factory->CanCreateNew() && Factory->ImportPriority >= 0 && Factory->SupportedClass == ULevelSequence::StaticClass())
+			{
+				if (ShotToDuplicate != nullptr)
+				{
+					NewAsset = AssetTools.DuplicateAssetWithDialog(NewShotName, NewShotPath, ShotToDuplicate->GetSequence());
+				}
+				else
+				{
+					NewAsset = AssetTools.CreateAssetWithDialog(NewShotName, NewShotPath, ULevelSequence::StaticClass(), Factory);
+				}
+				break;
+			}
+		}
+	}
+
+	if (NewAsset == nullptr)
+	{
+		return nullptr;
+	}
+
+	UMovieSceneSequence* NewSequence = Cast<UMovieSceneSequence>(NewAsset);
+
+	float Duration = NewSequence->GetMovieScene()->GetPlaybackRange().Size<float>();
+	if (ShotToDuplicate != nullptr)
+	{
+		Duration = ShotToDuplicate->GetEndTime() - ShotToDuplicate->GetStartTime();
+	}
+
+	// Create a cinematic shot section and set the time to be at the end of this shot. 
+	UMovieSceneCinematicShotTrack* CinematicShotTrack = FindOrCreateCinematicShotTrack();
+	UMovieSceneSubSection* NewSection = CinematicShotTrack->AddSequence(NewSequence, NewShotStartTime, Duration, bInsertShot);	
+
+	// Update the track instances to generate the new section
+	GetSequencer()->UpdateRuntimeInstances();
+
+	// Focus on the newly created shot
+	//GetSequencer()->FocusSequenceInstance(GetSequencer()->GetSequenceInstanceForSection(*NewSection));
+
+	return NewSection;
+}
+
+
+void FCinematicShotTrackEditor::InsertShotAtCurrentTime()
+{
+	const FScopedTransaction Transaction(LOCTEXT("InsertShot_Transaction", "Insert Shot"));
+
+	float NewShotStartTime = GetSequencer()->GetGlobalTime();
+
+	UMovieSceneCinematicShotTrack* CinematicShotTrack = FindOrCreateCinematicShotTrack();
+	FString NewShotName = MovieSceneToolHelpers::GenerateNewShotName(CinematicShotTrack->GetAllSections(), NewShotStartTime);
+
+	CreateShotInternal(NewShotName, NewShotStartTime);
+
+	GetSequencer()->NotifyMovieSceneDataChanged();
+}
+
+
+void FCinematicShotTrackEditor::InsertShot(UMovieSceneCinematicShotSection* Section)
+{
+	const FScopedTransaction Transaction(LOCTEXT("InsertShot_Transaction", "Insert Shot"));
+
+	// Start this new shot after the current shot
+	float NewShotStartTime = Section->GetEndTime();
+
+	UMovieSceneCinematicShotTrack* CinematicShotTrack = FindOrCreateCinematicShotTrack();
+	FString NewShotName = MovieSceneToolHelpers::GenerateNewShotName(CinematicShotTrack->GetAllSections(), NewShotStartTime);
+
+	CreateShotInternal(NewShotName, NewShotStartTime);
+
+	GetSequencer()->NotifyMovieSceneDataChanged();
+}
+
+
+void FCinematicShotTrackEditor::InsertFillerAtCurrentTime()
+{
+	const UMovieSceneToolsProjectSettings* ProjectSettings = GetDefault<UMovieSceneToolsProjectSettings>();
+
+	const FScopedTransaction Transaction(LOCTEXT("InsertFiller_Transaction", "Insert Filler"));
+
+	float NewShotStartTime = GetSequencer()->GetGlobalTime();
+
+	UMovieSceneCinematicShotTrack* CinematicShotTrack = FindOrCreateCinematicShotTrack();
+
+	float Duration = ProjectSettings->DefaultDuration;
+	const bool bInsertShot = true;
+
+	UMovieSceneSequence* NullSequence = nullptr;
+
+	UMovieSceneSubSection* NewSection = CinematicShotTrack->AddSequence(NullSequence, NewShotStartTime, Duration, bInsertShot);	
+
+	UMovieSceneCinematicShotSection* NewCinematicShotSection = Cast<UMovieSceneCinematicShotSection>(NewSection);
+
+	NewCinematicShotSection->SetShotDisplayName(FText(LOCTEXT("Filler", "Filler")));
+
+	GetSequencer()->NotifyMovieSceneDataChanged();
+}
+
+void FCinematicShotTrackEditor::DuplicateShot(UMovieSceneCinematicShotSection* Section)
+{
+	const FScopedTransaction Transaction(LOCTEXT("DuplicateShot_Transaction", "Duplicate Shot"));
+
+	// Start this new shot after the current shot
+	float NewShotStartTime = Section->GetEndTime();
+
+	UMovieSceneCinematicShotTrack* CinematicShotTrack = FindOrCreateCinematicShotTrack();
+	FString NewShotName = MovieSceneToolHelpers::GenerateNewShotName(CinematicShotTrack->GetAllSections(), NewShotStartTime);
+
+	UMovieSceneSubSection* NewShot = CreateShotInternal(NewShotName, NewShotStartTime, Section);
+	if (NewShot)
+	{
+		NewShot->SetEndTime(NewShot->GetStartTime() + Section->GetEndTime() - Section->GetStartTime());
+		NewShot->StartOffset = Section->StartOffset;
+		NewShot->TimeScale = Section->TimeScale;
+		NewShot->PrerollTime = Section->PrerollTime;
+	}
+
+	GetSequencer()->NotifyMovieSceneDataChanged();
+}
+
+
+void FCinematicShotTrackEditor::RenameShot(UMovieSceneCinematicShotSection* Section)
+{
+	//@todo
+}
+
+
+void FCinematicShotTrackEditor::NewTake(UMovieSceneCinematicShotSection* Section)
+{
+	const FScopedTransaction Transaction(LOCTEXT("NewTake_Transaction", "New Take"));
+
+	FString ShotPrefix;
+	uint32 ShotNumber = INDEX_NONE;
+	uint32 TakeNumber = INDEX_NONE;
+	if (MovieSceneToolHelpers::ParseShotName(Section->GetShotDisplayName().ToString(), ShotPrefix, ShotNumber, TakeNumber))
+	{
+		uint32 NewTakeNumber = TakeNumber + 1;
+		FString NewShotName = MovieSceneToolHelpers::ComposeShotName(ShotPrefix, ShotNumber, NewTakeNumber);
+
+		float NewShotStartTime = Section->GetStartTime();
+		float NewShotEndTime = Section->GetEndTime();
+		float NewShotStartOffset = Section->StartOffset;
+		float NewShotTimeScale = Section->TimeScale;
+		float NewShotPrerollTime = Section->PrerollTime;
+
+		const bool bInsertShot = false;
+		UMovieSceneSubSection* NewShot = CreateShotInternal(NewShotName, NewShotStartTime, Section, bInsertShot);
+
+		UMovieSceneCinematicShotTrack* CinematicShotTrack = FindOrCreateCinematicShotTrack();
+		CinematicShotTrack->RemoveSection(*Section);
+
+		if (NewShot != nullptr)
+		{
+			NewShot->SetStartTime(NewShotStartTime);
+			NewShot->SetEndTime(NewShotEndTime);
+			NewShot->StartOffset = NewShotStartOffset;
+			NewShot->TimeScale = NewShotTimeScale;
+			NewShot->PrerollTime = NewShotPrerollTime;
+		}
+	}
+
+	GetSequencer()->NotifyMovieSceneDataChanged();
+}
+
+
+void FCinematicShotTrackEditor::SwitchTake(UMovieSceneCinematicShotSection* Section, uint32 TakeNumber)
+{
+	const FScopedTransaction Transaction(LOCTEXT("SwitchTake_Transaction", "Switch Take"));
+
+	UObject* TakeObject = MovieSceneToolHelpers::GetTake(Section, TakeNumber);
+
+	if (TakeObject && TakeObject->IsA(UMovieSceneSequence::StaticClass()))
+	{
+		UMovieSceneSequence* MovieSceneSequence = CastChecked<UMovieSceneSequence>(TakeObject);
+
+		UMovieSceneCinematicShotTrack* CinematicShotTrack = FindOrCreateCinematicShotTrack();
+		
+		float NewShotStartTime = Section->GetStartTime();
+		float NewShotEndTime = Section->GetEndTime();
+		float NewShotStartOffset = Section->StartOffset;
+		float NewShotTimeScale = Section->TimeScale;
+		float NewShotPrerollTime = Section->PrerollTime;
+
+		const bool bInsertSequence = false;
+		const float Duration = NewShotEndTime - NewShotStartTime;
+		UMovieSceneSubSection* NewShot = CinematicShotTrack->AddSequence(MovieSceneSequence, NewShotStartTime, Duration, bInsertSequence);	
+		
+		CinematicShotTrack->RemoveSection(*Section);
+
+		if (NewShot != nullptr)
+		{
+			NewShot->SetStartTime(NewShotStartTime);
+			NewShot->SetEndTime(NewShotEndTime);
+			NewShot->StartOffset = NewShotStartOffset;
+			NewShot->TimeScale = NewShotTimeScale;
+			NewShot->PrerollTime = NewShotPrerollTime;
+		}
+	}
+
+	GetSequencer()->NotifyMovieSceneDataChanged();
+}
+
+
+/* FCinematicShotTrackEditor callbacks
+ *****************************************************************************/
+
+bool FCinematicShotTrackEditor::HandleAddCinematicShotTrackMenuEntryCanExecute() const
+{
+	UMovieScene* FocusedMovieScene = GetFocusedMovieScene();
+
+	return ((FocusedMovieScene != nullptr) && (FocusedMovieScene->FindMasterTrack<UMovieSceneCinematicShotTrack>() == nullptr));
+}
+
+
+void FCinematicShotTrackEditor::HandleAddCinematicShotTrackMenuEntryExecute()
+{
+	FindOrCreateCinematicShotTrack();
+}
+
+TSharedRef<SWidget> FCinematicShotTrackEditor::HandleAddCinematicShotComboButtonGetMenuContent()
+{
+	FMenuBuilder MenuBuilder(true, nullptr);
+
+	MenuBuilder.AddMenuEntry(
+			LOCTEXT("InsertShot", "Insert Shot"),
+			LOCTEXT("InsertShotTooltip", "Insert new shot"),
+			FSlateIcon(),
+			FUIAction(FExecuteAction::CreateSP(this, &FCinematicShotTrackEditor::InsertShotAtCurrentTime))
+	);
+
+	MenuBuilder.AddMenuEntry(
+			LOCTEXT("InsertFiller", "Insert Filler"),
+			LOCTEXT("InsertFillerTooltip", "Insert filler"),
+			FSlateIcon(),
+			FUIAction(FExecuteAction::CreateSP(this, &FCinematicShotTrackEditor::InsertFillerAtCurrentTime))
+	);
+
+	FAssetPickerConfig AssetPickerConfig;
+	{
+		AssetPickerConfig.OnAssetSelected = FOnAssetSelected::CreateRaw( this, &FCinematicShotTrackEditor::HandleAddCinematicShotComboButtonMenuEntryExecute);
+		AssetPickerConfig.bAllowNullSelection = false;
+		AssetPickerConfig.InitialAssetViewType = EAssetViewType::Tile;
+		AssetPickerConfig.Filter.ClassNames.Add(TEXT("LevelSequence"));
+	}
+
+	FContentBrowserModule& ContentBrowserModule = FModuleManager::Get().LoadModuleChecked<FContentBrowserModule>(TEXT("ContentBrowser"));
+
+	TSharedPtr<SBox> MenuEntry = SNew(SBox)
+		.WidthOverride(300.0f)
+		.HeightOverride(300.f)
+		[
+			ContentBrowserModule.Get().CreateAssetPicker(AssetPickerConfig)
+		];
+
+	MenuBuilder.AddWidget(MenuEntry.ToSharedRef(), FText::GetEmpty(), true);
+
+	return MenuBuilder.MakeWidget();
+}
+
+
+void FCinematicShotTrackEditor::HandleAddCinematicShotComboButtonMenuEntryExecute(const FAssetData& AssetData)
+{
+	FSlateApplication::Get().DismissAllMenus();
+
+	UObject* SelectedObject = AssetData.GetAsset();
+
+	if (SelectedObject && SelectedObject->IsA(UMovieSceneSequence::StaticClass()))
+	{
+		UMovieSceneSequence* MovieSceneSequence = CastChecked<UMovieSceneSequence>(AssetData.GetAsset());
+
+		AnimatablePropertyChanged( FOnKeyProperty::CreateRaw( this, &FCinematicShotTrackEditor::AddKeyInternal, MovieSceneSequence) );
+	}
+}
+
+bool FCinematicShotTrackEditor::AddKeyInternal(float KeyTime, UMovieSceneSequence* InMovieSceneSequence)
+{	
+	if (CanAddSubSequence(*InMovieSceneSequence))
+	{
+		UMovieSceneCinematicShotTrack* CinematicShotTrack = FindOrCreateCinematicShotTrack();
+		float EndTime = InMovieSceneSequence->GetMovieScene()->GetPlaybackRange().Size<float>();
+		CinematicShotTrack->AddSequence(InMovieSceneSequence, KeyTime, EndTime);
+
+		GetSequencer()->NotifyMovieSceneDataChanged();
+	
+		return true;
+	}
+
+	return false;
+}
+
+UMovieSceneCinematicShotTrack* FCinematicShotTrackEditor::FindOrCreateCinematicShotTrack()
+{
+	UMovieScene* FocusedMovieScene = GetFocusedMovieScene();
+	
+	if (FocusedMovieScene == nullptr)
+	{
+		return nullptr;
+	}
+
+	UMovieSceneCinematicShotTrack* CinematicShotTrack = FocusedMovieScene->FindMasterTrack<UMovieSceneCinematicShotTrack>();
+	if (CinematicShotTrack != nullptr)
+	{
+		return CinematicShotTrack;
+	}
+
+	const FScopedTransaction Transaction(LOCTEXT("AddCinematicShotTrack_Transaction", "Add Cinematic Shot Track"));
+	FocusedMovieScene->Modify();
+
+	auto NewTrack = FocusedMovieScene->AddMasterTrack<UMovieSceneCinematicShotTrack>();
+	ensure(NewTrack);
+
+	GetSequencer()->NotifyMovieSceneDataChanged();
+
+	return NewTrack;
+}
+
+ECheckBoxState FCinematicShotTrackEditor::AreShotsLocked() const
+{
+	if (GetSequencer()->IsPerspectiveViewportCameraCutEnabled())
+	{
+		return ECheckBoxState::Checked;
+	}
+	else
+	{
+		return ECheckBoxState::Unchecked;
+	}
+}
+
+
+void FCinematicShotTrackEditor::OnLockShotsClicked(ECheckBoxState CheckBoxState)
+{
+	if (CheckBoxState == ECheckBoxState::Checked)
+	{
+		for (int32 i = 0; i < GEditor->LevelViewportClients.Num(); ++i)
+		{		
+			FLevelEditorViewportClient* LevelVC = GEditor->LevelViewportClients[i];
+			if (LevelVC && LevelVC->IsPerspective() && LevelVC->AllowsCinematicPreview() && LevelVC->GetViewMode() != VMI_Unknown)
+			{
+				LevelVC->SetActorLock(nullptr);
+				LevelVC->bLockedCameraView = false;
+				LevelVC->UpdateViewForLockedActor();
+				LevelVC->Invalidate();
+			}
+		}
+		GetSequencer()->SetPerspectiveViewportCameraCutEnabled(true);
+	}
+	else
+	{
+		GetSequencer()->UpdateCameraCut(nullptr, nullptr);
+		GetSequencer()->SetPerspectiveViewportCameraCutEnabled(false);
+	}
+
+	GetSequencer()->UpdateRuntimeInstances();
+}
+
+
+FText FCinematicShotTrackEditor::GetLockShotsToolTip() const
+{
+	return AreShotsLocked() == ECheckBoxState::Checked ?
+		LOCTEXT("UnlockShots", "Unlock Viewport from Shots") :
+		LOCTEXT("LockShots", "Lock Viewport to Shots");
+}
+
+bool FCinematicShotTrackEditor::CanAddSubSequence(const UMovieSceneSequence& Sequence) const
+{
+	// prevent adding ourselves and ensure we have a valid movie scene
+	UMovieSceneSequence* FocusedSequence = GetSequencer()->GetFocusedMovieSceneSequence();
+
+	if ((FocusedSequence == nullptr) || (FocusedSequence == &Sequence) || (FocusedSequence->GetMovieScene() == nullptr))
+	{
+		return false;
+	}
+
+	// ensure that the other sequence has a valid movie scene
+	UMovieScene* SequenceMovieScene = Sequence.GetMovieScene();
+
+	if (SequenceMovieScene == nullptr)
+	{
+		return false;
+	}
+
+	// make sure we are not contained in the other sequence (circular dependency)
+	// @todo sequencer: this check is not sufficient (does not prevent circular dependencies of 2+ levels)
+	UMovieSceneCinematicShotTrack* SequenceCinematicShotTrack = SequenceMovieScene->FindMasterTrack<UMovieSceneCinematicShotTrack>();
+
+	if (SequenceCinematicShotTrack == nullptr)
+	{
+		return true;
+	}
+
+	return !SequenceCinematicShotTrack->ContainsSequence(*FocusedSequence, true);
+}
+
+void FCinematicShotTrackEditor::OnUpdateCameraCut(UObject* CameraObject)
+{
+	// Keep track of the camera when it switches so that the thumbnail can be drawn with the correct camera
+	CinematicShotCamera = Cast<ACameraActor>(CameraObject);
+}
+
+bool FCinematicShotTrackEditor::HandleSequenceAdded(float KeyTime, UMovieSceneSequence* Sequence)
+{
+	auto CinematicShotTrack = FindOrCreateCinematicShotTrack();
+	float Duration = Sequence->GetMovieScene()->GetPlaybackRange().Size<float>();
+	CinematicShotTrack->AddSequence(Sequence, KeyTime, Duration);
+
+	GetSequencer()->NotifyMovieSceneDataChanged();
+	return true;
+}
+
+#undef LOCTEXT_NAMESPACE
diff --git a/Engine/Source/Editor/MovieSceneTools/Private/TrackEditors/CinematicShotTrackEditor.h b/Engine/Source/Editor/MovieSceneTools/Private/TrackEditors/CinematicShotTrackEditor.h
new file mode 100644
--- /dev/null
+++ b/Engine/Source/Editor/MovieSceneTools/Private/TrackEditors/CinematicShotTrackEditor.h
@@ -0,0 +1,165 @@
+// Copyright 1998-2016 Epic Games, Inc. All Rights Reserved.
+
+#pragma once
+
+
+class FTrackEditorThumbnailPool;
+class UMovieSceneCinematicShotTrack;
+class UMovieSceneSubSection;
+
+
+
+/**
+ * Tools for cinematic shots.
+ */
+class FCinematicShotTrackEditor
+	: public FMovieSceneTrackEditor
+{
+public:
+
+	/**
+	 * Constructor
+	 *
+	 * @param InSequencer The sequencer instance to be used by this tool.
+	 */
+	FCinematicShotTrackEditor(TSharedRef<ISequencer> InSequencer);
+
+	/** Virtual destructor. */
+	virtual ~FCinematicShotTrackEditor() { }
+
+	/**
+	 * Creates an instance of this class.  Called by a sequencer .
+	 *
+	 * @param OwningSequencer The sequencer instance to be used by this tool.
+	 * @return The new instance of this class.
+	 */
+	static TSharedRef<ISequencerTrackEditor> CreateTrackEditor(TSharedRef<ISequencer> OwningSequencer);
+
+	TWeakObjectPtr<ACameraActor> GetCinematicShotCamera() const { return CinematicShotCamera; }
+
+public:
+
+	// ISequencerTrackEditor interface
+	virtual void OnInitialize() override;
+	virtual void OnRelease() override;
+	virtual void BuildAddTrackMenu(FMenuBuilder& MenuBuilder) override;
+	virtual TSharedPtr<SWidget> BuildOutlinerEditWidget(const FGuid& ObjectBinding, UMovieSceneTrack* Track, const FBuildEditWidgetParams& Params) override;
+	virtual TSharedRef<ISequencerSection> MakeSectionInterface(UMovieSceneSection& SectionObject, UMovieSceneTrack& Track) override;
+	virtual bool HandleAssetAdded(UObject* Asset, const FGuid& TargetObjectGuid) override;
+	virtual bool SupportsType(TSubclassOf<UMovieSceneTrack> Type ) const override;
+	virtual void Tick(float DeltaTime) override;
+	virtual const FSlateBrush* GetIconBrush() const override;
+
+
+	/*
+	 * Insert shot. 
+	 */
+	void InsertShotAtCurrentTime();
+	void InsertShot(UMovieSceneCinematicShotSection* Section);
+
+	/*
+	 * Insert filler.
+	 */
+	void InsertFillerAtCurrentTime();
+
+	/*
+	 * Duplicate shot. 
+	 *
+	 * @param Section The section to duplicate
+	 */
+	void DuplicateShot(UMovieSceneCinematicShotSection* Section);
+
+	/*
+	 * Rename shot. 
+	 *
+	 * @param Section The section to rename.
+	 */
+	void RenameShot(UMovieSceneCinematicShotSection* Section);
+
+	/*
+	 * New take. 
+	 *
+	 * @param Section The section to create a new take of.
+	 */
+	void NewTake(UMovieSceneCinematicShotSection* Section);
+
+	/*
+	* Switch take
+	*
+	* @param Section The section to switch the take of.
+	* @param TakeNumber The take number to switch to.
+	*/
+	void SwitchTake(UMovieSceneCinematicShotSection* Section, uint32 TakeNumber);
+
+
+private:
+
+	/*
+	 * Create shot
+	 *
+	 * @param NewShotName The new shot name.
+	 * @param NewShotStartTime The time to start the new shot at.
+	 * @param ShotToDuplicate The shot to duplicate.
+	 * @param bInsertShot Whether the shot should be inserted and the subsequent shots be pushed out.
+	 * @return The new shot.
+	 */
+	UMovieSceneSubSection* CreateShotInternal(FString& NewShotName, float NewShotStartTime, UMovieSceneCinematicShotSection* ShotToDuplicate = nullptr, const bool& bInsertShot = true);
+
+
+private:
+
+	/** Callback for determining whether the "Add Shot" menu entry can execute. */
+	bool HandleAddCinematicShotTrackMenuEntryCanExecute() const;
+
+	/** Callback for executing the "Add Shot Track" menu entry. */
+	void HandleAddCinematicShotTrackMenuEntryExecute();
+
+	/** Callback for generating the menu of the "Add Shot" combo button. */
+	TSharedRef<SWidget> HandleAddCinematicShotComboButtonGetMenuContent();
+
+	/** Callback for executing a menu entry in the "Add Shot" combo button. */
+	void HandleAddCinematicShotComboButtonMenuEntryExecute(const FAssetData& AssetData);
+
+	/** Find or create a cinematic shot track in the currently focused movie scene. */
+	UMovieSceneCinematicShotTrack* FindOrCreateCinematicShotTrack();
+
+	/** Delegate for AnimatablePropertyChanged in AddKey */
+	bool AddKeyInternal(float KeyTime, UMovieSceneSequence* InMovieSceneSequence);
+
+	/** Delegate for shots button lock state */
+	ECheckBoxState AreShotsLocked() const; 
+
+	/** Delegate for locked shots button */
+	void OnLockShotsClicked(ECheckBoxState CheckBoxState);
+
+	/** Delegate for shots button lock tooltip */
+	FText GetLockShotsToolTip() const; 
+
+	/**
+	 * Check whether the given sequence can be added as a sub-sequence.
+	 *
+	 * The purpose of this method is to disallow circular references
+	 * between sub-sequences in the focused movie scene.
+	 *
+	 * @param Sequence The sequence to check.
+	 * @return true if the sequence can be added as a sub-sequence, false otherwise.
+	 */
+	bool CanAddSubSequence(const UMovieSceneSequence& Sequence) const;
+
+	/** Called when our sequencer wants to switch cameras */
+	void OnUpdateCameraCut(UObject* CameraObject);
+
+	/** Callback for AnimatablePropertyChanged in HandleAssetAdded. */
+	bool HandleSequenceAdded(float KeyTime, UMovieSceneSequence* Sequence);
+
+private:
+
+	/** The Thumbnail pool which draws all the viewport thumbnails for the shot track. */
+	TSharedPtr<FTrackEditorThumbnailPool> ThumbnailPool;
+
+	/** The camera actor for the current cut. */
+	TWeakObjectPtr<ACameraActor> CinematicShotCamera;
+
+	/** Delegate binding handle for ISequencer::OnCameraCut */
+	FDelegateHandle OnCameraCutHandle;
+};
diff --git a/Engine/Source/Editor/MovieSceneTools/Private/TrackEditors/EventTrackEditor.cpp b/Engine/Source/Editor/MovieSceneTools/Private/TrackEditors/EventTrackEditor.cpp
--- a/Engine/Source/Editor/MovieSceneTools/Private/TrackEditors/EventTrackEditor.cpp
+++ b/Engine/Source/Editor/MovieSceneTools/Private/TrackEditors/EventTrackEditor.cpp
@@ -70,6 +70,10 @@ bool FEventTrackEditor::SupportsType(TSubclassOf<UMovieSceneTrack> Type) const
 	return (Type == UMovieSceneEventTrack::StaticClass());
 }
 
+const FSlateBrush* FEventTrackEditor::GetIconBrush() const
+{
+	return FEditorStyle::GetBrush("Sequencer.Tracks.Event");
+}
 
 /* FEventTrackEditor callbacks
  *****************************************************************************/
diff --git a/Engine/Source/Editor/MovieSceneTools/Private/TrackEditors/EventTrackEditor.h b/Engine/Source/Editor/MovieSceneTools/Private/TrackEditors/EventTrackEditor.h
--- a/Engine/Source/Editor/MovieSceneTools/Private/TrackEditors/EventTrackEditor.h
+++ b/Engine/Source/Editor/MovieSceneTools/Private/TrackEditors/EventTrackEditor.h
@@ -42,6 +42,7 @@ public:
 	virtual void BuildObjectBindingTrackMenu(FMenuBuilder& MenuBuilder, const FGuid& ObjectBinding, const UClass* ObjectClass) override;
 	virtual TSharedRef<ISequencerSection> MakeSectionInterface(UMovieSceneSection& SectionObject, UMovieSceneTrack& Track) override;
 	virtual bool SupportsType( TSubclassOf<UMovieSceneTrack> Type ) const override;
+	virtual const FSlateBrush* GetIconBrush() const override;
 
 private:
 
diff --git a/Engine/Source/Editor/MovieSceneTools/Private/TrackEditors/FadeTrackEditor.cpp b/Engine/Source/Editor/MovieSceneTools/Private/TrackEditors/FadeTrackEditor.cpp
--- a/Engine/Source/Editor/MovieSceneTools/Private/TrackEditors/FadeTrackEditor.cpp
+++ b/Engine/Source/Editor/MovieSceneTools/Private/TrackEditors/FadeTrackEditor.cpp
@@ -54,6 +54,11 @@ bool FFadeTrackEditor::SupportsType(TSubclassOf<UMovieSceneTrack> Type) const
 	return (Type == UMovieSceneFadeTrack::StaticClass());
 }
 
+const FSlateBrush* FFadeTrackEditor::GetIconBrush() const
+{
+	return FEditorStyle::GetBrush("Sequencer.Tracks.Fade");
+}
+
 
 /* FFadeTrackEditor callbacks
  *****************************************************************************/
diff --git a/Engine/Source/Editor/MovieSceneTools/Private/TrackEditors/MaterialTrackEditor.cpp b/Engine/Source/Editor/MovieSceneTools/Private/TrackEditors/MaterialTrackEditor.cpp
--- a/Engine/Source/Editor/MovieSceneTools/Private/TrackEditors/MaterialTrackEditor.cpp
+++ b/Engine/Source/Editor/MovieSceneTools/Private/TrackEditors/MaterialTrackEditor.cpp
@@ -5,6 +5,7 @@
 #include "MovieSceneMaterialTrack.h"
 #include "ParameterSection.h"
 #include "MovieSceneParameterSection.h"
+#include "SequencerUtilities.h"
 
 
 #define LOCTEXT_NAMESPACE "MaterialTrackEditor"
@@ -25,36 +26,12 @@ TSharedRef<ISequencerSection> FMaterialTrackEditor::MakeSectionInterface( UMovie
 }
 
 
-TSharedPtr<SWidget> FMaterialTrackEditor::BuildOutlinerEditWidget( const FGuid& ObjectBinding, UMovieSceneTrack* Track )
+TSharedPtr<SWidget> FMaterialTrackEditor::BuildOutlinerEditWidget( const FGuid& ObjectBinding, UMovieSceneTrack* Track, const FBuildEditWidgetParams& Params )
 {
 	UMovieSceneMaterialTrack* MaterialTrack = Cast<UMovieSceneMaterialTrack>(Track);
-	return
-		SNew( SComboButton )
-		.ButtonStyle( FEditorStyle::Get(), "FlatButton.Light" )
-		.OnGetMenuContent( this, &FMaterialTrackEditor::OnGetAddParameterMenuContent, ObjectBinding, MaterialTrack )
-		.ContentPadding( FMargin( 2, 0 ) )
-		.ButtonContent()
-		[
-			SNew( SHorizontalBox )
-
-			+ SHorizontalBox::Slot()
-			.VAlign( VAlign_Center )
-			.AutoWidth()
-			[
-				SNew( STextBlock )
-				.Font( FEditorStyle::Get().GetFontStyle( "FontAwesome.8" ) )
-				.Text( FText::FromString( FString( TEXT( "\xf067" ) ) ) /*fa-plus*/ )
-			]
-
-			+ SHorizontalBox::Slot()
-			.AutoWidth()
-			.Padding( 4, 0, 0, 0 )
-			[
-				SNew( STextBlock )
-				.Font( FEditorStyle::GetFontStyle( "Sequencer.AnimationOutliner.RegularFont" ) )
-				.Text( LOCTEXT( "AddTrackButton", "Parameter" ) )
-			]
-		];
+	FOnGetContent MenuContent = FOnGetContent::CreateSP(this, &FMaterialTrackEditor::OnGetAddParameterMenuContent, ObjectBinding, MaterialTrack);
+
+	return FSequencerUtilities::MakeAddButton(LOCTEXT( "AddParameterButton", "Parameter" ), MenuContent, Params.NodeIsHovered);
 }
 
 
diff --git a/Engine/Source/Editor/MovieSceneTools/Private/TrackEditors/MaterialTrackEditor.h b/Engine/Source/Editor/MovieSceneTools/Private/TrackEditors/MaterialTrackEditor.h
--- a/Engine/Source/Editor/MovieSceneTools/Private/TrackEditors/MaterialTrackEditor.h
+++ b/Engine/Source/Editor/MovieSceneTools/Private/TrackEditors/MaterialTrackEditor.h
@@ -24,7 +24,7 @@ public:
 
 	// ISequencerTrackEditor interface
 
-	virtual TSharedPtr<SWidget> BuildOutlinerEditWidget( const FGuid& ObjectBinding, UMovieSceneTrack* Track ) override;
+	virtual TSharedPtr<SWidget> BuildOutlinerEditWidget( const FGuid& ObjectBinding, UMovieSceneTrack* Track, const FBuildEditWidgetParams& Params ) override;
 	virtual TSharedRef<ISequencerSection> MakeSectionInterface( UMovieSceneSection& SectionObject, UMovieSceneTrack& Track ) override;
 
 protected:
diff --git a/Engine/Source/Editor/MovieSceneTools/Private/TrackEditors/ParticleParameterTrackEditor.cpp b/Engine/Source/Editor/MovieSceneTools/Private/TrackEditors/ParticleParameterTrackEditor.cpp
--- a/Engine/Source/Editor/MovieSceneTools/Private/TrackEditors/ParticleParameterTrackEditor.cpp
+++ b/Engine/Source/Editor/MovieSceneTools/Private/TrackEditors/ParticleParameterTrackEditor.cpp
@@ -5,6 +5,7 @@
 #include "MovieSceneParticleParameterTrack.h"
 #include "ParameterSection.h"
 #include "MovieSceneParameterSection.h"
+#include "SequencerUtilities.h"
 
 
 #define LOCTEXT_NAMESPACE "ParticleParameterTrackEditor"
@@ -32,36 +33,14 @@ TSharedRef<ISequencerSection> FParticleParameterTrackEditor::MakeSectionInterfac
 }
 
 
-TSharedPtr<SWidget> FParticleParameterTrackEditor::BuildOutlinerEditWidget( const FGuid& ObjectBinding, UMovieSceneTrack* Track )
+TSharedPtr<SWidget> FParticleParameterTrackEditor::BuildOutlinerEditWidget( const FGuid& ObjectBinding, UMovieSceneTrack* Track, const FBuildEditWidgetParams& Params )
 {
 	UMovieSceneParticleParameterTrack* ParticleParameterTrack = Cast<UMovieSceneParticleParameterTrack>(Track);
-	return
-		SNew( SComboButton )
-		.ButtonStyle( FEditorStyle::Get(), "FlatButton.Light" )
-		.OnGetMenuContent( this, &FParticleParameterTrackEditor::OnGetAddParameterMenuContent, ObjectBinding, ParticleParameterTrack )
-		.ContentPadding( FMargin( 2, 0 ) )
-		.ButtonContent()
-		[
-			SNew( SHorizontalBox )
-
-			+ SHorizontalBox::Slot()
-			.VAlign( VAlign_Center )
-			.AutoWidth()
-			[
-				SNew( STextBlock )
-				.Font( FEditorStyle::Get().GetFontStyle( "FontAwesome.8" ) )
-				.Text( FText::FromString( FString( TEXT( "\xf067" ) ) ) /*fa-plus*/ )
-			]
-
-			+ SHorizontalBox::Slot()
-			.AutoWidth()
-			.Padding( 4, 0, 0, 0 )
-			[
-				SNew( STextBlock )
-				.Font( FEditorStyle::GetFontStyle( "Sequencer.AnimationOutliner.RegularFont" ) )
-				.Text( LOCTEXT( "AddParameterButton", "Parameter" ) )
-			]
-		];
+	
+	// Create a container edit box
+	return FSequencerUtilities::MakeAddButton(LOCTEXT("ParameterText", "Parameter"),
+		FOnGetContent::CreateSP(this, &FParticleParameterTrackEditor::OnGetAddParameterMenuContent, ObjectBinding, ParticleParameterTrack),
+		Params.NodeIsHovered);
 }
 
 
diff --git a/Engine/Source/Editor/MovieSceneTools/Private/TrackEditors/ParticleParameterTrackEditor.h b/Engine/Source/Editor/MovieSceneTools/Private/TrackEditors/ParticleParameterTrackEditor.h
--- a/Engine/Source/Editor/MovieSceneTools/Private/TrackEditors/ParticleParameterTrackEditor.h
+++ b/Engine/Source/Editor/MovieSceneTools/Private/TrackEditors/ParticleParameterTrackEditor.h
@@ -28,7 +28,7 @@ public:
 
 	// ISequencerTrackEditor interface
 
-	virtual TSharedPtr<SWidget> BuildOutlinerEditWidget( const FGuid& ObjectBinding, UMovieSceneTrack* Track ) override;
+	virtual TSharedPtr<SWidget> BuildOutlinerEditWidget( const FGuid& ObjectBinding, UMovieSceneTrack* Track, const FBuildEditWidgetParams& Params ) override;
 	virtual TSharedRef<ISequencerSection> MakeSectionInterface( UMovieSceneSection& SectionObject, UMovieSceneTrack& Track ) override;
 	virtual bool SupportsType( TSubclassOf<UMovieSceneTrack> Type ) const override;
 
diff --git a/Engine/Source/Editor/MovieSceneTools/Private/TrackEditors/SlomoTrackEditor.cpp b/Engine/Source/Editor/MovieSceneTools/Private/TrackEditors/SlomoTrackEditor.cpp
--- a/Engine/Source/Editor/MovieSceneTools/Private/TrackEditors/SlomoTrackEditor.cpp
+++ b/Engine/Source/Editor/MovieSceneTools/Private/TrackEditors/SlomoTrackEditor.cpp
@@ -54,6 +54,11 @@ bool FSlomoTrackEditor::SupportsType(TSubclassOf<UMovieSceneTrack> Type) const
 	return (Type == UMovieSceneSlomoTrack::StaticClass());
 }
 
+const FSlateBrush* FSlomoTrackEditor::GetIconBrush() const
+{
+	return FEditorStyle::GetBrush("Sequencer.Tracks.Slomo");
+}
+
 
 /* FSlomoTrackEditor callbacks
  *****************************************************************************/
diff --git a/Engine/Source/Editor/MovieSceneTools/Private/TrackEditors/SpawnTrackEditor.cpp b/Engine/Source/Editor/MovieSceneTools/Private/TrackEditors/SpawnTrackEditor.cpp
--- a/Engine/Source/Editor/MovieSceneTools/Private/TrackEditors/SpawnTrackEditor.cpp
+++ b/Engine/Source/Editor/MovieSceneTools/Private/TrackEditors/SpawnTrackEditor.cpp
@@ -33,70 +33,6 @@ UMovieSceneTrack* FSpawnTrackEditor::AddTrack(UMovieScene* FocusedMovieScene, co
 	return NewTrack;
 }
 
-void FSpawnTrackEditor::BuildTrackContextMenu(FMenuBuilder& MenuBuilder, UMovieSceneTrack* Track)
-{
-	UMovieSceneSpawnTrack* SpawnTrack = Cast<UMovieSceneSpawnTrack>(Track);
-	if (!SpawnTrack)
-	{
-		return;
-	}
-
-	FMovieSceneSpawnable* Spawnable = GetSequencer()->GetFocusedMovieSceneSequence()->GetMovieScene()->FindSpawnable(SpawnTrack->GetObjectId());
-	if (!Spawnable)
-	{
-		return;
-	}
-
-	MenuBuilder.AddSubMenu(
-		LOCTEXT("OwnerLabel", "Spawned Object Owner"),
-		LOCTEXT("OwnerTooltip", "Specifies how the object spawned in this track is to be owned"),
-		FNewMenuDelegate::CreateSP(this, &FSpawnTrackEditor::AddSpawnOwnershipMenu, SpawnTrack)
-	);
-}
-
-void FSpawnTrackEditor::AddSpawnOwnershipMenu(FMenuBuilder& MenuBuilder, UMovieSceneSpawnTrack* Track)
-{
-	FMovieSceneSpawnable* Spawnable = GetSequencer()->GetFocusedMovieSceneSequence()->GetMovieScene()->FindSpawnable(Track->GetObjectId());
-
-	MenuBuilder.AddMenuEntry(
-		LOCTEXT("ThisSequence_Label", "This Sequence"),
-		LOCTEXT("ThisSequence_Tooltip", "Indicates that this sequence will own the spawned object. The object will be destroyed at the end of the sequence."),
-		FSlateIcon(),
-		FUIAction(
-			FExecuteAction::CreateLambda([=]{ Spawnable->SetSpawnOwnership(ESpawnOwnership::InnerSequence); }),
-			FCanExecuteAction(),
-			FIsActionChecked::CreateLambda([=]{ return Spawnable->GetSpawnOwnership() == ESpawnOwnership::InnerSequence; })
-		),
-		NAME_None,
-		EUserInterfaceActionType::ToggleButton
-	);
-
-	MenuBuilder.AddMenuEntry(
-		LOCTEXT("MasterSequence_Label", "Master Sequence"),
-		LOCTEXT("MasterSequence_Tooltip", "Indicates that the outermost sequence will own the spawned object. The object will be destroyed when the outermost sequence stops playing."),
-		FSlateIcon(),
-		FUIAction(
-			FExecuteAction::CreateLambda([=]{ Spawnable->SetSpawnOwnership(ESpawnOwnership::MasterSequence); }),
-			FCanExecuteAction(),
-			FIsActionChecked::CreateLambda([=]{ return Spawnable->GetSpawnOwnership() == ESpawnOwnership::MasterSequence; })
-		),
-		NAME_None,
-		EUserInterfaceActionType::ToggleButton
-	);
-
-	MenuBuilder.AddMenuEntry(
-		LOCTEXT("External_Label", "External"),
-		LOCTEXT("External_Tooltip", "Indicates this object's lifetime is managed externally once spawned. It will not be destroyed by sequencer."),
-		FSlateIcon(),
-		FUIAction(
-			FExecuteAction::CreateLambda([=]{ Spawnable->SetSpawnOwnership(ESpawnOwnership::External); }),
-			FCanExecuteAction(),
-			FIsActionChecked::CreateLambda([=]{ return Spawnable->GetSpawnOwnership() == ESpawnOwnership::External; })
-		),
-		NAME_None,
-		EUserInterfaceActionType::ToggleButton
-	);
-}
 
 void FSpawnTrackEditor::BuildObjectBindingTrackMenu(FMenuBuilder& MenuBuilder, const FGuid& ObjectBinding, const UClass* ObjectClass)
 {
diff --git a/Engine/Source/Editor/MovieSceneTools/Private/TrackEditors/SpawnTrackEditor.h b/Engine/Source/Editor/MovieSceneTools/Private/TrackEditors/SpawnTrackEditor.h
--- a/Engine/Source/Editor/MovieSceneTools/Private/TrackEditors/SpawnTrackEditor.h
+++ b/Engine/Source/Editor/MovieSceneTools/Private/TrackEditors/SpawnTrackEditor.h
@@ -40,17 +40,13 @@ public:
 
 	virtual UMovieSceneTrack* AddTrack(UMovieScene* FocusedMovieScene, const FGuid& ObjectHandle, TSubclassOf<class UMovieSceneTrack> TrackClass, FName UniqueTypeName) override;
 	virtual void BuildObjectBindingTrackMenu(FMenuBuilder& MenuBuilder, const FGuid& ObjectBinding, const UClass* ObjectClass) override;
-	virtual TSharedPtr<SWidget> BuildOutlinerEditWidget(const FGuid& ObjectBinding, UMovieSceneTrack* Track) override  { return TSharedPtr<SWidget>(); }
+	virtual TSharedPtr<SWidget> BuildOutlinerEditWidget(const FGuid& ObjectBinding, UMovieSceneTrack* Track, const FBuildEditWidgetParams& Params) override  { return TSharedPtr<SWidget>(); }
 	virtual bool HandleAssetAdded(UObject* Asset, const FGuid& TargetObjectGuid) override { return false; }
 	virtual bool SupportsType(TSubclassOf<UMovieSceneTrack> Type) const override;
-	virtual void BuildTrackContextMenu( FMenuBuilder& MenuBuilder, UMovieSceneTrack* Track ) override;
 
 private:
 
 	/** Callback for executing the "Add Spawn Track" menu entry. */
 	void HandleAddSpawnTrackMenuEntryExecute(FGuid ObjectBinding);
 	bool CanAddSpawnTrack(FGuid ObjectBinding) const;
-
-	/** Populate the sub menu for setting the ownership of the spawned object */
-	void AddSpawnOwnershipMenu(FMenuBuilder& MenuBuilder, UMovieSceneSpawnTrack* Track);
 };
diff --git a/Engine/Source/Editor/MovieSceneTools/Private/TrackEditors/SubTrackEditor.cpp b/Engine/Source/Editor/MovieSceneTools/Private/TrackEditors/SubTrackEditor.cpp
--- a/Engine/Source/Editor/MovieSceneTools/Private/TrackEditors/SubTrackEditor.cpp
+++ b/Engine/Source/Editor/MovieSceneTools/Private/TrackEditors/SubTrackEditor.cpp
@@ -4,7 +4,10 @@
 #include "MovieSceneSubSection.h"
 #include "MovieSceneSequence.h"
 #include "MovieSceneSubTrack.h"
+#include "MovieSceneCinematicShotTrack.h"
 #include "SubTrackEditor.h"
+#include "ContentBrowserModule.h"
+#include "SequencerUtilities.h"
 
 
 namespace SubTrackEditorConstants
@@ -244,6 +247,20 @@ void FSubTrackEditor::BuildAddTrackMenu(FMenuBuilder& MenuBuilder)
 	);
 }
 
+TSharedPtr<SWidget> FSubTrackEditor::BuildOutlinerEditWidget(const FGuid& ObjectBinding, UMovieSceneTrack* Track, const FBuildEditWidgetParams& Params)
+{
+	// Create a container edit box
+	return SNew(SHorizontalBox)
+
+	// Add the sub sequence combo box
+	+ SHorizontalBox::Slot()
+	.AutoWidth()
+	.VAlign(VAlign_Center)
+	[
+		FSequencerUtilities::MakeAddButton(LOCTEXT("SubText", "Sequence"), FOnGetContent::CreateSP(this, &FSubTrackEditor::HandleAddSubSequenceComboButtonGetMenuContent), Params.NodeIsHovered)
+	];
+}
+
 
 TSharedRef<ISequencerTrackEditor> FSubTrackEditor::CreateTrackEditor(TSharedRef<ISequencer> InSequencer)
 {
@@ -266,12 +283,22 @@ bool FSubTrackEditor::HandleAssetAdded(UObject* Asset, const FGuid& TargetObject
 		return false;
 	}
 
+	//@todo If there's already a cinematic shot track, allow that track to handle this asset
+	UMovieScene* FocusedMovieScene = GetFocusedMovieScene();
+
+	if (FocusedMovieScene != nullptr && FocusedMovieScene->FindMasterTrack<UMovieSceneCinematicShotTrack>() != nullptr)
+	{
+		return false;
+	}
+
 	if (CanAddSubSequence(*Sequence))
 	{
 		AnimatablePropertyChanged(FOnKeyProperty::CreateRaw(this, &FSubTrackEditor::HandleSequenceAdded, Sequence));
+
+		return true;
 	}
 
-	return true;
+	return false;
 }
 
 
@@ -281,6 +308,11 @@ bool FSubTrackEditor::SupportsType(TSubclassOf<UMovieSceneTrack> Type) const
 	return Type == UMovieSceneSubTrack::StaticClass();
 }
 
+const FSlateBrush* FSubTrackEditor::GetIconBrush() const
+{
+	return FEditorStyle::GetBrush("Sequencer.Tracks.Sub");
+}
+
 
 /* FSubTrackEditor callbacks
  *****************************************************************************/
@@ -338,11 +370,65 @@ void FSubTrackEditor::HandleAddSubTrackMenuEntryExecute()
 }
 
 
+TSharedRef<SWidget> FSubTrackEditor::HandleAddSubSequenceComboButtonGetMenuContent()
+{
+	FMenuBuilder MenuBuilder(true, nullptr);
+
+	FAssetPickerConfig AssetPickerConfig;
+	{
+		AssetPickerConfig.OnAssetSelected = FOnAssetSelected::CreateRaw( this, &FSubTrackEditor::HandleAddSubSequenceComboButtonMenuEntryExecute);
+		AssetPickerConfig.bAllowNullSelection = false;
+		AssetPickerConfig.InitialAssetViewType = EAssetViewType::Tile;
+		AssetPickerConfig.Filter.ClassNames.Add(TEXT("LevelSequence"));
+	}
+
+	FContentBrowserModule& ContentBrowserModule = FModuleManager::Get().LoadModuleChecked<FContentBrowserModule>(TEXT("ContentBrowser"));
+
+	TSharedPtr<SBox> MenuEntry = SNew(SBox)
+		.WidthOverride(300.0f)
+		.HeightOverride(300.f)
+		[
+			ContentBrowserModule.Get().CreateAssetPicker(AssetPickerConfig)
+		];
+
+	MenuBuilder.AddWidget(MenuEntry.ToSharedRef(), FText::GetEmpty(), true);
+
+	return MenuBuilder.MakeWidget();
+}
+
+void FSubTrackEditor::HandleAddSubSequenceComboButtonMenuEntryExecute(const FAssetData& AssetData)
+{
+	FSlateApplication::Get().DismissAllMenus();
+
+	UObject* SelectedObject = AssetData.GetAsset();
+
+	if (SelectedObject && SelectedObject->IsA(UMovieSceneSequence::StaticClass()))
+	{
+		UMovieSceneSequence* MovieSceneSequence = CastChecked<UMovieSceneSequence>(AssetData.GetAsset());
+
+		AnimatablePropertyChanged( FOnKeyProperty::CreateRaw( this, &FSubTrackEditor::AddKeyInternal, MovieSceneSequence) );
+	}
+}
+
+bool FSubTrackEditor::AddKeyInternal(float KeyTime, UMovieSceneSequence* InMovieSceneSequence)
+{	
+	if (CanAddSubSequence(*InMovieSceneSequence))
+	{
+		auto SubTrack = FindOrCreateMasterTrack<UMovieSceneSubTrack>().Track;
+		float Duration = InMovieSceneSequence->GetMovieScene()->GetPlaybackRange().Size<float>();
+		SubTrack->AddSequence(InMovieSceneSequence, KeyTime, Duration);
+
+		return true;
+	}
+
+	return false;
+}
+
 bool FSubTrackEditor::HandleSequenceAdded(float KeyTime, UMovieSceneSequence* Sequence)
 {
 	auto SubTrack = FindOrCreateMasterTrack<UMovieSceneSubTrack>().Track;
-	SubTrack->AddSequence(*Sequence, KeyTime);
-	SubTrack->SetDisplayName(LOCTEXT("SubTrackName", "Sequences"));
+	float Duration = Sequence->GetMovieScene()->GetPlaybackRange().Size<float>();
+	SubTrack->AddSequence(Sequence, KeyTime, Duration);
 
 	return true;
 }
diff --git a/Engine/Source/Editor/MovieSceneTools/Private/TrackEditors/SubTrackEditor.h b/Engine/Source/Editor/MovieSceneTools/Private/TrackEditors/SubTrackEditor.h
--- a/Engine/Source/Editor/MovieSceneTools/Private/TrackEditors/SubTrackEditor.h
+++ b/Engine/Source/Editor/MovieSceneTools/Private/TrackEditors/SubTrackEditor.h
@@ -37,10 +37,12 @@ public:
 	// ISequencerTrackEditor interface
 
 	virtual void BuildAddTrackMenu(FMenuBuilder& MenuBuilder) override;
+	virtual TSharedPtr<SWidget> BuildOutlinerEditWidget(const FGuid& ObjectBinding, UMovieSceneTrack* Track, const FBuildEditWidgetParams& Params) override;
 	virtual TSharedRef<ISequencerSection> MakeSectionInterface(UMovieSceneSection& SectionObject, UMovieSceneTrack& Track) override;
 	virtual bool HandleAssetAdded(UObject* Asset, const FGuid& TargetObjectGuid) override;
 	virtual bool SupportsType(TSubclassOf<UMovieSceneTrack> Type) const override;
-
+	virtual const FSlateBrush* GetIconBrush() const override;
+	
 protected:
 
 	/**
@@ -59,6 +61,15 @@ private:
 	/** Callback for executing the "Add Event Track" menu entry. */
 	void HandleAddSubTrackMenuEntryExecute();
 
+	/** Callback for generating the menu of the "Add Sequence" combo button. */
+	TSharedRef<SWidget> HandleAddSubSequenceComboButtonGetMenuContent();
+
+	/** Callback for executing a menu entry in the "Add Sequence" combo button. */
+	void HandleAddSubSequenceComboButtonMenuEntryExecute(const FAssetData& AssetData);
+
+	/** Delegate for AnimatablePropertyChanged in AddKey */
+	bool AddKeyInternal(float KeyTime, UMovieSceneSequence* InMovieSceneSequence);
+
 	/** Callback for AnimatablePropertyChanged in HandleAssetAdded. */
 	bool HandleSequenceAdded(float KeyTime, UMovieSceneSequence* Sequence);
 };
diff --git a/Engine/Source/Editor/MovieSceneTools/Private/TrackEditors/TransformTrackEditor.cpp b/Engine/Source/Editor/MovieSceneTools/Private/TrackEditors/TransformTrackEditor.cpp
--- a/Engine/Source/Editor/MovieSceneTools/Private/TrackEditors/TransformTrackEditor.cpp
+++ b/Engine/Source/Editor/MovieSceneTools/Private/TrackEditors/TransformTrackEditor.cpp
@@ -423,6 +423,11 @@ void F3DTransformTrackEditor::BindCommands(TSharedRef<FUICommandList> SequencerC
 
 void F3DTransformTrackEditor::BuildObjectBindingEditButtons(TSharedPtr<SHorizontalBox> EditBox, const FGuid& ObjectGuid, const UClass* ObjectClass)
 {
+	if (!ObjectClass->IsChildOf<ACameraActor>() && !ObjectClass->IsChildOf<UCameraComponent>())
+	{
+		return;
+	}
+
 	// If this is a camera track, add a button to lock the viewport to the camera
 	EditBox.Get()->AddSlot()
 		.VAlign(VAlign_Center)
diff --git a/Engine/Source/Editor/MovieSceneTools/Public/FloatCurveKeyArea.h b/Engine/Source/Editor/MovieSceneTools/Public/FloatCurveKeyArea.h
--- a/Engine/Source/Editor/MovieSceneTools/Public/FloatCurveKeyArea.h
+++ b/Engine/Source/Editor/MovieSceneTools/Public/FloatCurveKeyArea.h
@@ -44,6 +44,7 @@ public:
 	virtual FLinearColor GetColor() override;
 	virtual ERichCurveExtrapolation GetExtrapolationMode(bool bPreInfinity) const override;
 	virtual ERichCurveInterpMode GetKeyInterpMode(FKeyHandle KeyHandle) const override;
+	virtual TSharedPtr<FStructOnScope> GetKeyStruct(FKeyHandle KeyHandle) override;
 	virtual ERichCurveTangentMode GetKeyTangentMode(FKeyHandle KeyHandle) const override;
 	virtual float GetKeyTime(FKeyHandle KeyHandle) const override;
 	virtual UMovieSceneSection* GetOwningSection() override;
@@ -51,6 +52,7 @@ public:
 	virtual TArray<FKeyHandle> GetUnsortedKeyHandles() const override;
 	virtual FKeyHandle MoveKey(FKeyHandle KeyHandle, float DeltaPosition) override;
 	virtual void SetExtrapolationMode(ERichCurveExtrapolation ExtrapMode, bool bPreInfinity) override;
+	virtual bool CanSetExtrapolationMode() const override;
 	virtual void SetKeyInterpMode(FKeyHandle KeyHandle, ERichCurveInterpMode InterpMode) override;
 	virtual void SetKeyTangentMode(FKeyHandle KeyHandle, ERichCurveTangentMode TangentMode) override;
 	virtual void SetKeyTime(FKeyHandle KeyHandle, float NewKeyTime) const override;
diff --git a/Engine/Source/Editor/MovieSceneTools/Public/IntegralKeyArea.h b/Engine/Source/Editor/MovieSceneTools/Public/IntegralKeyArea.h
--- a/Engine/Source/Editor/MovieSceneTools/Public/IntegralKeyArea.h
+++ b/Engine/Source/Editor/MovieSceneTools/Public/IntegralKeyArea.h
@@ -30,6 +30,7 @@ public:
 	virtual FLinearColor GetColor() override;
 	virtual ERichCurveExtrapolation GetExtrapolationMode(bool bPreInfinity) const override;
 	virtual ERichCurveInterpMode GetKeyInterpMode(FKeyHandle KeyHandle) const override;
+	virtual TSharedPtr<FStructOnScope> GetKeyStruct(FKeyHandle KeyHandle) override;
 	virtual ERichCurveTangentMode GetKeyTangentMode(FKeyHandle KeyHandle) const override;
 	virtual float GetKeyTime(FKeyHandle KeyHandle) const override;
 	virtual UMovieSceneSection* GetOwningSection() override;
diff --git a/Engine/Source/Editor/MovieSceneTools/Public/MovieSceneToolHelpers.h b/Engine/Source/Editor/MovieSceneTools/Public/MovieSceneToolHelpers.h
--- a/Engine/Source/Editor/MovieSceneTools/Public/MovieSceneToolHelpers.h
+++ b/Engine/Source/Editor/MovieSceneTools/Public/MovieSceneToolHelpers.h
@@ -8,20 +8,69 @@ class UMovieSceneSection;
 class MOVIESCENETOOLS_API MovieSceneToolHelpers
 {
 public:
+
 	/**
-	* Trim section at the given time
-	*
-	* @param Sections The sections to trim
-	* @param Time	The time at which to trim
-	* @param bTrimLeft Trim left or trim right
-	*/
+	 * Trim section at the given time
+	 *
+	 * @param Sections The sections to trim
+	 * @param Time	The time at which to trim
+	 * @param bTrimLeft Trim left or trim right
+	 */
 	static void TrimSection(const TSet<TWeakObjectPtr<UMovieSceneSection>>& Sections, float Time, bool bTrimLeft);
 
 	/**
-	* Splits sections at the given time
-	*
-	* @param Sections The sections to split
-	* @param Time	The time at which to split
-	*/
+	 * Splits sections at the given time
+	 *
+	 * @param Sections The sections to split
+	 * @param Time	The time at which to split
+	 */
 	static void SplitSection(const TSet<TWeakObjectPtr<UMovieSceneSection>>& Sections, float Time);
+
+	/**
+	 * Parse a shot name into its components.
+	 *
+	 * @param ShotName The shot name to parse
+	 * @param ShotPrefix The parsed shot prefix
+	 * @param ShotNumber The parsed shot number
+	 * @param TakeNumber The parsed take number
+	 * &return Whether the shot name was parsed successfully
+	 */
+	static bool ParseShotName(const FString& ShotName, FString& ShotPrefix, uint32& ShotNumber, uint32& TakeNumber);
+
+	/**
+	 * Compose a shot name given its components.
+	 *
+	 * @param ShotPrefix The shot prefix to use
+	 * @param ShotNumber The shot number to use
+	 * @param TakeNumber The take number to use
+	 * @return The composed shot name
+	 */
+	static FString ComposeShotName(const FString& ShotPrefix, uint32 ShotNumber, uint32 TakeNumber);
+
+	/**
+	 * Generate a new shot name
+	 *
+	 * @param AllSections All the sections in the given shot track
+	 * @param Time The time to generate the new shot name at
+	 * @return The new shot name
+	 */
+	static FString GenerateNewShotName(const TArray<UMovieSceneSection*>& AllSections, float Time);
+
+	/**
+	 * Gather takes (level sequence assets that have the same shot prefix and shot number)
+	 * 
+	 * @param Section The section to gather takes from
+	 * @param TakeNumbers The gathered take numbers
+	 */
+	static void GatherTakes(const UMovieSceneSection* Section, TArray<uint32>& TakeNumbers);
+
+
+	/**
+	 * Get the asset associated with the take number
+	 *
+	 * @param Section The section to gather the take from
+	 * @param TakeNumber The take number to get
+	 * @return The asset
+	 */
+	static UObject* GetTake(const UMovieSceneSection* Section, uint32 TakeNumber);
 };
diff --git a/Engine/Source/Editor/MovieSceneTools/Public/NameCurveKeyArea.h b/Engine/Source/Editor/MovieSceneTools/Public/NameCurveKeyArea.h
--- a/Engine/Source/Editor/MovieSceneTools/Public/NameCurveKeyArea.h
+++ b/Engine/Source/Editor/MovieSceneTools/Public/NameCurveKeyArea.h
@@ -36,6 +36,7 @@ public:
 	virtual FLinearColor GetColor() override;
 	virtual ERichCurveExtrapolation GetExtrapolationMode(bool bPreInfinity) const override;
 	virtual ERichCurveInterpMode GetKeyInterpMode(FKeyHandle Keyhandle) const override;
+	virtual TSharedPtr<FStructOnScope> GetKeyStruct(FKeyHandle KeyHandle) override;
 	virtual ERichCurveTangentMode GetKeyTangentMode(FKeyHandle KeyHandle) const override;
 	virtual float GetKeyTime(FKeyHandle KeyHandle) const override;
 	virtual UMovieSceneSection* GetOwningSection() override;
diff --git a/Engine/Source/Editor/MovieSceneTools/Public/PropertyTrackEditor.h b/Engine/Source/Editor/MovieSceneTools/Public/PropertyTrackEditor.h
--- a/Engine/Source/Editor/MovieSceneTools/Public/PropertyTrackEditor.h
+++ b/Engine/Source/Editor/MovieSceneTools/Public/PropertyTrackEditor.h
@@ -144,6 +144,9 @@ protected:
 	virtual void InitializeNewTrack( TrackType* NewTrack, FPropertyChangedParams PropertyChangedParams )
 	{
 		NewTrack->SetPropertyNameAndPath( PropertyChangedParams.PropertyPath.Last()->GetFName(), PropertyChangedParams.GetPropertyPathString() );
+#if WITH_EDITORONLY_DATA
+		NewTrack->SetDisplayName(PropertyChangedParams.PropertyPath.Last()->GetDisplayNameText());
+#endif
 	}
 
 private:
diff --git a/Engine/Source/Editor/NiagaraEditor/Private/NiagaraSequencer.h b/Engine/Source/Editor/NiagaraEditor/Private/NiagaraSequencer.h
--- a/Engine/Source/Editor/NiagaraEditor/Private/NiagaraSequencer.h
+++ b/Engine/Source/Editor/NiagaraEditor/Private/NiagaraSequencer.h
@@ -120,7 +120,7 @@ public:
 	{
 	}
 
-	virtual void Update(float Position, float LastPosition, const TArray<UObject*>& RuntimeObjects, IMovieScenePlayer& Player, FMovieSceneSequenceInstance& SequenceInstance, EMovieSceneUpdatePass UpdatePass)
+	virtual void Update(EMovieSceneUpdateData& UpdateData, const TArray<UObject*>& RuntimeObjects, IMovieScenePlayer& Player, FMovieSceneSequenceInstance& SequenceInstance)
 	{
 	}
 
diff --git a/Engine/Source/Editor/Persona/Private/Persona.cpp b/Engine/Source/Editor/Persona/Private/Persona.cpp
--- a/Engine/Source/Editor/Persona/Private/Persona.cpp
+++ b/Engine/Source/Editor/Persona/Private/Persona.cpp
@@ -2617,7 +2617,7 @@ void FPersona::RedoAction()
 
 FSlateIcon FPersona::GetRecordStatusImage() const
 {
-	if (Recorder.InRecording())
+	if (IsRecording())
 	{
 		return FSlateIcon(FEditorStyle::GetStyleSetName(), "Persona.StopRecordAnimation");
 	}
@@ -2627,7 +2627,7 @@ FSlateIcon FPersona::GetRecordStatusImage() const
 
 FText FPersona::GetRecordMenuLabel() const
 {
-	if(Recorder.InRecording())
+	if(IsRecording())
 	{
 		return LOCTEXT("Persona_StopRecordAnimationMenuLabel", "Stop Record Animation");
 	}
@@ -2637,7 +2637,12 @@ FText FPersona::GetRecordMenuLabel() const
 
 FText FPersona::GetRecordStatusLabel() const
 {
-	if (Recorder.InRecording())
+	FPersonaModule& PersonaModule = FModuleManager::GetModuleChecked<FPersonaModule>(TEXT("Persona"));
+
+	bool bInRecording = false;
+	PersonaModule.OnIsRecordingActive().ExecuteIfBound(PreviewComponent, bInRecording);
+
+	if (bInRecording)
 	{
 		return LOCTEXT("Persona_StopRecordAnimationLabel", "Stop");
 	}
@@ -2647,7 +2652,12 @@ FText FPersona::GetRecordStatusLabel() const
 
 FText FPersona::GetRecordStatusTooltip() const
 {
-	if (Recorder.InRecording())
+	FPersonaModule& PersonaModule = FModuleManager::GetModuleChecked<FPersonaModule>(TEXT("Persona"));
+
+	bool bInRecording = false;
+	PersonaModule.OnIsRecordingActive().ExecuteIfBound(PreviewComponent, bInRecording);
+
+	if (bInRecording)
 	{
 		return LOCTEXT("Persona_StopRecordAnimation", "Stop Record Animation");
 	}
@@ -2663,13 +2673,18 @@ void FPersona::RecordAnimation()
 		return;
 	}
 
-	if (Recorder.InRecording())
+	FPersonaModule& PersonaModule = FModuleManager::GetModuleChecked<FPersonaModule>(TEXT("Persona"));
+
+	bool bInRecording = false;
+	PersonaModule.OnIsRecordingActive().ExecuteIfBound(PreviewComponent, bInRecording);
+
+	if (bInRecording)
 	{
-		Recorder.StopRecord(true);
+		PersonaModule.OnStopRecording().ExecuteIfBound(PreviewComponent);
 	}
 	else
 	{
-		Recorder.TriggerRecordAnimation(PreviewComponent);
+		PersonaModule.OnRecord().ExecuteIfBound(PreviewComponent);
 	}
 }
 
@@ -3043,10 +3058,11 @@ void FPersona::Tick(float DeltaTime)
 		Viewport.Pin()->RefreshViewport();
 	}
 
-	if (Recorder.InRecording())
+	if (IsRecording())
 	{
 		// make sure you don't allow switch previewcomponent
-		Recorder.UpdateRecord(PreviewComponent, DeltaTime);
+		FPersonaModule& PersonaModule = FModuleManager::GetModuleChecked<FPersonaModule>("Persona");
+		PersonaModule.OnTickRecording().ExecuteIfBound(PreviewComponent, DeltaTime);
 	}
 
 	if (PreviewComponent && LastCachedLODForPreviewComponent != PreviewComponent->PredictedLODLevel)
@@ -3327,6 +3343,38 @@ void FPersona::SetSelectedBlendProfile(UBlendProfile* InBlendProfile)
 	OnBlendProfileSelected.Broadcast(InBlendProfile);
 }
 
+bool FPersona::IsRecording() const
+{
+	FPersonaModule& PersonaModule = FModuleManager::GetModuleChecked<FPersonaModule>("Persona");
+
+	bool bInRecording = false;
+	PersonaModule.OnIsRecordingActive().ExecuteIfBound(PreviewComponent, bInRecording);
+
+	return bInRecording;
+}
+
+void FPersona::StopRecording()
+{
+	FPersonaModule& PersonaModule = FModuleManager::GetModuleChecked<FPersonaModule>("Persona");
+	PersonaModule.OnStopRecording().ExecuteIfBound(PreviewComponent);
+}
+
+UAnimSequence* FPersona::GetCurrentRecording() const
+{
+	FPersonaModule& PersonaModule = FModuleManager::GetModuleChecked<FPersonaModule>("Persona");
+	UAnimSequence* Recording = nullptr;
+	PersonaModule.OnGetCurrentRecording().ExecuteIfBound(PreviewComponent, Recording);
+	return Recording;
+}
+
+float FPersona::GetCurrentRecordingTime() const
+{
+	FPersonaModule& PersonaModule = FModuleManager::GetModuleChecked<FPersonaModule>("Persona");
+	float RecordingTime = 0.0f;
+	PersonaModule.OnGetCurrentRecordingTime().ExecuteIfBound(PreviewComponent, RecordingTime);
+	return RecordingTime;
+}
+
 static class FMeshHierarchyCmd : private FSelfRegisteringExec
 {
 public:
diff --git a/Engine/Source/Editor/Persona/Private/Persona.h b/Engine/Source/Editor/Persona/Private/Persona.h
--- a/Engine/Source/Editor/Persona/Private/Persona.h
+++ b/Engine/Source/Editor/Persona/Private/Persona.h
@@ -9,7 +9,6 @@
 #include "PersonaModule.h"
 #include "PersonaCommands.h"
 
-#include "AnimationRecorder.h"
 //////////////////////////////////////////////////////////////////////////
 // FPersona
 
@@ -177,6 +176,18 @@ public:
 	/** Called when the blend profile tab selects a profile */
 	void SetSelectedBlendProfile(UBlendProfile* InBlendProfile);
 
+	/** Check whether this Persona instance is recording */
+	bool IsRecording() const;
+
+	/** Stop recording in this Persona instance */
+	void StopRecording();
+
+	/** Get the currently recording animation */
+	UAnimSequence* GetCurrentRecording() const;
+
+	/** Get the currently recording animation time */
+	float GetCurrentRecordingTime() const;
+
 public:
 	//~ Begin IToolkit Interface
 	virtual FName GetToolkitContextFName() const override;
@@ -300,9 +311,6 @@ public:
 	mutable TWeakObjectPtr<UObject> SharedAnimAssetBeingEdited;
 	TWeakPtr<SDockTab> SharedAnimDocumentTab;
 
-	/** Animation recorder **/
-	FAnimationRecorder Recorder;
-
 public:
 	/** Viewport widget */
 	TWeakPtr<class SAnimationEditorViewportTabBody> Viewport;
diff --git a/Engine/Source/Editor/Persona/Private/SAnimationEditorViewport.cpp b/Engine/Source/Editor/Persona/Private/SAnimationEditorViewport.cpp
--- a/Engine/Source/Editor/Persona/Private/SAnimationEditorViewport.cpp
+++ b/Engine/Source/Editor/Persona/Private/SAnimationEditorViewport.cpp
@@ -1897,7 +1897,7 @@ bool SAnimationEditorViewportTabBody::IsSectionsDisplayMode(int32 DisplayMode) c
 EVisibility SAnimationEditorViewportTabBody::GetViewportCornerImageVisibility() const
 {
 	TSharedPtr<FPersona> Persona = PersonaPtr.Pin();	
-	return Persona->Recorder.InRecording()? EVisibility::Visible : EVisibility::Collapsed;
+	return Persona->IsRecording() ? EVisibility::Visible : EVisibility::Collapsed;
 }
 
 const FSlateBrush* SAnimationEditorViewportTabBody::GetViewportCornerImage() const
@@ -1928,7 +1928,7 @@ const FSlateBrush* SAnimationEditorViewportTabBody::GetViewportCornerImage() con
 EVisibility SAnimationEditorViewportTabBody::GetViewportCornerTextVisibility() const
 {
 	TSharedPtr<FPersona> Persona = PersonaPtr.Pin();
-	if (Persona->Recorder.InRecording())
+	if (Persona->IsRecording())
 	{
 		return EVisibility::Visible;
 	}
@@ -1947,10 +1947,11 @@ EVisibility SAnimationEditorViewportTabBody::GetViewportCornerTextVisibility() c
 FText SAnimationEditorViewportTabBody::GetViewportCornerText() const
 {
 	TSharedPtr<FPersona> Persona = PersonaPtr.Pin();
-	if(Persona->Recorder.InRecording())
+	if(Persona->IsRecording())
 	{
-		const FString& Name = Persona->Recorder.GetAnimationObject()->GetName();
-		float TimeRecorded = Persona->Recorder.GetTimeRecorded();
+		UAnimSequence* Recording = Persona->GetCurrentRecording();
+		const FString& Name = Recording ? Recording->GetName() : TEXT("None");
+		float TimeRecorded = Persona->GetCurrentRecordingTime();
 		FNumberFormattingOptions NumberOption;
 		NumberOption.MaximumFractionalDigits = 2;
 		NumberOption.MinimumFractionalDigits = 2;
@@ -1983,7 +1984,7 @@ FText SAnimationEditorViewportTabBody::GetViewportCornerText() const
 FText SAnimationEditorViewportTabBody::GetViewportCornerTooltip() const
 {
 	TSharedPtr<FPersona> Persona = PersonaPtr.Pin();
-	if(Persona->Recorder.InRecording())
+	if(Persona->IsRecording())
 	{
 		return LOCTEXT("RecordingStatusTooltip", "Shows the status of animation recording.\nClick to stop the recording.");
 	}
@@ -2000,9 +2001,9 @@ FReply SAnimationEditorViewportTabBody::ClickedOnViewportCornerText()
 	TSharedPtr<FPersona> Persona = PersonaPtr.Pin();
 	// if it's recording, it won't be able to see the message
 	// so disable it
-	if (Persona->Recorder.InRecording())
+	if (Persona->IsRecording())
 	{
-		Persona->Recorder.StopRecord(true);
+		Persona->StopRecording();
 	}
 	else 
 	{
diff --git a/Engine/Source/Editor/Persona/Private/SAnimationSequenceBrowser.cpp b/Engine/Source/Editor/Persona/Private/SAnimationSequenceBrowser.cpp
--- a/Engine/Source/Editor/Persona/Private/SAnimationSequenceBrowser.cpp
+++ b/Engine/Source/Editor/Persona/Private/SAnimationSequenceBrowser.cpp
@@ -1016,7 +1016,7 @@ bool SAnimationSequenceBrowser::IsToolTipPreviewVisible()
 {
 	bool bVisible = false;
 	// during persona recording, disable this
-	if( PersonaPtr.IsValid() && PersonaPtr.Pin()->Recorder.InRecording() == false 
+	if( PersonaPtr.IsValid() && PersonaPtr.Pin()->IsRecording() == false 
 		&& ViewportWidget.IsValid())
 	{
 		bVisible = ViewportWidget->GetVisibility() == EVisibility::Visible;
diff --git a/Engine/Source/Editor/Persona/Public/PersonaModule.h b/Engine/Source/Editor/Persona/Public/PersonaModule.h
--- a/Engine/Source/Editor/Persona/Public/PersonaModule.h
+++ b/Engine/Source/Editor/Persona/Public/PersonaModule.h
@@ -8,9 +8,17 @@
 
 #include "Toolkits/IToolkit.h"		// For EToolkitMode::Type
 #include "Toolkits/AssetEditorToolkit.h" 
+#include "ModuleInterface.h"
 
 extern const FName PersonaAppName;
 
+DECLARE_DELEGATE_TwoParams(FIsRecordingActive, USkeletalMeshComponent* /*Component*/, bool& /* bIsRecording */);
+DECLARE_DELEGATE_OneParam(FRecord, USkeletalMeshComponent* /*Component*/);
+DECLARE_DELEGATE_OneParam(FStopRecording, USkeletalMeshComponent* /*Component*/);
+DECLARE_DELEGATE_TwoParams(FGetCurrentRecording, USkeletalMeshComponent* /*Component*/, UAnimSequence*& /* OutRecording */);
+DECLARE_DELEGATE_TwoParams(FGetCurrentRecordingTime, USkeletalMeshComponent* /*Component*/, float& /* OutTime */);
+DECLARE_DELEGATE_TwoParams(FTickRecording, USkeletalMeshComponent* /*Component*/, float /* DeltaSeconds */);
+
 /**
  * Persona module manages the lifetime of all instances of Persona editors.
  */
@@ -49,6 +57,24 @@ public:
 	virtual TSharedPtr<FExtensibilityManager> GetMenuExtensibilityManager() {return MenuExtensibilityManager;}
 	virtual TSharedPtr<FExtensibilityManager> GetToolBarExtensibilityManager() {return ToolBarExtensibilityManager;}
 
+	/** Delegate used to query whether recording is active */
+	FIsRecordingActive& OnIsRecordingActive() { return IsRecordingActiveDelegate; }
+
+	/** Delegate used to start recording animation */
+	FRecord& OnRecord() { return RecordDelegate; }
+
+	/** Delegate used to stop recording animation */
+	FStopRecording& OnStopRecording() { return StopRecordingDelegate; }
+
+	/** Delegate used to get the currently recording animation */
+	FGetCurrentRecording& OnGetCurrentRecording() { return GetCurrentRecordingDelegate; }
+
+	/** Delegate used to get the currently recording animation time */
+	FGetCurrentRecordingTime& OnGetCurrentRecordingTime() { return GetCurrentRecordingTimeDelegate; }
+
+	/** Delegate used to tick the skelmesh component recording */
+	FTickRecording& OnTickRecording() { return TickRecordingDelegate; }
+
 private:
 	/** When a new AnimBlueprint is created, this will handle post creation work such as adding non-event default nodes */
 	void OnNewBlueprintCreated(UBlueprint* InBlueprint);
@@ -56,6 +82,24 @@ private:
 private:
 	TSharedPtr<FExtensibilityManager> MenuExtensibilityManager;
 	TSharedPtr<FExtensibilityManager> ToolBarExtensibilityManager;
+
+	/** Delegate used to query whether recording is active */
+	FIsRecordingActive IsRecordingActiveDelegate;
+
+	/** Delegate used to start recording animation */
+	FRecord RecordDelegate;
+
+	/** Delegate used to stop recording animation */
+	FStopRecording StopRecordingDelegate;
+
+	/** Delegate used to get the currently recording animation */
+	FGetCurrentRecording GetCurrentRecordingDelegate;
+
+	/** Delegate used to get the currently recording animation time */
+	FGetCurrentRecordingTime GetCurrentRecordingTimeDelegate;
+
+	/** Delegate used to tick the skelmesh component recording */
+	FTickRecording TickRecordingDelegate;
 };
 
 
diff --git a/Engine/Source/Editor/PhAT/Private/PhAT.cpp b/Engine/Source/Editor/PhAT/Private/PhAT.cpp
--- a/Engine/Source/Editor/PhAT/Private/PhAT.cpp
+++ b/Engine/Source/Editor/PhAT/Private/PhAT.cpp
@@ -32,6 +32,8 @@
 #include "Engine/StaticMesh.h"
 #include "EngineLogs.h"
 #include "PhysicalMaterials/PhysicalMaterial.h"
+#include "PersonaModule.h"
+
 DEFINE_LOG_CATEGORY(LogPhAT);
 #define LOCTEXT_NAMESPACE "PhAT"
 
@@ -3202,9 +3204,19 @@ bool FPhAT::IsRecordAvailable() const
 	return (SharedData->EditorSkelComp && SharedData->EditorSkelComp->SkeletalMesh/* && SharedData->bRunningSimulation */);
 }
 
+bool FPhAT::IsRecording() const
+{
+	FPersonaModule& PersonaModule = FModuleManager::LoadModuleChecked<FPersonaModule>(TEXT("Persona"));
+
+	bool bInRecording = false;
+	PersonaModule.OnIsRecordingActive().ExecuteIfBound(SharedData->EditorSkelComp, bInRecording);
+
+	return bInRecording;
+}
+
 FSlateIcon FPhAT::GetRecordStatusImage() const
 {
-	if(SharedData->Recorder.InRecording())
+	if(IsRecording())
 	{
 		return FSlateIcon(FEditorStyle::GetStyleSetName(), "Persona.StopRecordAnimation");
 	}
@@ -3214,7 +3226,7 @@ FSlateIcon FPhAT::GetRecordStatusImage() const
 
 FText FPhAT::GetRecordMenuLabel() const
 {
-	if(SharedData->Recorder.InRecording())
+	if(IsRecording())
 	{
 		return LOCTEXT("Persona_StopRecordAnimationMenuLabel", "Stop Record Animation");
 	}
@@ -3224,7 +3236,7 @@ FText FPhAT::GetRecordMenuLabel() const
 
 FText FPhAT::GetRecordStatusLabel() const
 {
-	if(SharedData->Recorder.InRecording())
+	if(IsRecording())
 	{
 		return LOCTEXT("Persona_StopRecordAnimationLabel", "Stop");
 	}
@@ -3234,7 +3246,7 @@ FText FPhAT::GetRecordStatusLabel() const
 
 FText FPhAT::GetRecordStatusTooltip() const
 {
-	if(SharedData->Recorder.InRecording())
+	if(IsRecording())
 	{
 		return LOCTEXT("Persona_StopRecordAnimation", "Stop Record Animation");
 	}
@@ -3250,23 +3262,15 @@ void FPhAT::RecordAnimation()
 		return;
 	}
 
-	if(SharedData->Recorder.InRecording())
-	{
-		UAnimSequence * AnimSeq = SharedData->Recorder.StopRecord(true);
-
-		// open the asset?
-		if (AnimSeq)
-		{
-			TArray<UObject*> ObjectsToSync;
-			ObjectsToSync.Add(AnimSeq);
+	FPersonaModule& PersonaModule = FModuleManager::LoadModuleChecked<FPersonaModule>(TEXT("Persona"));
 
-			FContentBrowserModule& ContentBrowserModule = FModuleManager::Get().LoadModuleChecked<FContentBrowserModule>("ContentBrowser");
-			ContentBrowserModule.Get().SyncBrowserToAssets(ObjectsToSync, true);
-		}
+	if(IsRecording())
+	{
+		PersonaModule.OnStopRecording().ExecuteIfBound(SharedData->EditorSkelComp);
 	}
 	else
 	{
-		SharedData->Recorder.TriggerRecordAnimation(SharedData->EditorSkelComp);
+		PersonaModule.OnRecord().ExecuteIfBound(SharedData->EditorSkelComp);
 	}
 }
 
diff --git a/Engine/Source/Editor/PhAT/Private/PhAT.h b/Engine/Source/Editor/PhAT/Private/PhAT.h
--- a/Engine/Source/Editor/PhAT/Private/PhAT.h
+++ b/Engine/Source/Editor/PhAT/Private/PhAT.h
@@ -96,6 +96,8 @@ public:
 
 	void OnFocusSelection();
 
+	bool IsRecording() const;
+
 private:
 
 	enum EPhatHierarchyFilterMode
diff --git a/Engine/Source/Editor/PhAT/Private/PhATPreviewViewportClient.cpp b/Engine/Source/Editor/PhAT/Private/PhATPreviewViewportClient.cpp
--- a/Engine/Source/Editor/PhAT/Private/PhATPreviewViewportClient.cpp
+++ b/Engine/Source/Editor/PhAT/Private/PhATPreviewViewportClient.cpp
@@ -18,6 +18,7 @@
 #include "PhysicsEngine/PhysicsSettings.h"
 #include "PhysicsEngine/PhysicsHandleComponent.h"
 #include "DrawDebugHelpers.h"
+#include "PersonaModule.h"
 
 FPhATEdPreviewViewportClient::FPhATEdPreviewViewportClient(TWeakPtr<FPhAT> InPhAT, TSharedPtr<FPhATSharedData> Data, const TSharedRef<SPhATPreviewViewport>& InPhATPreviewViewport)
 	: FEditorViewportClient(nullptr, &Data->PreviewScene, StaticCastSharedRef<SEditorViewport>(InPhATPreviewViewport))
@@ -684,10 +685,10 @@ void FPhATEdPreviewViewportClient::Tick(float DeltaSeconds)
 
 	World->Tick(LEVELTICK_All, DeltaSeconds * SharedData->EditorSimOptions->TimeDilation);
 
-	if(SharedData->Recorder.InRecording())
+	if (PhATPtr.Pin()->IsRecording())
 	{
-		// make sure you don't allow switch SharedData->EditorSkelComp
-		SharedData->Recorder.UpdateRecord(SharedData->EditorSkelComp, DeltaSeconds);
+		FPersonaModule& PersonaModule = FModuleManager::GetModuleChecked<FPersonaModule>("Persona");
+		PersonaModule.OnTickRecording().ExecuteIfBound(SharedData->EditorSkelComp, DeltaSeconds);
 	}
 }
 
diff --git a/Engine/Source/Editor/PhAT/Private/PhATSharedData.cpp b/Engine/Source/Editor/PhAT/Private/PhATSharedData.cpp
--- a/Engine/Source/Editor/PhAT/Private/PhATSharedData.cpp
+++ b/Engine/Source/Editor/PhAT/Private/PhATSharedData.cpp
@@ -65,9 +65,6 @@ FPhATSharedData::FPhATSharedData()
 	EditorSimOptions->HandleAngularDamping = MouseHandle->AngularDamping;
 	EditorSimOptions->HandleAngularStiffness = MouseHandle->AngularStiffness;
 	EditorSimOptions->InterpolationSpeed = MouseHandle->InterpolationSpeed;
-
-	// when record simulation, we'd like to record local to world change
-	Recorder.bRecordLocalToWorld = true;
 }
 
 FPhATSharedData::~FPhATSharedData()
diff --git a/Engine/Source/Editor/PhAT/Private/PhATSharedData.h b/Engine/Source/Editor/PhAT/Private/PhATSharedData.h
--- a/Engine/Source/Editor/PhAT/Private/PhATSharedData.h
+++ b/Engine/Source/Editor/PhAT/Private/PhATSharedData.h
@@ -4,7 +4,6 @@
 
 #include "PreviewScene.h"
 #include "PhysicsPublic.h"
-#include "AnimationRecorder.h"
 
 class UPhATEdSkeletalMeshComponent;
 class UPhysicsHandleComponent;
@@ -215,8 +214,6 @@ public:
 	UBodySetup * CopiedBodySetup;
 	UPhysicsConstraintTemplate * CopiedConstraintTemplate;
 
-	struct FAnimationRecorder Recorder;
-
 	/** Constraint editing */
 	TArray<FSelection> SelectedConstraints;
 	FSelection * GetSelectedConstraint()
diff --git a/Engine/Source/Editor/PlacementMode/PlacementMode.Build.cs b/Engine/Source/Editor/PlacementMode/PlacementMode.Build.cs
--- a/Engine/Source/Editor/PlacementMode/PlacementMode.Build.cs
+++ b/Engine/Source/Editor/PlacementMode/PlacementMode.Build.cs
@@ -21,8 +21,7 @@ public class PlacementMode : ModuleRules
                 "CollectionManager",
                 "LevelEditor",
                 "AssetTools",
-                "EditorStyle",
-				"BspMode"
+                "EditorStyle"
             } 
         );
     }
