Linux: fix native build (UE-26887).

diff --git a/Engine/Source/Editor/MovieSceneTools/Private/Sections/CinematicShotSection.cpp b/Engine/Source/Editor/MovieSceneTools/Private/Sections/CinematicShotSection.cpp
--- a/Engine/Source/Editor/MovieSceneTools/Private/Sections/CinematicShotSection.cpp
+++ b/Engine/Source/Editor/MovieSceneTools/Private/Sections/CinematicShotSection.cpp
@@ -6,7 +6,7 @@
 #include "CinematicShotTrackEditor.h"
 #include "Runtime/Engine/Public/Slate/SceneViewport.h"
 #include "SInlineEditableTextBlock.h"
-#include "CinematicShotTrackEDitor.h"
+#include "CinematicShotTrackEditor.h"
 #include "MovieSceneToolHelpers.h"
 
 #define LOCTEXT_NAMESPACE "FCinematicShotSection"
diff --git a/Engine/Source/Runtime/Core/Public/HAL/Platform.h b/Engine/Source/Runtime/Core/Public/HAL/Platform.h
--- a/Engine/Source/Runtime/Core/Public/HAL/Platform.h
+++ b/Engine/Source/Runtime/Core/Public/HAL/Platform.h
@@ -95,7 +95,7 @@
 
 //@port Identify the platform here and include the platform header to setup the platform types, etc
 #if PLATFORM_WINDOWS
-	#include "Windows/WindowsPlatform.h"
+	#include "Windows/WIndowsPlatform.h"
 #elif PLATFORM_PS4
 	#include "PS4/PS4Platform.h"
 #elif PLATFORM_XBOXONE
