Copying //UE4/Dev-Sequencer to //UE4/Main

diff --git a/Engine/Source/Editor/MovieSceneTools/Private/Sections/ThumbnailSection.cpp b/Engine/Source/Editor/MovieSceneTools/Private/Sections/ThumbnailSection.cpp
--- a/Engine/Source/Editor/MovieSceneTools/Private/Sections/ThumbnailSection.cpp
+++ b/Engine/Source/Editor/MovieSceneTools/Private/Sections/ThumbnailSection.cpp
@@ -12,10 +12,8 @@
 namespace ThumbnailSectionConstants
 {
 	const uint32 ThumbnailHeight = 90;
-	const uint32 TrackHeight = ThumbnailHeight+10; // some extra padding
 	const uint32 TrackWidth = 90;
 	const float SectionGripSize = 4.0f;
-	const FName ThumbnailGripBrushName = FName("Sequencer.Thumbnail.SectionHandle");
 }
 
 
@@ -74,12 +72,15 @@ void FThumbnailSection::DrawViewportThumbnail(TSharedPtr<FTrackEditorThumbnail>
 {
 	if (InternalViewportScene.IsValid())
 	{
-		SequencerPtr.Pin()->SetGlobalTime(TrackEditorThumbnail->GetTime());
+		EMovieScenePlayerStatus::Type PlaybackStatus = SequencerPtr.Pin()->GetPlaybackStatus();
+		SequencerPtr.Pin()->SetPlaybackStatus(EMovieScenePlayerStatus::Jumping);
+		SequencerPtr.Pin()->SetGlobalTimeDirectly(TrackEditorThumbnail->GetTime());
 		InternalViewportClient->SetActorLock(GetCameraObject());
 		InternalViewportClient->UpdateViewForLockedActor();
 		GWorld->SendAllEndOfFrameUpdates();
 		InternalViewportScene->Draw(false);
 		TrackEditorThumbnail->CopyTextureIn((FSlateRenderTargetRHI*)InternalViewportScene->GetViewportRenderTargetTexture());
+		SequencerPtr.Pin()->SetPlaybackStatus(PlaybackStatus);
 	}
 }
 
@@ -92,7 +93,7 @@ uint32 FThumbnailSection::GetThumbnailWidth() const
 
 void FThumbnailSection::RegenerateViewportThumbnails(const FIntPoint& Size)
 {
-	ACameraActor* CameraActor = GetCameraObject();
+	AActor* CameraActor = GetCameraObject();
 
 	if (CameraActor != nullptr && InternalViewportScene.IsValid())
 	{
@@ -107,7 +108,8 @@ void FThumbnailSection::RegenerateViewportThumbnails(const FIntPoint& Size)
 		StoredSize = Size;
 		StoredStartTime = Section->GetStartTime();
 	
-		float AspectRatio = CameraActor->CameraComponent->AspectRatio;
+		UCameraComponent* CameraComponent = MovieSceneHelpers::CameraComponentFromActor(CameraActor);
+		float AspectRatio = CameraComponent ? CameraComponent->AspectRatio : 1.f;
 
 		int32 NewThumbnailWidth = FMath::TruncToInt(ThumbnailSectionConstants::TrackWidth * AspectRatio);
 		int32 ThumbnailCount = FMath::DivideAndRoundUp(StoredSize.X, NewThumbnailWidth);
@@ -159,11 +161,6 @@ void FThumbnailSection::RegenerateViewportThumbnails(const FIntPoint& Size)
 /* ISequencerSection interface
  *****************************************************************************/
 
-bool FThumbnailSection::ShouldDrawKeyAreaBackground() const
-{
-	return false;
-}
-
 bool FThumbnailSection::AreSectionsConnected() const
 {
 	return true;
@@ -174,7 +171,7 @@ TSharedRef<SWidget> FThumbnailSection::GenerateSectionWidget()
 	return SNew(SBox)
 		.HAlign(HAlign_Left)
 		.VAlign(VAlign_Top)
-		.Padding(FMargin(15.0f, 7.0f, 0.0f, 0.0f))
+		.Padding(GetContentPadding())
 		[
 			SNew(SInlineEditableTextBlock)
 				.ToolTipText(CanRename() ? LOCTEXT("RenameThumbnail", "Click or hit F2 to rename") : FText::GetEmpty())
@@ -186,18 +183,6 @@ TSharedRef<SWidget> FThumbnailSection::GenerateSectionWidget()
 }
 
 
-FName FThumbnailSection::GetSectionGripLeftBrushName() const
-{
-	return ThumbnailSectionConstants::ThumbnailGripBrushName;
-}
-
-
-FName FThumbnailSection::GetSectionGripRightBrushName() const
-{
-	return ThumbnailSectionConstants::ThumbnailGripBrushName;
-}
-
-
 float FThumbnailSection::GetSectionGripSize() const
 {
 	return ThumbnailSectionConstants::SectionGripSize;
@@ -206,7 +191,7 @@ float FThumbnailSection::GetSectionGripSize() const
 
 float FThumbnailSection::GetSectionHeight() const
 {
-	return ThumbnailSectionConstants::TrackHeight;
+	return ThumbnailSectionConstants::ThumbnailHeight;
 }
 
 
@@ -221,40 +206,36 @@ FText FThumbnailSection::GetSectionTitle() const
 	return FText::GetEmpty();
 }
 
-
-int32 FThumbnailSection::OnPaintSection(const FGeometry& AllottedGeometry, const FSlateRect& SectionClippingRect, FSlateWindowElementList& OutDrawElements, int32 LayerId, bool bParentEnabled) const
+int32 FThumbnailSection::OnPaintSection( FSequencerSectionPainter& InPainter ) const
 {
-	const ESlateDrawEffect::Type DrawEffects = bParentEnabled ? ESlateDrawEffect::None : ESlateDrawEffect::DisabledEffect;
+	const ESlateDrawEffect::Type DrawEffects = InPainter.bParentEnabled ? ESlateDrawEffect::None : ESlateDrawEffect::DisabledEffect;
+
+	int32 LayerId = InPainter.LayerId;
 
 	if (Thumbnails.Num() != 0)
 	{
-		FGeometry ThumbnailAreaGeometry = AllottedGeometry.MakeChild( FVector2D(GetSectionGripSize(), 0.0f), AllottedGeometry.GetDrawSize() - FVector2D( GetSectionGripSize()*2, 0.0f ) );
-
-		FSlateDrawElement::MakeBox(
-				OutDrawElements,
-				LayerId,
-				ThumbnailAreaGeometry.ToPaintGeometry(),
-				FEditorStyle::GetBrush("Sequencer.GenericSection.Background"),
-				SectionClippingRect,
-				DrawEffects
-			);
-
 		// @todo Sequencer: Need a way to visualize the key here
+		FVector2D SectionSize = InPainter.SectionGeometry.GetLocalSize();
 
 		int32 ThumbnailIndex = 0;
 		for( const TSharedPtr<FTrackEditorThumbnail>& Thumbnail : Thumbnails )
 		{
-			FGeometry TruncatedGeometry = ThumbnailAreaGeometry.MakeChild(
-				Thumbnail->GetSize(),
-				FSlateLayoutTransform(ThumbnailAreaGeometry.Scale, FVector2D( ThumbnailIndex * ThumbnailWidth, 5.f))
-				);
+			FVector2D ThumbnailSize = Thumbnail->GetSize();
+
+			FGeometry TruncatedGeometry = InPainter.SectionGeometry.MakeChild(
+				ThumbnailSize,
+				FSlateLayoutTransform(
+					InPainter.SectionGeometry.Scale,
+					FVector2D( ThumbnailIndex * ThumbnailWidth, (SectionSize.Y - ThumbnailSize.Y)*.5f)
+				)
+			);
 			
 			FSlateDrawElement::MakeViewport(
-				OutDrawElements,
-				LayerId,
+				InPainter.DrawElements,
+				++LayerId,
 				TruncatedGeometry.ToPaintGeometry(),
 				Thumbnail,
-				SectionClippingRect,
+				InPainter.SectionClippingRect,
 				false,
 				false,
 				DrawEffects,
@@ -264,11 +245,11 @@ int32 FThumbnailSection::OnPaintSection(const FGeometry& AllottedGeometry, const
 			if(Thumbnail->GetFadeInCurve() > 0.0f )
 			{
 				FSlateDrawElement::MakeBox(
-					OutDrawElements,
-					LayerId+1,
+					InPainter.DrawElements,
+					++LayerId,
 					TruncatedGeometry.ToPaintGeometry(),
 					WhiteBrush,
-					SectionClippingRect,
+					InPainter.SectionClippingRect,
 					DrawEffects,
 					FLinearColor(1.0f, 1.0f, 1.0f, Thumbnail->GetFadeInCurve())
 					);
@@ -277,33 +258,17 @@ int32 FThumbnailSection::OnPaintSection(const FGeometry& AllottedGeometry, const
 			++ThumbnailIndex;
 		}
 	}
-	else
-	{
-		FSlateDrawElement::MakeBox( 
-			OutDrawElements,
-			LayerId,
-			AllottedGeometry.ToPaintGeometry(),
-			FEditorStyle::GetBrush("Sequencer.GenericSection.Background"),
-			SectionClippingRect,
-			DrawEffects
-		); 
-	}
 
-	return LayerId + 2;
+	return LayerId;
 }
 
 
 void FThumbnailSection::Tick(const FGeometry& AllottedGeometry, const FGeometry& ParentGeometry, const double InCurrentTime, const float InDeltaTime)
 {
-	FTimeToPixel TimeToPixelConverter(AllottedGeometry, TRange<float>(Section->GetStartTime(), Section->GetEndTime()));
-
-	// get the visible time range
-	VisibleTimeRange = TRange<float>(
-		TimeToPixelConverter.PixelToTime(-AllottedGeometry.Position.X),
-		TimeToPixelConverter.PixelToTime(-AllottedGeometry.Position.X + ParentGeometry.Size.X)
-	);
+	// The sequencer view range which will be used to determine whether the individual thumbnails should be drawn if the time range intersects the view range
+	VisibleTimeRange = SequencerPtr.Pin()->GetViewRange();
 
-	FIntPoint AllocatedSize = AllottedGeometry.MakeChild( FVector2D( GetSectionGripSize(), 0.0f ), FVector2D( AllottedGeometry.Size.X - (GetSectionGripSize()*2), AllottedGeometry.Size.Y) ).Size.IntPoint();
+	FIntPoint AllocatedSize = AllottedGeometry.GetLocalSize().IntPoint();
 	AllocatedSize.X = FMath::Max(AllocatedSize.X, 1);
 
 	float StartTime = Section->GetStartTime();
diff --git a/Engine/Source/Editor/MovieSceneTools/Private/Sections/ThumbnailSection.h b/Engine/Source/Editor/MovieSceneTools/Private/Sections/ThumbnailSection.h
--- a/Engine/Source/Editor/MovieSceneTools/Private/Sections/ThumbnailSection.h
+++ b/Engine/Source/Editor/MovieSceneTools/Private/Sections/ThumbnailSection.h
@@ -49,7 +49,7 @@ public:
 	void RegenerateViewportThumbnails(const FIntPoint& Size);
 
 	/** Get the camera that this thumbnail should draw from */
-	virtual ACameraActor* GetCameraObject() const { return nullptr; }
+	virtual AActor* GetCameraObject() const { return nullptr; }
 
 	/** Get whether the text is renameable */
 	virtual bool CanRename() const { return false; }
@@ -64,17 +64,14 @@ public:
 
 	// ISequencerSection interface
 
-	virtual bool ShouldDrawKeyAreaBackground() const override;
 	virtual bool AreSectionsConnected() const override;
 	virtual void GenerateSectionLayout(ISectionLayoutBuilder& LayoutBuilder) const override { }
 	virtual TSharedRef<SWidget> GenerateSectionWidget() override;
-	virtual FName GetSectionGripLeftBrushName() const override;
-	virtual FName GetSectionGripRightBrushName() const override;
 	virtual float GetSectionGripSize() const override;
 	virtual float GetSectionHeight() const override;
 	virtual UMovieSceneSection* GetSectionObject() override;
 	virtual FText GetSectionTitle() const override;
-	virtual int32 OnPaintSection( const FGeometry& AllottedGeometry, const FSlateRect& SectionClippingRect, FSlateWindowElementList& OutDrawElements, int32 LayerId, bool bParentEnabled ) const override;
+	virtual int32 OnPaintSection( FSequencerSectionPainter& InPainter ) const override;
 	virtual void Tick( const FGeometry& AllottedGeometry, const FGeometry& ParentGeometry, const double InCurrentTime, const float InDeltaTime ) override;
 
 protected:
diff --git a/Engine/Source/Editor/MovieSceneTools/Private/TrackEditorThumbnail/TrackEditorThumbnailPool.cpp b/Engine/Source/Editor/MovieSceneTools/Private/TrackEditorThumbnail/TrackEditorThumbnailPool.cpp
--- a/Engine/Source/Editor/MovieSceneTools/Private/TrackEditorThumbnail/TrackEditorThumbnailPool.cpp
+++ b/Engine/Source/Editor/MovieSceneTools/Private/TrackEditorThumbnail/TrackEditorThumbnailPool.cpp
@@ -22,15 +22,13 @@ void FTrackEditorThumbnailPool::AddThumbnailsNeedingRedraw(const TArray<TSharedP
 }
 
 
-void FTrackEditorThumbnailPool::DrawThumbnails()
+bool FTrackEditorThumbnailPool::DrawThumbnails()
 {
 	if (ThumbnailsNeedingDraw.Num() == 0)
 	{
-		return;
+		return false;
 	}
 
-	// save the global time
-	float SavedTime = Sequencer.Pin()->GetGlobalTime();
 	uint32 ThumbnailsDrawn = 0;
 
 	for (int32 ThumbnailsIndex = 0; ThumbnailsDrawn < MaxThumbnailsToDrawAtATime && ThumbnailsIndex < ThumbnailsNeedingDraw.Num(); ++ThumbnailsIndex)
@@ -63,10 +61,9 @@ void FTrackEditorThumbnailPool::DrawThumbnails()
 	{
 		// Ensure all buffers are updated
 		FlushRenderingCommands();
-
-		// restore the global time
-		Sequencer.Pin()->SetGlobalTime(SavedTime);
 	}
+
+	return ThumbnailsDrawn > 0;
 }
 
 
diff --git a/Engine/Source/Editor/MovieSceneTools/Private/TrackEditorThumbnail/TrackEditorThumbnailPool.h b/Engine/Source/Editor/MovieSceneTools/Private/TrackEditorThumbnail/TrackEditorThumbnailPool.h
--- a/Engine/Source/Editor/MovieSceneTools/Private/TrackEditorThumbnail/TrackEditorThumbnailPool.h
+++ b/Engine/Source/Editor/MovieSceneTools/Private/TrackEditorThumbnail/TrackEditorThumbnailPool.h
@@ -22,7 +22,8 @@ public:
 	void AddThumbnailsNeedingRedraw(const TArray<TSharedPtr<FTrackEditorThumbnail>>& InThumbnails);
 
 	/** Draws a small number of thumbnails that are enqueued for drawing */
-	void DrawThumbnails();
+	/* @return Whether thumbnails were drawn */
+	bool DrawThumbnails();
 
 	/** Informs the pool that the thumbnails passed in no longer need to be drawn */
 	void RemoveThumbnailsNeedingRedraw(const TArray< TSharedPtr<FTrackEditorThumbnail>>& InThumbnails);
diff --git a/Engine/Source/Editor/MovieSceneTools/Private/TrackEditors/ActorPickerTrackEditor.cpp b/Engine/Source/Editor/MovieSceneTools/Private/TrackEditors/ActorPickerTrackEditor.cpp
--- a/Engine/Source/Editor/MovieSceneTools/Private/TrackEditors/ActorPickerTrackEditor.cpp
+++ b/Engine/Source/Editor/MovieSceneTools/Private/TrackEditors/ActorPickerTrackEditor.cpp
@@ -108,31 +108,69 @@ void FActorPickerTrackEditor::ShowActorSubMenu(FMenuBuilder& MenuBuilder, FGuid
 
 void FActorPickerTrackEditor::ActorPicked(AActor* ParentActor, FGuid ObjectGuid, UMovieSceneSection* Section)
 {
-	USceneComponent* ComponentWithSockets = NULL;
+	TArray<USceneComponent*> ComponentsWithSockets;
 	if (ParentActor != NULL)
 	{
-		if (USceneComponent* RootComponent = Cast<USceneComponent>(ParentActor->GetRootComponent()))
+		TInlineComponentArray<USceneComponent*> Components(ParentActor);
+
+		for(USceneComponent* Component : Components)
 		{
-			if (RootComponent->HasAnySockets())
+			if (Component->HasAnySockets())
 			{
-				ComponentWithSockets = RootComponent;
+				ComponentsWithSockets.Add(Component);
 			}
 		}
 	}
 
 	// Show socket chooser if we have sockets to select
-	if (ComponentWithSockets != NULL)
+	if (ComponentsWithSockets.Num() > 0)
 	{
 		FLevelEditorModule& LevelEditorModule = FModuleManager::GetModuleChecked<FLevelEditorModule>( "LevelEditor");
 		TSharedPtr< ILevelEditor > LevelEditor = LevelEditorModule.GetFirstLevelEditor();
 
+		TSharedPtr<SWidget> MenuWidget;
+
+		if(ComponentsWithSockets.Num() > 1)
+		{
+			const bool bInShouldCloseWindowAfterMenuSelection = true;
+			FMenuBuilder MenuBuilder(bInShouldCloseWindowAfterMenuSelection, TSharedPtr<const FUICommandList>());
+
+			MenuBuilder.BeginSection(TEXT("ComponentsWithSockets"), LOCTEXT("ComponentsWithSockets", "Components With Sockets"));
+
+			for(USceneComponent* Component : ComponentsWithSockets)
+			{
+				auto NewMenuDelegate = [&](FMenuBuilder& ComponentMenuBuilder)
+				{
+					ComponentMenuBuilder.AddWidget(
+						SNew(SSocketChooserPopup)
+						.SceneComponent(Component)
+						.OnSocketChosen(this, &FActorPickerTrackEditor::ActorSocketPicked, Component, ParentActor, ObjectGuid, Section),
+						FText());
+				};
+
+				MenuBuilder.AddSubMenu(
+					FText::FromName(Component->GetFName()),
+					FText::Format(LOCTEXT("ComponentSubMenuToolTip", "View sockets for component {0}"), FText::FromName(Component->GetFName())),
+					FNewMenuDelegate::CreateLambda(NewMenuDelegate)
+					);
+			}
+		
+			MenuBuilder.EndSection();
+			MenuWidget = MenuBuilder.MakeWidget();
+		}
+		else
+		{
+			MenuWidget = 
+				SNew(SSocketChooserPopup)
+				.SceneComponent(ComponentsWithSockets[0])
+				.OnSocketChosen(this, &FActorPickerTrackEditor::ActorSocketPicked, ComponentsWithSockets[0], ParentActor, ObjectGuid, Section);		
+		}
+
 		// Create as context menu
 		FSlateApplication::Get().PushMenu(
 			LevelEditor.ToSharedRef(),
 			FWidgetPath(),
-			SNew(SSocketChooserPopup)
-			.SceneComponent( ComponentWithSockets )
-			.OnSocketChosen( this, &FActorPickerTrackEditor::ActorSocketPicked, ParentActor, ObjectGuid, Section ),
+			MenuWidget.ToSharedRef(),
 			FSlateApplication::Get().GetCursorPos(),
 			FPopupTransitionEffect( FPopupTransitionEffect::ContextMenu )
 			);
@@ -140,7 +178,7 @@ void FActorPickerTrackEditor::ActorPicked(AActor* ParentActor, FGuid ObjectGuid,
 	else
 	{
 		FSlateApplication::Get().DismissAllMenus();
-		ActorSocketPicked( NAME_None, ParentActor, ObjectGuid, Section );
+		ActorSocketPicked( NAME_None, nullptr, ParentActor, ObjectGuid, Section );
 	}
 }
 
diff --git a/Engine/Source/Editor/MovieSceneTools/Private/TrackEditors/ActorPickerTrackEditor.h b/Engine/Source/Editor/MovieSceneTools/Private/TrackEditors/ActorPickerTrackEditor.h
--- a/Engine/Source/Editor/MovieSceneTools/Private/TrackEditors/ActorPickerTrackEditor.h
+++ b/Engine/Source/Editor/MovieSceneTools/Private/TrackEditors/ActorPickerTrackEditor.h
@@ -15,7 +15,7 @@ public:
 	virtual bool IsActorPickable( const AActor* const ParentActor, FGuid ObjectBinding, UMovieSceneSection* InSection) { return false; }
 
 	/** Actor socked was picked */
-	virtual void ActorSocketPicked(const FName SocketName, AActor* ParentActor, FGuid ObjectBinding, UMovieSceneSection* Section) {}
+	virtual void ActorSocketPicked(const FName SocketName, USceneComponent* Component, AActor* ParentActor, FGuid ObjectBinding, UMovieSceneSection* Section) {}
 
 	/** Show a sub menu of the pickable actors */
 	void ShowActorSubMenu(FMenuBuilder& MenuBuilder, FGuid ObjectBinding, UMovieSceneSection* Section);
diff --git a/Engine/Source/Editor/MovieSceneTools/Private/TrackEditors/AttachTrackEditor.cpp b/Engine/Source/Editor/MovieSceneTools/Private/TrackEditors/AttachTrackEditor.cpp
--- a/Engine/Source/Editor/MovieSceneTools/Private/TrackEditors/AttachTrackEditor.cpp
+++ b/Engine/Source/Editor/MovieSceneTools/Private/TrackEditors/AttachTrackEditor.cpp
@@ -47,21 +47,9 @@ public:
 
 	virtual void GenerateSectionLayout( class ISectionLayoutBuilder& LayoutBuilder ) const override {}
 
-	virtual int32 OnPaintSection( const FGeometry& AllottedGeometry, const FSlateRect& SectionClippingRect, FSlateWindowElementList& OutDrawElements, int32 LayerId, bool bParentEnabled ) const override 
+	virtual int32 OnPaintSection( FSequencerSectionPainter& InPainter ) const override 
 	{
-		const ESlateDrawEffect::Type DrawEffects = bParentEnabled ? ESlateDrawEffect::None : ESlateDrawEffect::DisabledEffect;
-	
-		// Add a box for the section
-		FSlateDrawElement::MakeBox( 
-			OutDrawElements,
-			LayerId,
-			AllottedGeometry.ToPaintGeometry(),
-			FEditorStyle::GetBrush("Sequencer.GenericSection.Background"),
-			SectionClippingRect,
-			DrawEffects
-		); 
-
-		return LayerId;
+		return InPainter.PaintSectionBackground();
 	}
 	
 	virtual void BuildSectionContextMenu(FMenuBuilder& MenuBuilder, const FGuid& ObjectBinding) override
@@ -147,14 +135,14 @@ bool F3DAttachTrackEditor::IsActorPickable(const AActor* const ParentActor, FGui
 }
 
 
-void F3DAttachTrackEditor::ActorSocketPicked(const FName SocketName, AActor* ParentActor, FGuid ObjectGuid, UMovieSceneSection* Section)
+void F3DAttachTrackEditor::ActorSocketPicked(const FName SocketName, USceneComponent* Component, AActor* ParentActor, FGuid ObjectGuid, UMovieSceneSection* Section)
 {
 	if (ObjectGuid.IsValid())
 	{
 		TArray<UObject*> OutObjects;
 		GetSequencer()->GetRuntimeObjects( GetSequencer()->GetFocusedMovieSceneSequenceInstance(), ObjectGuid, OutObjects);
 
-		AnimatablePropertyChanged( FOnKeyProperty::CreateRaw( this, &F3DAttachTrackEditor::AddKeyInternal, OutObjects, SocketName, ParentActor) );
+		AnimatablePropertyChanged( FOnKeyProperty::CreateRaw( this, &F3DAttachTrackEditor::AddKeyInternal, OutObjects, SocketName, Component ? Component->GetFName() : NAME_None, ParentActor) );
 	}
 	else if (Section != nullptr)
 	{
@@ -167,11 +155,12 @@ void F3DAttachTrackEditor::ActorSocketPicked(const FName SocketName, AActor* Par
 		{
 			AttachSection->SetConstraintId(ActorId);
 			AttachSection->AttachSocketName = SocketName;
+			AttachSection->AttachComponentName = Component ? Component->GetFName() : NAME_None;
 		}
 	}
 }
 
-bool F3DAttachTrackEditor::AddKeyInternal( float KeyTime, const TArray<UObject*> Objects, const FName SocketName, AActor* ParentActor)
+bool F3DAttachTrackEditor::AddKeyInternal( float KeyTime, const TArray<UObject*> Objects, const FName SocketName, const FName ComponentName, AActor* ParentActor)
 {
 	bool bHandleCreated = false;
 	bool bTrackCreated = false;
@@ -222,7 +211,7 @@ bool F3DAttachTrackEditor::AddKeyInternal( float KeyTime, const TArray<UObject*>
 					}
 				}
 
-				Cast<UMovieScene3DAttachTrack>(Track)->AddConstraint( KeyTime, AttachEndTime, SocketName, ParentActorId );
+				Cast<UMovieScene3DAttachTrack>(Track)->AddConstraint( KeyTime, AttachEndTime, SocketName, ComponentName, ParentActorId );
 				bTrackModified = true;
 			}
 		}
diff --git a/Engine/Source/Editor/MovieSceneTools/Private/TrackEditors/AttachTrackEditor.h b/Engine/Source/Editor/MovieSceneTools/Private/TrackEditors/AttachTrackEditor.h
--- a/Engine/Source/Editor/MovieSceneTools/Private/TrackEditors/AttachTrackEditor.h
+++ b/Engine/Source/Editor/MovieSceneTools/Private/TrackEditors/AttachTrackEditor.h
@@ -43,10 +43,10 @@ public:
 
 	// FTrackEditorActorPicker
 	virtual bool IsActorPickable( const AActor* const ParentActor, FGuid ObjectBinding, UMovieSceneSection* InSection ) override;
-	virtual void ActorSocketPicked(const FName SocketName, AActor* ParentActor, FGuid ObjectBinding, UMovieSceneSection* Section) override;
+	virtual void ActorSocketPicked(const FName SocketName, USceneComponent* Component, AActor* ParentActor, FGuid ObjectBinding, UMovieSceneSection* Section) override;
 
 private:
 
 	/** Delegate for AnimatablePropertyChanged in AddKey */
-	bool AddKeyInternal(float KeyTime, const TArray<UObject*> Objects, const FName SocketName, AActor* ParentActor);
+	bool AddKeyInternal(float KeyTime, const TArray<UObject*> Objects, const FName SocketName, const FName ComponentName, AActor* ParentActor);
 };
diff --git a/Engine/Source/Editor/MovieSceneTools/Private/TrackEditors/AudioTrackEditor.cpp b/Engine/Source/Editor/MovieSceneTools/Private/TrackEditors/AudioTrackEditor.cpp
--- a/Engine/Source/Editor/MovieSceneTools/Private/TrackEditors/AudioTrackEditor.cpp
+++ b/Engine/Source/Editor/MovieSceneTools/Private/TrackEditors/AudioTrackEditor.cpp
@@ -58,11 +58,11 @@ USoundWave* DeriveSoundWave(UMovieSceneAudioSection* AudioSection)
 		// @todo Sequencer - Right now for sound cues, we just use the first sound wave in the cue
 		// In the future, it would be better to properly generate the sound cue's data after forcing determinism
 		const TArray<USoundNode*>& AllNodes = SoundCue->AllNodes;
-		for (int32 i = 0; i < AllNodes.Num() && SoundWave == nullptr; ++i)
+		for (int32 Index = 0; Index < AllNodes.Num() && SoundWave == nullptr; ++Index)
 		{
-			if (AllNodes[i]->IsA<USoundNodeWavePlayer>())
+			if (AllNodes[Index]->IsA<USoundNodeWavePlayer>())
 			{
-				SoundWave = Cast<USoundNodeWavePlayer>(AllNodes[i])->GetSoundWave();
+				SoundWave = Cast<USoundNodeWavePlayer>(AllNodes[Index])->GetSoundWave();
 			}
 		}
 	}
@@ -79,6 +79,37 @@ float DeriveUnloopedDuration(UMovieSceneAudioSection* AudioSection)
 	return Duration == INDEFINITELY_LOOPING_DURATION ? SoundWave->Duration : Duration;
 }
 
+/** The maximum number of channels we support */
+static const int32 MaxSupportedChannels = 2;
+/** The number of pixels between which to place control points for cubic interpolation */
+static const int32 SmoothingAmount = 6;
+/** The size of the sroked border of the audio wave */
+static const int32 StrokeBorderSize = 2;
+
+/** A specific sample from the audio, specifying peak and average amplitude over the sample's range */
+struct FAudioSample
+{
+	FAudioSample() : RMS(0.f), NumSamples(0) {}
+
+	float RMS;
+	int32 NumSamples;
+	TMap<int32, int32> Histogram;
+};
+
+/** A segment in a cubic spline */
+struct FSplineSegment
+{
+	FSplineSegment() : A(0.f), B(0.f), C(0.f), D(0.f) {}
+
+	/** Cubic polynomial coefficients for the equation f(x) = A + Bx + Cx^2 + Dx^3*/
+	float A, B, C, D;
+	/** The width of this segment */
+	float SampleSize;
+	/** The x-position of this segment */
+	float Position;
+};
+
+
 
 /**
  * The audio thumnail, which holds a texture which it can pass back to a viewport to render
@@ -88,7 +119,7 @@ class FAudioThumbnail
 	, public TSharedFromThis<FAudioThumbnail>
 {
 public:
-	FAudioThumbnail(UMovieSceneSection& InSection, TRange<float> DrawRange, int InTextureSize);
+	FAudioThumbnail(UMovieSceneSection& InSection, TRange<float> DrawRange, int InTextureSize, const FLinearColor& BaseColor, float DisplayScale);
 	~FAudioThumbnail();
 
 	/* ISlateViewport interface */
@@ -101,10 +132,13 @@ public:
 	
 private:
 	/** Generates the waveform preview and dumps it out to the OutBuffer */
-	void GenerateWaveformPreview(TArray<uint8>& OutBuffer, TRange<float> DrawRange);
+	void GenerateWaveformPreview(TArray<uint8>& OutBuffer, TRange<float> DrawRange, float DisplayScale);
+
+	/** Sample the audio data at the given lookup position. Appends the sample result to the Samples array */
+	void SampleAudio(int32 NumChannels, const int16* LookupData, int32 LookupStartIndex, int32 LookupEndIndex, int32 LookupSize, int32 MaxAmplitude, bool bShowIntensity);
 
-	/** Given a buffer of data, a PCM lookup, draws a line of the PCM waveform at the specified X of the data buffer */
-	void DrawWaveformLine(TArray<uint8>& OutData, int32 NumChannels, int16* LookupData, int32 LookupStartIndex, int32 LookupEndIndex, int32 LookupSize, int32 X, int32 XAxisHeight, int32 MaxAmplitude);
+	/** Generate a natural cubic spline from the sample buffer */
+	void GenerateSpline(int32 NumChannels, int32 SamplePositionOffset);
 
 private:
 	/** The section we are visualizing */
@@ -115,23 +149,52 @@ private:
 
 	/** Size of the texture */
 	int32 TextureSize;
+
+	/** Accumulation of audio samples for each channel */
+	TArray<FAudioSample> Samples[MaxSupportedChannels];
+
+	/** Spline segments generated from the above Samples array */
+	TArray<FSplineSegment> SplineSegments[MaxSupportedChannels];
+
+	/** Waveform colors */
+	FLinearColor BoundaryColorHSV;
+	FLinearColor FillColor_A, FillColor_B;
 };
 
+float Modulate(float Value, float Delta, float Range)
+{
+	Value = FMath::Fmod(Value + Delta, Range);
+	if (Value < 0.0f)
+	{
+		Value += Range;
+	}
+	return Value;
+}
 
-FAudioThumbnail::FAudioThumbnail(UMovieSceneSection& InSection, TRange<float> DrawRange, int32 InTextureSize)
+FAudioThumbnail::FAudioThumbnail(UMovieSceneSection& InSection, TRange<float> DrawRange, int32 InTextureSize, const FLinearColor& BaseColor, float DisplayScale)
 	: Section(InSection)
 	, Texture(NULL)
 	, TextureSize(InTextureSize)
 {
 	UMovieSceneAudioSection* AudioSection = Cast<UMovieSceneAudioSection>(&Section);
 	
+	const FLinearColor BaseHSV = BaseColor.LinearRGBToHSV();
+
+	const float BaseValue = FMath::Min(BaseHSV.B, .5f) * BaseHSV.A;
+	const float BaseSaturation = FMath::Max((BaseHSV.G - .45f), 0.f) * BaseHSV.A;
+
+	FillColor_A = FLinearColor(Modulate(BaseHSV.R, -2.5f, 360), BaseSaturation + .35f, BaseValue);
+	FillColor_B = FLinearColor(Modulate(BaseHSV.R,  2.5f, 360), BaseSaturation + .4f, BaseValue + .15f);
+
+	BoundaryColorHSV = FLinearColor(BaseHSV.R, BaseSaturation, BaseValue + .35f);
+
 	if (ShouldRender())
 	{
 		uint32 Size = GetSize().X * GetSize().Y * GPixelFormats[PF_B8G8R8A8].BlockBytes;
 		TArray<uint8> RawData;
 		RawData.AddZeroed(Size);
 
-		GenerateWaveformPreview(RawData, DrawRange);
+		GenerateWaveformPreview(RawData, DrawRange, DisplayScale);
 
 		FSlateTextureDataPtr BulkData(new FSlateTextureData(GetSize().X, GetSize().Y, GPixelFormats[PF_B8G8R8A8].BlockBytes, RawData));
 
@@ -162,8 +225,53 @@ FIntPoint FAudioThumbnail::GetSize() const {return FIntPoint(TextureSize, Animat
 FSlateShaderResource* FAudioThumbnail::GetViewportRenderTargetTexture() const {return Texture;}
 bool FAudioThumbnail::RequiresVsync() const {return false;}
 
+/** Lookup a pixel in the given data buffer based on the specified X and Y */
+uint8* LookupPixel(TArray<uint8>& Data, int32 X, int32 YPos, int32 Width, int32 Height, int32 Channel, int32 NumChannels)
+{
+	int32 Y = Height - YPos - 1;
+	if (NumChannels == 2)
+	{
+		Y = Channel == 0 ? Height/2 - YPos : Height/2 + YPos;
+	}
+
+	int32 Index = (Y * Width + X) * GPixelFormats[PF_B8G8R8A8].BlockBytes;
+	return &Data[Index];
+}
+
+/** Lerp between 2 HSV space colors */
+FLinearColor LerpHSV(const FLinearColor& A, const FLinearColor& B, float Alpha)
+{
+	float SrcHue = A.R;
+	float DestHue = B.R;
+
+	// Take the shortest path to the new hue
+	if (FMath::Abs(SrcHue - DestHue) > 180.0f)
+	{
+		if (DestHue > SrcHue)
+		{
+			SrcHue += 360.0f;
+		}
+		else
+		{
+			DestHue += 360.0f;
+		}
+	}
+
+	float NewHue = FMath::Fmod(FMath::Lerp(SrcHue, DestHue, Alpha), 360.0f);
+	if (NewHue < 0.0f)
+	{
+		NewHue += 360.0f;
+	}
+
+	return FLinearColor(
+		NewHue,
+		FMath::Lerp(A.G, B.G, Alpha),
+		FMath::Lerp(A.B, B.B, Alpha),
+		FMath::Lerp(A.A, B.A, Alpha)
+		);
+}
 
-void FAudioThumbnail::GenerateWaveformPreview(TArray<uint8>& OutData, TRange<float> DrawRange)
+void FAudioThumbnail::GenerateWaveformPreview(TArray<uint8>& OutData, TRange<float> DrawRange, float DisplayScale)
 {
 	UMovieSceneAudioSection* AudioSection = Cast<UMovieSceneAudioSection>(&Section);
 
@@ -172,9 +280,6 @@ void FAudioThumbnail::GenerateWaveformPreview(TArray<uint8>& OutData, TRange<flo
 	
 	check(SoundWave->NumChannels == 1 || SoundWave->NumChannels == 2);
 
-	// @todo Sequencer - This may be fixed when we switch to double precision, but the waveform flickers a bit on resize
-	// If double precision don't solve the issue, this system may have to be re-engineered
-
 	// decompress PCM data if necessary
 	if (SoundWave->RawPCMData == NULL)
 	{
@@ -189,9 +294,12 @@ void FAudioThumbnail::GenerateWaveformPreview(TArray<uint8>& OutData, TRange<flo
 		}
 	}
 
-	int16* LookupData = (int16*)SoundWave->RawPCMData;
-	int32 LookupDataSize = SoundWave->RawPCMDataSize;
-	int32 LookupSize = LookupDataSize * sizeof(uint8) / sizeof(int16);
+	const int32 NumChannels = SoundWave->NumChannels;
+	const int16* LookupData = (int16*)SoundWave->RawPCMData;
+	const int32 LookupDataSize = SoundWave->RawPCMDataSize;
+	const int32 LookupSize = LookupDataSize * sizeof(uint8) / sizeof(int16);
+
+	const bool bShowIntensity = AudioSection->ShouldShowIntensity();
 
 	// @todo Sequencer This fixes looping drawing by pretending we are only dealing with a SoundWave
 	TRange<float> AudioTrueRange = TRange<float>(
@@ -201,7 +309,22 @@ void FAudioThumbnail::GenerateWaveformPreview(TArray<uint8>& OutData, TRange<flo
 
 	float DrawRangeSize = DrawRange.Size<float>();
 
-	for (int32 X = 0; X < GetSize().X; ++X)
+	const int32 MaxAmplitude = NumChannels == 1 ? GetSize().Y : GetSize().Y / 2;
+
+	int32 DrawOffsetPx = FMath::Max(FMath::RoundToInt((DrawRange.GetLowerBoundValue() - AudioSection->GetRange().GetLowerBoundValue()) / DisplayScale), 0);
+
+
+	// In order to prevent flickering waveforms when moving the display position/range around, we have to lock our sample position and spline segments to the view range
+	float RangeLookupFraction = (SmoothingAmount*DisplayScale) / TrueRangeSize;
+	int32 LookupRange = FMath::Clamp(FMath::TruncToInt(RangeLookupFraction * LookupSize), 1, LookupSize);
+
+	int32 SampleLockOffset = DrawOffsetPx % SmoothingAmount;
+
+	int32 FirstSample = -2.f*SmoothingAmount - SampleLockOffset;
+	int32 LastSample = GetSize().X + 2*SmoothingAmount;
+
+	// Sample the audio one pixel to the left and right
+	for (int32 X = FirstSample; X < LastSample; ++X)
 	{
 		float LookupTime = ((float)(X - 0.5f) / (float)GetSize().X) * DrawRangeSize + DrawRange.GetLowerBoundValue();
 		float LookupFraction = (LookupTime - AudioTrueRange.GetLowerBoundValue()) / TrueRangeSize;
@@ -209,126 +332,271 @@ void FAudioThumbnail::GenerateWaveformPreview(TArray<uint8>& OutData, TRange<flo
 		
 		float NextLookupTime = ((float)(X + 0.5f) / (float)GetSize().X) * DrawRangeSize + DrawRange.GetLowerBoundValue();
 		float NextLookupFraction = (NextLookupTime - AudioTrueRange.GetLowerBoundValue()) / TrueRangeSize;
-		int32 NextLookupIndex = FMath::TruncToInt(NextLookupFraction * LookupSize);
+		int32 NextLookupIndex = FMath::TruncToInt(NextLookupFraction  * LookupSize);
 
-		DrawWaveformLine(OutData, SoundWave->NumChannels, LookupData, LookupIndex, NextLookupIndex, LookupSize, X, GetSize().Y / 2, GetSize().Y / 2);
+		if (LookupRange)
+		{
+			LookupIndex -= FMath::Clamp(LookupIndex, 0, LookupSize) % LookupRange;
+			NextLookupIndex = LookupIndex + LookupRange;
+		}
+		
+		SampleAudio(SoundWave->NumChannels, LookupData, LookupIndex, NextLookupIndex, LookupSize, MaxAmplitude, bShowIntensity);
 	}
-}
-
 
-void FAudioThumbnail::DrawWaveformLine(TArray<uint8>& OutData, int32 NumChannels, int16* LookupData, int32 LookupStartIndex, int32 LookupEndIndex, int32 LookupSize, int32 X, int32 XAxisHeight, int32 MaxAmplitude)
-{
-	int32 ModifiedLookupStartIndex = NumChannels == 2 ? (LookupStartIndex % 2 == 0 ? LookupStartIndex : LookupStartIndex - 1) : LookupStartIndex;
+	// Generate a spline
+	GenerateSpline(SoundWave->NumChannels, FirstSample);
 
-	int32 ClampedStartIndex = FMath::Max(0, ModifiedLookupStartIndex);
-	int32 ClampedEndIndex = FMath::Max(FMath::Min(LookupSize-1, LookupEndIndex), ClampedStartIndex + 1);
-	
-	int32 StepSize = NumChannels;
+	// Now draw the spline
+	const int32 Height = GetSize().Y;
+	const int32 Width = GetSize().X;
 
-	// optimization - don't take more than a maximum number of samples per pixel
-	int32 Range = ClampedEndIndex - ClampedStartIndex;
-	int32 SampleCount = Range / StepSize;
-	int32 MaxSampleCount = AnimatableAudioEditorConstants::MaxSamplesPerPixel;
-	int32 ModifiedStepSize = (SampleCount > MaxSampleCount ? SampleCount / MaxSampleCount : 1) * StepSize;
+	FLinearColor BoundaryColor = BoundaryColorHSV.HSVToLinearRGB();
 
-	// Two techniques for rendering smooth waveforms: Summed Density Histograms and Root Mean Square
-#if 1 // Summed Density Histogram
-	struct FChannelData
+	for (int32 ChannelIndex = 0; ChannelIndex < SoundWave->NumChannels; ++ChannelIndex)
 	{
-		FChannelData(int32 InMaxAmplitude)
-			: HistogramSum(0)
+		int32 SplineIndex = 0;
+
+		for (int32 X = 0; X < Width; ++X)
 		{
-			Histogram.AddZeroed(InMaxAmplitude);
-		}
+			bool bOutOfRange = SplineIndex >= SplineSegments[ChannelIndex].Num();
+			while (!bOutOfRange && X >= SplineSegments[ChannelIndex][SplineIndex].Position+SplineSegments[ChannelIndex][SplineIndex].SampleSize)
+			{
+				++SplineIndex;
+				bOutOfRange = SplineIndex >= SplineSegments[ChannelIndex].Num();
+			}
+			
+			if (bOutOfRange)
+			{
+				break;
+			}
 
-		TArray<int32> Histogram;
-		int32 HistogramSum;
-	};
-	TArray<FChannelData> ChannelData;
-	for (int32 i = 0; i < NumChannels; ++i)
-	{
-		ChannelData.Add(FChannelData(MaxAmplitude));
-	}
+			// Evaluate the spline
+			const float DistBetweenPts = (X-SplineSegments[ChannelIndex][SplineIndex].Position)/SplineSegments[ChannelIndex][SplineIndex].SampleSize;
+			const float Amplitude = 
+				SplineSegments[ChannelIndex][SplineIndex].A +
+				SplineSegments[ChannelIndex][SplineIndex].B * DistBetweenPts +
+				SplineSegments[ChannelIndex][SplineIndex].C * FMath::Pow(DistBetweenPts, 2) +
+				SplineSegments[ChannelIndex][SplineIndex].D * FMath::Pow(DistBetweenPts, 3);
 
-	// build up the histogram(s)
-	for (int32 i = ClampedStartIndex; i < ClampedEndIndex; i += ModifiedStepSize)
-	{
-		for (int32 ChannelIndex = 0; ChannelIndex < NumChannels; ++ChannelIndex)
-		{
-			int32 DataPoint = LookupData[i + ChannelIndex];
-			int32 HistogramData = FMath::Clamp(FMath::TruncToInt(FMath::Abs(DataPoint) / 32768.f * MaxAmplitude), 0, MaxAmplitude - 1);
+			// @todo: draw border according to gradient of curve to prevent aliasing on steep gradients? This would be non-trivial...
+			const float BoundaryStart = Amplitude - StrokeBorderSize * 0.5f;
+			const float BoundaryEnd = Amplitude + StrokeBorderSize * 0.5f;
 
-			++ChannelData[ChannelIndex].Histogram[HistogramData];
-			++ChannelData[ChannelIndex].HistogramSum;
+			const FAudioSample& Sample = Samples[ChannelIndex][X - FirstSample];
+
+			float IntensitySummation = bShowIntensity ? Sample.NumSamples : 0.f;
+			for (int32 PixelIndex = 0; PixelIndex < MaxAmplitude; ++PixelIndex)
+			{
+				uint8* Pixel = LookupPixel(OutData, X, PixelIndex, Width, Height, ChannelIndex, NumChannels);
+
+				if (bShowIntensity)
+				{
+					IntensitySummation -= Sample.Histogram.FindRef(PixelIndex);
+				}
+
+				const float PixelCenter = PixelIndex + 0.5f;
+
+				const float Dither = FMath::FRand() * .025f - .0125f;
+				const float GradLerp = FMath::Clamp(float(PixelIndex) / MaxAmplitude + Dither, 0.f, 1.f);
+				FLinearColor SolidFilledColor = LerpHSV(FillColor_A, FillColor_B, GradLerp);
+
+				float BorderBlend = 1.f;
+				if (PixelIndex <= FMath::TruncToInt(BoundaryStart))
+				{
+					BorderBlend = 1.f - FMath::Clamp(BoundaryStart - PixelIndex, 0.f, 1.f);
+				}
+				
+				FLinearColor Color = LerpHSV(SolidFilledColor, BoundaryColorHSV, BorderBlend).HSVToLinearRGB();
+
+				// Calculate alpha based on how far from the boundary we are
+				float Alpha = FMath::Max(FMath::Clamp(BoundaryEnd - PixelCenter, 0.f, 1.f), IntensitySummation / Sample.NumSamples);
+				if (Alpha <= 0.f)
+				{
+					break;
+				}
+
+				// Slate viewports must have pre-multiplied alpha
+				*Pixel++ = Color.B*Alpha*255;
+				*Pixel++ = Color.G*Alpha*255;
+				*Pixel++ = Color.R*Alpha*255;
+				*Pixel++ = Alpha*255;
+			}
 		}
 	}
+}
 
-	// output the data to the texture
-	const int32 MaxSupportedChannels = 2;
-	for (int32 ChannelIndex = 0; ChannelIndex < MaxSupportedChannels; ++ChannelIndex)
+void FAudioThumbnail::GenerateSpline(int32 NumChannels, int32 SamplePositionOffset)
+{
+	// Generate a cubic polynomial spline interpolating the samples
+	for (int32 ChannelIndex = 0; ChannelIndex < NumChannels; ++ChannelIndex)
 	{
-		FChannelData& ChannelDataUsed = ChannelData[ChannelIndex % NumChannels];
-		int32 IntensitySummation = 0;
-		for (int32 PixelIndex = MaxAmplitude - 1; PixelIndex >= 0; --PixelIndex)
+		TArray<FSplineSegment>& Segments = SplineSegments[ChannelIndex];
+
+		const int32 NumSamples = Samples[ChannelIndex].Num();
+
+		struct FControlPoint
+		{
+			float Value;
+			float Position;
+			int32 SampleSize;
+		};
+		TArray<FControlPoint> ControlPoints;
+
+		for (int SampleIndex = 0; SampleIndex < NumSamples; SampleIndex += SmoothingAmount)
 		{
-			IntensitySummation += ChannelDataUsed.Histogram[PixelIndex];
-			if (IntensitySummation > 0)
+			float RMS = 0.f;
+			int32 NumAvgs = FMath::Min(SmoothingAmount, NumSamples - SampleIndex);
+			
+			for (int32 SubIndex = 0; SubIndex < NumAvgs; ++SubIndex)
 			{
-				int32 AlphaIntensity = FMath::Clamp(IntensitySummation * 255 / ChannelDataUsed.HistogramSum, 0, 255);
-
-				int32 Y = ChannelIndex == 0 ? XAxisHeight - PixelIndex : XAxisHeight + PixelIndex;
-				int32 Index = (Y * GetSize().X + X) * GPixelFormats[PF_B8G8R8A8].BlockBytes;
-				OutData[Index+0] = 255;
-				OutData[Index+1] = 255;
-				OutData[Index+2] = 255;
-				OutData[Index+3] = AlphaIntensity;
+				RMS += FMath::Pow(Samples[ChannelIndex][SampleIndex + SubIndex].RMS, 2);
+			}
+
+			const int32 SegmentSize2 = NumAvgs / 2;
+			const int32 SegmentSize1 = NumAvgs - SegmentSize2;
+
+			RMS = FMath::Sqrt(RMS / NumAvgs);
+
+			FControlPoint& StartPoint = ControlPoints[ControlPoints.AddZeroed()];
+			StartPoint.Value = Samples[ChannelIndex][SampleIndex].RMS;
+			StartPoint.SampleSize = SegmentSize1;
+			StartPoint.Position = SampleIndex + SamplePositionOffset;
+
+			if (SegmentSize2 > 0)
+			{
+				FControlPoint& MidPoint = ControlPoints[ControlPoints.AddZeroed()];
+				MidPoint.Value = RMS;
+				MidPoint.SampleSize = SegmentSize2;
+				MidPoint.Position = SampleIndex + SamplePositionOffset + SegmentSize1;
 			}
 		}
+
+		if (ControlPoints.Num() <= 1)
+		{
+			continue;
+		}
+
+		const int32 LastIndex = ControlPoints.Num() - 1;
+
+		// Perform gaussian elimination on the following tridiagonal matrix that defines the piecewise cubic polynomial
+		// spline for n control points, given f(x), f'(x) and f''(x) continuity. Imposed boundary conditions are f''(0) = f''(n) = 0.
+		//	(D[i] = f[i]'(x))
+		//	1	2						D[i]	= 3(y[1] - y[0])
+		//	1	4	1					D[i+1]	= 3(y[2] - y[1])
+		//		1	4	1				|		|
+		//		\	\	\	\	\		|		|
+		//					1	4	1	|		= 3(y[n-1] - y[n-2])
+		//						1	2	D[n]	= 3(y[n] - y[n-1])
+		struct FMinimalMatrixComponent
+		{
+			float DiagComponent;
+			float KnownConstant;
+		};
+
+		TArray<FMinimalMatrixComponent> GaussianCoefficients;
+		GaussianCoefficients.AddZeroed(ControlPoints.Num());
+
+		// Setup the top left of the matrix
+		GaussianCoefficients[0].KnownConstant = 3.f * (ControlPoints[1].Value - ControlPoints[0].Value);
+		GaussianCoefficients[0].DiagComponent = 2.f;
+
+		// Calculate the diagonal component of each row, based on the eliminated value of the last
+		for (int32 Index = 1; Index < GaussianCoefficients.Num() - 1; ++Index)
+		{
+			GaussianCoefficients[Index].KnownConstant = (3.f * (ControlPoints[Index+1].Value - ControlPoints[Index-1].Value)) - (GaussianCoefficients[Index-1].KnownConstant / GaussianCoefficients[Index-1].DiagComponent);
+			GaussianCoefficients[Index].DiagComponent = 4.f - (1.f / GaussianCoefficients[Index-1].DiagComponent);
+		}
+		
+		// Setup the bottom right of the matrix
+		GaussianCoefficients[LastIndex].KnownConstant = (3.f * (ControlPoints[LastIndex].Value - ControlPoints[LastIndex-1].Value)) - (GaussianCoefficients[LastIndex-1].KnownConstant / GaussianCoefficients[LastIndex-1].DiagComponent);
+		GaussianCoefficients[LastIndex].DiagComponent = 2.f - (1.f / GaussianCoefficients[LastIndex-1].DiagComponent);
+
+		// Now we have an upper triangular matrix, we can use reverse substitution to calculate D[n] -> D[0]
+
+		TArray<float> FirstOrderDerivatives;
+		FirstOrderDerivatives.AddZeroed(GaussianCoefficients.Num());
+
+		FirstOrderDerivatives[LastIndex] = GaussianCoefficients[LastIndex].KnownConstant / GaussianCoefficients[LastIndex].DiagComponent;
+
+		for (int32 Index = GaussianCoefficients.Num() - 2; Index >= 0; --Index)
+		{
+			FirstOrderDerivatives[Index] = (GaussianCoefficients[Index].KnownConstant - FirstOrderDerivatives[Index+1]) / GaussianCoefficients[Index].DiagComponent;
+		}
+
+		// Now we know the first-order derivatives of each control point, calculating the interpolating polynomial is trivial
+		// f(x) = a + bx + cx^2 + dx^3
+		//	a = y
+		//	b = D[i]
+		//	c = 3(y[i+1] - y[i]) - 2D[i] - D[i+1]
+		//	d = 2(y[i] - y[i+1]) + 2D[i] + D[i+1]
+		for (int32 Index = 0; Index < FirstOrderDerivatives.Num() - 2; ++Index)
+		{
+			Segments.Emplace();
+			Segments.Last().A = ControlPoints[Index].Value;
+			Segments.Last().B = FirstOrderDerivatives[Index];
+			Segments.Last().C = 3.f*(ControlPoints[Index+1].Value - ControlPoints[Index].Value) - 2*FirstOrderDerivatives[Index] - FirstOrderDerivatives[Index+1];
+			Segments.Last().D = 2.f*(ControlPoints[Index].Value - ControlPoints[Index+1].Value) + FirstOrderDerivatives[Index] + FirstOrderDerivatives[Index+1];
+
+			Segments.Last().Position = ControlPoints[Index].Position;
+			Segments.Last().SampleSize = ControlPoints[Index].SampleSize;
+		}
 	}
-#else // Root Mean Square
-	int32 Samples = 0;
-	float SummedData = 0;
-	int32 DataAmplitude = 0;
-	for (int32 i = ClampedStartIndex; i < ClampedEndIndex; i += StepSize)
-	{
-		int32 DataPoint = LookupData[i];
-		SummedData += DataPoint * DataPoint;
-		++Samples;
-		DataAmplitude = FMath::Max(DataAmplitude, FMath::Abs(DataPoint));
-	}
-	float RootMeanSquare = FMath::Sqrt(SummedData / (float)Samples);
-	
-	int32 RMSHeight = FMath::Clamp(FMath::TruncToInt(RootMeanSquare / 32768.f * MaxAmplitude), 1, MaxAmplitude);
-	int32 MaxHeight = FMath::Clamp(FMath::TruncToInt(DataAmplitude / 32768.f * MaxAmplitude), 1, MaxAmplitude);
-	float HeightRange = FMath::Max(1, MaxHeight - RMSHeight);
+}
+
+void FAudioThumbnail::SampleAudio(int32 NumChannels, const int16* LookupData, int32 LookupStartIndex, int32 LookupEndIndex, int32 LookupSize, int32 MaxAmplitude, bool bShowIntensity)
+{
+	LookupStartIndex = NumChannels == 2 ? (LookupStartIndex % 2 == 0 ? LookupStartIndex : LookupStartIndex - 1) : LookupStartIndex;
+	LookupEndIndex = FMath::Max(LookupEndIndex, LookupStartIndex + 1);
 	
-	for (int32 i = 0; i < MaxHeight; ++i)
+	int32 StepSize = NumChannels;
+
+	// optimization - don't take more than a maximum number of samples per pixel
+	int32 Range = LookupEndIndex - LookupStartIndex;
+	int32 SampleCount = Range / StepSize;
+	int32 MaxSampleCount = AnimatableAudioEditorConstants::MaxSamplesPerPixel;
+	int32 ModifiedStepSize = (SampleCount > MaxSampleCount ? SampleCount / MaxSampleCount : 1) * StepSize;
+
+	for (int32 ChannelIndex = 0; ChannelIndex < NumChannels; ++ChannelIndex)
 	{
-		float Intensity = (MaxHeight - i) / HeightRange;
-		int32 IntegralIntensity = FMath::Clamp(FMath::TruncToInt(Intensity * 255), 0, 255);
-
-		int32 Y = XAxisHeight - i;
-		int32 Index = (Y * GetSize().X + X) * GPixelFormats[PF_B8G8R8A8].BlockBytes;
-		OutData[Index+0] = 255;
-		OutData[Index+1] = 255;
-		OutData[Index+2] = 255;
-		OutData[Index+3] = IntegralIntensity;
-
-		Y = XAxisHeight + i;
-		Index = (Y * GetSize().X + X) * GPixelFormats[PF_B8G8R8A8].BlockBytes;
-		OutData[Index+0] = 255;
-		OutData[Index+1] = 255;
-		OutData[Index+2] = 255;
-		OutData[Index+3] = IntegralIntensity;
+		FAudioSample& NewSample = Samples[ChannelIndex][Samples[ChannelIndex].Emplace()];
+
+		for (int32 Index = LookupStartIndex; Index < LookupEndIndex; Index += ModifiedStepSize)
+		{
+			if (Index < 0 || Index >= LookupSize)
+			{
+				NewSample.RMS += 0.f;
+				if (bShowIntensity)
+				{
+					++NewSample.Histogram.FindOrAdd(0);
+				}
+				++NewSample.NumSamples;
+				continue;
+			}
+
+			int32 DataPoint = LookupData[Index + ChannelIndex];
+			int32 Sample = FMath::Clamp(FMath::TruncToInt(FMath::Abs(DataPoint) / 32768.f * MaxAmplitude), 0, MaxAmplitude - 1);
+
+			NewSample.RMS += FMath::Pow(Sample, 2);
+			if (bShowIntensity)
+			{
+				++NewSample.Histogram.FindOrAdd(Sample);
+			}
+			++NewSample.NumSamples;
+		}
+
+		if (NewSample.NumSamples)
+		{
+			NewSample.RMS = (FMath::Sqrt(NewSample.RMS / NewSample.NumSamples));
+		}
 	}
-#endif
 }
 
 
-FAudioSection::FAudioSection( UMovieSceneSection& InSection, bool bOnAMasterTrack )
+FAudioSection::FAudioSection( UMovieSceneSection& InSection, bool bOnAMasterTrack, TWeakPtr<ISequencer> InSequencer )
 	: Section( InSection )
 	, StoredDrawRange(TRange<float>::Empty())
 	, bIsOnAMasterTrack(bOnAMasterTrack)
+	, Sequencer(InSequencer)
 {
 }
 
@@ -343,11 +611,6 @@ UMovieSceneSection* FAudioSection::GetSectionObject()
 	return &Section;
 }
 
-bool FAudioSection::ShouldDrawKeyAreaBackground() const
-{
-	return false;
-}
-
 FText FAudioSection::GetDisplayName() const
 {
 	return bIsOnAMasterTrack ? NSLOCTEXT("FAudioSection", "MasterAudioDisplayName", "Master Audio") :  NSLOCTEXT("FAudioSection", "AudioDisplayName", "Audio");
@@ -366,91 +629,94 @@ float FAudioSection::GetSectionHeight() const
 }
 
 
-int32 FAudioSection::OnPaintSection( const FGeometry& AllottedGeometry, const FSlateRect& SectionClippingRect, FSlateWindowElementList& OutDrawElements, int32 LayerId, bool bParentEnabled ) const
+int32 FAudioSection::OnPaintSection( FSequencerSectionPainter& Painter ) const
 {
-	const ESlateDrawEffect::Type DrawEffects = bParentEnabled ? ESlateDrawEffect::None : ESlateDrawEffect::DisabledEffect;
-
-	// Add a box for the section
-	FSlateDrawElement::MakeBox(
-		OutDrawElements,
-		LayerId,
-		AllottedGeometry.ToPaintGeometry(),
-		FEditorStyle::GetBrush("Sequencer.GenericSection.Background"),
-		SectionClippingRect,
-		DrawEffects,
-		FLinearColor(0.4f, 0.8f, 0.4f, 1.f)
-	);
+	int32 LayerId = Painter.PaintSectionBackground();
 
 	if (WaveformThumbnail.IsValid() && WaveformThumbnail->ShouldRender())
 	{
 		// @todo Sequencer draw multiple times if looping possibly - requires some thought about SoundCues
 		FSlateDrawElement::MakeViewport(
-			OutDrawElements,
-			LayerId+1,
-			AllottedGeometry.ToPaintGeometry(FVector2D(StoredXOffset, 0), FVector2D(StoredXSize, AllottedGeometry.Size.Y)),
+			Painter.DrawElements,
+			++LayerId,
+			Painter.SectionGeometry.ToPaintGeometry(FVector2D(StoredXOffset, 0), FVector2D(StoredXSize, Painter.SectionGeometry.Size.Y)),
 			WaveformThumbnail,
-			SectionClippingRect,
+			Painter.SectionClippingRect,
 			false,
-			false,
-			DrawEffects,
-			FLinearColor(1.f, 1.f, 1.f, 0.9f)
+			true,
+			Painter.bParentEnabled ? ESlateDrawEffect::None : ESlateDrawEffect::DisabledEffect,
+			FLinearColor::White
 		);
 	}
 
-	return LayerId+1;
+	return LayerId;
 }
 
 
 void FAudioSection::Tick( const FGeometry& AllottedGeometry, const FGeometry& ParentGeometry, const double InCurrentTime, const float InDeltaTime )
 {
 	UMovieSceneAudioSection* AudioSection = Cast<UMovieSceneAudioSection>(&Section);
+	UMovieSceneTrack* Track = Section.GetTypedOuter<UMovieSceneTrack>();
 
 	USoundWave* SoundWave = DeriveSoundWave(AudioSection);
-	if (SoundWave->NumChannels == 1 || SoundWave->NumChannels == 2)
+	if (Track && SoundWave->NumChannels == 1 || SoundWave->NumChannels == 2)
 	{
+		const FSlateRect ParentRect = TransformRect(
+			Concatenate(
+				ParentGeometry.GetAccumulatedLayoutTransform(),
+				AllottedGeometry.GetAccumulatedLayoutTransform().Inverse()
+			),
+			FSlateRect(FVector2D(0, 0), ParentGeometry.GetLocalSize())
+		);
+
+		const float LeftMostVisiblePixel = FMath::Max(ParentRect.Left, 0.f);
+		const float RightMostVisiblePixel = FMath::Min(ParentRect.Right, AllottedGeometry.GetLocalSize().X);
+
 		FTimeToPixel TimeToPixelConverter( AllottedGeometry, TRange<float>( AudioSection->GetStartTime(), AudioSection->GetEndTime() ) );
 
-		TRange<float> VisibleRange = TRange<float>(
-			TimeToPixelConverter.PixelToTime(-AllottedGeometry.Position.X),
-			TimeToPixelConverter.PixelToTime(-AllottedGeometry.Position.X + ParentGeometry.Size.X));
-		TRange<float> SectionRange = AudioSection->GetRange();
-		TRange<float> AudioRange = TRange<float>(
-			AudioSection->GetAudioStartTime(),
-			AudioSection->GetAudioStartTime() + DeriveUnloopedDuration(AudioSection) * AudioSection->GetAudioDilationFactor());
+		TRange<float> DrawRange = TRange<float>(
+			TimeToPixelConverter.PixelToTime(LeftMostVisiblePixel),
+			TimeToPixelConverter.PixelToTime(RightMostVisiblePixel)
+			);
 
-		TRange<float> DrawRange = TRange<float>::Intersection(TRange<float>::Intersection(SectionRange, AudioRange), VisibleRange);
-	
-		TRange<int32> PixelRange = TRange<int32>(
-			TimeToPixelConverter.TimeToPixel(DrawRange.GetLowerBoundValue()),
-			TimeToPixelConverter.TimeToPixel(DrawRange.GetUpperBoundValue()));
-		
 		// generate texture x offset and x size
-		int32 XOffset = PixelRange.GetLowerBoundValue() - TimeToPixelConverter.TimeToPixel(SectionRange.GetLowerBoundValue());
-		int32 XSize = PixelRange.Size<int32>();
+		int32 XOffset = LeftMostVisiblePixel;//PixelRange.GetLowerBoundValue() - TimeToPixelConverter.TimeToPixel(SectionRange.GetLowerBoundValue());
+		int32 XSize = RightMostVisiblePixel - LeftMostVisiblePixel;//PixelRange.Size<int32>();
 
 		if (!FMath::IsNearlyEqual(DrawRange.GetLowerBoundValue(), StoredDrawRange.GetLowerBoundValue()) ||
 			!FMath::IsNearlyEqual(DrawRange.GetUpperBoundValue(), StoredDrawRange.GetUpperBoundValue()) ||
-			XOffset != StoredXOffset || XSize != StoredXSize)
+			XOffset != StoredXOffset || XSize != StoredXSize || Track->GetColorTint() != StoredColor)
 		{
-			RegenerateWaveforms(DrawRange, XOffset, XSize);
+
+			float DisplayScale = XSize / DrawRange.Size<float>();
+
+			// Use the view range if possible, as it's much more stable than using the texture size and draw range
+			TSharedPtr<ISequencer> SequencerPin = Sequencer.Pin();
+			if (SequencerPin.IsValid())
+			{
+				DisplayScale = SequencerPin->GetViewRange().Size<float>() / ParentGeometry.GetLocalSize().X;
+			}
+
+			RegenerateWaveforms(DrawRange, XOffset, XSize, Track->GetColorTint(), DisplayScale);
 		}
 	}
 }
 
 
-void FAudioSection::RegenerateWaveforms(TRange<float> DrawRange, int32 XOffset, int32 XSize)
+void FAudioSection::RegenerateWaveforms(TRange<float> DrawRange, int32 XOffset, int32 XSize, const FColor& ColorTint, float DisplayScale)
 {
 	StoredDrawRange = DrawRange;
 	StoredXOffset = XOffset;
 	StoredXSize = XSize;
-	
+	StoredColor = ColorTint;
+
 	if (DrawRange.IsDegenerate() || DrawRange.IsEmpty() || Cast<UMovieSceneAudioSection>(&Section)->GetSound() == NULL)
 	{
 		WaveformThumbnail.Reset();
 	}
 	else
 	{
-		WaveformThumbnail = MakeShareable(new FAudioThumbnail(Section, DrawRange, XSize));
+		WaveformThumbnail = MakeShareable(new FAudioThumbnail(Section, DrawRange, XSize, ColorTint, DisplayScale));
 	}
 }
 
@@ -499,7 +765,7 @@ TSharedRef<ISequencerSection> FAudioTrackEditor::MakeSectionInterface( UMovieSce
 	check( SupportsType( SectionObject.GetOuter()->GetClass() ) );
 	
 	bool bIsAMasterTrack = Cast<UMovieScene>(Track.GetOuter())->IsAMasterTrack(Track);
-	return MakeShareable( new FAudioSection(SectionObject, bIsAMasterTrack) );
+	return MakeShareable( new FAudioSection(SectionObject, bIsAMasterTrack, GetSequencer()) );
 }
 
 
diff --git a/Engine/Source/Editor/MovieSceneTools/Private/TrackEditors/AudioTrackEditor.h b/Engine/Source/Editor/MovieSceneTools/Private/TrackEditors/AudioTrackEditor.h
--- a/Engine/Source/Editor/MovieSceneTools/Private/TrackEditors/AudioTrackEditor.h
+++ b/Engine/Source/Editor/MovieSceneTools/Private/TrackEditors/AudioTrackEditor.h
@@ -64,7 +64,7 @@ class FAudioSection
 public:
 
 	/** Constructor. */
-	FAudioSection(UMovieSceneSection& InSection, bool bOnAMasterTrack);
+	FAudioSection(UMovieSceneSection& InSection, bool bOnAMasterTrack, TWeakPtr<ISequencer> InSequencer);
 
 	/** Virtual destructor. */
 	virtual ~FAudioSection();
@@ -74,18 +74,17 @@ public:
 	// ISequencerSection interface
 
 	virtual UMovieSceneSection* GetSectionObject() override;
-	virtual bool ShouldDrawKeyAreaBackground() const override;
 	virtual FText GetDisplayName() const override;
 	virtual FText GetSectionTitle() const override;
 	virtual float GetSectionHeight() const override;
 	virtual void GenerateSectionLayout( class ISectionLayoutBuilder& LayoutBuilder) const override { }
-	virtual int32 OnPaintSection(const FGeometry& AllottedGeometry, const FSlateRect& SectionClippingRect, FSlateWindowElementList& OutDrawElements, int32 LayerId, bool bParentEnabled) const override;
+	virtual int32 OnPaintSection(FSequencerSectionPainter& Painter) const override;
 	virtual void Tick(const FGeometry& AllottedGeometry, const FGeometry& ParentGeometry, const double InCurrentTime, const float InDeltaTime) override;
 	
 private:
 
 	/* Re-creates the texture used to preview the waveform. */
-	void RegenerateWaveforms(TRange<float> DrawRange, int32 XOffset, int32 XSize);
+	void RegenerateWaveforms(TRange<float> DrawRange, int32 XOffset, int32 XSize, const FColor& ColorTint, float DisplayScale);
 
 private:
 
@@ -99,7 +98,10 @@ private:
 	TRange<float> StoredDrawRange;
 	int32 StoredXOffset;
 	int32 StoredXSize;
+	FColor StoredColor;
 
 	/** Whether this section is on a master audio track or an attached audio track. */
 	bool bIsOnAMasterTrack;
+
+	TWeakPtr<ISequencer> Sequencer;
 };
diff --git a/Engine/Source/Editor/MovieSceneTools/Private/TrackEditors/CameraCutTrackEditor.cpp b/Engine/Source/Editor/MovieSceneTools/Private/TrackEditors/CameraCutTrackEditor.cpp
--- a/Engine/Source/Editor/MovieSceneTools/Private/TrackEditors/CameraCutTrackEditor.cpp
+++ b/Engine/Source/Editor/MovieSceneTools/Private/TrackEditors/CameraCutTrackEditor.cpp
@@ -13,7 +13,6 @@
 #include "MovieSceneTrackEditor.h"
 #include "CameraCutTrackEditor.h"
 #include "CommonMovieSceneTools.h"
-#include "Camera/CameraActor.h"
 #include "AssetToolsModule.h"
 #include "CameraCutSection.h"
 #include "SequencerUtilities.h"
@@ -118,7 +117,12 @@ void FCameraCutTrackEditor::Tick(float DeltaTime)
 {
 	if (FSlateThrottleManager::Get().IsAllowingExpensiveTasks())
 	{
-		ThumbnailPool->DrawThumbnails();
+		float SavedTime = GetSequencer()->GetGlobalTime();
+
+		if (ThumbnailPool->DrawThumbnails())
+		{
+			GetSequencer()->SetGlobalTime(SavedTime);
+		}
 	}
 }
 
@@ -207,10 +211,13 @@ bool FCameraCutTrackEditor::IsCameraPickable(const AActor* const PickableActor)
 	if (PickableActor->IsListedInSceneOutliner() &&
 		!FActorEditorUtils::IsABuilderBrush(PickableActor) &&
 		!PickableActor->IsA( AWorldSettings::StaticClass() ) &&
-		!PickableActor->IsPendingKill() &&
-		PickableActor->IsA( ACameraActor::StaticClass()) )
-	{			
-		return true;
+		!PickableActor->IsPendingKill())
+	{	
+		UCameraComponent* CameraComponent = MovieSceneHelpers::CameraComponentFromActor(PickableActor);
+		if (CameraComponent)	
+		{
+			return true;
+		}
 	}
 	return false;
 }
diff --git a/Engine/Source/Editor/MovieSceneTools/Private/TrackEditors/CinematicShotTrackEditor.cpp b/Engine/Source/Editor/MovieSceneTools/Private/TrackEditors/CinematicShotTrackEditor.cpp
--- a/Engine/Source/Editor/MovieSceneTools/Private/TrackEditors/CinematicShotTrackEditor.cpp
+++ b/Engine/Source/Editor/MovieSceneTools/Private/TrackEditors/CinematicShotTrackEditor.cpp
@@ -14,7 +14,6 @@
 #include "MovieSceneToolHelpers.h"
 #include "CinematicShotTrackEditor.h"
 #include "CommonMovieSceneTools.h"
-#include "Camera/CameraActor.h"
 #include "CinematicShotSection.h"
 #include "SequencerUtilities.h"
 #include "ContentBrowserModule.h"
@@ -158,7 +157,12 @@ void FCinematicShotTrackEditor::Tick(float DeltaTime)
 {
 	if (FSlateThrottleManager::Get().IsAllowingExpensiveTasks())
 	{
-		ThumbnailPool->DrawThumbnails();
+		float SavedTime = GetSequencer()->GetGlobalTime();
+
+		if (ThumbnailPool->DrawThumbnails())
+		{
+			GetSequencer()->SetGlobalTime(SavedTime);
+		}
 	}
 }
 
@@ -630,7 +634,7 @@ bool FCinematicShotTrackEditor::CanAddSubSequence(const UMovieSceneSequence& Seq
 void FCinematicShotTrackEditor::OnUpdateCameraCut(UObject* CameraObject)
 {
 	// Keep track of the camera when it switches so that the thumbnail can be drawn with the correct camera
-	CinematicShotCamera = Cast<ACameraActor>(CameraObject);
+	CinematicShotCamera = Cast<AActor>(CameraObject);
 }
 
 bool FCinematicShotTrackEditor::HandleSequenceAdded(float KeyTime, UMovieSceneSequence* Sequence)
diff --git a/Engine/Source/Editor/MovieSceneTools/Private/TrackEditors/CinematicShotTrackEditor.h b/Engine/Source/Editor/MovieSceneTools/Private/TrackEditors/CinematicShotTrackEditor.h
--- a/Engine/Source/Editor/MovieSceneTools/Private/TrackEditors/CinematicShotTrackEditor.h
+++ b/Engine/Source/Editor/MovieSceneTools/Private/TrackEditors/CinematicShotTrackEditor.h
@@ -35,7 +35,7 @@ public:
 	 */
 	static TSharedRef<ISequencerTrackEditor> CreateTrackEditor(TSharedRef<ISequencer> OwningSequencer);
 
-	TWeakObjectPtr<ACameraActor> GetCinematicShotCamera() const { return CinematicShotCamera; }
+	TWeakObjectPtr<AActor> GetCinematicShotCamera() const { return CinematicShotCamera; }
 
 public:
 
@@ -158,7 +158,7 @@ private:
 	TSharedPtr<FTrackEditorThumbnailPool> ThumbnailPool;
 
 	/** The camera actor for the current cut. */
-	TWeakObjectPtr<ACameraActor> CinematicShotCamera;
+	TWeakObjectPtr<AActor> CinematicShotCamera;
 
 	/** Delegate binding handle for ISequencer::OnCameraCut */
 	FDelegateHandle OnCameraCutHandle;
diff --git a/Engine/Source/Editor/MovieSceneTools/Private/TrackEditors/LevelVisibilityTrackEditor.cpp b/Engine/Source/Editor/MovieSceneTools/Private/TrackEditors/LevelVisibilityTrackEditor.cpp
new file mode 100644
--- /dev/null
+++ b/Engine/Source/Editor/MovieSceneTools/Private/TrackEditors/LevelVisibilityTrackEditor.cpp
@@ -0,0 +1,142 @@
+// Copyright 1998-2015 Epic Games, Inc. All Rights Reserved.
+
+#include "MovieSceneToolsPrivatePCH.h"
+#include "LevelVisibilityTrackEditor.h"
+#include "LevelVisibilitySection.h"
+#include "MovieSceneLevelVisibilityTrack.h"
+#include "MovieSceneLevelVisibilitySection.h"
+#include "SequencerUtilities.h"
+
+#define LOCTEXT_NAMESPACE "LevelVisibilityTrackEditor.h"
+
+FLevelVisibilityTrackEditor::FLevelVisibilityTrackEditor( TSharedRef<ISequencer> InSequencer )
+	: FMovieSceneTrackEditor( InSequencer ) 
+{ }
+
+
+TSharedRef<ISequencerTrackEditor> FLevelVisibilityTrackEditor::CreateTrackEditor( TSharedRef<ISequencer> InSequencer )
+{
+	return MakeShareable( new FLevelVisibilityTrackEditor( InSequencer ) );
+}
+
+
+bool FLevelVisibilityTrackEditor::SupportsType( TSubclassOf<UMovieSceneTrack> Type ) const
+{
+	return Type == UMovieSceneLevelVisibilityTrack::StaticClass();
+}
+
+
+TSharedRef<ISequencerSection> FLevelVisibilityTrackEditor::MakeSectionInterface( UMovieSceneSection& SectionObject, UMovieSceneTrack& Track )
+{
+	UMovieSceneLevelVisibilitySection* LevelVisibilitySection = Cast<UMovieSceneLevelVisibilitySection>(&SectionObject);
+	check( SupportsType( SectionObject.GetOuter()->GetClass() ) && LevelVisibilitySection != nullptr );
+
+	return MakeShareable( new FLevelVisibilitySection( *LevelVisibilitySection ) );
+}
+
+void FLevelVisibilityTrackEditor::BuildAddTrackMenu( FMenuBuilder& MenuBuilder )
+{
+	UMovieSceneSequence* RootMovieSceneSequence = GetSequencer()->GetRootMovieSceneSequence();
+
+	if ( ( RootMovieSceneSequence == nullptr ) || ( RootMovieSceneSequence->GetClass()->GetName() != TEXT( "LevelSequence" ) ) )
+	{
+		return;
+	}
+
+	MenuBuilder.AddMenuEntry(
+		LOCTEXT("AddTrack", "Level Visibility Track" ),
+		LOCTEXT("AddAdTrackToolTip", "Adds a new track which can control level visibility." ),
+		FSlateIcon(FEditorStyle::GetStyleSetName(), "Sequencer.Tracks.LevelVisibility"),
+		FUIAction( FExecuteAction::CreateRaw( this, &FLevelVisibilityTrackEditor::OnAddTrack ) ) );
+}
+
+
+TSharedPtr<SWidget> FLevelVisibilityTrackEditor::BuildOutlinerEditWidget( const FGuid& ObjectBinding, UMovieSceneTrack* Track, const FBuildEditWidgetParams& Params )
+{
+	// Create a container edit box
+	return FSequencerUtilities::MakeAddButton( 
+		LOCTEXT( "AddVisibilityTrigger", "Visibility Trigger" ),
+		FOnGetContent::CreateSP( this, &FLevelVisibilityTrackEditor::BuildAddVisibilityTriggerMenu, Track ),
+		Params.NodeIsHovered );
+}
+
+
+void FLevelVisibilityTrackEditor::AddNewSection( UMovieScene* MovieScene, UMovieSceneTrack* LevelVisibilityTrack, ELevelVisibility Visibility )
+{
+	const FScopedTransaction Transaction( LOCTEXT( "AddLevelVisibilitySection_Transaction", "Add Level Visibility Trigger" ) );
+
+	LevelVisibilityTrack->Modify();
+
+	UMovieSceneLevelVisibilitySection* LevelVisibilitySection = CastChecked<UMovieSceneLevelVisibilitySection>( LevelVisibilityTrack->CreateNewSection() );
+	LevelVisibilitySection->SetVisibility( Visibility );
+	LevelVisibilitySection->SetStartTime( MovieScene->GetPlaybackRange().GetLowerBoundValue() );
+	LevelVisibilitySection->SetEndTime( MovieScene->GetPlaybackRange().GetUpperBoundValue() );
+
+	int32 RowIndex = 0;
+	for ( const UMovieSceneSection* Section : LevelVisibilityTrack->GetAllSections() )
+	{
+		RowIndex = FMath::Max( RowIndex, Section->GetRowIndex() );
+	}
+	LevelVisibilitySection->SetRowIndex(RowIndex + 1);
+
+	LevelVisibilityTrack->AddSection( *LevelVisibilitySection );
+}
+
+
+void FLevelVisibilityTrackEditor::OnAddTrack()
+{
+	UMovieScene* FocusedMovieScene = GetFocusedMovieScene();
+
+	if ( FocusedMovieScene == nullptr )
+	{
+		return;
+	}
+
+	const FScopedTransaction Transaction( LOCTEXT( "AddLevelVisibilityTrack_Transaction", "Add Level Visibility Track" ) );
+	FocusedMovieScene->Modify();
+
+	UMovieSceneLevelVisibilityTrack* NewTrack = FocusedMovieScene->AddMasterTrack<UMovieSceneLevelVisibilityTrack>();
+	checkf( NewTrack != nullptr, TEXT("Failed to create new level visibility track.") );
+
+	AddNewSection( FocusedMovieScene, NewTrack, ELevelVisibility::Visible );
+
+	GetSequencer()->NotifyMovieSceneDataChanged();
+}
+
+
+TSharedRef<SWidget> FLevelVisibilityTrackEditor::BuildAddVisibilityTriggerMenu( UMovieSceneTrack* LevelVisibilityTrack )
+{
+	FMenuBuilder MenuBuilder(true, nullptr);
+
+	MenuBuilder.AddMenuEntry(
+		LOCTEXT( "AddVisibleTrigger", "Visible" ),
+		LOCTEXT( "AddVisibleTriggerToolTip", "Add a trigger for visible levels." ),
+		FSlateIcon(),
+		FUIAction( FExecuteAction::CreateSP(
+			this, &FLevelVisibilityTrackEditor::OnAddNewSection, LevelVisibilityTrack, ELevelVisibility::Visible ) ) );
+
+	MenuBuilder.AddMenuEntry(
+		LOCTEXT( "AddHiddenTrigger", "Hidden" ),
+		LOCTEXT( "AddHiddenTriggerToolTip", "Add a trigger for hidden levels." ),
+		FSlateIcon(),
+		FUIAction( FExecuteAction::CreateSP(
+		this, &FLevelVisibilityTrackEditor::OnAddNewSection, LevelVisibilityTrack, ELevelVisibility::Hidden ) ) );
+
+	return MenuBuilder.MakeWidget();
+}
+
+
+void FLevelVisibilityTrackEditor::OnAddNewSection( UMovieSceneTrack* LevelVisibilityTrack, ELevelVisibility Visibility )
+{
+	UMovieScene* FocusedMovieScene = GetFocusedMovieScene();
+
+	if ( FocusedMovieScene == nullptr )
+	{
+		return;
+	}
+
+	AddNewSection( FocusedMovieScene, LevelVisibilityTrack, Visibility );
+	GetSequencer()->NotifyMovieSceneDataChanged();
+}
+
+#undef LOCTEXT_NAMESPACE
diff --git a/Engine/Source/Editor/MovieSceneTools/Private/TrackEditors/ParticleTrackEditor.cpp b/Engine/Source/Editor/MovieSceneTools/Private/TrackEditors/ParticleTrackEditor.cpp
--- a/Engine/Source/Editor/MovieSceneTools/Private/TrackEditors/ParticleTrackEditor.cpp
+++ b/Engine/Source/Editor/MovieSceneTools/Private/TrackEditors/ParticleTrackEditor.cpp
@@ -75,11 +75,13 @@ void FParticleSection::GenerateSectionLayout( class ISectionLayoutBuilder& Layou
 	LayoutBuilder.SetSectionAsKeyArea( MakeShareable( new FEnumKeyArea( ParticleSection->GetParticleCurve(), ParticleSection, ParticleKeyEnum ) ) );
 }
 
-int32 FParticleSection::OnPaintSection( const FGeometry& AllottedGeometry, const FSlateRect& SectionClippingRect, FSlateWindowElementList& OutDrawElements, int32 LayerId, bool bParentEnabled ) const
+int32 FParticleSection::OnPaintSection( FSequencerSectionPainter& InPainter ) const
 {
-	const ESlateDrawEffect::Type DrawEffects = bParentEnabled ? ESlateDrawEffect::None : ESlateDrawEffect::DisabledEffect;
+	const ESlateDrawEffect::Type DrawEffects = InPainter.bParentEnabled ? ESlateDrawEffect::None : ESlateDrawEffect::DisabledEffect;
 	UMovieSceneParticleSection* AnimSection = Cast<UMovieSceneParticleSection>( &Section );
-	FTimeToPixel TimeToPixelConverter( AllottedGeometry, OwningSequencer->GetViewRange() );
+	const FTimeToPixel& TimeToPixelConverter = InPainter.GetTimeConverter();
+
+	FLinearColor TrackColor;
 
 	// @todo Sequencer - These values should be cached and then refreshed only when the particle system changes.
 	bool bIsLooping = false;
@@ -90,6 +92,9 @@ int32 FParticleSection::OnPaintSection( const FGeometry& AllottedGeometry, const
 		UMovieSceneParticleTrack* ParentTrack = Cast<UMovieSceneParticleTrack>( ParticleSection->GetOuter() );
 		if ( ParentTrack != nullptr )
 		{
+			TrackColor = ParentTrack->GetColorTint();
+			TrackColor.A *= 0.4;
+
 			FGuid ObjectHandle;
 			for ( const FMovieSceneBinding& Binding : OwningSequencer->GetFocusedMovieSceneSequence()->GetMovieScene()->GetBindings() )
 			{
@@ -102,18 +107,20 @@ int32 FParticleSection::OnPaintSection( const FGeometry& AllottedGeometry, const
 
 			if ( ObjectHandle.IsValid() )
 			{
-				AEmitter* ParticleSystemActor = Cast<AEmitter>( OwningSequencer->GetFocusedMovieSceneSequenceInstance()->FindObject( ObjectHandle, *OwningSequencer ) );
-				if ( ParticleSystemActor != nullptr )
+				UObject* BoundObject = OwningSequencer->GetFocusedMovieSceneSequenceInstance()->FindObject( ObjectHandle, *OwningSequencer );
+				UParticleSystemComponent* ParticleSystemComponent = Cast<UParticleSystemComponent>(BoundObject);
+				if(AEmitter* ParticleSystemActor = Cast<AEmitter>(BoundObject))
 				{
-					UParticleSystemComponent* ParticleSystemComponent = ParticleSystemActor->GetParticleSystemComponent();
-					if ( ParticleSystemComponent != nullptr && ParticleSystemComponent->Template != nullptr )
+					ParticleSystemComponent = ParticleSystemActor->GetParticleSystemComponent();
+				}
+
+				if ( ParticleSystemComponent != nullptr && ParticleSystemComponent->Template != nullptr )
+				{
+					for ( UParticleEmitter* Emitter : ParticleSystemComponent->Template->Emitters )
 					{
-						for ( UParticleEmitter* Emitter : ParticleSystemComponent->Template->Emitters )
-						{
-							UParticleModuleRequired* RequiredModule = Emitter->GetLODLevel( 0 )->RequiredModule;
-							bIsLooping |= RequiredModule->EmitterLoops == 0;
-							LastEmitterEndTime = FMath::Max( LastEmitterEndTime, RequiredModule->EmitterDelay + RequiredModule->EmitterDuration );
-						}
+						UParticleModuleRequired* RequiredModule = Emitter->GetLODLevel( 0 )->RequiredModule;
+						bIsLooping |= RequiredModule->EmitterLoops == 0;
+						LastEmitterEndTime = FMath::Max( LastEmitterEndTime, RequiredModule->EmitterDelay + RequiredModule->EmitterDuration );
 					}
 				}
 			}
@@ -188,17 +195,25 @@ int32 FParticleSection::OnPaintSection( const FGeometry& AllottedGeometry, const
 		float XOffset = TimeToPixelConverter.TimeToPixel(DrawRange.GetLowerBoundValue());
 		float XSize = TimeToPixelConverter.TimeToPixel(DrawRange.GetUpperBoundValue()) - XOffset;
 		FSlateDrawElement::MakeBox(
-			OutDrawElements,
-			LayerId,
-			AllottedGeometry.ToPaintGeometry( FVector2D( XOffset, (AllottedGeometry.GetLocalSize().Y - SequencerSectionConstants::KeySize.Y) / 2 ), FVector2D( XSize, SequencerSectionConstants::KeySize.Y ) ),
-			FEditorStyle::GetBrush( "Sequencer.GenericSection.Background" ),
-			SectionClippingRect,
+			InPainter.DrawElements,
+			InPainter.LayerId,
+			InPainter.SectionGeometry.ToPaintGeometry( FVector2D( XOffset, (InPainter.SectionGeometry.GetLocalSize().Y - SequencerSectionConstants::KeySize.Y) / 2 ), FVector2D( XSize, SequencerSectionConstants::KeySize.Y ) ),
+			FEditorStyle::GetBrush( "Sequencer.Section.Background" ),
+			InPainter.SectionClippingRect,
+			DrawEffects
+			);
+		FSlateDrawElement::MakeBox(
+			InPainter.DrawElements,
+			InPainter.LayerId,
+			InPainter.SectionGeometry.ToPaintGeometry( FVector2D( XOffset, (InPainter.SectionGeometry.GetLocalSize().Y - SequencerSectionConstants::KeySize.Y) / 2 ), FVector2D( XSize, SequencerSectionConstants::KeySize.Y ) ),
+			FEditorStyle::GetBrush( "Sequencer.Section.BackgroundTint" ),
+			InPainter.SectionClippingRect,
 			DrawEffects,
-			FLinearColor( 0.0f, 0.0f, 1.f, 1.f )
+			TrackColor
 			);
 	}
 
-	return LayerId+1;
+	return InPainter.LayerId+1;
 }
 
 
@@ -249,7 +264,7 @@ TSharedRef<ISequencerSection> FParticleTrackEditor::MakeSectionInterface( UMovie
 
 void FParticleTrackEditor::BuildObjectBindingTrackMenu(FMenuBuilder& MenuBuilder, const FGuid& ObjectBinding, const UClass* ObjectClass)
 {
-	if (ObjectClass->IsChildOf(AEmitter::StaticClass()))
+	if (ObjectClass->IsChildOf(AEmitter::StaticClass()) || ObjectClass->IsChildOf(UParticleSystemComponent::StaticClass()))
 	{
 		const TSharedPtr<ISequencer> ParentSequencer = GetSequencer();
 
diff --git a/Engine/Source/Editor/MovieSceneTools/Private/TrackEditors/ParticleTrackEditor.h b/Engine/Source/Editor/MovieSceneTools/Private/TrackEditors/ParticleTrackEditor.h
--- a/Engine/Source/Editor/MovieSceneTools/Private/TrackEditors/ParticleTrackEditor.h
+++ b/Engine/Source/Editor/MovieSceneTools/Private/TrackEditors/ParticleTrackEditor.h
@@ -70,7 +70,7 @@ public:
 	virtual FText GetSectionTitle() const override { return FText::GetEmpty(); }
 	virtual float GetSectionHeight() const override;
 	virtual void GenerateSectionLayout( class ISectionLayoutBuilder& LayoutBuilder ) const override;
-	virtual int32 OnPaintSection( const FGeometry& AllottedGeometry, const FSlateRect& SectionClippingRect, FSlateWindowElementList& OutDrawElements, int32 LayerId, bool bParentEnabled ) const override;
+	virtual int32 OnPaintSection( FSequencerSectionPainter& InPainter ) const override;
 	virtual const FSlateBrush* GetKeyBrush( FKeyHandle KeyHandle ) const override;
 	virtual bool SectionIsResizable() const override { return false; }
 
diff --git a/Engine/Source/Editor/MovieSceneTools/Private/TrackEditors/PathTrackEditor.cpp b/Engine/Source/Editor/MovieSceneTools/Private/TrackEditors/PathTrackEditor.cpp
--- a/Engine/Source/Editor/MovieSceneTools/Private/TrackEditors/PathTrackEditor.cpp
+++ b/Engine/Source/Editor/MovieSceneTools/Private/TrackEditors/PathTrackEditor.cpp
@@ -51,18 +51,9 @@ public:
 		LayoutBuilder.AddKeyArea("Timing", LOCTEXT("TimingArea", "Timing"), MakeShareable( new FFloatCurveKeyArea ( &PathSection->GetTimingCurve( ), PathSection ) ) );
 	}
 
-	virtual int32 OnPaintSection( const FGeometry& AllottedGeometry, const FSlateRect& SectionClippingRect, FSlateWindowElementList& OutDrawElements, int32 LayerId, bool bParentEnabled ) const override 
+	virtual int32 OnPaintSection( FSequencerSectionPainter& InPainter ) const override 
 	{
-		// Add a box for the section
-		FSlateDrawElement::MakeBox( 
-			OutDrawElements,
-			LayerId,
-			AllottedGeometry.ToPaintGeometry(),
-			FEditorStyle::GetBrush("Sequencer.GenericSection.Background"),
-			SectionClippingRect
-		); 
-
-		return LayerId;
+		return InPainter.PaintSectionBackground();
 	}
 	
 	virtual void BuildSectionContextMenu(FMenuBuilder& MenuBuilder, const FGuid& ObjectBinding) override
@@ -164,7 +155,7 @@ bool F3DPathTrackEditor::IsActorPickable(const AActor* const ParentActor, FGuid
 }
 
 
-void F3DPathTrackEditor::ActorSocketPicked(const FName SocketName, AActor* ParentActor, FGuid ObjectGuid, UMovieSceneSection* Section)
+void F3DPathTrackEditor::ActorSocketPicked(const FName SocketName, USceneComponent* Component, AActor* ParentActor, FGuid ObjectGuid, UMovieSceneSection* Section)
 {
 	if (ObjectGuid.IsValid())
 	{
@@ -240,7 +231,7 @@ bool F3DPathTrackEditor::AddKeyInternal( float KeyTime, const TArray<UObject*> O
 					}
 				}
 
-				Cast<UMovieScene3DPathTrack>(Track)->AddConstraint( KeyTime, PathEndTime, NAME_None, SplineId );
+				Cast<UMovieScene3DPathTrack>(Track)->AddConstraint( KeyTime, PathEndTime, NAME_None, NAME_None, SplineId );
 				bTrackModified = true;
 			}
 		}
diff --git a/Engine/Source/Editor/MovieSceneTools/Private/TrackEditors/PathTrackEditor.h b/Engine/Source/Editor/MovieSceneTools/Private/TrackEditors/PathTrackEditor.h
--- a/Engine/Source/Editor/MovieSceneTools/Private/TrackEditors/PathTrackEditor.h
+++ b/Engine/Source/Editor/MovieSceneTools/Private/TrackEditors/PathTrackEditor.h
@@ -43,7 +43,7 @@ public:
 
 	// FTrackEditorActorPicker
 	virtual bool IsActorPickable( const AActor* const ParentActor, FGuid ObjectBinding, UMovieSceneSection* InSection ) override;
-	virtual void ActorSocketPicked(const FName SocketName, AActor* ParentActor, FGuid ObjectBinding, UMovieSceneSection* Section) override;
+	virtual void ActorSocketPicked(const FName SocketName, USceneComponent* Component, AActor* ParentActor, FGuid ObjectBinding, UMovieSceneSection* Section) override;
 
 private:
 
diff --git a/Engine/Source/Editor/MovieSceneTools/Private/TrackEditors/SkeletalAnimationTrackEditor.cpp b/Engine/Source/Editor/MovieSceneTools/Private/TrackEditors/SkeletalAnimationTrackEditor.cpp
--- a/Engine/Source/Editor/MovieSceneTools/Private/TrackEditors/SkeletalAnimationTrackEditor.cpp
+++ b/Engine/Source/Editor/MovieSceneTools/Private/TrackEditors/SkeletalAnimationTrackEditor.cpp
@@ -44,11 +44,6 @@ UMovieSceneSection* FSkeletalAnimationSection::GetSectionObject()
 }
 
 
-bool FSkeletalAnimationSection::ShouldDrawKeyAreaBackground() const
-{
-	return false;
-}
-
 FText FSkeletalAnimationSection::GetDisplayName() const
 {
 	return LOCTEXT("AnimationSection", "Animation");
@@ -72,24 +67,15 @@ float FSkeletalAnimationSection::GetSectionHeight() const
 }
 
 
-int32 FSkeletalAnimationSection::OnPaintSection( const FGeometry& AllottedGeometry, const FSlateRect& SectionClippingRect, FSlateWindowElementList& OutDrawElements, int32 LayerId, bool bParentEnabled ) const
+int32 FSkeletalAnimationSection::OnPaintSection( FSequencerSectionPainter& Painter ) const
 {
-	const ESlateDrawEffect::Type DrawEffects = bParentEnabled ? ESlateDrawEffect::None : ESlateDrawEffect::DisabledEffect;
+	const ESlateDrawEffect::Type DrawEffects = Painter.bParentEnabled ? ESlateDrawEffect::None : ESlateDrawEffect::DisabledEffect;
 
 	UMovieSceneSkeletalAnimationSection* AnimSection = Cast<UMovieSceneSkeletalAnimationSection>(&Section);
 	
-	FTimeToPixel TimeToPixelConverter( AllottedGeometry, TRange<float>( Section.GetStartTime(), Section.GetEndTime() ) );
-
-	// Add a box for the section
-	FSlateDrawElement::MakeBox(
-		OutDrawElements,
-		LayerId,
-		AllottedGeometry.ToPaintGeometry(),
-		FEditorStyle::GetBrush("Sequencer.GenericSection.Background"),
-		SectionClippingRect,
-		DrawEffects,
-		FLinearColor(0.7f, 0.4f, 0.7f, 1.f)
-	);
+	const FTimeToPixel& TimeToPixelConverter = Painter.GetTimeConverter();
+
+	int32 LayerId = Painter.PaintSectionBackground();
 
 	// Add lines where the animation starts and ends/loops
 	float CurrentTime = AnimSection->GetStartTime();
@@ -103,21 +89,21 @@ int32 FSkeletalAnimationSection::OnPaintSection( const FGeometry& AllottedGeomet
 
 			TArray<FVector2D> Points;
 			Points.Add(FVector2D(CurrentPixels, 0));
-			Points.Add(FVector2D(CurrentPixels, AllottedGeometry.Size.Y));
+			Points.Add(FVector2D(CurrentPixels, Painter.SectionGeometry.Size.Y));
 
 			FSlateDrawElement::MakeLines(
-				OutDrawElements,
-				LayerId + 2,
-				AllottedGeometry.ToPaintGeometry(),
+				Painter.DrawElements,
+				++LayerId,
+				Painter.SectionGeometry.ToPaintGeometry(),
 				Points,
-				SectionClippingRect,
+				Painter.SectionClippingRect,
 				DrawEffects
 			);
 		}
 		CurrentTime += SeqLength;
 	}
 
-	return LayerId+3;
+	return LayerId;
 }
 
 
@@ -207,7 +193,7 @@ bool FSkeletalAnimationTrackEditor::HandleAssetAdded(UObject* Asset, const FGuid
 
 void FSkeletalAnimationTrackEditor::BuildObjectBindingTrackMenu(FMenuBuilder& MenuBuilder, const FGuid& ObjectBinding, const UClass* ObjectClass)
 {
-	if (ObjectClass->IsChildOf(ASkeletalMeshActor::StaticClass()))
+	if (ObjectClass->IsChildOf(USkeletalMeshComponent::StaticClass()) || ObjectClass->IsChildOf(AActor::StaticClass()))
 	{
 		const TSharedPtr<ISequencer> ParentSequencer = GetSequencer();
 
@@ -315,9 +301,7 @@ USkeleton* FSkeletalAnimationTrackEditor::AcquireSkeletonFromObjectGuid(const FG
 	USkeleton* Skeleton = NULL;
 	for (int32 i = 0; i < OutObjects.Num(); ++i)
 	{
-		AActor* Actor = Cast<AActor>(OutObjects[i]);
-
-		if (Actor != NULL)
+		if (AActor* Actor = Cast<AActor>(OutObjects[i]))
 		{
 			TInlineComponentArray<USkeletalMeshComponent*> SkeletalMeshComponents;
 			Actor->GetComponents(SkeletalMeshComponents);
@@ -333,6 +317,14 @@ USkeleton* FSkeletalAnimationTrackEditor::AcquireSkeletonFromObjectGuid(const FG
 				}
 			}
 		}
+		else if(USkeletalMeshComponent* SkeletalMeshComponent = Cast<USkeletalMeshComponent>(OutObjects[i]))
+		{
+			if (SkeletalMeshComponent->SkeletalMesh && SkeletalMeshComponent->SkeletalMesh->Skeleton)
+			{
+				check(!Skeleton);
+				Skeleton = SkeletalMeshComponent->SkeletalMesh->Skeleton;
+			}
+		}
 	}
 
 	return Skeleton;
diff --git a/Engine/Source/Editor/MovieSceneTools/Private/TrackEditors/SkeletalAnimationTrackEditor.h b/Engine/Source/Editor/MovieSceneTools/Private/TrackEditors/SkeletalAnimationTrackEditor.h
--- a/Engine/Source/Editor/MovieSceneTools/Private/TrackEditors/SkeletalAnimationTrackEditor.h
+++ b/Engine/Source/Editor/MovieSceneTools/Private/TrackEditors/SkeletalAnimationTrackEditor.h
@@ -72,12 +72,11 @@ public:
 	// ISequencerSection interface
 
 	virtual UMovieSceneSection* GetSectionObject() override;
-	virtual bool ShouldDrawKeyAreaBackground() const override;
 	virtual FText GetDisplayName() const override;
 	virtual FText GetSectionTitle() const override;
 	virtual float GetSectionHeight() const override;
 	virtual void GenerateSectionLayout( class ISectionLayoutBuilder& LayoutBuilder ) const override {}
-	virtual int32 OnPaintSection( const FGeometry& AllottedGeometry, const FSlateRect& SectionClippingRect, FSlateWindowElementList& OutDrawElements, int32 LayerId, bool bParentEnabled ) const override;
+	virtual int32 OnPaintSection( FSequencerSectionPainter& Painter ) const override;
 
 private:
 
diff --git a/Engine/Source/Editor/MovieSceneTools/Private/TrackEditors/SubTrackEditor.cpp b/Engine/Source/Editor/MovieSceneTools/Private/TrackEditors/SubTrackEditor.cpp
--- a/Engine/Source/Editor/MovieSceneTools/Private/TrackEditors/SubTrackEditor.cpp
+++ b/Engine/Source/Editor/MovieSceneTools/Private/TrackEditors/SubTrackEditor.cpp
@@ -8,6 +8,7 @@
 #include "SubTrackEditor.h"
 #include "ContentBrowserModule.h"
 #include "SequencerUtilities.h"
+#include "SequencerSectionPainter.h"
 
 
 namespace SubTrackEditorConstants
@@ -44,11 +45,6 @@ public:
 		// do nothing
 	}
 
-	virtual bool ShouldDrawKeyAreaBackground() const override
-	{
-		return false;
-	}
-
 	virtual FText GetDisplayName() const override
 	{
 		return GetSectionTitle();
@@ -68,9 +64,11 @@ public:
 	{
 		return FText::FromString(SectionObject.GetSequence()->GetName());
 	}
-
-	virtual int32 OnPaintSection(const FGeometry& AllottedGeometry, const FSlateRect& SectionClippingRect, FSlateWindowElementList& OutDrawElements, int32 LayerId, bool bParentEnabled) const override 
+	
+	virtual int32 OnPaintSection( FSequencerSectionPainter& InPainter ) const override
 	{
+		int32 LayerId = InPainter.PaintSectionBackground();
+
 		const float SectionSize = SectionObject.GetTimeSize();
 
 		if (SectionSize <= 0.0f)
@@ -78,8 +76,8 @@ public:
 			return LayerId;
 		}
 
-		const float DrawScale = AllottedGeometry.Size.X / SectionSize;
-		const ESlateDrawEffect::Type DrawEffects = bParentEnabled
+		const float DrawScale = InPainter.SectionGeometry.Size.X / SectionSize;
+		const ESlateDrawEffect::Type DrawEffects = InPainter.bParentEnabled
 			? ESlateDrawEffect::None
 			: ESlateDrawEffect::DisabledEffect;
 
@@ -92,14 +90,14 @@ public:
 		const float WorkingSize = SectionObject.TimeScale * MovieScene->GetEditorData().WorkingRange.Size<float>();
 
 		FSlateDrawElement::MakeBox(
-			OutDrawElements,
+			InPainter.DrawElements,
 			LayerId,
-			AllottedGeometry.ToPaintGeometry(
+			InPainter.SectionGeometry.ToPaintGeometry(
 				FVector2D(WorkingStart * DrawScale, 0.f),
-				FVector2D(WorkingSize * DrawScale, AllottedGeometry.Size.Y)
+				FVector2D(WorkingSize * DrawScale, InPainter.SectionGeometry.Size.Y)
 			),
-			FEditorStyle::GetBrush("Sequencer.GenericSection.Background"),
-			SectionClippingRect,
+			FEditorStyle::GetBrush("Sequencer.Section.SelectedTrackTint"),
+			InPainter.SectionClippingRect,
 			DrawEffects,
 			FColor(220, 120, 120)
 		);
@@ -108,14 +106,14 @@ public:
 		if (StartOffset < 0.0f)
 		{
 			FSlateDrawElement::MakeBox(
-				OutDrawElements,
-				LayerId + 1,
-				AllottedGeometry.ToPaintGeometry(
+				InPainter.DrawElements,
+				++LayerId,
+				InPainter.SectionGeometry.ToPaintGeometry(
 					FVector2D(0.0f, 0.f),
-					FVector2D(-StartOffset * DrawScale, AllottedGeometry.Size.Y)
+					FVector2D(-StartOffset * DrawScale, InPainter.SectionGeometry.Size.Y)
 				),
 				FEditorStyle::GetBrush("WhiteBrush"),
-				SectionClippingRect,
+				InPainter.SectionClippingRect,
 				ESlateDrawEffect::None,
 				FLinearColor::Black.CopyWithNewOpacity(0.2f)
 			);
@@ -125,14 +123,14 @@ public:
 		if (StartOffset <= 0)
 		{
 			FSlateDrawElement::MakeBox(
-				OutDrawElements,
-				LayerId + 2,
-				AllottedGeometry.ToPaintGeometry(
+				InPainter.DrawElements,
+				++LayerId,
+				InPainter.SectionGeometry.ToPaintGeometry(
 					FVector2D(-StartOffset * DrawScale, 0.f),
-					FVector2D(1.0f, AllottedGeometry.Size.Y)
+					FVector2D(1.0f, InPainter.SectionGeometry.Size.Y)
 				),
 				FEditorStyle::GetBrush("WhiteBrush"),
-				SectionClippingRect,
+				InPainter.SectionClippingRect,
 				ESlateDrawEffect::None,
 				FColor(32, 128, 32)	// 120, 75, 50 (HSV)
 			);
@@ -144,14 +142,14 @@ public:
 		if (PlaybackEnd < SectionSize)
 		{
 			FSlateDrawElement::MakeBox(
-				OutDrawElements,
-				LayerId + 1,
-				AllottedGeometry.ToPaintGeometry(
+				InPainter.DrawElements,
+				++LayerId,
+				InPainter.SectionGeometry.ToPaintGeometry(
 					FVector2D(PlaybackEnd * DrawScale, 0.f),
-					FVector2D((SectionSize - PlaybackEnd) * DrawScale, AllottedGeometry.Size.Y)
+					FVector2D((SectionSize - PlaybackEnd) * DrawScale, InPainter.SectionGeometry.Size.Y)
 				),
 				FEditorStyle::GetBrush("WhiteBrush"),
-				SectionClippingRect,
+				InPainter.SectionClippingRect,
 				ESlateDrawEffect::None,
 				FLinearColor::Black.CopyWithNewOpacity(0.2f)
 			);
@@ -161,14 +159,14 @@ public:
 		if (PlaybackEnd <= SectionSize)
 		{
 			FSlateDrawElement::MakeBox(
-				OutDrawElements,
-				LayerId + 2,
-				AllottedGeometry.ToPaintGeometry(
+				InPainter.DrawElements,
+				++LayerId,
+				InPainter.SectionGeometry.ToPaintGeometry(
 					FVector2D(PlaybackEnd * DrawScale, 0.f),
-					FVector2D(1.0f, AllottedGeometry.Size.Y)
+					FVector2D(1.0f, InPainter.SectionGeometry.Size.Y)
 				),
 				FEditorStyle::GetBrush("WhiteBrush"),
-				SectionClippingRect,
+				InPainter.SectionClippingRect,
 				ESlateDrawEffect::None,
 				FColor(128, 32, 32)	// 0, 75, 50 (HSV)
 			);
@@ -178,14 +176,14 @@ public:
 		int32 NumTracks = MovieScene->GetPossessableCount() + MovieScene->GetSpawnableCount() + MovieScene->GetMasterTracks().Num();
 
 		FSlateDrawElement::MakeText(
-			OutDrawElements,
-			LayerId + 3,
-			AllottedGeometry.ToOffsetPaintGeometry(FVector2D(5.0f, 32.0f)),
+			InPainter.DrawElements,
+			++LayerId,
+			InPainter.SectionGeometry.ToOffsetPaintGeometry(FVector2D(5.0f, 32.0f)),
 			FText::Format(LOCTEXT("NumTracksFormat", "{0} track(s)"), FText::AsNumber(NumTracks)),
 			FEditorStyle::GetFontStyle("NormalFont"),
-			SectionClippingRect,
+			InPainter.SectionClippingRect,
 			DrawEffects,
-			FLinearColor::White
+			FColor(200, 200, 200)
 		);
 
 		return LayerId;
@@ -195,7 +193,7 @@ public:
 	{
 		if( MouseEvent.GetEffectingButton() == EKeys::LeftMouseButton )
 		{
-			Sequencer.Pin()->FocusSequenceInstance(SequenceInstance.Pin().ToSharedRef());
+			Sequencer.Pin()->FocusSequenceInstance(SectionObject);
 		}
 
 		return FReply::Handled();
diff --git a/Engine/Source/Editor/MovieSceneTools/Private/TrackEditors/TransformTrackEditor.cpp b/Engine/Source/Editor/MovieSceneTools/Private/TrackEditors/TransformTrackEditor.cpp
--- a/Engine/Source/Editor/MovieSceneTools/Private/TrackEditors/TransformTrackEditor.cpp
+++ b/Engine/Source/Editor/MovieSceneTools/Private/TrackEditors/TransformTrackEditor.cpp
@@ -54,9 +54,9 @@ public:
 
 	virtual void GenerateSectionLayout( class ISectionLayoutBuilder& LayoutBuilder ) const override
 	{
-		static const FLinearColor BlueKeyAreaColor(0.0f, 0.0f, 0.7f, 0.2f);
-		static const FLinearColor GreenKeyAreaColor(0.0f, 0.7f, 0.0f, 0.2f);
-		static const FLinearColor RedKeyAreaColor(0.7f, 0.0f, 0.0f, 0.2f);
+		static const FLinearColor BlueKeyAreaColor(0.0f, 0.0f, 0.7f, 0.5f);
+		static const FLinearColor GreenKeyAreaColor(0.0f, 0.7f, 0.0f, 0.5f);
+		static const FLinearColor RedKeyAreaColor(0.7f, 0.0f, 0.0f, 0.5f);
 
 		UMovieScene3DTransformSection* TransformSection = Cast<UMovieScene3DTransformSection>( &Section );
 
@@ -92,21 +92,9 @@ public:
 		LayoutBuilder.PopCategory();
 	}
 
-	virtual int32 OnPaintSection( const FGeometry& AllottedGeometry, const FSlateRect& SectionClippingRect, FSlateWindowElementList& OutDrawElements, int32 LayerId, bool bParentEnabled ) const override 
+	virtual int32 OnPaintSection( FSequencerSectionPainter& InPainter ) const override 
 	{
-		const ESlateDrawEffect::Type DrawEffects = bParentEnabled ? ESlateDrawEffect::None : ESlateDrawEffect::DisabledEffect;
-	
-		// Add a box for the section
-		FSlateDrawElement::MakeBox( 
-			OutDrawElements,
-			LayerId,
-			AllottedGeometry.ToPaintGeometry(),
-			FEditorStyle::GetBrush("Sequencer.GenericSection.Background"),
-			SectionClippingRect,
-			DrawEffects
-		); 
-
-		return LayerId;
+		return InPainter.PaintSectionBackground();
 	}
 
 	void SetIntermediateValue(UMovieSceneTrack* Track, FTransformData IntermediateData )
@@ -423,38 +411,54 @@ void F3DTransformTrackEditor::BindCommands(TSharedRef<FUICommandList> SequencerC
 
 void F3DTransformTrackEditor::BuildObjectBindingEditButtons(TSharedPtr<SHorizontalBox> EditBox, const FGuid& ObjectGuid, const UClass* ObjectClass)
 {
-	if (!ObjectClass->IsChildOf<ACameraActor>() && !ObjectClass->IsChildOf<UCameraComponent>())
+	bool bHasCameraComponent = false;
+
+	TArray<UObject*> OutObjects;
+	GetSequencer()->GetRuntimeObjects( GetSequencer()->GetFocusedMovieSceneSequenceInstance(), ObjectGuid, OutObjects);
+
+	for (UObject* Object : OutObjects)
 	{
-		return;
+		AActor* Actor = Cast<AActor>( Object );
+		if (Actor)
+		{
+			UCameraComponent* CameraComponent = MovieSceneHelpers::CameraComponentFromActor(Actor);
+			if (CameraComponent)
+			{
+				bHasCameraComponent = true;
+			}
+		}
 	}
 
-	// If this is a camera track, add a button to lock the viewport to the camera
-	EditBox.Get()->AddSlot()
-		.VAlign(VAlign_Center)
-		.HAlign(HAlign_Right)
-		.AutoWidth()
-		.Padding(4, 0, 0, 0)
-		[
-			SNew(SCheckBox)		
-				.IsFocusable(false)
-				.Visibility(this, &F3DTransformTrackEditor::IsCameraVisible, ObjectGuid)
-				.IsChecked(this, &F3DTransformTrackEditor::IsCameraLocked, ObjectGuid)
-				.OnCheckStateChanged(this, &F3DTransformTrackEditor::OnLockCameraClicked, ObjectGuid)
-				.ToolTipText(this, &F3DTransformTrackEditor::GetLockCameraToolTip, ObjectGuid)
-				.ForegroundColor(FLinearColor::White)
-				.CheckedImage(FEditorStyle::GetBrush("Sequencer.LockCamera"))
-				.CheckedHoveredImage(FEditorStyle::GetBrush("Sequencer.LockCamera"))
-				.CheckedPressedImage(FEditorStyle::GetBrush("Sequencer.LockCamera"))
-				.UncheckedImage(FEditorStyle::GetBrush("Sequencer.UnlockCamera"))
-				.UncheckedHoveredImage(FEditorStyle::GetBrush("Sequencer.UnlockCamera"))
-				.UncheckedPressedImage(FEditorStyle::GetBrush("Sequencer.UnlockCamera"))
-		];
+	if (bHasCameraComponent)
+	{
+		// If this is a camera track, add a button to lock the viewport to the camera
+		EditBox.Get()->AddSlot()
+			.VAlign(VAlign_Center)
+			.HAlign(HAlign_Right)
+			.AutoWidth()
+			.Padding(4, 0, 0, 0)
+			[
+				SNew(SCheckBox)		
+					.IsFocusable(false)
+					.Visibility(this, &F3DTransformTrackEditor::IsCameraVisible, ObjectGuid)
+					.IsChecked(this, &F3DTransformTrackEditor::IsCameraLocked, ObjectGuid)
+					.OnCheckStateChanged(this, &F3DTransformTrackEditor::OnLockCameraClicked, ObjectGuid)
+					.ToolTipText(this, &F3DTransformTrackEditor::GetLockCameraToolTip, ObjectGuid)
+					.ForegroundColor(FLinearColor::White)
+					.CheckedImage(FEditorStyle::GetBrush("Sequencer.LockCamera"))
+					.CheckedHoveredImage(FEditorStyle::GetBrush("Sequencer.LockCamera"))
+					.CheckedPressedImage(FEditorStyle::GetBrush("Sequencer.LockCamera"))
+					.UncheckedImage(FEditorStyle::GetBrush("Sequencer.UnlockCamera"))
+					.UncheckedHoveredImage(FEditorStyle::GetBrush("Sequencer.UnlockCamera"))
+					.UncheckedPressedImage(FEditorStyle::GetBrush("Sequencer.UnlockCamera"))
+			];
+	}
 }
 
 
 void F3DTransformTrackEditor::BuildObjectBindingTrackMenu(FMenuBuilder& MenuBuilder, const FGuid& ObjectBinding, const UClass* ObjectClass)
 {
-	if (ObjectClass->IsChildOf(AActor::StaticClass()))
+	if (ObjectClass->IsChildOf(AActor::StaticClass()) || ObjectClass->IsChildOf(USceneComponent::StaticClass()))
 	{
 		FKeyParams KeyParams;
 		KeyParams.bCreateTrackIfMissing = true;
@@ -489,7 +493,7 @@ F3DTransformTrackEditor::IsCameraVisible(FGuid ObjectGuid) const
 
 	for (UObject* Object : OutObjects)
 	{
-		ACameraActor* Actor = Cast<ACameraActor>( Object );
+		AActor* Actor = Cast<AActor>( Object );
 		if (Actor)
 		{
 			return EVisibility::Visible;
@@ -504,11 +508,11 @@ ECheckBoxState F3DTransformTrackEditor::IsCameraLocked(FGuid ObjectGuid) const
 	TArray<UObject*> OutObjects;
 	GetSequencer()->GetRuntimeObjects( GetSequencer()->GetFocusedMovieSceneSequenceInstance(), ObjectGuid, OutObjects);
 
-	TWeakObjectPtr<ACameraActor> CameraActor;
+	TWeakObjectPtr<AActor> CameraActor;
 
 	for (UObject* Object : OutObjects)
 	{
-		ACameraActor* Actor = Cast<ACameraActor>( Object );
+		AActor* Actor = Cast<AActor>( Object );
 		if (Actor)
 		{
 			CameraActor = Actor;
@@ -537,11 +541,11 @@ void F3DTransformTrackEditor::OnLockCameraClicked(ECheckBoxState CheckBoxState,
 	TArray<UObject*> OutObjects;
 	GetSequencer()->GetRuntimeObjects( GetSequencer()->GetFocusedMovieSceneSequenceInstance(), ObjectGuid, OutObjects);
 
-	TWeakObjectPtr<ACameraActor> CameraActor;
+	TWeakObjectPtr<AActor> CameraActor;
 
 	for (UObject* Object : OutObjects)
 	{
-		ACameraActor* Actor = Cast<ACameraActor>( Object );
+		AActor* Actor = Cast<AActor>( Object );
 		if (Actor)
 		{
 			CameraActor = Actor;
@@ -607,11 +611,11 @@ FText F3DTransformTrackEditor::GetLockCameraToolTip(FGuid ObjectGuid) const
 	TArray<UObject*> OutObjects;
 	GetSequencer()->GetRuntimeObjects( GetSequencer()->GetFocusedMovieSceneSequenceInstance(), ObjectGuid, OutObjects);
 
-	TWeakObjectPtr<ACameraActor> CameraActor;
+	TWeakObjectPtr<AActor> CameraActor;
 
 	for (UObject* Object : OutObjects)
 	{
-		ACameraActor* Actor = Cast<ACameraActor>( Object );
+		AActor* Actor = Cast<AActor>( Object );
 		if (Actor)
 		{
 			CameraActor = Actor;
@@ -696,25 +700,32 @@ void F3DTransformTrackEditor::AddTransformKeysForObject( UObject* Object, EKey3D
 	if ( Actor != nullptr && SceneComponent != nullptr )
 	{
 		FTransformData CurrentTransform( SceneComponent );
-		AddTransformKeys( Actor, FTransformData(), CurrentTransform, ChannelToKey, false, KeyParams );
+		if(Object->GetClass()->IsChildOf(AActor::StaticClass()))
+		{
+			AddTransformKeys( Actor, FTransformData(), CurrentTransform, ChannelToKey, false, KeyParams );
+		}
+		else if(Object->GetClass()->IsChildOf(USceneComponent::StaticClass()))
+		{
+			AddTransformKeys( SceneComponent, FTransformData(), CurrentTransform, ChannelToKey, false, KeyParams );
+		}
 	}
 }
 
 
-void F3DTransformTrackEditor::AddTransformKeys( AActor* ActorToKey, const FTransformData& LastTransform, const FTransformData& CurrentTransform, EKey3DTransformChannel::Type ChannelsToKey, bool bUnwindRotation, FKeyParams KeyParams )
+void F3DTransformTrackEditor::AddTransformKeys( UObject* ObjectToKey, const FTransformData& LastTransform, const FTransformData& CurrentTransform, EKey3DTransformChannel::Type ChannelsToKey, bool bUnwindRotation, FKeyParams KeyParams )
 {
 	TArray<FTransformKey> NewKeys;
 	TArray<FTransformKey> DefaultKeys;
 	GetTransformKeys(LastTransform, CurrentTransform, ChannelsToKey, bUnwindRotation, NewKeys, DefaultKeys);
 
-	AnimatablePropertyChanged( FOnKeyProperty::CreateRaw(this, &F3DTransformTrackEditor::OnAddTransformKeys, ActorToKey, &NewKeys, &DefaultKeys, CurrentTransform, KeyParams ) );
+	AnimatablePropertyChanged( FOnKeyProperty::CreateRaw(this, &F3DTransformTrackEditor::OnAddTransformKeys, ObjectToKey, &NewKeys, &DefaultKeys, CurrentTransform, KeyParams ) );
 }
 
 
-bool F3DTransformTrackEditor::OnAddTransformKeys( float Time, AActor* ActorToKey, TArray<FTransformKey>* NewKeys, TArray<FTransformKey>* DefaultKeys, FTransformData CurrentTransform, FKeyParams KeyParams )
+bool F3DTransformTrackEditor::OnAddTransformKeys( float Time, UObject* ObjectToKey, TArray<FTransformKey>* NewKeys, TArray<FTransformKey>* DefaultKeys, FTransformData CurrentTransform, FKeyParams KeyParams )
 {
 	TArray<UObject*> ObjectsToKey;
-	ObjectsToKey.Add(ActorToKey);
+	ObjectsToKey.Add(ObjectToKey);
 	
 	return AddKeysToObjects(
 		ObjectsToKey,
diff --git a/Engine/Source/Editor/MovieSceneTools/Private/TrackEditors/TransformTrackEditor.h b/Engine/Source/Editor/MovieSceneTools/Private/TrackEditors/TransformTrackEditor.h
--- a/Engine/Source/Editor/MovieSceneTools/Private/TrackEditors/TransformTrackEditor.h
+++ b/Engine/Source/Editor/MovieSceneTools/Private/TrackEditors/TransformTrackEditor.h
@@ -117,17 +117,17 @@ private:
 	* @param bUnwindRotation Whether or not rotation key values should be unwound.
 	* @param KeyParams Parameters which control how the keys are added.
 	*/
-	void AddTransformKeys( AActor* ActorToKey, const FTransformData& LastTransform, const FTransformData& CurrentTransform, EKey3DTransformChannel::Type ChannelsToKey, bool bUnwindRotation, FKeyParams KeyParams );
+	void AddTransformKeys( UObject* ObjectToKey, const FTransformData& LastTransform, const FTransformData& CurrentTransform, EKey3DTransformChannel::Type ChannelsToKey, bool bUnwindRotation, FKeyParams KeyParams );
 
 	/**
 	* Delegate target of AnimatablePropertyChanged which actually adds the keys.
 
 	* @param Time The time to add keys.
-	* @param ActorToKey The actor to add keys to.
+	* @param ObjectToKey The object to add keys to.
 	* @param Keys The keys to add.
 	* @param KeyParams Parameters which control how the keys are added.
 	*/
-	bool OnAddTransformKeys( float Time, AActor* ActorToKey, TArray<FTransformKey>* NewKeys, TArray<FTransformKey>* DefaultKeys, FTransformData CurrentTransform, FKeyParams KeyParams );
+	bool OnAddTransformKeys( float Time, UObject* ObjectToKey, TArray<FTransformKey>* NewKeys, TArray<FTransformKey>* DefaultKeys, FTransformData CurrentTransform, FKeyParams KeyParams );
 
 private:
 
diff --git a/Engine/Source/Editor/MovieSceneTools/Public/FloatCurveKeyArea.h b/Engine/Source/Editor/MovieSceneTools/Public/FloatCurveKeyArea.h
--- a/Engine/Source/Editor/MovieSceneTools/Public/FloatCurveKeyArea.h
+++ b/Engine/Source/Editor/MovieSceneTools/Public/FloatCurveKeyArea.h
@@ -14,7 +14,7 @@ class MOVIESCENETOOLS_API FFloatCurveKeyArea
 {
 public:
 
-	FFloatCurveKeyArea(FRichCurve* InCurve, UMovieSceneSection* InOwningSection, const FLinearColor& InColor = FLinearColor(0.1f, 0.1f, 0.1f, 0.7f))
+	FFloatCurveKeyArea(FRichCurve* InCurve, UMovieSceneSection* InOwningSection, TOptional<FLinearColor> InColor = TOptional<FLinearColor>())
 		: Color(InColor)
 		, Curve(InCurve)
 		, OwningSection(InOwningSection)
@@ -41,7 +41,7 @@ public:
 	virtual bool CanCreateKeyEditor() override;
 	virtual TSharedRef<SWidget> CreateKeyEditor(ISequencer* Sequencer) override;
 	virtual void DeleteKey(FKeyHandle KeyHandle) override;
-	virtual FLinearColor GetColor() override;
+	virtual TOptional<FLinearColor> GetColor() override;
 	virtual ERichCurveExtrapolation GetExtrapolationMode(bool bPreInfinity) const override;
 	virtual ERichCurveInterpMode GetKeyInterpMode(FKeyHandle KeyHandle) const override;
 	virtual TSharedPtr<FStructOnScope> GetKeyStruct(FKeyHandle KeyHandle) override;
@@ -62,7 +62,7 @@ public:
 private:
 
 	/** The key area's color. */
-	FLinearColor Color;
+	TOptional<FLinearColor> Color;
 
 	/** The curve which provides the keys for this key area. */
 	FRichCurve* Curve;
diff --git a/Engine/Source/Editor/MovieSceneTools/Public/IntegralKeyArea.h b/Engine/Source/Editor/MovieSceneTools/Public/IntegralKeyArea.h
--- a/Engine/Source/Editor/MovieSceneTools/Public/IntegralKeyArea.h
+++ b/Engine/Source/Editor/MovieSceneTools/Public/IntegralKeyArea.h
@@ -27,7 +27,7 @@ public:
 	virtual TArray<FKeyHandle> AddKeyUnique( float Time, EMovieSceneKeyInterpolation InKeyInterpolation, float TimeToCopyFrom = FLT_MAX ) override;
 	virtual TOptional<FKeyHandle> DuplicateKey(FKeyHandle KeyToDuplicate) override;
 	virtual void DeleteKey(FKeyHandle KeyHandle) override;
-	virtual FLinearColor GetColor() override;
+	virtual TOptional<FLinearColor> GetColor() override;
 	virtual ERichCurveExtrapolation GetExtrapolationMode(bool bPreInfinity) const override;
 	virtual ERichCurveInterpMode GetKeyInterpMode(FKeyHandle KeyHandle) const override;
 	virtual TSharedPtr<FStructOnScope> GetKeyStruct(FKeyHandle KeyHandle) override;
