Sequencer: Fix linux build

diff --git a/Engine/Source/Editor/MovieSceneTools/Private/TrackEditors/AudioTrackEditor.cpp b/Engine/Source/Editor/MovieSceneTools/Private/TrackEditors/AudioTrackEditor.cpp
--- a/Engine/Source/Editor/MovieSceneTools/Private/TrackEditors/AudioTrackEditor.cpp
+++ b/Engine/Source/Editor/MovieSceneTools/Private/TrackEditors/AudioTrackEditor.cpp
@@ -659,7 +659,7 @@ void FAudioSection::Tick( const FGeometry& AllottedGeometry, const FGeometry& Pa
 	UMovieSceneTrack* Track = Section.GetTypedOuter<UMovieSceneTrack>();
 
 	USoundWave* SoundWave = DeriveSoundWave(AudioSection);
-	if (Track && SoundWave->NumChannels == 1 || SoundWave->NumChannels == 2)
+	if (Track && (SoundWave->NumChannels == 1 || SoundWave->NumChannels == 2))
 	{
 		const FSlateRect ParentRect = TransformRect(
 			Concatenate(
