Copying //UE4/Dev-Sequencer to //UE4/Main

diff --git a/Engine/Source/Editor/LevelEditor/Private/SLevelViewportToolBar.cpp b/Engine/Source/Editor/LevelEditor/Private/SLevelViewportToolBar.cpp
--- a/Engine/Source/Editor/LevelEditor/Private/SLevelViewportToolBar.cpp
+++ b/Engine/Source/Editor/LevelEditor/Private/SLevelViewportToolBar.cpp
@@ -124,22 +124,24 @@ void SLevelViewportToolBar::Construct( const FArguments& InArgs )
 		.ColorAndOpacity( this, &SViewportToolBar::OnGetColorAndOpacity )
 		.ForegroundColor( FEditorStyle::GetSlateColor(DefaultForegroundName) )
 		[
-			SNew( SVerticalBox )
-			+ SVerticalBox::Slot()
-			.AutoHeight()
+			SNew( SHorizontalBox )
+			+ SHorizontalBox::Slot()
+			.AutoWidth()
+			.Padding( ToolbarSlotPadding )
+			[
+				SNew( SEditorViewportToolbarMenu )
+				.ParentToolBar( SharedThis( this ) )
+				.Cursor( EMouseCursor::Default )
+				.Image( "EditorViewportToolBar.MenuDropdown" )
+				.AddMetaData<FTagMetaData>(FTagMetaData(TEXT("EditorViewportToolBar.MenuDropdown")))
+				.OnGetMenuContent( this, &SLevelViewportToolBar::GenerateOptionsMenu )
+			]
+
+			+ SHorizontalBox::Slot()
 			[
 				SNew( SHorizontalBox )
-				+ SHorizontalBox::Slot()
-				.AutoWidth()
-				.Padding( ToolbarSlotPadding )
-				[
-					SNew( SEditorViewportToolbarMenu )
-					.ParentToolBar( SharedThis( this ) )
-					.Cursor( EMouseCursor::Default )
-					.Image( "EditorViewportToolBar.MenuDropdown" )
-					.AddMetaData<FTagMetaData>(FTagMetaData(TEXT("EditorViewportToolBar.MenuDropdown")))
-					.OnGetMenuContent( this, &SLevelViewportToolBar::GenerateOptionsMenu )
-				]
+				.Visibility(Viewport.Pin().Get(), &SLevelViewport::GetFullToolbarVisibility)
+
 				+ SHorizontalBox::Slot()
 				.AutoWidth()
 				.Padding( ToolbarSlotPadding )
@@ -192,7 +194,7 @@ void SLevelViewportToolBar::Construct( const FArguments& InArgs )
 					SNew(STransformViewportToolBar)
 					.Viewport(ViewportRef)
 					.CommandList(LevelEditorModule.GetGlobalLevelEditorActions())
-					.Extenders(LevelEditorModule.GetToolBarExtensibilityManager()->GetAllExtenders())					
+					.Extenders(LevelEditorModule.GetToolBarExtensibilityManager()->GetAllExtenders())
 					.Visibility(ViewportRef, &SLevelViewport::GetTransformToolbarVisibility)
 				]
 				+ SHorizontalBox::Slot()
@@ -503,6 +505,7 @@ TSharedRef<SWidget> SLevelViewportToolBar::GenerateOptionsMenu() const
 			OptionsMenuBuilder.AddMenuEntry( FEditorViewportCommands::Get().ToggleRealTime );
 			OptionsMenuBuilder.AddMenuEntry( FEditorViewportCommands::Get().ToggleStats );
 			OptionsMenuBuilder.AddMenuEntry( FEditorViewportCommands::Get().ToggleFPS );
+			OptionsMenuBuilder.AddMenuEntry( LevelViewportActions.ToggleViewportToolbar );
 
 			FText HideAllLabel = LOCTEXT("HideAllLabel", "Hide All");
 			TArray< FLevelViewportCommands::FShowMenuCommand > HideStatsMenu;
@@ -692,6 +695,17 @@ void SLevelViewportToolBar::GeneratePlacedCameraMenuEntries( FMenuBuilder& Build
 	}
 }
 
+void SLevelViewportToolBar::GenerateViewportTypeMenu( FMenuBuilder& Builder ) const
+{
+	FLevelEditorModule& LevelEditorModule = FModuleManager::GetModuleChecked<FLevelEditorModule>(TEXT("LevelEditor"));
+	LevelEditorModule.IterateViewportTypes([&](FName, const FViewportTypeDefinition& InDefinition) {
+		if (InDefinition.ToggleCommand.IsValid())
+		{
+			Builder.AddMenuEntry(InDefinition.ToggleCommand);
+		}
+	});
+}
+
 TSharedRef<SWidget> SLevelViewportToolBar::GenerateCameraMenu() const
 {
 	Viewport.Pin()->OnFloatingButtonClicked();
@@ -711,7 +725,6 @@ TSharedRef<SWidget> SLevelViewportToolBar::GenerateCameraMenu() const
 		CameraMenuBuilder.AddMenuEntry(FEditorViewportCommands::Get().Back);
 	CameraMenuBuilder.EndSection();
 
-
 	TArray<ACameraActor*> Cameras;
 
 	for( TActorIterator<ACameraActor> It(GetWorld().Get()); It; ++It )
@@ -736,6 +749,30 @@ TSharedRef<SWidget> SLevelViewportToolBar::GenerateCameraMenu() const
 		CameraMenuBuilder.EndSection();
 	}
 
+	{
+		FLevelEditorModule& LevelEditorModule = FModuleManager::GetModuleChecked<FLevelEditorModule>(TEXT("LevelEditor"));
+
+		int32 NumCustomViewportTypes = 0;
+		LevelEditorModule.IterateViewportTypes([&](FName, const FViewportTypeDefinition&){ ++NumCustomViewportTypes; });
+
+		FText ViewportTypesHeading = LOCTEXT("ViewportTypes", "Viewport Type");
+		const uint32 MaxViewportTypesInTopLevelMenu = 4;
+		if( NumCustomViewportTypes > MaxViewportTypesInTopLevelMenu )
+		{
+			CameraMenuBuilder.BeginSection("ViewportTypes");
+			CameraMenuBuilder.AddSubMenu( ViewportTypesHeading, FText(), FNewMenuDelegate::CreateSP(this, &SLevelViewportToolBar::GenerateViewportTypeMenu) );
+			CameraMenuBuilder.EndSection();
+		}
+		else
+		{
+			CameraMenuBuilder.BeginSection("ViewportTypes", ViewportTypesHeading);
+			GenerateViewportTypeMenu(CameraMenuBuilder);
+			CameraMenuBuilder.EndSection();
+
+		}
+		CameraMenuBuilder.EndSection();
+	}
+
 	return CameraMenuBuilder.MakeWidget();
 }
 
@@ -851,44 +888,6 @@ void SLevelViewportToolBar::GenerateViewportConfigsMenu(FMenuBuilder& MenuBuilde
 			);
 	}
 	MenuBuilder.EndSection();
-
-	MenuBuilder.BeginSection("LevelViewportSpecializedConfigs", LOCTEXT("SpecializedConfigHeader", "Specialized") );
-	{
-		FToolBarBuilder SpecializedButtons(CommandList, FMultiBoxCustomization::None);
-		SpecializedButtons.SetLabelVisibility(EVisibility::Collapsed);
-		SpecializedButtons.SetStyle(&FEditorStyle::Get(), "ViewportLayoutToolbar");
-
-
-		FLevelEditorModule& LevelEditorModule = FModuleManager::Get().GetModuleChecked<FLevelEditorModule>("LevelEditor");
-		LevelEditorModule.IterateCustomViewportLayouts([&](FName InName, const FCustomViewportLayoutDefinition& Definition){
-
-			SpecializedButtons.AddToolBarButton(
-				FExecuteAction::CreateSP(Viewport.Pin().Get(), &SLevelViewport::OnSetViewportConfiguration, InName ),
-				NAME_None,
-				Definition.DisplayName,
-				Definition.Description,
-				Definition.Icon
-			);
-
-		});
-
-		MenuBuilder.AddWidget(
-			SNew(SHorizontalBox)
-			+SHorizontalBox::Slot()
-			.AutoWidth()
-			[
-				SpecializedButtons.MakeWidget()
-			]
-			+SHorizontalBox::Slot()
-			.FillWidth(1)
-			[
-				SNullWidget::NullWidget
-			],
-			FText::GetEmpty(),
-			true
-		);
-	}
-	MenuBuilder.EndSection();
 }
 
 
diff --git a/Engine/Source/Editor/LevelEditor/Private/SLevelViewportToolBar.h b/Engine/Source/Editor/LevelEditor/Private/SLevelViewportToolBar.h
--- a/Engine/Source/Editor/LevelEditor/Private/SLevelViewportToolBar.h
+++ b/Engine/Source/Editor/LevelEditor/Private/SLevelViewportToolBar.h
@@ -109,6 +109,13 @@ private:
 	 */
 	void GeneratePlacedCameraMenuEntries( FMenuBuilder& Builder, TArray<ACameraActor*> Cameras ) const;
 
+	/**
+	 * Generates menu entries for changing the type of the viewport
+	 *
+	 * @param Builder	The menu builder to add menu entries to
+	 */
+	void GenerateViewportTypeMenu( FMenuBuilder& Builder ) const;
+
 	/**
 	 * Generates the toolbar view menu content 
 	 *
diff --git a/Engine/Source/Editor/LevelEditor/Public/LevelEditor.h b/Engine/Source/Editor/LevelEditor/Public/LevelEditor.h
--- a/Engine/Source/Editor/LevelEditor/Public/LevelEditor.h
+++ b/Engine/Source/Editor/LevelEditor/Public/LevelEditor.h
@@ -8,7 +8,7 @@
 #include "ModuleInterface.h"
 #include "SEditorViewport.h"
 #include "Toolkits/AssetEditorToolkit.h" // For FExtensibilityManager
-#include "CustomViewportLayout.h"
+#include "ViewportTypeDefinition.h"
 
 extern const FName LevelEditorApp;
 
@@ -230,38 +230,29 @@ public:
 
 public:
 	
-	/** Register a custom viewport layout for the level editor */
-	void RegisterCustomViewportLayout(FName InLayoutName, const FCustomViewportLayoutDefinition& InDefinition)
+	/** Register a viewport type for the level editor */
+	void RegisterViewportType(FName InLayoutName, const FViewportTypeDefinition& InDefinition)
 	{
-		CustomViewportLayouts.Add(InLayoutName, InDefinition);
+		CustomViewports.Add(InLayoutName, InDefinition);
 	}
 
-	/** Unregister a previously registered custom viewport layout */
-	void UnRegisterCustomViewportLayout(FName InLayoutName)
+	/** Unregister a previously registered viewport type */
+	void UnregisterViewportType(FName InLayoutName)
 	{
-		CustomViewportLayouts.Remove(InLayoutName);
+		CustomViewports.Remove(InLayoutName);
 	}
 
-	/** Iterate all registered custom viewport layouts */
-	void IterateCustomViewportLayouts(TFunctionRef<void(FName, const FCustomViewportLayoutDefinition&)> Iter)
+	/** Iterate all registered viewport types */
+	void IterateViewportTypes(TFunctionRef<void(FName, const FViewportTypeDefinition&)> Iter)
 	{
-		for (auto& Pair : CustomViewportLayouts)
+		for (auto& Pair : CustomViewports)
 		{
 			Iter(Pair.Key, Pair.Value);
 		}
 	}
 
-	/** Create an instanceof a custom viewport layout from the specified layout name */
-	TSharedPtr<FLevelViewportLayout> FactoryCustomViewportLayout(FName InLayoutName) const
-	{
-		const FCustomViewportLayoutDefinition* Definition = CustomViewportLayouts.Find(InLayoutName);
-		if (Definition)
-		{
-			return Definition->FactoryFunction();
-		}
-		
-		return nullptr;
-	}
+	/** Create an instance of a custom viewport from the specified viewport type name */
+	TSharedRef<IViewportLayoutEntity> FactoryViewport(FName InTypeName, const FViewportConstructionArgs& ConstructionArgs) const;
 
 private:
 	/**
@@ -336,8 +327,8 @@ private:
 	/* Holds the Editor's tab manager */
 	TSharedPtr<FTabManager> LevelEditorTabManager;
 
-	/** Map of named viewport layouts to factory functions */
-	TMap<FName, FCustomViewportLayoutDefinition> CustomViewportLayouts;
+	/** Map of named viewport types to factory functions */
+	TMap<FName, FViewportTypeDefinition> CustomViewports;
 };
 
 #endif // __LevelEditor_h__
diff --git a/Engine/Source/Editor/LevelEditor/Public/LevelViewportActions.h b/Engine/Source/Editor/LevelEditor/Public/LevelViewportActions.h
--- a/Engine/Source/Editor/LevelEditor/Public/LevelViewportActions.h
+++ b/Engine/Source/Editor/LevelEditor/Public/LevelViewportActions.h
@@ -162,6 +162,8 @@ public:
 	TSharedPtr< FUICommandInfo > ViewportConfig_FourPanesBottom;
 	TSharedPtr< FUICommandInfo > ViewportConfig_FourPanes2x2;
 
+	TSharedPtr< FUICommandInfo > ToggleViewportToolbar;
+
 	TSharedPtr< FUICommandInfo > EnablePreviewMesh;
 	TSharedPtr< FUICommandInfo > CyclePreviewMesh;
 
diff --git a/Engine/Source/Editor/LevelEditor/Public/LevelViewportLayout.h b/Engine/Source/Editor/LevelEditor/Public/LevelViewportLayout.h
--- a/Engine/Source/Editor/LevelEditor/Public/LevelViewportLayout.h
+++ b/Engine/Source/Editor/LevelEditor/Public/LevelViewportLayout.h
@@ -9,7 +9,31 @@
 #include "LevelEditor.h"
 #include "LevelViewportActions.h"
 
+class ILevelEditor;
 class SLevelViewport;
+class FLevelViewportLayout;
+
+/** Arguments for constructing a viewport */
+struct FViewportConstructionArgs
+{
+	FViewportConstructionArgs()
+		: ViewportType(LVT_Perspective)
+		, bRealtime(false)
+	{}
+
+	/** The viewport's parent layout */
+	TSharedPtr<FLevelViewportLayout> ParentLayout;
+	/** The viewport's parent level editor */
+	TWeakPtr<ILevelEditor> ParentLevelEditor;
+	/** The viewport's desired type */
+	ELevelViewportType ViewportType;
+	/** Whether the viewport should default to realtime */
+	bool bRealtime;
+	/** A config key for loading/saving settings for the viewport */
+	FString ConfigKey;
+	/** Widget enabled attribute */
+	TAttribute<bool> IsEnabled;
+};
 
 /** Interface that defines an entity within a viewport layout */
 class IViewportLayoutEntity : public TSharedFromThis<IViewportLayoutEntity>
@@ -24,6 +48,9 @@ public:
 	/** Optionally return this entity as an SLevelViewport. Legacy function for interop with code that used to deal directly with SLevelViewports */
 	virtual TSharedPtr<SLevelViewport> AsLevelViewport() const { return nullptr; }
 
+	/** Get this viewport's level editor viewport client */
+	virtual FLevelEditorViewportClient& GetLevelViewportClient() const = 0;
+
 	/** Check if this entity has an active play in editor viewport */
 	virtual bool IsPlayInEditorViewportActive() const = 0;
 
@@ -38,26 +65,9 @@ public:
 
 	/** Called to save this item's settings in the specified config section */
 	virtual void SaveConfig(const FString& ConfigSection) = 0;
-};
-
-class LEVELEDITOR_API FLevelViewportLayoutEntity : public IViewportLayoutEntity
-{
-public:
-
-	FLevelViewportLayoutEntity(TSharedRef<SLevelViewport> InLevelViewport);
-	virtual TSharedRef<SWidget> AsWidget() const override;
-	virtual TSharedPtr<SLevelViewport> AsLevelViewport() const override;
-
-	bool IsPlayInEditorViewportActive() const;
-	void RegisterGameViewportIfPIE();
-	void SetKeyboardFocus();
-	void OnLayoutDestroyed();
-	void SaveConfig(const FString& ConfigSection);
-
-protected:
 
-	/** This entity's level viewport */
-	TSharedRef<SLevelViewport> LevelViewport;
+	/** Get the type of this viewport as a name */
+	virtual FName GetType() const = 0;
 };
 
 /**
@@ -85,7 +95,7 @@ public:
 	 * @param LayoutString			The layout string loaded from file to custom build the layout with
 	 * @param InParentLevelEditor	Optional level editor parent to use for new viewports
 	 */
-	TSharedRef<SWidget> BuildViewportLayout( TSharedPtr<SDockTab> InParentDockTab, TSharedPtr<class FLevelViewportTabContent> InParentTab, const FString& LayoutString, TWeakPtr<class SLevelEditor> InParentLevelEditor = NULL );
+	TSharedRef<SWidget> BuildViewportLayout( TSharedPtr<SDockTab> InParentDockTab, TSharedPtr<class FLevelViewportTabContent> InParentTab, const FString& LayoutString, TWeakPtr<ILevelEditor> InParentLevelEditor = NULL );
 
 	/**
 	 * Makes a request to maximize a specific viewport and hide the others in this layout
@@ -258,7 +268,7 @@ protected:
 	TWeakPtr< class FLevelViewportTabContent > ParentTabContent;
 
 	/** The optional parent level editor for this layout */
-	TWeakPtr< class SLevelEditor > ParentLevelEditor;
+	TWeakPtr< ILevelEditor > ParentLevelEditor;
 
 	/** The current maximized viewport if any */
 	FName MaximizedViewport;
diff --git a/Engine/Source/Editor/LevelEditor/Public/LevelViewportTabContent.h b/Engine/Source/Editor/LevelEditor/Public/LevelViewportTabContent.h
--- a/Engine/Source/Editor/LevelEditor/Public/LevelViewportTabContent.h
+++ b/Engine/Source/Editor/LevelEditor/Public/LevelViewportTabContent.h
@@ -19,7 +19,7 @@ class LEVELEDITOR_API FLevelViewportTabContent : public TSharedFromThis<FLevelVi
 {
 public:
 	/** Starts the tab content object and creates the initial layout based on the layout string */
-	void Initialize(TSharedPtr<class SLevelEditor> InParentLevelEditor, TSharedPtr<SDockTab> InParentTab, const FString& InLayoutString);
+	void Initialize(TSharedPtr<ILevelEditor> InParentLevelEditor, TSharedPtr<SDockTab> InParentTab, const FString& InLayoutString);
 
 	/** Returns whether the tab is currently shown */
 	bool IsVisible() const;
@@ -71,7 +71,7 @@ private:
 
 private:
 	TWeakPtr<class SDockTab> ParentTab;
-	TWeakPtr<class SLevelEditor> ParentLevelEditor;
+	TWeakPtr<ILevelEditor> ParentLevelEditor;
 	FString LayoutString;
 
 	/** Current layout */
diff --git a/Engine/Source/Editor/LevelEditor/Public/SLevelViewport.h b/Engine/Source/Editor/LevelEditor/Public/SLevelViewport.h
--- a/Engine/Source/Editor/LevelEditor/Public/SLevelViewport.h
+++ b/Engine/Source/Editor/LevelEditor/Public/SLevelViewport.h
@@ -8,6 +8,7 @@
 
 class SActorPreview;
 class SGameLayerManager;
+class ILevelEditor;
 
 /**
  * Encapsulates an SViewport and an SLevelViewportToolBar
@@ -22,7 +23,8 @@ public:
 
 		SLATE_ARGUMENT( TWeakPtr<class FEditorModeTools>, EditorModeTools )
 		SLATE_ARGUMENT( TSharedPtr<class FLevelViewportLayout>, ParentLayout )
-		SLATE_ARGUMENT( TWeakPtr<class SLevelEditor>, ParentLevelEditor )
+		SLATE_ARGUMENT( TWeakPtr<ILevelEditor>, ParentLevelEditor )
+		SLATE_ARGUMENT( TSharedPtr<FLevelEditorViewportClient>, LevelEditorViewportClient )
 		SLATE_ARGUMENT( ELevelViewportType, ViewportType )
 		SLATE_ARGUMENT( bool, Realtime )
 		SLATE_ARGUMENT( FString, ConfigKey )
@@ -255,7 +257,7 @@ public:
 	FString GetDeviceProfileString( ) const;
 
 	/** Get the parent level editor for this viewport */
-	TWeakPtr<class SLevelEditor> GetParentLevelEditor() const { return ParentLevelEditor; }
+	TWeakPtr<ILevelEditor> GetParentLevelEditor() const { return ParentLevelEditor; }
 
 	/** Called to get the level text */
 	FText GetCurrentLevelText( bool bDrawOnlyLabel ) const;
@@ -264,13 +266,18 @@ public:
 	FText GetCurrentFeatureLevelPreviewText( bool bDrawOnlyLabel ) const;
 	
 	/** @return The visibility of the current level text display */
-	EVisibility GetCurrentLevelTextVisibility() const;
+	virtual EVisibility GetCurrentLevelTextVisibility() const;
 
 	/** @return The visibility of the current feature level preview text display */
 	EVisibility GetCurrentFeatureLevelPreviewTextVisibility() const;
 
 	/** @return The visibility of the viewport controls popup */
-	EVisibility GetViewportControlsVisibility() const;
+	virtual EVisibility GetViewportControlsVisibility() const;
+
+	/**
+	 * Called to get the visibility of the level viewport toolbar
+	 */
+	virtual EVisibility GetToolBarVisibility() const;
 
 	/**
 	 * Sets the current layout on the parent tab that this viewport belongs to.
@@ -287,6 +294,18 @@ public:
 	 */
 	bool IsViewportConfigurationSet(FName ConfigurationName) const;
 
+	/** Get this level viewport widget's type within its parent layout */
+	FName GetViewportTypeWithinLayout() const;
+
+	/** Set this level viewport widget's type within its parent layout */
+	void SetViewportTypeWithinLayout(FName InLayoutType);
+
+	/** Activates the specified viewport type in the layout, if it's not already, or reverts to default if it is. */
+	void ToggleViewportTypeActivationWithinLayout(FName InLayoutType);
+
+	/** Checks if the specified layout type matches our current viewport type. */
+	bool IsViewportTypeWithinLayoutEqual(FName InLayoutType);
+
 	/** For the specified actor, toggle Pinned/Unpinned of it's ActorPreview */
 	void ToggleActorPreviewIsPinned(TWeakObjectPtr<AActor> PreviewActor);
 
@@ -296,6 +315,9 @@ public:
 	/** Actions to perform whenever the viewports floating buttons are pressed */
 	void OnFloatingButtonClicked();
 
+	/** Get the visibility for items considered to be part of the 'full' viewport toolbar */
+	EVisibility GetFullToolbarVisibility() const { return bShowFullToolbar ? EVisibility::Visible : EVisibility::Collapsed; }
+
 protected:
 	/** SEditorViewport interface */
 	virtual TSharedRef<FEditorViewportClient> MakeEditorViewportClient() override;
@@ -554,11 +576,7 @@ private:
 	* @param InCommandName	The command used by the functions
 	*/
 	void BindStatCommand(const TSharedPtr<FUICommandInfo> InMenuItem, const FString& InCommandName);
-
-	/**
-	 * Called to get the visibility of the level viewport toolbar
-	 */
-	EVisibility GetToolBarVisibility() const;
+	
 	/**
 	 * Called to build the viewport context menu when the user is Drag Dropping from the Content Browser
 	 */
@@ -639,12 +657,18 @@ private:
 	/** Gets the active scene viewport for the game */
 	const FSceneViewport* GetGameSceneViewport() const;
 
+	/** Called when the user toggles the full toolbar */
+	void OnToggleShowFullToolbar() { bShowFullToolbar = !bShowFullToolbar; }
+
+	/** Check whether we should display the full toolbar or not */
+	bool ShouldShowFullToolbar() const { return bShowFullToolbar; }
+
 private:
 	/** Tab which this viewport is located in */
 	TWeakPtr<class FLevelViewportLayout> ParentLayout;
 
 	/** Pointer to the parent level editor for this viewport */
-	TWeakPtr<class SLevelEditor> ParentLevelEditor;
+	TWeakPtr<ILevelEditor> ParentLevelEditor;
 
 	/** Viewport overlay widget exposed to game systems when running play-in-editor */
 	TSharedPtr<SOverlay> PIEViewportOverlayWidget;
@@ -787,6 +811,9 @@ private:
 	/** The users value for allowing throttling, we restore this value when we lose focus. */
 	int32 UserAllowThrottlingValue;
 
+	/** Whether to show a full toolbar, or a compact one */
+	bool bShowFullToolbar;
+
 protected:
 	void LockActorInternal(AActor* NewActorToLock);
 
diff --git a/Engine/Source/Editor/LevelEditor/Public/ViewportTypeDefinition.h b/Engine/Source/Editor/LevelEditor/Public/ViewportTypeDefinition.h
new file mode 100644
--- /dev/null
+++ b/Engine/Source/Editor/LevelEditor/Public/ViewportTypeDefinition.h
@@ -0,0 +1,30 @@
+// Copyright 1998-2015 Epic Games, Inc. All Rights Reserved.
+
+#pragma once
+
+class IViewportLayoutEntity;
+struct FViewportConstructionArgs;
+
+/** Definition of a custom viewport */
+struct FViewportTypeDefinition
+{
+	typedef TFunction<TSharedRef<IViewportLayoutEntity>(const FViewportConstructionArgs&)> FFactoryFunctionType;
+
+	template<typename T>
+	static FViewportTypeDefinition FromType()
+	{
+		return FViewportTypeDefinition([](const FViewportConstructionArgs& Args) -> TSharedRef<IViewportLayoutEntity> {
+			return MakeShareable(new T(Args));
+		});
+	}
+
+	FViewportTypeDefinition(const FFactoryFunctionType& InFactoryFunction)
+		: FactoryFunction(MoveTemp(InFactoryFunction))
+	{}
+
+	/** A UI command for toggling activation this viewport */
+	TSharedPtr<FUICommandInfo> ToggleCommand;
+
+	/** Function used to create a new instance of the viewport */
+	FFactoryFunctionType FactoryFunction;
+};
diff --git a/Engine/Source/Editor/MovieSceneTools/MovieSceneTools.Build.cs b/Engine/Source/Editor/MovieSceneTools/MovieSceneTools.Build.cs
--- a/Engine/Source/Editor/MovieSceneTools/MovieSceneTools.Build.cs
+++ b/Engine/Source/Editor/MovieSceneTools/MovieSceneTools.Build.cs
@@ -41,7 +41,9 @@ public class MovieSceneTools : ModuleRules
                 "EditorStyle",
 				"EditorWidgets",
 				"RenderCore",
-				"RHI"
+				"RHI",
+				"ShaderCore",
+				"SequenceRecorder",
 			}
 		);
 
@@ -52,7 +54,7 @@ public class MovieSceneTools : ModuleRules
 				"Sequencer",
                 "Settings",
 				"SceneOutliner",
-				"PropertyEditor"
+				"PropertyEditor",
 			}
 		);
 
@@ -61,8 +63,10 @@ public class MovieSceneTools : ModuleRules
                 "AssetRegistry",
 				"AssetTools",
 				"SceneOutliner",
-				"PropertyEditor"
+				"PropertyEditor",
 			}
 		);
-	}
+
+        CircularlyReferencedDependentModules.Add("Sequencer");
+    }
 }
diff --git a/Engine/Source/Editor/MovieSceneTools/Private/CurveKeyAreas/BoolKeyArea.cpp b/Engine/Source/Editor/MovieSceneTools/Private/CurveKeyAreas/BoolKeyArea.cpp
--- a/Engine/Source/Editor/MovieSceneTools/Private/CurveKeyAreas/BoolKeyArea.cpp
+++ b/Engine/Source/Editor/MovieSceneTools/Private/CurveKeyAreas/BoolKeyArea.cpp
@@ -20,7 +20,14 @@ TSharedRef<SWidget> FBoolKeyArea::CreateKeyEditor(ISequencer* Sequencer)
 		.Sequencer(Sequencer)
 		.OwningSection(OwningSection)
 		.Curve(&Curve)
+		.OnValueChanged(this, &FBoolKeyArea::OnValueChanged)
 		.IntermediateValue_Lambda([this] {
 			return IntermediateValue;
 		});
 };
+
+void FBoolKeyArea::OnValueChanged(bool InValue)
+{
+	ClearIntermediateValue();
+}
+
diff --git a/Engine/Source/Editor/MovieSceneTools/Private/CurveKeyAreas/EnumKeyArea.cpp b/Engine/Source/Editor/MovieSceneTools/Private/CurveKeyAreas/EnumKeyArea.cpp
--- a/Engine/Source/Editor/MovieSceneTools/Private/CurveKeyAreas/EnumKeyArea.cpp
+++ b/Engine/Source/Editor/MovieSceneTools/Private/CurveKeyAreas/EnumKeyArea.cpp
@@ -20,8 +20,15 @@ TSharedRef<SWidget> FEnumKeyArea::CreateKeyEditor(ISequencer* Sequencer)
 		.Sequencer(Sequencer)
 		.OwningSection(OwningSection)
 		.Curve(&Curve)
+		.OnValueChanged(this, &FEnumKeyArea::OnValueChanged)
 		.Enum(Enum)
 		.IntermediateValue_Lambda([this] {
 			return IntermediateValue;
 		});
 };
+
+void FEnumKeyArea::OnValueChanged(int32 InValue)
+{
+	ClearIntermediateValue();
+}
+
diff --git a/Engine/Source/Editor/MovieSceneTools/Private/CurveKeyAreas/FloatCurveKeyArea.cpp b/Engine/Source/Editor/MovieSceneTools/Private/CurveKeyAreas/FloatCurveKeyArea.cpp
--- a/Engine/Source/Editor/MovieSceneTools/Private/CurveKeyAreas/FloatCurveKeyArea.cpp
+++ b/Engine/Source/Editor/MovieSceneTools/Private/CurveKeyAreas/FloatCurveKeyArea.cpp
@@ -90,6 +90,7 @@ TSharedRef<SWidget> FFloatCurveKeyArea::CreateKeyEditor(ISequencer* Sequencer)
 		.Sequencer(Sequencer)
 		.OwningSection(OwningSection)
 		.Curve(Curve)
+		.OnValueChanged(this, &FFloatCurveKeyArea::OnValueChanged)
 		.IntermediateValue_Lambda([this] {
 			return IntermediateValue;
 		});
@@ -292,3 +293,8 @@ void FFloatCurveKeyArea::PasteKeys(const FMovieSceneClipboardKeyTrack& KeyTrack,
 		return true;
 	});
 }
+
+void FFloatCurveKeyArea::OnValueChanged(float InValue)
+{
+	ClearIntermediateValue();
+}
diff --git a/Engine/Source/Editor/MovieSceneTools/Private/CurveKeyEditors/SBoolCurveKeyEditor.cpp b/Engine/Source/Editor/MovieSceneTools/Private/CurveKeyEditors/SBoolCurveKeyEditor.cpp
--- a/Engine/Source/Editor/MovieSceneTools/Private/CurveKeyEditors/SBoolCurveKeyEditor.cpp
+++ b/Engine/Source/Editor/MovieSceneTools/Private/CurveKeyEditors/SBoolCurveKeyEditor.cpp
@@ -11,6 +11,7 @@ void SBoolCurveKeyEditor::Construct(const FArguments& InArgs)
 	Sequencer = InArgs._Sequencer;
 	OwningSection = InArgs._OwningSection;
 	Curve = InArgs._Curve;
+	OnValueChangedEvent = InArgs._OnValueChanged;
 	IntermediateValue = InArgs._IntermediateValue;
 	float CurrentTime = Sequencer->GetCurrentLocalTime(*Sequencer->GetFocusedMovieSceneSequence());
 	ChildSlot
@@ -56,15 +57,19 @@ void SBoolCurveKeyEditor::OnCheckStateChanged(ECheckBoxState NewCheckboxState)
 				OwningSection->SetEndTime(CurrentTime);
 			}
 		}
-
+		
+		int32 NewValue = NewCheckboxState == ECheckBoxState::Checked ? 1 : 0;
 		if (Curve->GetNumKeys() == 0)
 		{
-			Curve->SetDefaultValue(NewCheckboxState == ECheckBoxState::Checked ? 1 : 0);
+			Curve->SetDefaultValue(NewValue);
 		}
 		else
 		{
-			Curve->UpdateOrAddKey(CurrentTime, NewCheckboxState == ECheckBoxState::Checked ? 1 : 0);
+			Curve->UpdateOrAddKey(CurrentTime, NewValue);
 		}
+
+		OnValueChangedEvent.ExecuteIfBound(NewCheckboxState == ECheckBoxState::Checked);
+
 		Sequencer->UpdateRuntimeInstances();
 	}
 }
diff --git a/Engine/Source/Editor/MovieSceneTools/Private/CurveKeyEditors/SBoolCurveKeyEditor.h b/Engine/Source/Editor/MovieSceneTools/Private/CurveKeyEditors/SBoolCurveKeyEditor.h
--- a/Engine/Source/Editor/MovieSceneTools/Private/CurveKeyEditors/SBoolCurveKeyEditor.h
+++ b/Engine/Source/Editor/MovieSceneTools/Private/CurveKeyEditors/SBoolCurveKeyEditor.h
@@ -8,10 +8,14 @@
 class SBoolCurveKeyEditor : public SCompoundWidget
 {
 public:
+	/** Notification for bool value change */
+	DECLARE_DELEGATE_OneParam( FOnValueChanged, bool );
+
 	SLATE_BEGIN_ARGS(SBoolCurveKeyEditor) {}
 		SLATE_ARGUMENT(ISequencer*, Sequencer)
 		SLATE_ARGUMENT(UMovieSceneSection*, OwningSection)
 		SLATE_ARGUMENT(FIntegralCurve*, Curve)
+		SLATE_EVENT(FOnValueChanged, OnValueChanged)
 		SLATE_ATTRIBUTE(TOptional<bool>, IntermediateValue)
 	SLATE_END_ARGS();
 
@@ -25,5 +29,6 @@ private:
 	ISequencer* Sequencer;
 	UMovieSceneSection* OwningSection;
 	FIntegralCurve* Curve;
+	FOnValueChanged OnValueChangedEvent;
 	TAttribute<TOptional<bool>> IntermediateValue;
 };
\ No newline at end of file
diff --git a/Engine/Source/Editor/MovieSceneTools/Private/CurveKeyEditors/SEnumCurveKeyEditor.cpp b/Engine/Source/Editor/MovieSceneTools/Private/CurveKeyEditors/SEnumCurveKeyEditor.cpp
--- a/Engine/Source/Editor/MovieSceneTools/Private/CurveKeyEditors/SEnumCurveKeyEditor.cpp
+++ b/Engine/Source/Editor/MovieSceneTools/Private/CurveKeyEditors/SEnumCurveKeyEditor.cpp
@@ -10,109 +10,57 @@ void SEnumCurveKeyEditor::Construct(const FArguments& InArgs)
 	Sequencer = InArgs._Sequencer;
 	OwningSection = InArgs._OwningSection;
 	Curve = InArgs._Curve;
-	Enum = InArgs._Enum;
-	IntermediateValue = InArgs._IntermediateValue;
-
-	for (int32 i = 0; i < InArgs._Enum->NumEnums() - 1; i++)
-	{
-		if (Enum->HasMetaData( TEXT("Hidden"), i ) == false)
-		{
-			VisibleEnumNameIndices.Add(MakeShareable(new int32(i)));
-		}
-	}
+	OnValueChangedEvent = InArgs._OnValueChanged;
 
 	ChildSlot
 	[
-		SAssignNew(EnumValuesCombo, SComboBox<TSharedPtr<int32>>)
-		.ButtonStyle(FEditorStyle::Get(), "FlatButton.Light")
-		.OptionsSource(&VisibleEnumNameIndices)
-		.OnGenerateWidget(this, &SEnumCurveKeyEditor::OnGenerateWidget)
-		.OnSelectionChanged(this, &SEnumCurveKeyEditor::OnComboSelectionChanged)
-		.OnComboBoxOpening(this, &SEnumCurveKeyEditor::OnComboMenuOpening)
-		.ContentPadding(FMargin(2, 0))
-		[
-			SNew(STextBlock)
-			.Font(FEditorStyle::GetFontStyle("Sequencer.AnimationOutliner.RegularFont"))
-			.Text(this, &SEnumCurveKeyEditor::GetCurrentValue)
-		]
+		MovieSceneToolHelpers::MakeEnumComboBox(
+			InArgs._Enum,
+			TAttribute<int32>::Create(TAttribute<int32>::FGetter::CreateSP(this, &SEnumCurveKeyEditor::OnGetCurrentValue)),
+			FOnEnumSelectionChanged::CreateSP(this, &SEnumCurveKeyEditor::OnComboSelectionChanged),
+			InArgs._IntermediateValue)
 	];
-	bUpdatingSelectionInternally = false;
 }
 
-FText SEnumCurveKeyEditor::GetCurrentValue() const
+int32 SEnumCurveKeyEditor::OnGetCurrentValue() const
 {
-	if ( IntermediateValue.IsSet() && IntermediateValue.Get().IsSet() )
-	{
-		int32 IntermediateNameIndex = Enum->GetIndexByValue( IntermediateValue.Get().GetValue() );
-		return Enum->GetDisplayNameText( IntermediateNameIndex );
-	}
-
 	float CurrentTime = Sequencer->GetCurrentLocalTime(*Sequencer->GetFocusedMovieSceneSequence());
-	int32 CurrentNameIndex = Enum->GetIndexByValue(Curve->Evaluate(CurrentTime));
-	return Enum->GetDisplayNameText(CurrentNameIndex);
-}
-
-TSharedRef<SWidget> SEnumCurveKeyEditor::OnGenerateWidget(TSharedPtr<int32> InItem)
-{
-	return SNew(STextBlock)
-		.Text(Enum->GetDisplayNameText(*InItem));
+	return Curve->Evaluate(CurrentTime);
 }
 
-void SEnumCurveKeyEditor::OnComboSelectionChanged(TSharedPtr<int32> InSelectedItem, ESelectInfo::Type SelectInfo)
+void SEnumCurveKeyEditor::OnComboSelectionChanged(int32 InSelectedItem, ESelectInfo::Type SelectInfo)
 {
-	if (bUpdatingSelectionInternally == false)
+	FScopedTransaction Transaction(LOCTEXT("SetEnumKey", "Set Enum Key Value"));
+	OwningSection->SetFlags(RF_Transactional);
+	if (OwningSection->TryModify())
 	{
-		FScopedTransaction Transaction(LOCTEXT("SetEnumKey", "Set Enum Key Value"));
-		OwningSection->SetFlags(RF_Transactional);
-		if (OwningSection->TryModify())
-		{
-			float CurrentTime = Sequencer->GetCurrentLocalTime(*Sequencer->GetFocusedMovieSceneSequence());
-
-			bool bKeyWillBeAdded = Curve->IsKeyHandleValid(Curve->FindKey(CurrentTime)) == false;
-			if (bKeyWillBeAdded)
-			{
-				if (OwningSection->GetStartTime() > CurrentTime)
-				{
-					OwningSection->SetStartTime(CurrentTime);
-				}
-				if (OwningSection->GetEndTime() < CurrentTime)
-				{
-					OwningSection->SetEndTime(CurrentTime);
-				}
-			}
+		float CurrentTime = Sequencer->GetCurrentLocalTime(*Sequencer->GetFocusedMovieSceneSequence());
 
-			int32 SelectedValue = Enum->GetValueByIndex(*InSelectedItem);
-			if (Curve->GetNumKeys() == 0)
+		bool bKeyWillBeAdded = Curve->IsKeyHandleValid(Curve->FindKey(CurrentTime)) == false;
+		if (bKeyWillBeAdded)
+		{
+			if (OwningSection->GetStartTime() > CurrentTime)
 			{
-				Curve->SetDefaultValue(SelectedValue);
+				OwningSection->SetStartTime(CurrentTime);
 			}
-			else
+			if (OwningSection->GetEndTime() < CurrentTime)
 			{
-				Curve->UpdateOrAddKey(CurrentTime, SelectedValue);
+				OwningSection->SetEndTime(CurrentTime);
 			}
-			Sequencer->UpdateRuntimeInstances();
 		}
-	}
-}
 
-void SEnumCurveKeyEditor::OnComboMenuOpening()
-{
-	float CurrentTime = Sequencer->GetCurrentLocalTime( *Sequencer->GetFocusedMovieSceneSequence() );
-	int32 CurrentNameIndex = Enum->GetIndexByValue( Curve->Evaluate( CurrentTime ) );
-	TSharedPtr<int32> FoundNameIndexItem;
-	for ( int32 i = 0; i < VisibleEnumNameIndices.Num(); i++ )
-	{
-		if ( *VisibleEnumNameIndices[i] == CurrentNameIndex )
+		if (Curve->GetNumKeys() == 0)
 		{
-			FoundNameIndexItem = VisibleEnumNameIndices[i];
-			break;
+			Curve->SetDefaultValue(InSelectedItem);
 		}
-	}
-	if ( FoundNameIndexItem.IsValid() )
-	{
-		bUpdatingSelectionInternally = true;
-		EnumValuesCombo->SetSelectedItem(FoundNameIndexItem);
-		bUpdatingSelectionInternally = false;
+		else
+		{
+			Curve->UpdateOrAddKey(CurrentTime, InSelectedItem);
+		}
+
+		OnValueChangedEvent.ExecuteIfBound(InSelectedItem);
+
+		Sequencer->UpdateRuntimeInstances();
 	}
 }
 
diff --git a/Engine/Source/Editor/MovieSceneTools/Private/CurveKeyEditors/SEnumCurveKeyEditor.h b/Engine/Source/Editor/MovieSceneTools/Private/CurveKeyEditors/SEnumCurveKeyEditor.h
--- a/Engine/Source/Editor/MovieSceneTools/Private/CurveKeyEditors/SEnumCurveKeyEditor.h
+++ b/Engine/Source/Editor/MovieSceneTools/Private/CurveKeyEditors/SEnumCurveKeyEditor.h
@@ -8,31 +8,27 @@
 class SEnumCurveKeyEditor : public SCompoundWidget
 {
 public:
+	/** Notification for enum value change */
+	DECLARE_DELEGATE_OneParam( FOnValueChanged, int32 );
+
 	SLATE_BEGIN_ARGS(SEnumCurveKeyEditor) {}
 		SLATE_ARGUMENT(ISequencer*, Sequencer)
 		SLATE_ARGUMENT(UMovieSceneSection*, OwningSection)
 		SLATE_ARGUMENT(FIntegralCurve*, Curve)
+		SLATE_EVENT(FOnValueChanged, OnValueChanged)
 		SLATE_ARGUMENT(const UEnum*, Enum)
-		SLATE_ATTRIBUTE( TOptional<uint8>, IntermediateValue )
+		SLATE_ATTRIBUTE(TOptional<uint8>, IntermediateValue)
 	SLATE_END_ARGS();
 
 	void Construct(const FArguments& InArgs);
 
 private:
-	FText GetCurrentValue() const;
-
-	TSharedRef<SWidget> OnGenerateWidget(TSharedPtr<int32> InItem);
-
-	void OnComboSelectionChanged(TSharedPtr<int32> InSelectedItem, ESelectInfo::Type SelectInfo);
+	int32 OnGetCurrentValue() const;
 
-	void OnComboMenuOpening();
+	void OnComboSelectionChanged(int32 InSelectedItem, ESelectInfo::Type SelectInfo);
 
 	ISequencer* Sequencer;
 	UMovieSceneSection* OwningSection;
 	FIntegralCurve* Curve;
-	const UEnum* Enum;
-	TArray<TSharedPtr<int32>> VisibleEnumNameIndices;
-	TSharedPtr<SComboBox<TSharedPtr<int32>>> EnumValuesCombo;
-	bool bUpdatingSelectionInternally;
-	TAttribute<TOptional<uint8>> IntermediateValue;
+	FOnValueChanged OnValueChangedEvent;
 };
\ No newline at end of file
diff --git a/Engine/Source/Editor/MovieSceneTools/Private/CurveKeyEditors/SFloatCurveKeyEditor.cpp b/Engine/Source/Editor/MovieSceneTools/Private/CurveKeyEditors/SFloatCurveKeyEditor.cpp
--- a/Engine/Source/Editor/MovieSceneTools/Private/CurveKeyEditors/SFloatCurveKeyEditor.cpp
+++ b/Engine/Source/Editor/MovieSceneTools/Private/CurveKeyEditors/SFloatCurveKeyEditor.cpp
@@ -5,15 +5,31 @@
 
 #define LOCTEXT_NAMESPACE "FloatCurveKeyEditor"
 
+
+template<typename T>
+struct SNonThrottledSpinBox : SSpinBox<T>
+{
+	virtual FReply OnMouseButtonDown(const FGeometry& MyGeometry, const FPointerEvent& MouseEvent) override
+	{
+		FReply Reply = SSpinBox<T>::OnMouseButtonDown(MyGeometry, MouseEvent);
+		if (Reply.IsEventHandled())
+		{
+			Reply.PreventThrottling();
+		}
+		return Reply;
+	}
+};
+
 void SFloatCurveKeyEditor::Construct(const FArguments& InArgs)
 {
 	Sequencer = InArgs._Sequencer;
 	OwningSection = InArgs._OwningSection;
 	Curve = InArgs._Curve;
+	OnValueChangedEvent = InArgs._OnValueChanged;
 	IntermediateValue = InArgs._IntermediateValue;
 	ChildSlot
 	[
-		SNew(SSpinBox<float>)
+		SNew(SNonThrottledSpinBox<float>)
 		.Style(&FEditorStyle::GetWidgetStyle<FSpinBoxStyle>("Sequencer.HyperlinkSpinBox"))
 		.Font(FEditorStyle::GetFontStyle("Sequencer.AnimationOutliner.RegularFont"))
 		.MinValue(TOptional<float>())
@@ -88,6 +104,9 @@ void SFloatCurveKeyEditor::OnValueChanged(float Value)
 				OwningSection->SetEndTime(CurrentTime);
 			}
 		}
+
+		OnValueChangedEvent.ExecuteIfBound(Value);
+
 		Sequencer->UpdateRuntimeInstances();
 	}
 }
diff --git a/Engine/Source/Editor/MovieSceneTools/Private/CurveKeyEditors/SFloatCurveKeyEditor.h b/Engine/Source/Editor/MovieSceneTools/Private/CurveKeyEditors/SFloatCurveKeyEditor.h
--- a/Engine/Source/Editor/MovieSceneTools/Private/CurveKeyEditors/SFloatCurveKeyEditor.h
+++ b/Engine/Source/Editor/MovieSceneTools/Private/CurveKeyEditors/SFloatCurveKeyEditor.h
@@ -8,10 +8,14 @@
 class SFloatCurveKeyEditor : public SCompoundWidget
 {
 public:
+	/** Notification for numeric value change */
+	DECLARE_DELEGATE_OneParam( FOnValueChanged, float );
+
 	SLATE_BEGIN_ARGS(SFloatCurveKeyEditor) {}
 		SLATE_ARGUMENT(ISequencer*, Sequencer)
 		SLATE_ARGUMENT(UMovieSceneSection*, OwningSection)
 		SLATE_ARGUMENT(FRichCurve*, Curve)
+		SLATE_EVENT(FOnValueChanged, OnValueChanged)
 		SLATE_ATTRIBUTE(TOptional<float>, IntermediateValue)
 	SLATE_END_ARGS();
 
@@ -30,5 +34,6 @@ private:
 	ISequencer* Sequencer;
 	UMovieSceneSection* OwningSection;
 	FRichCurve* Curve;
+	FOnValueChanged OnValueChangedEvent;
 	TAttribute<TOptional<float>> IntermediateValue;
 };
\ No newline at end of file
diff --git a/Engine/Source/Editor/MovieSceneTools/Private/MovieSceneToolHelpers.cpp b/Engine/Source/Editor/MovieSceneTools/Private/MovieSceneToolHelpers.cpp
--- a/Engine/Source/Editor/MovieSceneTools/Private/MovieSceneToolHelpers.cpp
+++ b/Engine/Source/Editor/MovieSceneTools/Private/MovieSceneToolHelpers.cpp
@@ -199,8 +199,11 @@ FString MovieSceneToolHelpers::GenerateNewShotName(const TArray<UMovieSceneSecti
 		if (ParseShotName(BeforeShot->GetShotDisplayName().ToString(), BeforeShotPrefix, BeforeShotNumber, BeforeTakeNumber) &&
 			ParseShotName(NextShot->GetShotDisplayName().ToString(), NextShotPrefix, NextShotNumber, NextTakeNumber))
 		{
-			uint32 NewShotNumber = BeforeShotNumber + ( (NextShotNumber - BeforeShotNumber) / 2); // what if we can't find one? or conflicts with another?
-			return ComposeShotName(BeforeShotPrefix, NewShotNumber, ProjectSettings->FirstTakeNumber);
+			if (BeforeShotNumber < NextShotNumber)
+			{
+				uint32 NewShotNumber = BeforeShotNumber + ( (NextShotNumber - BeforeShotNumber) / 2); // what if we can't find one? or conflicts with another?
+				return ComposeShotName(BeforeShotPrefix, NewShotNumber, ProjectSettings->FirstTakeNumber);
+			}
 		}
 	}
 
@@ -220,9 +223,6 @@ void MovieSceneToolHelpers::GatherTakes(const UMovieSceneSection* Section, TArra
 	FAssetData ShotData(Shot->GetSequence()->GetOuter());
 
 	FString ShotPackagePath = ShotData.PackagePath.ToString();
-	int32 ShotLastSlashPos = INDEX_NONE;
-	ShotPackagePath.FindLastChar(TCHAR('/'), ShotLastSlashPos);
-	ShotPackagePath = ShotPackagePath.Left(ShotLastSlashPos);
 
 	FString ShotPrefix;
 	uint32 ShotNumber = INDEX_NONE;
@@ -238,9 +238,6 @@ void MovieSceneToolHelpers::GatherTakes(const UMovieSceneSection* Section, TArra
 		for (auto AssetObject : ObjectList)
 		{
 			FString AssetPackagePath = AssetObject.PackagePath.ToString();
-			int32 AssetLastSlashPos = INDEX_NONE;
-			AssetPackagePath.FindLastChar(TCHAR('/'), AssetLastSlashPos);
-			AssetPackagePath = AssetPackagePath.Left(AssetLastSlashPos);
 
 			if (AssetPackagePath == ShotPackagePath)
 			{
@@ -311,4 +308,118 @@ UObject* MovieSceneToolHelpers::GetTake(const UMovieSceneSection* Section, uint3
 	return nullptr;
 }
 
+class SEnumCombobox : public SComboBox<TSharedPtr<int32>>
+{
+public:
+	SLATE_BEGIN_ARGS(SEnumCombobox) {}
+
+	SLATE_ATTRIBUTE(TOptional<uint8>, IntermediateValue)
+	SLATE_ATTRIBUTE(int32, CurrentValue)
+	SLATE_ARGUMENT(FOnEnumSelectionChanged, OnEnumSelectionChanged)
+
+	SLATE_END_ARGS()
+
+	void Construct(const FArguments& InArgs, const UEnum* InEnum)
+	{
+		Enum = InEnum;
+		IntermediateValue = InArgs._IntermediateValue;
+		CurrentValue = InArgs._CurrentValue;
+		check(CurrentValue.IsBound());
+		OnEnumSelectionChangedDelegate = InArgs._OnEnumSelectionChanged;
+
+		bUpdatingSelectionInternally = false;
+
+		for (int32 i = 0; i < Enum->NumEnums() - 1; i++)
+		{
+			if (Enum->HasMetaData( TEXT("Hidden"), i ) == false)
+			{
+				VisibleEnumNameIndices.Add(MakeShareable(new int32(i)));
+			}
+		}
+
+		SComboBox::Construct(SComboBox<TSharedPtr<int32>>::FArguments()
+			.ButtonStyle(FEditorStyle::Get(), "FlatButton.Light")
+			.OptionsSource(&VisibleEnumNameIndices)
+			.OnGenerateWidget_Lambda([this](TSharedPtr<int32> InItem)
+			{
+				return SNew(STextBlock)
+					.Text(Enum->GetDisplayNameText(*InItem));
+			})
+			.OnSelectionChanged(this, &SEnumCombobox::OnComboSelectionChanged)
+			.OnComboBoxOpening(this, &SEnumCombobox::OnComboMenuOpening)
+			.ContentPadding(FMargin(2, 0))
+			[
+				SNew(STextBlock)
+				.Font(FEditorStyle::GetFontStyle("Sequencer.AnimationOutliner.RegularFont"))
+				.Text(this, &SEnumCombobox::GetCurrentValue)
+			]);
+	}
+
+private:
+	FText GetCurrentValue() const
+	{
+		if(IntermediateValue.IsSet() && IntermediateValue.Get().IsSet())
+		{
+			int32 IntermediateNameIndex = Enum->GetIndexByValue(IntermediateValue.Get().GetValue());
+			return Enum->GetDisplayNameText(IntermediateNameIndex);
+		}
 
+		int32 CurrentNameIndex = Enum->GetIndexByValue(CurrentValue.Get());
+		return Enum->GetDisplayNameText(CurrentNameIndex);
+	}
+
+	TSharedRef<SWidget> OnGenerateWidget(TSharedPtr<int32> InItem)
+	{
+		return SNew(STextBlock)
+			.Text(Enum->GetDisplayNameText(*InItem));
+	}
+
+	void OnComboSelectionChanged(TSharedPtr<int32> InSelectedItem, ESelectInfo::Type SelectInfo)
+	{
+		if (bUpdatingSelectionInternally == false)
+		{
+			OnEnumSelectionChangedDelegate.ExecuteIfBound(*InSelectedItem, SelectInfo);
+		}
+	}
+
+	void OnComboMenuOpening()
+	{
+		int32 CurrentNameIndex = Enum->GetIndexByValue(CurrentValue.Get());
+		TSharedPtr<int32> FoundNameIndexItem;
+		for ( int32 i = 0; i < VisibleEnumNameIndices.Num(); i++ )
+		{
+			if ( *VisibleEnumNameIndices[i] == CurrentNameIndex )
+			{
+				FoundNameIndexItem = VisibleEnumNameIndices[i];
+				break;
+			}
+		}
+		if ( FoundNameIndexItem.IsValid() )
+		{
+			bUpdatingSelectionInternally = true;
+			SetSelectedItem(FoundNameIndexItem);
+			bUpdatingSelectionInternally = false;
+		}
+	}	
+
+private:
+	const UEnum* Enum;
+
+	TAttribute<int32> CurrentValue;
+
+	TAttribute<TOptional<uint8>> IntermediateValue;
+
+	TArray<TSharedPtr<int32>> VisibleEnumNameIndices;
+
+	bool bUpdatingSelectionInternally;
+
+	FOnEnumSelectionChanged OnEnumSelectionChangedDelegate;
+};
+
+TSharedRef<SWidget> MovieSceneToolHelpers::MakeEnumComboBox(const UEnum* InEnum, TAttribute<int32> InCurrentValue, FOnEnumSelectionChanged InOnSelectionChanged, TAttribute<TOptional<uint8>> InIntermediateValue)
+{
+	return SNew(SEnumCombobox, InEnum)
+		.CurrentValue(InCurrentValue)
+		.OnEnumSelectionChanged(InOnSelectionChanged)
+		.IntermediateValue(InIntermediateValue);
+}
\ No newline at end of file
diff --git a/Engine/Source/Editor/MovieSceneTools/Private/MovieSceneToolsModule.cpp b/Engine/Source/Editor/MovieSceneTools/Private/MovieSceneToolsModule.cpp
--- a/Engine/Source/Editor/MovieSceneTools/Private/MovieSceneToolsModule.cpp
+++ b/Engine/Source/Editor/MovieSceneTools/Private/MovieSceneToolsModule.cpp
@@ -16,6 +16,7 @@
 #include "FloatPropertyTrackEditor.h"
 #include "VectorPropertyTrackEditor.h"
 #include "VisibilityPropertyTrackEditor.h"
+#include "ActorReferencePropertyTrackEditor.h"
 
 #include "TransformTrackEditor.h"
 #include "CameraCutTrackEditor.h"
@@ -33,6 +34,8 @@
 #include "FadeTrackEditor.h"
 #include "SpawnTrackEditor.h"
 #include "LevelVisibilityTrackEditor.h"
+#include "CameraAnimTrackEditor.h"
+#include "CameraShakeTrackEditor.h"
 
 #include "MovieSceneClipboard.h"
 #include "SequencerClipboardReconciler.h"
@@ -72,6 +75,7 @@ public:
 		FloatPropertyTrackCreateEditorHandle = SequencerModule.RegisterTrackEditor_Handle( FOnCreateTrackEditor::CreateStatic( &FFloatPropertyTrackEditor::CreateTrackEditor ) );
 		VectorPropertyTrackCreateEditorHandle = SequencerModule.RegisterTrackEditor_Handle( FOnCreateTrackEditor::CreateStatic( &FVectorPropertyTrackEditor::CreateTrackEditor ) );
 		VisibilityPropertyTrackCreateEditorHandle = SequencerModule.RegisterTrackEditor_Handle( FOnCreateTrackEditor::CreateStatic( &FVisibilityPropertyTrackEditor::CreateTrackEditor ) );
+		ActorReferencePropertyTrackCreateEditorHandle = SequencerModule.RegisterTrackEditor_Handle(FOnCreateTrackEditor::CreateStatic( &FActorReferencePropertyTrackEditor::CreateTrackEditor ) );
 
 		// register specialty track editors
 		AnimationTrackCreateEditorHandle = SequencerModule.RegisterTrackEditor_Handle( FOnCreateTrackEditor::CreateStatic( &FSkeletalAnimationTrackEditor::CreateTrackEditor ) );
@@ -90,6 +94,8 @@ public:
 		FadeTrackCreateEditorHandle = SequencerModule.RegisterTrackEditor_Handle( FOnCreateTrackEditor::CreateStatic( &FFadeTrackEditor::CreateTrackEditor ) );
 		SpawnTrackCreateEditorHandle = SequencerModule.RegisterTrackEditor_Handle( FOnCreateTrackEditor::CreateStatic( &FSpawnTrackEditor::CreateTrackEditor ) );
 		LevelVisibilityTrackCreateEditorHandle = SequencerModule.RegisterTrackEditor_Handle( FOnCreateTrackEditor::CreateStatic( &FLevelVisibilityTrackEditor::CreateTrackEditor ) );
+		CameraAnimTrackCreateEditorHandle = SequencerModule.RegisterTrackEditor_Handle(FOnCreateTrackEditor::CreateStatic(&FCameraAnimTrackEditor::CreateTrackEditor));
+		CameraShakeTrackCreateEditorHandle = SequencerModule.RegisterTrackEditor_Handle(FOnCreateTrackEditor::CreateStatic(&FCameraShakeTrackEditor::CreateTrackEditor));
 
 		RegisterClipboardConversions();
 	}
@@ -115,6 +121,7 @@ public:
 		SequencerModule.UnRegisterTrackEditor_Handle( FloatPropertyTrackCreateEditorHandle );
 		SequencerModule.UnRegisterTrackEditor_Handle( VectorPropertyTrackCreateEditorHandle );
 		SequencerModule.UnRegisterTrackEditor_Handle( VisibilityPropertyTrackCreateEditorHandle );
+		SequencerModule.UnRegisterTrackEditor_Handle( ActorReferencePropertyTrackCreateEditorHandle );
 
 		// unregister specialty track editors
 		SequencerModule.UnRegisterTrackEditor_Handle( AnimationTrackCreateEditorHandle );
@@ -132,7 +139,8 @@ public:
 		SequencerModule.UnRegisterTrackEditor_Handle( ComponentMaterialTrackCreateEditorHandle );
 		SequencerModule.UnRegisterTrackEditor_Handle( FadeTrackCreateEditorHandle );
 		SequencerModule.UnRegisterTrackEditor_Handle( SpawnTrackCreateEditorHandle );
-		SequencerModule.UnRegisterTrackEditor_Handle( LevelVisibilityTrackCreateEditorHandle );
+		SequencerModule.UnRegisterTrackEditor_Handle( CameraAnimTrackCreateEditorHandle );
+		SequencerModule.UnRegisterTrackEditor_Handle( CameraShakeTrackCreateEditorHandle );
 	}
 
 	void RegisterClipboardConversions()
@@ -178,6 +186,7 @@ private:
 	FDelegateHandle FloatPropertyTrackCreateEditorHandle;
 	FDelegateHandle VectorPropertyTrackCreateEditorHandle;
 	FDelegateHandle VisibilityPropertyTrackCreateEditorHandle;
+	FDelegateHandle ActorReferencePropertyTrackCreateEditorHandle;
 
 	FDelegateHandle AnimationTrackCreateEditorHandle;
 	FDelegateHandle AttachTrackCreateEditorHandle;
@@ -195,6 +204,8 @@ private:
 	FDelegateHandle FadeTrackCreateEditorHandle;
 	FDelegateHandle SpawnTrackCreateEditorHandle;
 	FDelegateHandle LevelVisibilityTrackCreateEditorHandle;
+	FDelegateHandle CameraAnimTrackCreateEditorHandle;
+	FDelegateHandle CameraShakeTrackCreateEditorHandle;
 };
 
 
diff --git a/Engine/Source/Editor/MovieSceneTools/Private/MovieSceneToolsProjectSettings.cpp b/Engine/Source/Editor/MovieSceneTools/Private/MovieSceneToolsProjectSettings.cpp
--- a/Engine/Source/Editor/MovieSceneTools/Private/MovieSceneToolsProjectSettings.cpp
+++ b/Engine/Source/Editor/MovieSceneTools/Private/MovieSceneToolsProjectSettings.cpp
@@ -7,7 +7,8 @@
 UMovieSceneToolsProjectSettings::UMovieSceneToolsProjectSettings()
 	: DefaultStartTime(0.f)
 	, DefaultDuration(5.f)
-	, ShotPrefix(TEXT("sht"))
+	, ShotDirectory(TEXT("shots"))
+	, ShotPrefix(TEXT("shot"))
 	, FirstShotNumber(10)
 	, ShotIncrement(10)
 	, ShotNumDigits(4)
diff --git a/Engine/Source/Editor/MovieSceneTools/Private/MovieSceneToolsUserSettings.cpp b/Engine/Source/Editor/MovieSceneTools/Private/MovieSceneToolsUserSettings.cpp
new file mode 100644
--- /dev/null
+++ b/Engine/Source/Editor/MovieSceneTools/Private/MovieSceneToolsUserSettings.cpp
@@ -0,0 +1,26 @@
+// Copyright 1998-2016 Epic Games, Inc. All Rights Reserved.
+
+#include "MovieSceneToolsPrivatePCH.h"
+#include "MovieSceneToolsUserSettings.h"
+
+#include "DetailLayoutBuilder.h"
+#include "DetailCategoryBuilder.h"
+#include "DetailWidgetRow.h"
+#include "IDetailPropertyRow.h"
+
+UMovieSceneUserThumbnailSettings::UMovieSceneUserThumbnailSettings(const FObjectInitializer& Initializer)
+	: Super(Initializer)
+{
+	ThumbnailSize = FIntPoint(128, 72);
+	bDrawThumbnails = true;
+}
+
+void UMovieSceneUserThumbnailSettings::PostEditChangeProperty(FPropertyChangedEvent& PropertyChangedEvent)
+{
+	Super::PostEditChangeProperty(PropertyChangedEvent);
+
+	ThumbnailSize.X = FMath::Clamp(ThumbnailSize.X, 1, 1024);
+	ThumbnailSize.Y = FMath::Clamp(ThumbnailSize.Y, 1, 1024);
+
+	SaveConfig();
+}
\ No newline at end of file
diff --git a/Engine/Source/Editor/MovieSceneTools/Private/Sections/ActorReferencePropertySection.cpp b/Engine/Source/Editor/MovieSceneTools/Private/Sections/ActorReferencePropertySection.cpp
new file mode 100644
--- /dev/null
+++ b/Engine/Source/Editor/MovieSceneTools/Private/Sections/ActorReferencePropertySection.cpp
@@ -0,0 +1,21 @@
+// Copyright 1998-2016 Epic Games, Inc. All Rights Reserved.
+
+#include "MovieSceneToolsPrivatePCH.h"
+#include "ActorReferencePropertySection.h"
+#include "MovieSceneActorReferenceSection.h"
+
+void FActorReferencePropertySection::GenerateSectionLayout( class ISectionLayoutBuilder& LayoutBuilder ) const
+{
+	UMovieSceneActorReferenceSection* ActorReferenceSection = Cast<UMovieSceneActorReferenceSection>( &SectionObject );
+	KeyArea = MakeShareable( new FActorReferenceKeyArea( ActorReferenceSection->GetActorReferenceCurve(), ActorReferenceSection ) );
+	LayoutBuilder.SetSectionAsKeyArea( KeyArea.ToSharedRef() );
+}
+
+void FActorReferencePropertySection::SetIntermediateValue( FPropertyChangedParams PropertyChangedParams )
+{
+}
+
+void FActorReferencePropertySection::ClearIntermediateValue()
+{
+	KeyArea->ClearIntermediateValue();
+}
\ No newline at end of file
diff --git a/Engine/Source/Editor/MovieSceneTools/Private/Sections/ActorReferencePropertySection.h b/Engine/Source/Editor/MovieSceneTools/Private/Sections/ActorReferencePropertySection.h
new file mode 100644
--- /dev/null
+++ b/Engine/Source/Editor/MovieSceneTools/Private/Sections/ActorReferencePropertySection.h
@@ -0,0 +1,51 @@
+// Copyright 1998-2016 Epic Games, Inc. All Rights Reserved.
+
+#pragma once
+
+#include "ByteKeyArea.h"
+
+
+/** A key area for displaying and editing an actor reference property. */
+class FActorReferenceKeyArea
+	: public FIntegralKeyArea < int32 >
+{
+public:
+
+	FActorReferenceKeyArea( FIntegralCurve& InCurve, UMovieSceneSection* InOwningSection )
+		: FIntegralKeyArea( InCurve, InOwningSection )
+	{ }
+
+public:
+
+	// FIntegralKeyArea overrides
+
+	virtual bool CanCreateKeyEditor() override { return false; }
+	virtual TSharedRef<SWidget> CreateKeyEditor( ISequencer* Sequencer ) override { return SNew( SSpacer ); }
+
+};
+
+
+/**
+ * A property section for actor references.
+ */
+class FActorReferencePropertySection
+	: public FPropertySection
+{
+public:
+
+	FActorReferencePropertySection(UMovieSceneSection& InSectionObject, const FText& SectionName)
+		: FPropertySection(InSectionObject, SectionName)
+	{ }
+
+public:
+
+	// FPropertySection interface 
+
+	virtual void GenerateSectionLayout(class ISectionLayoutBuilder& LayoutBuilder) const override;
+	virtual void SetIntermediateValue(FPropertyChangedParams PropertyChangedParams) override;
+	virtual void ClearIntermediateValue() override;
+
+private:
+
+	mutable TSharedPtr<FActorReferenceKeyArea> KeyArea;
+};
diff --git a/Engine/Source/Editor/MovieSceneTools/Private/Sections/BoolPropertySection.cpp b/Engine/Source/Editor/MovieSceneTools/Private/Sections/BoolPropertySection.cpp
--- a/Engine/Source/Editor/MovieSceneTools/Private/Sections/BoolPropertySection.cpp
+++ b/Engine/Source/Editor/MovieSceneTools/Private/Sections/BoolPropertySection.cpp
@@ -19,28 +19,25 @@ int32 FBoolPropertySection::OnPaintSection( FSequencerSectionPainter& Painter )
 
 	TArray<float> SectionSwitchTimes;
 
-	if ( SectionObject.IsInfinite() && Sequencer->GetViewRange().GetLowerBoundValue() < SectionObject.GetStartTime() )
-	{
-		SectionSwitchTimes.Add( Sequencer->GetViewRange().GetLowerBoundValue() );
-	}
-	SectionSwitchTimes.Add( SectionObject.GetStartTime() );
+	const FTimeToPixel& TimeConverter = Painter.GetTimeConverter();
+
+	// Add the start time
+	const float StartTime = TimeConverter.PixelToTime(0.f);
+	const float EndTime = TimeConverter.PixelToTime(Painter.SectionGeometry.GetLocalSize().X);
+	
+	SectionSwitchTimes.Add(StartTime);
 
 	FIntegralCurve& BoolCurve = BoolSection->GetCurve();
 	for ( TArray<FIntegralKey>::TConstIterator It( BoolCurve.GetKeyIterator() ); It; ++It )
 	{
 		float Time = It->Time;
-		if ( Time > SectionObject.GetStartTime() && Time < SectionObject.GetEndTime() )
+		if ( Time > StartTime && Time < EndTime )
 		{
 			SectionSwitchTimes.Add( Time );
 		}
 	}
 
-	if ( SectionObject.IsInfinite() && Sequencer->GetViewRange().GetUpperBoundValue() > SectionObject.GetEndTime() )
-	{
-		SectionSwitchTimes.Add( Sequencer->GetViewRange().GetUpperBoundValue() );
-	}
-	SectionSwitchTimes.Add( SectionObject.GetEndTime() );
-
+	SectionSwitchTimes.Add(EndTime);
 
 	int32 LayerId = Painter.PaintSectionBackground();
 
@@ -48,8 +45,6 @@ int32 FBoolPropertySection::OnPaintSection( FSequencerSectionPainter& Painter )
 		? ESlateDrawEffect::None
 		: ESlateDrawEffect::DisabledEffect;
 
-	const FTimeToPixel& TimeConverter = Painter.GetTimeConverter();
-
 	static const int32 Height = 5;
 	const float VerticalOffset = Painter.SectionGeometry.GetLocalSize().Y * .5f - Height * .5f;
 
@@ -58,8 +53,8 @@ int32 FBoolPropertySection::OnPaintSection( FSequencerSectionPainter& Painter )
 	for ( int32 i = 0; i < SectionSwitchTimes.Num() - 1; ++i )
 	{
 		float ThisTime = SectionSwitchTimes[i];
-
-		const FColor Color = BoolSection->Eval(ThisTime) ? FColor(0, 255, 0, 127) : FColor(255, 0, 0, 127);
+		
+		const FColor Color = BoolSection->Eval(ThisTime) ? FColor(0, 255, 0, 200) : FColor(255, 0, 0, 200);
 		
 		FVector2D StartPos(TimeConverter.TimeToPixel(ThisTime), VerticalOffset);
 		FVector2D Size(TimeConverter.TimeToPixel(SectionSwitchTimes[i+1]) - StartPos.X, Height);
diff --git a/Engine/Source/Editor/MovieSceneTools/Private/Sections/CameraCutSection.cpp b/Engine/Source/Editor/MovieSceneTools/Private/Sections/CameraCutSection.cpp
--- a/Engine/Source/Editor/MovieSceneTools/Private/Sections/CameraCutSection.cpp
+++ b/Engine/Source/Editor/MovieSceneTools/Private/Sections/CameraCutSection.cpp
@@ -5,7 +5,7 @@
 #include "Runtime/MovieSceneTracks/Public/Sections/MovieSceneCameraCutSection.h"
 #include "Runtime/Engine/Public/Slate/SceneViewport.h"
 #include "SInlineEditableTextBlock.h"
-
+#include "MovieSceneToolsUserSettings.h"
 
 #define LOCTEXT_NAMESPACE "FCameraCutSection"
 
@@ -25,8 +25,37 @@ FCameraCutSection::~FCameraCutSection()
 /* ISequencerSection interface
  *****************************************************************************/
 
+void FCameraCutSection::SetSingleTime(float GlobalTime)
+{
+	UMovieSceneCameraCutSection* CameraCutSection = Cast<UMovieSceneCameraCutSection>(Section);
+	if (CameraCutSection)
+	{
+		CameraCutSection->SetThumbnailReferenceOffset(GlobalTime - CameraCutSection->GetStartTime());
+	}
+}
+
+void FCameraCutSection::Tick(const FGeometry& AllottedGeometry, const FGeometry& ClippedGeometry, const double InCurrentTime, const float InDeltaTime)
+{
+	UMovieSceneCameraCutSection* CameraCutSection = Cast<UMovieSceneCameraCutSection>(Section);
+	if (CameraCutSection)
+	{
+		if (GetDefault<UMovieSceneUserThumbnailSettings>()->bDrawSingleThumbnails)
+		{
+			ThumbnailCache.SetSingleReferenceFrame(CameraCutSection->GetThumbnailReferenceOffset() + CameraCutSection->GetStartTime());
+		}
+		else
+		{
+			ThumbnailCache.SetSingleReferenceFrame(TOptional<float>());
+		}
+	}
+
+	FThumbnailSection::Tick(AllottedGeometry, ClippedGeometry, InCurrentTime, InDeltaTime);
+}
+
 void FCameraCutSection::BuildSectionContextMenu(FMenuBuilder& MenuBuilder, const FGuid& ObjectBinding)
 {
+	FThumbnailSection::BuildSectionContextMenu(MenuBuilder, ObjectBinding);
+
 	UWorld* World = GEditor->GetEditorWorldContext().World();
 
 	if (World == nullptr)
@@ -34,15 +63,16 @@ void FCameraCutSection::BuildSectionContextMenu(FMenuBuilder& MenuBuilder, const
 		return;
 	}
 
+	const AActor* CameraActor = GetCameraForFrame(Section->GetStartTime());
+
 	// get list of available cameras
 	TArray<AActor*> AllCameras;
-	TArray<AActor*> SelectedCameras;
 
 	for (FActorIterator ActorIt(World); ActorIt; ++ActorIt)
 	{
 		AActor* Actor = *ActorIt;
 
-		if ((Actor != GetCameraObject()) && Actor->IsListedInSceneOutliner())
+		if ((Actor != CameraActor) && Actor->IsListedInSceneOutliner())
 		{
 			UCameraComponent* CameraComponent = MovieSceneHelpers::CameraComponentFromActor(Actor);
 			if (CameraComponent)
@@ -84,41 +114,31 @@ FText FCameraCutSection::GetDisplayName() const
 /* FThumbnailSection interface
  *****************************************************************************/
 
-AActor* FCameraCutSection::GetCameraObject() const
+const AActor* FCameraCutSection::GetCameraForFrame(float Time) const
 {
 	UMovieSceneCameraCutSection* CameraCutSection = Cast<UMovieSceneCameraCutSection>(Section);
+	TSharedPtr<ISequencer> Sequencer = SequencerPtr.Pin();
 
-	TArray<UObject*> OutObjects;
-	// @todo sequencer: the director track may be able to get cameras from sub-movie scenes
-	SequencerPtr.Pin()->GetRuntimeObjects(SequencerPtr.Pin()->GetRootMovieSceneSequenceInstance(),  CameraCutSection->GetCameraGuid(), OutObjects);
-
-	AActor* ReturnCam = nullptr;
-
-	if (OutObjects.Num() > 0)
+	if (CameraCutSection && Sequencer.IsValid())
 	{
-		ReturnCam = Cast<AActor>(OutObjects[0]);
+		TSharedRef<FMovieSceneSequenceInstance> SequenceInstance = Sequencer->GetFocusedMovieSceneSequenceInstance();
 
-		if (ReturnCam != nullptr)
+		AActor* Actor = Cast<AActor>(SequenceInstance->FindObject(CameraCutSection->GetCameraGuid(), *Sequencer));
+		if (Actor)
 		{
-			return ReturnCam;
+			return Actor;
 		}
-	}
-
-	// Otherwise look in the current movie scene
-	OutObjects.Empty();
-	SequencerPtr.Pin()->GetRuntimeObjects(SequencerPtr.Pin()->GetFocusedMovieSceneSequenceInstance(), CameraCutSection->GetCameraGuid(), OutObjects);
-
-	if (OutObjects.Num() > 0)
-	{
-		ReturnCam = Cast<AActor>(OutObjects[0]);
-
-		if (ReturnCam != nullptr)
+		else
 		{
-			return ReturnCam;
+			FMovieSceneSpawnable* Spawnable = SequenceInstance->GetSequence()->GetMovieScene()->FindSpawnable(CameraCutSection->GetCameraGuid());
+			if (Spawnable)
+			{
+				return GetDefault<AActor>(Spawnable->GetClass());
+			}
 		}
 	}
 
-	return ReturnCam;
+	return nullptr;
 }
 
 float FCameraCutSection::GetSectionHeight() const
@@ -128,7 +148,7 @@ float FCameraCutSection::GetSectionHeight() const
 
 FMargin FCameraCutSection::GetContentPadding() const
 {
-	return FMargin(6.f, 12.f);
+	return FMargin(6.f, 10.f);
 }
 
 int32 FCameraCutSection::OnPaintSection(FSequencerSectionPainter& InPainter) const
@@ -141,7 +161,7 @@ int32 FCameraCutSection::OnPaintSection(FSequencerSectionPainter& InPainter) con
 
 FText FCameraCutSection::HandleThumbnailTextBlockText() const
 {
-	AActor* CameraActor = GetCameraObject();
+	const AActor* CameraActor = GetCameraForFrame(Section->GetStartTime());
 	if (CameraActor)
 	{
 		return FText::FromString(CameraActor->GetActorLabel());
diff --git a/Engine/Source/Editor/MovieSceneTools/Private/Sections/CameraCutSection.h b/Engine/Source/Editor/MovieSceneTools/Private/Sections/CameraCutSection.h
--- a/Engine/Source/Editor/MovieSceneTools/Private/Sections/CameraCutSection.h
+++ b/Engine/Source/Editor/MovieSceneTools/Private/Sections/CameraCutSection.h
@@ -2,7 +2,6 @@
 
 #pragma once
 
-class FThumbnailSection;
 class FTrackEditorThumbnail;
 class FTrackEditorThumbnailPool;
 class IMenu;
@@ -27,7 +26,7 @@ public:
 public:
 
 	// ISequencerSection interface
-
+	virtual void Tick(const FGeometry& AllottedGeometry, const FGeometry& ClippedGeometry, const double InCurrentTime, const float InDeltaTime) override;
 	virtual void BuildSectionContextMenu(FMenuBuilder& MenuBuilder, const FGuid& ObjectBinding) override;
 	virtual FText GetDisplayName() const override;
 	virtual float GetSectionHeight() const override;
@@ -35,11 +34,14 @@ public:
 	virtual FMargin GetContentPadding() const override;
 	// FThumbnail interface
 
-	virtual AActor* GetCameraObject() const override;
+	virtual void SetSingleTime(float GlobalTime) override;
 	virtual FText HandleThumbnailTextBlockText() const override;
 
 private:
 
+	/** Get a representative camera for the given time */
+	const AActor* GetCameraForFrame(float Time) const;
+
 	/** Callback for executing a "Set Camera" menu entry in the context menu. */
 	void HandleSetCameraMenuEntryExecute(AActor* InCamera);
 
diff --git a/Engine/Source/Editor/MovieSceneTools/Private/Sections/CinematicShotSection.cpp b/Engine/Source/Editor/MovieSceneTools/Private/Sections/CinematicShotSection.cpp
--- a/Engine/Source/Editor/MovieSceneTools/Private/Sections/CinematicShotSection.cpp
+++ b/Engine/Source/Editor/MovieSceneTools/Private/Sections/CinematicShotSection.cpp
@@ -6,11 +6,21 @@
 #include "CinematicShotTrackEditor.h"
 #include "Runtime/Engine/Public/Slate/SceneViewport.h"
 #include "SInlineEditableTextBlock.h"
-#include "CinematicShotTrackEditor.h"
 #include "MovieSceneToolHelpers.h"
+#include "MovieSceneToolsUserSettings.h"
 
 #define LOCTEXT_NAMESPACE "FCinematicShotSection"
 
+FCinematicShotSection::FCinematicSectionCache::FCinematicSectionCache(UMovieSceneCinematicShotSection* Section)
+	: ActualStartTime(0.f), TimeScale(0.f)
+{
+	if (Section)
+	{
+		ActualStartTime = Section->GetStartTime() - Section->StartOffset;
+		TimeScale = Section->TimeScale;
+	}
+}
+
 
 /* FCinematicShotSection structors
  *****************************************************************************/
@@ -20,6 +30,9 @@ FCinematicShotSection::FCinematicShotSection(TSharedPtr<ISequencer> InSequencer,
 	, SectionObject(*CastChecked<UMovieSceneCinematicShotSection>(&InSection))
 	, Sequencer(InSequencer)
 	, CinematicShotTrackEditor(InCinematicShotTrackEditor)
+	, InitialStartOffsetDuringResize(0.f)
+	, InitialStartTimeDuringResize(0.f)
+	, ThumbnailCacheData(&SectionObject)
 {
 	if (InSequencer->HasSequenceInstanceForSection(InSection))
 	{
@@ -34,7 +47,7 @@ FCinematicShotSection::~FCinematicShotSection()
 
 float FCinematicShotSection::GetSectionHeight() const
 {
-	return FThumbnailSection::GetSectionHeight() + 2*13.f;
+	return FThumbnailSection::GetSectionHeight() + 2*9.f;
 }
 
 FMargin FCinematicShotSection::GetContentPadding() const
@@ -42,6 +55,57 @@ FMargin FCinematicShotSection::GetContentPadding() const
 	return FMargin(8.f, 15.f);
 }
 
+void FCinematicShotSection::SetSingleTime(float GlobalTime)
+{
+	SectionObject.SetThumbnailReferenceOffset(GlobalTime - SectionObject.GetStartTime());
+}
+
+void FCinematicShotSection::BeginResizeSection()
+{
+	InitialStartOffsetDuringResize = SectionObject.StartOffset;
+	InitialStartTimeDuringResize = SectionObject.GetStartTime();
+}
+
+void FCinematicShotSection::ResizeSection(ESequencerSectionResizeMode ResizeMode, float ResizeTime)
+{
+	// Adjust the start offset when resizing from the beginning
+	if (ResizeMode == SSRM_LeadingEdge)
+	{
+		float StartOffset = (ResizeTime - InitialStartTimeDuringResize) / SectionObject.TimeScale;
+		StartOffset += InitialStartOffsetDuringResize;
+
+		// Ensure start offset is not less than 0
+		StartOffset = FMath::Max(StartOffset, 0.f);
+
+		SectionObject.StartOffset = StartOffset;
+	}
+
+	FThumbnailSection::ResizeSection(ResizeMode, ResizeTime);
+}
+
+void FCinematicShotSection::Tick(const FGeometry& AllottedGeometry, const FGeometry& ClippedGeometry, const double InCurrentTime, const float InDeltaTime)
+{
+	// Set cached data
+	FCinematicSectionCache NewCacheData(&SectionObject);
+	if (NewCacheData != ThumbnailCacheData)
+	{
+		ThumbnailCache.ForceRedraw();
+	}
+	ThumbnailCacheData = NewCacheData;
+
+	// Update single reference frame settings
+	if (GetDefault<UMovieSceneUserThumbnailSettings>()->bDrawSingleThumbnails)
+	{
+		ThumbnailCache.SetSingleReferenceFrame(SectionObject.GetStartTime() + SectionObject.GetThumbnailReferenceOffset());
+	}
+	else
+	{
+		ThumbnailCache.SetSingleReferenceFrame(TOptional<float>());
+	}
+
+	FThumbnailSection::Tick(AllottedGeometry, ClippedGeometry, InCurrentTime, InDeltaTime);
+}
+
 int32 FCinematicShotSection::OnPaintSection(FSequencerSectionPainter& InPainter) const
 {
 	static const FSlateBrush* FilmBorder = FEditorStyle::GetBrush("Sequencer.Section.FilmBorder");
@@ -78,6 +142,8 @@ int32 FCinematicShotSection::OnPaintSection(FSequencerSectionPainter& InPainter)
 
 void FCinematicShotSection::BuildSectionContextMenu(FMenuBuilder& MenuBuilder, const FGuid& ObjectBinding)
 {
+	FThumbnailSection::BuildSectionContextMenu(MenuBuilder, ObjectBinding);
+
 	MenuBuilder.BeginSection(NAME_None, LOCTEXT("ShotMenuText", "Shot"));
 	{
 		if (SequenceInstance.IsValid())
@@ -150,11 +216,6 @@ FText FCinematicShotSection::GetDisplayName() const
 /* FCinematicShotSection callbacks
  *****************************************************************************/
 
-AActor* FCinematicShotSection::GetCameraObject() const
-{
-	return CinematicShotTrackEditor.IsValid() ? CinematicShotTrackEditor.Pin()->GetCinematicShotCamera().Get() : nullptr;
-}
-
 FText FCinematicShotSection::HandleThumbnailTextBlockText() const
 {
 	return SectionObject.GetShotDisplayName();
diff --git a/Engine/Source/Editor/MovieSceneTools/Private/Sections/CinematicShotSection.h b/Engine/Source/Editor/MovieSceneTools/Private/Sections/CinematicShotSection.h
--- a/Engine/Source/Editor/MovieSceneTools/Private/Sections/CinematicShotSection.h
+++ b/Engine/Source/Editor/MovieSceneTools/Private/Sections/CinematicShotSection.h
@@ -30,15 +30,18 @@ public:
 
 	// ISequencerSection interface
 
+	virtual void Tick(const FGeometry& AllottedGeometry, const FGeometry& ClippedGeometry, const double InCurrentTime, const float InDeltaTime) override;
 	virtual int32 OnPaintSection( FSequencerSectionPainter& Painter ) const override;
 	virtual void BuildSectionContextMenu(FMenuBuilder& MenuBuilder, const FGuid& ObjectBinding) override;
 	virtual FText GetDisplayName() const override;
 	virtual FReply OnSectionDoubleClicked(const FGeometry& SectionGeometry, const FPointerEvent& MouseEvent) override;
 	virtual float GetSectionHeight() const override;
 	virtual FMargin GetContentPadding() const override;
-	// FThumbnail interface
+	virtual void BeginResizeSection() override;
+	virtual void ResizeSection(ESequencerSectionResizeMode ResizeMode, float ResizeTime) override;
 
-	virtual AActor* GetCameraObject() const override;
+	// FThumbnail interface
+	virtual void SetSingleTime(float GlobalTime) override;
 	virtual bool CanRename() const override { return true; }
 	virtual FText HandleThumbnailTextBlockText() const override;
 	virtual void HandleThumbnailTextBlockTextCommitted(const FText& NewThumbnailName, ETextCommit::Type CommitType) override;
@@ -61,4 +64,26 @@ private:
 
 	/** The cinematic shot track editor that contains this section */
 	TWeakPtr<FCinematicShotTrackEditor> CinematicShotTrackEditor;
+
+	/** Cached start offset value valid only during resize */
+	float InitialStartOffsetDuringResize;
+	
+	/** Cached start time valid only during resize */
+	float InitialStartTimeDuringResize;
+
+	struct FCinematicSectionCache
+	{
+		FCinematicSectionCache(UMovieSceneCinematicShotSection* Section = nullptr);
+
+		bool operator!=(const FCinematicSectionCache& RHS) const
+		{
+			return !FMath::IsNearlyEqual(ActualStartTime, RHS.ActualStartTime, 0.001f) || TimeScale != RHS.TimeScale;
+		}
+
+		float ActualStartTime;
+		float TimeScale;
+	};
+
+	/** Cached section thumbnail data */
+	FCinematicSectionCache ThumbnailCacheData;
 };
diff --git a/Engine/Source/Editor/MovieSceneTools/Private/Sections/ColorPropertySection.cpp b/Engine/Source/Editor/MovieSceneTools/Private/Sections/ColorPropertySection.cpp
--- a/Engine/Source/Editor/MovieSceneTools/Private/Sections/ColorPropertySection.cpp
+++ b/Engine/Source/Editor/MovieSceneTools/Private/Sections/ColorPropertySection.cpp
@@ -25,21 +25,24 @@ void FColorPropertySection::GenerateSectionLayout( class ISectionLayoutBuilder&
 
 int32 FColorPropertySection::OnPaintSection( FSequencerSectionPainter& Painter ) const
 {
-	int32 LayerId = Painter.PaintSectionBackground(FColor(255,255,255));
+	int32 LayerId = Painter.PaintSectionBackground();
 
 	const ESlateDrawEffect::Type DrawEffects = Painter.bParentEnabled ? ESlateDrawEffect::None : ESlateDrawEffect::DisabledEffect;
 
 	const UMovieSceneColorSection* ColorSection = Cast<const UMovieSceneColorSection>( &SectionObject );
 
-	float StartTime = ColorSection->GetStartTime();
-	float EndTime = ColorSection->GetEndTime();
-	float SectionDuration = EndTime - StartTime;
+	const FTimeToPixel& TimeConverter = Painter.GetTimeConverter();
+
+	const float StartTime = TimeConverter.PixelToTime(0.f);
+	const float EndTime = TimeConverter.PixelToTime(Painter.SectionGeometry.GetLocalSize().X);
+	const float SectionDuration = EndTime - StartTime;
 
 	if ( !FMath::IsNearlyZero( SectionDuration ) )
 	{
-		FVector2D GradientSize = FVector2D( Painter.SectionGeometry.Size.X, (Painter.SectionGeometry.Size.Y / 4) - 3.0f );
+		FVector2D GradientSize = FVector2D( Painter.SectionGeometry.Size.X - 2.f, (Painter.SectionGeometry.Size.Y / 4) - 3.0f );
 
-		FPaintGeometry PaintGeometry = Painter.SectionGeometry.ToPaintGeometry( FVector2D( 0, 0 ), GradientSize );
+		FPaintGeometry PaintGeometry = Painter.SectionGeometry.ToPaintGeometry( FVector2D( 1.f, 1.f ), GradientSize );
+		FSlateRect ClippingRect = Painter.SectionClippingRect.InsetBy(1.f);
 
 		// If we are showing a background pattern and the colors is transparent, draw a checker pattern
 		FSlateDrawElement::MakeBox(
@@ -47,7 +50,7 @@ int32 FColorPropertySection::OnPaintSection( FSequencerSectionPainter& Painter )
 			LayerId,
 			PaintGeometry,
 			FEditorStyle::GetBrush( "Checker" ),
-			Painter.SectionClippingRect,
+			ClippingRect,
 			DrawEffects);
 
 		TArray<FSlateGradientStop> GradientStops;
@@ -73,7 +76,7 @@ int32 FColorPropertySection::OnPaintSection( FSequencerSectionPainter& Painter )
 				PaintGeometry,
 				GradientStops,
 				Orient_Vertical,
-				Painter.SectionClippingRect,
+				ClippingRect,
 				DrawEffects
 				);
 		}
diff --git a/Engine/Source/Editor/MovieSceneTools/Private/Sections/ThumbnailSection.cpp b/Engine/Source/Editor/MovieSceneTools/Private/Sections/ThumbnailSection.cpp
--- a/Engine/Source/Editor/MovieSceneTools/Private/Sections/ThumbnailSection.cpp
+++ b/Engine/Source/Editor/MovieSceneTools/Private/Sections/ThumbnailSection.cpp
@@ -2,9 +2,10 @@
 
 #include "MovieSceneToolsPrivatePCH.h"
 #include "ISectionLayoutBuilder.h"
-#include "Runtime/Engine/Public/Slate/SceneViewport.h"
 #include "SInlineEditableTextBlock.h"
-
+#include "PropertyEditorModule.h"
+#include "IDetailsView.h"
+#include "MovieSceneToolsUserSettings.h"
 
 #define LOCTEXT_NAMESPACE "FThumbnailSection"
 
@@ -23,137 +24,53 @@ namespace ThumbnailSectionConstants
 FThumbnailSection::FThumbnailSection(TSharedPtr<ISequencer> InSequencer, TSharedPtr<FTrackEditorThumbnailPool> InThumbnailPool, UMovieSceneSection& InSection)
 	: Section(&InSection)
 	, SequencerPtr(InSequencer)
-	, ThumbnailPool(InThumbnailPool)
-	, Thumbnails()
-	, ThumbnailWidth(0)
-	, StoredSize(ForceInit)
-	, StoredStartTime(0.f)
-	, VisibleTimeRange(TRange<float>::Empty())
-	, InternalViewportScene(nullptr)
-	, InternalViewportClient(nullptr)
+	, ThumbnailCache(InThumbnailPool, this)
 {
-	// @todo Sequencer optimize This code shouldn't even be called if this is offscreen
-	// the following code could be generated on demand when the section is scrolled onscreen
-	InternalViewportClient = MakeShareable(new FLevelEditorViewportClient(nullptr));
-	{
-		InternalViewportClient->ViewportType = LVT_Perspective;
-		InternalViewportClient->bDisableInput = true;
-		InternalViewportClient->bDrawAxes = false;
-		InternalViewportClient->EngineShowFlags = FEngineShowFlags(ESFIM_Game);
-		InternalViewportClient->EngineShowFlags.DisableAdvancedFeatures();
-		InternalViewportClient->SetRealtime( false );
-	}
-
-	InternalViewportScene = MakeShareable(new FSceneViewport(InternalViewportClient.Get(), nullptr));
-	InternalViewportClient->Viewport = InternalViewportScene.Get();
-
 	WhiteBrush = FEditorStyle::GetBrush("WhiteBrush");
+	GetMutableDefault<UMovieSceneUserThumbnailSettings>()->OnForceRedraw().AddRaw(this, &FThumbnailSection::RedrawThumbnails);
 }
 
-
 FThumbnailSection::~FThumbnailSection()
 {
-	if (ThumbnailPool.IsValid())
-	{
-		ThumbnailPool.Pin()->RemoveThumbnailsNeedingRedraw(Thumbnails);
-	}
-	
-	if (InternalViewportClient.IsValid())
-	{
-		InternalViewportClient->Viewport = nullptr;
-	}
+	GetMutableDefault<UMovieSceneUserThumbnailSettings>()->OnForceRedraw().RemoveAll(this);
 }
 
+void FThumbnailSection::RedrawThumbnails()
+{
+	ThumbnailCache.ForceRedraw();
+}
 
 /* FThumbnailSection interface
  *****************************************************************************/
 
-void FThumbnailSection::DrawViewportThumbnail(TSharedPtr<FTrackEditorThumbnail> TrackEditorThumbnail)
+void FThumbnailSection::PreDraw(FTrackEditorThumbnail& Thumbnail, FLevelEditorViewportClient& ViewportClient, FSceneViewport& SceneViewport)
 {
-	if (InternalViewportScene.IsValid())
+	TSharedPtr<ISequencer> Sequencer = SequencerPtr.Pin();
+	if (Sequencer.IsValid())
 	{
-		EMovieScenePlayerStatus::Type PlaybackStatus = SequencerPtr.Pin()->GetPlaybackStatus();
-		SequencerPtr.Pin()->SetPlaybackStatus(EMovieScenePlayerStatus::Jumping);
-		SequencerPtr.Pin()->SetGlobalTimeDirectly(TrackEditorThumbnail->GetTime());
-		InternalViewportClient->SetActorLock(GetCameraObject());
-		InternalViewportClient->UpdateViewForLockedActor();
-		GWorld->SendAllEndOfFrameUpdates();
-		InternalViewportScene->Draw(false);
-		TrackEditorThumbnail->CopyTextureIn((FSlateRenderTargetRHI*)InternalViewportScene->GetViewportRenderTargetTexture());
-		SequencerPtr.Pin()->SetPlaybackStatus(PlaybackStatus);
-	}
-}
+		Sequencer->EnterSilentMode();
 
+		AActor* Camera = nullptr;
+		FDelegateHandle Handle = Sequencer->OnCameraCut().AddLambda([&](UObject* InObject, bool){
+			Camera = Cast<AActor>(InObject);
+		});
 
-uint32 FThumbnailSection::GetThumbnailWidth() const
-{
-	return ThumbnailWidth;
-}
+		SavedPlaybackStatus = Sequencer->GetPlaybackStatus();
+		Sequencer->SetPlaybackStatus(EMovieScenePlayerStatus::Jumping);
+		Sequencer->SetGlobalTimeDirectly(Thumbnail.GetEvalPosition());
 
+		ViewportClient.SetActorLock(Camera);
+		Sequencer->OnCameraCut().Remove(Handle);
+	}
+}
 
-void FThumbnailSection::RegenerateViewportThumbnails(const FIntPoint& Size)
+void FThumbnailSection::PostDraw(FTrackEditorThumbnail& Thumbnail, FLevelEditorViewportClient& ViewportClient, FSceneViewport& SceneViewport)
 {
-	AActor* CameraActor = GetCameraObject();
-
-	if (CameraActor != nullptr && InternalViewportScene.IsValid())
+	TSharedPtr<ISequencer> Sequencer = SequencerPtr.Pin();
+	if (Sequencer.IsValid())
 	{
-		ThumbnailPool.Pin()->RemoveThumbnailsNeedingRedraw(Thumbnails);
-		Thumbnails.Empty();
-
-		if (Size.X == 0 || Size.Y == 0)
-		{
-			return;
-		}
-
-		StoredSize = Size;
-		StoredStartTime = Section->GetStartTime();
-	
-		UCameraComponent* CameraComponent = MovieSceneHelpers::CameraComponentFromActor(CameraActor);
-		float AspectRatio = CameraComponent ? CameraComponent->AspectRatio : 1.f;
-
-		int32 NewThumbnailWidth = FMath::TruncToInt(ThumbnailSectionConstants::TrackWidth * AspectRatio);
-		int32 ThumbnailCount = FMath::DivideAndRoundUp(StoredSize.X, NewThumbnailWidth);
-
-		if (InternalViewportClient.IsValid())
-		{
-			InternalViewportClient->SetAllowCinematicPreview(false);
-			InternalViewportClient->SetActorLock(CameraActor);
-		}
-
-		if (NewThumbnailWidth != ThumbnailWidth)
-		{
-			ThumbnailWidth = NewThumbnailWidth;
-			InternalViewportScene->UpdateViewportRHI(false, ThumbnailWidth, ThumbnailSectionConstants::ThumbnailHeight, EWindowMode::Windowed);
-		}
-
-		float StartTime = Section->GetStartTime();
-		float EndTime = Section->GetEndTime();
-		float DeltaTime = EndTime - StartTime;
-
-		int32 TotalX = StoredSize.X;
-
-		// @todo sequencer optimize this to use a single viewport and a single thumbnail
-		// instead paste the textures into the thumbnail at specific offsets
-		for (int32 ThumbnailIndex = 0; ThumbnailIndex < ThumbnailCount; ++ThumbnailIndex)
-		{
-			int32 StartX = ThumbnailIndex * ThumbnailWidth;
-			float FractionThrough = (float)StartX / (float)TotalX;
-			float Time = StartTime + DeltaTime * FractionThrough;
-
-			int32 NextStartX = (ThumbnailIndex+1) * ThumbnailWidth;
-			float NextFractionThrough = (float)NextStartX / (float)TotalX;
-			float NextTime = FMath::Min(StartTime + DeltaTime * NextFractionThrough, EndTime);
-
-			FIntPoint ThumbnailSize(NextStartX-StartX, ThumbnailSectionConstants::ThumbnailHeight);
-
-			check(FractionThrough >= 0.f && FractionThrough <= 1.f && NextFractionThrough >= 0.f);
-			TSharedPtr<FTrackEditorThumbnail> NewThumbnail = MakeShareable(new FTrackEditorThumbnail(SharedThis(this), ThumbnailSize, TRange<float>(Time, NextTime)));
-
-			Thumbnails.Add(NewThumbnail);
-		}
-
-		// @todo Sequencer Optimize Only say a thumbnail needs redraw if it is onscreen
-		ThumbnailPool.Pin()->AddThumbnailsNeedingRedraw(Thumbnails);
+		Thumbnail.SetupFade(Sequencer->GetSequencerWidget());
+		Sequencer->ExitSilentMode();
 	}
 }
 
@@ -182,6 +99,66 @@ TSharedRef<SWidget> FThumbnailSection::GenerateSectionWidget()
 		];
 }
 
+void FThumbnailSection::BuildSectionContextMenu(FMenuBuilder& MenuBuilder, const FGuid& ObjectBinding)
+{
+	MenuBuilder.BeginSection(NAME_None, LOCTEXT("ViewMenuText", "View"));
+	{
+		MenuBuilder.AddSubMenu(
+			LOCTEXT("ThumbnailsMenu", "Thumbnails"),
+			FText(),
+			FNewMenuDelegate::CreateLambda([=](FMenuBuilder& InMenuBuilder){
+
+				TSharedPtr<ISequencer> Sequencer = SequencerPtr.Pin();
+
+				FText CurrentTime = FText::FromString(Sequencer->GetZeroPadNumericTypeInterface()->ToString(Sequencer->GetGlobalTime()));
+
+				InMenuBuilder.BeginSection(NAME_None, LOCTEXT("ThisSectionText", "This Section"));
+				{
+					InMenuBuilder.AddMenuEntry(
+						LOCTEXT("RefreshText", "Refresh"),
+						LOCTEXT("RefreshTooltip", "Refresh this section's thumbnails"),
+						FSlateIcon(),
+						FUIAction(FExecuteAction::CreateRaw(this, &FThumbnailSection::RedrawThumbnails))
+					);
+					InMenuBuilder.AddMenuEntry(
+						FText::Format(LOCTEXT("SetSingleTime", "Set Thumbnail Time To {0}"), CurrentTime),
+						LOCTEXT("SetSingleTimeTooltip", "Defines the time at which this section should draw its single thumbnail to the current cursor position"),
+						FSlateIcon(),
+						FUIAction(
+							FExecuteAction::CreateLambda([=]{
+								SetSingleTime(Sequencer->GetGlobalTime());
+								GetMutableDefault<UMovieSceneUserThumbnailSettings>()->bDrawSingleThumbnails = true;
+								GetMutableDefault<UMovieSceneUserThumbnailSettings>()->SaveConfig();
+							})
+						)
+					);
+				}
+				InMenuBuilder.EndSection();
+
+				InMenuBuilder.BeginSection(NAME_None, LOCTEXT("GlobalSettingsText", "Global Settings"));
+				{
+					InMenuBuilder.AddMenuEntry(
+						LOCTEXT("RefreshAllText", "Refresh All"),
+						LOCTEXT("RefreshAllTooltip", "Refresh all sections' thumbnails"),
+						FSlateIcon(),
+						FUIAction(FExecuteAction::CreateLambda([]{
+							GetDefault<UMovieSceneUserThumbnailSettings>()->BroadcastRedrawThumbnails();
+						}))
+					);
+
+					FPropertyEditorModule& PropertyModule = FModuleManager::LoadModuleChecked<FPropertyEditorModule>("PropertyEditor");
+
+					FDetailsViewArgs Args(false, false, false, FDetailsViewArgs::HideNameArea);
+					TSharedRef<IDetailsView> DetailView = PropertyModule.CreateDetailView(Args);
+					DetailView->SetObject(GetMutableDefault<UMovieSceneUserThumbnailSettings>());
+					InMenuBuilder.AddWidget(DetailView, FText(), true);
+				}
+				InMenuBuilder.EndSection();
+			})
+		);
+	}
+	MenuBuilder.EndSection();
+}
 
 float FThumbnailSection::GetSectionGripSize() const
 {
@@ -191,7 +168,15 @@ float FThumbnailSection::GetSectionGripSize() const
 
 float FThumbnailSection::GetSectionHeight() const
 {
-	return ThumbnailSectionConstants::ThumbnailHeight;
+	auto* Settings = GetDefault<UMovieSceneUserThumbnailSettings>();
+	if (Settings->bDrawThumbnails)
+	{
+		return GetDefault<UMovieSceneUserThumbnailSettings>()->ThumbnailSize.Y;
+	}
+	else
+	{
+		return FEditorStyle::GetFontStyle("NormalFont").Size + 8.f;
+	}
 }
 
 
@@ -208,74 +193,95 @@ FText FThumbnailSection::GetSectionTitle() const
 
 int32 FThumbnailSection::OnPaintSection( FSequencerSectionPainter& InPainter ) const
 {
+	if (!GetDefault<UMovieSceneUserThumbnailSettings>()->bDrawThumbnails)
+	{
+		return InPainter.LayerId;
+	}
+
+	static const float SectionThumbnailPadding = 4.f;
+
 	const ESlateDrawEffect::Type DrawEffects = InPainter.bParentEnabled ? ESlateDrawEffect::None : ESlateDrawEffect::DisabledEffect;
 
 	int32 LayerId = InPainter.LayerId;
 
-	if (Thumbnails.Num() != 0)
-	{
-		// @todo Sequencer: Need a way to visualize the key here
-		FVector2D SectionSize = InPainter.SectionGeometry.GetLocalSize();
+	const FGeometry& SectionGeometry = InPainter.SectionGeometry;
+
+	// @todo Sequencer: Need a way to visualize the key here
+
+	const TRange<float> VisibleRange = SequencerPtr.Pin()->GetViewRange();
+	const TRange<float> SectionRange = Section->IsInfinite() ? VisibleRange : Section->GetRange();
+
+	const float TimePerPx = SectionRange.Size<float>() / InPainter.SectionGeometry.GetLocalSize().X;
+	
+	FSlateRect ThumbnailClipRect = SectionGeometry.GetClippingRect().InsetBy(FMargin(SectionThumbnailPadding, 0.f)).IntersectionWith(InPainter.SectionClippingRect);
 
-		int32 ThumbnailIndex = 0;
-		for( const TSharedPtr<FTrackEditorThumbnail>& Thumbnail : Thumbnails )
+	for (const TSharedPtr<FTrackEditorThumbnail>& Thumbnail : ThumbnailCache.GetThumbnails())
+	{
+		const float Fade = Thumbnail->bHasFinishedDrawing ? Thumbnail->GetFadeInCurve() : 1.f;
+		FIntPoint ThumbnailSize = Thumbnail->GetSize();
+		
+		// Calculate the paint geometry for this thumbnail
+		TOptional<float> SingleReferenceFrame = ThumbnailCache.GetSingleReferenceFrame();
+
+		// Single thumbnails are always draw at the start of the section, clamped to the visible range
+		// Thumbnail sequences draw relative to their actual position in the sequence
+		const int32 Offset = SingleReferenceFrame.IsSet() ?
+			FMath::Max((VisibleRange.GetLowerBoundValue() - SectionRange.GetLowerBoundValue()) / TimePerPx, 0.f) + SectionThumbnailPadding :
+			(Thumbnail->GetTimeRange().GetLowerBoundValue() - SectionRange.GetLowerBoundValue()) / TimePerPx;
+
+		FPaintGeometry PaintGeometry = SectionGeometry.ToPaintGeometry(
+			ThumbnailSize,
+			FSlateLayoutTransform(
+				SectionGeometry.Scale,
+				FVector2D( Offset, (SectionGeometry.GetLocalSize().Y - ThumbnailSize.Y)*.5f)
+			)
+		);
+
+		if (Fade <= 1.f)
 		{
-			FVector2D ThumbnailSize = Thumbnail->GetSize();
-
-			FGeometry TruncatedGeometry = InPainter.SectionGeometry.MakeChild(
-				ThumbnailSize,
-				FSlateLayoutTransform(
-					InPainter.SectionGeometry.Scale,
-					FVector2D( ThumbnailIndex * ThumbnailWidth, (SectionSize.Y - ThumbnailSize.Y)*.5f)
-				)
-			);
-			
 			FSlateDrawElement::MakeViewport(
 				InPainter.DrawElements,
-				++LayerId,
-				TruncatedGeometry.ToPaintGeometry(),
+				LayerId,
+				PaintGeometry,
 				Thumbnail,
-				InPainter.SectionClippingRect,
+				ThumbnailClipRect,
 				false,
 				false,
-				DrawEffects,
-				FLinearColor::White
-			);
-
-			if(Thumbnail->GetFadeInCurve() > 0.0f )
-			{
-				FSlateDrawElement::MakeBox(
-					InPainter.DrawElements,
-					++LayerId,
-					TruncatedGeometry.ToPaintGeometry(),
-					WhiteBrush,
-					InPainter.SectionClippingRect,
-					DrawEffects,
-					FLinearColor(1.0f, 1.0f, 1.0f, Thumbnail->GetFadeInCurve())
-					);
-			}
+				DrawEffects
+				);
+		}
 
-			++ThumbnailIndex;
+		if (Fade > 0.0f)
+		{
+			FSlateDrawElement::MakeBox(
+				InPainter.DrawElements,
+				LayerId + 1,
+				PaintGeometry,
+				WhiteBrush,
+				ThumbnailClipRect,
+				DrawEffects,
+				FLinearColor(0.0f, 0.0f, 0.0f, Fade)
+				);
 		}
 	}
 
-	return LayerId;
+	return LayerId + 2;
 }
 
 
 void FThumbnailSection::Tick(const FGeometry& AllottedGeometry, const FGeometry& ParentGeometry, const double InCurrentTime, const float InDeltaTime)
 {
-	// The sequencer view range which will be used to determine whether the individual thumbnails should be drawn if the time range intersects the view range
-	VisibleTimeRange = SequencerPtr.Pin()->GetViewRange();
+	if (FSlateThrottleManager::Get().IsAllowingExpensiveTasks() && GetDefault<UMovieSceneUserThumbnailSettings>()->bDrawThumbnails)
+	{
+		const UMovieSceneUserThumbnailSettings* Settings = GetDefault<UMovieSceneUserThumbnailSettings>();
 
-	FIntPoint AllocatedSize = AllottedGeometry.GetLocalSize().IntPoint();
-	AllocatedSize.X = FMath::Max(AllocatedSize.X, 1);
+		FIntPoint AllocatedSize = AllottedGeometry.GetLocalSize().IntPoint();
+		AllocatedSize.X = FMath::Max(AllocatedSize.X, 1);
 
-	float StartTime = Section->GetStartTime();
+		const TRange<float> VisibleRange = SequencerPtr.Pin()->GetViewRange();
+		const TRange<float> SectionRange = Section->IsInfinite() ? VisibleRange : Section->GetRange();
 
-	if (!FSlateApplication::Get().HasAnyMouseCaptor() && ( AllocatedSize.X != StoredSize.X || !FMath::IsNearlyEqual(StartTime, StoredStartTime) ) )
-	{
-		RegenerateViewportThumbnails(AllocatedSize);
+		ThumbnailCache.Update(SectionRange, VisibleRange, AllocatedSize, Settings->ThumbnailSize);
 	}
 }
 
diff --git a/Engine/Source/Editor/MovieSceneTools/Private/Sections/ThumbnailSection.h b/Engine/Source/Editor/MovieSceneTools/Private/Sections/ThumbnailSection.h
--- a/Engine/Source/Editor/MovieSceneTools/Private/Sections/ThumbnailSection.h
+++ b/Engine/Source/Editor/MovieSceneTools/Private/Sections/ThumbnailSection.h
@@ -2,19 +2,24 @@
 
 #pragma once
 
+#include "TrackEditorThumbnail.h"
 
-class FTrackEditorThumbnail;
 class FTrackEditorThumbnailPool;
 class IMenu;
 class ISectionLayoutBuilder;
 class UMovieSceneSection;
 
+struct FThumbnailCameraSettings
+{
+	float AspectRatio;
+};
 
 /**
  * Thumbnail section, which paints and ticks the appropriate section.
  */
 class FThumbnailSection
 	: public ISequencerSection
+	, public IThumbnailClient
 	, public TSharedFromThis<FThumbnailSection>
 {
 public:
@@ -23,34 +28,16 @@ public:
 	FThumbnailSection(TSharedPtr<ISequencer> InSequencer, TSharedPtr<FTrackEditorThumbnailPool> InThumbnailPool, UMovieSceneSection& InSection);
 
 	/** Virtual destructor. */
-	virtual ~FThumbnailSection();
+	~FThumbnailSection();
 
 public:
 
-	/** Draws the passed in viewport thumbnail and copies it to the thumbnail's texture. */
-	void DrawViewportThumbnail(TSharedPtr<FTrackEditorThumbnail> TrackEditorThumbnail);
-
 	/** @return The sequencer widget owning the MovieScene section. */
 	TSharedRef<SWidget> GetSequencerWidget()
 	{
 		return SequencerPtr.Pin()->GetSequencerWidget();
 	}
 
-	/** Gets the thumbnail width. */
-	uint32 GetThumbnailWidth() const;
-
-	/** Gets the time range of what in the sequencer is visible. */
-	TRange<float> GetVisibleTimeRange() const
-	{
-		return VisibleTimeRange;
-	}
-
-	/** Regenerates all viewports and thumbnails at the new size. */
-	void RegenerateViewportThumbnails(const FIntPoint& Size);
-
-	/** Get the camera that this thumbnail should draw from */
-	virtual AActor* GetCameraObject() const { return nullptr; }
-
 	/** Get whether the text is renameable */
 	virtual bool CanRename() const { return false; }
 
@@ -62,10 +49,14 @@ public:
 
 public:
 
+	/** Set this thumbnail section to draw a single thumbnail at the specified time */
+	virtual void SetSingleTime(float GlobalTime) = 0;
+
 	// ISequencerSection interface
 
 	virtual bool AreSectionsConnected() const override;
 	virtual void GenerateSectionLayout(ISectionLayoutBuilder& LayoutBuilder) const override { }
+	virtual void BuildSectionContextMenu(FMenuBuilder& MenuBuilder, const FGuid& ObjectBinding) override;
 	virtual TSharedRef<SWidget> GenerateSectionWidget() override;
 	virtual float GetSectionGripSize() const override;
 	virtual float GetSectionHeight() const override;
@@ -74,37 +65,26 @@ public:
 	virtual int32 OnPaintSection( FSequencerSectionPainter& InPainter ) const override;
 	virtual void Tick( const FGeometry& AllottedGeometry, const FGeometry& ParentGeometry, const double InCurrentTime, const float InDeltaTime ) override;
 
+	// IThumbnailClient interface
+	virtual void PreDraw(FTrackEditorThumbnail& Thumbnail, FLevelEditorViewportClient& ViewportClient, FSceneViewport& SceneViewport) override;
+	virtual void PostDraw(FTrackEditorThumbnail& Thumbnail, FLevelEditorViewportClient& ViewportClient, FSceneViewport& SceneViewport) override;
+
 protected:
 
+	/** Called to force a redraw of this section's thumbnails */
+	void RedrawThumbnails();
+
 	/** The section we are visualizing. */
 	UMovieSceneSection* Section;
 
 	/** The parent sequencer we are a part of. */
 	TWeakPtr<ISequencer> SequencerPtr;
 
-	/** The thumbnail pool that we are sending all of our thumbnails to. */
-	TWeakPtr<FTrackEditorThumbnailPool> ThumbnailPool;
-
 	/** A list of all thumbnails this CameraCut section has. */
-	TArray<TSharedPtr<FTrackEditorThumbnail>> Thumbnails;
-
-	/** The width of our thumbnails. */
-	uint32 ThumbnailWidth;
-
-	/** The stored size of this section in the Slate geometry. */
-	FIntPoint StoredSize;
-
-	/** The stored start time, to query for invalidations. */
-	float StoredStartTime;
-
-	/** Cached Time Range of the visible parent section area. */
-	TRange<float> VisibleTimeRange;
-	
-	/** An internal viewport scene we use to render the thumbnails with. */
-	TSharedPtr<FSceneViewport> InternalViewportScene;
+	FTrackEditorThumbnailCache ThumbnailCache;
 
-	/** An internal editor viewport client to render the thumbnails with. */
-	TSharedPtr<FLevelEditorViewportClient> InternalViewportClient;
+	/** Saved playback status. Used for restoring state when rendering thumbnails */
+	EMovieScenePlayerStatus::Type SavedPlaybackStatus;
 	
 	/** Fade brush. */
 	const FSlateBrush* WhiteBrush;
