Copying //UE4/Dev-Sequencer to //UE4/Main

diff --git a/Engine/Source/Editor/MovieSceneTools/Private/TrackEditorThumbnail/TrackEditorThumbnail.cpp b/Engine/Source/Editor/MovieSceneTools/Private/TrackEditorThumbnail/TrackEditorThumbnail.cpp
--- a/Engine/Source/Editor/MovieSceneTools/Private/TrackEditorThumbnail/TrackEditorThumbnail.cpp
+++ b/Engine/Source/Editor/MovieSceneTools/Private/TrackEditorThumbnail/TrackEditorThumbnail.cpp
@@ -1,88 +1,208 @@
 // Copyright 1998-2016 Epic Games, Inc. All Rights Reserved.
 
 #include "MovieSceneToolsPrivatePCH.h"
-#include "Runtime/Engine/Public/Slate/SlateTextures.h"
-
+#include "SlateTextures.h"
+#include "SceneViewport.h"
+#include "ScreenRendering.h"
+#include "RenderCore.h"
+#include "RHIStaticStates.h"
+#include "RendererInterface.h"
 
 namespace TrackEditorThumbnailConstants
 {
 	const double ThumbnailFadeInDuration = 0.25f;
 }
 
+class FThumbnailViewportClient : public FLevelEditorViewportClient
+{
+public:
+	FThumbnailViewportClient() : FLevelEditorViewportClient(nullptr) {}
+
+	float CurrentWorldTime, DeltaWorldTime;
+
+	virtual FSceneView* CalcSceneView(FSceneViewFamily* ViewFamily) override
+	{
+		FSceneView* View = FEditorViewportClient::CalcSceneView(ViewFamily);
+
+		// Artificially set the world times so that graphics settings apply correctly (we don't tick the world when rendering thumbnails)
+		ViewFamily->CurrentWorldTime = CurrentWorldTime;
+		ViewFamily->DeltaWorldTime = DeltaWorldTime;
+
+		View->FinalPostProcessSettings.bOverride_AutoExposureSpeedDown = View->FinalPostProcessSettings.bOverride_AutoExposureSpeedUp = true;
+		View->FinalPostProcessSettings.AutoExposureSpeedDown = View->FinalPostProcessSettings.AutoExposureSpeedUp = 0.02f;
+		return View;
+	}
+
+};
+
 
 /* FTrackEditorThumbnail structors
  *****************************************************************************/
 
-FTrackEditorThumbnail::FTrackEditorThumbnail(TSharedPtr<FThumbnailSection> InSection, const FIntPoint& InSize, TRange<float> InTimeRange)
-	: OwningSection(InSection)
+FTrackEditorThumbnail::FTrackEditorThumbnail(const FOnThumbnailDraw& InOnDraw, const FIntPoint& InSize, TRange<float> InTimeRange, float InPosition)
+	: OnDraw(InOnDraw)
 	, Size(InSize)
 	, Texture(nullptr)
 	, TimeRange(InTimeRange)
+	, Position(InPosition)
 	, FadeInCurve(0.0f, TrackEditorThumbnailConstants::ThumbnailFadeInDuration)
 {
-	Texture = new FSlateTexture2DRHIRef(GetSize().X, GetSize().Y, PF_B8G8R8A8, nullptr, TexCreate_Dynamic, true);
-
-	BeginInitResource( Texture );
-
+	SortOrder = 0;
 }
 
 
 FTrackEditorThumbnail::~FTrackEditorThumbnail()
 {
-	BeginReleaseResource(Texture);
-	FlushRenderingCommands();
-
-	if (Texture) 
+	if (Texture && !bHasFinishedDrawing)
 	{
-		delete Texture;
+		FlushRenderingCommands();
 	}
+	DestroyTexture();
 }
 
+void FTrackEditorThumbnail::DestroyTexture()
+{
+	if (Texture)
+	{
+		ENQUEUE_UNIQUE_RENDER_COMMAND_ONEPARAMETER(DestroyTexture,
+			FSlateTexture2DRHIRef*, Texture, Texture,
+		{
+			Texture->ReleaseResource();
+			delete Texture;
+		});
+
+		Texture = nullptr;
+	}
+}
 
 /* FTrackEditorThumbnail interface
  *****************************************************************************/
 
-void FTrackEditorThumbnail::CopyTextureIn(FSlateRenderTargetRHI* InTexture)
+void FTrackEditorThumbnail::CopyTextureIn(TSharedPtr<FSceneViewport> SceneViewport)
 {
-	ENQUEUE_UNIQUE_RENDER_COMMAND_TWOPARAMETER(ReadTexture,
-		FSlateRenderTargetRHI*, RenderTarget, InTexture,
-		FSlateTexture2DRHIRef*, TargetTexture, Texture,
+	FSlateRenderTargetRHI* RenderTarget = (FSlateRenderTargetRHI*)SceneViewport->GetViewportRenderTargetTexture();
+
+	if (!Texture || !RenderTarget)
+	{
+		return;
+	}
+	
+	// This code is a little manual because CopyToResolveTarget on its own can't resolve a sub rect without offsetting it inside the destination texture
+	// So we render our own rectangle onto a render target of the right size, so we can maintain the correct aspect ratio and fov settings of the camera,
+	// but still fulfil the desired thumbnail size.
+	FSlateTexture2DRHIRef* TargetTexture = Texture;
+
+	static const FName RendererModuleName( "Renderer" );
+	IRendererModule* RendererModule = &FModuleManager::GetModuleChecked<IRendererModule>(RendererModuleName);
+	FThreadSafeBool* bHasFinishedDrawingPtr = &bHasFinishedDrawing;
+
+	// Note: We explicitly capture the viewport scene here to ensure that the render target lives as long as this render command.
+	// This means we don't have to flush the rendering commands all the time
+	auto RenderCommand = [TargetTexture, RendererModule, RenderTarget, SceneViewport, bHasFinishedDrawingPtr](FRHICommandListImmediate& RHICmdList){
+
+		const FIntPoint TargetSize(TargetTexture->GetWidth(), TargetTexture->GetHeight());
+
+		FPooledRenderTargetDesc OutputDesc = FPooledRenderTargetDesc::Create2DDesc(
+			TargetSize,
+			PF_B8G8R8A8,
+			FClearValueBinding::None,
+			TexCreate_None,
+			TexCreate_RenderTargetable,
+			false);
+
+		TRefCountPtr<IPooledRenderTarget> ResampleTexturePooledRenderTarget;
+		RendererModule->RenderTargetPoolFindFreeElement(RHICmdList, OutputDesc, ResampleTexturePooledRenderTarget, TEXT("ResampleTexture"));
+		check(ResampleTexturePooledRenderTarget);
+
+		const FSceneRenderTargetItem& DestRenderTarget = ResampleTexturePooledRenderTarget->GetRenderTargetItem();
+
+		SetRenderTarget(RHICmdList, DestRenderTarget.TargetableTexture, FTextureRHIRef());
+
+		RHICmdList.SetViewport(0, 0, 0.0f, TargetSize.X, TargetSize.Y, 1.0f);
+
+		RHICmdList.SetBlendState(TStaticBlendState<>::GetRHI());
+		RHICmdList.SetRasterizerState(TStaticRasterizerState<>::GetRHI());
+		RHICmdList.SetDepthStencilState(TStaticDepthStencilState<false,CF_Always>::GetRHI());
+
+		const ERHIFeatureLevel::Type FeatureLevel = GMaxRHIFeatureLevel;
+		
+		TShaderMap<FGlobalShaderType>* ShaderMap = GetGlobalShaderMap(FeatureLevel);
+		TShaderMapRef<FScreenVS> VertexShader(ShaderMap);
+		TShaderMapRef<FScreenPS> PixelShader(ShaderMap);
+
+		static FGlobalBoundShaderState BoundShaderState;
+		SetGlobalBoundShaderState(RHICmdList, FeatureLevel, BoundShaderState, RendererModule->GetFilterVertexDeclaration().VertexDeclarationRHI, *VertexShader, *PixelShader);
+
+		FTexture2DRHIRef SourceBackBuffer = RenderTarget->GetRHIRef();
+		PixelShader->SetParameters(RHICmdList, TStaticSamplerState<SF_Bilinear>::GetRHI(), SourceBackBuffer);
+
+		const float Left = (SourceBackBuffer->GetSizeX() - TargetTexture->GetWidth()) * .5f;
+		const float Top = (SourceBackBuffer->GetSizeY() - TargetTexture->GetHeight()) * .5f;
+
+		const float U = Left / float(SourceBackBuffer->GetSizeX());
+		const float V = Top / float(SourceBackBuffer->GetSizeY());
+		const float SizeU = float(TargetTexture->GetWidth()) / float(SourceBackBuffer->GetSizeX());
+		const float SizeV = float(TargetTexture->GetHeight()) / float(SourceBackBuffer->GetSizeY());
+
+		RendererModule->DrawRectangle(
+			RHICmdList,
+			0, 0,									// Dest X, Y
+			TargetSize.X,							// Dest Width
+			TargetSize.Y,							// Dest Height
+			U, V,									// Source U, V
+			SizeU, SizeV,							// Source USize, VSize
+			TargetSize,								// Target buffer size
+			FIntPoint(1, 1),						// Source texture size
+			*VertexShader,
+			EDRF_Default);
+
+		// Asynchronously copy render target from GPU to CPU
+		const bool bKeepOriginalSurface = false;
+		RHICmdList.CopyToResolveTarget(
+			DestRenderTarget.TargetableTexture,
+			TargetTexture->GetTypedResource(),
+			bKeepOriginalSurface,
+			FResolveParams());
+
+		*bHasFinishedDrawingPtr = true;
+	};
+
+	ENQUEUE_UNIQUE_RENDER_COMMAND_ONEPARAMETER(
+		ResolveCaptureFrameTexture,
+		TFunction<void(FRHICommandListImmediate&)>, InRenderCommand, RenderCommand,
 	{
-		RHICmdList.CopyToResolveTarget(RenderTarget->GetRHIRef(), TargetTexture->GetTypedResource(), false, FResolveParams());
+		InRenderCommand(RHICmdList);
 	});
 }
 
 
 void FTrackEditorThumbnail::DrawThumbnail()
 {
-	OwningSection.Pin()->DrawViewportThumbnail(SharedThis(this));
-	FadeInCurve.PlayReverse( OwningSection.Pin()->GetSequencerWidget() );
-}
-
+	// Delay texture creation until we actually draw the thumbnail
+	if (Size.X > 1 && Size.Y > 1)
+	{
+		DestroyTexture();
+		Texture = new FSlateTexture2DRHIRef(Size.X, Size.Y, PF_B8G8R8A8, nullptr, TexCreate_Dynamic, true);
+		BeginInitResource(Texture);
+	}
 
-float FTrackEditorThumbnail::GetFadeInCurve() const 
-{
-	return FadeInCurve.GetLerp();
+	OnDraw.ExecuteIfBound(*this);
 }
 
-
-float FTrackEditorThumbnail::GetTime() const
+void FTrackEditorThumbnail::SetupFade(const TSharedRef<SWidget>& InWidget)
 {
-	return TimeRange.GetLowerBoundValue();
+	FadeInCurve.PlayReverse(InWidget);
+	FadeInCurve.Pause();
 }
 
-
-bool FTrackEditorThumbnail::IsValid() const
+void FTrackEditorThumbnail::PlayFade()
 {
-	return OwningSection.IsValid();
+	FadeInCurve.Resume();
 }
 
-
-bool FTrackEditorThumbnail::IsVisible() const
+float FTrackEditorThumbnail::GetFadeInCurve() const 
 {
-	return OwningSection.IsValid()
-		? TimeRange.Overlaps(OwningSection.Pin()->GetVisibleTimeRange())
-		: false;
+	return FadeInCurve.GetLerp();
 }
 
 
@@ -105,3 +225,363 @@ bool FTrackEditorThumbnail::RequiresVsync() const
 {
 	return false;
 }
+
+
+FTrackEditorThumbnailCache::FTrackEditorThumbnailCache(const TSharedPtr<FTrackEditorThumbnailPool>& InThumbnailPool, IThumbnailClient* InThumbnailClient)
+	: ThumbnailClient(InThumbnailClient)
+	, ThumbnailPool(InThumbnailPool)
+{
+	FrameCount = 0;
+	bForceRedraw = false;
+}
+
+
+FTrackEditorThumbnailCache::~FTrackEditorThumbnailCache()
+{
+	TSharedPtr<FTrackEditorThumbnailPool> PinnedPool = ThumbnailPool.Pin();
+	if (PinnedPool.IsValid())
+	{
+		PinnedPool->RemoveThumbnailsNeedingRedraw(Thumbnails);
+	}
+
+	if (InternalViewportClient.IsValid())
+	{
+		InternalViewportClient->Viewport = nullptr;
+	}
+}
+
+void FTrackEditorThumbnailCache::SetSingleReferenceFrame(TOptional<float> InReferenceFrame)
+{
+	CurrentCache.SingleReferenceFrame = InReferenceFrame;
+}
+
+void FTrackEditorThumbnailCache::Update(const TRange<float>& NewRange, const TRange<float>& VisibleRange, const FIntPoint& AllottedSize, const FIntPoint& InDesiredSize)
+{
+	PreviousCache.TimeRange = CurrentCache.TimeRange;
+	PreviousCache.VisibleRange = CurrentCache.VisibleRange;
+	PreviousCache.AllottedSize = CurrentCache.AllottedSize;
+	PreviousCache.DesiredSize = CurrentCache.DesiredSize;
+
+	CurrentCache.TimeRange = NewRange;
+	CurrentCache.VisibleRange = VisibleRange;
+	CurrentCache.AllottedSize = AllottedSize;
+	CurrentCache.DesiredSize = InDesiredSize;
+
+	Revalidate();
+
+	// Only update the single reference frame value once we've updated, since that can get set at any time, but Update() many be throttled
+	PreviousCache.SingleReferenceFrame = CurrentCache.SingleReferenceFrame;
+}
+
+FIntPoint FTrackEditorThumbnailCache::CalculateTextureSize() const
+{
+	UCameraComponent* CameraComponent = InternalViewportClient->GetCameraComponentForView();
+	float DesiredRatio = CameraComponent ? CameraComponent->AspectRatio : 1.77f;
+
+	if (CurrentCache.DesiredSize.X <= 0 || CurrentCache.DesiredSize.Y <= 0)
+	{
+		return FIntPoint(0,0);
+	}
+
+	float SizeRatio = float(CurrentCache.DesiredSize.X) / CurrentCache.DesiredSize.Y;
+
+	FIntPoint NewSize = CurrentCache.DesiredSize;
+	if (SizeRatio > DesiredRatio)
+	{
+		// Take width
+		NewSize = FIntPoint(CurrentCache.DesiredSize.X, CurrentCache.DesiredSize.X / DesiredRatio);
+	}
+	else if (SizeRatio < DesiredRatio)
+	{
+		// Take height
+		NewSize = FIntPoint(CurrentCache.DesiredSize.Y * DesiredRatio, CurrentCache.DesiredSize.Y);
+	}
+
+	return NewSize;
+}
+
+bool FTrackEditorThumbnailCache::ShouldRegenerateEverything() const
+{
+	if (bForceRedraw)
+	{
+		return true;
+	}
+
+	const float PreviousScale = PreviousCache.TimeRange.Size<float>() / PreviousCache.AllottedSize.X;
+	const float CurrentScale = CurrentCache.TimeRange.Size<float>() / CurrentCache.AllottedSize.X;
+	const float Threshold = PreviousScale * 0.01f;
+
+	return PreviousCache.DesiredSize != CurrentCache.DesiredSize || !FMath::IsNearlyEqual(PreviousScale, CurrentScale, Threshold);
+}
+
+void FTrackEditorThumbnailCache::DrawViewportThumbnail(FTrackEditorThumbnail& TrackEditorThumbnail)
+{
+	if (InternalViewportScene.IsValid())
+	{
+		// Ask the client to setup the frame
+		if (ThumbnailClient)
+		{
+			ThumbnailClient->PreDraw(TrackEditorThumbnail, *InternalViewportClient, *InternalViewportScene);
+		}
+		
+		// Finalize the view
+		InternalViewportClient->bLockedCameraView = true;
+		InternalViewportClient->UpdateViewForLockedActor();
+		GWorld->SendAllEndOfFrameUpdates();
+
+		// Update the viewport RHI if necessary
+		FIntPoint Size = CalculateTextureSize();
+		if (InternalViewportScene->GetSize() != Size)
+		{
+			InternalViewportScene->UpdateViewportRHI(false, Size.X, Size.Y, EWindowMode::Windowed);
+		}
+
+		InternalViewportClient->DeltaWorldTime = TrackEditorThumbnail.GetEvalPosition() - InternalViewportClient->CurrentWorldTime;
+		InternalViewportClient->CurrentWorldTime = TrackEditorThumbnail.GetEvalPosition();
+
+		// Draw the frame. If our total frame count is < 3 we re-render some benign frames to ensure the view is correctly set up (first few frames can be black)
+		do
+		{
+			InternalViewportScene->Draw(false);
+		}
+		while (++FrameCount < 3);
+
+		// Ask the client to finalize the frame
+		if (ThumbnailClient)
+		{
+			ThumbnailClient->PostDraw(TrackEditorThumbnail, *InternalViewportClient, *InternalViewportScene);
+		}
+
+		// Copy the render target into our texture
+		TrackEditorThumbnail.CopyTextureIn(InternalViewportScene);
+	}
+}
+
+void FTrackEditorThumbnailCache::Revalidate()
+{
+	if (CurrentCache == PreviousCache && !bForceRedraw)
+	{
+		return;
+	}
+
+	if (FMath::IsNearlyZero(CurrentCache.TimeRange.Size<float>()))
+	{
+		// Can't generate thumbnails for this
+		ThumbnailPool.Pin()->RemoveThumbnailsNeedingRedraw(ThumbnailsNeedingRedraw);
+		ThumbnailsNeedingRedraw.Reset();
+		Thumbnails.Reset();
+		return;
+	}
+
+	if (ShouldRegenerateEverything())
+	{
+		ThumbnailPool.Pin()->RemoveThumbnailsNeedingRedraw(ThumbnailsNeedingRedraw);
+		Thumbnails.Reset();
+	}
+	
+	ThumbnailsNeedingRedraw.Reset();
+
+	if (CurrentCache.SingleReferenceFrame.IsSet())
+	{
+		if (!Thumbnails.Num() || bForceRedraw || CurrentCache.SingleReferenceFrame != PreviousCache.SingleReferenceFrame)
+		{
+			UpdateSingleThumbnail();
+		}
+	}
+	else
+	{
+		UpdateFilledThumbnails();
+	}
+
+	if (ThumbnailsNeedingRedraw.Num())
+	{
+		ThumbnailPool.Pin()->AddThumbnailsNeedingRedraw(ThumbnailsNeedingRedraw);
+	}
+	if (Thumbnails.Num())
+	{
+		Setup();
+	}
+
+	bForceRedraw = false;
+}
+
+void FTrackEditorThumbnailCache::UpdateSingleThumbnail()
+{
+	Thumbnails.Reset();
+
+	const float TimePerPx = CurrentCache.TimeRange.Size<float>() / CurrentCache.AllottedSize.X;
+	const float HalfRange = CurrentCache.DesiredSize.X*TimePerPx*.5f;
+	const float EvalPosition = CurrentCache.SingleReferenceFrame.GetValue();
+
+	TSharedPtr<FTrackEditorThumbnail> NewThumbnail = MakeShareable(new FTrackEditorThumbnail(
+		FOnThumbnailDraw::CreateRaw(this, &FTrackEditorThumbnailCache::DrawViewportThumbnail),
+		CurrentCache.DesiredSize,
+		TRange<float>(EvalPosition - HalfRange, EvalPosition + HalfRange),
+		EvalPosition
+	));
+
+	Thumbnails.Add(NewThumbnail);
+	ThumbnailsNeedingRedraw.Add(NewThumbnail);
+}
+
+void FTrackEditorThumbnailCache::UpdateFilledThumbnails()
+{
+	// Remove any thumbnails from the front of the array that aren't in the actual time range of this section (we keep stuff around outside of the visible range)
+	{
+		int32 Index = 0;
+		for (; Index < Thumbnails.Num(); ++Index)
+		{
+			if (Thumbnails[Index]->GetTimeRange().Overlaps(CurrentCache.TimeRange))
+			{
+				break;
+			}
+		}
+		if (Index)
+		{
+			TArray<TSharedPtr<FTrackEditorThumbnail>> Remove;
+			Remove.Append(&Thumbnails[0], Index);
+			ThumbnailPool.Pin()->RemoveThumbnailsNeedingRedraw(Remove);
+
+			Thumbnails.RemoveAt(0, Index, false);
+		}
+	}
+
+	// Remove any thumbnails from the back of the array that aren't in the *actual* time range of this section (we keep stuff around outside of the visible range)
+	{
+		int32 NumToRemove = 0;
+		for (int32 Index = Thumbnails.Num() - 1; Index >= 0; --Index)
+		{
+			if (!Thumbnails[Index]->GetTimeRange().Overlaps(CurrentCache.TimeRange))
+			{
+				++NumToRemove;
+			}
+			else
+			{
+				break;
+			}
+		}
+		
+		if (NumToRemove)
+		{
+			TArray<TSharedPtr<FTrackEditorThumbnail>> Remove;
+			Remove.Append(&Thumbnails[Thumbnails.Num() - NumToRemove], NumToRemove);
+			ThumbnailPool.Pin()->RemoveThumbnailsNeedingRedraw(Remove);
+
+			Thumbnails.RemoveAt(Thumbnails.Num() - NumToRemove, NumToRemove, false);
+		}
+	}
+
+	TRange<float> MaxRange(CurrentCache.VisibleRange.GetLowerBoundValue() - CurrentCache.VisibleRange.Size<float>(), CurrentCache.VisibleRange.GetUpperBoundValue() + CurrentCache.VisibleRange.Size<float>());
+	TRange<float> Boundary = TRange<float>::Intersection(CurrentCache.TimeRange, MaxRange);
+
+	if (!Boundary.IsEmpty())
+	{
+		GenerateFront(Boundary);
+		GenerateBack(Boundary);
+	}
+
+	if (Thumbnails.Num())
+	{
+		for (const TSharedPtr<FTrackEditorThumbnail>& Thumbnail : Thumbnails)
+		{
+			Thumbnail->SortOrder = Thumbnail->GetTimeRange().Overlaps(CurrentCache.VisibleRange) ? 1 : 10;
+		}
+	}
+}
+
+void FTrackEditorThumbnailCache::GenerateFront(const TRange<float>& Boundary)
+{
+	if (!Thumbnails.Num())
+	{
+		return;
+	}
+
+	const float TimePerPx = CurrentCache.TimeRange.Size<float>() / CurrentCache.AllottedSize.X;
+	float EndTime = Thumbnails[0]->GetTimeRange().GetLowerBoundValue();
+
+	while (EndTime > Boundary.GetLowerBoundValue())
+	{
+		FIntPoint TextureSize = CurrentCache.DesiredSize;
+
+		TRange<float> TimeRange(EndTime - TextureSize.X * TimePerPx, EndTime);
+
+		// Evaluate the thumbnail along the length of its duration, based on its position in the sequence
+		const float FrameLength = TimeRange.Size<float>();
+		float TotalLerp = (TimeRange.GetLowerBoundValue() - CurrentCache.TimeRange.GetLowerBoundValue()) / (CurrentCache.TimeRange.Size<float>() - FrameLength);
+		float EvalPosition = CurrentCache.TimeRange.GetLowerBoundValue() + FMath::Clamp(TotalLerp, 0.f, .99f)*CurrentCache.TimeRange.Size<float>();
+
+		TSharedPtr<FTrackEditorThumbnail> NewThumbnail = MakeShareable(new FTrackEditorThumbnail(
+			FOnThumbnailDraw::CreateRaw(this, &FTrackEditorThumbnailCache::DrawViewportThumbnail),
+			TextureSize,
+			TimeRange,
+			EvalPosition
+		));
+
+		Thumbnails.Insert(NewThumbnail, 0);
+		ThumbnailsNeedingRedraw.Add(NewThumbnail);
+
+		EndTime = TimeRange.GetLowerBoundValue();
+	}
+}
+
+void FTrackEditorThumbnailCache::GenerateBack(const TRange<float>& Boundary)
+{
+	const float TimePerPx = CurrentCache.TimeRange.Size<float>() / CurrentCache.AllottedSize.X;
+	float StartTime = Thumbnails.Num() ? Thumbnails.Last()->GetTimeRange().GetUpperBoundValue() : Boundary.GetLowerBoundValue();
+
+	while (StartTime < Boundary.GetUpperBoundValue())
+	{
+		FIntPoint TextureSize = CurrentCache.DesiredSize;
+
+		{
+			// Move the thumbnail to the center of the space if we're the only thumbnail, and we don't fit on
+			float Overflow = TextureSize.X*TimePerPx - CurrentCache.TimeRange.Size<float>();
+			if (Thumbnails.Num() == 0 && Overflow > 0)
+			{
+				StartTime -= Overflow*.5f;
+			}
+		}
+
+		TRange<float> TimeRange(StartTime, StartTime + TextureSize.X * TimePerPx);
+
+		const float FrameLength = TimeRange.Size<float>();
+		float TotalLerp = (TimeRange.GetLowerBoundValue() - CurrentCache.TimeRange.GetLowerBoundValue()) / (CurrentCache.TimeRange.Size<float>() - FrameLength);
+		float EvalPosition = CurrentCache.TimeRange.GetLowerBoundValue() + FMath::Clamp(TotalLerp, 0.f, .99f)*CurrentCache.TimeRange.Size<float>();
+
+		TSharedPtr<FTrackEditorThumbnail> NewThumbnail = MakeShareable(new FTrackEditorThumbnail(
+			FOnThumbnailDraw::CreateRaw(this, &FTrackEditorThumbnailCache::DrawViewportThumbnail),
+			TextureSize,
+			TimeRange,
+			EvalPosition
+		));
+
+		NewThumbnail->SortOrder = TimeRange.Overlaps(CurrentCache.VisibleRange) ? 1 : 10;
+
+		Thumbnails.Add(NewThumbnail);
+		ThumbnailsNeedingRedraw.Add(NewThumbnail);
+
+		StartTime = TimeRange.GetUpperBoundValue();
+	}
+}
+
+void FTrackEditorThumbnailCache::Setup()
+{
+	if (!InternalViewportClient.IsValid())
+	{
+		InternalViewportClient = MakeShareable(new FThumbnailViewportClient());
+		InternalViewportClient->ViewportType = LVT_Perspective;
+		InternalViewportClient->bDisableInput = true;
+		InternalViewportClient->bDrawAxes = false;
+		InternalViewportClient->EngineShowFlags = FEngineShowFlags(ESFIM_Game);
+		// @todo: DisableAdvancedFeatures()? Faster, but can yield grossly incorrect thumbnails
+		InternalViewportClient->EngineShowFlags.SetMotionBlur(false);
+		InternalViewportClient->SetAllowCinematicPreview(false);
+		InternalViewportClient->SetRealtime(false);
+	}
+
+	if (!InternalViewportScene.IsValid())
+	{
+		InternalViewportScene = MakeShareable(new FSceneViewport(InternalViewportClient.Get(), nullptr));
+		InternalViewportClient->Viewport = InternalViewportScene.Get();
+	}
+}
\ No newline at end of file
diff --git a/Engine/Source/Editor/MovieSceneTools/Private/TrackEditorThumbnail/TrackEditorThumbnail.h b/Engine/Source/Editor/MovieSceneTools/Private/TrackEditorThumbnail/TrackEditorThumbnail.h
--- a/Engine/Source/Editor/MovieSceneTools/Private/TrackEditorThumbnail/TrackEditorThumbnail.h
+++ b/Engine/Source/Editor/MovieSceneTools/Private/TrackEditorThumbnail/TrackEditorThumbnail.h
@@ -6,7 +6,11 @@
 class FThumbnailSection;
 class FSlateRenderTargetRHI;
 class FSlateTexture2DRHIRef;
+class FTrackEditorThumbnail;
+class FTrackEditorThumbnailPool;
+class FTrackEditorThumbnailCache;
 
+DECLARE_DELEGATE_OneParam(FOnThumbnailDraw, FTrackEditorThumbnail&);
 
 /**
  * Track Editor Thumbnail, which keeps a Texture to be displayed by a viewport.
@@ -18,7 +22,7 @@ class FTrackEditorThumbnail
 public:
 
 	/** Create and initialize a new instance. */
-	FTrackEditorThumbnail(TSharedPtr<FThumbnailSection> InSection, const FIntPoint& InSize, TRange<float> InTimeRange);
+	FTrackEditorThumbnail(const FOnThumbnailDraw& InOnDraw, const FIntPoint& InSize, TRange<float> InTimeRange, float InPosition);
 
 	/** Virtual destructor. */
 	virtual ~FTrackEditorThumbnail();
@@ -26,21 +30,23 @@ public:
 public:
 
 	/** Copies the incoming render target to this thumbnails texture. */
-	void CopyTextureIn(FSlateRenderTargetRHI* InTexture);
+	void CopyTextureIn(TSharedPtr<FSceneViewport> SceneViewport);
 
 	/** Renders the thumbnail to the texture. */
 	void DrawThumbnail();
 
+	/** Prompt this thumbnail to fade in */
+	void SetupFade(const TSharedRef<SWidget>& InWidget);
+	void PlayFade();
+
 	/** Gets the curve for fading in the thumbnail. */
 	float GetFadeInCurve() const;
 
-	/** Gets the time that this thumbnail is a rendering of. */
-	float GetTime() const;
-
-	bool IsValid() const;
+	/** Get the full time-range that this thumbnail occupies */
+	const TRange<float> GetTimeRange() const { return TimeRange; }
 
-	/** Returns whether this thumbnail is visible based on the thumbnail section geometry visibility. */
-	bool IsVisible() const;
+	/** Get the time at which this thumbnail should be drawn */
+	float GetEvalPosition() const { return Position; }
 
 public:
 
@@ -52,9 +58,23 @@ public:
 
 private:
 
-	/** Parent thumbnail section we are a thumbnail of. */
-	TWeakPtr<FThumbnailSection> OwningSection;
-	
+	/** Destroy the texture */
+	void DestroyTexture();
+
+public:
+
+	/** Sort order for this thumbnail */
+	int32 SortOrder;
+
+	/** True when this thumbnail has been drawn, false otherwise */
+	FThreadSafeBool bHasFinishedDrawing;
+
+private:
+
+	/** Delegate to use to draw the thumbnail. */
+	FOnThumbnailDraw OnDraw;
+
+	/** The size of the texture */
 	FIntPoint Size;
 
 	/** The Texture RHI that holds the thumbnail. */
@@ -63,6 +83,118 @@ private:
 	/** Where in time this thumbnail is a rendering of. */
 	TRange<float> TimeRange;
 
+	/** The position we should actually render (within the above time range). */
+	float Position;
+
 	/** Fade curve to display while the thumbnail is redrawing. */
 	FCurveSequence FadeInCurve;
 };
+
+struct IThumbnailClient
+{
+	virtual void PreDraw(FTrackEditorThumbnail& TrackEditorThumbnail, FLevelEditorViewportClient& ViewportClient, FSceneViewport& SceneViewport) { }
+	virtual void PostDraw(FTrackEditorThumbnail& TrackEditorThumbnail, FLevelEditorViewportClient& ViewportClient, FSceneViewport& SceneViewport) { }
+};
+
+/** Cache data */
+struct FThumbnailCacheData
+{
+	FThumbnailCacheData() : VisibleRange(0.f), TimeRange(0.f), AllottedSize(0,0), DesiredSize(0, 0) {}
+
+	bool operator==(const FThumbnailCacheData& RHS) const
+	{
+		return
+			AllottedSize == RHS.AllottedSize && 
+			VisibleRange == RHS.VisibleRange &&
+			TimeRange == RHS.TimeRange &&
+			DesiredSize == RHS.DesiredSize &&
+			SingleReferenceFrame == RHS.SingleReferenceFrame;
+	}
+
+	bool operator!=(const FThumbnailCacheData& RHS) const
+	{
+		return
+			AllottedSize != RHS.AllottedSize ||
+			VisibleRange != RHS.VisibleRange ||
+			TimeRange != RHS.TimeRange ||
+			DesiredSize != RHS.DesiredSize ||
+			SingleReferenceFrame != RHS.SingleReferenceFrame;
+	}
+
+	/** The visible range of our thumbnails we can see on the UI */
+	TRange<float> VisibleRange;
+	/** The total range to generate thumbnails for */
+	TRange<float> TimeRange;
+	/** Physical size of the thumbnail area */
+	FIntPoint AllottedSize;
+	/** Desired frame size constraint */
+	FIntPoint DesiredSize;
+	/** Set when we want to render a single reference frame */
+	TOptional<float> SingleReferenceFrame;
+};
+
+class FTrackEditorThumbnailCache
+{
+public:
+	FTrackEditorThumbnailCache(const TSharedPtr<FTrackEditorThumbnailPool>& ThumbnailPool, IThumbnailClient* InThumbnailClient);
+	~FTrackEditorThumbnailCache();
+
+	void ForceRedraw() { bForceRedraw = true; }
+
+	void SetSingleReferenceFrame(TOptional<float> InReferenceFrame);
+	TOptional<float> GetSingleReferenceFrame() const { return CurrentCache.SingleReferenceFrame; }
+
+	void Update(const TRange<float>& NewRange, const TRange<float>& VisibleRange, const FIntPoint& AllottedSize, const FIntPoint& InDesiredSize);
+
+	void Revalidate();
+
+	void DrawViewportThumbnail(FTrackEditorThumbnail& TrackEditorThumbnail);
+
+	const TArray<TSharedPtr<FTrackEditorThumbnail>>& GetThumbnails() const
+	{
+		return Thumbnails;
+	}
+
+protected:
+
+	void Setup();
+
+	bool ShouldRegenerateEverything() const;
+
+	FIntPoint CalculateTextureSize() const;
+
+	void UpdateSingleThumbnail();
+
+	void UpdateFilledThumbnails();
+
+	void GenerateFront(const TRange<float>& Boundary);
+
+	void GenerateBack(const TRange<float>& Boundary);
+
+protected:
+
+	/** Thumbnail client used for paint notifications */
+	IThumbnailClient* ThumbnailClient;
+	
+	/** An internal viewport scene we use to render the thumbnails with. */
+	TSharedPtr<FSceneViewport> InternalViewportScene;
+
+	/** An internal editor viewport client to render the thumbnails with. */
+	TSharedPtr<class FThumbnailViewportClient> InternalViewportClient;
+
+	/** The thumbnail pool that we are sending all of our thumbnails to. */
+	TWeakPtr<FTrackEditorThumbnailPool> ThumbnailPool;
+
+	FThumbnailCacheData CurrentCache;
+
+	FThumbnailCacheData PreviousCache;
+
+	TArray<TSharedPtr<FTrackEditorThumbnail>> Thumbnails;
+	TArray<TSharedPtr<FTrackEditorThumbnail>> ThumbnailsNeedingRedraw;
+
+	/** The number of frames we've rendered */
+	uint32 FrameCount;
+
+	/** Whether to force a redraw or not */
+	bool bForceRedraw;
+};
diff --git a/Engine/Source/Editor/MovieSceneTools/Private/TrackEditorThumbnail/TrackEditorThumbnailPool.cpp b/Engine/Source/Editor/MovieSceneTools/Private/TrackEditorThumbnail/TrackEditorThumbnailPool.cpp
--- a/Engine/Source/Editor/MovieSceneTools/Private/TrackEditorThumbnail/TrackEditorThumbnailPool.cpp
+++ b/Engine/Source/Editor/MovieSceneTools/Private/TrackEditorThumbnail/TrackEditorThumbnailPool.cpp
@@ -6,11 +6,13 @@
 /* FShotSequencerSection structors
  *****************************************************************************/
 
-FTrackEditorThumbnailPool::FTrackEditorThumbnailPool(TSharedPtr<ISequencer> InSequencer, uint32 InMaxThumbnailsToDrawAtATime)
+FTrackEditorThumbnailPool::FTrackEditorThumbnailPool(TSharedPtr<ISequencer> InSequencer)
 	: Sequencer(InSequencer)
 	, ThumbnailsNeedingDraw()
-	, MaxThumbnailsToDrawAtATime(InMaxThumbnailsToDrawAtATime)
-{ }
+	, bNeedsSort(false)
+{
+	TimeOfLastDraw = TimeOfLastUpdate = 0;
+}
 
 
 /* FShotSequencerSection interface
@@ -19,50 +21,100 @@ FTrackEditorThumbnailPool::FTrackEditorThumbnailPool(TSharedPtr<ISequencer> InSe
 void FTrackEditorThumbnailPool::AddThumbnailsNeedingRedraw(const TArray<TSharedPtr<FTrackEditorThumbnail>>& InThumbnails)
 {
 	ThumbnailsNeedingDraw.Append(InThumbnails);
+	bNeedsSort = true;
 }
 
-
 bool FTrackEditorThumbnailPool::DrawThumbnails()
 {
-	if (ThumbnailsNeedingDraw.Num() == 0)
+	int32 ThumbnailsDrawn = 0;
+
+	// Apply sorting if necessary
+	if (bNeedsSort)
 	{
-		return false;
+		auto SortFunc = [](const TSharedPtr<FTrackEditorThumbnail>& A, const TSharedPtr<FTrackEditorThumbnail>& B){
+			if (A->SortOrder == B->SortOrder)
+			{
+				return A->GetTimeRange().GetLowerBoundValue() < B->GetTimeRange().GetLowerBoundValue();
+			}
+			return A->SortOrder < B->SortOrder;
+		};
+		ThumbnailsNeedingDraw.Sort(SortFunc);
+		bNeedsSort = false;
 	}
 
-	uint32 ThumbnailsDrawn = 0;
+	const double CurrentTime = FSlateApplication::Get().GetCurrentTime();
 
-	for (int32 ThumbnailsIndex = 0; ThumbnailsDrawn < MaxThumbnailsToDrawAtATime && ThumbnailsIndex < ThumbnailsNeedingDraw.Num(); ++ThumbnailsIndex)
+	// Only allow drawing if we're not waiting on other thumbnails
+	bool bAllowDraw = !ThumbnailsBeingDrawn.Num();
+	for (int32 ThumbnailIndex = ThumbnailsBeingDrawn.Num() - 1; ThumbnailIndex >= 0; --ThumbnailIndex)
 	{
-		TSharedPtr<FTrackEditorThumbnail> Thumbnail = ThumbnailsNeedingDraw[ThumbnailsIndex];
-			
-		if (Thumbnail->IsVisible())
+		if (ThumbnailsBeingDrawn[ThumbnailIndex]->bHasFinishedDrawing)
 		{
-			
-			bool bIsEnabled = Sequencer.Pin()->IsPerspectiveViewportCameraCutEnabled();
-
-			Sequencer.Pin()->SetPerspectiveViewportCameraCutEnabled(false);
-			
-			Thumbnail->DrawThumbnail();
-			
-			Sequencer.Pin()->SetPerspectiveViewportCameraCutEnabled(bIsEnabled);
-			++ThumbnailsDrawn;
+			ThumbnailsBeingDrawn[ThumbnailIndex]->PlayFade();
+			ThumbnailsBeingDrawn.RemoveAt(ThumbnailIndex, 1, false);
+			TimeOfLastDraw = CurrentTime;
+		}
+		else
+		{
+			bAllowDraw = false;
+		}
+	}
 
-			ThumbnailsNeedingDraw.Remove(Thumbnail);
+	if (bAllowDraw)
+	{
+		if (!ThumbnailsNeedingDraw.Num() && !ThumbnailsBeingDrawn.Num())
+		{
+			TimeOfLastDraw = CurrentTime;
 		}
-		else if (!Thumbnail->IsValid())
+		else
 		{
-			ensure(0);
+			const float MinThumbnailsPerS = 2.f;
+			const float MaxThumbnailsPerS = 120.f;
 
-			ThumbnailsNeedingDraw.Remove(Thumbnail);
+			const float MinFramerate = 10.f;
+			const float MaxFramerate = 90.f;
+
+			const float AverageDeltaTime = FSlateApplication::Get().GetAverageDeltaTime();
+
+			const float Lerp = FMath::Max(((1.f/AverageDeltaTime) - MinFramerate), 0.f) / (MaxFramerate - MinFramerate);
+			const float NumToDrawPerS = FMath::Lerp(MinThumbnailsPerS, MaxThumbnailsPerS, Lerp);
+
+			const float TimeThreshold = 1.f / NumToDrawPerS;
+
+			const float TimeElapsed = CurrentTime - TimeOfLastDraw;
+
+			if (!FMath::IsNearlyEqual(float(CurrentTime - TimeOfLastUpdate), AverageDeltaTime, AverageDeltaTime*2.f))
+			{
+				// Don't generate thumbnails if we haven't had an update within a reasonable time - assume some blocking task
+				TimeOfLastDraw = CurrentTime;
+			}
+			else
+			{
+				const int32 NumToDraw = FMath::TruncToInt(TimeElapsed / TimeThreshold);
+
+				for (; ThumbnailsDrawn < FMath::Min(NumToDraw, ThumbnailsNeedingDraw.Num()); ++ThumbnailsDrawn)
+				{
+					const TSharedPtr<FTrackEditorThumbnail>& Thumbnail = ThumbnailsNeedingDraw[ThumbnailsDrawn];
+
+					bool bIsEnabled = Sequencer.Pin()->IsPerspectiveViewportCameraCutEnabled();
+					Sequencer.Pin()->SetPerspectiveViewportCameraCutEnabled(false);
+
+					Thumbnail->DrawThumbnail();
+					ThumbnailsBeingDrawn.Add(Thumbnail);
+
+					Sequencer.Pin()->SetPerspectiveViewportCameraCutEnabled(bIsEnabled);
+				}
+			}
 		}
 	}
-		
+
 	if (ThumbnailsDrawn > 0)
 	{
-		// Ensure all buffers are updated
-		FlushRenderingCommands();
+		ThumbnailsNeedingDraw.RemoveAt(0, ThumbnailsDrawn, false);
 	}
 
+	TimeOfLastUpdate = CurrentTime;
+
 	return ThumbnailsDrawn > 0;
 }
 
@@ -72,5 +124,6 @@ void FTrackEditorThumbnailPool::RemoveThumbnailsNeedingRedraw(const TArray< TSha
 	for (int32 i = 0; i < InThumbnails.Num(); ++i)
 	{
 		ThumbnailsNeedingDraw.Remove(InThumbnails[i]);
+		ThumbnailsBeingDrawn.Remove(InThumbnails[i]);
 	}
 }
diff --git a/Engine/Source/Editor/MovieSceneTools/Private/TrackEditorThumbnail/TrackEditorThumbnailPool.h b/Engine/Source/Editor/MovieSceneTools/Private/TrackEditorThumbnail/TrackEditorThumbnailPool.h
--- a/Engine/Source/Editor/MovieSceneTools/Private/TrackEditorThumbnail/TrackEditorThumbnailPool.h
+++ b/Engine/Source/Editor/MovieSceneTools/Private/TrackEditorThumbnail/TrackEditorThumbnailPool.h
@@ -14,7 +14,7 @@ class FTrackEditorThumbnailPool
 {
 public:
 
-	FTrackEditorThumbnailPool(TSharedPtr<ISequencer> InSequencer, uint32 InMaxThumbnailsToDrawAtATime = 1);
+	FTrackEditorThumbnailPool(TSharedPtr<ISequencer> InSequencer);
 
 public:
 
@@ -36,6 +36,12 @@ private:
 	/** Thumbnails enqueued for drawing */
 	TArray< TSharedPtr<FTrackEditorThumbnail> > ThumbnailsNeedingDraw;
 
-	/** How many thumbnails we are allowed to draw in a single DrawThumbnails call */
-	uint32 MaxThumbnailsToDrawAtATime;
+	/** Thumbnails that are currently being drawn */
+	TArray< TSharedPtr<FTrackEditorThumbnail> > ThumbnailsBeingDrawn;
+
+	double TimeOfLastDraw;
+	double TimeOfLastUpdate;
+
+	/** Whether we need to sort */
+	bool bNeedsSort;
 };
diff --git a/Engine/Source/Editor/MovieSceneTools/Private/TrackEditors/AttachTrackEditor.cpp b/Engine/Source/Editor/MovieSceneTools/Private/TrackEditors/AttachTrackEditor.cpp
--- a/Engine/Source/Editor/MovieSceneTools/Private/TrackEditors/AttachTrackEditor.cpp
+++ b/Engine/Source/Editor/MovieSceneTools/Private/TrackEditors/AttachTrackEditor.cpp
@@ -117,7 +117,7 @@ void F3DAttachTrackEditor::BuildObjectBindingTrackMenu(FMenuBuilder& MenuBuilder
 bool F3DAttachTrackEditor::IsActorPickable(const AActor* const ParentActor, FGuid ObjectBinding, UMovieSceneSection* InSection)
 {
 	// Can't pick the object that this track binds
-	TArray<UObject*> OutObjects;
+	TArray<TWeakObjectPtr<UObject>> OutObjects;
 	GetSequencer()->GetRuntimeObjects( GetSequencer()->GetFocusedMovieSceneSequenceInstance(), ObjectBinding, OutObjects);
 	if (OutObjects.Contains(ParentActor))
 	{
@@ -137,14 +137,7 @@ bool F3DAttachTrackEditor::IsActorPickable(const AActor* const ParentActor, FGui
 
 void F3DAttachTrackEditor::ActorSocketPicked(const FName SocketName, USceneComponent* Component, AActor* ParentActor, FGuid ObjectGuid, UMovieSceneSection* Section)
 {
-	if (ObjectGuid.IsValid())
-	{
-		TArray<UObject*> OutObjects;
-		GetSequencer()->GetRuntimeObjects( GetSequencer()->GetFocusedMovieSceneSequenceInstance(), ObjectGuid, OutObjects);
-
-		AnimatablePropertyChanged( FOnKeyProperty::CreateRaw( this, &F3DAttachTrackEditor::AddKeyInternal, OutObjects, SocketName, Component ? Component->GetFName() : NAME_None, ParentActor) );
-	}
-	else if (Section != nullptr)
+	if (Section != nullptr)
 	{
 		const FScopedTransaction Transaction(LOCTEXT("UndoSetAttach", "Set Attach"));
 
@@ -158,9 +151,16 @@ void F3DAttachTrackEditor::ActorSocketPicked(const FName SocketName, USceneCompo
 			AttachSection->AttachComponentName = Component ? Component->GetFName() : NAME_None;
 		}
 	}
+	else if (ObjectGuid.IsValid())
+	{
+		TArray<TWeakObjectPtr<UObject>> OutObjects;
+		GetSequencer()->GetRuntimeObjects( GetSequencer()->GetFocusedMovieSceneSequenceInstance(), ObjectGuid, OutObjects);
+
+		AnimatablePropertyChanged( FOnKeyProperty::CreateRaw( this, &F3DAttachTrackEditor::AddKeyInternal, OutObjects, SocketName, Component ? Component->GetFName() : NAME_None, ParentActor) );
+	}
 }
 
-bool F3DAttachTrackEditor::AddKeyInternal( float KeyTime, const TArray<UObject*> Objects, const FName SocketName, const FName ComponentName, AActor* ParentActor)
+bool F3DAttachTrackEditor::AddKeyInternal( float KeyTime, const TArray<TWeakObjectPtr<UObject>> Objects, const FName SocketName, const FName ComponentName, AActor* ParentActor)
 {
 	bool bHandleCreated = false;
 	bool bTrackCreated = false;
@@ -182,7 +182,7 @@ bool F3DAttachTrackEditor::AddKeyInternal( float KeyTime, const TArray<UObject*>
 
 	for( int32 ObjectIndex = 0; ObjectIndex < Objects.Num(); ++ObjectIndex )
 	{
-		UObject* Object = Objects[ObjectIndex];
+		UObject* Object = Objects[ObjectIndex].Get();
 
 		FFindOrCreateHandleResult HandleResult = FindOrCreateHandleToObject( Object );
 		FGuid ObjectHandle = HandleResult.Handle;
diff --git a/Engine/Source/Editor/MovieSceneTools/Private/TrackEditors/AttachTrackEditor.h b/Engine/Source/Editor/MovieSceneTools/Private/TrackEditors/AttachTrackEditor.h
--- a/Engine/Source/Editor/MovieSceneTools/Private/TrackEditors/AttachTrackEditor.h
+++ b/Engine/Source/Editor/MovieSceneTools/Private/TrackEditors/AttachTrackEditor.h
@@ -48,5 +48,5 @@ public:
 private:
 
 	/** Delegate for AnimatablePropertyChanged in AddKey */
-	bool AddKeyInternal(float KeyTime, const TArray<UObject*> Objects, const FName SocketName, const FName ComponentName, AActor* ParentActor);
+	bool AddKeyInternal(float KeyTime, const TArray<TWeakObjectPtr<UObject>> Objects, const FName SocketName, const FName ComponentName, AActor* ParentActor);
 };
diff --git a/Engine/Source/Editor/MovieSceneTools/Private/TrackEditors/AudioTrackEditor.cpp b/Engine/Source/Editor/MovieSceneTools/Private/TrackEditors/AudioTrackEditor.cpp
--- a/Engine/Source/Editor/MovieSceneTools/Private/TrackEditors/AudioTrackEditor.cpp
+++ b/Engine/Source/Editor/MovieSceneTools/Private/TrackEditors/AudioTrackEditor.cpp
@@ -25,6 +25,9 @@
 #include "Sound/SoundBase.h"
 #include "Sound/SoundWave.h"
 #include "Sound/SoundCue.h"
+#include "ContentBrowserModule.h"
+#include "SequencerUtilities.h"
+#include "AssetRegistryModule.h"
 
 
 #define LOCTEXT_NAMESPACE "FAudioTrackEditor"
@@ -288,9 +291,11 @@ void FAudioThumbnail::GenerateWaveformPreview(TArray<uint8>& OutData, TRange<flo
 		FAudioDevice* AudioDevice = GEngine->GetMainAudioDevice();
 		if (AudioDevice)
 		{
-			SoundWave->InitAudioResource(AudioDevice->GetRuntimeFormat(SoundWave));
-			FAsyncAudioDecompress TempDecompress(SoundWave);
-			TempDecompress.StartSynchronousTask();
+			if ( SoundWave->InitAudioResource( AudioDevice->GetRuntimeFormat( SoundWave ) ) && (SoundWave->DecompressionType != DTYPE_RealTime || SoundWave->CachedRealtimeFirstBuffer == nullptr ) )
+			{
+				FAsyncAudioDecompress TempDecompress(SoundWave);
+				TempDecompress.StartSynchronousTask();
+			}
 		}
 	}
 
@@ -299,6 +304,11 @@ void FAudioThumbnail::GenerateWaveformPreview(TArray<uint8>& OutData, TRange<flo
 	const int32 LookupDataSize = SoundWave->RawPCMDataSize;
 	const int32 LookupSize = LookupDataSize * sizeof(uint8) / sizeof(int16);
 
+	if (!LookupData)
+	{
+		return;
+	}
+
 	const bool bShowIntensity = AudioSection->ShouldShowIntensity();
 
 	// @todo Sequencer This fixes looping drawing by pretending we are only dealing with a SoundWave
@@ -768,6 +778,19 @@ TSharedRef<ISequencerSection> FAudioTrackEditor::MakeSectionInterface( UMovieSce
 	return MakeShareable( new FAudioSection(SectionObject, bIsAMasterTrack, GetSequencer()) );
 }
 
+TSharedPtr<SWidget> FAudioTrackEditor::BuildOutlinerEditWidget(const FGuid& ObjectBinding, UMovieSceneTrack* Track, const FBuildEditWidgetParams& Params)
+{
+	// Create a container edit box
+	return SNew(SHorizontalBox)
+
+	// Add the audio combo box
+	+ SHorizontalBox::Slot()
+	.AutoWidth()
+	.VAlign(VAlign_Center)
+	[
+		FSequencerUtilities::MakeAddButton(LOCTEXT("AudioText", "Audio"), FOnGetContent::CreateSP(this, &FAudioTrackEditor::BuildAudioSubMenu, Track), Params.NodeIsHovered)
+	];
+}
 
 bool FAudioTrackEditor::HandleAssetAdded(UObject* Asset, const FGuid& TargetObjectGuid)
 {
@@ -777,7 +800,7 @@ bool FAudioTrackEditor::HandleAssetAdded(UObject* Asset, const FGuid& TargetObje
 		
 		if (TargetObjectGuid.IsValid())
 		{
-			TArray<UObject*> OutObjects;
+			TArray<TWeakObjectPtr<UObject>> OutObjects;
 			GetSequencer()->GetRuntimeObjects(GetSequencer()->GetFocusedMovieSceneSequenceInstance(), TargetObjectGuid, OutObjects);
 
 			AnimatablePropertyChanged( FOnKeyProperty::CreateRaw(this, &FAudioTrackEditor::AddNewAttachedSound, Sound, OutObjects));
@@ -805,7 +828,7 @@ bool FAudioTrackEditor::AddNewMasterSound( float KeyTime, USoundBase* Sound )
 }
 
 
-bool FAudioTrackEditor::AddNewAttachedSound( float KeyTime, USoundBase* Sound, TArray<UObject*> ObjectsToAttachTo )
+bool FAudioTrackEditor::AddNewAttachedSound( float KeyTime, USoundBase* Sound, TArray<TWeakObjectPtr<UObject>> ObjectsToAttachTo )
 {
 	bool bHandleCreated = false;
 	bool bTrackCreated = false;
@@ -813,7 +836,7 @@ bool FAudioTrackEditor::AddNewAttachedSound( float KeyTime, USoundBase* Sound, T
 
 	for( int32 ObjectIndex = 0; ObjectIndex < ObjectsToAttachTo.Num(); ++ObjectIndex )
 	{
-		UObject* Object = ObjectsToAttachTo[ObjectIndex];
+		UObject* Object = ObjectsToAttachTo[ObjectIndex].Get();
 
 		FFindOrCreateHandleResult HandleResult = FindOrCreateHandleToObject( Object );
 		FGuid ObjectHandle = HandleResult.Handle;
@@ -862,5 +885,63 @@ void FAudioTrackEditor::HandleAddAudioTrackMenuEntryExecute()
 	GetSequencer()->NotifyMovieSceneDataChanged();
 }
 
+TSharedRef<SWidget> FAudioTrackEditor::BuildAudioSubMenu(UMovieSceneTrack* Track)
+{
+	FAssetRegistryModule& AssetRegistryModule = FModuleManager::LoadModuleChecked<FAssetRegistryModule>(TEXT("AssetRegistry"));
+	TArray<FName> ClassNames;
+	ClassNames.Add(USoundBase::StaticClass()->GetFName());
+	TSet<FName> DerivedClassNames;
+	AssetRegistryModule.Get().GetDerivedClassNames(ClassNames, TSet<FName>(), DerivedClassNames);
+
+	FMenuBuilder MenuBuilder(true, nullptr);
+
+	FAssetPickerConfig AssetPickerConfig;
+	{
+		AssetPickerConfig.OnAssetSelected = FOnAssetSelected::CreateRaw( this, &FAudioTrackEditor::OnAudioAssetSelected, Track);
+		AssetPickerConfig.bAllowNullSelection = false;
+		AssetPickerConfig.InitialAssetViewType = EAssetViewType::List;
+		for (auto ClassName : DerivedClassNames)
+		{
+			AssetPickerConfig.Filter.ClassNames.Add(ClassName);
+		}
+	}
+
+	FContentBrowserModule& ContentBrowserModule = FModuleManager::Get().LoadModuleChecked<FContentBrowserModule>(TEXT("ContentBrowser"));
 
+	TSharedPtr<SBox> MenuEntry = SNew(SBox)
+		.WidthOverride(300.0f)
+		.HeightOverride(300.f)
+		[
+			ContentBrowserModule.Get().CreateAssetPicker(AssetPickerConfig)
+		];
+
+	MenuBuilder.AddWidget(MenuEntry.ToSharedRef(), FText::GetEmpty(), true);
+
+	return MenuBuilder.MakeWidget();
+}
+
+
+void FAudioTrackEditor::OnAudioAssetSelected(const FAssetData& AssetData, UMovieSceneTrack* Track)
+{
+	FSlateApplication::Get().DismissAllMenus();
+
+	UObject* SelectedObject = AssetData.GetAsset();
+
+	if (SelectedObject)
+	{
+		USoundBase* NewSound = CastChecked<USoundBase>(AssetData.GetAsset());
+		if (NewSound != nullptr)
+		{
+			const FScopedTransaction Transaction(NSLOCTEXT("Sequencer", "AddAudio_Transaction", "Add Audio"));
+
+			auto AudioTrack = Cast<UMovieSceneAudioTrack>(Track);
+			AudioTrack->Modify();
+
+			float KeyTime = GetSequencer()->GetGlobalTime();
+			AudioTrack->AddNewSound( NewSound, KeyTime );
+
+			GetSequencer()->NotifyMovieSceneDataChanged();
+		}
+	}
+}
 #undef LOCTEXT_NAMESPACE
diff --git a/Engine/Source/Editor/MovieSceneTools/Private/TrackEditors/AudioTrackEditor.h b/Engine/Source/Editor/MovieSceneTools/Private/TrackEditors/AudioTrackEditor.h
--- a/Engine/Source/Editor/MovieSceneTools/Private/TrackEditors/AudioTrackEditor.h
+++ b/Engine/Source/Editor/MovieSceneTools/Private/TrackEditors/AudioTrackEditor.h
@@ -34,6 +34,7 @@ public:
 	// ISequencerTrackEditor interface
 
 	virtual void BuildAddTrackMenu(FMenuBuilder& MenuBuilder) override;
+	virtual TSharedPtr<SWidget> BuildOutlinerEditWidget(const FGuid& ObjectBinding, UMovieSceneTrack* Track, const FBuildEditWidgetParams& Params) override;
 	virtual bool HandleAssetAdded(UObject* Asset, const FGuid& TargetObjectGuid) override;
 	virtual TSharedRef<ISequencerSection> MakeSectionInterface(UMovieSceneSection& SectionObject, UMovieSceneTrack& Track) override;
 	virtual bool SupportsType(TSubclassOf<UMovieSceneTrack> Type) const override;
@@ -45,12 +46,18 @@ protected:
 	bool AddNewMasterSound(float KeyTime, class USoundBase* Sound);
 
 	/** Delegate for AnimatablePropertyChanged in HandleAssetAdded for attached sounds */
-	bool AddNewAttachedSound(float KeyTime, class USoundBase* Sound, TArray<UObject*> ObjectsToAttachTo);
+	bool AddNewAttachedSound(float KeyTime, class USoundBase* Sound, TArray<TWeakObjectPtr<UObject>> ObjectsToAttachTo);
 
 private:
 
 	/** Callback for executing the "Add Event Track" menu entry. */
 	void HandleAddAudioTrackMenuEntryExecute();
+
+	/** Audio sub menu */
+	TSharedRef<SWidget> BuildAudioSubMenu(UMovieSceneTrack* Track);
+
+	/** Audio asset selected */
+	void OnAudioAssetSelected(const FAssetData& AssetData, UMovieSceneTrack* Track);
 };
 
 
diff --git a/Engine/Source/Editor/MovieSceneTools/Private/TrackEditors/CameraAnimTrackEditor.cpp b/Engine/Source/Editor/MovieSceneTools/Private/TrackEditors/CameraAnimTrackEditor.cpp
new file mode 100644
--- /dev/null
+++ b/Engine/Source/Editor/MovieSceneTools/Private/TrackEditors/CameraAnimTrackEditor.cpp
@@ -0,0 +1,290 @@
+// Copyright 1998-2016 Epic Games, Inc. All Rights Reserved.
+
+#include "MovieSceneToolsPrivatePCH.h"
+#include "MovieScene.h"
+#include "MovieSceneSection.h"
+#include "MovieSceneCameraAnimTrack.h"
+#include "MovieSceneCameraAnimSection.h"
+#include "ISequencerSection.h"
+#include "ISectionLayoutBuilder.h"
+#include "MovieSceneTrackEditor.h"
+#include "CameraAnimTrackEditor.h"
+#include "AssetRegistryModule.h"
+#include "ContentBrowserModule.h"
+#include "SequencerUtilities.h"
+
+#define LOCTEXT_NAMESPACE "FCameraAnimTrackEditor"
+
+
+class FCameraAnimSection : public ISequencerSection
+{
+public:
+	FCameraAnimSection(UMovieSceneSection& InSection)
+		: Section( InSection )
+	{ }
+
+	/** ISequencerSection interface */
+	virtual UMovieSceneSection* GetSectionObject() override
+	{ 
+		return &Section;
+	}
+
+	virtual FText GetDisplayName() const override
+	{ 
+		return LOCTEXT("CameraAnimSection", "Camera Anim");
+	}
+	
+	virtual FText GetSectionTitle() const override 
+	{ 
+		UMovieSceneCameraAnimSection const* const AnimSection = Cast<UMovieSceneCameraAnimSection>(&Section);
+		UCameraAnim const* const Anim = AnimSection ? AnimSection->GetCameraAnim() : nullptr;
+		if (Anim)
+		{
+			return FText::FromString(Anim->GetName());
+		}
+		return LOCTEXT("NoCameraAnimSection", "No Camera Anim");
+	}
+
+	virtual void GenerateSectionLayout( class ISectionLayoutBuilder& LayoutBuilder ) const override
+	{
+// 		UMovieSceneCameraAnimSection* PathSection = Cast<UMovieSceneCameraAnimSection>(&Section);
+// 		LayoutBuilder.AddKeyArea("Timing", LOCTEXT("TimingArea", "Timing"), MakeShareable( new FFloatCurveKeyArea ( &PathSection->GetTimingCurve( ), PathSection ) ) );
+	}
+
+	virtual int32 OnPaintSection( FSequencerSectionPainter& InPainter ) const override 
+	{
+		return InPainter.PaintSectionBackground();
+	}
+
+private:
+	/** The section we are visualizing */
+	UMovieSceneSection& Section;
+};
+
+
+FCameraAnimTrackEditor::FCameraAnimTrackEditor(TSharedRef<ISequencer> InSequencer)
+	: FMovieSceneTrackEditor( InSequencer ) 
+{ 
+}
+
+
+TSharedRef<ISequencerTrackEditor> FCameraAnimTrackEditor::CreateTrackEditor(TSharedRef<ISequencer> InSequencer)
+{
+	return MakeShareable(new FCameraAnimTrackEditor(InSequencer));
+}
+
+
+bool FCameraAnimTrackEditor::SupportsType(TSubclassOf<UMovieSceneTrack> Type) const
+{
+	return Type == UMovieSceneCameraAnimTrack::StaticClass();
+}
+
+
+TSharedRef<ISequencerSection> FCameraAnimTrackEditor::MakeSectionInterface( UMovieSceneSection& SectionObject, UMovieSceneTrack& Track )
+{
+	check( SupportsType( SectionObject.GetOuter()->GetClass() ) );
+	return MakeShareable(new FCameraAnimSection(SectionObject));
+}
+
+
+void FCameraAnimTrackEditor::AddKey(const FGuid& ObjectGuid)
+{
+	// Load the asset registry module
+	FAssetRegistryModule& AssetRegistryModule = FModuleManager::LoadModuleChecked<FAssetRegistryModule>(TEXT("AssetRegistry"));
+
+	// Collect a full list of assets with the specified class
+	TArray<FAssetData> AssetDataList;
+	AssetRegistryModule.Get().GetAssetsByClass(UCameraAnim::StaticClass()->GetFName(), AssetDataList);
+
+	if (AssetDataList.Num())
+	{
+		TSharedPtr< SWindow > Parent = FSlateApplication::Get().GetActiveTopLevelWindow();
+		if (Parent.IsValid())
+		{
+			FSlateApplication::Get().PushMenu(
+				Parent.ToSharedRef(),
+				FWidgetPath(),
+				BuildCameraAnimSubMenu(ObjectGuid),
+				FSlateApplication::Get().GetCursorPos(),
+				FPopupTransitionEffect(FPopupTransitionEffect::TypeInPopup)
+				);
+		}
+	}
+}
+
+
+bool FCameraAnimTrackEditor::HandleAssetAdded(UObject* Asset, const FGuid& TargetObjectGuid)
+{
+	UCameraAnim* const CameraAnim = Cast<UCameraAnim>(Asset);
+	if (CameraAnim)
+	{
+		if (TargetObjectGuid.IsValid())
+		{
+			TArray<TWeakObjectPtr<UObject>> OutObjects;
+
+			GetSequencer()->GetRuntimeObjects(GetSequencer()->GetFocusedMovieSceneSequenceInstance(), TargetObjectGuid, OutObjects);
+			AnimatablePropertyChanged(FOnKeyProperty::CreateRaw(this, &FCameraAnimTrackEditor::AddKeyInternal, OutObjects, CameraAnim));
+
+			return true;
+		}
+	}
+
+	return false;
+}
+
+
+
+void FCameraAnimTrackEditor::BuildObjectBindingTrackMenu(FMenuBuilder& MenuBuilder, const FGuid& ObjectBinding, const UClass* ObjectClass)
+{
+	// only offer this track if we can find a camera component
+	UCameraComponent const* const CamComponent = AcquireCameraComponentFromObjectGuid(ObjectBinding);
+	if (CamComponent)
+	{
+		const TSharedPtr<ISequencer> ParentSequencer = GetSequencer();
+
+		// Load the asset registry module
+		FAssetRegistryModule& AssetRegistryModule = FModuleManager::LoadModuleChecked<FAssetRegistryModule>(TEXT("AssetRegistry"));
+
+		// Collect a full list of assets with the specified class
+		TArray<FAssetData> AssetDataList;
+		AssetRegistryModule.Get().GetAssetsByClass(UCameraAnim::StaticClass()->GetFName(), AssetDataList);
+
+		if (AssetDataList.Num())
+		{
+			MenuBuilder.AddSubMenu(
+				LOCTEXT("AddCameraAnim", "Camera Anim"), NSLOCTEXT("Sequencer", "AddCameraAnimTooltip", "Adds an additive camera animation track."),
+				FNewMenuDelegate::CreateRaw(this, &FCameraAnimTrackEditor::AddCameraAnimSubMenu, ObjectBinding)
+				);
+		}
+	}
+}
+
+TSharedRef<SWidget> FCameraAnimTrackEditor::BuildCameraAnimSubMenu(FGuid ObjectBinding)
+{
+	FMenuBuilder MenuBuilder(true, nullptr);
+
+	AddCameraAnimSubMenu(MenuBuilder, ObjectBinding);
+
+	return MenuBuilder.MakeWidget();
+}
+
+
+void FCameraAnimTrackEditor::AddCameraAnimSubMenu(FMenuBuilder& MenuBuilder, FGuid ObjectBinding)
+{
+	FAssetPickerConfig AssetPickerConfig;
+	{
+		AssetPickerConfig.OnAssetSelected = FOnAssetSelected::CreateRaw(this, &FCameraAnimTrackEditor::OnCameraAnimAssetSelected, ObjectBinding);
+		AssetPickerConfig.bAllowNullSelection = false;
+		AssetPickerConfig.InitialAssetViewType = EAssetViewType::List;
+		AssetPickerConfig.Filter.ClassNames.Add(UCameraAnim::StaticClass()->GetFName());
+//		AssetPickerConfig.Filter.TagsAndValues.Add(TEXT("Skeleton"), FAssetData(Skeleton).GetExportTextName());
+	}
+
+	FContentBrowserModule& ContentBrowserModule = FModuleManager::Get().LoadModuleChecked<FContentBrowserModule>(TEXT("ContentBrowser"));
+
+	TSharedPtr<SBox> MenuEntry = SNew(SBox)
+		.WidthOverride(300.0f)
+		.HeightOverride(300.f)
+		[
+			ContentBrowserModule.Get().CreateAssetPicker(AssetPickerConfig)
+		];
+
+	MenuBuilder.AddWidget(MenuEntry.ToSharedRef(), FText::GetEmpty(), true);
+}
+
+TSharedPtr<SWidget> FCameraAnimTrackEditor::BuildOutlinerEditWidget(const FGuid& ObjectBinding, UMovieSceneTrack* Track, const FBuildEditWidgetParams& Params)
+{
+	// Create a container edit box
+	return SNew(SHorizontalBox)
+
+	// Add the camera anim combo box
+	+ SHorizontalBox::Slot()
+	.AutoWidth()
+	.VAlign(VAlign_Center)
+	[
+		FSequencerUtilities::MakeAddButton(LOCTEXT("AddCameraAnim", "Camera Anim"), FOnGetContent::CreateSP(this, &FCameraAnimTrackEditor::BuildCameraAnimSubMenu, ObjectBinding), Params.NodeIsHovered)
+	];
+}
+
+void FCameraAnimTrackEditor::OnCameraAnimAssetSelected(const FAssetData& AssetData, FGuid ObjectBinding)
+{
+	FSlateApplication::Get().DismissAllMenus();
+
+	UObject* SelectedObject = AssetData.GetAsset();
+
+	if (SelectedObject && SelectedObject->IsA(UCameraAnim::StaticClass()))
+	{
+		UCameraAnim* const CameraAnim = CastChecked<UCameraAnim>(AssetData.GetAsset());
+
+		TArray<TWeakObjectPtr<UObject>> OutObjects;
+
+		GetSequencer()->GetRuntimeObjects(GetSequencer()->GetFocusedMovieSceneSequenceInstance(), ObjectBinding, OutObjects);
+		AnimatablePropertyChanged(FOnKeyProperty::CreateRaw(this, &FCameraAnimTrackEditor::AddKeyInternal, OutObjects, CameraAnim));
+	}
+}
+
+
+
+bool FCameraAnimTrackEditor::AddKeyInternal(float KeyTime, const TArray<TWeakObjectPtr<UObject>> Objects, UCameraAnim* CameraAnim)
+{
+	bool bHandleCreated = false;
+	bool bTrackCreated = false;
+	bool bTrackModified = false;
+
+	for (int32 ObjectIndex = 0; ObjectIndex < Objects.Num(); ++ObjectIndex)
+	{
+		UObject* Object = Objects[ObjectIndex].Get();
+
+		FFindOrCreateHandleResult HandleResult = FindOrCreateHandleToObject(Object);
+		FGuid ObjectHandle = HandleResult.Handle;
+		bHandleCreated |= HandleResult.bWasCreated;
+		if (ObjectHandle.IsValid())
+		{
+			FFindOrCreateTrackResult TrackResult = FindOrCreateTrackForObject(ObjectHandle, UMovieSceneCameraAnimTrack::StaticClass());
+			UMovieSceneTrack* Track = TrackResult.Track;
+			bTrackCreated |= TrackResult.bWasCreated;
+
+			if (ensure(Track))
+			{
+				Cast<UMovieSceneCameraAnimTrack>(Track)->AddNewCameraAnim(KeyTime, CameraAnim);
+				bTrackModified = true;
+			}
+		}
+	}
+
+	return bHandleCreated || bTrackCreated || bTrackModified;
+}
+
+
+UCameraComponent* FCameraAnimTrackEditor::AcquireCameraComponentFromObjectGuid(const FGuid& Guid)
+{
+	TArray<TWeakObjectPtr<UObject>> OutObjects;
+	GetSequencer()->GetRuntimeObjects(GetSequencer()->GetFocusedMovieSceneSequenceInstance(), Guid, OutObjects);
+
+	USkeleton* Skeleton = nullptr;
+	for (int32 i = 0; i < OutObjects.Num(); ++i)
+	{
+		UObject* const Obj = OutObjects[i].Get();
+	
+		if (AActor* const Actor = Cast<AActor>(Obj))
+		{
+			UCameraComponent* const CameraComp = MovieSceneHelpers::CameraComponentFromActor(Actor);
+			if (CameraComp)
+			{
+				return CameraComp;
+			}
+		}
+		else if (UCameraComponent* const CameraComp = Cast<UCameraComponent>(Obj))
+		{
+			if (CameraComp->bIsActive)
+			{
+				return CameraComp;
+			}
+		}
+	}
+
+	return nullptr;
+}
+
+
+#undef LOCTEXT_NAMESPACE
diff --git a/Engine/Source/Editor/MovieSceneTools/Private/TrackEditors/CameraAnimTrackEditor.h b/Engine/Source/Editor/MovieSceneTools/Private/TrackEditors/CameraAnimTrackEditor.h
new file mode 100644
--- /dev/null
+++ b/Engine/Source/Editor/MovieSceneTools/Private/TrackEditors/CameraAnimTrackEditor.h
@@ -0,0 +1,58 @@
+// Copyright 1998-2016 Epic Games, Inc. All Rights Reserved.
+
+#pragma once
+
+#include "MovieSceneCameraAnimSection.h"
+#include "MovieSceneCameraAnimTrack.h"
+
+/**
+ * Tools for playing a camera anim
+ */
+class FCameraAnimTrackEditor : public FMovieSceneTrackEditor
+{
+public:
+
+	/**
+	 * Constructor
+	 *
+	 * @param InSequencer The sequencer instance to be used by this tool
+	 */
+	FCameraAnimTrackEditor(TSharedRef<ISequencer> InSequencer);
+
+	/** Virtual destructor. */
+	virtual ~FCameraAnimTrackEditor() { }
+
+	/**
+	 * Creates an instance of this class.  Called by a sequencer 
+	 *
+	 * @param OwningSequencer The sequencer instance to be used by this tool
+	 * @return The new instance of this class
+	 */
+	static TSharedRef<ISequencerTrackEditor> CreateTrackEditor( TSharedRef<ISequencer> OwningSequencer );
+
+public:
+
+	// ISequencerTrackEditor interface
+	virtual void AddKey(const FGuid& ObjectGuid) override;
+	virtual void BuildObjectBindingTrackMenu(FMenuBuilder& MenuBuilder, const FGuid& ObjectBinding, const UClass* ObjectClass) override;
+	virtual bool HandleAssetAdded(UObject* Asset, const FGuid& TargetObjectGuid) override;
+	virtual TSharedRef<ISequencerSection> MakeSectionInterface(UMovieSceneSection& SectionObject, UMovieSceneTrack& Track) override;
+	virtual bool SupportsType(TSubclassOf<UMovieSceneTrack> Type) const override;
+	virtual TSharedPtr<SWidget> BuildOutlinerEditWidget(const FGuid& ObjectBinding, UMovieSceneTrack* Track, const FBuildEditWidgetParams& Params) override;
+
+private:
+
+	/** Delegate for AnimatablePropertyChanged in AddKey */
+	bool AddKeyInternal(float KeyTime, const TArray<TWeakObjectPtr<UObject>> Objects, UCameraAnim* CameraAnim);
+
+	/** Animation sub menu */
+	TSharedRef<SWidget> BuildCameraAnimSubMenu(FGuid ObjectBinding);
+	void AddCameraAnimSubMenu(FMenuBuilder& MenuBuilder, FGuid ObjectBinding);
+
+	/** Animation asset selected */
+	void OnCameraAnimAssetSelected(const FAssetData& AssetData, FGuid ObjectBinding);
+
+	class UCameraComponent* AcquireCameraComponentFromObjectGuid(const FGuid& Guid);
+};
+
+
diff --git a/Engine/Source/Editor/MovieSceneTools/Private/TrackEditors/CameraCutTrackEditor.cpp b/Engine/Source/Editor/MovieSceneTools/Private/TrackEditors/CameraCutTrackEditor.cpp
--- a/Engine/Source/Editor/MovieSceneTools/Private/TrackEditors/CameraCutTrackEditor.cpp
+++ b/Engine/Source/Editor/MovieSceneTools/Private/TrackEditors/CameraCutTrackEditor.cpp
@@ -115,17 +115,30 @@ bool FCameraCutTrackEditor::SupportsType(TSubclassOf<UMovieSceneTrack> Type) con
 
 void FCameraCutTrackEditor::Tick(float DeltaTime)
 {
-	if (FSlateThrottleManager::Get().IsAllowingExpensiveTasks())
+	TSharedPtr<ISequencer> SequencerPin = GetSequencer();
+	if (!SequencerPin.IsValid())
 	{
-		float SavedTime = GetSequencer()->GetGlobalTime();
+		return;
+	}
+
+	EMovieScenePlayerStatus::Type PlaybackState = SequencerPin->GetPlaybackStatus();
+
+	if (FSlateThrottleManager::Get().IsAllowingExpensiveTasks() && PlaybackState != EMovieScenePlayerStatus::Playing && PlaybackState != EMovieScenePlayerStatus::Scrubbing)
+	{
+		SequencerPin->EnterSilentMode();
+
+		float SavedTime = SequencerPin->GetGlobalTime();
 
-		if (ThumbnailPool->DrawThumbnails())
+		if (DeltaTime > 0.f && ThumbnailPool->DrawThumbnails())
 		{
-			GetSequencer()->SetGlobalTime(SavedTime);
+			SequencerPin->SetGlobalTimeDirectly(SavedTime);
 		}
+
+		SequencerPin->ExitSilentMode();
 	}
 }
 
+
 const FSlateBrush* FCameraCutTrackEditor::GetIconBrush() const
 {
 	return FEditorStyle::GetBrush("Sequencer.Tracks.CameraCut");
diff --git a/Engine/Source/Editor/MovieSceneTools/Private/TrackEditors/CameraShakeTrackEditor.cpp b/Engine/Source/Editor/MovieSceneTools/Private/TrackEditors/CameraShakeTrackEditor.cpp
new file mode 100644
--- /dev/null
+++ b/Engine/Source/Editor/MovieSceneTools/Private/TrackEditors/CameraShakeTrackEditor.cpp
@@ -0,0 +1,298 @@
+// Copyright 1998-2016 Epic Games, Inc. All Rights Reserved.
+
+#include "MovieSceneToolsPrivatePCH.h"
+#include "MovieScene.h"
+#include "MovieSceneSection.h"
+#include "MovieSceneCameraShakeTrack.h"
+#include "MovieSceneCameraShakeSection.h"
+#include "ISequencerSection.h"
+#include "ISectionLayoutBuilder.h"
+#include "MovieSceneTrackEditor.h"
+#include "CameraShakeTrackEditor.h"
+#include "AssetRegistryModule.h"
+#include "ContentBrowserModule.h"
+#include "SequencerUtilities.h"
+
+#define LOCTEXT_NAMESPACE "FCameraShakeTrackEditor"
+
+static const FName ParentClassTagName(TEXT("ParentClass"));
+static const FString CameraShakeClassPath(TEXT("Class'/Script/Engine.CameraShake'"));
+
+
+class FCameraShakeSection : public ISequencerSection
+{
+public:
+	FCameraShakeSection(UMovieSceneSection& InSection)
+		: Section(InSection)
+	{ }
+
+	/** ISequencerSection interface */
+	virtual UMovieSceneSection* GetSectionObject() override
+	{
+		return &Section;
+	}
+
+	virtual FText GetDisplayName() const override
+	{
+		return LOCTEXT("CameraShakeSection", "Camera Shake");
+	}
+
+	virtual FText GetSectionTitle() const override
+	{
+		UMovieSceneCameraShakeSection const* const ShakeSection = Cast<UMovieSceneCameraShakeSection>(&Section);
+		UClass const* const Shake = ShakeSection ? ShakeSection->GetCameraShakeClass() : nullptr;
+		if (Shake)
+		{
+			return FText::FromString(Shake->GetName());
+		}
+		return LOCTEXT("NoCameraShakeSection", "No Camera Shake");
+	}
+
+	virtual void GenerateSectionLayout(class ISectionLayoutBuilder& LayoutBuilder) const override
+	{
+// 		UMovieSceneCameraShakeSection* PathSection = Cast<UMovieSceneCameraShakeSection>(&Section);
+// 		LayoutBuilder.AddKeyArea("Timing", LOCTEXT("TimingArea", "Timing"), MakeShareable( new FFloatCurveKeyArea ( &PathSection->GetTimingCurve( ), PathSection ) ) );
+	}
+
+	virtual int32 OnPaintSection(FSequencerSectionPainter& InPainter) const override
+	{
+		return InPainter.PaintSectionBackground();
+	}
+
+private:
+
+	/** The section we are visualizing */
+	UMovieSceneSection& Section;
+};
+
+
+FCameraShakeTrackEditor::FCameraShakeTrackEditor(TSharedRef<ISequencer> InSequencer)
+	: FMovieSceneTrackEditor( InSequencer ) 
+{ 
+}
+
+
+TSharedRef<ISequencerTrackEditor> FCameraShakeTrackEditor::CreateTrackEditor(TSharedRef<ISequencer> InSequencer)
+{
+	return MakeShareable(new FCameraShakeTrackEditor(InSequencer));
+}
+
+
+bool FCameraShakeTrackEditor::SupportsType(TSubclassOf<UMovieSceneTrack> Type) const
+{
+	return Type == UMovieSceneCameraShakeTrack::StaticClass();
+}
+
+
+TSharedRef<ISequencerSection> FCameraShakeTrackEditor::MakeSectionInterface(UMovieSceneSection& SectionObject, UMovieSceneTrack& Track)
+{
+	check(SupportsType(SectionObject.GetOuter()->GetClass()));
+	return MakeShareable(new FCameraShakeSection(SectionObject));
+}
+
+
+void FCameraShakeTrackEditor::AddKey(const FGuid& ObjectGuid)
+{
+	// Load the asset registry module
+	FAssetRegistryModule& AssetRegistryModule = FModuleManager::LoadModuleChecked<FAssetRegistryModule>(TEXT("AssetRegistry"));
+
+	// Collect a full list of assets with the specified class
+	TArray<FAssetData> AssetDataList;
+	AssetRegistryModule.Get().GetAssetsByClass(UCameraShake::StaticClass()->GetFName(), AssetDataList);
+
+	if (AssetDataList.Num())
+	{
+		TSharedPtr< SWindow > Parent = FSlateApplication::Get().GetActiveTopLevelWindow();
+		if (Parent.IsValid())
+		{
+			FSlateApplication::Get().PushMenu(
+				Parent.ToSharedRef(),
+				FWidgetPath(),
+				BuildCameraShakeSubMenu(ObjectGuid),
+				FSlateApplication::Get().GetCursorPos(),
+				FPopupTransitionEffect(FPopupTransitionEffect::TypeInPopup)
+				);
+		}
+	}
+}
+
+bool FCameraShakeTrackEditor::HandleAssetAdded(UObject* Asset, const FGuid& TargetObjectGuid)
+{
+	if (TargetObjectGuid.IsValid())
+	{
+		UBlueprint* const SelectedObject = dynamic_cast<UBlueprint*>(Asset);
+		if (SelectedObject && SelectedObject->GeneratedClass && SelectedObject->GeneratedClass->IsChildOf(UCameraShake::StaticClass()))
+		{
+			TSubclassOf<UCameraShake> const ShakeClass = *(SelectedObject->GeneratedClass);
+
+			TArray<TWeakObjectPtr<UObject>> OutObjects;
+			GetSequencer()->GetRuntimeObjects(GetSequencer()->GetFocusedMovieSceneSequenceInstance(), TargetObjectGuid, OutObjects);
+			AnimatablePropertyChanged(FOnKeyProperty::CreateRaw(this, &FCameraShakeTrackEditor::AddKeyInternal, OutObjects, ShakeClass));
+
+			return true;
+		}
+	}
+
+	return false;
+}
+
+void FCameraShakeTrackEditor::BuildObjectBindingTrackMenu(FMenuBuilder& MenuBuilder, const FGuid& ObjectBinding, const UClass* ObjectClass)
+{
+	// only offer this track if we can find a camera component
+	UCameraComponent const* const CamComponent = AcquireCameraComponentFromObjectGuid(ObjectBinding);
+	if (CamComponent)
+	{
+		const TSharedPtr<ISequencer> ParentSequencer = GetSequencer();
+
+		// Load the asset registry module
+		FAssetRegistryModule& AssetRegistryModule = FModuleManager::LoadModuleChecked<FAssetRegistryModule>(TEXT("AssetRegistry"));
+
+		// Collect a full list of assets with the specified class
+		TArray<FAssetData> AssetDataList;
+		{
+			FARFilter Filter;
+			Filter.ClassNames.Add(UBlueprint::StaticClass()->GetFName());
+
+			TMultiMap<FName, FString> TagsAndValuesFilter;
+			TagsAndValuesFilter.Add(ParentClassTagName, CameraShakeClassPath);
+			Filter.TagsAndValues = TagsAndValuesFilter;
+
+			AssetRegistryModule.Get().GetAssets(Filter, AssetDataList);
+		}
+
+		if (AssetDataList.Num())
+		{
+			MenuBuilder.AddSubMenu(
+				LOCTEXT("AddCameraShake", "Camera Shake"), NSLOCTEXT("Sequencer", "AddCameraShakeTooltip", "Adds an additive camera shake track."),
+				FNewMenuDelegate::CreateRaw(this, &FCameraShakeTrackEditor::AddCameraShakeSubMenu, ObjectBinding)
+				);
+		}
+	}
+}
+
+TSharedRef<SWidget> FCameraShakeTrackEditor::BuildCameraShakeSubMenu(FGuid ObjectBinding)
+{
+	FMenuBuilder MenuBuilder(true, nullptr);
+
+	AddCameraShakeSubMenu(MenuBuilder, ObjectBinding);
+
+	return MenuBuilder.MakeWidget();
+}
+
+
+void FCameraShakeTrackEditor::AddCameraShakeSubMenu(FMenuBuilder& MenuBuilder, FGuid ObjectBinding)
+{
+	FAssetPickerConfig AssetPickerConfig;
+	{
+		AssetPickerConfig.OnAssetSelected = FOnAssetSelected::CreateRaw(this, &FCameraShakeTrackEditor::OnCameraShakeAssetSelected, ObjectBinding);
+		AssetPickerConfig.bAllowNullSelection = false;
+		AssetPickerConfig.InitialAssetViewType = EAssetViewType::List;
+		AssetPickerConfig.Filter.ClassNames.Add(UBlueprint::StaticClass()->GetFName());
+		AssetPickerConfig.Filter.TagsAndValues.Add(ParentClassTagName, CameraShakeClassPath);
+	}
+
+	FContentBrowserModule& ContentBrowserModule = FModuleManager::Get().LoadModuleChecked<FContentBrowserModule>(TEXT("ContentBrowser"));
+
+	TSharedPtr<SBox> MenuEntry = SNew(SBox)
+		.WidthOverride(300.0f)
+		.HeightOverride(300.f)
+		[
+			ContentBrowserModule.Get().CreateAssetPicker(AssetPickerConfig)
+		];
+
+	MenuBuilder.AddWidget(MenuEntry.ToSharedRef(), FText::GetEmpty(), true);
+}
+
+TSharedPtr<SWidget> FCameraShakeTrackEditor::BuildOutlinerEditWidget(const FGuid& ObjectBinding, UMovieSceneTrack* Track, const FBuildEditWidgetParams& Params)
+{
+	// Create a container edit box
+	return SNew(SHorizontalBox)
+
+	// Add the camera shake combo box
+	+ SHorizontalBox::Slot()
+	.AutoWidth()
+	.VAlign(VAlign_Center)
+	[
+		FSequencerUtilities::MakeAddButton(LOCTEXT("AddCameraShake", "Camera Shake"), FOnGetContent::CreateSP(this, &FCameraShakeTrackEditor::BuildCameraShakeSubMenu, ObjectBinding), Params.NodeIsHovered)
+	];
+}
+
+
+void FCameraShakeTrackEditor::OnCameraShakeAssetSelected(const FAssetData& AssetData, FGuid ObjectBinding)
+{
+	FSlateApplication::Get().DismissAllMenus();
+
+	UBlueprint* const SelectedObject = dynamic_cast<UBlueprint*>(AssetData.GetAsset());
+	if (SelectedObject && SelectedObject->GeneratedClass && SelectedObject->GeneratedClass->IsChildOf(UCameraShake::StaticClass()))
+	{
+		TSubclassOf<UCameraShake> const ShakeClass = *(SelectedObject->GeneratedClass);
+
+		TArray<TWeakObjectPtr<UObject>> OutObjects;
+
+		GetSequencer()->GetRuntimeObjects(GetSequencer()->GetFocusedMovieSceneSequenceInstance(), ObjectBinding, OutObjects);
+		AnimatablePropertyChanged(FOnKeyProperty::CreateRaw(this, &FCameraShakeTrackEditor::AddKeyInternal, OutObjects, ShakeClass));
+	}
+}
+
+
+bool FCameraShakeTrackEditor::AddKeyInternal(float KeyTime, const TArray<TWeakObjectPtr<UObject>> Objects, TSubclassOf<UCameraShake> ShakeClass)
+{
+	bool bHandleCreated = false;
+	bool bTrackCreated = false;
+	bool bTrackModified = false;
+
+	for (int32 ObjectIndex = 0; ObjectIndex < Objects.Num(); ++ObjectIndex)
+	{
+		UObject* Object = Objects[ObjectIndex].Get();
+
+		FFindOrCreateHandleResult HandleResult = FindOrCreateHandleToObject(Object);
+		FGuid ObjectHandle = HandleResult.Handle;
+		bHandleCreated |= HandleResult.bWasCreated;
+		if (ObjectHandle.IsValid())
+		{
+			FFindOrCreateTrackResult TrackResult = FindOrCreateTrackForObject(ObjectHandle, UMovieSceneCameraShakeTrack::StaticClass());
+			UMovieSceneTrack* Track = TrackResult.Track;
+			bTrackCreated |= TrackResult.bWasCreated;
+
+			if (ensure(Track))
+			{
+				Cast<UMovieSceneCameraShakeTrack>(Track)->AddNewCameraShake(KeyTime, ShakeClass);
+				bTrackModified = true;
+			}
+		}
+	}
+
+	return bHandleCreated || bTrackCreated || bTrackModified;
+}
+
+UCameraComponent* FCameraShakeTrackEditor::AcquireCameraComponentFromObjectGuid(const FGuid& Guid)
+{
+	TArray<TWeakObjectPtr<UObject>> OutObjects;
+	GetSequencer()->GetRuntimeObjects(GetSequencer()->GetFocusedMovieSceneSequenceInstance(), Guid, OutObjects);
+
+	USkeleton* Skeleton = nullptr;
+	for (int32 i = 0; i < OutObjects.Num(); ++i)
+	{
+		UObject* const Obj = OutObjects[i].Get();
+		if (AActor* const Actor = Cast<AActor>(Obj))
+		{
+			UCameraComponent* const CameraComp = MovieSceneHelpers::CameraComponentFromActor(Actor);
+			if (CameraComp)
+			{
+				return CameraComp;
+			}
+		}
+		else if (UCameraComponent* const CameraComp = Cast<UCameraComponent>(Obj))
+		{
+			if (CameraComp->bIsActive)
+			{
+				return CameraComp;
+			}
+		}
+	}
+
+	return nullptr;
+}
+
+
+#undef LOCTEXT_NAMESPACE
diff --git a/Engine/Source/Editor/MovieSceneTools/Private/TrackEditors/CameraShakeTrackEditor.h b/Engine/Source/Editor/MovieSceneTools/Private/TrackEditors/CameraShakeTrackEditor.h
new file mode 100644
--- /dev/null
+++ b/Engine/Source/Editor/MovieSceneTools/Private/TrackEditors/CameraShakeTrackEditor.h
@@ -0,0 +1,58 @@
+// Copyright 1998-2016 Epic Games, Inc. All Rights Reserved.
+
+#pragma once
+
+#include "MovieSceneCameraShakeSection.h"
+#include "MovieSceneCameraShakeTrack.h"
+
+/**
+* Tools for playing a camera shake
+*/
+class FCameraShakeTrackEditor : public FMovieSceneTrackEditor
+{
+public:
+
+	/**
+	* Constructor
+	*
+	* @param InSequencer The sequencer instance to be used by this tool
+	*/
+	FCameraShakeTrackEditor(TSharedRef<ISequencer> InSequencer);
+
+	/** Virtual destructor. */
+	virtual ~FCameraShakeTrackEditor() { }
+
+	/**
+	* Creates an instance of this class.  Called by a sequencer
+	*
+	* @param OwningSequencer The sequencer instance to be used by this tool
+	* @return The new instance of this class
+	*/
+	static TSharedRef<ISequencerTrackEditor> CreateTrackEditor(TSharedRef<ISequencer> OwningSequencer);
+
+public:
+
+	// ISequencerTrackEditor interface
+	virtual void AddKey(const FGuid& ObjectGuid) override;
+	virtual void BuildObjectBindingTrackMenu(FMenuBuilder& MenuBuilder, const FGuid& ObjectBinding, const UClass* ObjectClass) override;
+	virtual bool HandleAssetAdded(UObject* Asset, const FGuid& TargetObjectGuid) override;
+	virtual TSharedRef<ISequencerSection> MakeSectionInterface(UMovieSceneSection& SectionObject, UMovieSceneTrack& Track) override;
+	virtual bool SupportsType(TSubclassOf<UMovieSceneTrack> Type) const override;
+	virtual TSharedPtr<SWidget> BuildOutlinerEditWidget(const FGuid& ObjectBinding, UMovieSceneTrack* Track, const FBuildEditWidgetParams& Params) override;
+
+private:
+
+	/** Delegate for AnimatablePropertyChanged in AddKey */
+	bool AddKeyInternal(float KeyTime, const TArray<TWeakObjectPtr<UObject>> Objects, TSubclassOf<UCameraShake> ShakeClass);
+
+	/** Animation sub menu */
+	TSharedRef<SWidget> BuildCameraShakeSubMenu(FGuid ObjectBinding);
+	void AddCameraShakeSubMenu(FMenuBuilder& MenuBuilder, FGuid ObjectBinding);
+
+	/** Animation asset selected */
+	void OnCameraShakeAssetSelected(const FAssetData& AssetData, FGuid ObjectBinding);
+
+	class UCameraComponent* AcquireCameraComponentFromObjectGuid(const FGuid& Guid);
+};
+
+
diff --git a/Engine/Source/Editor/MovieSceneTools/Private/TrackEditors/CinematicShotTrackEditor.cpp b/Engine/Source/Editor/MovieSceneTools/Private/TrackEditors/CinematicShotTrackEditor.cpp
--- a/Engine/Source/Editor/MovieSceneTools/Private/TrackEditors/CinematicShotTrackEditor.cpp
+++ b/Engine/Source/Editor/MovieSceneTools/Private/TrackEditors/CinematicShotTrackEditor.cpp
@@ -155,14 +155,26 @@ bool FCinematicShotTrackEditor::SupportsType(TSubclassOf<UMovieSceneTrack> Type)
 
 void FCinematicShotTrackEditor::Tick(float DeltaTime)
 {
-	if (FSlateThrottleManager::Get().IsAllowingExpensiveTasks())
+	TSharedPtr<ISequencer> SequencerPin = GetSequencer();
+	if (!SequencerPin.IsValid())
 	{
-		float SavedTime = GetSequencer()->GetGlobalTime();
+		return;
+	}
+
+	EMovieScenePlayerStatus::Type PlaybackState = SequencerPin->GetPlaybackStatus();
+
+	if (FSlateThrottleManager::Get().IsAllowingExpensiveTasks() && PlaybackState != EMovieScenePlayerStatus::Playing && PlaybackState != EMovieScenePlayerStatus::Scrubbing)
+	{
+		SequencerPin->EnterSilentMode();
+
+		float SavedTime = SequencerPin->GetGlobalTime();
 
-		if (ThumbnailPool->DrawThumbnails())
+		if (DeltaTime > 0.f && ThumbnailPool->DrawThumbnails())
 		{
-			GetSequencer()->SetGlobalTime(SavedTime);
+			SequencerPin->SetGlobalTimeDirectly(SavedTime);
 		}
+
+		SequencerPin->ExitSilentMode();
 	}
 }
 
@@ -205,11 +217,17 @@ UMovieSceneSubSection* FCinematicShotTrackEditor::CreateShotInternal(FString& Ne
 
 	FString NewShotDirectory = MovieSceneToolHelpers::ComposeShotName(NewShotPrefix, NewShotNumber, INDEX_NONE);
 	FString NewShotPath = SequencePath;
-	NewShotPath /= NewShotDirectory; // put this in the shot directory, ie. /Game/cine/max/sht0010
+
+	FString ShotDirectory = ProjectSettings->ShotDirectory;
+	if (!ShotDirectory.IsEmpty())
+	{
+		NewShotPath /= ShotDirectory;
+	}
+	NewShotPath /= NewShotDirectory; // put this in the shot directory, ie. /Game/cine/max/shots/shot0010
 
 	// Make sure this shot path is unique
 	FString NewPackageName = NewShotPath;
-	NewPackageName /= NewShotName; // ie. /Game/cine/max/sht0010/sht0010_001
+	NewPackageName /= NewShotName; // ie. /Game/cine/max/shots/shot0010/shot0010_001
 	if (!IsPackageNameUnique(ObjectList, NewPackageName))
 	{
 		while (1)
@@ -218,6 +236,10 @@ UMovieSceneSubSection* FCinematicShotTrackEditor::CreateShotInternal(FString& Ne
 			NewShotName = MovieSceneToolHelpers::ComposeShotName(NewShotPrefix, NewShotNumber, NewTakeNumber);
 			NewShotDirectory = MovieSceneToolHelpers::ComposeShotName(NewShotPrefix, NewShotNumber, INDEX_NONE);
 			NewShotPath = SequencePath;
+			if (!ShotDirectory.IsEmpty())
+			{
+				NewShotPath /= ShotDirectory;
+			}
 			NewShotPath /= NewShotDirectory;
 
 			NewPackageName = NewShotPath;
@@ -267,16 +289,23 @@ UMovieSceneSubSection* FCinematicShotTrackEditor::CreateShotInternal(FString& Ne
 		Duration = ShotToDuplicate->GetEndTime() - ShotToDuplicate->GetStartTime();
 	}
 
-	// Create a cinematic shot section and set the time to be at the end of this shot. 
 	UMovieSceneCinematicShotTrack* CinematicShotTrack = FindOrCreateCinematicShotTrack();
+
+	// Cut any existing overlapping sections on the top row
+	for (auto ShotSection : CinematicShotTrack->GetAllSections())
+	{
+		if (ShotSection->GetRowIndex() == 0 && NewShotStartTime > ShotSection->GetStartTime() && NewShotStartTime < ShotSection->GetEndTime())
+		{
+			ShotSection->SetEndTime(NewShotStartTime);
+		}
+	}
+
+	// Create a cinematic shot section and set the time to be at the end of this shot. 
 	UMovieSceneSubSection* NewSection = CinematicShotTrack->AddSequence(NewSequence, NewShotStartTime, Duration, bInsertShot);	
 
 	// Update the track instances to generate the new section
 	GetSequencer()->UpdateRuntimeInstances();
 
-	// Focus on the newly created shot
-	//GetSequencer()->FocusSequenceInstance(GetSequencer()->GetSequenceInstanceForSection(*NewSection));
-
 	return NewSection;
 }
 
@@ -463,14 +492,14 @@ TSharedRef<SWidget> FCinematicShotTrackEditor::HandleAddCinematicShotComboButton
 
 	MenuBuilder.AddMenuEntry(
 			LOCTEXT("InsertShot", "Insert Shot"),
-			LOCTEXT("InsertShotTooltip", "Insert new shot"),
+			LOCTEXT("InsertShotTooltip", "Insert new shot at current time"),
 			FSlateIcon(),
 			FUIAction(FExecuteAction::CreateSP(this, &FCinematicShotTrackEditor::InsertShotAtCurrentTime))
 	);
 
 	MenuBuilder.AddMenuEntry(
 			LOCTEXT("InsertFiller", "Insert Filler"),
-			LOCTEXT("InsertFillerTooltip", "Insert filler"),
+			LOCTEXT("InsertFillerTooltip", "Insert filler at current time"),
 			FSlateIcon(),
 			FUIAction(FExecuteAction::CreateSP(this, &FCinematicShotTrackEditor::InsertFillerAtCurrentTime))
 	);
@@ -631,7 +660,7 @@ bool FCinematicShotTrackEditor::CanAddSubSequence(const UMovieSceneSequence& Seq
 	return !SequenceCinematicShotTrack->ContainsSequence(*FocusedSequence, true);
 }
 
-void FCinematicShotTrackEditor::OnUpdateCameraCut(UObject* CameraObject)
+void FCinematicShotTrackEditor::OnUpdateCameraCut(UObject* CameraObject, bool bJumpCut)
 {
 	// Keep track of the camera when it switches so that the thumbnail can be drawn with the correct camera
 	CinematicShotCamera = Cast<AActor>(CameraObject);
diff --git a/Engine/Source/Editor/MovieSceneTools/Private/TrackEditors/CinematicShotTrackEditor.h b/Engine/Source/Editor/MovieSceneTools/Private/TrackEditors/CinematicShotTrackEditor.h
--- a/Engine/Source/Editor/MovieSceneTools/Private/TrackEditors/CinematicShotTrackEditor.h
+++ b/Engine/Source/Editor/MovieSceneTools/Private/TrackEditors/CinematicShotTrackEditor.h
@@ -147,7 +147,7 @@ private:
 	bool CanAddSubSequence(const UMovieSceneSequence& Sequence) const;
 
 	/** Called when our sequencer wants to switch cameras */
-	void OnUpdateCameraCut(UObject* CameraObject);
+	void OnUpdateCameraCut(UObject* CameraObject, bool bJumpCut);
 
 	/** Callback for AnimatablePropertyChanged in HandleAssetAdded. */
 	bool HandleSequenceAdded(float KeyTime, UMovieSceneSequence* Sequence);
diff --git a/Engine/Source/Editor/MovieSceneTools/Private/TrackEditors/LevelVisibilityTrackEditor.cpp b/Engine/Source/Editor/MovieSceneTools/Private/TrackEditors/LevelVisibilityTrackEditor.cpp
--- a/Engine/Source/Editor/MovieSceneTools/Private/TrackEditors/LevelVisibilityTrackEditor.cpp
+++ b/Engine/Source/Editor/MovieSceneTools/Private/TrackEditors/LevelVisibilityTrackEditor.cpp
@@ -25,6 +25,10 @@ bool FLevelVisibilityTrackEditor::SupportsType( TSubclassOf<UMovieSceneTrack> Ty
 	return Type == UMovieSceneLevelVisibilityTrack::StaticClass();
 }
 
+const FSlateBrush* FLevelVisibilityTrackEditor::GetIconBrush() const
+{
+	return FEditorStyle::GetBrush("Sequencer.Tracks.LevelVisibility");
+}
 
 TSharedRef<ISequencerSection> FLevelVisibilityTrackEditor::MakeSectionInterface( UMovieSceneSection& SectionObject, UMovieSceneTrack& Track )
 {
diff --git a/Engine/Source/Editor/MovieSceneTools/Private/TrackEditors/ParticleTrackEditor.cpp b/Engine/Source/Editor/MovieSceneTools/Private/TrackEditors/ParticleTrackEditor.cpp
--- a/Engine/Source/Editor/MovieSceneTools/Private/TrackEditors/ParticleTrackEditor.cpp
+++ b/Engine/Source/Editor/MovieSceneTools/Private/TrackEditors/ParticleTrackEditor.cpp
@@ -38,7 +38,7 @@ namespace AnimatableParticleEditorConstants
 
 FParticleSection::FParticleSection( UMovieSceneSection& InSection, TSharedRef<ISequencer>InOwningSequencer )
 	: Section( InSection )
-	, OwningSequencer( InOwningSequencer )
+	, OwningSequencerPtr( InOwningSequencer )
 {
 	ParticleKeyEnum = FindObject<UEnum>( ANY_PACKAGE, TEXT( "EParticleKey" ) );
 	checkf( ParticleKeyEnum != nullptr, TEXT( "FParticleSection could not find the EParticleKey UEnum by name." ) )
@@ -77,6 +77,13 @@ void FParticleSection::GenerateSectionLayout( class ISectionLayoutBuilder& Layou
 
 int32 FParticleSection::OnPaintSection( FSequencerSectionPainter& InPainter ) const
 {
+	TSharedPtr<ISequencer> OwningSequencer = OwningSequencerPtr.Pin();
+
+	if (!OwningSequencer.IsValid())
+	{
+		return InPainter.LayerId + 1;
+	}
+
 	const ESlateDrawEffect::Type DrawEffects = InPainter.bParentEnabled ? ESlateDrawEffect::None : ESlateDrawEffect::DisabledEffect;
 	UMovieSceneParticleSection* AnimSection = Cast<UMovieSceneParticleSection>( &Section );
 	const FTimeToPixel& TimeToPixelConverter = InPainter.GetTimeConverter();
@@ -93,7 +100,6 @@ int32 FParticleSection::OnPaintSection( FSequencerSectionPainter& InPainter ) co
 		if ( ParentTrack != nullptr )
 		{
 			TrackColor = ParentTrack->GetColorTint();
-			TrackColor.A *= 0.4;
 
 			FGuid ObjectHandle;
 			for ( const FMovieSceneBinding& Binding : OwningSequencer->GetFocusedMovieSceneSequence()->GetMovieScene()->GetBindings() )
@@ -235,6 +241,24 @@ const FSlateBrush* FParticleSection::GetKeyBrush( FKeyHandle KeyHandle ) const
 	return nullptr;
 }
 
+FVector2D FParticleSection::GetKeyBrushOrigin( FKeyHandle KeyHandle ) const
+{
+	UMovieSceneParticleSection* ParticleSection = Cast<UMovieSceneParticleSection>( &Section );
+	if ( ParticleSection != nullptr )
+	{
+		FIntegralKey ParticleKey = ParticleSection->GetParticleCurve().GetKey(KeyHandle);
+		if ( (EParticleKey::Type)ParticleKey.Value == EParticleKey::Activate )
+		{
+			return FVector2D(-1.0f, 1.0f);
+		}
+		else if ( (EParticleKey::Type)ParticleKey.Value == EParticleKey::Deactivate )
+		{
+			return FVector2D(1.0f, 1.0f);
+		}
+	}
+	return FVector2D(0.0f, 0.0f);
+}
+
 
 FParticleTrackEditor::FParticleTrackEditor( TSharedRef<ISequencer> InSequencer )
 	: FMovieSceneTrackEditor( InSequencer ) 
@@ -280,21 +304,21 @@ void FParticleTrackEditor::BuildObjectBindingTrackMenu(FMenuBuilder& MenuBuilder
 
 void FParticleTrackEditor::AddParticleKey( const FGuid ObjectGuid )
 {
-	TArray<UObject*> OutObjects;
+	TArray<TWeakObjectPtr<UObject>> OutObjects;
 
 	GetSequencer()->GetRuntimeObjects( GetSequencer()->GetFocusedMovieSceneSequenceInstance(), ObjectGuid, OutObjects );
 	AnimatablePropertyChanged( FOnKeyProperty::CreateRaw( this, &FParticleTrackEditor::AddKeyInternal, OutObjects ) );
 }
 
 
-bool FParticleTrackEditor::AddKeyInternal( float KeyTime, const TArray<UObject*> Objects )
+bool FParticleTrackEditor::AddKeyInternal( float KeyTime, const TArray<TWeakObjectPtr<UObject>> Objects )
 {
 	bool bHandleCreated = false;
 	bool bTrackCreated = false;
 
 	for( int32 ObjectIndex = 0; ObjectIndex < Objects.Num(); ++ObjectIndex )
 	{
-		UObject* Object = Objects[ObjectIndex];
+		UObject* Object = Objects[ObjectIndex].Get();
 
 		FFindOrCreateHandleResult HandleResult = FindOrCreateHandleToObject( Object );
 		FGuid ObjectHandle = HandleResult.Handle;
diff --git a/Engine/Source/Editor/MovieSceneTools/Private/TrackEditors/ParticleTrackEditor.h b/Engine/Source/Editor/MovieSceneTools/Private/TrackEditors/ParticleTrackEditor.h
--- a/Engine/Source/Editor/MovieSceneTools/Private/TrackEditors/ParticleTrackEditor.h
+++ b/Engine/Source/Editor/MovieSceneTools/Private/TrackEditors/ParticleTrackEditor.h
@@ -46,7 +46,7 @@ public:
 private:
 
 	/** Delegate for AnimatablePropertyChanged in AddKey. */
-	virtual bool AddKeyInternal( float KeyTime, const TArray<UObject*> Objects);
+	virtual bool AddKeyInternal( float KeyTime, const TArray<TWeakObjectPtr<UObject>> Objects);
 };
 
 
@@ -72,6 +72,7 @@ public:
 	virtual void GenerateSectionLayout( class ISectionLayoutBuilder& LayoutBuilder ) const override;
 	virtual int32 OnPaintSection( FSequencerSectionPainter& InPainter ) const override;
 	virtual const FSlateBrush* GetKeyBrush( FKeyHandle KeyHandle ) const override;
+	virtual FVector2D GetKeyBrushOrigin( FKeyHandle KeyHandle ) const override;
 	virtual bool SectionIsResizable() const override { return false; }
 
 private:
@@ -80,7 +81,7 @@ private:
 	UMovieSceneSection& Section;
 
 	/** The sequencer that owns this section */
-	TSharedRef<ISequencer> OwningSequencer;
+	TWeakPtr<ISequencer> OwningSequencerPtr;
 
 	/** The UEnum for the EParticleKey enum */
 	const UEnum* ParticleKeyEnum;
diff --git a/Engine/Source/Editor/MovieSceneTools/Private/TrackEditors/PathTrackEditor.cpp b/Engine/Source/Editor/MovieSceneTools/Private/TrackEditors/PathTrackEditor.cpp
--- a/Engine/Source/Editor/MovieSceneTools/Private/TrackEditors/PathTrackEditor.cpp
+++ b/Engine/Source/Editor/MovieSceneTools/Private/TrackEditors/PathTrackEditor.cpp
@@ -115,7 +115,7 @@ void F3DPathTrackEditor::BuildObjectBindingTrackMenu(FMenuBuilder& MenuBuilder,
 bool F3DPathTrackEditor::IsActorPickable(const AActor* const ParentActor, FGuid ObjectBinding, UMovieSceneSection* InSection)
 {
 	// Can't pick the object that this track binds
-	TArray<UObject*> OutObjects;
+	TArray<TWeakObjectPtr<UObject>> OutObjects;
 	GetSequencer()->GetRuntimeObjects( GetSequencer()->GetFocusedMovieSceneSequenceInstance(), ObjectBinding, OutObjects);
 	if (OutObjects.Contains(ParentActor))
 	{
@@ -130,7 +130,7 @@ bool F3DPathTrackEditor::IsActorPickable(const AActor* const ParentActor, FGuid
 
 		if (ConstraintId.IsValid())
 		{
-			TArray<UObject*> ConstraintObjects;
+			TArray<TWeakObjectPtr<UObject>> ConstraintObjects;
 			GetSequencer()->GetRuntimeObjects( GetSequencer()->GetFocusedMovieSceneSequenceInstance(), ConstraintId, ConstraintObjects);
 			if (ConstraintObjects.Contains(ParentActor))
 			{
@@ -157,14 +157,7 @@ bool F3DPathTrackEditor::IsActorPickable(const AActor* const ParentActor, FGuid
 
 void F3DPathTrackEditor::ActorSocketPicked(const FName SocketName, USceneComponent* Component, AActor* ParentActor, FGuid ObjectGuid, UMovieSceneSection* Section)
 {
-	if (ObjectGuid.IsValid())
-	{
-		TArray<UObject*> OutObjects;
-		GetSequencer()->GetRuntimeObjects( GetSequencer()->GetFocusedMovieSceneSequenceInstance(), ObjectGuid, OutObjects);
-
-		AnimatablePropertyChanged( FOnKeyProperty::CreateRaw( this, &F3DPathTrackEditor::AddKeyInternal, OutObjects, ParentActor) );
-	}
-	else if (Section != nullptr)
+	if (Section != nullptr)
 	{
 		const FScopedTransaction Transaction(LOCTEXT("UndoSetPath", "Set Path"));
 
@@ -176,9 +169,15 @@ void F3DPathTrackEditor::ActorSocketPicked(const FName SocketName, USceneCompone
 			PathSection->SetConstraintId(SplineId);
 		}
 	}
+	else if (ObjectGuid.IsValid())
+	{
+		TArray<TWeakObjectPtr<UObject>> OutObjects;
+		GetSequencer()->GetRuntimeObjects( GetSequencer()->GetFocusedMovieSceneSequenceInstance(), ObjectGuid, OutObjects);
+		AnimatablePropertyChanged( FOnKeyProperty::CreateRaw( this, &F3DPathTrackEditor::AddKeyInternal, OutObjects, ParentActor) );
+	}
 }
 
-bool F3DPathTrackEditor::AddKeyInternal( float KeyTime, const TArray<UObject*> Objects, AActor* ParentActor)
+bool F3DPathTrackEditor::AddKeyInternal( float KeyTime, const TArray<TWeakObjectPtr<UObject>> Objects, AActor* ParentActor)
 {
 	bool bHandleCreated = false;
 	bool bTrackCreated = false;
@@ -202,7 +201,7 @@ bool F3DPathTrackEditor::AddKeyInternal( float KeyTime, const TArray<UObject*> O
 
 	for( int32 ObjectIndex = 0; ObjectIndex < Objects.Num(); ++ObjectIndex )
 	{
-		UObject* Object = Objects[ObjectIndex];
+		UObject* Object = Objects[ObjectIndex].Get();
 
 		FFindOrCreateHandleResult HandleResult = FindOrCreateHandleToObject( Object );
 		FGuid ObjectHandle = HandleResult.Handle;
diff --git a/Engine/Source/Editor/MovieSceneTools/Private/TrackEditors/PathTrackEditor.h b/Engine/Source/Editor/MovieSceneTools/Private/TrackEditors/PathTrackEditor.h
--- a/Engine/Source/Editor/MovieSceneTools/Private/TrackEditors/PathTrackEditor.h
+++ b/Engine/Source/Editor/MovieSceneTools/Private/TrackEditors/PathTrackEditor.h
@@ -48,7 +48,7 @@ public:
 private:
 
 	/** Delegate for AnimatablePropertyChanged in AddKey */
-	bool AddKeyInternal(float KeyTime, const TArray<UObject*> Objects, AActor* ParentActor);
+	bool AddKeyInternal(float KeyTime, const TArray<TWeakObjectPtr<UObject>> Objects, AActor* ParentActor);
 };
 
 
diff --git a/Engine/Source/Editor/MovieSceneTools/Private/TrackEditors/PropertyTrackEditors/ActorReferencePropertyTrackEditor.cpp b/Engine/Source/Editor/MovieSceneTools/Private/TrackEditors/PropertyTrackEditors/ActorReferencePropertyTrackEditor.cpp
new file mode 100644
--- /dev/null
+++ b/Engine/Source/Editor/MovieSceneTools/Private/TrackEditors/PropertyTrackEditors/ActorReferencePropertyTrackEditor.cpp
@@ -0,0 +1,33 @@
+// Copyright 1998-2016 Epic Games, Inc. All Rights Reserved.
+
+#include "MovieSceneToolsPrivatePCH.h"
+#include "ActorReferencePropertyTrackEditor.h"
+#include "ActorReferencePropertySection.h"
+#include "MatineeImportTools.h"
+
+
+TSharedRef<ISequencerTrackEditor> FActorReferencePropertyTrackEditor::CreateTrackEditor( TSharedRef<ISequencer> OwningSequencer )
+{
+	return MakeShareable(new FActorReferencePropertyTrackEditor(OwningSequencer));
+}
+
+
+TSharedRef<FPropertySection> FActorReferencePropertyTrackEditor::MakePropertySectionInterface( UMovieSceneSection& SectionObject, UMovieSceneTrack& Track )
+{
+	return MakeShareable(new FActorReferencePropertySection(SectionObject, Track.GetDisplayName()));
+}
+
+
+void FActorReferencePropertyTrackEditor::GenerateKeysFromPropertyChanged( const FPropertyChangedParams& PropertyChangedParams, TArray<FGuid>& NewGeneratedKeys, TArray<FGuid>& DefaultGeneratedKeys )
+{
+
+	AActor* NewReferencedActor = PropertyChangedParams.GetPropertyValue<AActor*>();
+	if ( NewReferencedActor != nullptr )
+	{
+		FGuid ActorGuid = GetSequencer()->GetHandleToObject( NewReferencedActor );
+		if ( ActorGuid.IsValid() )
+		{
+			NewGeneratedKeys.Add( ActorGuid );
+		}
+	}
+}
diff --git a/Engine/Source/Editor/MovieSceneTools/Private/TrackEditors/PropertyTrackEditors/ActorReferencePropertyTrackEditor.h b/Engine/Source/Editor/MovieSceneTools/Private/TrackEditors/PropertyTrackEditors/ActorReferencePropertyTrackEditor.h
new file mode 100644
--- /dev/null
+++ b/Engine/Source/Editor/MovieSceneTools/Private/TrackEditors/PropertyTrackEditors/ActorReferencePropertyTrackEditor.h
@@ -0,0 +1,39 @@
+// Copyright 1998-2016 Epic Games, Inc. All Rights Reserved.
+
+#pragma once
+
+#include "MovieSceneActorReferenceTrack.h"
+#include "MovieSceneActorReferenceSection.h"
+
+/**
+ * A property track editor for actor references.
+ */
+class FActorReferencePropertyTrackEditor
+	: public FPropertyTrackEditor<UMovieSceneActorReferenceTrack, UMovieSceneActorReferenceSection, FGuid>
+{
+public:
+
+	/**
+	 * Constructor
+	 *
+	 * @param InSequencer The sequencer instance to be used by this tool
+	 */
+	FActorReferencePropertyTrackEditor( TSharedRef<ISequencer> InSequencer )
+		: FPropertyTrackEditor( InSequencer, NAME_ObjectProperty )
+	{ }
+
+	/**
+	 * Creates an instance of this class.  Called by a sequencer 
+	 *
+	 * @param OwningSequencer The sequencer instance to be used by this tool
+	 * @return The new instance of this class
+	 */
+	static TSharedRef<ISequencerTrackEditor> CreateTrackEditor( TSharedRef<ISequencer> OwningSequencer );
+
+protected:
+
+	// FPropertyTrackEditor interface
+
+	virtual TSharedRef<FPropertySection> MakePropertySectionInterface( UMovieSceneSection& SectionObject, UMovieSceneTrack& Track ) override;
+	virtual void GenerateKeysFromPropertyChanged( const FPropertyChangedParams& PropertyChangedParams, TArray<FGuid>& NewGeneratedKeys, TArray<FGuid>& DefaultGeneratedKeys ) override;
+};
diff --git a/Engine/Source/Editor/MovieSceneTools/Private/TrackEditors/SkeletalAnimationTrackEditor.cpp b/Engine/Source/Editor/MovieSceneTools/Private/TrackEditors/SkeletalAnimationTrackEditor.cpp
--- a/Engine/Source/Editor/MovieSceneTools/Private/TrackEditors/SkeletalAnimationTrackEditor.cpp
+++ b/Engine/Source/Editor/MovieSceneTools/Private/TrackEditors/SkeletalAnimationTrackEditor.cpp
@@ -21,6 +21,7 @@
 #include "ContentBrowserModule.h"
 #include "MatineeImportTools.h"
 #include "Matinee/InterpTrackAnimControl.h"
+#include "SequencerUtilities.h"
 
 
 namespace SkeletalAnimationEditorConstants
@@ -147,16 +148,13 @@ void FSkeletalAnimationTrackEditor::AddKey(const FGuid& ObjectGuid)
 
 		if (AssetDataList.Num())
 		{
-			FMenuBuilder MenuBuilder(true, NULL);
-			BuildAnimationSubMenu(MenuBuilder, ObjectGuid, Skeleton);
-			
 			TSharedPtr< SWindow > Parent = FSlateApplication::Get().GetActiveTopLevelWindow(); 
 			if (Parent.IsValid())
 			{
 				FSlateApplication::Get().PushMenu(
 					Parent.ToSharedRef(),
 					FWidgetPath(),
-					MenuBuilder.MakeWidget(),
+					BuildAnimationSubMenu(ObjectGuid, Skeleton),
 					FSlateApplication::Get().GetCursorPos(),
 					FPopupTransitionEffect(FPopupTransitionEffect::TypeInPopup)
 					);
@@ -178,7 +176,7 @@ bool FSkeletalAnimationTrackEditor::HandleAssetAdded(UObject* Asset, const FGuid
 
 			if (Skeleton && Skeleton == AnimSequence->GetSkeleton())
 			{
-				TArray<UObject*> OutObjects;
+				TArray<TWeakObjectPtr<UObject>> OutObjects;
 
 				GetSequencer()->GetRuntimeObjects(GetSequencer()->GetFocusedMovieSceneSequenceInstance(), TargetObjectGuid, OutObjects);
 				AnimatablePropertyChanged(FOnKeyProperty::CreateRaw(this, &FSkeletalAnimationTrackEditor::AddKeyInternal, OutObjects, AnimSequence));
@@ -212,15 +210,23 @@ void FSkeletalAnimationTrackEditor::BuildObjectBindingTrackMenu(FMenuBuilder& Me
 			{
 				MenuBuilder.AddSubMenu(
 					LOCTEXT("AddAnimation", "Animation"), NSLOCTEXT("Sequencer", "AddAnimationTooltip", "Adds an animation track."),
-					FNewMenuDelegate::CreateRaw(this, &FSkeletalAnimationTrackEditor::BuildAnimationSubMenu, ObjectBinding, Skeleton)
+					FNewMenuDelegate::CreateRaw(this, &FSkeletalAnimationTrackEditor::AddAnimationSubMenu, ObjectBinding, Skeleton)
 				);
 			}
 		}
 	}
 }
 
+TSharedRef<SWidget> FSkeletalAnimationTrackEditor::BuildAnimationSubMenu(FGuid ObjectBinding, USkeleton* Skeleton)
+{
+	FMenuBuilder MenuBuilder(true, nullptr);
+
+	AddAnimationSubMenu(MenuBuilder, ObjectBinding, Skeleton);
 
-void FSkeletalAnimationTrackEditor::BuildAnimationSubMenu(FMenuBuilder& MenuBuilder, FGuid ObjectBinding, USkeleton* Skeleton)
+	return MenuBuilder.MakeWidget();
+}
+
+void FSkeletalAnimationTrackEditor::AddAnimationSubMenu(FMenuBuilder& MenuBuilder, FGuid ObjectBinding, USkeleton* Skeleton)
 {
 	FAssetPickerConfig AssetPickerConfig;
 	{
@@ -254,7 +260,7 @@ void FSkeletalAnimationTrackEditor::OnAnimationAssetSelected(const FAssetData& A
 	{
 		UAnimSequence* AnimSequence = CastChecked<UAnimSequence>(AssetData.GetAsset());
 
-		TArray<UObject*> OutObjects;
+		TArray<TWeakObjectPtr<UObject>> OutObjects;
 
 		GetSequencer()->GetRuntimeObjects( GetSequencer()->GetFocusedMovieSceneSequenceInstance(), ObjectBinding, OutObjects);
 		AnimatablePropertyChanged( FOnKeyProperty::CreateRaw( this, &FSkeletalAnimationTrackEditor::AddKeyInternal, OutObjects, AnimSequence) );
@@ -262,7 +268,7 @@ void FSkeletalAnimationTrackEditor::OnAnimationAssetSelected(const FAssetData& A
 }
 
 
-bool FSkeletalAnimationTrackEditor::AddKeyInternal( float KeyTime, const TArray<UObject*> Objects, class UAnimSequence* AnimSequence )
+bool FSkeletalAnimationTrackEditor::AddKeyInternal( float KeyTime, const TArray<TWeakObjectPtr<UObject>> Objects, class UAnimSequence* AnimSequence )
 {
 	bool bHandleCreated = false;
 	bool bTrackCreated = false;
@@ -270,7 +276,7 @@ bool FSkeletalAnimationTrackEditor::AddKeyInternal( float KeyTime, const TArray<
 
 	for( int32 ObjectIndex = 0; ObjectIndex < Objects.Num(); ++ObjectIndex )
 	{
-		UObject* Object = Objects[ObjectIndex];
+		UObject* Object = Objects[ObjectIndex].Get();
 
 		FFindOrCreateHandleResult HandleResult = FindOrCreateHandleToObject( Object );
 		FGuid ObjectHandle = HandleResult.Handle;
@@ -295,13 +301,15 @@ bool FSkeletalAnimationTrackEditor::AddKeyInternal( float KeyTime, const TArray<
 
 USkeleton* FSkeletalAnimationTrackEditor::AcquireSkeletonFromObjectGuid(const FGuid& Guid)
 {
-	TArray<UObject*> OutObjects;
+	TArray<TWeakObjectPtr<UObject>> OutObjects;
 	GetSequencer()->GetRuntimeObjects(GetSequencer()->GetFocusedMovieSceneSequenceInstance(), Guid, OutObjects);
 
 	USkeleton* Skeleton = NULL;
 	for (int32 i = 0; i < OutObjects.Num(); ++i)
 	{
-		if (AActor* Actor = Cast<AActor>(OutObjects[i]))
+		UObject* Object = OutObjects[i].Get();
+
+		if (AActor* Actor = Cast<AActor>(Object))
 		{
 			TInlineComponentArray<USkeletalMeshComponent*> SkeletalMeshComponents;
 			Actor->GetComponents(SkeletalMeshComponents);
@@ -317,7 +325,7 @@ USkeleton* FSkeletalAnimationTrackEditor::AcquireSkeletonFromObjectGuid(const FG
 				}
 			}
 		}
-		else if(USkeletalMeshComponent* SkeletalMeshComponent = Cast<USkeletalMeshComponent>(OutObjects[i]))
+		else if(USkeletalMeshComponent* SkeletalMeshComponent = Cast<USkeletalMeshComponent>(Object))
 		{
 			if (SkeletalMeshComponent->SkeletalMesh && SkeletalMeshComponent->SkeletalMesh->Skeleton)
 			{
@@ -352,5 +360,28 @@ void FSkeletalAnimationTrackEditor::BuildTrackContextMenu( FMenuBuilder& MenuBui
 			FCanExecuteAction::CreateLambda( [=]()->bool { return MatineeAnimControlTrack != nullptr && MatineeAnimControlTrack->AnimSeqs.Num() > 0 && SkeletalAnimationTrack != nullptr; } ) ) );
 }
 
+TSharedPtr<SWidget> FSkeletalAnimationTrackEditor::BuildOutlinerEditWidget(const FGuid& ObjectBinding, UMovieSceneTrack* Track, const FBuildEditWidgetParams& Params)
+{
+	USkeleton* Skeleton = AcquireSkeletonFromObjectGuid(ObjectBinding);
+
+	if (Skeleton)
+	{
+		// Create a container edit box
+		return SNew(SHorizontalBox)
+
+		// Add the animation combo box
+		+ SHorizontalBox::Slot()
+		.AutoWidth()
+		.VAlign(VAlign_Center)
+		[
+			FSequencerUtilities::MakeAddButton(LOCTEXT("AnimationText", "Animation"), FOnGetContent::CreateSP(this, &FSkeletalAnimationTrackEditor::BuildAnimationSubMenu, ObjectBinding, Skeleton), Params.NodeIsHovered)
+		];
+	}
+
+	else
+	{
+		return TSharedPtr<SWidget>();
+	}
+}
 
 #undef LOCTEXT_NAMESPACE
diff --git a/Engine/Source/Editor/MovieSceneTools/Private/TrackEditors/SkeletalAnimationTrackEditor.h b/Engine/Source/Editor/MovieSceneTools/Private/TrackEditors/SkeletalAnimationTrackEditor.h
--- a/Engine/Source/Editor/MovieSceneTools/Private/TrackEditors/SkeletalAnimationTrackEditor.h
+++ b/Engine/Source/Editor/MovieSceneTools/Private/TrackEditors/SkeletalAnimationTrackEditor.h
@@ -37,17 +37,19 @@ public:
 	virtual TSharedRef<ISequencerSection> MakeSectionInterface( UMovieSceneSection& SectionObject, UMovieSceneTrack& Track ) override;
 	virtual bool SupportsType( TSubclassOf<UMovieSceneTrack> Type ) const override;
 	virtual void BuildTrackContextMenu( FMenuBuilder& MenuBuilder, UMovieSceneTrack* Track ) override;
+	virtual TSharedPtr<SWidget> BuildOutlinerEditWidget(const FGuid& ObjectBinding, UMovieSceneTrack* Track, const FBuildEditWidgetParams& Params) override;
 
 private:
 
 	/** Animation sub menu */
-	void BuildAnimationSubMenu(FMenuBuilder& MenuBuilder, FGuid ObjectBinding, USkeleton* Skeleton);
+	TSharedRef<SWidget> BuildAnimationSubMenu(FGuid ObjectBinding, USkeleton* Skeleton);
+	void AddAnimationSubMenu(FMenuBuilder& MenuBuilder, FGuid ObjectBinding, USkeleton* Skeleton);
 
 	/** Animation asset selected */
 	void OnAnimationAssetSelected(const FAssetData& AssetData, FGuid ObjectBinding);
 
 	/** Delegate for AnimatablePropertyChanged in AddKey */
-	bool AddKeyInternal(float KeyTime, const TArray<UObject*> Objects, class UAnimSequence* AnimSequence);
+	bool AddKeyInternal(float KeyTime, const TArray<TWeakObjectPtr<UObject>> Objects, class UAnimSequence* AnimSequence);
 
 	/** Gets a skeleton from an object guid in the movie scene */
 	class USkeleton* AcquireSkeletonFromObjectGuid(const FGuid& Guid);
