Copying //UE4/Dev-Sequencer to //UE4/Main (Source: //UE4/Dev-Sequencer @ 2945541)

diff --git a/Engine/Source/Editor/DetailCustomizations/Private/CameraFocusSettingsCustomization.cpp b/Engine/Source/Editor/DetailCustomizations/Private/CameraFocusSettingsCustomization.cpp
--- a/Engine/Source/Editor/DetailCustomizations/Private/CameraFocusSettingsCustomization.cpp
+++ b/Engine/Source/Editor/DetailCustomizations/Private/CameraFocusSettingsCustomization.cpp
@@ -9,7 +9,6 @@
 
 static FName NAME_Category(TEXT("Category"));
 static FString ManualFocusSettingsString(TEXT("Manual Focus Settings"));
-static FString SpotFocusSettingsString(TEXT("Spot Focus Settings"));
 static FString TrackingFocusSettingsString(TEXT("Tracking Focus Settings"));
 static FString GeneralFocusSettingsString(TEXT("Focus Settings"));
 
@@ -60,10 +59,6 @@ void FCameraFocusSettingsCustomization::CustomizeChildren(TSharedRef<IPropertyHa
 		{
 			PropertyRow.Visibility(TAttribute<EVisibility>(this, &FCameraFocusSettingsCustomization::IsManualSettingGroupVisible));
 		}
-		else if (Category == SpotFocusSettingsString)
-		{
-			PropertyRow.Visibility(TAttribute<EVisibility>(this, &FCameraFocusSettingsCustomization::IsSpotSettingGroupVisible));
-		}
 		else if (Category == TrackingFocusSettingsString)
 		{
 			PropertyRow.Visibility(TAttribute<EVisibility>(this, &FCameraFocusSettingsCustomization::IsTrackingSettingGroupVisible));
@@ -72,29 +67,84 @@ void FCameraFocusSettingsCustomization::CustomizeChildren(TSharedRef<IPropertyHa
 		{
 			PropertyRow.Visibility(TAttribute<EVisibility>(this, &FCameraFocusSettingsCustomization::IsGeneralSettingGroupVisible));
 		}
+
+		// special customization to show scene depth picker widget
+		if (Iter.Value() == ManualFocusDistanceHandle)
+		{
+			TSharedPtr<SWidget> NameWidget;
+			TSharedPtr<SWidget> ValueWidget;
+			FDetailWidgetRow Row;
+			PropertyRow.GetDefaultWidgets(NameWidget, ValueWidget, Row);
+
+			PropertyRow.CustomWidget(/*bShowChildren*/ true)
+			.NameContent()
+			[
+				NameWidget.ToSharedRef()
+			]
+			.ValueContent()
+			[
+				SNew(SHorizontalBox)
+				+ SHorizontalBox::Slot()
+				.VAlign(VAlign_Center)
+				[
+					ValueWidget.ToSharedRef()
+				]
+				+ SHorizontalBox::Slot()
+				.Padding(2.0f, 0.0f)
+				.AutoWidth()
+				.VAlign(VAlign_Center)
+				[
+					PropertyCustomizationHelpers::MakeSceneDepthPicker(FOnSceneDepthLocationSelected::CreateSP(this, &FCameraFocusSettingsCustomization::OnSceneDepthLocationSelected))
+				]
+			];
+		}
 	}
+
 }
 
-EVisibility FCameraFocusSettingsCustomization::IsManualSettingGroupVisible() const
+void FCameraFocusSettingsCustomization::OnSceneDepthLocationSelected(FVector PickedSceneLoc)
 {
-	uint8 FocusMethodNumber;
-	FocusMethodHandle->GetValue(FocusMethodNumber);
-	ECameraFocusMethod const FocusMethod = static_cast<ECameraFocusMethod>(FocusMethodNumber);
-	if (FocusMethod == ECameraFocusMethod::Manual)
+	if (PickedSceneLoc != FVector::ZeroVector)
 	{
-		// if focus method is set to none, all non-none setting groups are collapsed
-		return EVisibility::Visible;
-	}
+		// find the camera component and set it relative to that
+		UCameraComponent* OuterCameraComponent = nullptr;
+		{
+			TArray<UObject*> OuterObjects;
+			ManualFocusDistanceHandle->GetOuterObjects(OuterObjects);
+			for (UObject* Obj : OuterObjects)
+			{
+				UCameraComponent* const CamComp = dynamic_cast<UCameraComponent*>(Obj);
+				if (CamComp)
+				{
+					OuterCameraComponent = CamComp;
+					break;
+				}
+			}
+		}
 
-	return EVisibility::Collapsed;
+		if (OuterCameraComponent)
+		{
+			FVector const CamToPickedLoc = PickedSceneLoc - OuterCameraComponent->GetComponentLocation();
+			FVector const CamForward = OuterCameraComponent->GetComponentRotation().Vector();
+
+			// if picked behind camera, don't set it
+			if ((CamToPickedLoc | CamForward) > 0.f)
+			{
+				float const FinalSceneDepth = CamToPickedLoc.ProjectOnToNormal(CamForward).Size();
+
+				const FScopedTransaction Transaction(LOCTEXT("PickedSceneDepth", "Pick Scene Depth"));
+				ensure(ManualFocusDistanceHandle->SetValue(FinalSceneDepth, EPropertyValueSetFlags::NotTransactable) == FPropertyAccess::Result::Success);
+			}
+		}
+	}
 }
 
-EVisibility FCameraFocusSettingsCustomization::IsSpotSettingGroupVisible() const
+EVisibility FCameraFocusSettingsCustomization::IsManualSettingGroupVisible() const
 {
 	uint8 FocusMethodNumber;
 	FocusMethodHandle->GetValue(FocusMethodNumber);
 	ECameraFocusMethod const FocusMethod = static_cast<ECameraFocusMethod>(FocusMethodNumber);
-	if (FocusMethod == ECameraFocusMethod::Spot)
+	if (FocusMethod == ECameraFocusMethod::Manual)
 	{
 		// if focus method is set to none, all non-none setting groups are collapsed
 		return EVisibility::Visible;
diff --git a/Engine/Source/Editor/DetailCustomizations/Private/CameraFocusSettingsCustomization.h b/Engine/Source/Editor/DetailCustomizations/Private/CameraFocusSettingsCustomization.h
--- a/Engine/Source/Editor/DetailCustomizations/Private/CameraFocusSettingsCustomization.h
+++ b/Engine/Source/Editor/DetailCustomizations/Private/CameraFocusSettingsCustomization.h
@@ -20,4 +20,6 @@ private:
 	EVisibility IsSpotSettingGroupVisible() const;
 	EVisibility IsTrackingSettingGroupVisible() const;
 	EVisibility IsGeneralSettingGroupVisible() const;
+
+	void OnSceneDepthLocationSelected(FVector PickedSceneLoc);
 };
diff --git a/Engine/Source/Editor/DetailCustomizations/Private/CameraLensSettingsCustomization.cpp b/Engine/Source/Editor/DetailCustomizations/Private/CameraLensSettingsCustomization.cpp
--- a/Engine/Source/Editor/DetailCustomizations/Private/CameraLensSettingsCustomization.cpp
+++ b/Engine/Source/Editor/DetailCustomizations/Private/CameraLensSettingsCustomization.cpp
@@ -77,11 +77,10 @@ void FCameraLensSettingsCustomization::CustomizeChildren(TSharedRef<IPropertyHan
 	MinFStopHandle = PropertyHandles.FindChecked(GET_MEMBER_NAME_CHECKED(FCameraLensSettings, MinFStop));
 	MaxFStopHandle = PropertyHandles.FindChecked(GET_MEMBER_NAME_CHECKED(FCameraLensSettings, MaxFStop));
 	MinFocusDistanceHandle = PropertyHandles.FindChecked(GET_MEMBER_NAME_CHECKED(FCameraLensSettings, MinimumFocusDistance));
-	MaxReproRatioHandle = PropertyHandles.FindChecked(GET_MEMBER_NAME_CHECKED(FCameraLensSettings, MaxReproductionRatio));
 
 	for (auto Iter(PropertyHandles.CreateConstIterator()); Iter; ++Iter)
 	{
-		if ((Iter.Value() == MinFocusDistanceHandle) || (Iter.Value() == MaxReproRatioHandle))
+		if (Iter.Value() == MinFocusDistanceHandle)
 		{
 			// skip showing these in the panel for now, as we don't really use them
 			continue;
@@ -123,7 +122,6 @@ void FCameraLensSettingsCustomization::OnPresetChanged(TSharedPtr<FString> NewSe
 				ensure(MinFStopHandle->SetValue(P.LensSettings.MinFStop, EPropertyValueSetFlags::InteractiveChange | EPropertyValueSetFlags::NotTransactable) == FPropertyAccess::Result::Success);
 				ensure(MaxFStopHandle->SetValue(P.LensSettings.MaxFStop, EPropertyValueSetFlags::InteractiveChange | EPropertyValueSetFlags::NotTransactable) == FPropertyAccess::Result::Success);
 				ensure(MinFocusDistanceHandle->SetValue(P.LensSettings.MinimumFocusDistance, EPropertyValueSetFlags::InteractiveChange | EPropertyValueSetFlags::NotTransactable) == FPropertyAccess::Result::Success);
-				ensure(MaxReproRatioHandle->SetValue(P.LensSettings.MaxReproductionRatio, EPropertyValueSetFlags::NotTransactable) == FPropertyAccess::Result::Success);
 
 				break;
 			}
@@ -165,9 +163,6 @@ TSharedPtr<FString> FCameraLensSettingsCustomization::GetPresetString() const
 	float MinFocusDistance;
 	MinFocusDistanceHandle->GetValue(MinFocusDistance);
 
-	float MaxReproRatio;
-	MaxReproRatioHandle->GetValue(MaxReproRatio);
-
 	// search presets for one that matches
 	TArray<FNamedLensPreset> const& Presets = UCineCameraComponent::GetLensPresets();
 	int32 const NumPresets = Presets.Num();
@@ -175,7 +170,7 @@ TSharedPtr<FString> FCameraLensSettingsCustomization::GetPresetString() const
 	{
 		FNamedLensPreset const& P = Presets[PresetIdx];
 		if ((P.LensSettings.MinFocalLength == MinFocalLength) && (P.LensSettings.MaxFocalLength == MaxFocalLength) && (P.LensSettings.MinFStop == MinFStop) && 
-			(P.LensSettings.MaxFStop == MaxFStop) && (P.LensSettings.MinimumFocusDistance == MinFocusDistance) && (P.LensSettings.MaxReproductionRatio == MaxReproRatio))
+			(P.LensSettings.MaxFStop == MaxFStop) && (P.LensSettings.MinimumFocusDistance == MinFocusDistance) )
 		{
 			// this is the one
 			if (PresetComboList.IsValidIndex(PresetIdx + 1))
diff --git a/Engine/Source/Editor/DetailCustomizations/Private/CaptureResolutionCustomization.cpp b/Engine/Source/Editor/DetailCustomizations/Private/CaptureResolutionCustomization.cpp
--- a/Engine/Source/Editor/DetailCustomizations/Private/CaptureResolutionCustomization.cpp
+++ b/Engine/Source/Editor/DetailCustomizations/Private/CaptureResolutionCustomization.cpp
@@ -25,7 +25,8 @@ void FCaptureResolutionCustomization::CustomizeHeader( TSharedRef<IPropertyHandl
 	Resolutions.Add(MakeShareable( new FPredefinedResolution{LOCTEXT("ResolutionC", "640 x 360 (16:9)"), 640, 360} ));
 	Resolutions.Add(MakeShareable( new FPredefinedResolution{LOCTEXT("ResolutionD", "1280 x 720 (16:9)"), 1280, 720} ));
 	Resolutions.Add(MakeShareable( new FPredefinedResolution{LOCTEXT("ResolutionE", "1920 x 1080 (16:9)"), 1920, 1080} ));
-	Resolutions.Add(MakeShareable( new FPredefinedResolution{LOCTEXT("ResolutionF", "Custom"), 1920, 1080} ));
+	Resolutions.Add(MakeShareable( new FPredefinedResolution{LOCTEXT("ResolutionF", "3840 x 2160 (16:9)"), 3840, 2160} ));
+	Resolutions.Add(MakeShareable( new FPredefinedResolution{LOCTEXT("ResolutionG", "Custom"), 1920, 1080} ));
 
 	uint32 CurrentResX = 0, CurrentResY = 0;
 	ResXHandle->GetValue(CurrentResX);
diff --git a/Engine/Source/Editor/EditorStyle/Private/SlateEditorStyle.cpp b/Engine/Source/Editor/EditorStyle/Private/SlateEditorStyle.cpp
--- a/Engine/Source/Editor/EditorStyle/Private/SlateEditorStyle.cpp
+++ b/Engine/Source/Editor/EditorStyle/Private/SlateEditorStyle.cpp
@@ -1007,6 +1007,12 @@ void FSlateEditorStyle::FStyle::SetupGeneralStyles()
 			.SetHovered(IMAGE_BRUSH( "/Sequencer/Transport_Bar/Record_24x_OFF", Icon24x24 ))
 			.SetPressed(IMAGE_BRUSH( "/Sequencer/Transport_Bar/Record_24x", Icon24x24 ))
 		);
+
+		Set( "Animation.Recording", FButtonStyle( Button )
+			.SetNormal(IMAGE_BRUSH( "/Sequencer/Transport_Bar/Recording_24x_OFF", Icon24x24 ))
+			.SetHovered(IMAGE_BRUSH( "/Sequencer/Transport_Bar/Recording_24x_OFF", Icon24x24 ))
+			.SetPressed(IMAGE_BRUSH( "/Sequencer/Transport_Bar/Recording_24x", Icon24x24 ))
+		);
 	}
 
 	// Message Log
@@ -1703,12 +1709,12 @@ void FSlateEditorStyle::FStyle::SetupGeneralStyles()
 
 		Set( "Sequencer.NotificationImage_AddedPlayMovieSceneEvent", new IMAGE_BRUSH( "Old/Checkbox_checked", Icon16x16 ) );
 
-		Set( "Sequencer.Save", new IMAGE_BRUSH( "Sequencer/Main_Icons/Icon_Sequencer_Save_24x", Icon40x40 ) );
-		Set( "Sequencer.Save.Small", new IMAGE_BRUSH( "Sequencer/Main_Icons/Icon_Sequencer_Save_24x", Icon20x20 ) );
-		Set( "Sequencer.SaveAs", new IMAGE_BRUSH( "Sequencer/Main_Icons/Icon_Sequencer_Save_As_24x", Icon40x40 ) );
-		Set( "Sequencer.SaveAs.Small", new IMAGE_BRUSH( "Sequencer/Main_Icons/Icon_Sequencer_Save_As_24x", Icon20x20 ) );
-		Set( "Sequencer.DiscardChanges", new IMAGE_BRUSH( "Sequencer/Main_Icons/Icon_Sequencer_Revert_24x", Icon40x40 ) );
-		Set( "Sequencer.DiscardChanges.Small", new IMAGE_BRUSH( "Sequencer/Main_Icons/Icon_Sequencer_Revert_24x", Icon20x20 ) );
+		Set( "Sequencer.Save", new IMAGE_BRUSH( "Sequencer/Main_Icons/Icon_Sequencer_Save_24x", Icon48x48 ) );
+		Set( "Sequencer.Save.Small", new IMAGE_BRUSH( "Sequencer/Main_Icons/Icon_Sequencer_Save_24x", Icon24x24 ) );
+		Set( "Sequencer.SaveAs", new IMAGE_BRUSH( "Sequencer/Main_Icons/Icon_Sequencer_Save_As_24x", Icon48x48 ) );
+		Set( "Sequencer.SaveAs.Small", new IMAGE_BRUSH( "Sequencer/Main_Icons/Icon_Sequencer_Save_As_24x", Icon24x24 ) );
+		Set( "Sequencer.DiscardChanges", new IMAGE_BRUSH( "Sequencer/Main_Icons/Icon_Sequencer_Revert_24x", Icon48x48 ) );
+		Set( "Sequencer.DiscardChanges.Small", new IMAGE_BRUSH( "Sequencer/Main_Icons/Icon_Sequencer_Revert_24x", Icon24x24 ) );
 		Set( "Sequencer.GenericGripLeft", new BOX_BRUSH( "Sequencer/GenericGripLeft", FMargin(5.f/16.f) ) );
 		Set( "Sequencer.GenericGripRight", new BOX_BRUSH( "Sequencer/GenericGripRight", FMargin(5.f/16.f) ) );
 		Set( "Sequencer.SectionArea.Background", new FSlateColorBrush( FColor::White ) );
@@ -1737,29 +1743,32 @@ void FSlateEditorStyle::FStyle::SetupGeneralStyles()
 		Set( "Sequencer.AnimationOutliner.TopLevelBorder_Expanded", new BOX_BRUSH( "Sequencer/TopLevelNodeBorder_Expanded", FMargin(4.0f/16.0f) ) );
 		Set( "Sequencer.AnimationOutliner.TopLevelBorder_Collapsed", new BOX_BRUSH( "Sequencer/TopLevelNodeBorder_Collapsed", FMargin(4.0f/16.0f) ) );
 		Set( "Sequencer.AnimationOutliner.DefaultBorder", new FSlateColorBrush( FLinearColor::White ) );
+		Set( "Sequencer.AnimationOutliner.TransparentBorder", new FSlateColorBrush( FLinearColor::Transparent ) );
 		Set( "Sequencer.AnimationOutliner.BoldFont", TTF_CORE_FONT( "Fonts/Roboto-Bold", 11 ) );
 		Set( "Sequencer.AnimationOutliner.RegularFont", TTF_CORE_FONT( "Fonts/Roboto-Regular", 9 ) );
 		Set( "Sequencer.ShotFilter", new IMAGE_BRUSH( "Sequencer/FilteredArea", FVector2D(74,74), FLinearColor::White, ESlateBrushTileType::Both ) );
 		Set( "Sequencer.KeyMark", new IMAGE_BRUSH("Sequencer/KeyMark", FVector2D(3,21), FLinearColor::White, ESlateBrushTileType::NoTile ) );
-		Set( "Sequencer.SetAutoKeyModeAll", new IMAGE_BRUSH( "Sequencer/Main_Icons/Icon_Sequencer_Auto_Key_All_24x", Icon40x40 ) );
-		Set( "Sequencer.SetAutoKeyModeAll.Small", new IMAGE_BRUSH( "Sequencer/Main_Icons/Icon_Sequencer_Auto_Key_All_24x", Icon20x20 ) );
-		Set( "Sequencer.SetAutoKeyModeAnimated", new IMAGE_BRUSH("Sequencer/Main_Icons/Icon_Sequencer_Auto_Key_All_Animated_24x", Icon40x40 ) );
-		Set( "Sequencer.SetAutoKeyModeAnimated.Small", new IMAGE_BRUSH("Sequencer/Main_Icons/Icon_Sequencer_Auto_Key_All_Animated_24x", Icon20x20 ) );
-		Set( "Sequencer.SetAutoKeyModeNone", new IMAGE_BRUSH("Sequencer/Main_Icons/Icon_Sequencer_Disable_Auto_Key_24x", Icon40x40));
-		Set( "Sequencer.SetAutoKeyModeNone.Small", new IMAGE_BRUSH("Sequencer/Main_Icons/Icon_Sequencer_Disable_Auto_Key_24x", Icon20x20 ) );		
-		Set( "Sequencer.KeyAllEnabled", new IMAGE_BRUSH( "Sequencer/Main_Icons/Icon_Sequencer_Key_All_24x", Icon40x40 ) );
-		Set( "Sequencer.KeyAllEnabled.Small", new IMAGE_BRUSH( "Sequencer/Main_Icons/Icon_Sequencer_Key_All_24x", Icon20x20 ) );
-		Set( "Sequencer.KeyAllDisabled", new IMAGE_BRUSH( "Sequencer/Main_Icons/Icon_Sequencer_Key_Part_24x", Icon40x40 ) );
-		Set( "Sequencer.KeyAllDisabled.Small", new IMAGE_BRUSH( "Sequencer/Main_Icons/Icon_Sequencer_Key_Part_24x", Icon20x20 ) );
-		Set( "Sequencer.ToggleIsSnapEnabled", new IMAGE_BRUSH( "Sequencer/Main_Icons/Icon_Sequencer_Snap_24x", Icon40x40 ) );
-		Set( "Sequencer.ToggleIsSnapEnabled.Small", new IMAGE_BRUSH( "Sequencer/Main_Icons/Icon_Sequencer_Snap_24x", Icon20x20 ) );
-		Set( "Sequencer.ToggleShowCurveEditor", new IMAGE_BRUSH("Icons/SequencerIcons/icon_Sequencer_CurveEditor_24x", Icon40x40) );
-		Set( "Sequencer.ToggleShowCurveEditor.Small", new IMAGE_BRUSH("Icons/SequencerIcons/icon_Sequencer_CurveEditor_24x", Icon20x20) );
-		Set( "Sequencer.ToggleAutoScroll", new IMAGE_BRUSH( "Icons/icon_Sequencer_ToggleAutoScroll_40x", Icon40x40 ) );
+		Set( "Sequencer.SetAutoKeyModeAll", new IMAGE_BRUSH( "Sequencer/Main_Icons/Icon_Sequencer_Auto_Key_All_24x", Icon48x48 ) );
+		Set( "Sequencer.SetAutoKeyModeAll.Small", new IMAGE_BRUSH( "Sequencer/Main_Icons/Icon_Sequencer_Auto_Key_All_24x", Icon24x24 ) );
+		Set( "Sequencer.SetAutoKeyModeAnimated", new IMAGE_BRUSH("Sequencer/Main_Icons/Icon_Sequencer_Auto_Key_All_Animated_24x", Icon48x48 ) );
+		Set( "Sequencer.SetAutoKeyModeAnimated.Small", new IMAGE_BRUSH("Sequencer/Main_Icons/Icon_Sequencer_Auto_Key_All_Animated_24x", Icon24x24 ) );
+		Set( "Sequencer.SetAutoKeyModeNone", new IMAGE_BRUSH("Sequencer/Main_Icons/Icon_Sequencer_Disable_Auto_Key_24x", Icon48x48));
+		Set( "Sequencer.SetAutoKeyModeNone.Small", new IMAGE_BRUSH("Sequencer/Main_Icons/Icon_Sequencer_Disable_Auto_Key_24x", Icon24x24 ) );		
+		Set( "Sequencer.KeyAllEnabled", new IMAGE_BRUSH( "Sequencer/Main_Icons/Icon_Sequencer_Key_All_24x", Icon48x48 ) );
+		Set( "Sequencer.KeyAllEnabled.Small", new IMAGE_BRUSH( "Sequencer/Main_Icons/Icon_Sequencer_Key_All_24x", Icon24x24 ) );
+		Set( "Sequencer.KeyAllDisabled", new IMAGE_BRUSH( "Sequencer/Main_Icons/Icon_Sequencer_Key_Part_24x", Icon48x48 ) );
+		Set( "Sequencer.KeyAllDisabled.Small", new IMAGE_BRUSH( "Sequencer/Main_Icons/Icon_Sequencer_Key_Part_24x", Icon24x24 ) );
+		Set( "Sequencer.ToggleIsSnapEnabled", new IMAGE_BRUSH( "Sequencer/Main_Icons/Icon_Sequencer_Snap_24x", Icon48x48 ) );
+		Set( "Sequencer.ToggleIsSnapEnabled.Small", new IMAGE_BRUSH( "Sequencer/Main_Icons/Icon_Sequencer_Snap_24x", Icon24x24 ) );
+		Set( "Sequencer.ToggleShowCurveEditor", new IMAGE_BRUSH("Icons/SequencerIcons/icon_Sequencer_CurveEditor_24x", Icon48x48) );
+		Set( "Sequencer.ToggleShowCurveEditor.Small", new IMAGE_BRUSH("Icons/SequencerIcons/icon_Sequencer_CurveEditor_24x", Icon24x24) );
+		Set( "Sequencer.ToggleAutoScroll", new IMAGE_BRUSH( "Icons/icon_Sequencer_ToggleAutoScroll_40x", Icon48x48 ) );
 		Set( "Sequencer.ToggleAutoScroll.Small", new IMAGE_BRUSH( "Icons/icon_Sequencer_ToggleAutoScroll_16x", Icon16x16 ) );
 		Set( "Sequencer.MoveTool.Small", new IMAGE_BRUSH( "Icons/SequencerIcons/icon_Sequencer_Move_24x", Icon16x16 ) );
 		Set( "Sequencer.MarqueeTool.Small", new IMAGE_BRUSH( "Icons/SequencerIcons/icon_Sequencer_Marquee_24x", Icon16x16 ) );
-		Set( "Sequencer.RenderMovie.Small", new IMAGE_BRUSH( "Sequencer/Main_Icons/Icon_Sequencer_Create_Movie_24x", Icon16x16 ) );
+		Set( "Sequencer.RenderMovie.Small", new IMAGE_BRUSH( "Sequencer/Main_Icons/Icon_Sequencer_Create_Movie_24x", Icon24x24 ) );
+		Set( "Sequencer.CreateCamera.Small", new IMAGE_BRUSH( "Sequencer/Main_Icons/Icon_Sequencer_Create_Camera_24x", Icon24x24 ) );
+		Set( "Sequencer.FindInContentBrowser.Small", new IMAGE_BRUSH( "Sequencer/Main_Icons/Icon_Sequencer_Find_In_Content_Browser_24x", Icon24x24 ) );
 		Set( "Sequencer.LockCamera", new IMAGE_BRUSH( "Sequencer/Main_Icons/Icon_Sequencer_Look_Thru_24x", Icon16x16 ) );
 		Set( "Sequencer.UnlockCamera", new IMAGE_BRUSH( "Sequencer/Main_Icons/Icon_Sequencer_Look_Thru_24x", Icon16x16, FLinearColor(1.f, 1.f, 1.f, 0.5f) ) );
 		Set( "Sequencer.Thumbnail.SectionHandle", new IMAGE_BRUSH( "Old/White", Icon16x16, FLinearColor::Black ) );
@@ -1767,12 +1776,12 @@ void FSlateEditorStyle::FStyle::SetupGeneralStyles()
 		Set( "Sequencer.TrackHoverHighlight_Bottom", new IMAGE_BRUSH( TEXT("Sequencer/TrackHoverHighlight_Bottom"), FVector2D(4, 4) ) );
 		Set( "Sequencer.SpawnableIconOverlay", new IMAGE_BRUSH( TEXT("Sequencer/SpawnableIconOverlay"), FVector2D(13, 13) ) );
 
-		Set( "Sequencer.General", new IMAGE_BRUSH( "Sequencer/Main_Icons/Icon_Sequencer_General_Options_24x", Icon24x24 ) );
-		Set( "Sequencer.General.Small", new IMAGE_BRUSH( "Sequencer/Main_Icons/Icon_Sequencer_General_Options_24x", Icon20x20 ) );
-		Set( "Sequencer.Time", new IMAGE_BRUSH( "Sequencer/Main_Icons/Icon_Sequencer_Time_24x", Icon24x24 ) );
-		Set( "Sequencer.Time.Small", new IMAGE_BRUSH( "Sequencer/Main_Icons/Icon_Sequencer_Time_24x", Icon20x20 ) );
-		Set( "Sequencer.Value", new IMAGE_BRUSH( "Sequencer/Main_Icons/Icon_Sequencer_Value_24x", Icon24x24 ) );
-		Set( "Sequencer.Value.Small", new IMAGE_BRUSH( "Sequencer/Main_Icons/Icon_Sequencer_Value_24x", Icon20x20 ) );
+		Set( "Sequencer.General", new IMAGE_BRUSH( "Sequencer/Main_Icons/Icon_Sequencer_General_Options_24x", Icon48x48 ) );
+		Set( "Sequencer.General.Small", new IMAGE_BRUSH( "Sequencer/Main_Icons/Icon_Sequencer_General_Options_24x", Icon24x24 ) );
+		Set( "Sequencer.Time", new IMAGE_BRUSH( "Sequencer/Main_Icons/Icon_Sequencer_Time_24x", Icon48x48 ) );
+		Set( "Sequencer.Time.Small", new IMAGE_BRUSH( "Sequencer/Main_Icons/Icon_Sequencer_Time_24x", Icon24x24 ) );
+		Set( "Sequencer.Value", new IMAGE_BRUSH( "Sequencer/Main_Icons/Icon_Sequencer_Value_24x", Icon48x48 ) );
+		Set( "Sequencer.Value.Small", new IMAGE_BRUSH( "Sequencer/Main_Icons/Icon_Sequencer_Value_24x", Icon24x24 ) );
 
 		Set( "Sequencer.TrackArea.LaneColor", FLinearColor(0.3f, 0.3f, 0.3f, 0.3f) );
 
@@ -6372,8 +6381,10 @@ void FSlateEditorStyle::FStyle::SetupToolkitStyles()
 		Set("RichCurveEditor.ZoomToFit", new IMAGE_BRUSH("Icons/icon_CurveEditor_ZoomToFit_40x", Icon40x40));
 		Set("RichCurveEditor.ZoomToFit.Small", new IMAGE_BRUSH("Icons/icon_CurveEditor_ZoomToFit_40x", Icon20x20));
 
-		Set("RichCurveEditor.ToggleSnapping", new IMAGE_BRUSH("Icons/icon_CurveEditor_ToggleSnap_40x", Icon40x40));
-		Set("RichCurveEditor.ToggleSnapping.Small", new IMAGE_BRUSH("Icons/icon_CurveEditor_ToggleSnap_40x", Icon20x20));
+		Set("RichCurveEditor.ToggleInputSnapping", new IMAGE_BRUSH("Icons/icon_CurveEditor_ToggleInputSnap_40x", Icon40x40));
+		Set("RichCurveEditor.ToggleInputSnapping.Small", new IMAGE_BRUSH("Icons/icon_CurveEditor_ToggleInputSnap_40x", Icon20x20));
+		Set("RichCurveEditor.ToggleOutputSnapping", new IMAGE_BRUSH("Icons/icon_CurveEditor_ToggleOutputSnap_40x", Icon40x40));
+		Set("RichCurveEditor.ToggleOutputSnapping.Small", new IMAGE_BRUSH("Icons/icon_CurveEditor_ToggleOutputSnap_40x", Icon20x20));
 
 		Set("RichCurveEditor.InterpolationCubicAuto", new IMAGE_BRUSH("Icons/icon_CurveEditor_Auto_40x", Icon40x40));
 		Set("RichCurveEditor.InterpolationCubicAuto.Small", new IMAGE_BRUSH("Icons/icon_CurveEditor_Auto_40x", Icon20x20));
diff --git a/Engine/Source/Editor/FontEditor/Private/FontEditor.cpp b/Engine/Source/Editor/FontEditor/Private/FontEditor.cpp
--- a/Engine/Source/Editor/FontEditor/Private/FontEditor.cpp
+++ b/Engine/Source/Editor/FontEditor/Private/FontEditor.cpp
@@ -1054,7 +1054,11 @@ bool FFontEditor::RecreateFontObject(const EFontCacheType NewCacheType)
 
 	if(FontFactoryPtr && FontFactoryPtr->ConfigureProperties())
 	{
-		bSuccess = UFactory::StaticImportObject(Font->GetClass(), Font->GetOuter(), *Font->GetName(), RF_Public|RF_Standalone, TEXT(""), nullptr, FontFactoryPtr) != nullptr;
+		bool OutCanceled = false;
+		if (FontFactoryPtr->ImportObject(Font->GetClass(), Font->GetOuter(), *Font->GetName(), RF_Public | RF_Standalone, TEXT(""), nullptr, OutCanceled) != nullptr)
+		{
+			bSuccess = true;
+		}
 	}
 
 	if(bSuccess)
diff --git a/Engine/Source/Editor/LevelEditor/Private/LevelEditor.cpp b/Engine/Source/Editor/LevelEditor/Private/LevelEditor.cpp
--- a/Engine/Source/Editor/LevelEditor/Private/LevelEditor.cpp
+++ b/Engine/Source/Editor/LevelEditor/Private/LevelEditor.cpp
@@ -235,7 +235,8 @@ void FLevelEditorModule::StartupModule()
 	// Bind level editor commands shared across an instance
 	BindGlobalLevelEditorCommands();
 
-	RegisterViewportType("Default", FViewportTypeDefinition::FromType<FLevelViewportLayoutEntity>());
+	FViewportTypeDefinition ViewportType = FViewportTypeDefinition::FromType<FLevelViewportLayoutEntity>(FLevelViewportCommands::Get().SetDefaultViewportType);
+	RegisterViewportType("Default", ViewportType);
 
 	const IWorkspaceMenuStructure& MenuStructure = WorkspaceMenu::GetMenuStructure();
 
@@ -553,7 +554,7 @@ TSharedRef<IViewportLayoutEntity> FLevelEditorModule::FactoryViewport(FName InTy
 		return Definition->FactoryFunction(ConstructionArgs);
 	}
 
-	return FViewportTypeDefinition::FromType<FLevelViewportLayoutEntity>().FactoryFunction(ConstructionArgs);
+	return MakeShareable(new FLevelViewportLayoutEntity(ConstructionArgs));
 }
 
 void FLevelEditorModule::BindGlobalLevelEditorCommands()
diff --git a/Engine/Source/Editor/LevelEditor/Private/LevelViewportActions.cpp b/Engine/Source/Editor/LevelEditor/Private/LevelViewportActions.cpp
--- a/Engine/Source/Editor/LevelEditor/Private/LevelViewportActions.cpp
+++ b/Engine/Source/Editor/LevelEditor/Private/LevelViewportActions.cpp
@@ -50,6 +50,8 @@ void FLevelViewportCommands::RegisterCommands()
 	UI_COMMAND( ViewportConfig_FourPanesBottom, "Layout Four Panes (one bottom, three top)", "Changes the viewport arrangement to four panes, one on the bottom, three on the top", EUserInterfaceActionType::ToggleButton, FInputChord() );
 	UI_COMMAND( ViewportConfig_FourPanes2x2, "Layout Four Panes (2x2)", "Changes the viewport arrangement to four panes, in a 2x2 grid", EUserInterfaceActionType::ToggleButton, FInputChord() );
 
+	UI_COMMAND( SetDefaultViewportType, "Default Viewport", "Reconfigures this viewport to the default arrangement", EUserInterfaceActionType::RadioButton, FInputChord(EModifierKey::Shift, EKeys::D) );
+
 	UI_COMMAND( ToggleViewportToolbar, "Show Toolbar", "Defines whether a toolbar should be displayed on this viewport", EUserInterfaceActionType::ToggleButton, FInputChord(EModifierKey::Control | EModifierKey::Shift, EKeys::T) );
 
 	UI_COMMAND( ApplyMaterialToActor, "Apply Material", "Attempts to apply a dropped material to this object", EUserInterfaceActionType::Button, FInputChord() );
diff --git a/Engine/Source/Editor/LevelEditor/Private/LevelViewportTabContent.cpp b/Engine/Source/Editor/LevelEditor/Private/LevelViewportTabContent.cpp
--- a/Engine/Source/Editor/LevelEditor/Private/LevelViewportTabContent.cpp
+++ b/Engine/Source/Editor/LevelEditor/Private/LevelViewportTabContent.cpp
@@ -93,6 +93,31 @@ void FLevelViewportTabContent::SetViewportConfiguration(const FName& Configurati
 	UpdateViewportTabWidget();
 }
 
+void FLevelViewportTabContent::RefreshViewportConfiguration()
+{
+	check(ActiveLevelViewportLayout.IsValid());
+
+	FName ConfigurationName = ActiveLevelViewportLayout->GetLayoutTypeName();
+	for (auto& Pair : ActiveLevelViewportLayout->GetViewports())
+	{
+		if (Pair.Value->AsWidget()->HasFocusedDescendants())
+		{
+			PreviouslyFocusedViewport = Pair.Key;
+			break;
+		}
+	}
+
+	ActiveLevelViewportLayout->SaveLayoutString(LayoutString);
+	ActiveLevelViewportLayout.Reset();
+
+	bool bSwitchingLayouts = false;
+	ActiveLevelViewportLayout = ConstructViewportLayoutByTypeName(ConfigurationName, bSwitchingLayouts);
+	check(ActiveLevelViewportLayout.IsValid());
+
+	UpdateViewportTabWidget();
+}
+
+
 bool FLevelViewportTabContent::IsViewportConfigurationSet(const FName& ConfigurationName) const
 {
 	if (ActiveLevelViewportLayout.IsValid())
@@ -114,8 +139,17 @@ void FLevelViewportTabContent::UpdateViewportTabWidget()
 	if (ParentTabPinned.IsValid() && ActiveLevelViewportLayout.IsValid())
 	{
 		TSharedRef<SWidget> LayoutWidget = ActiveLevelViewportLayout->BuildViewportLayout(ParentTabPinned, SharedThis(this), LayoutString, ParentLevelEditor);
-
 		ParentTabPinned->SetContent(LayoutWidget);
+
+		if (PreviouslyFocusedViewport.IsSet())
+		{
+			TSharedPtr<IViewportLayoutEntity> ViewportToFocus = ActiveLevelViewportLayout->GetViewports().FindRef(PreviouslyFocusedViewport.GetValue());
+			if (ViewportToFocus.IsValid())
+			{
+				ViewportToFocus->SetKeyboardFocus();
+			}
+			PreviouslyFocusedViewport = TOptional<FName>();
+		}
 	}
 }
 
diff --git a/Engine/Source/Editor/LevelEditor/Private/SActorDetails.cpp b/Engine/Source/Editor/LevelEditor/Private/SActorDetails.cpp
--- a/Engine/Source/Editor/LevelEditor/Private/SActorDetails.cpp
+++ b/Engine/Source/Editor/LevelEditor/Private/SActorDetails.cpp
@@ -584,6 +584,12 @@ bool SActorDetails::IsPropertyReadOnly(const FPropertyAndParent& PropertyAndPare
 
 bool SActorDetails::IsPropertyEditingEnabled() const
 {
+	FLevelEditorModule& LevelEditor = FModuleManager::GetModuleChecked<FLevelEditorModule>("LevelEditor");
+	if (!LevelEditor.AreObjectsEditable(DetailsView->GetSelectedObjects()))
+	{
+		return false;
+	}
+
 	bool bIsEditable = true;
 	const TArray<FSCSEditorTreeNodePtrType> SelectedNodes = SCSEditor->GetSelectedNodes();
 	for (const auto& Node : SelectedNodes)
diff --git a/Engine/Source/Editor/LevelEditor/Private/SLevelViewport.cpp b/Engine/Source/Editor/LevelEditor/Private/SLevelViewport.cpp
--- a/Engine/Source/Editor/LevelEditor/Private/SLevelViewport.cpp
+++ b/Engine/Source/Editor/LevelEditor/Private/SLevelViewport.cpp
@@ -1305,9 +1305,9 @@ void SLevelViewport::BindViewCommands( FUICommandList& CommandList )
 		FIsActionChecked::CreateSP( this, &SLevelViewport::IsViewportConfigurationSet, LevelViewportConfigurationNames::FourPanes2x2 ));
 
 	auto ProcessViewportTypeActions = [&](FName InViewportTypeName, const FViewportTypeDefinition& InDefinition){
-		if (InDefinition.ToggleCommand.IsValid())
+		if (InDefinition.ActivationCommand.IsValid())
 		{
-			CommandList.MapAction(InDefinition.ToggleCommand, FUIAction(
+			CommandList.MapAction(InDefinition.ActivationCommand, FUIAction(
 				FExecuteAction::CreateSP(this, &SLevelViewport::ToggleViewportTypeActivationWithinLayout, InViewportTypeName),
 				FCanExecuteAction(),
 				FIsActionChecked::CreateSP(this, &SLevelViewport::IsViewportTypeWithinLayoutEqual, InViewportTypeName)
@@ -3173,6 +3173,16 @@ bool SLevelViewport::IsActorPreviewPinned( TWeakObjectPtr<AActor> PreviewActor )
 
 void SLevelViewport::UpdateActorPreviewViewports()
 {
+	// Remove any previews that are locked to the same actor as the level viewport client's actor lock
+	for( int32 PreviewIndex = 0; PreviewIndex < ActorPreviews.Num(); ++PreviewIndex )
+	{
+		AActor* ExistingActor = ActorPreviews[PreviewIndex].Actor.Get();
+		if (ExistingActor && LevelViewportClient->IsActorLocked(ExistingActor))
+		{
+			RemoveActorPreview( PreviewIndex-- );
+		}
+	}
+
 	// Look for actors that we no longer want to preview
 	for( auto ActorPreviewIt = ActorPreviews.CreateConstIterator(); ActorPreviewIt; ++ActorPreviewIt )
 	{
@@ -3334,6 +3344,22 @@ void SLevelViewport::OnSetViewportConfiguration(FName ConfigurationName)
 	}
 }
 
+void SLevelViewport::RefreshViewportConfiguration()
+{
+	TSharedPtr<FLevelViewportLayout> LayoutPinned = ParentLayout.Pin();
+	if (LayoutPinned.IsValid())
+	{
+		TSharedPtr<FLevelViewportTabContent> ViewportTabPinned = LayoutPinned->GetParentTabContent().Pin();
+		if (ViewportTabPinned.IsValid())
+		{
+			// Viewport clients are going away.  Any current one is invalid.
+			GCurrentLevelEditingViewportClient = nullptr;
+			ViewportTabPinned->RefreshViewportConfiguration();
+			FSlateApplication::Get().DismissAllMenus();
+		}
+	}
+}
+
 bool SLevelViewport::IsViewportConfigurationSet(FName ConfigurationName) const
 {
 	TSharedPtr<FLevelViewportLayout> LayoutPinned = ParentLayout.Pin();
@@ -3371,17 +3397,13 @@ void SLevelViewport::SetViewportTypeWithinLayout(FName InLayoutType)
 		GConfig->SetString( *IniSection, *( ConfigKey + TEXT(".TypeWithinLayout") ), *InLayoutType.ToString(), GEditorPerProjectIni );
 
 		// Force a refresh of the tab content
-		OnSetViewportConfiguration(LayoutPinned->GetLayoutTypeName());
+		RefreshViewportConfiguration();
 	}
 }
 
 void SLevelViewport::ToggleViewportTypeActivationWithinLayout(FName InLayoutType)
 {
-	if (GetViewportTypeWithinLayout() == InLayoutType)
-	{
-		SetViewportTypeWithinLayout("Default");
-	}
-	else
+	if (GetViewportTypeWithinLayout() != InLayoutType)
 	{
 		SetViewportTypeWithinLayout(InLayoutType);
 	}
diff --git a/Engine/Source/Editor/LevelEditor/Private/SLevelViewportToolBar.cpp b/Engine/Source/Editor/LevelEditor/Private/SLevelViewportToolBar.cpp
--- a/Engine/Source/Editor/LevelEditor/Private/SLevelViewportToolBar.cpp
+++ b/Engine/Source/Editor/LevelEditor/Private/SLevelViewportToolBar.cpp
@@ -699,9 +699,9 @@ void SLevelViewportToolBar::GenerateViewportTypeMenu( FMenuBuilder& Builder ) co
 {
 	FLevelEditorModule& LevelEditorModule = FModuleManager::GetModuleChecked<FLevelEditorModule>(TEXT("LevelEditor"));
 	LevelEditorModule.IterateViewportTypes([&](FName, const FViewportTypeDefinition& InDefinition) {
-		if (InDefinition.ToggleCommand.IsValid())
+		if (InDefinition.ActivationCommand.IsValid())
 		{
-			Builder.AddMenuEntry(InDefinition.ToggleCommand);
+			Builder.AddMenuEntry(InDefinition.ActivationCommand);
 		}
 	});
 }
diff --git a/Engine/Source/Editor/LevelEditor/Public/LevelEditor.h b/Engine/Source/Editor/LevelEditor/Public/LevelEditor.h
--- a/Engine/Source/Editor/LevelEditor/Public/LevelEditor.h
+++ b/Engine/Source/Editor/LevelEditor/Public/LevelEditor.h
@@ -12,6 +12,8 @@
 
 extern const FName LevelEditorApp;
 
+DECLARE_DELEGATE_RetVal_OneParam(bool, FAreObjectsEditable, const TArray<TWeakObjectPtr<UObject>>&);
+
 /**
  * Level editor module
  */
@@ -234,6 +236,34 @@ public:
 	DECLARE_DELEGATE_RetVal_OneParam(UAnimSequence*, FCaptureSingleFrameAnimSequence, USkeletalMeshComponent* /*Component*/);
 	virtual FCaptureSingleFrameAnimSequence& OnCaptureSingleFrameAnimSequence() { return CaptureSingleFrameAnimSequenceDelegate; }
 
+public:
+
+	/** Add a delegate that will get called to check whether the specified objects should be editable on the details panel or not */
+	void AddEditableObjectPredicate(const FAreObjectsEditable& InPredicate)
+	{
+		check(InPredicate.IsBound());
+		AreObjectsEditableDelegates.Add(InPredicate);
+	}
+
+	/** Remove a delegate that was added via AddEditableObjectPredicate */
+	void RemoveEditableObjectPredicate(FDelegateHandle InPredicateHandle)
+	{
+		AreObjectsEditableDelegates.RemoveAll([=](const FAreObjectsEditable& P){ return P.GetHandle() == InPredicateHandle; });
+	}
+
+	/** Check whether the specified objects are editable */
+	bool AreObjectsEditable(const TArray<TWeakObjectPtr<UObject>>& InObjects) const
+	{
+		for (const FAreObjectsEditable& Predicate : AreObjectsEditableDelegates)
+		{
+			if (!Predicate.Execute(InObjects))
+			{
+				return false;
+			}
+		}
+		return true;
+	}
+
 public:
 	
 	/** Register a viewport type for the level editor */
@@ -335,6 +365,9 @@ private:
 
 	/** Map of named viewport types to factory functions */
 	TMap<FName, FViewportTypeDefinition> CustomViewports;
+
+	/** Array of delegates that are used to check if the specified objects should be editable on the details panel */
+	TArray<FAreObjectsEditable> AreObjectsEditableDelegates;
 };
 
 #endif // __LevelEditor_h__
diff --git a/Engine/Source/Editor/LevelEditor/Public/LevelViewportActions.h b/Engine/Source/Editor/LevelEditor/Public/LevelViewportActions.h
--- a/Engine/Source/Editor/LevelEditor/Public/LevelViewportActions.h
+++ b/Engine/Source/Editor/LevelEditor/Public/LevelViewportActions.h
@@ -162,6 +162,8 @@ public:
 	TSharedPtr< FUICommandInfo > ViewportConfig_FourPanesBottom;
 	TSharedPtr< FUICommandInfo > ViewportConfig_FourPanes2x2;
 
+	TSharedPtr< FUICommandInfo > SetDefaultViewportType;
+
 	TSharedPtr< FUICommandInfo > ToggleViewportToolbar;
 
 	TSharedPtr< FUICommandInfo > EnablePreviewMesh;
diff --git a/Engine/Source/Editor/LevelEditor/Public/LevelViewportTabContent.h b/Engine/Source/Editor/LevelEditor/Public/LevelViewportTabContent.h
--- a/Engine/Source/Editor/LevelEditor/Public/LevelViewportTabContent.h
+++ b/Engine/Source/Editor/LevelEditor/Public/LevelViewportTabContent.h
@@ -34,6 +34,11 @@ public:
 	 */
 	void SetViewportConfiguration(const FName& ConfigurationName);
 
+	/**
+	 * Refresh the current layout
+	 */
+	void RefreshViewportConfiguration();
+
 	/**
 	 * Returns whether the named layout is currently selected
 	 *
@@ -74,6 +79,8 @@ private:
 	TWeakPtr<ILevelEditor> ParentLevelEditor;
 	FString LayoutString;
 
+	TOptional<FName> PreviouslyFocusedViewport;
+
 	/** Current layout */
 	TSharedPtr< class FLevelViewportLayout > ActiveLevelViewportLayout;
 };
diff --git a/Engine/Source/Editor/LevelEditor/Public/SLevelViewport.h b/Engine/Source/Editor/LevelEditor/Public/SLevelViewport.h
--- a/Engine/Source/Editor/LevelEditor/Public/SLevelViewport.h
+++ b/Engine/Source/Editor/LevelEditor/Public/SLevelViewport.h
@@ -290,6 +290,11 @@ public:
 	 */
 	void OnSetViewportConfiguration(FName ConfigurationName);
 
+	/**
+	 * Refresh the current viewport configuration
+	 */
+	void RefreshViewportConfiguration();
+
 	/**
 	 * Returns whether the named layout is currently selected on the parent tab that this viewport belongs to.
 	 *
diff --git a/Engine/Source/Editor/LevelEditor/Public/ViewportTypeDefinition.h b/Engine/Source/Editor/LevelEditor/Public/ViewportTypeDefinition.h
--- a/Engine/Source/Editor/LevelEditor/Public/ViewportTypeDefinition.h
+++ b/Engine/Source/Editor/LevelEditor/Public/ViewportTypeDefinition.h
@@ -11,19 +11,20 @@ struct FViewportTypeDefinition
 	typedef TFunction<TSharedRef<IViewportLayoutEntity>(const FViewportConstructionArgs&)> FFactoryFunctionType;
 
 	template<typename T>
-	static FViewportTypeDefinition FromType()
+	static FViewportTypeDefinition FromType(const TSharedPtr<FUICommandInfo>& ActivationCommand)
 	{
 		return FViewportTypeDefinition([](const FViewportConstructionArgs& Args) -> TSharedRef<IViewportLayoutEntity> {
 			return MakeShareable(new T(Args));
-		});
+		}, ActivationCommand);
 	}
 
-	FViewportTypeDefinition(const FFactoryFunctionType& InFactoryFunction)
-		: FactoryFunction(MoveTemp(InFactoryFunction))
+	FViewportTypeDefinition(const FFactoryFunctionType& InFactoryFunction, const TSharedPtr<FUICommandInfo>& InActivationCommand)
+		: ActivationCommand(InActivationCommand)
+		, FactoryFunction(MoveTemp(InFactoryFunction))
 	{}
 
 	/** A UI command for toggling activation this viewport */
-	TSharedPtr<FUICommandInfo> ToggleCommand;
+	TSharedPtr<FUICommandInfo> ActivationCommand;
 
 	/** Function used to create a new instance of the viewport */
 	FFactoryFunctionType FactoryFunction;
diff --git a/Engine/Source/Editor/MovieSceneTools/Private/MatineeImportTools.cpp b/Engine/Source/Editor/MovieSceneTools/Private/MatineeImportTools.cpp
--- a/Engine/Source/Editor/MovieSceneTools/Private/MatineeImportTools.cpp
+++ b/Engine/Source/Editor/MovieSceneTools/Private/MatineeImportTools.cpp
@@ -3,19 +3,44 @@
 #include "MovieSceneToolsPrivatePCH.h"
 #include "MatineeImportTools.h"
 
+#include "Matinee/MatineeActor.h"
+#include "Matinee/InterpData.h"
+#include "Matinee/InterpGroupInst.h"
+#include "Matinee/InterpTrackColorProp.h"
+#include "Matinee/InterpTrackBoolProp.h"
 #include "Matinee/InterpTrackFloatBase.h"
 #include "Matinee/InterpTrackMove.h"
 #include "Matinee/InterpTrackAnimControl.h"
-
+#include "Matinee/InterpTrackSound.h"
+#include "Matinee/InterpTrackFade.h"
+#include "Matinee/InterpTrackDirector.h"
+#include "Matinee/InterpTrackEvent.h"
+#include "Matinee/InterpTrackVisibility.h"
+
+#include "MovieSceneColorTrack.h"
+#include "MovieSceneBoolTrack.h"
 #include "MovieSceneFloatTrack.h"
 #include "MovieScene3DTransformTrack.h"
 #include "MovieSceneParticleTrack.h"
 #include "MovieSceneSkeletalAnimationTrack.h"
-
+#include "MovieSceneAudioTrack.h"
+#include "MovieSceneFadeTrack.h"
+#include "MovieSceneCameraCutTrack.h"
+#include "MovieSceneEventTrack.h"
+#include "MovieSceneVisibilityTrack.h"
+
+#include "MovieSceneColorSection.h"
+#include "MovieSceneBoolSection.h"
 #include "MovieSceneFloatSection.h"
 #include "MovieScene3DTransformSection.h"
 #include "MovieSceneSkeletalAnimationSection.h"
+#include "MovieSceneAudioSection.h"
+#include "MovieSceneFadeSection.h"
+#include "MovieSceneCameraCutSection.h"
+#include "MovieSceneEventSection.h"
+#include "MovieSceneVisibilitySection.h"
 
+#include "MovieSceneSequence.h"
 
 ERichCurveInterpMode FMatineeImportTools::MatineeInterpolationToRichCurveInterpolation( EInterpCurveMode CurveMode )
 {
@@ -43,6 +68,9 @@ ERichCurveTangentMode FMatineeImportTools::MatineeInterpolationToRichCurveTangen
 	case CIM_CurveBreak:
 		return ERichCurveTangentMode::RCTM_Break;
 	case CIM_CurveUser:
+	// Import auto-clamped curves as user curves because rich curves don't have support for clamped tangents, and if the
+	// user moves the converted keys, the tangents will get mangled.
+	case CIM_CurveAutoClamped:
 		return ERichCurveTangentMode::RCTM_User;
 	default:
 		return ERichCurveTangentMode::RCTM_Auto;
@@ -79,7 +107,43 @@ void FMatineeImportTools::SetOrAddKey( FRichCurve& Curve, float Time, float Valu
 }
 
 
-void FMatineeImportTools::CopyInterpFloatTrack( TSharedRef<ISequencer> Sequencer, UInterpTrackFloatBase* MatineeFloatTrack, UMovieSceneFloatTrack* FloatTrack )
+bool FMatineeImportTools::CopyInterpBoolTrack( UInterpTrackBoolProp* MatineeBoolTrack, UMovieSceneBoolTrack* BoolTrack )
+{
+	const FScopedTransaction Transaction( NSLOCTEXT( "Sequencer", "PasteMatineeFBoolTrack", "Paste Matinee Bool Track" ) );
+	bool bSectionCreated = false;
+
+	BoolTrack->Modify();
+
+	float KeyTime = MatineeBoolTrack->GetKeyframeTime( 0 );
+	UMovieSceneBoolSection* Section = Cast<UMovieSceneBoolSection>( MovieSceneHelpers::FindSectionAtTime( BoolTrack->GetAllSections(), KeyTime ) );
+	if ( Section == nullptr )
+	{
+		Section = Cast<UMovieSceneBoolSection>( BoolTrack->CreateNewSection() );
+		BoolTrack->AddSection( *Section );
+		Section->SetIsInfinite(true);
+		bSectionCreated = true;
+	}
+	if (Section->TryModify())
+	{
+		float SectionMin = Section->GetStartTime();
+		float SectionMax = Section->GetEndTime();
+
+		FIntegralCurve& BoolCurve = Section->GetCurve();
+		for ( const auto& Point : MatineeBoolTrack->BoolTrack )
+		{
+			BoolCurve.UpdateOrAddKey(Point.Time, Point.Value);
+			SectionMin = FMath::Min( SectionMin, Point.Time );
+			SectionMax = FMath::Max( SectionMax, Point.Time );
+		}
+
+		Section->SetStartTime( SectionMin );
+		Section->SetEndTime( SectionMax );
+	}
+
+	return bSectionCreated;
+}
+
+bool FMatineeImportTools::CopyInterpFloatTrack( UInterpTrackFloatBase* MatineeFloatTrack, UMovieSceneFloatTrack* FloatTrack )
 {
 	const FScopedTransaction Transaction( NSLOCTEXT( "Sequencer", "PasteMatineeFloatTrack", "Paste Matinee Float Track" ) );
 	bool bSectionCreated = false;
@@ -92,6 +156,7 @@ void FMatineeImportTools::CopyInterpFloatTrack( TSharedRef<ISequencer> Sequencer
 	{
 		Section = Cast<UMovieSceneFloatSection>( FloatTrack->CreateNewSection() );
 		FloatTrack->AddSection( *Section );
+		Section->SetIsInfinite(true);
 		bSectionCreated = true;
 	}
 	if (Section->TryModify())
@@ -109,18 +174,61 @@ void FMatineeImportTools::CopyInterpFloatTrack( TSharedRef<ISequencer> Sequencer
 
 		Section->SetStartTime( SectionMin );
 		Section->SetEndTime( SectionMax );
+	}
+
+	return bSectionCreated;
+}
+
+
+bool FMatineeImportTools::CopyInterpColorTrack( UInterpTrackColorProp* ColorPropTrack, UMovieSceneColorTrack* ColorTrack )
+{
+	const FScopedTransaction Transaction( NSLOCTEXT( "Sequencer", "PasteMatineeColorTrack", "Paste Matinee Color Track" ) );
+	bool bSectionCreated = false;
 
-		if ( bSectionCreated )
+	ColorTrack->Modify();
+
+	float KeyTime = ColorPropTrack->GetKeyframeTime( 0 );
+	UMovieSceneColorSection* Section = Cast<UMovieSceneColorSection>( MovieSceneHelpers::FindSectionAtTime( ColorTrack->GetAllSections(), KeyTime ) );
+	if ( Section == nullptr )
+	{
+		Section = Cast<UMovieSceneColorSection>( ColorTrack->CreateNewSection() );
+		ColorTrack->AddSection( *Section );
+		Section->GetRedCurve().SetDefaultValue(0.f);
+		Section->GetGreenCurve().SetDefaultValue(0.f);
+		Section->GetBlueCurve().SetDefaultValue(0.f);
+		Section->GetAlphaCurve().SetDefaultValue(1.f);
+		Section->SetIsInfinite(true);
+		bSectionCreated = true;
+	}
+
+	if (Section->TryModify())
+	{
+		float SectionMin = Section->GetStartTime();
+		float SectionMax = Section->GetEndTime();
+
+		FRichCurve& RedCurve = Section->GetRedCurve();
+		FRichCurve& GreenCurve = Section->GetGreenCurve();
+		FRichCurve& BlueCurve = Section->GetBlueCurve();
+
+		for ( const auto& Point : ColorPropTrack->VectorTrack.Points )
 		{
-			Sequencer->NotifyMovieSceneDataChanged();
+			FMatineeImportTools::SetOrAddKey( RedCurve, Point.InVal, Point.OutVal.X, Point.ArriveTangent.X, Point.LeaveTangent.X, Point.InterpMode );
+			FMatineeImportTools::SetOrAddKey( GreenCurve, Point.InVal, Point.OutVal.Y, Point.ArriveTangent.Y, Point.LeaveTangent.Y, Point.InterpMode );
+			FMatineeImportTools::SetOrAddKey( BlueCurve, Point.InVal, Point.OutVal.Z, Point.ArriveTangent.Z, Point.LeaveTangent.Z, Point.InterpMode );
+			SectionMin = FMath::Min( SectionMin, Point.InVal );
+			SectionMax = FMath::Max( SectionMax, Point.InVal );
 		}
+
+		Section->SetStartTime( SectionMin );
+		Section->SetEndTime( SectionMax );
 	}
-}
 
+	return bSectionCreated;
+}
 
-void FMatineeImportTools::CopyInterpMoveTrack( TSharedRef<ISequencer> Sequencer, UInterpTrackMove* MoveTrack, UMovieScene3DTransformTrack* TransformTrack )
+bool FMatineeImportTools::CopyInterpMoveTrack( UInterpTrackMove* MoveTrack, UMovieScene3DTransformTrack* TransformTrack )
 {
-	const FScopedTransaction Transaction( NSLOCTEXT( "Sequencer", "PasteMatineeMoveTrack", "Paste Matinee move track" ) );
+	const FScopedTransaction Transaction( NSLOCTEXT( "Sequencer", "PasteMatineeMoveTrack", "Paste Matinee Move Track" ) );
 	bool bSectionCreated = false;
 
 	TransformTrack->Modify();
@@ -134,6 +242,7 @@ void FMatineeImportTools::CopyInterpMoveTrack( TSharedRef<ISequencer> Sequencer,
 		Section->GetScaleCurve(EAxis::Y).SetDefaultValue(1);
 		Section->GetScaleCurve(EAxis::Z).SetDefaultValue(1);
 		TransformTrack->AddSection( *Section );
+		Section->SetIsInfinite(true);
 		bSectionCreated = true;
 	}
 
@@ -170,18 +279,15 @@ void FMatineeImportTools::CopyInterpMoveTrack( TSharedRef<ISequencer> Sequencer,
 
 		Section->SetStartTime( SectionMin );
 		Section->SetEndTime( SectionMax );
-
-		if ( bSectionCreated )
-		{
-			Sequencer->NotifyMovieSceneDataChanged();
-		}
 	}
+
+	return bSectionCreated;
 }
 
 
-void FMatineeImportTools::CopyInterpParticleTrack( TSharedRef<ISequencer> Sequencer, UInterpTrackToggle* MatineeToggleTrack, UMovieSceneParticleTrack* ParticleTrack )
+bool FMatineeImportTools::CopyInterpParticleTrack( UInterpTrackToggle* MatineeToggleTrack, UMovieSceneParticleTrack* ParticleTrack )
 {
-	const FScopedTransaction Transaction( NSLOCTEXT( "Sequencer", "PasteMatineeParticleTrack", "Paste Matinee particle track" ) );
+	const FScopedTransaction Transaction( NSLOCTEXT( "Sequencer", "PasteMatineeParticleTrack", "Paste Matinee Particle Track" ) );
 	bool bSectionCreated = false;
 
 	ParticleTrack->Modify();
@@ -214,20 +320,18 @@ void FMatineeImportTools::CopyInterpParticleTrack( TSharedRef<ISequencer> Sequen
 
 		Section->SetStartTime( SectionMin );
 		Section->SetEndTime( SectionMax );
-
-		if ( bSectionCreated )
-		{
-			Sequencer->NotifyMovieSceneDataChanged();
-		}
 	}
+
+	return bSectionCreated;
 }
 
 
-void FMatineeImportTools::CopyInterpAnimControlTrack( TSharedRef<ISequencer> Sequencer, UInterpTrackAnimControl* MatineeAnimControlTrack, UMovieSceneSkeletalAnimationTrack* SkeletalAnimationTrack )
+bool FMatineeImportTools::CopyInterpAnimControlTrack( UInterpTrackAnimControl* MatineeAnimControlTrack, UMovieSceneSkeletalAnimationTrack* SkeletalAnimationTrack, float EndPlaybackRange )
 {
 	// @todo - Sequencer - Add support for slot names once they are implemented.
-	const FScopedTransaction Transaction( NSLOCTEXT( "Sequencer", "PasteMatineeAnimTrack", "Paste Matinee anim track" ) );
-
+	const FScopedTransaction Transaction( NSLOCTEXT( "Sequencer", "PasteMatineeAnimTrack", "Paste Matinee Anim Track" ) );
+	bool bSectionCreated = false;
+	
 	SkeletalAnimationTrack->Modify();
 	SkeletalAnimationTrack->RemoveAllAnimationData();
 
@@ -244,12 +348,19 @@ void FMatineeImportTools::CopyInterpAnimControlTrack( TSharedRef<ISequencer> Seq
 			}
 			else
 			{
-				EndTime = Sequencer->GetFocusedMovieSceneSequence()->GetMovieScene()->GetPlaybackRange().GetUpperBoundValue();
+				EndTime = EndPlaybackRange;
 			}
 		}
 		else
 		{
 			EndTime = AnimSeq.StartTime + ( ( ( AnimSeq.AnimSeq->SequenceLength - AnimSeq.AnimEndOffset ) - AnimSeq.AnimStartOffset ) / AnimSeq.AnimPlayRate );
+
+			// Clamp to next clip's start time
+			if (i+1 < MatineeAnimControlTrack->AnimSeqs.Num())
+			{
+				float NextStartTime = MatineeAnimControlTrack->AnimSeqs[i+1].StartTime;
+				EndTime = FMath::Min(NextStartTime, EndTime);
+			}
 		}
 
 		UMovieSceneSkeletalAnimationSection* NewSection = Cast<UMovieSceneSkeletalAnimationSection>( SkeletalAnimationTrack->CreateNewSection() );
@@ -259,9 +370,202 @@ void FMatineeImportTools::CopyInterpAnimControlTrack( TSharedRef<ISequencer> Seq
 		NewSection->SetEndOffset( AnimSeq.AnimEndOffset );
 		NewSection->SetPlayRate( AnimSeq.AnimPlayRate );
 		NewSection->SetAnimSequence( AnimSeq.AnimSeq );
+		NewSection->SetSlotName( MatineeAnimControlTrack->SlotName );
 
 		SkeletalAnimationTrack->AddSection( *NewSection );
+		bSectionCreated = true;
 	}
 
-	Sequencer->NotifyMovieSceneDataChanged();
+	return bSectionCreated;
+}
+
+bool FMatineeImportTools::CopyInterpSoundTrack( UInterpTrackSound* MatineeSoundTrack, UMovieSceneAudioTrack* AudioTrack )
+{
+	const FScopedTransaction Transaction( NSLOCTEXT( "Sequencer", "PasteMatineeSoundTrack", "Paste Matinee Sound Track" ) );
+	bool bSectionCreated = false;
+
+	AudioTrack->Modify();
+	
+	int MaxSectionRowIndex = -1;
+	for ( UMovieSceneSection* Section : AudioTrack->GetAllSections() )
+	{
+		MaxSectionRowIndex = FMath::Max( MaxSectionRowIndex, Section->GetRowIndex() );
+	}
+
+	for ( const FSoundTrackKey& SoundTrackKey : MatineeSoundTrack->Sounds )
+	{
+		AudioTrack->AddNewSound( SoundTrackKey.Sound, SoundTrackKey.Time );
+
+		UMovieSceneAudioSection* NewAudioSection = Cast<UMovieSceneAudioSection>(AudioTrack->GetAllSections().Last());
+		NewAudioSection->SetRowIndex( MaxSectionRowIndex + 1 );
+		// The audio dilation is inverted before it's set as the pitch, so we need to invert the pitch here to get the same result.
+		NewAudioSection->SetAudioDilationFactor( 1 / SoundTrackKey.Pitch );
+		NewAudioSection->SetAudioVolume( SoundTrackKey.Volume );
+
+		AudioTrack->AddSection( *NewAudioSection );
+		bSectionCreated = true;
+	}
+
+	return bSectionCreated;
+}
+
+bool FMatineeImportTools::CopyInterpFadeTrack( UInterpTrackFade* MatineeFadeTrack, UMovieSceneFadeTrack* FadeTrack )
+{
+	const FScopedTransaction Transaction( NSLOCTEXT( "Sequencer", "PasteMatineeFadeTrack", "Paste Matinee Fade Track" ) );
+	bool bSectionCreated = false;
+
+	FadeTrack->Modify();
+	
+	float KeyTime = MatineeFadeTrack->GetKeyframeTime( 0 );
+	UMovieSceneFadeSection* Section = Cast<UMovieSceneFadeSection>( MovieSceneHelpers::FindSectionAtTime( FadeTrack->GetAllSections(), KeyTime ) );
+	if ( Section == nullptr )
+	{
+		Section = Cast<UMovieSceneFadeSection>( FadeTrack->CreateNewSection() );
+		FadeTrack->AddSection( *Section );
+		bSectionCreated = true;
+	}
+	if (Section->TryModify())
+	{
+		float SectionMin = Section->GetStartTime();		
+		float SectionMax = Section->GetEndTime();
+		
+		FRichCurve& FloatCurve = Section->GetFloatCurve();
+		for ( const auto& Point : MatineeFadeTrack->FloatTrack.Points )
+		{
+			FMatineeImportTools::SetOrAddKey( FloatCurve, Point.InVal, Point.OutVal, Point.ArriveTangent, Point.LeaveTangent, Point.InterpMode );
+			SectionMin = FMath::Min( SectionMin, Point.InVal );
+			SectionMax = FMath::Max( SectionMax, Point.InVal );
+		}
+
+		Section->SetStartTime( SectionMin );
+		Section->SetEndTime( SectionMax );
+	
+		Section->FadeColor = MatineeFadeTrack->FadeColor;
+		Section->bFadeAudio = MatineeFadeTrack->bFadeAudio;
+	}
+
+	return bSectionCreated;
+}
+
+bool FMatineeImportTools::CopyInterpDirectorTrack( UInterpTrackDirector* DirectorTrack, UMovieSceneCameraCutTrack* CameraCutTrack, AMatineeActor* MatineeActor, UMovieSceneSequence* MovieSceneSequence )
+{
+	const FScopedTransaction Transaction( NSLOCTEXT( "Sequencer", "PasteMatineeDirectorTrack", "Paste Matinee Director Track" ) );
+	bool bCutsAdded = false;
+
+	CameraCutTrack->Modify();
+	
+	for (FDirectorTrackCut TrackCut : DirectorTrack->CutTrack)
+	{
+		int32 GroupIndex = MatineeActor->MatineeData->FindGroupByName(TrackCut.TargetCamGroup);
+		
+		UInterpGroupInst* ViewGroupInst = (GroupIndex != INDEX_NONE) ? MatineeActor->FindFirstGroupInstByName( TrackCut.TargetCamGroup.ToString() ) : NULL;
+		if ( GroupIndex != INDEX_NONE && ViewGroupInst )
+		{
+			// Find a valid move track for this cut.
+			UInterpGroup* Group = MatineeActor->MatineeData->InterpGroups[GroupIndex];
+			if (Group)
+			{
+				AActor* CameraActor = ViewGroupInst->GetGroupActor();
+		
+				FGuid CameraHandle = MovieSceneSequence->FindPossessableObjectId(*CameraActor);
+				if (CameraHandle.IsValid())
+				{
+					CameraCutTrack->AddNewCameraCut(CameraHandle, TrackCut.Time);
+					bCutsAdded = true;
+				}
+			}
+		}
+	}
+
+	return bCutsAdded;
+}
+
+bool FMatineeImportTools::CopyInterpEventTrack( UInterpTrackEvent* MatineeEventTrack, UMovieSceneEventTrack* EventTrack )
+{
+	const FScopedTransaction Transaction( NSLOCTEXT( "Sequencer", "PasteMatineeEventTrack", "Paste Matinee Event Track" ) );
+	bool bSectionCreated = false;
+
+	EventTrack->Modify();
+
+	if (MatineeEventTrack->EventTrack.Num())
+	{
+		float KeyTime = MatineeEventTrack->EventTrack[0].Time;
+		UMovieSceneEventSection* Section = Cast<UMovieSceneEventSection>( MovieSceneHelpers::FindSectionAtTime( EventTrack->GetAllSections(), KeyTime ) );
+		if ( Section == nullptr )
+		{
+			Section = Cast<UMovieSceneEventSection>( EventTrack->CreateNewSection() );
+			EventTrack->AddSection( *Section );
+			bSectionCreated = true;
+		}
+		if (Section->TryModify())
+		{
+			float SectionMin = Section->GetStartTime();
+			float SectionMax = Section->GetEndTime();
+			
+			FNameCurve& EventCurve = Section->GetEventCurve();
+			for (FEventTrackKey EventTrackKey : MatineeEventTrack->EventTrack)
+			{
+				EventCurve.UpdateOrAddKey(EventTrackKey.Time, EventTrackKey.EventName);
+				SectionMin = FMath::Min( SectionMin, EventTrackKey.Time );
+				SectionMax = FMath::Max( SectionMax, EventTrackKey.Time );
+			}
+
+			Section->SetStartTime( SectionMin );
+			Section->SetEndTime( SectionMax );
+		}
+	}	
+	
+	return bSectionCreated;
+}
+
+bool FMatineeImportTools::CopyInterpVisibilityTrack( UInterpTrackVisibility* MatineeVisibilityTrack, UMovieSceneVisibilityTrack* VisibilityTrack )
+{
+	const FScopedTransaction Transaction( NSLOCTEXT( "Sequencer", "PasteMatineeVisibilityTrack", "Paste Matinee Visibility track" ) );
+	bool bSectionCreated = false;
+
+	VisibilityTrack->Modify();
+
+	if (MatineeVisibilityTrack->VisibilityTrack.Num())
+	{
+		float KeyTime = MatineeVisibilityTrack->VisibilityTrack[0].Time;
+		UMovieSceneVisibilitySection* Section = Cast<UMovieSceneVisibilitySection>( MovieSceneHelpers::FindSectionAtTime( VisibilityTrack->GetAllSections(), KeyTime ) );
+		if ( Section == nullptr )
+		{
+			Section = Cast<UMovieSceneVisibilitySection>( VisibilityTrack->CreateNewSection() );
+			VisibilityTrack->AddSection( *Section );
+			bSectionCreated = true;
+		}
+		if (Section->TryModify())
+		{
+			float SectionMin = Section->GetStartTime();
+			float SectionMax = Section->GetEndTime();
+			
+			bool bVisible = true;
+			FIntegralCurve& VisibilityCurve = Section->GetCurve();
+			for (FVisibilityTrackKey VisibilityTrackKey : MatineeVisibilityTrack->VisibilityTrack)
+			{
+				if (VisibilityTrackKey.Action == EVisibilityTrackAction::EVTA_Hide)
+				{
+					bVisible = false;
+				}
+				else if (VisibilityTrackKey.Action == EVisibilityTrackAction::EVTA_Hide)
+				{
+					bVisible = true;
+				}
+				else if (VisibilityTrackKey.Action == EVisibilityTrackAction::EVTA_Toggle)
+				{
+					bVisible = !bVisible;
+				}
+
+				VisibilityCurve.UpdateOrAddKey(VisibilityTrackKey.Time, bVisible);
+				SectionMin = FMath::Min( SectionMin, VisibilityTrackKey.Time );
+				SectionMax = FMath::Max( SectionMax, VisibilityTrackKey.Time );
+			}
+
+			Section->SetStartTime( SectionMin );
+			Section->SetEndTime( SectionMax );
+		}
+	}	
+	
+	return bSectionCreated;
 }
diff --git a/Engine/Source/Editor/MovieSceneTools/Private/Sections/CameraCutSection.cpp b/Engine/Source/Editor/MovieSceneTools/Private/Sections/CameraCutSection.cpp
--- a/Engine/Source/Editor/MovieSceneTools/Private/Sections/CameraCutSection.cpp
+++ b/Engine/Source/Editor/MovieSceneTools/Private/Sections/CameraCutSection.cpp
@@ -133,7 +133,7 @@ const AActor* FCameraCutSection::GetCameraForFrame(float Time) const
 			FMovieSceneSpawnable* Spawnable = SequenceInstance->GetSequence()->GetMovieScene()->FindSpawnable(CameraCutSection->GetCameraGuid());
 			if (Spawnable)
 			{
-				return GetDefault<AActor>(Spawnable->GetClass());
+				return Cast<AActor>(Spawnable->GetObjectTemplate());
 			}
 		}
 	}
diff --git a/Engine/Source/Editor/MovieSceneTools/Private/Sections/CinematicShotSection.cpp b/Engine/Source/Editor/MovieSceneTools/Private/Sections/CinematicShotSection.cpp
--- a/Engine/Source/Editor/MovieSceneTools/Private/Sections/CinematicShotSection.cpp
+++ b/Engine/Source/Editor/MovieSceneTools/Private/Sections/CinematicShotSection.cpp
@@ -240,7 +240,10 @@ FReply FCinematicShotSection::OnSectionDoubleClicked(const FGeometry& SectionGeo
 	{
 		if (SequenceInstance.IsValid())
 		{
-			Sequencer.Pin()->FocusSequenceInstance(SectionObject);
+			if (SectionObject.GetSequence())
+			{
+				Sequencer.Pin()->FocusSequenceInstance(SectionObject);
+			}
 		}
 	}
 
diff --git a/Engine/Source/Editor/MovieSceneTools/Private/Sections/ColorPropertySection.cpp b/Engine/Source/Editor/MovieSceneTools/Private/Sections/ColorPropertySection.cpp
--- a/Engine/Source/Editor/MovieSceneTools/Private/Sections/ColorPropertySection.cpp
+++ b/Engine/Source/Editor/MovieSceneTools/Private/Sections/ColorPropertySection.cpp
@@ -195,7 +195,7 @@ FLinearColor FColorPropertySection::FindSlateColor(const FName& ColorName) const
 				return *Property->ContainerPtrToValuePtr<FLinearColor>(RuntimeObject);
 			}
 
-			return Property->ContainerPtrToValuePtr<FColor>(RuntimeObject)->ReinterpretAsLinear();
+			return *Property->ContainerPtrToValuePtr<FColor>(RuntimeObject);
 		}
 	}
 
@@ -214,7 +214,7 @@ void FColorPropertySection::SetIntermediateValue( FPropertyChangedParams Propert
 
 	if ( bIsFColor )
 	{
-		ColorValue = PropertyChangedParams.GetPropertyValue<FColor>().ReinterpretAsLinear();
+		ColorValue = PropertyChangedParams.GetPropertyValue<FColor>();
 	}
 	else
 	{
diff --git a/Engine/Source/Editor/MovieSceneTools/Private/TrackEditorThumbnail/TrackEditorThumbnail.cpp b/Engine/Source/Editor/MovieSceneTools/Private/TrackEditorThumbnail/TrackEditorThumbnail.cpp
--- a/Engine/Source/Editor/MovieSceneTools/Private/TrackEditorThumbnail/TrackEditorThumbnail.cpp
+++ b/Engine/Source/Editor/MovieSceneTools/Private/TrackEditorThumbnail/TrackEditorThumbnail.cpp
@@ -364,10 +364,10 @@ void FTrackEditorThumbnailCache::Revalidate()
 		return;
 	}
 
-	if (FMath::IsNearlyZero(CurrentCache.TimeRange.Size<float>()))
+	if (FMath::IsNearlyZero(CurrentCache.TimeRange.Size<float>()) || CurrentCache.TimeRange.IsEmpty())
 	{
 		// Can't generate thumbnails for this
-		ThumbnailPool.Pin()->RemoveThumbnailsNeedingRedraw(ThumbnailsNeedingRedraw);
+		ThumbnailPool.Pin()->RemoveThumbnailsNeedingRedraw(Thumbnails);
 		ThumbnailsNeedingRedraw.Reset();
 		Thumbnails.Reset();
 		return;
@@ -375,7 +375,7 @@ void FTrackEditorThumbnailCache::Revalidate()
 
 	if (ShouldRegenerateEverything())
 	{
-		ThumbnailPool.Pin()->RemoveThumbnailsNeedingRedraw(ThumbnailsNeedingRedraw);
+		ThumbnailPool.Pin()->RemoveThumbnailsNeedingRedraw(Thumbnails);
 		Thumbnails.Reset();
 	}
 	
diff --git a/Engine/Source/Editor/MovieSceneTools/Private/TrackEditors/AudioTrackEditor.cpp b/Engine/Source/Editor/MovieSceneTools/Private/TrackEditors/AudioTrackEditor.cpp
--- a/Engine/Source/Editor/MovieSceneTools/Private/TrackEditors/AudioTrackEditor.cpp
+++ b/Engine/Source/Editor/MovieSceneTools/Private/TrackEditors/AudioTrackEditor.cpp
@@ -28,6 +28,8 @@
 #include "ContentBrowserModule.h"
 #include "SequencerUtilities.h"
 #include "AssetRegistryModule.h"
+#include "MatineeImportTools.h"
+#include "Matinee/InterpTrackSound.h"
 
 
 #define LOCTEXT_NAMESPACE "FAudioTrackEditor"
@@ -49,6 +51,10 @@ USoundWave* DeriveSoundWave(UMovieSceneAudioSection* AudioSection)
 	USoundWave* SoundWave = NULL;
 	
 	USoundBase* Sound = AudioSection->GetSound();
+	if (!Sound)
+	{
+		return SoundWave;
+	}
 
 	if (Sound->IsA<USoundWave>())
 	{
@@ -634,7 +640,13 @@ FText FAudioSection::GetDisplayName() const
 
 FText FAudioSection::GetSectionTitle() const
 {
-	return FText::FromString( Cast<UMovieSceneAudioSection>(&Section)->GetSound()->GetName() );
+	UMovieSceneAudioSection* AudioSection = Cast<UMovieSceneAudioSection>(&Section);
+	if (AudioSection && AudioSection->GetSound())
+	{
+		return FText::FromString(AudioSection->GetSound()->GetName());
+	}
+	
+	return NSLOCTEXT("FAudioSection", "NoAudioTitleName", "No Audio");
 }
 
 
@@ -672,7 +684,7 @@ void FAudioSection::Tick( const FGeometry& AllottedGeometry, const FGeometry& Pa
 	UMovieSceneTrack* Track = Section.GetTypedOuter<UMovieSceneTrack>();
 
 	USoundWave* SoundWave = DeriveSoundWave(AudioSection);
-	if (Track && (SoundWave->NumChannels == 1 || SoundWave->NumChannels == 2))
+	if (Track && SoundWave && (SoundWave->NumChannels == 1 || SoundWave->NumChannels == 2))
 	{
 		const FSlateRect ParentRect = TransformRect(
 			Concatenate(
@@ -698,9 +710,9 @@ void FAudioSection::Tick( const FGeometry& AllottedGeometry, const FGeometry& Pa
 
 		if (!FMath::IsNearlyEqual(DrawRange.GetLowerBoundValue(), StoredDrawRange.GetLowerBoundValue()) ||
 			!FMath::IsNearlyEqual(DrawRange.GetUpperBoundValue(), StoredDrawRange.GetUpperBoundValue()) ||
-			XOffset != StoredXOffset || XSize != StoredXSize || Track->GetColorTint() != StoredColor)
+			XOffset != StoredXOffset || XSize != StoredXSize || Track->GetColorTint() != StoredColor ||
+			StoredSoundWave != SoundWave)
 		{
-
 			float DisplayScale = XSize / DrawRange.Size<float>();
 
 			// Use the view range if possible, as it's much more stable than using the texture size and draw range
@@ -711,8 +723,15 @@ void FAudioSection::Tick( const FGeometry& AllottedGeometry, const FGeometry& Pa
 			}
 
 			RegenerateWaveforms(DrawRange, XOffset, XSize, Track->GetColorTint(), DisplayScale);
+			StoredSoundWave = SoundWave;
 		}
 	}
+	else
+	{
+		WaveformThumbnail.Reset();
+		StoredDrawRange = TRange<float>::Empty();
+		StoredSoundWave.Reset();
+	}
 }
 
 
@@ -768,6 +787,37 @@ bool FAudioTrackEditor::SupportsType( TSubclassOf<UMovieSceneTrack> Type ) const
 	return Type == UMovieSceneAudioTrack::StaticClass();
 }
 
+
+void CopyInterpSoundTrack(TSharedRef<ISequencer> Sequencer, UInterpTrackSound* MatineeSoundTrack, UMovieSceneAudioTrack* AudioTrack)
+{
+	if (FMatineeImportTools::CopyInterpSoundTrack(MatineeSoundTrack, AudioTrack))
+	{
+		Sequencer.Get().NotifyMovieSceneDataChanged();
+	}
+}
+
+void FAudioTrackEditor::BuildTrackContextMenu( FMenuBuilder& MenuBuilder, UMovieSceneTrack* Track )
+{
+	UInterpTrackSound* MatineeSoundTrack = nullptr;
+	for ( UObject* CopyPasteObject : GUnrealEd->MatineeCopyPasteBuffer )
+	{
+		MatineeSoundTrack = Cast<UInterpTrackSound>( CopyPasteObject );
+		if ( MatineeSoundTrack != nullptr )
+		{
+			break;
+		}
+	}
+	UMovieSceneAudioTrack* AudioTrack = Cast<UMovieSceneAudioTrack>( Track );
+	MenuBuilder.AddMenuEntry(
+		NSLOCTEXT( "Sequencer", "PasteMatineeSoundTrack", "Paste Matinee sound track" ),
+		NSLOCTEXT( "Sequencer", "PasteMatineeSoundTrackTooltip", "Pastes keys from a Matinee sound track into this track." ),
+		FSlateIcon(),
+		FUIAction(
+		FExecuteAction::CreateStatic( &CopyInterpSoundTrack, GetSequencer().ToSharedRef(), MatineeSoundTrack, AudioTrack ),
+		FCanExecuteAction::CreateLambda( [=]()->bool { return MatineeSoundTrack != nullptr && MatineeSoundTrack->Sounds.Num() > 0 && AudioTrack != nullptr; } ) ) );
+}
+
+
 const FSlateBrush* FAudioTrackEditor::GetIconBrush() const
 {
 	return FEditorStyle::GetBrush("Sequencer.Tracks.Audio");
@@ -821,11 +871,15 @@ bool FAudioTrackEditor::HandleAssetAdded(UObject* Asset, const FGuid& TargetObje
 
 bool FAudioTrackEditor::AddNewMasterSound( float KeyTime, USoundBase* Sound )
 {
-	UMovieSceneTrack* Track = FindOrCreateMasterTrack<UMovieSceneAudioTrack>().Track;
+	FFindOrCreateMasterTrackResult<UMovieSceneAudioTrack> TrackResult = FindOrCreateMasterTrack<UMovieSceneAudioTrack>();
+	UMovieSceneTrack* Track = TrackResult.Track;
 
 	auto AudioTrack = Cast<UMovieSceneAudioTrack>(Track);
 	AudioTrack->AddNewSound( Sound, KeyTime );
-	AudioTrack->SetDisplayName(LOCTEXT("AudioTrackName", "Audio"));
+	if (TrackResult.bWasCreated)
+	{
+		AudioTrack->SetDisplayName(LOCTEXT("AudioTrackName", "Audio"));
+	}
 
 	return true;
 }
diff --git a/Engine/Source/Editor/MovieSceneTools/Private/TrackEditors/AudioTrackEditor.h b/Engine/Source/Editor/MovieSceneTools/Private/TrackEditors/AudioTrackEditor.h
--- a/Engine/Source/Editor/MovieSceneTools/Private/TrackEditors/AudioTrackEditor.h
+++ b/Engine/Source/Editor/MovieSceneTools/Private/TrackEditors/AudioTrackEditor.h
@@ -38,6 +38,7 @@ public:
 	virtual bool HandleAssetAdded(UObject* Asset, const FGuid& TargetObjectGuid) override;
 	virtual TSharedRef<ISequencerSection> MakeSectionInterface(UMovieSceneSection& SectionObject, UMovieSceneTrack& Track) override;
 	virtual bool SupportsType(TSubclassOf<UMovieSceneTrack> Type) const override;
+	virtual void BuildTrackContextMenu( FMenuBuilder& MenuBuilder, UMovieSceneTrack* Track ) override;
 	virtual const FSlateBrush* GetIconBrush() const override;
 	
 protected:
@@ -107,6 +108,9 @@ private:
 	int32 StoredXSize;
 	FColor StoredColor;
 
+	/** Stored sound wave to determine when it is invalidated. */
+	TWeakObjectPtr<USoundWave> StoredSoundWave;
+
 	/** Whether this section is on a master audio track or an attached audio track. */
 	bool bIsOnAMasterTrack;
 
diff --git a/Engine/Source/Editor/MovieSceneTools/Private/TrackEditors/CinematicShotTrackEditor.cpp b/Engine/Source/Editor/MovieSceneTools/Private/TrackEditors/CinematicShotTrackEditor.cpp
--- a/Engine/Source/Editor/MovieSceneTools/Private/TrackEditors/CinematicShotTrackEditor.cpp
+++ b/Engine/Source/Editor/MovieSceneTools/Private/TrackEditors/CinematicShotTrackEditor.cpp
@@ -403,7 +403,15 @@ void FCinematicShotTrackEditor::NewTake(UMovieSceneCinematicShotSection* Section
 	uint32 TakeNumber = INDEX_NONE;
 	if (MovieSceneToolHelpers::ParseShotName(Section->GetShotDisplayName().ToString(), ShotPrefix, ShotNumber, TakeNumber))
 	{
-		uint32 NewTakeNumber = TakeNumber + 1;
+		TArray<uint32> TakeNumbers;
+		uint32 CurrentTakeNumber;
+		MovieSceneToolHelpers::GatherTakes(Section, TakeNumbers, CurrentTakeNumber);
+		uint32 NewTakeNumber = CurrentTakeNumber;
+		if (TakeNumbers.Num() > 0)
+		{
+			NewTakeNumber = TakeNumbers[TakeNumbers.Num()-1] + 1;
+		}
+
 		FString NewShotName = MovieSceneToolHelpers::ComposeShotName(ShotPrefix, ShotNumber, NewTakeNumber);
 
 		float NewShotStartTime = Section->GetStartTime();
@@ -415,16 +423,19 @@ void FCinematicShotTrackEditor::NewTake(UMovieSceneCinematicShotSection* Section
 		const bool bInsertShot = false;
 		UMovieSceneSubSection* NewShot = CreateShotInternal(NewShotName, NewShotStartTime, Section, bInsertShot);
 
-		UMovieSceneCinematicShotTrack* CinematicShotTrack = FindOrCreateCinematicShotTrack();
-		CinematicShotTrack->RemoveSection(*Section);
-
-		if (NewShot != nullptr)
+		if (NewShot)
 		{
-			NewShot->SetStartTime(NewShotStartTime);
-			NewShot->SetEndTime(NewShotEndTime);
-			NewShot->StartOffset = NewShotStartOffset;
-			NewShot->TimeScale = NewShotTimeScale;
-			NewShot->PrerollTime = NewShotPrerollTime;
+			UMovieSceneCinematicShotTrack* CinematicShotTrack = FindOrCreateCinematicShotTrack();
+			CinematicShotTrack->RemoveSection(*Section);
+
+			if (NewShot != nullptr)
+			{
+				NewShot->SetStartTime(NewShotStartTime);
+				NewShot->SetEndTime(NewShotEndTime);
+				NewShot->StartOffset = NewShotStartOffset;
+				NewShot->TimeScale = NewShotTimeScale;
+				NewShot->PrerollTime = NewShotPrerollTime;
+			}
 		}
 	}
 
diff --git a/Engine/Source/Editor/MovieSceneTools/Private/TrackEditors/MaterialTrackEditor.cpp b/Engine/Source/Editor/MovieSceneTools/Private/TrackEditors/MaterialTrackEditor.cpp
--- a/Engine/Source/Editor/MovieSceneTools/Private/TrackEditors/MaterialTrackEditor.cpp
+++ b/Engine/Source/Editor/MovieSceneTools/Private/TrackEditors/MaterialTrackEditor.cpp
@@ -175,8 +175,13 @@ bool FComponentMaterialTrackEditor::SupportsType( TSubclassOf<UMovieSceneTrack>
 
 UMaterialInterface* FComponentMaterialTrackEditor::GetMaterialInterfaceForTrack( FGuid ObjectBinding, UMovieSceneMaterialTrack* MaterialTrack )
 {
-	UObject* ComponentObject = GetSequencer()->GetFocusedMovieSceneSequenceInstance()->FindObject( ObjectBinding, *GetSequencer() );
-	UPrimitiveComponent* Component = Cast<UPrimitiveComponent>( ComponentObject );
+	TSharedPtr<ISequencer> Sequencer = GetSequencer();
+	if (!Sequencer.IsValid())
+	{
+		return nullptr;
+	}
+
+	UPrimitiveComponent* Component = Cast<UPrimitiveComponent>(Sequencer->FindSpawnedObjectOrTemplate( ObjectBinding ));
 	UMovieSceneComponentMaterialTrack* ComponentMaterialTrack = Cast<UMovieSceneComponentMaterialTrack>( MaterialTrack );
 	if ( Component != nullptr && ComponentMaterialTrack != nullptr )
 	{
diff --git a/Engine/Source/Editor/MovieSceneTools/Private/TrackEditors/ParticleParameterTrackEditor.cpp b/Engine/Source/Editor/MovieSceneTools/Private/TrackEditors/ParticleParameterTrackEditor.cpp
--- a/Engine/Source/Editor/MovieSceneTools/Private/TrackEditors/ParticleParameterTrackEditor.cpp
+++ b/Engine/Source/Editor/MovieSceneTools/Private/TrackEditors/ParticleParameterTrackEditor.cpp
@@ -73,7 +73,9 @@ TSharedRef<SWidget> FParticleParameterTrackEditor::OnGetAddParameterMenuContent(
 {
 	FMenuBuilder AddParameterMenuBuilder( true, nullptr );
 
-	AEmitter* Emitter = Cast<AEmitter>( GetSequencer()->GetFocusedMovieSceneSequenceInstance()->FindObject( ObjectBinding, *GetSequencer() ) );
+	TSharedPtr<ISequencer> Sequencer = GetSequencer();
+	AEmitter* Emitter = Sequencer.IsValid() ? Cast<AEmitter>( Sequencer->FindSpawnedObjectOrTemplate( ObjectBinding ) ) : nullptr;
+
 	if ( Emitter != nullptr )
 	{
 		UParticleSystemComponent* ParticleSystemComponent = Emitter->GetParticleSystemComponent();
diff --git a/Engine/Source/Editor/MovieSceneTools/Private/TrackEditors/ParticleTrackEditor.cpp b/Engine/Source/Editor/MovieSceneTools/Private/TrackEditors/ParticleTrackEditor.cpp
--- a/Engine/Source/Editor/MovieSceneTools/Private/TrackEditors/ParticleTrackEditor.cpp
+++ b/Engine/Source/Editor/MovieSceneTools/Private/TrackEditors/ParticleTrackEditor.cpp
@@ -304,38 +304,36 @@ void FParticleTrackEditor::BuildObjectBindingTrackMenu(FMenuBuilder& MenuBuilder
 
 void FParticleTrackEditor::AddParticleKey( const FGuid ObjectGuid )
 {
-	TArray<TWeakObjectPtr<UObject>> OutObjects;
+	TSharedPtr<ISequencer> Sequencer = GetSequencer();
+	UObject* Object = Sequencer.IsValid() ? Sequencer->FindSpawnedObjectOrTemplate(ObjectGuid) : nullptr;
 
-	GetSequencer()->GetRuntimeObjects( GetSequencer()->GetFocusedMovieSceneSequenceInstance(), ObjectGuid, OutObjects );
-	AnimatablePropertyChanged( FOnKeyProperty::CreateRaw( this, &FParticleTrackEditor::AddKeyInternal, OutObjects ) );
+	if (Object)
+	{
+		AnimatablePropertyChanged( FOnKeyProperty::CreateRaw( this, &FParticleTrackEditor::AddKeyInternal, Object ) );
+	}
 }
 
 
-bool FParticleTrackEditor::AddKeyInternal( float KeyTime, const TArray<TWeakObjectPtr<UObject>> Objects )
+bool FParticleTrackEditor::AddKeyInternal( float KeyTime, UObject* Object )
 {
 	bool bHandleCreated = false;
 	bool bTrackCreated = false;
 
-	for( int32 ObjectIndex = 0; ObjectIndex < Objects.Num(); ++ObjectIndex )
-	{
-		UObject* Object = Objects[ObjectIndex].Get();
+	FFindOrCreateHandleResult HandleResult = FindOrCreateHandleToObject( Object );
+	FGuid ObjectHandle = HandleResult.Handle;
+	bHandleCreated |= HandleResult.bWasCreated;
 
-		FFindOrCreateHandleResult HandleResult = FindOrCreateHandleToObject( Object );
-		FGuid ObjectHandle = HandleResult.Handle;
-		bHandleCreated |= HandleResult.bWasCreated;
+	if (ObjectHandle.IsValid())
+	{
+		FFindOrCreateTrackResult TrackResult = FindOrCreateTrackForObject(ObjectHandle, UMovieSceneParticleTrack::StaticClass());
+		UMovieSceneTrack* Track = TrackResult.Track;
+		bTrackCreated |= TrackResult.bWasCreated;
 
-		if (ObjectHandle.IsValid())
+		if (bTrackCreated && ensure(Track))
 		{
-			FFindOrCreateTrackResult TrackResult = FindOrCreateTrackForObject(ObjectHandle, UMovieSceneParticleTrack::StaticClass());
-			UMovieSceneTrack* Track = TrackResult.Track;
-			bTrackCreated |= TrackResult.bWasCreated;
-
-			if (bTrackCreated && ensure(Track))
-			{
-				UMovieSceneParticleTrack* ParticleTrack = Cast<UMovieSceneParticleTrack>(Track);
-				ParticleTrack->AddNewSection(KeyTime);
-				ParticleTrack->SetDisplayName(LOCTEXT("TrackName", "Particle System"));
-			}
+			UMovieSceneParticleTrack* ParticleTrack = Cast<UMovieSceneParticleTrack>(Track);
+			ParticleTrack->AddNewSection(KeyTime);
+			ParticleTrack->SetDisplayName(LOCTEXT("TrackName", "Particle System"));
 		}
 	}
 
@@ -343,6 +341,14 @@ bool FParticleTrackEditor::AddKeyInternal( float KeyTime, const TArray<TWeakObje
 }
 
 
+void CopyInterpParticleTrack(TSharedRef<ISequencer> Sequencer, UInterpTrackToggle* MatineeToggleTrack, UMovieSceneParticleTrack* ParticleTrack)
+{
+	if (FMatineeImportTools::CopyInterpParticleTrack(MatineeToggleTrack, ParticleTrack))
+	{
+		Sequencer.Get().NotifyMovieSceneDataChanged();
+	}
+}
+
 void FParticleTrackEditor::BuildTrackContextMenu( FMenuBuilder& MenuBuilder, UMovieSceneTrack* Track )
 {
 	UInterpTrackToggle* MatineeToggleTrack = nullptr;
@@ -360,7 +366,7 @@ void FParticleTrackEditor::BuildTrackContextMenu( FMenuBuilder& MenuBuilder, UMo
 		NSLOCTEXT( "Sequencer", "PasteMatineeToggleTrackTooltip", "Pastes keys from a Matinee particle track into this track." ),
 		FSlateIcon(),
 		FUIAction(
-			FExecuteAction::CreateStatic( &FMatineeImportTools::CopyInterpParticleTrack, GetSequencer().ToSharedRef(), MatineeToggleTrack, ParticleTrack ),
+			FExecuteAction::CreateStatic( &CopyInterpParticleTrack, GetSequencer().ToSharedRef(), MatineeToggleTrack, ParticleTrack ),
 			FCanExecuteAction::CreateLambda( [=]()->bool { return MatineeToggleTrack != nullptr && MatineeToggleTrack->ToggleTrack.Num() > 0 && ParticleTrack != nullptr; } ) ) );
 }
 
diff --git a/Engine/Source/Editor/MovieSceneTools/Private/TrackEditors/ParticleTrackEditor.h b/Engine/Source/Editor/MovieSceneTools/Private/TrackEditors/ParticleTrackEditor.h
--- a/Engine/Source/Editor/MovieSceneTools/Private/TrackEditors/ParticleTrackEditor.h
+++ b/Engine/Source/Editor/MovieSceneTools/Private/TrackEditors/ParticleTrackEditor.h
@@ -46,7 +46,7 @@ public:
 private:
 
 	/** Delegate for AnimatablePropertyChanged in AddKey. */
-	virtual bool AddKeyInternal( float KeyTime, const TArray<TWeakObjectPtr<UObject>> Objects);
+	virtual bool AddKeyInternal( float KeyTime, UObject* Object);
 };
 
 
diff --git a/Engine/Source/Editor/MovieSceneTools/Private/TrackEditors/PropertyTrackEditors/ColorPropertyTrackEditor.cpp b/Engine/Source/Editor/MovieSceneTools/Private/TrackEditors/PropertyTrackEditors/ColorPropertyTrackEditor.cpp
--- a/Engine/Source/Editor/MovieSceneTools/Private/TrackEditors/PropertyTrackEditors/ColorPropertyTrackEditor.cpp
+++ b/Engine/Source/Editor/MovieSceneTools/Private/TrackEditors/PropertyTrackEditors/ColorPropertyTrackEditor.cpp
@@ -4,6 +4,8 @@
 #include "MovieSceneColorTrack.h"
 #include "ColorPropertyTrackEditor.h"
 #include "ColorPropertySection.h"
+#include "MatineeImportTools.h"
+#include "Matinee/InterpTrackColorProp.h"
 
 
 FName FColorPropertyTrackEditor::RedName( "R" );
@@ -38,7 +40,7 @@ void FColorPropertyTrackEditor::GenerateKeysFromPropertyChanged( const FProperty
 
 	if (bIsFColor)
 	{
-		ColorValue = PropertyChangedParams.GetPropertyValue<FColor>().ReinterpretAsLinear();
+		ColorValue = FLinearColor( PropertyChangedParams.GetPropertyValue<FColor>() );
 	}
 	else
 	{
@@ -64,3 +66,33 @@ void FColorPropertyTrackEditor::GenerateKeysFromPropertyChanged( const FProperty
 	TArray<FColorKey>& AlphaKeys = ChannelName == NAME_None || ChannelName == AlphaName ? NewGeneratedKeys : DefaultGeneratedKeys;
 	AlphaKeys.Add( FColorKey( EKeyColorChannel::Alpha, ColorValue.A, bIsSlateColor ) );
 }
+
+void CopyInterpColorTrack(TSharedRef<ISequencer> Sequencer, UInterpTrackColorProp* ColorPropTrack, UMovieSceneColorTrack* ColorTrack)
+{
+	if (FMatineeImportTools::CopyInterpColorTrack(ColorPropTrack, ColorTrack))
+	{
+		Sequencer.Get().NotifyMovieSceneDataChanged();
+	}
+}
+
+void FColorPropertyTrackEditor::BuildTrackContextMenu( FMenuBuilder& MenuBuilder, UMovieSceneTrack* Track )
+{
+	UInterpTrackColorProp* ColorPropTrack = nullptr;
+	for ( UObject* CopyPasteObject : GUnrealEd->MatineeCopyPasteBuffer )
+	{
+		ColorPropTrack = Cast<UInterpTrackColorProp>( CopyPasteObject );
+		if ( ColorPropTrack != nullptr )
+		{
+			break;
+		}
+	}
+	UMovieSceneColorTrack* ColorTrack = Cast<UMovieSceneColorTrack>( Track );
+	MenuBuilder.AddMenuEntry(
+		NSLOCTEXT( "Sequencer", "PasteMatineeTrack", "Paste Matinee Color Track" ),
+		NSLOCTEXT( "Sequencer", "PasteMatineeTrackTooltip", "Pastes keys from a Matinee color track into this track." ),
+		FSlateIcon(),
+		FUIAction(
+			FExecuteAction::CreateStatic( &CopyInterpColorTrack, GetSequencer().ToSharedRef(), ColorPropTrack, ColorTrack ),
+			FCanExecuteAction::CreateLambda( [=]()->bool { return ColorPropTrack != nullptr && ColorPropTrack->GetNumKeys() > 0 && ColorTrack != nullptr; } ) ) );
+}
+
diff --git a/Engine/Source/Editor/MovieSceneTools/Private/TrackEditors/PropertyTrackEditors/ColorPropertyTrackEditor.h b/Engine/Source/Editor/MovieSceneTools/Private/TrackEditors/PropertyTrackEditors/ColorPropertyTrackEditor.h
--- a/Engine/Source/Editor/MovieSceneTools/Private/TrackEditors/PropertyTrackEditors/ColorPropertyTrackEditor.h
+++ b/Engine/Source/Editor/MovieSceneTools/Private/TrackEditors/PropertyTrackEditors/ColorPropertyTrackEditor.h
@@ -38,6 +38,10 @@ protected:
 	virtual TSharedRef<FPropertySection> MakePropertySectionInterface( UMovieSceneSection& SectionObject, UMovieSceneTrack& Track ) override;
 	virtual void GenerateKeysFromPropertyChanged( const FPropertyChangedParams& PropertyChangedParams, TArray<FColorKey>& NewGeneratedKeys, TArray<FColorKey>& DefaultGeneratedKeys ) override;
 
+	// ISequencerTrackEditor interface
+
+	virtual void BuildTrackContextMenu( FMenuBuilder& MenuBuilder, UMovieSceneTrack* Track ) override;
+
 private:
 	static FName RedName;
 	static FName GreenName;
diff --git a/Engine/Source/Editor/MovieSceneTools/Private/TrackEditors/PropertyTrackEditors/FloatPropertyTrackEditor.cpp b/Engine/Source/Editor/MovieSceneTools/Private/TrackEditors/PropertyTrackEditors/FloatPropertyTrackEditor.cpp
--- a/Engine/Source/Editor/MovieSceneTools/Private/TrackEditors/PropertyTrackEditors/FloatPropertyTrackEditor.cpp
+++ b/Engine/Source/Editor/MovieSceneTools/Private/TrackEditors/PropertyTrackEditors/FloatPropertyTrackEditor.cpp
@@ -25,6 +25,14 @@ void FFloatPropertyTrackEditor::GenerateKeysFromPropertyChanged( const FProperty
 }
 
 
+void CopyInterpFloatTrack(TSharedRef<ISequencer> Sequencer, UInterpTrackFloatBase* MatineeFloatTrack, UMovieSceneFloatTrack* FloatTrack)
+{
+	if (FMatineeImportTools::CopyInterpFloatTrack(MatineeFloatTrack, FloatTrack))
+	{
+		Sequencer.Get().NotifyMovieSceneDataChanged();
+	}
+}
+
 void FFloatPropertyTrackEditor::BuildTrackContextMenu( FMenuBuilder& MenuBuilder, UMovieSceneTrack* Track )
 {
 	UInterpTrackFloatBase* MatineeFloatTrack = nullptr;
@@ -42,6 +50,6 @@ void FFloatPropertyTrackEditor::BuildTrackContextMenu( FMenuBuilder& MenuBuilder
 		NSLOCTEXT( "Sequencer", "PasteMatineeFloatTrackTooltip", "Pastes keys from a Matinee float track into this track." ),
 		FSlateIcon(),
 		FUIAction(
-			FExecuteAction::CreateStatic( &FMatineeImportTools::CopyInterpFloatTrack, GetSequencer().ToSharedRef(), MatineeFloatTrack, FloatTrack ),
+			FExecuteAction::CreateStatic( &CopyInterpFloatTrack, GetSequencer().ToSharedRef(), MatineeFloatTrack, FloatTrack ),
 			FCanExecuteAction::CreateLambda( [=]()->bool { return MatineeFloatTrack != nullptr && MatineeFloatTrack->GetNumKeys() > 0 && FloatTrack != nullptr; } ) ) );
 }
diff --git a/Engine/Source/Editor/MovieSceneTools/Private/TrackEditors/SkeletalAnimationTrackEditor.cpp b/Engine/Source/Editor/MovieSceneTools/Private/TrackEditors/SkeletalAnimationTrackEditor.cpp
--- a/Engine/Source/Editor/MovieSceneTools/Private/TrackEditors/SkeletalAnimationTrackEditor.cpp
+++ b/Engine/Source/Editor/MovieSceneTools/Private/TrackEditors/SkeletalAnimationTrackEditor.cpp
@@ -30,7 +30,6 @@ namespace SkeletalAnimationEditorConstants
 	const uint32 AnimationTrackHeight = 20;
 }
 
-
 #define LOCTEXT_NAMESPACE "FSkeletalAnimationTrackEditor"
 
 
@@ -166,7 +165,9 @@ void FSkeletalAnimationTrackEditor::AddKey(const FGuid& ObjectGuid)
 
 bool FSkeletalAnimationTrackEditor::HandleAssetAdded(UObject* Asset, const FGuid& TargetObjectGuid)
 {
-	if (Asset->IsA<UAnimSequence>())
+	TSharedPtr<ISequencer> Sequencer = GetSequencer();
+
+	if (Asset->IsA<UAnimSequence>() && Sequencer.IsValid())
 	{
 		UAnimSequence* AnimSequence = Cast<UAnimSequence>(Asset);
 		
@@ -176,10 +177,10 @@ bool FSkeletalAnimationTrackEditor::HandleAssetAdded(UObject* Asset, const FGuid
 
 			if (Skeleton && Skeleton == AnimSequence->GetSkeleton())
 			{
-				TArray<TWeakObjectPtr<UObject>> OutObjects;
+				
 
-				GetSequencer()->GetRuntimeObjects(GetSequencer()->GetFocusedMovieSceneSequenceInstance(), TargetObjectGuid, OutObjects);
-				AnimatablePropertyChanged(FOnKeyProperty::CreateRaw(this, &FSkeletalAnimationTrackEditor::AddKeyInternal, OutObjects, AnimSequence));
+				UObject* Object = Sequencer->FindSpawnedObjectOrTemplate(TargetObjectGuid);
+				AnimatablePropertyChanged(FOnKeyProperty::CreateRaw(this, &FSkeletalAnimationTrackEditor::AddKeyInternal, Object, AnimSequence));
 
 				return true;
 			}
@@ -255,43 +256,37 @@ void FSkeletalAnimationTrackEditor::OnAnimationAssetSelected(const FAssetData& A
 	FSlateApplication::Get().DismissAllMenus();
 
 	UObject* SelectedObject = AssetData.GetAsset();
+	TSharedPtr<ISequencer> Sequencer = GetSequencer();
 
-	if (SelectedObject && SelectedObject->IsA(UAnimSequence::StaticClass()))
+	if (SelectedObject && SelectedObject->IsA(UAnimSequence::StaticClass()) && Sequencer.IsValid())
 	{
 		UAnimSequence* AnimSequence = CastChecked<UAnimSequence>(AssetData.GetAsset());
 
-		TArray<TWeakObjectPtr<UObject>> OutObjects;
-
-		GetSequencer()->GetRuntimeObjects( GetSequencer()->GetFocusedMovieSceneSequenceInstance(), ObjectBinding, OutObjects);
-		AnimatablePropertyChanged( FOnKeyProperty::CreateRaw( this, &FSkeletalAnimationTrackEditor::AddKeyInternal, OutObjects, AnimSequence) );
+		UObject* Object = Sequencer->FindSpawnedObjectOrTemplate(ObjectBinding);
+		AnimatablePropertyChanged( FOnKeyProperty::CreateRaw( this, &FSkeletalAnimationTrackEditor::AddKeyInternal, Object, AnimSequence) );
 	}
 }
 
 
-bool FSkeletalAnimationTrackEditor::AddKeyInternal( float KeyTime, const TArray<TWeakObjectPtr<UObject>> Objects, class UAnimSequence* AnimSequence )
+bool FSkeletalAnimationTrackEditor::AddKeyInternal( float KeyTime, UObject* Object, class UAnimSequence* AnimSequence )
 {
 	bool bHandleCreated = false;
 	bool bTrackCreated = false;
 	bool bTrackModified = false;
 
-	for( int32 ObjectIndex = 0; ObjectIndex < Objects.Num(); ++ObjectIndex )
+	FFindOrCreateHandleResult HandleResult = FindOrCreateHandleToObject( Object );
+	FGuid ObjectHandle = HandleResult.Handle;
+	bHandleCreated |= HandleResult.bWasCreated;
+	if (ObjectHandle.IsValid())
 	{
-		UObject* Object = Objects[ObjectIndex].Get();
+		FFindOrCreateTrackResult TrackResult = FindOrCreateTrackForObject(ObjectHandle, UMovieSceneSkeletalAnimationTrack::StaticClass());
+		UMovieSceneTrack* Track = TrackResult.Track;
+		bTrackCreated |= TrackResult.bWasCreated;
 
-		FFindOrCreateHandleResult HandleResult = FindOrCreateHandleToObject( Object );
-		FGuid ObjectHandle = HandleResult.Handle;
-		bHandleCreated |= HandleResult.bWasCreated;
-		if (ObjectHandle.IsValid())
+		if (ensure(Track))
 		{
-			FFindOrCreateTrackResult TrackResult = FindOrCreateTrackForObject(ObjectHandle, UMovieSceneSkeletalAnimationTrack::StaticClass());
-			UMovieSceneTrack* Track = TrackResult.Track;
-			bTrackCreated |= TrackResult.bWasCreated;
-
-			if (ensure(Track))
-			{
-				Cast<UMovieSceneSkeletalAnimationTrack>(Track)->AddNewAnimation( KeyTime, AnimSequence );
-				bTrackModified = true;
-			}
+			Cast<UMovieSceneSkeletalAnimationTrack>(Track)->AddNewAnimation( KeyTime, AnimSequence );
+			bTrackModified = true;
 		}
 	}
 
@@ -301,44 +296,46 @@ bool FSkeletalAnimationTrackEditor::AddKeyInternal( float KeyTime, const TArray<
 
 USkeleton* FSkeletalAnimationTrackEditor::AcquireSkeletonFromObjectGuid(const FGuid& Guid)
 {
-	TArray<TWeakObjectPtr<UObject>> OutObjects;
-	GetSequencer()->GetRuntimeObjects(GetSequencer()->GetFocusedMovieSceneSequenceInstance(), Guid, OutObjects);
+	TSharedPtr<ISequencer> Sequencer = GetSequencer();
+	UObject* BoundObject = Sequencer.IsValid() ? Sequencer->FindSpawnedObjectOrTemplate(Guid) : nullptr;
 
-	USkeleton* Skeleton = NULL;
-	for (int32 i = 0; i < OutObjects.Num(); ++i)
+	if (AActor* Actor = Cast<AActor>(BoundObject))
 	{
-		UObject* Object = OutObjects[i].Get();
+		TInlineComponentArray<USkeletalMeshComponent*> SkeletalMeshComponents;
+		Actor->GetComponents(SkeletalMeshComponents);
 
-		if (AActor* Actor = Cast<AActor>(Object))
+		for (int32 j = 0; j <SkeletalMeshComponents.Num(); ++j)
 		{
-			TInlineComponentArray<USkeletalMeshComponent*> SkeletalMeshComponents;
-			Actor->GetComponents(SkeletalMeshComponents);
-
-			for (int32 j = 0; j <SkeletalMeshComponents.Num(); ++j)
+			USkeletalMeshComponent* SkeletalMeshComp = SkeletalMeshComponents[j];
+			if (SkeletalMeshComp->SkeletalMesh && SkeletalMeshComp->SkeletalMesh->Skeleton)
 			{
-				USkeletalMeshComponent* SkeletalMeshComp = SkeletalMeshComponents[j];
-				if (SkeletalMeshComp->SkeletalMesh && SkeletalMeshComp->SkeletalMesh->Skeleton)
-				{
-					// @todo Multiple actors, multiple components
-					check(!Skeleton);
-					Skeleton = SkeletalMeshComp->SkeletalMesh->Skeleton;
-				}
+				// @todo Multiple actors, multiple components
+				return SkeletalMeshComp->SkeletalMesh->Skeleton;
 			}
 		}
-		else if(USkeletalMeshComponent* SkeletalMeshComponent = Cast<USkeletalMeshComponent>(Object))
+	}
+	else if(USkeletalMeshComponent* SkeletalMeshComponent = Cast<USkeletalMeshComponent>(BoundObject))
+	{
+		if (SkeletalMeshComponent->SkeletalMesh)
 		{
-			if (SkeletalMeshComponent->SkeletalMesh && SkeletalMeshComponent->SkeletalMesh->Skeleton)
-			{
-				check(!Skeleton);
-				Skeleton = SkeletalMeshComponent->SkeletalMesh->Skeleton;
-			}
+			return SkeletalMeshComponent->SkeletalMesh->Skeleton;
 		}
 	}
 
-	return Skeleton;
+	return nullptr;
 }
 
 
+void CopyInterpAnimControlTrack(TSharedRef<ISequencer> Sequencer, UInterpTrackAnimControl* MatineeAnimControlTrack, UMovieSceneSkeletalAnimationTrack* SkeletalAnimationTrack)
+{
+	float EndPlaybackRange = Sequencer.Get().GetFocusedMovieSceneSequence()->GetMovieScene()->GetPlaybackRange().GetUpperBoundValue();
+
+	if (FMatineeImportTools::CopyInterpAnimControlTrack(MatineeAnimControlTrack, SkeletalAnimationTrack, EndPlaybackRange))
+	{
+		Sequencer.Get().NotifyMovieSceneDataChanged();
+	}
+}
+
 void FSkeletalAnimationTrackEditor::BuildTrackContextMenu( FMenuBuilder& MenuBuilder, UMovieSceneTrack* Track )
 {
 	UInterpTrackAnimControl* MatineeAnimControlTrack = nullptr;
@@ -356,7 +353,7 @@ void FSkeletalAnimationTrackEditor::BuildTrackContextMenu( FMenuBuilder& MenuBui
 		NSLOCTEXT( "Sequencer", "PasteMatineeAnimControlTrackTooltip", "Pastes keys from a Matinee float track into this track." ),
 		FSlateIcon(),
 		FUIAction(
-			FExecuteAction::CreateStatic( &FMatineeImportTools::CopyInterpAnimControlTrack, GetSequencer().ToSharedRef(), MatineeAnimControlTrack, SkeletalAnimationTrack ),
+			FExecuteAction::CreateStatic( &CopyInterpAnimControlTrack, GetSequencer().ToSharedRef(), MatineeAnimControlTrack, SkeletalAnimationTrack ),
 			FCanExecuteAction::CreateLambda( [=]()->bool { return MatineeAnimControlTrack != nullptr && MatineeAnimControlTrack->AnimSeqs.Num() > 0 && SkeletalAnimationTrack != nullptr; } ) ) );
 }
 
diff --git a/Engine/Source/Editor/MovieSceneTools/Private/TrackEditors/SkeletalAnimationTrackEditor.h b/Engine/Source/Editor/MovieSceneTools/Private/TrackEditors/SkeletalAnimationTrackEditor.h
--- a/Engine/Source/Editor/MovieSceneTools/Private/TrackEditors/SkeletalAnimationTrackEditor.h
+++ b/Engine/Source/Editor/MovieSceneTools/Private/TrackEditors/SkeletalAnimationTrackEditor.h
@@ -49,7 +49,7 @@ private:
 	void OnAnimationAssetSelected(const FAssetData& AssetData, FGuid ObjectBinding);
 
 	/** Delegate for AnimatablePropertyChanged in AddKey */
-	bool AddKeyInternal(float KeyTime, const TArray<TWeakObjectPtr<UObject>> Objects, class UAnimSequence* AnimSequence);
+	bool AddKeyInternal(float KeyTime, UObject* Object, class UAnimSequence* AnimSequence);
 
 	/** Gets a skeleton from an object guid in the movie scene */
 	class USkeleton* AcquireSkeletonFromObjectGuid(const FGuid& Guid);
diff --git a/Engine/Source/Editor/MovieSceneTools/Private/TrackEditors/SubTrackEditor.cpp b/Engine/Source/Editor/MovieSceneTools/Private/TrackEditors/SubTrackEditor.cpp
--- a/Engine/Source/Editor/MovieSceneTools/Private/TrackEditors/SubTrackEditor.cpp
+++ b/Engine/Source/Editor/MovieSceneTools/Private/TrackEditors/SubTrackEditor.cpp
@@ -71,7 +71,7 @@ public:
 		{
 			const FString& ActorToRecord = UMovieSceneSubSection::GetActorToRecord();
 
-			ISequenceRecorder& SequenceRecorder = FModuleManager::GetModuleChecked<ISequenceRecorder>("SequenceRecorder");
+			ISequenceRecorder& SequenceRecorder = FModuleManager::LoadModuleChecked<ISequenceRecorder>("SequenceRecorder");
 			if(SequenceRecorder.IsRecording())
 			{
 				if(ActorToRecord.Len() > 0)
@@ -305,7 +305,7 @@ void FSubTrackEditor::BuildAddTrackMenu(FMenuBuilder& MenuBuilder)
 	}
 
 	MenuBuilder.AddMenuEntry(
-		LOCTEXT("AddSubTrack", "Sub Track"),
+		LOCTEXT("AddSubTrack", "Subscenes Track"),
 		LOCTEXT("AddSubTooltip", "Adds a new track that can contain other sequences."),
 		FSlateIcon(FEditorStyle::GetStyleSetName(), "Sequencer.Tracks.Sub"),
 		FUIAction(
@@ -436,6 +436,23 @@ void FSubTrackEditor::HandleAddSubTrackMenuEntryExecute()
 	GetSequencer()->NotifyMovieSceneDataChanged();
 }
 
+/** Helper function - get the first PIE world (or first PIE client world if there is more than one) */
+static UWorld* GetFirstPIEWorld()
+{
+	for (const FWorldContext& Context : GEngine->GetWorldContexts())
+	{
+		if (Context.World()->IsPlayInEditor())
+		{
+			if(Context.World()->GetNetMode() == ENetMode::NM_Standalone ||
+				(Context.World()->GetNetMode() == ENetMode::NM_Client && Context.PIEInstance == 2))
+			{
+				return Context.World();
+			}
+		}
+	}
+
+	return nullptr;
+}
 
 TSharedRef<SWidget> FSubTrackEditor::HandleAddSubSequenceComboButtonGetMenuContent()
 {
@@ -445,11 +462,27 @@ TSharedRef<SWidget> FSubTrackEditor::HandleAddSubSequenceComboButtonGetMenuConte
 	{
 		MenuBuilder.AddMenuEntry(
 			LOCTEXT("RecordNewSequence", "Record New Sequence"), 
-			LOCTEXT("RecordNewSequence_ToolTip", "Record a new level seqeuence into this sub-track from gameplay/simulation etc.\nThis only primes the track for recording. Click the record button to begin recording into this track once primed.\nOnly one sequence can be recorded at a time."), 
+			LOCTEXT("RecordNewSequence_ToolTip", "Record a new level sequence into this sub-track from gameplay/simulation etc.\nThis only primes the track for recording. Click the record button to begin recording into this track once primed.\nOnly one sequence can be recorded at a time."), 
 			FSlateIcon(), 
 			FUIAction(
-				FExecuteAction::CreateSP(this, &FSubTrackEditor::HandleRecordNewSequence),
+				FExecuteAction::CreateSP(this, &FSubTrackEditor::HandleRecordNewSequence, FString()),
 				FCanExecuteAction::CreateSP(this, &FSubTrackEditor::CanRecordNewSequence)));
+
+		if(UWorld* PIEWorld = GetFirstPIEWorld())
+		{
+			APlayerController* Controller = GEngine->GetFirstLocalPlayerController(PIEWorld);
+			if(Controller && Controller->GetPawn())
+			{
+				FString ActorNameToRecord = Controller->GetPawn()->GetFName().ToString();
+				MenuBuilder.AddMenuEntry(
+					LOCTEXT("RecordNewSequenceFromPlayer", "Record New Sequence From Current Player"), 
+					LOCTEXT("RecordNewSequenceFromPlayer_ToolTip", "Record a new level sequence into this sub track using the current player's pawn.\nThis only primes the track for recording. Click the record button to begin recording into this track once primed.\nOnly one sequence can be recorded at a time."), 
+					FSlateIcon(), 
+					FUIAction(
+						FExecuteAction::CreateSP(this, &FSubTrackEditor::HandleRecordNewSequence, ActorNameToRecord),
+						FCanExecuteAction::CreateSP(this, &FSubTrackEditor::CanRecordNewSequence)));
+			}
+		}
 	}
 	MenuBuilder.EndSection();
 
@@ -521,14 +554,14 @@ bool FSubTrackEditor::CanRecordNewSequence() const
 	return !UMovieSceneSubSection::IsSetAsRecording();
 }
 
-void FSubTrackEditor::HandleRecordNewSequence()
+void FSubTrackEditor::HandleRecordNewSequence(FString NameOfActorToRecord)
 {
 	FSlateApplication::Get().DismissAllMenus();
 
-	AnimatablePropertyChanged( FOnKeyProperty::CreateRaw( this, &FSubTrackEditor::HandleRecordNewSequenceInternal) );
+	AnimatablePropertyChanged( FOnKeyProperty::CreateRaw( this, &FSubTrackEditor::HandleRecordNewSequenceInternal, NameOfActorToRecord) );
 }
 
-bool FSubTrackEditor::HandleRecordNewSequenceInternal(float KeyTime)
+bool FSubTrackEditor::HandleRecordNewSequenceInternal(float KeyTime, FString NameOfActorToRecord)
 {
 	auto SubTrack = FindOrCreateMasterTrack<UMovieSceneSubTrack>().Track;
 	UMovieSceneSubSection* Section = SubTrack->AddSequenceToRecord();
@@ -536,6 +569,7 @@ bool FSubTrackEditor::HandleRecordNewSequenceInternal(float KeyTime)
 	// @todo: we could default to the same directory as a parent sequence, or the last sequence recorded. Lots of options!
 	Section->SetTargetSequenceName(GetDefault<USequenceRecorderSettings>()->SequenceName);
 	Section->SetTargetPathToRecordTo(GetDefault<USequenceRecorderSettings>()->SequenceRecordingBasePath.Path);
+	Section->SetNameOfActorToRecord(NameOfActorToRecord);
 
 	return true;
 }
diff --git a/Engine/Source/Editor/MovieSceneTools/Private/TrackEditors/SubTrackEditor.h b/Engine/Source/Editor/MovieSceneTools/Private/TrackEditors/SubTrackEditor.h
--- a/Engine/Source/Editor/MovieSceneTools/Private/TrackEditors/SubTrackEditor.h
+++ b/Engine/Source/Editor/MovieSceneTools/Private/TrackEditors/SubTrackEditor.h
@@ -77,8 +77,8 @@ private:
 	bool CanRecordNewSequence() const;
 
 	/** Handle recording new sequence into a sub track */
-	void HandleRecordNewSequence();
+	void HandleRecordNewSequence(FString NameOfActorToRecord);
 
 	/** Actually handles the adding of the section */
-	bool HandleRecordNewSequenceInternal(float KeyTime);
+	bool HandleRecordNewSequenceInternal(float KeyTime, FString NameOfActorToRecord);
 };
diff --git a/Engine/Source/Editor/MovieSceneTools/Private/TrackEditors/TransformTrackEditor.cpp b/Engine/Source/Editor/MovieSceneTools/Private/TrackEditors/TransformTrackEditor.cpp
--- a/Engine/Source/Editor/MovieSceneTools/Private/TrackEditors/TransformTrackEditor.cpp
+++ b/Engine/Source/Editor/MovieSceneTools/Private/TrackEditors/TransformTrackEditor.cpp
@@ -239,6 +239,14 @@ bool F3DTransformTrackEditor::SupportsType( TSubclassOf<UMovieSceneTrack> Type )
 }
 
 
+void CopyInterpMoveTrack(TSharedRef<ISequencer> Sequencer, UInterpTrackMove* MoveTrack, UMovieScene3DTransformTrack* TransformTrack)
+{
+	if (FMatineeImportTools::CopyInterpMoveTrack(MoveTrack, TransformTrack))
+	{
+		Sequencer.Get().NotifyMovieSceneDataChanged();
+	}
+}
+
 void F3DTransformTrackEditor::BuildTrackContextMenu( FMenuBuilder& MenuBuilder, UMovieSceneTrack* Track )
 {
 	UInterpTrackMove* MoveTrack = nullptr;
@@ -256,7 +264,7 @@ void F3DTransformTrackEditor::BuildTrackContextMenu( FMenuBuilder& MenuBuilder,
 		NSLOCTEXT( "Sequencer", "PasteMatineeTrackTooltip", "Pastes keys from a Matinee move track into this track." ),
 		FSlateIcon(),
 		FUIAction(
-			FExecuteAction::CreateStatic( &FMatineeImportTools::CopyInterpMoveTrack, GetSequencer().ToSharedRef(), MoveTrack, TransformTrack ),
+			FExecuteAction::CreateStatic( &CopyInterpMoveTrack, GetSequencer().ToSharedRef(), MoveTrack, TransformTrack ),
 			FCanExecuteAction::CreateLambda( [=]()->bool { return MoveTrack != nullptr && MoveTrack->GetNumKeys() > 0 && TransformTrack != nullptr; } ) ) );
 
 	auto AnimSubMenuDelegate = [](FMenuBuilder& InMenuBuilder, TSharedRef<ISequencer> InSequencer, UMovieScene3DTransformTrack* InTransformTrack)
@@ -827,7 +835,7 @@ void F3DTransformTrackEditor::ImportAnimSequenceTransforms(const FAssetData& Ass
 				const UClass* ObjectClass = nullptr;
 				if(FMovieSceneSpawnable* Spawnable = MovieScene->FindSpawnable(ObjectBinding))
 				{
-					ObjectClass = Spawnable->GetClass();
+					ObjectClass = Spawnable->GetObjectTemplate()->GetClass();
 				}
 				else if(FMovieScenePossessable* Possessable = MovieScene->FindPossessable(ObjectBinding))
 				{
