Sequencer: Fix non unity builds.

diff --git a/Engine/Plugins/Media/WmfMedia/Source/WmfMedia/Private/Wmf/WmfMediaResolver.cpp b/Engine/Plugins/Media/WmfMedia/Source/WmfMedia/Private/Wmf/WmfMediaResolver.cpp
--- a/Engine/Plugins/Media/WmfMedia/Source/WmfMedia/Private/Wmf/WmfMediaResolver.cpp
+++ b/Engine/Plugins/Media/WmfMedia/Source/WmfMedia/Private/Wmf/WmfMediaResolver.cpp
@@ -1,6 +1,7 @@
 // Copyright 1998-2016 Epic Games, Inc. All Rights Reserved.
 
 #include "WmfMediaPCH.h"
+#include "WmfMediaByteStream.h"
 #include "WmfMediaResolver.h"
 #include "WmfMediaResolveState.h"
 #include "WmfMediaUtils.h"
diff --git a/Engine/Source/Editor/MovieSceneTools/Private/MovieSceneToolsProjectSettingsCustomization.cpp b/Engine/Source/Editor/MovieSceneTools/Private/MovieSceneToolsProjectSettingsCustomization.cpp
--- a/Engine/Source/Editor/MovieSceneTools/Private/MovieSceneToolsProjectSettingsCustomization.cpp
+++ b/Engine/Source/Editor/MovieSceneTools/Private/MovieSceneToolsProjectSettingsCustomization.cpp
@@ -2,6 +2,7 @@
 
 #include "MovieSceneToolsPrivatePCH.h"
 #include "MovieSceneToolsProjectSettingsCustomization.h"
+#include "MovieSceneToolsProjectSettings.h"
 
 #include "PropertyEditorModule.h"
 #include "DetailLayoutBuilder.h"
diff --git a/Engine/Source/Editor/Sequencer/Private/SequencerContextMenus.cpp b/Engine/Source/Editor/Sequencer/Private/SequencerContextMenus.cpp
--- a/Engine/Source/Editor/Sequencer/Private/SequencerContextMenus.cpp
+++ b/Engine/Source/Editor/Sequencer/Private/SequencerContextMenus.cpp
@@ -20,6 +20,7 @@
 #include "IntegralKeyDetailsCustomization.h"
 #include "MovieSceneSubSection.h"
 #include "MovieSceneCinematicShotSection.h"
+#include "Curves/IntegralCurve.h"
 
 #define LOCTEXT_NAMESPACE "SequencerContextMenus"
 
