Copying //UE4/Dev-Sequencer to //UE4/Dev-Main (Source: //UE4/Dev-Sequencer @ 3116759)

diff --git a/Engine/Source/Editor/MovieSceneTools/Private/Sections/ColorPropertySection.cpp b/Engine/Source/Editor/MovieSceneTools/Private/Sections/ColorPropertySection.cpp
--- a/Engine/Source/Editor/MovieSceneTools/Private/Sections/ColorPropertySection.cpp
+++ b/Engine/Source/Editor/MovieSceneTools/Private/Sections/ColorPropertySection.cpp
@@ -8,20 +8,32 @@
 #include "MovieSceneSequence.h"
 #include "FloatCurveKeyArea.h"
 
+const FName SlateColorName("SlateColor");
 
 void FColorPropertySection::GenerateSectionLayout( class ISectionLayoutBuilder& LayoutBuilder ) const
 {
 	UMovieSceneColorSection* ColorSection = Cast<UMovieSceneColorSection>( &SectionObject );
 
-	RedKeyArea = MakeShareable( new FFloatCurveKeyArea( &ColorSection->GetRedCurve(), ColorSection ) ) ;
-	GreenKeyArea = MakeShareable( new FFloatCurveKeyArea( &ColorSection->GetGreenCurve(), ColorSection ) ) ;
-	BlueKeyArea = MakeShareable( new FFloatCurveKeyArea( &ColorSection->GetBlueCurve(), ColorSection ) ) ;
-	AlphaKeyArea = MakeShareable( new FFloatCurveKeyArea( &ColorSection->GetAlphaCurve(), ColorSection ) );
-
-	LayoutBuilder.AddKeyArea( "R", NSLOCTEXT( "FColorPropertySection", "RedArea", "Red" ), RedKeyArea.ToSharedRef() );
-	LayoutBuilder.AddKeyArea( "G", NSLOCTEXT( "FColorPropertySection", "GreenArea", "Green" ), GreenKeyArea.ToSharedRef() );
-	LayoutBuilder.AddKeyArea( "B", NSLOCTEXT( "FColorPropertySection", "BlueArea", "Blue" ), BlueKeyArea.ToSharedRef() );
-	LayoutBuilder.AddKeyArea( "A", NSLOCTEXT( "FColorPropertySection", "OpacityArea", "Opacity" ), AlphaKeyArea.ToSharedRef() );
+	TAttribute<TOptional<float>> RedExternalValue = TAttribute<TOptional<float>>::Create(
+		TAttribute<TOptional<float>>::FGetter::CreateRaw(this, &FColorPropertySection::GetColorRedValue));
+	TSharedRef<FFloatCurveKeyArea> RedKeyArea = MakeShareable( new FFloatCurveKeyArea( &ColorSection->GetRedCurve(), RedExternalValue, ColorSection ) ) ;
+
+	TAttribute<TOptional<float>> GreenExternalValue = TAttribute<TOptional<float>>::Create(
+		TAttribute<TOptional<float>>::FGetter::CreateRaw(this, &FColorPropertySection::GetColorGreenValue));
+	TSharedRef<FFloatCurveKeyArea> GreenKeyArea = MakeShareable( new FFloatCurveKeyArea( &ColorSection->GetGreenCurve(), GreenExternalValue, ColorSection ) ) ;
+	
+	TAttribute<TOptional<float>> BlueExternalValue = TAttribute<TOptional<float>>::Create(
+		TAttribute<TOptional<float>>::FGetter::CreateRaw(this, &FColorPropertySection::GetColorBlueValue));
+	TSharedRef<FFloatCurveKeyArea> BlueKeyArea = MakeShareable( new FFloatCurveKeyArea( &ColorSection->GetBlueCurve(), BlueExternalValue, ColorSection ) ) ;
+	
+	TAttribute<TOptional<float>> AlphaExternalValue = TAttribute<TOptional<float>>::Create(
+		TAttribute<TOptional<float>>::FGetter::CreateRaw(this, &FColorPropertySection::GetColorAlphaValue));
+	TSharedRef<FFloatCurveKeyArea> AlphaKeyArea = MakeShareable( new FFloatCurveKeyArea( &ColorSection->GetAlphaCurve(), AlphaExternalValue, ColorSection ) );
+
+	LayoutBuilder.AddKeyArea( "R", NSLOCTEXT( "FColorPropertySection", "RedArea", "Red" ), RedKeyArea );
+	LayoutBuilder.AddKeyArea( "G", NSLOCTEXT( "FColorPropertySection", "GreenArea", "Green" ), GreenKeyArea );
+	LayoutBuilder.AddKeyArea( "B", NSLOCTEXT( "FColorPropertySection", "BlueArea", "Blue" ), BlueKeyArea );
+	LayoutBuilder.AddKeyArea( "A", NSLOCTEXT( "FColorPropertySection", "OpacityArea", "Opacity" ), AlphaKeyArea );
 }
 
 
@@ -95,9 +107,18 @@ int32 FColorPropertySection::OnPaintSection( FSequencerSectionPainter& Painter )
 
 void FColorPropertySection::ConsolidateColorCurves( TArray< TKeyValuePair<float, FLinearColor> >& OutColorKeys, const UMovieSceneColorSection* Section ) const
 {
-	// Get the default color of the first instance
-	static const FName SlateColorName("SlateColor");
-	FLinearColor DefaultColor = FindSlateColor(SlateColorName);
+	// Get the default color from the current property value, or use black if the current property
+	// value can't be found.
+	FLinearColor DefaultColor;
+	TOptional<FLinearColor> CurrentPropertyValue = GetPropertyValueAsLinearColor();
+	if (CurrentPropertyValue.IsSet())
+	{
+		DefaultColor = CurrentPropertyValue.GetValue();
+	}
+	else
+	{
+		DefaultColor = FLinearColor::Black;
+	}
 
 	// @todo Sequencer Optimize - This could all get cached, instead of recalculating everything every OnPaint
 
@@ -163,84 +184,57 @@ void FColorPropertySection::ConsolidateColorCurves( TArray< TKeyValuePair<float,
 }
 
 
-FLinearColor FColorPropertySection::FindSlateColor(const FName& ColorName) const
+TOptional<FLinearColor> FColorPropertySection::GetPropertyValueAsLinearColor() const
 {
-	const UMovieSceneSequence* FocusedSequence = Sequencer->GetFocusedMovieSceneSequence();
-	const TArray<FMovieSceneBinding>& FocusedBindings = FocusedSequence->GetMovieScene()->GetBindings();
-
-	for (const FMovieSceneBinding& Binding : FocusedBindings)
+	UStructProperty* ColorStructProperty = Cast<UStructProperty>(GetProperty());
+	if (ColorStructProperty != nullptr)
 	{
-		for (const UMovieSceneTrack* BindingTrack : Binding.GetTracks())
+		if (ColorStructProperty->Struct->GetFName() == SlateColorName)
 		{
-			if (BindingTrack != &Track)
+			TOptional<FSlateColor> SlateColor = GetPropertyValue<FSlateColor>();
+			if (SlateColor.IsSet())
 			{
-				continue;
-			}
-
-			UObject* RuntimeObject = Sequencer->GetFocusedMovieSceneSequenceInstance()->FindObject(Binding.GetObjectGuid(), *Sequencer);
-
-			if (RuntimeObject == nullptr)
-			{
-				continue;
-			}
-
-			UProperty* Property = RuntimeObject->GetClass()->FindPropertyByName(Track.GetPropertyName());
-			UStructProperty* ColorStructProp = Cast<UStructProperty>(Property);
-			
-			if ((ColorStructProp == nullptr) || (ColorStructProp->Struct == nullptr))
-			{
-				continue;
+				return TOptional<FLinearColor>(SlateColor.GetValue().GetSpecifiedColor());
 			}
+		}
 
-			if (ColorStructProp->Struct->GetFName() == ColorName)
-			{
-				return (*Property->ContainerPtrToValuePtr<FSlateColor>(RuntimeObject)).GetSpecifiedColor();
-			}
+		if (ColorStructProperty->Struct->GetFName() == NAME_LinearColor)
+		{
+			return GetPropertyValue<FLinearColor>();
+		}
 
-			if (ColorStructProp->Struct->GetFName() == NAME_LinearColor)
+		if (ColorStructProperty->Struct->GetFName() == NAME_Color)
+		{
+			TOptional<FColor> Color = GetPropertyValue<FColor>();
+			if (Color.IsSet())
 			{
-				return *Property->ContainerPtrToValuePtr<FLinearColor>(RuntimeObject);
+				return TOptional<FLinearColor>(Color.GetValue());
 			}
-
-			return *Property->ContainerPtrToValuePtr<FColor>(RuntimeObject);
 		}
 	}
-
-	return FLinearColor(0.0f, 0.0f, 0.0f, 0.0f);
+	return TOptional<FLinearColor>();
 }
 
-
-void FColorPropertySection::SetIntermediateValue( FPropertyChangedParams PropertyChangedParams )
+TOptional<float> FColorPropertySection::GetColorRedValue() const
 {
-	const UStructProperty* StructProp = Cast<const UStructProperty>( PropertyChangedParams.PropertyPath.Last() );
-	FName StructName = StructProp->Struct->GetFName();
-	bool bIsFColor = StructName == NAME_Color;
-	bool bIsFLinearColor = StructName == NAME_LinearColor;
-	bool bIsSlateColor = StructName == FName( "SlateColor" );
-
-	FLinearColor ColorValue;
-
-	if ( bIsFColor )
-	{
-		ColorValue = PropertyChangedParams.GetPropertyValue<FColor>();
-	}
-	else
-	{
-		ColorValue = PropertyChangedParams.GetPropertyValue<FLinearColor>();
-	}
+	TOptional<FLinearColor> CurrentValue = GetPropertyValueAsLinearColor();
+	return CurrentValue.IsSet() ? CurrentValue.GetValue().R : TOptional<float>();
+}
 
-	// @todo - Sequencer - This will cause slightly different values to show up in the key editors for FColors.
-	RedKeyArea->SetIntermediateValue(ColorValue.R);
-	GreenKeyArea->SetIntermediateValue(ColorValue.G);
-	BlueKeyArea->SetIntermediateValue(ColorValue.B);
-	AlphaKeyArea->SetIntermediateValue(ColorValue.A);
+TOptional<float> FColorPropertySection::GetColorGreenValue() const
+{
+	TOptional<FLinearColor> CurrentValue = GetPropertyValueAsLinearColor();
+	return CurrentValue.IsSet() ? CurrentValue.GetValue().G : TOptional<float>();
 }
 
+TOptional<float> FColorPropertySection::GetColorBlueValue() const
+{
+	TOptional<FLinearColor> CurrentValue = GetPropertyValueAsLinearColor();
+	return CurrentValue.IsSet() ? CurrentValue.GetValue().B : TOptional<float>();
+}
 
-void FColorPropertySection::ClearIntermediateValue()
+TOptional<float> FColorPropertySection::GetColorAlphaValue() const
 {
-	RedKeyArea->ClearIntermediateValue();
-	GreenKeyArea->ClearIntermediateValue();
-	BlueKeyArea->ClearIntermediateValue();
-	AlphaKeyArea->ClearIntermediateValue();
+	TOptional<FLinearColor> CurrentValue = GetPropertyValueAsLinearColor();
+	return CurrentValue.IsSet() ? CurrentValue.GetValue().A : TOptional<float>();
 }
diff --git a/Engine/Source/Editor/MovieSceneTools/Private/Sections/ColorPropertySection.h b/Engine/Source/Editor/MovieSceneTools/Private/Sections/ColorPropertySection.h
--- a/Engine/Source/Editor/MovieSceneTools/Private/Sections/ColorPropertySection.h
+++ b/Engine/Source/Editor/MovieSceneTools/Private/Sections/ColorPropertySection.h
@@ -2,16 +2,12 @@
 
 #pragma once
 
-#include "FloatCurveKeyArea.h"
 #include "PropertySection.h"
 
 
-class FFloatCurveKeyArea;
 class ISequencer;
 class UMovieSceneColorSection;
-class UMovieSceneColorTrack;
 class UMovieSceneSection;
-class UMovieSceneTrack;
 
 
 /**
@@ -23,18 +19,19 @@ class FColorPropertySection
 public:
 
 	/**
-	 * Create and initialize a new instance.
-	 *
-	 * @param InSectionObject
-	 * @param SectionName The name of the section.
-	 * @param InSequencer The sequencer that manages the section.
-	 * @param InTrack The track that owns the section.
-	 */
-	FColorPropertySection(UMovieSceneSection& InSectionObject, ISequencer* InSequencer, UMovieSceneTrack& InTrack)
-		: FPropertySection(InSectionObject, InTrack.GetDisplayName())
-		, Sequencer(InSequencer)
-		, Track(*Cast<UMovieSceneColorTrack>(&InTrack))
-	{ }
+	* Creates a new color property section.
+	*
+	* @param InSequencer The sequencer which is controlling this property section.
+	* @param InObjectBinding The object binding for the object which owns the property that this section is animating.
+	* @param InPropertyName The name of the property which is animated by this section.
+	* @param InPropertyPath A string representing the path to the property which is animated by this section.
+	* @param InSectionObject The section object which is being displayed and edited.
+	* @param InDisplayName A display name for the section being displayed and edited.
+	*/
+	FColorPropertySection(ISequencer* InSequencer, FGuid InObjectBinding, FName InPropertyName, const FString& InPropertyPath, UMovieSceneSection& InSectionObject, const FText& InDisplayName)
+		: FPropertySection(InSequencer, InObjectBinding, InPropertyName, InPropertyPath, InSectionObject, InDisplayName)
+	{
+	}
 
 public:
 
@@ -42,27 +39,17 @@ public:
 
 	virtual void GenerateSectionLayout(class ISectionLayoutBuilder& LayoutBuilder) const override;
 	virtual int32 OnPaintSection(FSequencerSectionPainter& Painter) const override;
-	virtual void SetIntermediateValue( FPropertyChangedParams PropertyChangedParams ) override;
-	virtual void ClearIntermediateValue() override;
 
 protected:
 
 	/** Consolidate color curves for all track sections. */
 	void ConsolidateColorCurves(TArray< TKeyValuePair<float, FLinearColor> >& OutColorKeys, const UMovieSceneColorSection* Section) const;
 	
-	/** Find the Slate color of the specified name in the track. */
-	FLinearColor FindSlateColor(const FName& ColorName) const;
+	/** Gets current value of the property being edited as a linear color. */
+	TOptional<FLinearColor> GetPropertyValueAsLinearColor() const;
 
-private:
-
-	mutable TSharedPtr<FFloatCurveKeyArea> RedKeyArea;
-	mutable TSharedPtr<FFloatCurveKeyArea> GreenKeyArea;
-	mutable TSharedPtr<FFloatCurveKeyArea> BlueKeyArea;
-	mutable TSharedPtr<FFloatCurveKeyArea> AlphaKeyArea;
-
-	/** The sequencer that manages the section. */
-	ISequencer* Sequencer;
-
-	/** The track that owns the section. */
-	UMovieSceneColorTrack& Track;
+	TOptional<float> GetColorRedValue() const;
+	TOptional<float> GetColorGreenValue() const;
+	TOptional<float> GetColorBlueValue() const;
+	TOptional<float> GetColorAlphaValue() const;
 };
diff --git a/Engine/Source/Editor/MovieSceneTools/Private/Sections/FloatPropertySection.cpp b/Engine/Source/Editor/MovieSceneTools/Private/Sections/FloatPropertySection.cpp
--- a/Engine/Source/Editor/MovieSceneTools/Private/Sections/FloatPropertySection.cpp
+++ b/Engine/Source/Editor/MovieSceneTools/Private/Sections/FloatPropertySection.cpp
@@ -4,23 +4,20 @@
 #include "FloatCurveKeyArea.h"
 #include "FloatPropertySection.h"
 #include "MovieSceneFloatSection.h"
-
+#include "MovieScenePropertyTrack.h"
 
 void FFloatPropertySection::GenerateSectionLayout(class ISectionLayoutBuilder& LayoutBuilder) const
 {
 	UMovieSceneFloatSection* FloatSection = Cast<UMovieSceneFloatSection>(&SectionObject);
-	KeyArea = MakeShareable(new FFloatCurveKeyArea(&FloatSection->GetFloatCurve(), FloatSection));
-	LayoutBuilder.SetSectionAsKeyArea(KeyArea.ToSharedRef());
-}
-
-
-void FFloatPropertySection::SetIntermediateValue(FPropertyChangedParams PropertyChangedParams)
-{
-	KeyArea->SetIntermediateValue(PropertyChangedParams.GetPropertyValue<float>());
-}
-
-
-void FFloatPropertySection::ClearIntermediateValue()
-{
-	KeyArea->ClearIntermediateValue();
+	TAttribute<TOptional<float>> ExternalValue;
+	if (CanGetPropertyValue())
+	{
+		ExternalValue.Bind(TAttribute<TOptional<float>>::FGetter::CreateLambda([&]
+		{
+			return GetPropertyValue<float>();
+		}));
+	}
+	TSharedRef<FFloatCurveKeyArea> KeyArea = MakeShareable(
+		new FFloatCurveKeyArea(&FloatSection->GetFloatCurve(), ExternalValue, FloatSection));
+	LayoutBuilder.SetSectionAsKeyArea(KeyArea);
 }
diff --git a/Engine/Source/Editor/MovieSceneTools/Private/Sections/FloatPropertySection.h b/Engine/Source/Editor/MovieSceneTools/Private/Sections/FloatPropertySection.h
--- a/Engine/Source/Editor/MovieSceneTools/Private/Sections/FloatPropertySection.h
+++ b/Engine/Source/Editor/MovieSceneTools/Private/Sections/FloatPropertySection.h
@@ -5,7 +5,6 @@
 #include "PropertySection.h"
 
 
-class FFloatCurveKeyArea;
 class UMovieSceneSection;
 
 
@@ -17,19 +16,36 @@ class FFloatPropertySection
 {
 public:
 
-	FFloatPropertySection(UMovieSceneSection& InSectionObject, const FText& SectionName)
-		: FPropertySection(InSectionObject, SectionName)
-	{ }
+	/**
+	* Creates a new float property section for editing non-property float sections.
+	*
+	* @param InSectionObject The section object which is being displayed and edited.
+	* @param InDisplayName A display name for the section being displayed and edited.
+	* @TODO Make another section which is a good base class for non-property float sections.
+	*/
+	FFloatPropertySection(UMovieSceneSection& InSectionObject, const FText& InDisplayName)
+		: FPropertySection(InSectionObject, InDisplayName)
+	{
+	}
+
+	/**
+	* Creates a new float property section.
+	*
+	* @param InSequencer The sequencer which is controlling this property section.
+	* @param InObjectBinding The object binding for the object which owns the property that this section is animating.
+	* @param InPropertyName The name of the property which is animated by this section.
+	* @param InPropertyPath A string representing the path to the property which is animated by this section.
+	* @param InSectionObject The section object which is being displayed and edited.
+	* @param InDisplayName A display name for the section being displayed and edited.
+	*/
+	FFloatPropertySection(ISequencer* InSequencer, FGuid InObjectBinding, FName InPropertyName, const FString& InPropertyPath, UMovieSceneSection& InSectionObject, const FText& InDisplayName)
+		: FPropertySection(InSequencer, InObjectBinding, InPropertyName, InPropertyPath, InSectionObject, InDisplayName)
+	{
+	}
 
 public:
 
 	//~ FPropertySection interface 
 
 	virtual void GenerateSectionLayout(class ISectionLayoutBuilder& LayoutBuilder) const override;
-	virtual void SetIntermediateValue(FPropertyChangedParams PropertyChangedParams) override;
-	virtual void ClearIntermediateValue() override;
-
-private:
-
-	mutable TSharedPtr<FFloatCurveKeyArea> KeyArea;
 };
diff --git a/Engine/Source/Editor/MovieSceneTools/Private/Sections/ParameterSection.h b/Engine/Source/Editor/MovieSceneTools/Private/Sections/ParameterSection.h
--- a/Engine/Source/Editor/MovieSceneTools/Private/Sections/ParameterSection.h
+++ b/Engine/Source/Editor/MovieSceneTools/Private/Sections/ParameterSection.h
@@ -20,13 +20,6 @@ public:
 		: FPropertySection(InSectionObject, SectionName)
 	{ }
 
-public:
-
-	//~ FPropertySection interface
-
-	virtual void SetIntermediateValue(FPropertyChangedParams PropertyChangedParams) override { }
-	virtual void ClearIntermediateValue() override { }
-
 public:
 
 	//~ ISequencerSection interface
diff --git a/Engine/Source/Editor/MovieSceneTools/Private/Sections/StringPropertySection.cpp b/Engine/Source/Editor/MovieSceneTools/Private/Sections/StringPropertySection.cpp
--- a/Engine/Source/Editor/MovieSceneTools/Private/Sections/StringPropertySection.cpp
+++ b/Engine/Source/Editor/MovieSceneTools/Private/Sections/StringPropertySection.cpp
@@ -9,30 +9,10 @@
 void FStringPropertySection::GenerateSectionLayout(class ISectionLayoutBuilder& LayoutBuilder) const
 {
 	UMovieSceneStringSection* StringSection = Cast<UMovieSceneStringSection>(&SectionObject);
-	KeyArea = MakeShareable(new FStringCurveKeyArea(&StringSection->GetStringCurve(), StringSection));
-	LayoutBuilder.SetSectionAsKeyArea(KeyArea.ToSharedRef());
-}
-
 
-void FStringPropertySection::SetIntermediateValue(FPropertyChangedParams PropertyChangedParams)
-{
-	void* CurrentObject = PropertyChangedParams.ObjectsThatChanged[0];
-	void* PropertyValue = nullptr;
-	for (int32 i = 0; i < PropertyChangedParams.PropertyPath.Num(); i++)
-	{
-		CurrentObject = PropertyChangedParams.PropertyPath[i]->ContainerPtrToValuePtr<FString>(CurrentObject, 0);
-	}
+	TAttribute<TOptional<FString>> ExternalValue;
+	ExternalValue.Bind(TAttribute<TOptional<FString>>::FGetter::CreateLambda([&] { return GetPropertyValue<FString>(); }));
 
-	const UStrProperty* StrProperty = Cast<const UStrProperty>( PropertyChangedParams.PropertyPath.Last() );
-	if ( StrProperty )
-	{
-		FString StrPropertyValue = StrProperty->GetPropertyValue(CurrentObject);
-		KeyArea->SetIntermediateValue(StrPropertyValue);
-	}
-}
-
-
-void FStringPropertySection::ClearIntermediateValue()
-{
-	KeyArea->ClearIntermediateValue();
+	TSharedRef<FStringCurveKeyArea> KeyArea = MakeShareable(new FStringCurveKeyArea(&StringSection->GetStringCurve(), ExternalValue, StringSection));
+	LayoutBuilder.SetSectionAsKeyArea(KeyArea);
 }
diff --git a/Engine/Source/Editor/MovieSceneTools/Private/Sections/StringPropertySection.h b/Engine/Source/Editor/MovieSceneTools/Private/Sections/StringPropertySection.h
--- a/Engine/Source/Editor/MovieSceneTools/Private/Sections/StringPropertySection.h
+++ b/Engine/Source/Editor/MovieSceneTools/Private/Sections/StringPropertySection.h
@@ -5,7 +5,6 @@
 #include "PropertySection.h"
 
 
-class FStringCurveKeyArea;
 class UMovieSceneSection;
 
 
@@ -17,19 +16,24 @@ class FStringPropertySection
 {
 public:
 
-	FStringPropertySection(UMovieSceneSection& InSectionObject, const FText& SectionName)
-		: FPropertySection(InSectionObject, SectionName)
-	{ }
+	/**
+	* Creates a new string property section.
+	*
+	* @param InSequencer The sequencer which is controlling this property section.
+	* @param InObjectBinding The object binding for the object which owns the property that this section is animating.
+	* @param InPropertyName The name of the property which is animated by this section.
+	* @param InPropertyPath A string representing the path to the property which is animated by this section.
+	* @param InSectionObject The section object which is being displayed and edited.
+	* @param InDisplayName A display name for the section being displayed and edited.
+	*/
+	FStringPropertySection(ISequencer* InSequencer, FGuid InObjectBinding, FName InPropertyName, const FString& InPropertyPath, UMovieSceneSection& InSectionObject, const FText& InDisplayName)
+		: FPropertySection(InSequencer, InObjectBinding, InPropertyName, InPropertyPath, InSectionObject, InDisplayName)
+	{
+	}
 
 public:
 
 	//~ FPropertySection interface 
 
 	virtual void GenerateSectionLayout(class ISectionLayoutBuilder& LayoutBuilder) const override;
-	virtual void SetIntermediateValue(FPropertyChangedParams PropertyChangedParams) override;
-	virtual void ClearIntermediateValue() override;
-
-private:
-
-	mutable TSharedPtr<FStringCurveKeyArea> KeyArea;
 };
diff --git a/Engine/Source/Editor/MovieSceneTools/Private/Sections/VectorPropertySection.cpp b/Engine/Source/Editor/MovieSceneTools/Private/Sections/VectorPropertySection.cpp
--- a/Engine/Source/Editor/MovieSceneTools/Private/Sections/VectorPropertySection.cpp
+++ b/Engine/Source/Editor/MovieSceneTools/Private/Sections/VectorPropertySection.cpp
@@ -9,76 +9,76 @@
 void FVectorPropertySection::GenerateSectionLayout(class ISectionLayoutBuilder& LayoutBuilder) const
 {
 	UMovieSceneVectorSection* VectorSection = Cast<UMovieSceneVectorSection>(&SectionObject);
-	int32 ChannelsUsed = VectorSection->GetChannelsUsed();
+	ChannelsUsed = VectorSection->GetChannelsUsed();
 	check(ChannelsUsed >= 2 && ChannelsUsed <= 4);
 
-	XKeyArea = MakeShareable(new FFloatCurveKeyArea(&VectorSection->GetCurve(0), VectorSection));
-	LayoutBuilder.AddKeyArea("Vector.X", NSLOCTEXT("FVectorPropertySection", "XArea", "X"), XKeyArea.ToSharedRef());
+	TAttribute<TOptional<float>> XExternalValue = TAttribute<TOptional<float>>::Create(
+		TAttribute<TOptional<float>>::FGetter::CreateRaw(this, &FVectorPropertySection::GetVectorXValue));
+	TSharedRef<FFloatCurveKeyArea> XKeyArea = MakeShareable(new FFloatCurveKeyArea(&VectorSection->GetCurve(0), XExternalValue, VectorSection));
+	LayoutBuilder.AddKeyArea("Vector.X", NSLOCTEXT("FVectorPropertySection", "XArea", "X"), XKeyArea);
 
-	YKeyArea = MakeShareable(new FFloatCurveKeyArea(&VectorSection->GetCurve(1), VectorSection));
-	LayoutBuilder.AddKeyArea("Vector.Y", NSLOCTEXT("FVectorPropertySection", "YArea", "Y"), YKeyArea.ToSharedRef());
+	TAttribute<TOptional<float>> YExternalValue = TAttribute<TOptional<float>>::Create(
+		TAttribute<TOptional<float>>::FGetter::CreateRaw(this, &FVectorPropertySection::GetVectorYValue));
+	TSharedRef<FFloatCurveKeyArea> YKeyArea = MakeShareable(new FFloatCurveKeyArea(&VectorSection->GetCurve(1), YExternalValue, VectorSection));
+	LayoutBuilder.AddKeyArea("Vector.Y", NSLOCTEXT("FVectorPropertySection", "YArea", "Y"), YKeyArea);
 
 	if (ChannelsUsed >= 3)
 	{
-		ZKeyArea = MakeShareable(new FFloatCurveKeyArea(&VectorSection->GetCurve(2), VectorSection));
-		LayoutBuilder.AddKeyArea("Vector.Z", NSLOCTEXT("FVectorPropertySection", "ZArea", "Z"), ZKeyArea.ToSharedRef());
+		TAttribute<TOptional<float>> ZExternalValue = TAttribute<TOptional<float>>::Create(
+			TAttribute<TOptional<float>>::FGetter::CreateRaw(this, &FVectorPropertySection::GetVectorZValue));
+		TSharedRef<FFloatCurveKeyArea> ZKeyArea = MakeShareable(new FFloatCurveKeyArea(&VectorSection->GetCurve(2), ZExternalValue, VectorSection));
+		LayoutBuilder.AddKeyArea("Vector.Z", NSLOCTEXT("FVectorPropertySection", "ZArea", "Z"), ZKeyArea);
 	}
 
 	if (ChannelsUsed >= 4)
 	{
-		WKeyArea = MakeShareable(new FFloatCurveKeyArea(&VectorSection->GetCurve(3), VectorSection));
-		LayoutBuilder.AddKeyArea("Vector.W", NSLOCTEXT("FVectorPropertySection", "WArea", "W"), WKeyArea.ToSharedRef());
+		TAttribute<TOptional<float>> WExternalValue = TAttribute<TOptional<float>>::Create(
+			TAttribute<TOptional<float>>::FGetter::CreateRaw(this, &FVectorPropertySection::GetVectorWValue));
+		TSharedRef<FFloatCurveKeyArea> WKeyArea = MakeShareable(new FFloatCurveKeyArea(&VectorSection->GetCurve(3), WExternalValue, VectorSection));
+		LayoutBuilder.AddKeyArea("Vector.W", NSLOCTEXT("FVectorPropertySection", "WArea", "W"), WKeyArea);
 	}
 }
 
 
-void FVectorPropertySection::SetIntermediateValue(FPropertyChangedParams PropertyChangedParams)
+TOptional<FVector4> FVectorPropertySection::GetPropertyValueAsVector4() const
 {
-	const UStructProperty* StructProp = Cast<const UStructProperty>(PropertyChangedParams.PropertyPath.Last());
-
-	FName StructName = StructProp->Struct->GetFName();
-
-	bool bIsVector2D = StructName == NAME_Vector2D;
-	bool bIsVector = StructName == NAME_Vector;
-	bool bIsVector4 = StructName == NAME_Vector4;
-
-	// Get the vector value from the property
-	if (bIsVector2D)
+	if (ChannelsUsed == 2)
 	{
-		FVector2D Vector =  PropertyChangedParams.GetPropertyValue<FVector2D>();
-		XKeyArea->SetIntermediateValue(Vector.X);
-		YKeyArea->SetIntermediateValue(Vector.Y);
+		TOptional<FVector2D> Vector = GetPropertyValue<FVector2D>();
+		return Vector.IsSet() ? TOptional<FVector4>(FVector4(Vector.GetValue().X, Vector.GetValue().Y, 0, 0)) : TOptional<FVector4>();
 	}
-	else if (bIsVector)
+	else if (ChannelsUsed == 3)
 	{
-		FVector Vector = PropertyChangedParams.GetPropertyValue<FVector>();
-		XKeyArea->SetIntermediateValue(Vector.X);
-		YKeyArea->SetIntermediateValue(Vector.Y);
-		ZKeyArea->SetIntermediateValue(Vector.Z);
+		TOptional<FVector> Vector = GetPropertyValue<FVector>();
+		return Vector.IsSet() ? TOptional<FVector4>(FVector4(Vector.GetValue().X, Vector.GetValue().Y, Vector.GetValue().Z, 0)) : TOptional<FVector4>();
 	}
-	else // if (bIsVector4)
+	else // ChannelsUsed == 4
 	{
-		FVector4 Vector = PropertyChangedParams.GetPropertyValue<FVector4>();
-		XKeyArea->SetIntermediateValue(Vector.X);
-		YKeyArea->SetIntermediateValue(Vector.Y);
-		ZKeyArea->SetIntermediateValue(Vector.Z);
-		WKeyArea->SetIntermediateValue(Vector.W);
+		return GetPropertyValue<FVector4>();
 	}
 }
 
 
-void FVectorPropertySection::ClearIntermediateValue()
+TOptional<float> FVectorPropertySection::GetVectorXValue() const
 {
-	XKeyArea->ClearIntermediateValue();
-	YKeyArea->ClearIntermediateValue();
+	TOptional<FVector4> Vector = GetPropertyValueAsVector4();
+	return Vector.IsSet() ? TOptional<float>(Vector.GetValue().X) : TOptional<float>();
+}
 
-	if (ZKeyArea.IsValid())
-	{
-		ZKeyArea->ClearIntermediateValue();
-	}
+TOptional<float> FVectorPropertySection::GetVectorYValue() const
+{
+	TOptional<FVector4> Vector = GetPropertyValueAsVector4();
+	return Vector.IsSet() ? TOptional<float>(Vector.GetValue().Y) : TOptional<float>();
+}
 
-	if (WKeyArea.IsValid())
-	{
-		WKeyArea->ClearIntermediateValue();
-	}
+TOptional<float> FVectorPropertySection::GetVectorZValue() const
+{
+	TOptional<FVector4> Vector = GetPropertyValueAsVector4();
+	return Vector.IsSet() ? TOptional<float>(Vector.GetValue().Z) : TOptional<float>();
+}
+
+TOptional<float> FVectorPropertySection::GetVectorWValue() const
+{
+	TOptional<FVector4> Vector = GetPropertyValueAsVector4();
+	return Vector.IsSet() ? TOptional<float>(Vector.GetValue().W) : TOptional<float>();
 }
diff --git a/Engine/Source/Editor/MovieSceneTools/Private/Sections/VectorPropertySection.h b/Engine/Source/Editor/MovieSceneTools/Private/Sections/VectorPropertySection.h
--- a/Engine/Source/Editor/MovieSceneTools/Private/Sections/VectorPropertySection.h
+++ b/Engine/Source/Editor/MovieSceneTools/Private/Sections/VectorPropertySection.h
@@ -5,7 +5,6 @@
 #include "PropertySection.h"
 
 
-class FFloatCurveKeyArea;
 class UMovieSceneSection;
 
 
@@ -17,22 +16,37 @@ class FVectorPropertySection
 {
 public:
 
-	FVectorPropertySection(UMovieSceneSection& InSectionObject, const FText& SectionName)
-		: FPropertySection(InSectionObject, SectionName)
-	{ }
+	/**
+	* Creates a new vector property section.
+	*
+	* @param InSequencer The sequencer which is controlling this property section.
+	* @param InObjectBinding The object binding for the object which owns the property that this section is animating.
+	* @param InPropertyName The name of the property which is animated by this section.
+	* @param InPropertyPath A string representing the path to the property which is animated by this section.
+	* @param InSectionObject The section object which is being displayed and edited.
+	* @param InDisplayName A display name for the section being displayed and edited.
+	*/
+	FVectorPropertySection(ISequencer* InSequencer, FGuid InObjectBinding, FName InPropertyName, const FString& InPropertyPath, UMovieSceneSection& InSectionObject, const FText& InDisplayName)
+		: FPropertySection(InSequencer, InObjectBinding, InPropertyName, InPropertyPath, InSectionObject, InDisplayName)
+	{
+	}
 
 public:
 
 	//~ FPropertySectionInterface
 
 	virtual void GenerateSectionLayout(class ISectionLayoutBuilder& LayoutBuilder) const override;
-	virtual void SetIntermediateValue(FPropertyChangedParams PropertyChangedParams) override;
-	virtual void ClearIntermediateValue() override;
 
 private:
 
-	mutable TSharedPtr<FFloatCurveKeyArea> XKeyArea;
-	mutable TSharedPtr<FFloatCurveKeyArea> YKeyArea;
-	mutable TSharedPtr<FFloatCurveKeyArea> ZKeyArea;
-	mutable TSharedPtr<FFloatCurveKeyArea> WKeyArea;
+	TOptional<FVector4> GetPropertyValueAsVector4() const;
+
+	TOptional<float> GetVectorXValue() const;
+	TOptional<float> GetVectorYValue() const;
+	TOptional<float> GetVectorZValue() const;
+	TOptional<float> GetVectorWValue() const;
+
+private:
+
+	mutable int32 ChannelsUsed;
 };
diff --git a/Engine/Source/Editor/MovieSceneTools/Private/Sections/VisibilityPropertySection.cpp b/Engine/Source/Editor/MovieSceneTools/Private/Sections/VisibilityPropertySection.cpp
--- a/Engine/Source/Editor/MovieSceneTools/Private/Sections/VisibilityPropertySection.cpp
+++ b/Engine/Source/Editor/MovieSceneTools/Private/Sections/VisibilityPropertySection.cpp
@@ -3,9 +3,25 @@
 #include "MovieSceneToolsPrivatePCH.h"
 #include "BoolKeyArea.h"
 #include "VisibilityPropertySection.h"
+#include "MovieSceneVisibilitySection.h"
 
 
-void FVisibilityPropertySection::SetIntermediateValue(FPropertyChangedParams PropertyChangedParams)
+void FVisibilityPropertySection::GenerateSectionLayout(class ISectionLayoutBuilder& LayoutBuilder) const
 {
-	KeyArea->SetIntermediateValue(!PropertyChangedParams.GetPropertyValue<bool>());
-}
+	UMovieSceneVisibilitySection* VisibilitySection = Cast<UMovieSceneVisibilitySection>(&SectionObject);
+	TAttribute<TOptional<bool>> ExternalValue;
+	ExternalValue.Bind(TAttribute<TOptional<bool>>::FGetter::CreateLambda([&]
+	{
+		UObject* RuntimeObject = GetRuntimeObjectAndUpdatePropertyBindings();
+
+		// Bool property values are stored in a bit field so using a straight cast of the pointer to get their value does not
+		// work.  Instead use the actual property to get the correct value.
+		const UBoolProperty* BoolProperty = Cast<const UBoolProperty>( GetProperty());
+		uint32* ValuePtr = BoolProperty->ContainerPtrToValuePtr<uint32>(RuntimeObject);
+		bool BoolPropertyValue = !BoolProperty->GetPropertyValue(ValuePtr);
+		void *PropertyValue = &BoolPropertyValue;
+		return TOptional<bool>(*((bool*)PropertyValue));
+	}));
+	TSharedRef<FBoolKeyArea> KeyArea = MakeShareable(new FBoolKeyArea(VisibilitySection->GetCurve(), ExternalValue, VisibilitySection));
+	LayoutBuilder.SetSectionAsKeyArea(KeyArea);
+}
\ No newline at end of file
diff --git a/Engine/Source/Editor/MovieSceneTools/Private/Sections/VisibilityPropertySection.h b/Engine/Source/Editor/MovieSceneTools/Private/Sections/VisibilityPropertySection.h
--- a/Engine/Source/Editor/MovieSceneTools/Private/Sections/VisibilityPropertySection.h
+++ b/Engine/Source/Editor/MovieSceneTools/Private/Sections/VisibilityPropertySection.h
@@ -13,11 +13,23 @@ class FVisibilityPropertySection
 {
 public:
 
-	FVisibilityPropertySection(UMovieSceneSection& InSectionObject, const FText& SectionName, ISequencer* InSequencer)
-		: FBoolPropertySection(InSectionObject, SectionName, InSequencer)
+	/**
+	* Creates a new property section for a boolean hidden property which is displayed inverted as 'visible'.
+	*
+	* @param InSequencer The sequencer which is controlling this property section.
+	* @param InObjectBinding The object binding for the object which owns the property that this section is animating.
+	* @param InPropertyName The name of the property which is animated by this section.
+	* @param InPropertyPath A string representing the path to the property which is animated by this section.
+	* @param InSectionObject The section object which is being displayed and edited.
+	* @param InDisplayName A display name for the section being displayed and edited.
+	*/
+	FVisibilityPropertySection(ISequencer* InSequencer, FGuid InObjectBinding, FName InPropertyName, const FString& InPropertyPath, UMovieSceneSection& InSectionObject, const FText& InDisplayName)
+		: FBoolPropertySection(InSequencer, InObjectBinding, InPropertyName, InPropertyPath, InSectionObject, InDisplayName)
 	{
 		DisplayName = FText::FromString(TEXT("Visible"));
 	}
 
-	virtual void SetIntermediateValue(FPropertyChangedParams PropertyChangedParams) override;
+	//~ FPropertySection interface
+
+	virtual void GenerateSectionLayout(class ISectionLayoutBuilder& LayoutBuilder) const override;
 };
diff --git a/Engine/Source/Editor/MovieSceneTools/Private/TrackEditors/AttachTrackEditor.cpp b/Engine/Source/Editor/MovieSceneTools/Private/TrackEditors/AttachTrackEditor.cpp
--- a/Engine/Source/Editor/MovieSceneTools/Private/TrackEditors/AttachTrackEditor.cpp
+++ b/Engine/Source/Editor/MovieSceneTools/Private/TrackEditors/AttachTrackEditor.cpp
@@ -93,7 +93,7 @@ bool F3DAttachTrackEditor::SupportsType( TSubclassOf<UMovieSceneTrack> Type ) co
 }
 
 
-TSharedRef<ISequencerSection> F3DAttachTrackEditor::MakeSectionInterface( UMovieSceneSection& SectionObject, UMovieSceneTrack& Track )
+TSharedRef<ISequencerSection> F3DAttachTrackEditor::MakeSectionInterface( UMovieSceneSection& SectionObject, UMovieSceneTrack& Track, FGuid ObjectBinding )
 {
 	check( SupportsType( SectionObject.GetOuter()->GetClass() ) );
 
diff --git a/Engine/Source/Editor/MovieSceneTools/Private/TrackEditors/AttachTrackEditor.h b/Engine/Source/Editor/MovieSceneTools/Private/TrackEditors/AttachTrackEditor.h
--- a/Engine/Source/Editor/MovieSceneTools/Private/TrackEditors/AttachTrackEditor.h
+++ b/Engine/Source/Editor/MovieSceneTools/Private/TrackEditors/AttachTrackEditor.h
@@ -38,7 +38,7 @@ public:
 	// ISequencerTrackEditor interface
 
 	virtual void BuildObjectBindingTrackMenu(FMenuBuilder& MenuBuilder, const FGuid& ObjectBinding, const UClass* ObjectClass) override;
-	virtual TSharedRef<ISequencerSection> MakeSectionInterface( UMovieSceneSection& SectionObject, UMovieSceneTrack& Track ) override;
+	virtual TSharedRef<ISequencerSection> MakeSectionInterface( UMovieSceneSection& SectionObject, UMovieSceneTrack& Track, FGuid ObjectBinding ) override;
 	virtual bool SupportsType( TSubclassOf<UMovieSceneTrack> Type ) const override;
 
 	// FTrackEditorActorPicker
diff --git a/Engine/Source/Editor/MovieSceneTools/Private/TrackEditors/AudioTrackEditor.cpp b/Engine/Source/Editor/MovieSceneTools/Private/TrackEditors/AudioTrackEditor.cpp
--- a/Engine/Source/Editor/MovieSceneTools/Private/TrackEditors/AudioTrackEditor.cpp
+++ b/Engine/Source/Editor/MovieSceneTools/Private/TrackEditors/AudioTrackEditor.cpp
@@ -825,7 +825,7 @@ const FSlateBrush* FAudioTrackEditor::GetIconBrush() const
 	return FEditorStyle::GetBrush("Sequencer.Tracks.Audio");
 }
 
-TSharedRef<ISequencerSection> FAudioTrackEditor::MakeSectionInterface( UMovieSceneSection& SectionObject, UMovieSceneTrack& Track )
+TSharedRef<ISequencerSection> FAudioTrackEditor::MakeSectionInterface( UMovieSceneSection& SectionObject, UMovieSceneTrack& Track, FGuid ObjectBinding )
 {
 	check( SupportsType( SectionObject.GetOuter()->GetClass() ) );
 	
diff --git a/Engine/Source/Editor/MovieSceneTools/Private/TrackEditors/AudioTrackEditor.h b/Engine/Source/Editor/MovieSceneTools/Private/TrackEditors/AudioTrackEditor.h
--- a/Engine/Source/Editor/MovieSceneTools/Private/TrackEditors/AudioTrackEditor.h
+++ b/Engine/Source/Editor/MovieSceneTools/Private/TrackEditors/AudioTrackEditor.h
@@ -36,7 +36,7 @@ public:
 	virtual void BuildAddTrackMenu(FMenuBuilder& MenuBuilder) override;
 	virtual TSharedPtr<SWidget> BuildOutlinerEditWidget(const FGuid& ObjectBinding, UMovieSceneTrack* Track, const FBuildEditWidgetParams& Params) override;
 	virtual bool HandleAssetAdded(UObject* Asset, const FGuid& TargetObjectGuid) override;
-	virtual TSharedRef<ISequencerSection> MakeSectionInterface(UMovieSceneSection& SectionObject, UMovieSceneTrack& Track) override;
+	virtual TSharedRef<ISequencerSection> MakeSectionInterface(UMovieSceneSection& SectionObject, UMovieSceneTrack& Track, FGuid ObjectBinding) override;
 	virtual bool SupportsType(TSubclassOf<UMovieSceneTrack> Type) const override;
 	virtual void BuildTrackContextMenu( FMenuBuilder& MenuBuilder, UMovieSceneTrack* Track ) override;
 	virtual const FSlateBrush* GetIconBrush() const override;
diff --git a/Engine/Source/Editor/MovieSceneTools/Private/TrackEditors/CameraAnimTrackEditor.cpp b/Engine/Source/Editor/MovieSceneTools/Private/TrackEditors/CameraAnimTrackEditor.cpp
--- a/Engine/Source/Editor/MovieSceneTools/Private/TrackEditors/CameraAnimTrackEditor.cpp
+++ b/Engine/Source/Editor/MovieSceneTools/Private/TrackEditors/CameraAnimTrackEditor.cpp
@@ -80,7 +80,7 @@ bool FCameraAnimTrackEditor::SupportsType(TSubclassOf<UMovieSceneTrack> Type) co
 }
 
 
-TSharedRef<ISequencerSection> FCameraAnimTrackEditor::MakeSectionInterface( UMovieSceneSection& SectionObject, UMovieSceneTrack& Track )
+TSharedRef<ISequencerSection> FCameraAnimTrackEditor::MakeSectionInterface( UMovieSceneSection& SectionObject, UMovieSceneTrack& Track, FGuid ObjectBinding )
 {
 	check( SupportsType( SectionObject.GetOuter()->GetClass() ) );
 	return MakeShareable(new FCameraAnimSection(SectionObject));
diff --git a/Engine/Source/Editor/MovieSceneTools/Private/TrackEditors/CameraAnimTrackEditor.h b/Engine/Source/Editor/MovieSceneTools/Private/TrackEditors/CameraAnimTrackEditor.h
--- a/Engine/Source/Editor/MovieSceneTools/Private/TrackEditors/CameraAnimTrackEditor.h
+++ b/Engine/Source/Editor/MovieSceneTools/Private/TrackEditors/CameraAnimTrackEditor.h
@@ -36,7 +36,7 @@ public:
 	virtual void AddKey(const FGuid& ObjectGuid) override;
 	virtual void BuildObjectBindingTrackMenu(FMenuBuilder& MenuBuilder, const FGuid& ObjectBinding, const UClass* ObjectClass) override;
 	virtual bool HandleAssetAdded(UObject* Asset, const FGuid& TargetObjectGuid) override;
-	virtual TSharedRef<ISequencerSection> MakeSectionInterface(UMovieSceneSection& SectionObject, UMovieSceneTrack& Track) override;
+	virtual TSharedRef<ISequencerSection> MakeSectionInterface(UMovieSceneSection& SectionObject, UMovieSceneTrack& Track, FGuid ObjectBinding) override;
 	virtual bool SupportsType(TSubclassOf<UMovieSceneTrack> Type) const override;
 	virtual TSharedPtr<SWidget> BuildOutlinerEditWidget(const FGuid& ObjectBinding, UMovieSceneTrack* Track, const FBuildEditWidgetParams& Params) override;
 
diff --git a/Engine/Source/Editor/MovieSceneTools/Private/TrackEditors/CameraCutTrackEditor.cpp b/Engine/Source/Editor/MovieSceneTools/Private/TrackEditors/CameraCutTrackEditor.cpp
--- a/Engine/Source/Editor/MovieSceneTools/Private/TrackEditors/CameraCutTrackEditor.cpp
+++ b/Engine/Source/Editor/MovieSceneTools/Private/TrackEditors/CameraCutTrackEditor.cpp
@@ -100,7 +100,7 @@ TSharedPtr<SWidget> FCameraCutTrackEditor::BuildOutlinerEditWidget(const FGuid&
 }
 
 
-TSharedRef<ISequencerSection> FCameraCutTrackEditor::MakeSectionInterface(UMovieSceneSection& SectionObject, UMovieSceneTrack& Track)
+TSharedRef<ISequencerSection> FCameraCutTrackEditor::MakeSectionInterface(UMovieSceneSection& SectionObject, UMovieSceneTrack& Track, FGuid ObjectBinding)
 {
 	check(SupportsType(SectionObject.GetOuter()->GetClass()));
 
@@ -177,35 +177,6 @@ UMovieSceneCameraCutTrack* FCameraCutTrackEditor::FindOrCreateCameraCutTrack()
 }
 
 
-UFactory* FCameraCutTrackEditor::GetAssetFactoryForNewCameraCut( UClass* SequenceClass )
-{
-	static TWeakObjectPtr<UFactory> CameraCutFactory;
-
-	if( !CameraCutFactory.IsValid() || CameraCutFactory->SupportedClass != SequenceClass )
-	{
-		TArray<UFactory*> Factories;
-		for (TObjectIterator<UClass> It; It; ++It)
-		{
-			UClass* Class = *It;
-			if (Class->IsChildOf(UFactory::StaticClass()) && !Class->HasAnyClassFlags(CLASS_Abstract))
-			{
-				UFactory* Factory = Class->GetDefaultObject<UFactory>();
-				if (Factory->CanCreateNew())
-				{
-					UClass* SupportedClass = Factory->GetSupportedClass();
-					if ( SupportedClass == SequenceClass )
-					{
-						CameraCutFactory = Factory;
-					}
-				}
-			}
-		}
-	}
-
-	return CameraCutFactory.Get();
-}
-
-
 /* FCameraCutTrackEditor callbacks
  *****************************************************************************/
 
diff --git a/Engine/Source/Editor/MovieSceneTools/Private/TrackEditors/CameraShakeTrackEditor.cpp b/Engine/Source/Editor/MovieSceneTools/Private/TrackEditors/CameraShakeTrackEditor.cpp
--- a/Engine/Source/Editor/MovieSceneTools/Private/TrackEditors/CameraShakeTrackEditor.cpp
+++ b/Engine/Source/Editor/MovieSceneTools/Private/TrackEditors/CameraShakeTrackEditor.cpp
@@ -84,7 +84,7 @@ bool FCameraShakeTrackEditor::SupportsType(TSubclassOf<UMovieSceneTrack> Type) c
 }
 
 
-TSharedRef<ISequencerSection> FCameraShakeTrackEditor::MakeSectionInterface(UMovieSceneSection& SectionObject, UMovieSceneTrack& Track)
+TSharedRef<ISequencerSection> FCameraShakeTrackEditor::MakeSectionInterface(UMovieSceneSection& SectionObject, UMovieSceneTrack& Track, FGuid ObjectBinding)
 {
 	check(SupportsType(SectionObject.GetOuter()->GetClass()));
 	return MakeShareable(new FCameraShakeSection(SectionObject));
diff --git a/Engine/Source/Editor/MovieSceneTools/Private/TrackEditors/CinematicShotTrackEditor.cpp b/Engine/Source/Editor/MovieSceneTools/Private/TrackEditors/CinematicShotTrackEditor.cpp
--- a/Engine/Source/Editor/MovieSceneTools/Private/TrackEditors/CinematicShotTrackEditor.cpp
+++ b/Engine/Source/Editor/MovieSceneTools/Private/TrackEditors/CinematicShotTrackEditor.cpp
@@ -120,7 +120,7 @@ TSharedPtr<SWidget> FCinematicShotTrackEditor::BuildOutlinerEditWidget(const FGu
 }
 
 
-TSharedRef<ISequencerSection> FCinematicShotTrackEditor::MakeSectionInterface(UMovieSceneSection& SectionObject, UMovieSceneTrack& Track)
+TSharedRef<ISequencerSection> FCinematicShotTrackEditor::MakeSectionInterface(UMovieSceneSection& SectionObject, UMovieSceneTrack& Track, FGuid ObjectBinding)
 {
 	check(SupportsType(SectionObject.GetOuter()->GetClass()));
 
diff --git a/Engine/Source/Editor/MovieSceneTools/Private/TrackEditors/EventTrackEditor.cpp b/Engine/Source/Editor/MovieSceneTools/Private/TrackEditors/EventTrackEditor.cpp
--- a/Engine/Source/Editor/MovieSceneTools/Private/TrackEditors/EventTrackEditor.cpp
+++ b/Engine/Source/Editor/MovieSceneTools/Private/TrackEditors/EventTrackEditor.cpp
@@ -59,7 +59,7 @@ void FEventTrackEditor::BuildObjectBindingTrackMenu(FMenuBuilder& MenuBuilder, c
 }
 
 
-TSharedRef<ISequencerSection> FEventTrackEditor::MakeSectionInterface(UMovieSceneSection& SectionObject, UMovieSceneTrack& Track)
+TSharedRef<ISequencerSection> FEventTrackEditor::MakeSectionInterface(UMovieSceneSection& SectionObject, UMovieSceneTrack& Track, FGuid ObjectBinding)
 {
 	return MakeShareable(new FEventTrackSection(SectionObject, GetSequencer()));
 }
diff --git a/Engine/Source/Editor/MovieSceneTools/Private/TrackEditors/FadeTrackEditor.cpp b/Engine/Source/Editor/MovieSceneTools/Private/TrackEditors/FadeTrackEditor.cpp
--- a/Engine/Source/Editor/MovieSceneTools/Private/TrackEditors/FadeTrackEditor.cpp
+++ b/Engine/Source/Editor/MovieSceneTools/Private/TrackEditors/FadeTrackEditor.cpp
@@ -66,7 +66,7 @@ public:
 
 			for (auto Time : TimesWithKeys)
 			{
-				float Value = FadeSection->Eval(Time);
+				float Value = FadeSection->Eval(Time, 0.f);
 			
 				FLinearColor Color = FLinearColor::Black;
 				Color.A = Value*255.f;
@@ -112,17 +112,14 @@ FFadeTrackEditor::FFadeTrackEditor(TSharedRef<ISequencer> InSequencer)
 	: FFloatPropertyTrackEditor(InSequencer)
 { }
 
-/* FFPropertyTrackEditor interface
+/* ISequencerTrackEditor interface
  *****************************************************************************/
 
-TSharedRef<FPropertySection> FFadeTrackEditor::MakePropertySectionInterface( UMovieSceneSection& SectionObject, UMovieSceneTrack& Track )
+TSharedRef<ISequencerSection> FFadeTrackEditor::MakeSectionInterface(UMovieSceneSection& SectionObject, UMovieSceneTrack& Track, FGuid ObjectBinding)
 {
 	return MakeShareable(new FFadeSection(SectionObject, Track.GetDisplayName()));
 }
 
-/* ISequencerTrackEditor interface
- *****************************************************************************/
-
 void FFadeTrackEditor::BuildAddTrackMenu(FMenuBuilder& MenuBuilder)
 {
 	UMovieSceneSequence* RootMovieSceneSequence = GetSequencer()->GetRootMovieSceneSequence();
diff --git a/Engine/Source/Editor/MovieSceneTools/Private/TrackEditors/FadeTrackEditor.h b/Engine/Source/Editor/MovieSceneTools/Private/TrackEditors/FadeTrackEditor.h
--- a/Engine/Source/Editor/MovieSceneTools/Private/TrackEditors/FadeTrackEditor.h
+++ b/Engine/Source/Editor/MovieSceneTools/Private/TrackEditors/FadeTrackEditor.h
@@ -36,12 +36,9 @@ public:
 
 public:
 
-	//~ FPropertyTrackEditor interface
-
-	virtual TSharedRef<FPropertySection> MakePropertySectionInterface(UMovieSceneSection& SectionObject, UMovieSceneTrack& Track) override;
-
 	// ISequencerTrackEditor interface
 
+	virtual TSharedRef<ISequencerSection> MakeSectionInterface(UMovieSceneSection& SectionObject, UMovieSceneTrack& Track, FGuid ObjectBinding) override;
 	virtual void BuildAddTrackMenu(FMenuBuilder& MenuBuilder) override;
 	virtual bool SupportsType(TSubclassOf<UMovieSceneTrack> Type) const override;
 	virtual const FSlateBrush* GetIconBrush() const override;
diff --git a/Engine/Source/Editor/MovieSceneTools/Private/TrackEditors/LevelVisibilityTrackEditor.cpp b/Engine/Source/Editor/MovieSceneTools/Private/TrackEditors/LevelVisibilityTrackEditor.cpp
--- a/Engine/Source/Editor/MovieSceneTools/Private/TrackEditors/LevelVisibilityTrackEditor.cpp
+++ b/Engine/Source/Editor/MovieSceneTools/Private/TrackEditors/LevelVisibilityTrackEditor.cpp
@@ -30,7 +30,7 @@ const FSlateBrush* FLevelVisibilityTrackEditor::GetIconBrush() const
 	return FEditorStyle::GetBrush("Sequencer.Tracks.LevelVisibility");
 }
 
-TSharedRef<ISequencerSection> FLevelVisibilityTrackEditor::MakeSectionInterface( UMovieSceneSection& SectionObject, UMovieSceneTrack& Track )
+TSharedRef<ISequencerSection> FLevelVisibilityTrackEditor::MakeSectionInterface( UMovieSceneSection& SectionObject, UMovieSceneTrack& Track, FGuid ObjectBinding )
 {
 	UMovieSceneLevelVisibilitySection* LevelVisibilitySection = Cast<UMovieSceneLevelVisibilitySection>(&SectionObject);
 	check( SupportsType( SectionObject.GetOuter()->GetClass() ) && LevelVisibilitySection != nullptr );
diff --git a/Engine/Source/Editor/MovieSceneTools/Private/TrackEditors/LevelVisibilityTrackEditor.h b/Engine/Source/Editor/MovieSceneTools/Private/TrackEditors/LevelVisibilityTrackEditor.h
--- a/Engine/Source/Editor/MovieSceneTools/Private/TrackEditors/LevelVisibilityTrackEditor.h
+++ b/Engine/Source/Editor/MovieSceneTools/Private/TrackEditors/LevelVisibilityTrackEditor.h
@@ -36,7 +36,7 @@ public:
 
 	// ISequencerTrackEditor interface
 
-	virtual TSharedRef<ISequencerSection> MakeSectionInterface( UMovieSceneSection& SectionObject, UMovieSceneTrack& Track ) override;
+	virtual TSharedRef<ISequencerSection> MakeSectionInterface( UMovieSceneSection& SectionObject, UMovieSceneTrack& Track, FGuid ObjectBinding ) override;
 	virtual bool SupportsType( TSubclassOf<UMovieSceneTrack> Type ) const override;
 	virtual const FSlateBrush* GetIconBrush() const override;
 	virtual void BuildAddTrackMenu( FMenuBuilder& MenuBuilder ) override;
diff --git a/Engine/Source/Editor/MovieSceneTools/Private/TrackEditors/MaterialTrackEditor.cpp b/Engine/Source/Editor/MovieSceneTools/Private/TrackEditors/MaterialTrackEditor.cpp
--- a/Engine/Source/Editor/MovieSceneTools/Private/TrackEditors/MaterialTrackEditor.cpp
+++ b/Engine/Source/Editor/MovieSceneTools/Private/TrackEditors/MaterialTrackEditor.cpp
@@ -17,7 +17,7 @@ FMaterialTrackEditor::FMaterialTrackEditor( TSharedRef<ISequencer> InSequencer )
 }
 
 
-TSharedRef<ISequencerSection> FMaterialTrackEditor::MakeSectionInterface( UMovieSceneSection& SectionObject, UMovieSceneTrack& Track )
+TSharedRef<ISequencerSection> FMaterialTrackEditor::MakeSectionInterface( UMovieSceneSection& SectionObject, UMovieSceneTrack& Track, FGuid ObjectBinding )
 {
 	UMovieSceneParameterSection* ParameterSection = Cast<UMovieSceneParameterSection>(&SectionObject);
 	checkf( ParameterSection != nullptr, TEXT("Unsupported section type.") );
diff --git a/Engine/Source/Editor/MovieSceneTools/Private/TrackEditors/ParticleParameterTrackEditor.cpp b/Engine/Source/Editor/MovieSceneTools/Private/TrackEditors/ParticleParameterTrackEditor.cpp
--- a/Engine/Source/Editor/MovieSceneTools/Private/TrackEditors/ParticleParameterTrackEditor.cpp
+++ b/Engine/Source/Editor/MovieSceneTools/Private/TrackEditors/ParticleParameterTrackEditor.cpp
@@ -24,7 +24,7 @@ TSharedRef<ISequencerTrackEditor> FParticleParameterTrackEditor::CreateTrackEdit
 }
 
 
-TSharedRef<ISequencerSection> FParticleParameterTrackEditor::MakeSectionInterface( UMovieSceneSection& SectionObject, UMovieSceneTrack& Track )
+TSharedRef<ISequencerSection> FParticleParameterTrackEditor::MakeSectionInterface( UMovieSceneSection& SectionObject, UMovieSceneTrack& Track, FGuid ObjectBinding )
 {
 	UMovieSceneParameterSection* ParameterSection = Cast<UMovieSceneParameterSection>(&SectionObject);
 	checkf( ParameterSection != nullptr, TEXT("Unsupported section type.") );
diff --git a/Engine/Source/Editor/MovieSceneTools/Private/TrackEditors/ParticleTrackEditor.cpp b/Engine/Source/Editor/MovieSceneTools/Private/TrackEditors/ParticleTrackEditor.cpp
--- a/Engine/Source/Editor/MovieSceneTools/Private/TrackEditors/ParticleTrackEditor.cpp
+++ b/Engine/Source/Editor/MovieSceneTools/Private/TrackEditors/ParticleTrackEditor.cpp
@@ -277,7 +277,7 @@ bool FParticleTrackEditor::SupportsType( TSubclassOf<UMovieSceneTrack> Type ) co
 }
 
 
-TSharedRef<ISequencerSection> FParticleTrackEditor::MakeSectionInterface( UMovieSceneSection& SectionObject, UMovieSceneTrack& Track )
+TSharedRef<ISequencerSection> FParticleTrackEditor::MakeSectionInterface( UMovieSceneSection& SectionObject, UMovieSceneTrack& Track, FGuid ObjectBinding )
 {
 	check( SupportsType( SectionObject.GetOuter()->GetClass() ) );
 
diff --git a/Engine/Source/Editor/MovieSceneTools/Private/TrackEditors/PathTrackEditor.cpp b/Engine/Source/Editor/MovieSceneTools/Private/TrackEditors/PathTrackEditor.cpp
--- a/Engine/Source/Editor/MovieSceneTools/Private/TrackEditors/PathTrackEditor.cpp
+++ b/Engine/Source/Editor/MovieSceneTools/Private/TrackEditors/PathTrackEditor.cpp
@@ -92,7 +92,7 @@ bool F3DPathTrackEditor::SupportsType( TSubclassOf<UMovieSceneTrack> Type ) cons
 }
 
 
-TSharedRef<ISequencerSection> F3DPathTrackEditor::MakeSectionInterface( UMovieSceneSection& SectionObject, UMovieSceneTrack& Track )
+TSharedRef<ISequencerSection> F3DPathTrackEditor::MakeSectionInterface( UMovieSceneSection& SectionObject, UMovieSceneTrack& Track, FGuid ObjectBinding )
 {
 	check( SupportsType( SectionObject.GetOuter()->GetClass() ) );
 
diff --git a/Engine/Source/Editor/MovieSceneTools/Private/TrackEditors/PropertyTrackEditors/ActorReferencePropertyTrackEditor.cpp b/Engine/Source/Editor/MovieSceneTools/Private/TrackEditors/PropertyTrackEditors/ActorReferencePropertyTrackEditor.cpp
--- a/Engine/Source/Editor/MovieSceneTools/Private/TrackEditors/PropertyTrackEditors/ActorReferencePropertyTrackEditor.cpp
+++ b/Engine/Source/Editor/MovieSceneTools/Private/TrackEditors/PropertyTrackEditors/ActorReferencePropertyTrackEditor.cpp
@@ -12,9 +12,11 @@ TSharedRef<ISequencerTrackEditor> FActorReferencePropertyTrackEditor::CreateTrac
 }
 
 
-TSharedRef<FPropertySection> FActorReferencePropertyTrackEditor::MakePropertySectionInterface( UMovieSceneSection& SectionObject, UMovieSceneTrack& Track )
+TSharedRef<ISequencerSection> FActorReferencePropertyTrackEditor::MakeSectionInterface(UMovieSceneSection& SectionObject, UMovieSceneTrack& Track, FGuid ObjectBinding)
 {
-	return MakeShareable(new FActorReferencePropertySection(SectionObject, Track.GetDisplayName()));
+	UMovieScenePropertyTrack* PropertyTrack = Cast<UMovieScenePropertyTrack>(&Track);
+	checkf(PropertyTrack != nullptr, TEXT("Incompatible track in FActorReferencePropertyTrackEditor"));
+	return MakeShareable(new FActorReferencePropertySection(GetSequencer().Get(), ObjectBinding, PropertyTrack->GetPropertyName(), PropertyTrack->GetPropertyPath(), SectionObject, Track.GetDisplayName()));
 }
 
 
diff --git a/Engine/Source/Editor/MovieSceneTools/Private/TrackEditors/PropertyTrackEditors/ActorReferencePropertyTrackEditor.h b/Engine/Source/Editor/MovieSceneTools/Private/TrackEditors/PropertyTrackEditors/ActorReferencePropertyTrackEditor.h
--- a/Engine/Source/Editor/MovieSceneTools/Private/TrackEditors/PropertyTrackEditors/ActorReferencePropertyTrackEditor.h
+++ b/Engine/Source/Editor/MovieSceneTools/Private/TrackEditors/PropertyTrackEditors/ActorReferencePropertyTrackEditor.h
@@ -35,10 +35,13 @@ public:
 	 */
 	static TSharedRef<ISequencerTrackEditor> CreateTrackEditor(TSharedRef<ISequencer> OwningSequencer);
 
+	//~ ISequencerTrackEditor interface
+
+	virtual TSharedRef<ISequencerSection> MakeSectionInterface(UMovieSceneSection& SectionObject, UMovieSceneTrack& Track, FGuid ObjectBinding) override;
+
 protected:
 
 	//~ FPropertyTrackEditor interface
 
-	virtual TSharedRef<FPropertySection> MakePropertySectionInterface(UMovieSceneSection& SectionObject, UMovieSceneTrack& Track) override;
 	virtual void GenerateKeysFromPropertyChanged(const FPropertyChangedParams& PropertyChangedParams, TArray<FGuid>& NewGeneratedKeys, TArray<FGuid>& DefaultGeneratedKeys) override;
 };
diff --git a/Engine/Source/Editor/MovieSceneTools/Private/TrackEditors/PropertyTrackEditors/BoolPropertyTrackEditor.cpp b/Engine/Source/Editor/MovieSceneTools/Private/TrackEditors/PropertyTrackEditors/BoolPropertyTrackEditor.cpp
--- a/Engine/Source/Editor/MovieSceneTools/Private/TrackEditors/PropertyTrackEditors/BoolPropertyTrackEditor.cpp
+++ b/Engine/Source/Editor/MovieSceneTools/Private/TrackEditors/PropertyTrackEditors/BoolPropertyTrackEditor.cpp
@@ -12,9 +12,11 @@ TSharedRef<ISequencerTrackEditor> FBoolPropertyTrackEditor::CreateTrackEditor(TS
 }
 
 
-TSharedRef<FPropertySection> FBoolPropertyTrackEditor::MakePropertySectionInterface(UMovieSceneSection& SectionObject, UMovieSceneTrack& Track)
+TSharedRef<ISequencerSection> FBoolPropertyTrackEditor::MakeSectionInterface(UMovieSceneSection& SectionObject, UMovieSceneTrack& Track, FGuid ObjectBinding)
 {
-	return MakeShareable(new FBoolPropertySection(SectionObject, Track.GetDisplayName(), GetSequencer().Get()));
+	UMovieScenePropertyTrack* PropertyTrack = Cast<UMovieScenePropertyTrack>(&Track);
+	checkf(PropertyTrack != nullptr, TEXT("Incompatible track in FBoolPropertyTrackEditor"));
+	return MakeShareable(new FBoolPropertySection(GetSequencer().Get(), ObjectBinding, PropertyTrack->GetPropertyName(), PropertyTrack->GetPropertyPath(), SectionObject, Track.GetDisplayName()));
 }
 
 
diff --git a/Engine/Source/Editor/MovieSceneTools/Private/TrackEditors/PropertyTrackEditors/BoolPropertyTrackEditor.h b/Engine/Source/Editor/MovieSceneTools/Private/TrackEditors/PropertyTrackEditors/BoolPropertyTrackEditor.h
--- a/Engine/Source/Editor/MovieSceneTools/Private/TrackEditors/PropertyTrackEditors/BoolPropertyTrackEditor.h
+++ b/Engine/Source/Editor/MovieSceneTools/Private/TrackEditors/PropertyTrackEditors/BoolPropertyTrackEditor.h
@@ -35,10 +35,13 @@ public:
 	 */
 	static TSharedRef<ISequencerTrackEditor> CreateTrackEditor( TSharedRef<ISequencer> OwningSequencer );
 
+	//~ ISequencerTrackEditor interface
+
+	virtual TSharedRef<ISequencerSection> MakeSectionInterface(UMovieSceneSection& SectionObject, UMovieSceneTrack& Track, FGuid ObjectBinding) override;
+
 protected:
 
 	//~ FPropertyTrackEditor interface
 
-	virtual TSharedRef<FPropertySection> MakePropertySectionInterface( UMovieSceneSection& SectionObject, UMovieSceneTrack& Track ) override;
 	virtual void GenerateKeysFromPropertyChanged( const FPropertyChangedParams& PropertyChangedParams, TArray<bool>& NewGeneratedKeys, TArray<bool>& DefaultGeneratedKeys ) override;
 };
diff --git a/Engine/Source/Editor/MovieSceneTools/Private/TrackEditors/PropertyTrackEditors/BytePropertyTrackEditor.cpp b/Engine/Source/Editor/MovieSceneTools/Private/TrackEditors/PropertyTrackEditors/BytePropertyTrackEditor.cpp
--- a/Engine/Source/Editor/MovieSceneTools/Private/TrackEditors/PropertyTrackEditors/BytePropertyTrackEditor.cpp
+++ b/Engine/Source/Editor/MovieSceneTools/Private/TrackEditors/PropertyTrackEditors/BytePropertyTrackEditor.cpp
@@ -13,9 +13,12 @@ TSharedRef<ISequencerTrackEditor> FBytePropertyTrackEditor::CreateTrackEditor( T
 }
 
 
-TSharedRef<FPropertySection> FBytePropertyTrackEditor::MakePropertySectionInterface( UMovieSceneSection& SectionObject, UMovieSceneTrack& Track )
+TSharedRef<ISequencerSection> FBytePropertyTrackEditor::MakeSectionInterface(UMovieSceneSection& SectionObject, UMovieSceneTrack& Track, FGuid ObjectBinding)
 {
-	return MakeShareable( new FBytePropertySection( SectionObject, Track.GetDisplayName(), Cast<UMovieSceneByteTrack>( SectionObject.GetOuter() )->GetEnum() ) );
+	UMovieSceneByteTrack* ByteTrack = Cast<UMovieSceneByteTrack>(&Track);
+	checkf(ByteTrack != nullptr, TEXT("Incompatible track in FBoolPropertyTrackEditor"));
+	return MakeShareable(new FBytePropertySection(GetSequencer().Get(), ObjectBinding, ByteTrack->GetPropertyName(), ByteTrack->GetPropertyPath(),
+		SectionObject, Track.GetDisplayName(), ByteTrack->GetEnum()));
 }
 
 
diff --git a/Engine/Source/Editor/MovieSceneTools/Private/TrackEditors/PropertyTrackEditors/BytePropertyTrackEditor.h b/Engine/Source/Editor/MovieSceneTools/Private/TrackEditors/PropertyTrackEditors/BytePropertyTrackEditor.h
--- a/Engine/Source/Editor/MovieSceneTools/Private/TrackEditors/PropertyTrackEditors/BytePropertyTrackEditor.h
+++ b/Engine/Source/Editor/MovieSceneTools/Private/TrackEditors/PropertyTrackEditors/BytePropertyTrackEditor.h
@@ -40,11 +40,11 @@ public:
 	//~ ISequencerTrackEditor interface
 
 	virtual UMovieSceneTrack* AddTrack(UMovieScene* FocusedMovieScene, const FGuid& ObjectHandle, TSubclassOf<class UMovieSceneTrack> TrackClass, FName UniqueTypeName) override;
+	virtual TSharedRef<ISequencerSection> MakeSectionInterface(UMovieSceneSection& SectionObject, UMovieSceneTrack& Track, FGuid ObjectBinding) override;
 
 protected:
 
 	//~ FPropertyTrackEditor interface
 
-	virtual TSharedRef<FPropertySection> MakePropertySectionInterface(UMovieSceneSection& SectionObject, UMovieSceneTrack& Track) override;
 	virtual void GenerateKeysFromPropertyChanged(const FPropertyChangedParams& PropertyChangedParams, TArray<uint8>& NewGeneratedKeys, TArray<uint8>& DefaultGeneratedKeys) override;
 };
diff --git a/Engine/Source/Editor/MovieSceneTools/Private/TrackEditors/PropertyTrackEditors/ColorPropertyTrackEditor.cpp b/Engine/Source/Editor/MovieSceneTools/Private/TrackEditors/PropertyTrackEditors/ColorPropertyTrackEditor.cpp
--- a/Engine/Source/Editor/MovieSceneTools/Private/TrackEditors/PropertyTrackEditors/ColorPropertyTrackEditor.cpp
+++ b/Engine/Source/Editor/MovieSceneTools/Private/TrackEditors/PropertyTrackEditors/ColorPropertyTrackEditor.cpp
@@ -21,9 +21,11 @@ TSharedRef<ISequencerTrackEditor> FColorPropertyTrackEditor::CreateTrackEditor(T
 }
 
 
-TSharedRef<FPropertySection> FColorPropertyTrackEditor::MakePropertySectionInterface( UMovieSceneSection& SectionObject, UMovieSceneTrack& Track )
+TSharedRef<ISequencerSection> FColorPropertyTrackEditor::MakeSectionInterface(UMovieSceneSection& SectionObject, UMovieSceneTrack& Track, FGuid ObjectBinding)
 {
-	return MakeShareable(new FColorPropertySection(SectionObject, GetSequencer().Get(), Track));
+	UMovieScenePropertyTrack* PropertyTrack = Cast<UMovieScenePropertyTrack>(&Track);
+	checkf(PropertyTrack != nullptr, TEXT("Incompatible track in FColorPropertyTrackEditor"));
+	return MakeShareable(new FColorPropertySection(GetSequencer().Get(), ObjectBinding, PropertyTrack->GetPropertyName(), PropertyTrack->GetPropertyPath(), SectionObject, Track.GetDisplayName()));
 }
 
 
@@ -107,5 +109,8 @@ void FColorPropertyTrackEditor::BuildTrackContextMenu( FMenuBuilder& MenuBuilder
 			FExecuteAction::CreateStatic( &CopyInterpColorTrack, GetSequencer().ToSharedRef(), ColorPropTrack, ColorTrack ) : 
 			FExecuteAction::CreateStatic( &CopyInterpLinearColorTrack, GetSequencer().ToSharedRef(), LinearColorPropTrack, ColorTrack ),			
 			FCanExecuteAction::CreateLambda( [=]()->bool { return ((ColorPropTrack != nullptr && ColorPropTrack->GetNumKeys() > 0) || (LinearColorPropTrack != nullptr && LinearColorPropTrack->GetNumKeys() > 0)) && ColorTrack != nullptr; } ) ) );
+
+	MenuBuilder.AddMenuSeparator();
+	FKeyframeTrackEditor::BuildTrackContextMenu(MenuBuilder, Track);
 }
 
diff --git a/Engine/Source/Editor/MovieSceneTools/Private/TrackEditors/PropertyTrackEditors/ColorPropertyTrackEditor.h b/Engine/Source/Editor/MovieSceneTools/Private/TrackEditors/PropertyTrackEditors/ColorPropertyTrackEditor.h
--- a/Engine/Source/Editor/MovieSceneTools/Private/TrackEditors/PropertyTrackEditors/ColorPropertyTrackEditor.h
+++ b/Engine/Source/Editor/MovieSceneTools/Private/TrackEditors/PropertyTrackEditors/ColorPropertyTrackEditor.h
@@ -35,11 +35,14 @@ public:
 	 */
 	static TSharedRef<ISequencerTrackEditor> CreateTrackEditor(TSharedRef<ISequencer> OwningSequencer);
 
+	//~ ISequencerTrackEditor interface
+
+	virtual TSharedRef<ISequencerSection> MakeSectionInterface(UMovieSceneSection& SectionObject, UMovieSceneTrack& Track, FGuid ObjectBinding) override;
+
 protected:
 
 	//~ FPropertyTrackEditor interface
 
-	virtual TSharedRef<FPropertySection> MakePropertySectionInterface(UMovieSceneSection& SectionObject, UMovieSceneTrack& Track) override;
 	virtual void GenerateKeysFromPropertyChanged(const FPropertyChangedParams& PropertyChangedParams, TArray<FColorKey>& NewGeneratedKeys, TArray<FColorKey>& DefaultGeneratedKeys) override;
 
 protected:
diff --git a/Engine/Source/Editor/MovieSceneTools/Private/TrackEditors/PropertyTrackEditors/FloatPropertyTrackEditor.cpp b/Engine/Source/Editor/MovieSceneTools/Private/TrackEditors/PropertyTrackEditors/FloatPropertyTrackEditor.cpp
--- a/Engine/Source/Editor/MovieSceneTools/Private/TrackEditors/PropertyTrackEditors/FloatPropertyTrackEditor.cpp
+++ b/Engine/Source/Editor/MovieSceneTools/Private/TrackEditors/PropertyTrackEditors/FloatPropertyTrackEditor.cpp
@@ -13,9 +13,17 @@ TSharedRef<ISequencerTrackEditor> FFloatPropertyTrackEditor::CreateTrackEditor(
 }
 
 
-TSharedRef<FPropertySection> FFloatPropertyTrackEditor::MakePropertySectionInterface( UMovieSceneSection& SectionObject, UMovieSceneTrack& Track )
+TSharedRef<ISequencerSection> FFloatPropertyTrackEditor::MakeSectionInterface(UMovieSceneSection& SectionObject, UMovieSceneTrack& Track, FGuid ObjectBinding)
 {
-	return MakeShareable(new FFloatPropertySection(SectionObject, Track.GetDisplayName()));
+	UMovieScenePropertyTrack* PropertyTrack = Cast<UMovieScenePropertyTrack>(&Track);
+	if (PropertyTrack != nullptr)
+	{
+		return MakeShareable(new FFloatPropertySection(GetSequencer().Get(), ObjectBinding, PropertyTrack->GetPropertyName(), PropertyTrack->GetPropertyPath(), SectionObject, Track.GetDisplayName()));
+	}
+	else
+	{
+		return MakeShareable(new FFloatPropertySection(SectionObject, Track.GetDisplayName()));
+	}
 }
 
 
@@ -52,4 +60,7 @@ void FFloatPropertyTrackEditor::BuildTrackContextMenu( FMenuBuilder& MenuBuilder
 		FUIAction(
 			FExecuteAction::CreateStatic( &CopyInterpFloatTrack, GetSequencer().ToSharedRef(), MatineeFloatTrack, FloatTrack ),
 			FCanExecuteAction::CreateLambda( [=]()->bool { return MatineeFloatTrack != nullptr && MatineeFloatTrack->GetNumKeys() > 0 && FloatTrack != nullptr; } ) ) );
+
+	MenuBuilder.AddMenuSeparator();
+	FKeyframeTrackEditor::BuildTrackContextMenu(MenuBuilder, Track);
 }
diff --git a/Engine/Source/Editor/MovieSceneTools/Private/TrackEditors/PropertyTrackEditors/FloatPropertyTrackEditor.h b/Engine/Source/Editor/MovieSceneTools/Private/TrackEditors/PropertyTrackEditors/FloatPropertyTrackEditor.h
--- a/Engine/Source/Editor/MovieSceneTools/Private/TrackEditors/PropertyTrackEditors/FloatPropertyTrackEditor.h
+++ b/Engine/Source/Editor/MovieSceneTools/Private/TrackEditors/PropertyTrackEditors/FloatPropertyTrackEditor.h
@@ -41,10 +41,11 @@ public:
 
 	virtual void BuildTrackContextMenu(FMenuBuilder& MenuBuilder, UMovieSceneTrack* Track) override;
 
-protected:
+	//~ ISequencerTrackEditor interface
+
+	virtual TSharedRef<ISequencerSection> MakeSectionInterface(UMovieSceneSection& SectionObject, UMovieSceneTrack& Track, FGuid ObjectBinding) override;
 
-	//~ FPropertyTrackEditor interface
+protected:
 
-	virtual TSharedRef<FPropertySection> MakePropertySectionInterface(UMovieSceneSection& SectionObject, UMovieSceneTrack& Track) override;
 	virtual void GenerateKeysFromPropertyChanged(const FPropertyChangedParams& PropertyChangedParams, TArray<float>& NewGeneratedKeys, TArray<float>& DefaultGeneratedKeys) override;
 };
diff --git a/Engine/Source/Editor/MovieSceneTools/Private/TrackEditors/PropertyTrackEditors/StringPropertyTrackEditor.cpp b/Engine/Source/Editor/MovieSceneTools/Private/TrackEditors/PropertyTrackEditors/StringPropertyTrackEditor.cpp
--- a/Engine/Source/Editor/MovieSceneTools/Private/TrackEditors/PropertyTrackEditors/StringPropertyTrackEditor.cpp
+++ b/Engine/Source/Editor/MovieSceneTools/Private/TrackEditors/PropertyTrackEditors/StringPropertyTrackEditor.cpp
@@ -11,9 +11,11 @@ TSharedRef<ISequencerTrackEditor> FStringPropertyTrackEditor::CreateTrackEditor(
 }
 
 
-TSharedRef<FPropertySection> FStringPropertyTrackEditor::MakePropertySectionInterface( UMovieSceneSection& SectionObject, UMovieSceneTrack& Track )
+TSharedRef<ISequencerSection> FStringPropertyTrackEditor::MakeSectionInterface(UMovieSceneSection& SectionObject, UMovieSceneTrack& Track, FGuid ObjectBinding)
 {
-	return MakeShareable(new FStringPropertySection(SectionObject, Track.GetDisplayName()));
+	UMovieScenePropertyTrack* PropertyTrack = Cast<UMovieScenePropertyTrack>(&Track);
+	checkf(PropertyTrack != nullptr, TEXT("Incompatible track in FStringPropertyTrackEditor"));
+	return MakeShareable(new FStringPropertySection(GetSequencer().Get(), ObjectBinding, PropertyTrack->GetPropertyName(), PropertyTrack->GetPropertyPath(), SectionObject, Track.GetDisplayName()));
 }
 
 
diff --git a/Engine/Source/Editor/MovieSceneTools/Private/TrackEditors/PropertyTrackEditors/VectorPropertyTrackEditor.cpp b/Engine/Source/Editor/MovieSceneTools/Private/TrackEditors/PropertyTrackEditors/VectorPropertyTrackEditor.cpp
--- a/Engine/Source/Editor/MovieSceneTools/Private/TrackEditors/PropertyTrackEditors/VectorPropertyTrackEditor.cpp
+++ b/Engine/Source/Editor/MovieSceneTools/Private/TrackEditors/PropertyTrackEditors/VectorPropertyTrackEditor.cpp
@@ -17,9 +17,11 @@ TSharedRef<ISequencerTrackEditor> FVectorPropertyTrackEditor::CreateTrackEditor(
 }
 
 
-TSharedRef<FPropertySection> FVectorPropertyTrackEditor::MakePropertySectionInterface( UMovieSceneSection& SectionObject, UMovieSceneTrack& Track )
+TSharedRef<ISequencerSection> FVectorPropertyTrackEditor::MakeSectionInterface(UMovieSceneSection& SectionObject, UMovieSceneTrack& Track, FGuid ObjectBinding)
 {
-	return MakeShareable(new FVectorPropertySection(SectionObject, Track.GetDisplayName()));
+	UMovieScenePropertyTrack* PropertyTrack = Cast<UMovieScenePropertyTrack>(&Track);
+	checkf(PropertyTrack != nullptr, TEXT("Incompatible track in FVectorTrackEditor"));
+	return MakeShareable(new FVectorPropertySection(GetSequencer().Get(), ObjectBinding, PropertyTrack->GetPropertyName(), PropertyTrack->GetPropertyPath(), SectionObject, Track.GetDisplayName()));
 }
 
 void FVectorPropertyTrackEditor::GenerateKeysFromPropertyChanged( const FPropertyChangedParams& PropertyChangedParams, TArray<FVectorKey>& NewGeneratedKeys, TArray<FVectorKey>& DefaultGeneratedKeys )
diff --git a/Engine/Source/Editor/MovieSceneTools/Private/TrackEditors/PropertyTrackEditors/VisibilityPropertyTrackEditor.cpp b/Engine/Source/Editor/MovieSceneTools/Private/TrackEditors/PropertyTrackEditors/VisibilityPropertyTrackEditor.cpp
--- a/Engine/Source/Editor/MovieSceneTools/Private/TrackEditors/PropertyTrackEditors/VisibilityPropertyTrackEditor.cpp
+++ b/Engine/Source/Editor/MovieSceneTools/Private/TrackEditors/PropertyTrackEditors/VisibilityPropertyTrackEditor.cpp
@@ -12,9 +12,11 @@ TSharedRef<ISequencerTrackEditor> FVisibilityPropertyTrackEditor::CreateTrackEdi
 }
 
 
-TSharedRef<FPropertySection> FVisibilityPropertyTrackEditor::MakePropertySectionInterface( UMovieSceneSection& SectionObject, UMovieSceneTrack& Track )
+TSharedRef<ISequencerSection> FVisibilityPropertyTrackEditor::MakeSectionInterface(UMovieSceneSection& SectionObject, UMovieSceneTrack& Track, FGuid ObjectBinding)
 {
-	return MakeShareable(new FVisibilityPropertySection(SectionObject, Track.GetDisplayName(), GetSequencer().Get()));
+	UMovieScenePropertyTrack* PropertyTrack = Cast<UMovieScenePropertyTrack>(&Track);
+	checkf(PropertyTrack != nullptr, TEXT("Incompatible track in FVisibilityPropertyTrackEditor"));
+	return MakeShareable(new FVisibilityPropertySection(GetSequencer().Get(), ObjectBinding, PropertyTrack->GetPropertyName(), PropertyTrack->GetPropertyPath(), SectionObject, Track.GetDisplayName()));
 }
 
 
diff --git a/Engine/Source/Editor/MovieSceneTools/Private/TrackEditors/SkeletalAnimationTrackEditor.cpp b/Engine/Source/Editor/MovieSceneTools/Private/TrackEditors/SkeletalAnimationTrackEditor.cpp
--- a/Engine/Source/Editor/MovieSceneTools/Private/TrackEditors/SkeletalAnimationTrackEditor.cpp
+++ b/Engine/Source/Editor/MovieSceneTools/Private/TrackEditors/SkeletalAnimationTrackEditor.cpp
@@ -151,7 +151,7 @@ bool FSkeletalAnimationTrackEditor::SupportsType( TSubclassOf<UMovieSceneTrack>
 }
 
 
-TSharedRef<ISequencerSection> FSkeletalAnimationTrackEditor::MakeSectionInterface( UMovieSceneSection& SectionObject, UMovieSceneTrack& Track )
+TSharedRef<ISequencerSection> FSkeletalAnimationTrackEditor::MakeSectionInterface( UMovieSceneSection& SectionObject, UMovieSceneTrack& Track, FGuid ObjectBinding )
 {
 	check( SupportsType( SectionObject.GetOuter()->GetClass() ) );
 	
diff --git a/Engine/Source/Editor/MovieSceneTools/Private/TrackEditors/SpawnTrackEditor.cpp b/Engine/Source/Editor/MovieSceneTools/Private/TrackEditors/SpawnTrackEditor.cpp
--- a/Engine/Source/Editor/MovieSceneTools/Private/TrackEditors/SpawnTrackEditor.cpp
+++ b/Engine/Source/Editor/MovieSceneTools/Private/TrackEditors/SpawnTrackEditor.cpp
@@ -4,6 +4,7 @@
 #include "MovieSceneSpawnTrack.h"
 #include "MovieSceneTrack.h"
 #include "SpawnTrackEditor.h"
+#include "BoolPropertySection.h"
 
 
 #define LOCTEXT_NAMESPACE "FSpawnTrackEditor"
@@ -61,6 +62,12 @@ bool FSpawnTrackEditor::SupportsType(TSubclassOf<UMovieSceneTrack> Type) const
 }
 
 
+TSharedRef<ISequencerSection> FSpawnTrackEditor::MakeSectionInterface(UMovieSceneSection& SectionObject, UMovieSceneTrack& Track, FGuid ObjectBinding)
+{
+	return MakeShareable(new FBoolPropertySection(SectionObject, Track.GetDisplayName()));
+}
+
+
 void FSpawnTrackEditor::HandleAddSpawnTrackMenuEntryExecute(FGuid ObjectBinding)
 {
 	FScopedTransaction AddSpawnTrackTransaction(LOCTEXT("AddSpawnTrack_Transaction", "Add Spawn Track"));
diff --git a/Engine/Source/Editor/MovieSceneTools/Private/TrackEditors/SpawnTrackEditor.h b/Engine/Source/Editor/MovieSceneTools/Private/TrackEditors/SpawnTrackEditor.h
--- a/Engine/Source/Editor/MovieSceneTools/Private/TrackEditors/SpawnTrackEditor.h
+++ b/Engine/Source/Editor/MovieSceneTools/Private/TrackEditors/SpawnTrackEditor.h
@@ -13,6 +13,8 @@ class UMovieSceneSpawnTrack;
  * A property track editor for controlling the lifetime of a sapwnable object
  */
 class FSpawnTrackEditor
+	// TODO: Don't derive from bool property track editor here since it doesn't actually edit
+	// bool property tracks.
 	: public FBoolPropertyTrackEditor
 {
 public:
@@ -43,6 +45,7 @@ public:
 	virtual TSharedPtr<SWidget> BuildOutlinerEditWidget(const FGuid& ObjectBinding, UMovieSceneTrack* Track, const FBuildEditWidgetParams& Params) override  { return TSharedPtr<SWidget>(); }
 	virtual bool HandleAssetAdded(UObject* Asset, const FGuid& TargetObjectGuid) override { return false; }
 	virtual bool SupportsType(TSubclassOf<UMovieSceneTrack> Type) const override;
+	virtual TSharedRef<ISequencerSection> MakeSectionInterface(UMovieSceneSection& SectionObject, UMovieSceneTrack& Track, FGuid ObjectBinding) override;
 
 private:
 
diff --git a/Engine/Source/Editor/MovieSceneTools/Private/TrackEditors/SubTrackEditor.cpp b/Engine/Source/Editor/MovieSceneTools/Private/TrackEditors/SubTrackEditor.cpp
--- a/Engine/Source/Editor/MovieSceneTools/Private/TrackEditors/SubTrackEditor.cpp
+++ b/Engine/Source/Editor/MovieSceneTools/Private/TrackEditors/SubTrackEditor.cpp
@@ -349,7 +349,7 @@ TSharedRef<ISequencerTrackEditor> FSubTrackEditor::CreateTrackEditor(TSharedRef<
 }
 
 
-TSharedRef<ISequencerSection> FSubTrackEditor::MakeSectionInterface(UMovieSceneSection& SectionObject, UMovieSceneTrack& Track)
+TSharedRef<ISequencerSection> FSubTrackEditor::MakeSectionInterface(UMovieSceneSection& SectionObject, UMovieSceneTrack& Track, FGuid ObjectBinding)
 {
 	return MakeShareable(new FSubSection(GetSequencer(), SectionObject, Track.GetDisplayName()));
 }
diff --git a/Engine/Source/Editor/MovieSceneTools/Private/TrackEditors/TransformTrackEditor.cpp b/Engine/Source/Editor/MovieSceneTools/Private/TrackEditors/TransformTrackEditor.cpp
--- a/Engine/Source/Editor/MovieSceneTools/Private/TrackEditors/TransformTrackEditor.cpp
+++ b/Engine/Source/Editor/MovieSceneTools/Private/TrackEditors/TransformTrackEditor.cpp
@@ -31,8 +31,10 @@ class F3DTransformSection
 {
 public:
 
-	F3DTransformSection( UMovieSceneSection& InSection )
+	F3DTransformSection(ISequencer* InSequencer, FGuid InObjectBinding, UMovieSceneSection& InSection)
 		: Section( InSection )
+		, Sequencer(InSequencer)
+		, ObjectBinding(InObjectBinding)
 	{ }
 
 	/** Virtual destructor. */
@@ -62,35 +64,62 @@ public:
 
 		UMovieScene3DTransformSection* TransformSection = Cast<UMovieScene3DTransformSection>( &Section );
 
-		TranslationXKeyArea = MakeShareable(new FFloatCurveKeyArea(&TransformSection->GetTranslationCurve(EAxis::X), TransformSection, RedKeyAreaColor));
-		TranslationYKeyArea = MakeShareable(new FFloatCurveKeyArea(&TransformSection->GetTranslationCurve(EAxis::Y), TransformSection, GreenKeyAreaColor));
-		TranslationZKeyArea = MakeShareable(new FFloatCurveKeyArea(&TransformSection->GetTranslationCurve(EAxis::Z), TransformSection, BlueKeyAreaColor));
+		// Translation
+		TAttribute<TOptional<float>> TranslationXExternalValue = TAttribute<TOptional<float>>::Create(
+			TAttribute<TOptional<float>>::FGetter::CreateRaw(this, &F3DTransformSection::GetTranslationValue, EAxis::X));
+		TSharedRef<FFloatCurveKeyArea> TranslationXKeyArea = MakeShareable(new FFloatCurveKeyArea(&TransformSection->GetTranslationCurve(EAxis::X), TranslationXExternalValue, TransformSection, RedKeyAreaColor));
 
-		RotationXKeyArea = MakeShareable(new FFloatCurveKeyArea(&TransformSection->GetRotationCurve(EAxis::X), TransformSection, RedKeyAreaColor));
-		RotationYKeyArea = MakeShareable(new FFloatCurveKeyArea(&TransformSection->GetRotationCurve(EAxis::Y), TransformSection, GreenKeyAreaColor));
-		RotationZKeyArea = MakeShareable(new FFloatCurveKeyArea(&TransformSection->GetRotationCurve(EAxis::Z), TransformSection, BlueKeyAreaColor));
+		TAttribute<TOptional<float>> TranslationYExternalValue = TAttribute<TOptional<float>>::Create(
+			TAttribute<TOptional<float>>::FGetter::CreateRaw(this, &F3DTransformSection::GetTranslationValue, EAxis::Y));
+		TSharedRef<FFloatCurveKeyArea> TranslationYKeyArea = MakeShareable(new FFloatCurveKeyArea(&TransformSection->GetTranslationCurve(EAxis::Y), TranslationYExternalValue, TransformSection, GreenKeyAreaColor));
 
-		ScaleXKeyArea = MakeShareable(new FFloatCurveKeyArea(&TransformSection->GetScaleCurve(EAxis::X), TransformSection, RedKeyAreaColor));
-		ScaleYKeyArea = MakeShareable(new FFloatCurveKeyArea(&TransformSection->GetScaleCurve(EAxis::Y), TransformSection, GreenKeyAreaColor));
-		ScaleZKeyArea = MakeShareable(new FFloatCurveKeyArea(&TransformSection->GetScaleCurve(EAxis::Z), TransformSection, BlueKeyAreaColor));
+		TAttribute<TOptional<float>> TranslationZExternalValue = TAttribute<TOptional<float>>::Create(
+			TAttribute<TOptional<float>>::FGetter::CreateRaw(this, &F3DTransformSection::GetTranslationValue, EAxis::Z));
+		TSharedRef<FFloatCurveKeyArea> TranslationZKeyArea = MakeShareable(new FFloatCurveKeyArea(&TransformSection->GetTranslationCurve(EAxis::Z), TranslationZExternalValue, TransformSection, BlueKeyAreaColor));
+
+		// Rotation
+		TAttribute<TOptional<float>> RotationXExternalValue = TAttribute<TOptional<float>>::Create(
+			TAttribute<TOptional<float>>::FGetter::CreateRaw(this, &F3DTransformSection::GetRotationValue, EAxis::X));
+		TSharedRef<FFloatCurveKeyArea> RotationXKeyArea = MakeShareable(new FFloatCurveKeyArea(&TransformSection->GetRotationCurve(EAxis::X), RotationXExternalValue, TransformSection, RedKeyAreaColor));
+
+		TAttribute<TOptional<float>> RotationYExternalValue = TAttribute<TOptional<float>>::Create(
+			TAttribute<TOptional<float>>::FGetter::CreateRaw(this, &F3DTransformSection::GetRotationValue, EAxis::Y));
+		TSharedRef<FFloatCurveKeyArea> RotationYKeyArea = MakeShareable(new FFloatCurveKeyArea(&TransformSection->GetRotationCurve(EAxis::Y), RotationYExternalValue, TransformSection, GreenKeyAreaColor));
+
+		TAttribute<TOptional<float>> RotationZExternalValue = TAttribute<TOptional<float>>::Create(
+			TAttribute<TOptional<float>>::FGetter::CreateRaw(this, &F3DTransformSection::GetRotationValue, EAxis::Z));
+		TSharedRef<FFloatCurveKeyArea> RotationZKeyArea = MakeShareable(new FFloatCurveKeyArea(&TransformSection->GetRotationCurve(EAxis::Z), RotationZExternalValue, TransformSection, BlueKeyAreaColor));
+
+		// Scale
+		TAttribute<TOptional<float>> ScaleXExternalValue = TAttribute<TOptional<float>>::Create(
+			TAttribute<TOptional<float>>::FGetter::CreateRaw(this, &F3DTransformSection::GetScaleValue, EAxis::X));
+		TSharedRef<FFloatCurveKeyArea> ScaleXKeyArea = MakeShareable(new FFloatCurveKeyArea(&TransformSection->GetScaleCurve(EAxis::X), ScaleXExternalValue, TransformSection, RedKeyAreaColor));
+
+		TAttribute<TOptional<float>> ScaleYExternalValue = TAttribute<TOptional<float>>::Create(
+			TAttribute<TOptional<float>>::FGetter::CreateRaw(this, &F3DTransformSection::GetScaleValue, EAxis::Y));
+		TSharedRef<FFloatCurveKeyArea> ScaleYKeyArea = MakeShareable(new FFloatCurveKeyArea(&TransformSection->GetScaleCurve(EAxis::Y), ScaleYExternalValue, TransformSection, GreenKeyAreaColor));
+
+		TAttribute<TOptional<float>> ScaleZExternalValue = TAttribute<TOptional<float>>::Create(
+			TAttribute<TOptional<float>>::FGetter::CreateRaw(this, &F3DTransformSection::GetScaleValue, EAxis::Z));
+		TSharedRef<FFloatCurveKeyArea> ScaleZKeyArea = MakeShareable(new FFloatCurveKeyArea(&TransformSection->GetScaleCurve(EAxis::Z), ScaleZExternalValue, TransformSection, BlueKeyAreaColor));
 
 		// This generates the tree structure for the transform section
 		LayoutBuilder.PushCategory( "Location", NSLOCTEXT("FTransformSection", "LocationArea", "Location") );
-			LayoutBuilder.AddKeyArea("Location.X", NSLOCTEXT("FTransformSection", "LocXArea", "X"), TranslationXKeyArea.ToSharedRef() );
-			LayoutBuilder.AddKeyArea("Location.Y", NSLOCTEXT("FTransformSection", "LocYArea", "Y"), TranslationYKeyArea.ToSharedRef() );
-			LayoutBuilder.AddKeyArea("Location.Z", NSLOCTEXT("FTransformSection", "LocZArea", "Z"), TranslationZKeyArea.ToSharedRef() );
+			LayoutBuilder.AddKeyArea("Location.X", NSLOCTEXT("FTransformSection", "LocXArea", "X"), TranslationXKeyArea );
+			LayoutBuilder.AddKeyArea("Location.Y", NSLOCTEXT("FTransformSection", "LocYArea", "Y"), TranslationYKeyArea );
+			LayoutBuilder.AddKeyArea("Location.Z", NSLOCTEXT("FTransformSection", "LocZArea", "Z"), TranslationZKeyArea );
 		LayoutBuilder.PopCategory();
 
 		LayoutBuilder.PushCategory( "Rotation", NSLOCTEXT("FTransformSection", "RotationArea", "Rotation") );
-			LayoutBuilder.AddKeyArea("Rotation.X", NSLOCTEXT("FTransformSection", "RotXArea", "X"), RotationXKeyArea.ToSharedRef() );
-			LayoutBuilder.AddKeyArea("Rotation.Y", NSLOCTEXT("FTransformSection", "RotYArea", "Y"), RotationYKeyArea.ToSharedRef() );
-			LayoutBuilder.AddKeyArea("Rotation.Z", NSLOCTEXT("FTransformSection", "RotZArea", "Z"), RotationZKeyArea.ToSharedRef() );
+			LayoutBuilder.AddKeyArea("Rotation.X", NSLOCTEXT("FTransformSection", "RotXArea", "X"), RotationXKeyArea );
+			LayoutBuilder.AddKeyArea("Rotation.Y", NSLOCTEXT("FTransformSection", "RotYArea", "Y"), RotationYKeyArea );
+			LayoutBuilder.AddKeyArea("Rotation.Z", NSLOCTEXT("FTransformSection", "RotZArea", "Z"), RotationZKeyArea );
 		LayoutBuilder.PopCategory();
 
 		LayoutBuilder.PushCategory( "Scale", NSLOCTEXT("FTransformSection", "ScaleArea", "Scale") );
-			LayoutBuilder.AddKeyArea("Scale.X", NSLOCTEXT("FTransformSection", "ScaleXArea", "X"), ScaleXKeyArea.ToSharedRef() );
-			LayoutBuilder.AddKeyArea("Scale.Y", NSLOCTEXT("FTransformSection", "ScaleYArea", "Y"), ScaleYKeyArea.ToSharedRef() );
-			LayoutBuilder.AddKeyArea("Scale.Z", NSLOCTEXT("FTransformSection", "ScaleZArea", "Z"), ScaleZKeyArea.ToSharedRef() );
+			LayoutBuilder.AddKeyArea("Scale.X", NSLOCTEXT("FTransformSection", "ScaleXArea", "X"), ScaleXKeyArea );
+			LayoutBuilder.AddKeyArea("Scale.Y", NSLOCTEXT("FTransformSection", "ScaleYArea", "Y"), ScaleYKeyArea );
+			LayoutBuilder.AddKeyArea("Scale.Z", NSLOCTEXT("FTransformSection", "ScaleZArea", "Z"), ScaleZKeyArea );
 		LayoutBuilder.PopCategory();
 	}
 
@@ -99,37 +128,78 @@ public:
 		return InPainter.PaintSectionBackground();
 	}
 
-	void SetIntermediateValue(UMovieSceneTrack* Track, FTransformData IntermediateData )
+
+
+private:
+
+	TOptional<float> GetTranslationValue(EAxis::Type Axis) const
 	{
-		if ( Section.GetOuter() == Track )
+		USceneComponent* SceneComponent = GetSceneComponent();
+		if (SceneComponent != nullptr)
 		{
-			TranslationXKeyArea->SetIntermediateValue( IntermediateData.Translation.X );
-			TranslationYKeyArea->SetIntermediateValue( IntermediateData.Translation.Y );
-			TranslationZKeyArea->SetIntermediateValue( IntermediateData.Translation.Z );
+			switch (Axis)
+			{
+			case EAxis::X:
+				return TOptional<float>(SceneComponent->GetRelativeTransform().GetTranslation().X);
+			case EAxis::Y:
+				return TOptional<float>(SceneComponent->GetRelativeTransform().GetTranslation().Y);
+			case EAxis::Z:
+				return TOptional<float>(SceneComponent->GetRelativeTransform().GetTranslation().Z);
+			}
+		}
+		return TOptional<float>();
+	}
 
-			RotationXKeyArea->SetIntermediateValue( IntermediateData.Rotation.Euler().X );
-			RotationYKeyArea->SetIntermediateValue( IntermediateData.Rotation.Euler().Y );
-			RotationZKeyArea->SetIntermediateValue( IntermediateData.Rotation.Euler().Z );
+	TOptional<float> GetRotationValue(EAxis::Type Axis) const
+	{
+		USceneComponent* SceneComponent = GetSceneComponent();
+		if (SceneComponent != nullptr)
+		{
+			FRotator Rotator = SceneComponent->RelativeRotation;
 
-			ScaleXKeyArea->SetIntermediateValue( IntermediateData.Scale.X );
-			ScaleYKeyArea->SetIntermediateValue( IntermediateData.Scale.Y );
-			ScaleZKeyArea->SetIntermediateValue( IntermediateData.Scale.Z );
+			switch (Axis)
+			{
+			case EAxis::X:
+				return TOptional<float>(Rotator.Roll);
+			case EAxis::Y:
+				return TOptional<float>(Rotator.Pitch);
+			case EAxis::Z:
+				return TOptional<float>(Rotator.Yaw);
+			}
 		}
+		return TOptional<float>();
 	}
 
-	void ClearIntermediateValue()
+	TOptional<float> GetScaleValue(EAxis::Type Axis) const
 	{
-		TranslationXKeyArea->ClearIntermediateValue();
-		TranslationYKeyArea->ClearIntermediateValue();
-		TranslationZKeyArea->ClearIntermediateValue();
-
-		RotationXKeyArea->ClearIntermediateValue();
-		RotationYKeyArea->ClearIntermediateValue();
-		RotationZKeyArea->ClearIntermediateValue();
+		USceneComponent* SceneComponent = GetSceneComponent();
+		if (SceneComponent != nullptr)
+		{
+			switch (Axis)
+			{
+			case EAxis::X:
+				return TOptional<float>(SceneComponent->GetRelativeTransform().GetScale3D().X);
+			case EAxis::Y:
+				return TOptional<float>(SceneComponent->GetRelativeTransform().GetScale3D().Y);
+			case EAxis::Z:
+				return TOptional<float>(SceneComponent->GetRelativeTransform().GetScale3D().Z);
+			}
+		}
+		return TOptional<float>();
+	}
 
-		ScaleXKeyArea->ClearIntermediateValue();
-		ScaleYKeyArea->ClearIntermediateValue();
-		ScaleZKeyArea->ClearIntermediateValue();
+	USceneComponent* GetSceneComponent() const
+	{
+		if (Sequencer != nullptr)
+		{
+			TArray<TWeakObjectPtr<UObject>> RuntimeObjects;
+			Sequencer->GetRuntimeObjects(Sequencer->GetFocusedMovieSceneSequenceInstance(), ObjectBinding, RuntimeObjects);
+			if (RuntimeObjects.Num() == 1)
+			{
+				return MovieSceneHelpers::SceneComponentFromRuntimeObject(RuntimeObjects[0].Get());
+			}
+		}
+		return nullptr;
 	}
 
 private:
@@ -137,17 +207,9 @@ private:
 	/** The section we are visualizing */
 	UMovieSceneSection& Section;
 
-	mutable TSharedPtr<FFloatCurveKeyArea> TranslationXKeyArea;
-	mutable TSharedPtr<FFloatCurveKeyArea> TranslationYKeyArea;
-	mutable TSharedPtr<FFloatCurveKeyArea> TranslationZKeyArea;
-
-	mutable TSharedPtr<FFloatCurveKeyArea> RotationXKeyArea;
-	mutable TSharedPtr<FFloatCurveKeyArea> RotationYKeyArea;
-	mutable TSharedPtr<FFloatCurveKeyArea> RotationZKeyArea;
+	ISequencer* Sequencer;
 
-	mutable TSharedPtr<FFloatCurveKeyArea> ScaleXKeyArea;
-	mutable TSharedPtr<FFloatCurveKeyArea> ScaleYKeyArea;
-	mutable TSharedPtr<FFloatCurveKeyArea> ScaleZKeyArea;
+	FGuid ObjectBinding;
 };
 
 
@@ -326,16 +388,16 @@ void F3DTransformTrackEditor::BuildTrackContextMenu( FMenuBuilder& MenuBuilder,
 		NSLOCTEXT( "Sequencer", "ImportTransformsTooltip", "Import transform keys from an animation sequence's root motion." ),
 		FNewMenuDelegate::CreateLambda(AnimSubMenuDelegate, GetSequencer().ToSharedRef(), TransformTrack), 
 		bInOpenSubMenuOnClick);
+
+	MenuBuilder.AddMenuSeparator();
+	FKeyframeTrackEditor::BuildTrackContextMenu(MenuBuilder, Track);
 }
 
 
-TSharedRef<ISequencerSection> F3DTransformTrackEditor::MakeSectionInterface( UMovieSceneSection& SectionObject, UMovieSceneTrack& Track )
+TSharedRef<ISequencerSection> F3DTransformTrackEditor::MakeSectionInterface(UMovieSceneSection& SectionObject, UMovieSceneTrack& Track, FGuid ObjectBinding)
 {
-	check( SupportsType( SectionObject.GetOuter()->GetClass() ) );
-	TSharedRef<F3DTransformSection> TransformSection = MakeShareable( new F3DTransformSection( SectionObject ) );
-	OnSetIntermediateValueFromTransformChange.AddSP(TransformSection, &F3DTransformSection::SetIntermediateValue );
-	GetSequencer()->OnGlobalTimeChanged().AddSP(TransformSection, &F3DTransformSection::ClearIntermediateValue );
-	return TransformSection;
+	check(SupportsType(SectionObject.GetOuter()->GetClass()));
+	return MakeShareable(new F3DTransformSection(GetSequencer().Get(), ObjectBinding, SectionObject));
 }
 
 
@@ -403,25 +465,9 @@ void F3DTransformTrackEditor::OnTransformChanged( UObject& InObject )
 		// Build new transform data
 		FTransformData NewTransformData( SceneComponentThatChanged );
 
-		FKeyParams KeyParams;
-		KeyParams.bCreateKeyIfUnchanged = GetSequencer()->GetKeyAllEnabled();
-		KeyParams.bCreateKeyOnlyWhenAutoKeying = true;
-		if ( GetSequencer()->GetAutoKeyMode() == EAutoKeyMode::KeyAll )
-		{
-			KeyParams.bCreateHandleIfMissing = true;
-			KeyParams.bCreateTrackIfMissing = true;
-			KeyParams.bCreateKeyIfEmpty = true;
-		}
-		else
-		{
-			KeyParams.bCreateHandleIfMissing = false;
-			KeyParams.bCreateTrackIfMissing = false;
-			KeyParams.bCreateKeyIfEmpty = false;
-		}
-
 		const bool bUnwindRotation = GetSequencer()->IsRecordingLive();
 
-		AddTransformKeys(Actor, ExistingTransform, NewTransformData, EKey3DTransformChannel::All, bUnwindRotation, KeyParams);
+		AddTransformKeys(Actor, ExistingTransform, NewTransformData, EKey3DTransformChannel::All, bUnwindRotation, ESequencerKeyMode::AutoKey);
 	}
 }
 
@@ -439,19 +485,13 @@ void F3DTransformTrackEditor::OnAddTransformKeysForSelectedObjects( EKey3DTransf
 		}
 	}
 
-	FKeyParams KeyParams;
-	KeyParams.bCreateHandleIfMissing = true;
-	KeyParams.bCreateTrackIfMissing = true;
-	KeyParams.bCreateKeyIfUnchanged = true;
-	KeyParams.bCreateKeyIfEmpty = true;
-
 	USelection* CurrentSelection = GEditor->GetSelectedActors();
 	TArray<UObject*> SelectedActors;
 	CurrentSelection->GetSelectedObjects( AActor::StaticClass(), SelectedActors );
 
 	for (TArray<UObject*>::TIterator It(SelectedActors); It; ++It)
 	{
-		AddTransformKeysForObject(*It, Channel, KeyParams);
+		AddTransformKeysForObject(*It, Channel, ESequencerKeyMode::ManualKeyForced);
 	}
 }
 
@@ -529,16 +569,12 @@ void F3DTransformTrackEditor::BuildObjectBindingTrackMenu(FMenuBuilder& MenuBuil
 {
 	if (ObjectClass->IsChildOf(AActor::StaticClass()) || ObjectClass->IsChildOf(USceneComponent::StaticClass()))
 	{
-		FKeyParams KeyParams;
-		KeyParams.bCreateTrackIfMissing = true;
-		KeyParams.bCreateKeyIfUnchanged = true;
-
 		MenuBuilder.AddMenuEntry(
 			NSLOCTEXT("Sequencer", "AddTransform", "Transform"),
 			NSLOCTEXT("Sequencer", "AddPTransformTooltip", "Adds a transform track."),
 			FSlateIcon(),
 			FUIAction( 
-				FExecuteAction::CreateSP( this, &F3DTransformTrackEditor::AddTransformKeysForHandle, ObjectBinding, EKey3DTransformChannel::All, KeyParams ),
+				FExecuteAction::CreateSP( this, &F3DTransformTrackEditor::AddTransformKeysForHandle, ObjectBinding, EKey3DTransformChannel::All, ESequencerKeyMode::ManualKey ),
 				FCanExecuteAction::CreateSP( this, &F3DTransformTrackEditor::CanAddTransformTrackForActorHandle, ObjectBinding ) ) );
 	}
 }
@@ -778,19 +814,19 @@ void F3DTransformTrackEditor::GetTransformKeys( const FTransformData& LastTransf
 	GetKeysForVector( bLastVectorIsValid, LastTransform.Scale, CurrentTransform.Scale, EKey3DTransformChannel::Scale, ChannelsToKey, bUnwindRotation, OutNewKeys, OutDefaultKeys );
 }
 
-void F3DTransformTrackEditor::AddTransformKeysForHandle( FGuid ObjectHandle, EKey3DTransformChannel::Type ChannelToKey, FKeyParams KeyParams )
+void F3DTransformTrackEditor::AddTransformKeysForHandle( FGuid ObjectHandle, EKey3DTransformChannel::Type ChannelToKey, ESequencerKeyMode KeyMode )
 {
 	TArray<TWeakObjectPtr<UObject>> OutObjects;
 	GetSequencer()->GetRuntimeObjects( GetSequencer()->GetFocusedMovieSceneSequenceInstance(), ObjectHandle, OutObjects );
 
 	for ( TWeakObjectPtr<UObject> Object : OutObjects )
 	{
-		AddTransformKeysForObject(Object.Get(), ChannelToKey, KeyParams);
+		AddTransformKeysForObject(Object.Get(), ChannelToKey, KeyMode);
 	}
 }
 
 
-void F3DTransformTrackEditor::AddTransformKeysForObject( UObject* Object, EKey3DTransformChannel::Type ChannelToKey, FKeyParams KeyParams )
+void F3DTransformTrackEditor::AddTransformKeysForObject( UObject* Object, EKey3DTransformChannel::Type ChannelToKey, ESequencerKeyMode KeyMode )
 {
 	AActor* Actor = nullptr;
 	USceneComponent* SceneComponent = nullptr;
@@ -800,27 +836,27 @@ void F3DTransformTrackEditor::AddTransformKeysForObject( UObject* Object, EKey3D
 		FTransformData CurrentTransform( SceneComponent );
 		if(Object->GetClass()->IsChildOf(AActor::StaticClass()))
 		{
-			AddTransformKeys( Actor, FTransformData(), CurrentTransform, ChannelToKey, false, KeyParams );
+			AddTransformKeys( Actor, FTransformData(), CurrentTransform, ChannelToKey, false, KeyMode );
 		}
 		else if(Object->GetClass()->IsChildOf(USceneComponent::StaticClass()))
 		{
-			AddTransformKeys( SceneComponent, FTransformData(), CurrentTransform, ChannelToKey, false, KeyParams );
+			AddTransformKeys( SceneComponent, FTransformData(), CurrentTransform, ChannelToKey, false, KeyMode );
 		}
 	}
 }
 
 
-void F3DTransformTrackEditor::AddTransformKeys( UObject* ObjectToKey, const FTransformData& LastTransform, const FTransformData& CurrentTransform, EKey3DTransformChannel::Type ChannelsToKey, bool bUnwindRotation, FKeyParams KeyParams )
+void F3DTransformTrackEditor::AddTransformKeys( UObject* ObjectToKey, const FTransformData& LastTransform, const FTransformData& CurrentTransform, EKey3DTransformChannel::Type ChannelsToKey, bool bUnwindRotation, ESequencerKeyMode KeyMode )
 {
 	TArray<FTransformKey> NewKeys;
 	TArray<FTransformKey> DefaultKeys;
 	GetTransformKeys(LastTransform, CurrentTransform, ChannelsToKey, bUnwindRotation, NewKeys, DefaultKeys);
 
-	AnimatablePropertyChanged( FOnKeyProperty::CreateRaw(this, &F3DTransformTrackEditor::OnAddTransformKeys, ObjectToKey, &NewKeys, &DefaultKeys, CurrentTransform, KeyParams ) );
+	AnimatablePropertyChanged( FOnKeyProperty::CreateRaw(this, &F3DTransformTrackEditor::OnAddTransformKeys, ObjectToKey, &NewKeys, &DefaultKeys, CurrentTransform, KeyMode ) );
 }
 
 
-bool F3DTransformTrackEditor::OnAddTransformKeys( float Time, UObject* ObjectToKey, TArray<FTransformKey>* NewKeys, TArray<FTransformKey>* DefaultKeys, FTransformData CurrentTransform, FKeyParams KeyParams )
+bool F3DTransformTrackEditor::OnAddTransformKeys( float Time, UObject* ObjectToKey, TArray<FTransformKey>* NewKeys, TArray<FTransformKey>* DefaultKeys, FTransformData CurrentTransform, ESequencerKeyMode KeyMode )
 {
 	TArray<UObject*> ObjectsToKey;
 	ObjectsToKey.Add(ObjectToKey);
@@ -830,24 +866,16 @@ bool F3DTransformTrackEditor::OnAddTransformKeys( float Time, UObject* ObjectToK
 		Time,
 		*NewKeys,
 		*DefaultKeys,
-		KeyParams,
+		KeyMode,
 		UMovieScene3DTransformTrack::StaticClass(),
 		TransformPropertyName,
 		[this](UMovieScene3DTransformTrack* NewTrack) {
 			NewTrack->SetPropertyNameAndPath(TransformPropertyName, TransformPropertyName.ToString());
-		},
-		[&](UMovieScene3DTransformTrack* NewTrack) {
-			SetIntermediateValueFromTransformChange(NewTrack, CurrentTransform);
 		}
 	);
 }
 
 
-void F3DTransformTrackEditor::SetIntermediateValueFromTransformChange( UMovieSceneTrack* Track, FTransformData TransformData )
-{
-	OnSetIntermediateValueFromTransformChange.Broadcast(Track, TransformData);
-}
-
 void F3DTransformTrackEditor::ImportAnimSequenceTransforms(const FAssetData& Asset, TSharedRef<ISequencer> Sequencer, UMovieScene3DTransformTrack* TransformTrack)
 {
 	FSlateApplication::Get().DismissAllMenus();
diff --git a/Engine/Source/Editor/MovieSceneTools/Private/TrackEditors/TransformTrackEditor.h b/Engine/Source/Editor/MovieSceneTools/Private/TrackEditors/TransformTrackEditor.h
--- a/Engine/Source/Editor/MovieSceneTools/Private/TrackEditors/TransformTrackEditor.h
+++ b/Engine/Source/Editor/MovieSceneTools/Private/TrackEditors/TransformTrackEditor.h
@@ -40,7 +40,7 @@ public:
 	virtual void BindCommands(TSharedRef<FUICommandList> SequencerCommandBindings) override;
 	virtual void BuildObjectBindingEditButtons(TSharedPtr<SHorizontalBox> EditBox, const FGuid& ObjectBinding, const UClass* ObjectClass) override;
 	virtual void BuildObjectBindingTrackMenu(FMenuBuilder& MenuBuilder, const FGuid& ObjectBinding, const UClass* ObjectClass) override;
-	virtual TSharedRef<ISequencerSection> MakeSectionInterface( UMovieSceneSection& SectionObject, UMovieSceneTrack& Track ) override;
+	virtual TSharedRef<ISequencerSection> MakeSectionInterface( UMovieSceneSection& SectionObject, UMovieSceneTrack& Track, FGuid ObjectBinding ) override;
 	virtual void OnRelease() override;
 	virtual bool SupportsType( TSubclassOf<UMovieSceneTrack> Type ) const override;
 	virtual void BuildTrackContextMenu( FMenuBuilder& MenuBuilder, UMovieSceneTrack* Track ) override;
@@ -97,7 +97,7 @@ private:
 	 * @param ChannelToKey The channels to add keys to.
 	 * @param KeyParams Parameters which control how the keys are added. 
 	 */
-	void AddTransformKeysForHandle( FGuid ObjectHandle, EKey3DTransformChannel::Type ChannelToKey, FKeyParams KeyParams );
+	void AddTransformKeysForHandle( FGuid ObjectHandle, EKey3DTransformChannel::Type ChannelToKey, ESequencerKeyMode KeyMode );
 
 	/**
 	* Adds transform keys to a specific object.
@@ -106,7 +106,7 @@ private:
 	* @param ChannelToKey The channels to add keys to.
 	* @param KeyParams Parameters which control how the keys are added.
 	*/
-	void AddTransformKeysForObject( UObject* Object, EKey3DTransformChannel::Type ChannelToKey, FKeyParams KeyParams );
+	void AddTransformKeysForObject( UObject* Object, EKey3DTransformChannel::Type ChannelToKey, ESequencerKeyMode KeyMode );
 
 	/**
 	* Adds keys to a specific actor.
@@ -117,7 +117,7 @@ private:
 	* @param bUnwindRotation Whether or not rotation key values should be unwound.
 	* @param KeyParams Parameters which control how the keys are added.
 	*/
-	void AddTransformKeys( UObject* ObjectToKey, const FTransformData& LastTransform, const FTransformData& CurrentTransform, EKey3DTransformChannel::Type ChannelsToKey, bool bUnwindRotation, FKeyParams KeyParams );
+	void AddTransformKeys( UObject* ObjectToKey, const FTransformData& LastTransform, const FTransformData& CurrentTransform, EKey3DTransformChannel::Type ChannelsToKey, bool bUnwindRotation, ESequencerKeyMode KeyMode );
 
 	/**
 	* Delegate target of AnimatablePropertyChanged which actually adds the keys.
@@ -127,14 +127,10 @@ private:
 	* @param Keys The keys to add.
 	* @param KeyParams Parameters which control how the keys are added.
 	*/
-	bool OnAddTransformKeys( float Time, UObject* ObjectToKey, TArray<FTransformKey>* NewKeys, TArray<FTransformKey>* DefaultKeys, FTransformData CurrentTransform, FKeyParams KeyParams );
+	bool OnAddTransformKeys( float Time, UObject* ObjectToKey, TArray<FTransformKey>* NewKeys, TArray<FTransformKey>* DefaultKeys, FTransformData CurrentTransform, ESequencerKeyMode KeyMode );
 
 private:
 
-	DECLARE_MULTICAST_DELEGATE_TwoParams( FOnSetIntermediateValueFromTransformChange, UMovieSceneTrack*, FTransformData )
-
-	void SetIntermediateValueFromTransformChange( UMovieSceneTrack* Track, FTransformData TransformData );
-
 	/** Import an animation sequence's root transforms into a transform section */
 	static void ImportAnimSequenceTransforms(const FAssetData& Asset, TSharedRef<class ISequencer> Sequencer, UMovieScene3DTransformTrack* TransformTrack);
 
@@ -144,6 +140,4 @@ private:
 
 	/** Mapping of objects to their existing transform data (for comparing against new transform data) */
 	TMap< TWeakObjectPtr<UObject>, FTransformData > ObjectToExistingTransform;
-
-	FOnSetIntermediateValueFromTransformChange OnSetIntermediateValueFromTransformChange;
 };
