Copying //UE4/Dev-Sequencer to //UE4/Dev-Main (Source: //UE4/Dev-Sequencer @ 3116759)

diff --git a/Engine/Source/Editor/MovieSceneTools/Public/BoolKeyArea.h b/Engine/Source/Editor/MovieSceneTools/Public/BoolKeyArea.h
--- a/Engine/Source/Editor/MovieSceneTools/Public/BoolKeyArea.h
+++ b/Engine/Source/Editor/MovieSceneTools/Public/BoolKeyArea.h
@@ -16,11 +16,29 @@ class FBoolKeyArea
 {
 public:
 
-	/* Create and initialize a new instance. */
+	/**
+	* Creates a new key area for editing bool curves.
+	*
+	* @param InCurve The integral curve which has the bool keys.
+	* @param InOwningSection The section which owns the curve which is being displayed and edited by this area.
+	*/
 	FBoolKeyArea(FIntegralCurve& InCurve, UMovieSceneSection* InOwningSection)
 		: FIntegralKeyArea(InCurve, InOwningSection)
 	{ }
 
+	/**
+	* Creates a new key area for editing bool curves whose value can be overridden externally.
+	*
+	* @param InCurve The integral curve which has the bool keys.
+	* @param InExternalValue An attribute which can provide an external value for this key area.  External values are
+	*        useful for things like property tracks where the property value can change without changing the animation
+	*        and we want to be able to key and update using the new property value.
+	* @param InOwningSection The section which owns the curve which is being displayed and edited by this area.
+	*/
+	FBoolKeyArea(FIntegralCurve& InCurve, TAttribute<TOptional<bool>> ExternalValue, UMovieSceneSection* InOwningSection)
+		: FIntegralKeyArea(InCurve, ExternalValue, InOwningSection)
+	{ }
+
 public:
 
 	//~ FIntegralKeyArea interface
@@ -30,5 +48,7 @@ public:
 
 protected:
 
-	void OnValueChanged(bool InValue);
+	//~ FIntegralKeyArea interface
+
+	virtual bool ConvertCurveValueToIntegralType(int32 CurveValue) const override;
 };
diff --git a/Engine/Source/Editor/MovieSceneTools/Public/ByteKeyArea.h b/Engine/Source/Editor/MovieSceneTools/Public/ByteKeyArea.h
--- a/Engine/Source/Editor/MovieSceneTools/Public/ByteKeyArea.h
+++ b/Engine/Source/Editor/MovieSceneTools/Public/ByteKeyArea.h
@@ -16,15 +16,39 @@ class FByteKeyArea
 {
 public:
 
-	/* Create and initialize a new instance. */
+	/**
+	* Creates a new key area for editing byte curves.
+	*
+	* @param InCurve The integral curve which has the byte keys.
+	* @param InOwningSection The section which owns the curve which is being displayed and edited by this area.
+	*/
 	FByteKeyArea(FIntegralCurve& InCurve, UMovieSceneSection* InOwningSection)
 		: FIntegralKeyArea(InCurve, InOwningSection)
 	{ }
 
+	/**
+	* Creates a new key area for editing byte curves whose value can be overridden externally.
+	*
+	* @param InCurve The integral curve which has the byte keys.
+	* @param InExternalValue An attribute which can provide an external value for this key area.  External values are
+	*        useful for things like property tracks where the property value can change without changing the animation
+	*        and we want to be able to key and update using the new property value.
+	* @param InOwningSection The section which owns the curve which is being displayed and edited by this area.
+	*/
+	FByteKeyArea(FIntegralCurve& InCurve, TAttribute<TOptional<uint8>> ExternalValue, UMovieSceneSection* InOwningSection)
+		: FIntegralKeyArea(InCurve, ExternalValue, InOwningSection)
+	{ }
+
 public:
 
 	//~ FIntegralKeyArea interface
 
 	virtual bool CanCreateKeyEditor() override;
 	virtual TSharedRef<SWidget> CreateKeyEditor(ISequencer* Sequencer) override;
+
+protected:
+
+	//~ FIntegralKeyArea interface
+
+	virtual uint8 ConvertCurveValueToIntegralType(int32 CurveValue) const override;
 };
diff --git a/Engine/Source/Editor/MovieSceneTools/Public/ClipboardTypes.h b/Engine/Source/Editor/MovieSceneTools/Public/ClipboardTypes.h
--- a/Engine/Source/Editor/MovieSceneTools/Public/ClipboardTypes.h
+++ b/Engine/Source/Editor/MovieSceneTools/Public/ClipboardTypes.h
@@ -38,6 +38,11 @@ namespace MovieSceneClipboard
 		static FName Name("String");
 		return Name;
 	}
+	template<> inline FName GetKeyTypeName<FGuid>()
+	{
+		static FName Name("Guid");
+		return Name;
+	}
 }
 
 #endif
\ No newline at end of file
diff --git a/Engine/Source/Editor/MovieSceneTools/Public/EnumKeyArea.h b/Engine/Source/Editor/MovieSceneTools/Public/EnumKeyArea.h
--- a/Engine/Source/Editor/MovieSceneTools/Public/EnumKeyArea.h
+++ b/Engine/Source/Editor/MovieSceneTools/Public/EnumKeyArea.h
@@ -19,12 +19,33 @@ class FEnumKeyArea
 {
 public:
 
-	/* Create and initialize a new instance. */
+	/**
+	* Creates a new key area for editing enum curves.
+	*
+	* @param InCurve The integral curve which has the enum keys.
+	* @param InOwningSection The section which owns the curve which is being displayed and edited by this area.
+	* @param InEnum The enum which is used to generate display values for this area.
+	*/
 	FEnumKeyArea(FIntegralCurve& InCurve, UMovieSceneSection* InOwningSection, const UEnum* InEnum)
 		: FByteKeyArea(InCurve, InOwningSection)
 		, Enum(InEnum)
 	{ }
 
+	/**
+	* Creates a new key area for editing enum curves whose value can be overridden externally.
+	*
+	* @param InCurve The integral curve which has the enum keys.
+	* @param InExternalValue An attribute which can provide an external value for this key area.  External values are
+	*        useful for things like property tracks where the property value can change without changing the animation
+	*        and we want to be able to key and update using the new property value.
+	* @param InOwningSection The section which owns the curve which is being displayed and edited by this area.
+	* @param InEnum The enum which is used to generate display values for this area.
+	*/
+	FEnumKeyArea(FIntegralCurve& InCurve, TAttribute<TOptional<uint8>> ExternalValue, UMovieSceneSection* InOwningSection, const UEnum* InEnum)
+		: FByteKeyArea(InCurve, ExternalValue, InOwningSection)
+		, Enum(InEnum)
+	{ }
+
 public:
 
 	//~ FByteKeyArea interface
@@ -32,10 +53,6 @@ public:
 	virtual bool CanCreateKeyEditor() override;
 	virtual TSharedRef<SWidget> CreateKeyEditor(ISequencer* Sequencer) override;
 
-protected:
-
-	void OnValueChanged(int32 InValue);
-
 private:
 
 	/** The enum which provides available integral values for this key area. */
diff --git a/Engine/Source/Editor/MovieSceneTools/Public/FloatCurveKeyArea.h b/Engine/Source/Editor/MovieSceneTools/Public/FloatCurveKeyArea.h
--- a/Engine/Source/Editor/MovieSceneTools/Public/FloatCurveKeyArea.h
+++ b/Engine/Source/Editor/MovieSceneTools/Public/FloatCurveKeyArea.h
@@ -18,23 +18,35 @@ class MOVIESCENETOOLS_API FFloatCurveKeyArea
 {
 public:
 
+	/**
+	* Creates a new key area for editing float curves.
+	*
+	* @param InCurve The rich curve which has the float keys.
+	* @param InOwningSection The section which owns the curve which is being displayed and edited by this area.
+	* @param InColor An optional color which should be used to draw the background of this key area.
+	*/
 	FFloatCurveKeyArea(FRichCurve* InCurve, UMovieSceneSection* InOwningSection, TOptional<FLinearColor> InColor = TOptional<FLinearColor>())
 		: Color(InColor)
 		, Curve(InCurve)
 		, OwningSection(InOwningSection)
 	{ }
 
-public:
-
-	void ClearIntermediateValue()
-	{
-		IntermediateValue.Reset();
-	}
-
-	void SetIntermediateValue(float InIntermediateValue)
-	{
-		IntermediateValue = InIntermediateValue;
-	}
+	/**
+	* Creates a new key area for editing float curves whose value can be overridden externally.
+	*
+	* @param InCurve The rich curve which has the float keys.
+	* @param InExternalValue An attribute which can provide an external value for this key area.  External values are
+	*        useful for things like property tracks where the property value can change without changing the animation
+	*        and we want to be able to key and update using the new property value.
+	* @param InOwningSection The section which owns the curve which is being displayed and edited by this area.
+	* @param InColor An optional color which should be used to draw the background of this key area.
+	*/
+	FFloatCurveKeyArea(FRichCurve* InCurve, TAttribute<TOptional<float>> InExternalValue, UMovieSceneSection* InOwningSection, TOptional<FLinearColor> InColor = TOptional<FLinearColor>())
+		: Color(InColor)
+		, Curve(InCurve)
+		, OwningSection(InOwningSection)
+		, ExternalValue(InExternalValue)
+	{ }
 
 public:
 
@@ -78,5 +90,6 @@ private:
 	/** The section that owns this key area. */
 	UMovieSceneSection* OwningSection;
 
-	TOptional<float> IntermediateValue;
+	/** The value displayed by this key area supplied externally. */
+	TAttribute<TOptional<float>> ExternalValue;
 };
diff --git a/Engine/Source/Editor/MovieSceneTools/Public/IntegralKeyArea.h b/Engine/Source/Editor/MovieSceneTools/Public/IntegralKeyArea.h
--- a/Engine/Source/Editor/MovieSceneTools/Public/IntegralKeyArea.h
+++ b/Engine/Source/Editor/MovieSceneTools/Public/IntegralKeyArea.h
@@ -50,7 +50,7 @@ public:
 protected:
 
 	virtual void EvaluateAndAddKey(float Time, float TimeToCopyFrom, FKeyHandle& CurrentKey) = 0;
-	virtual void SetKeyValue(float Time) = 0;
+	virtual void UpdateKeyWithExternalValue(float Time) = 0;
 
 protected:
 
@@ -71,10 +71,30 @@ class MOVIESCENETOOLS_API FIntegralKeyArea
 {
 public:
 
+	/**
+	* Creates a new key area for editing integral curves.
+	*
+	* @param InCurve The integral curve which has the integral keys.
+	* @param InOwningSection The section which owns the curve which is being displayed and edited by this area.
+	*/
 	FIntegralKeyArea(FIntegralCurve& InCurve, UMovieSceneSection* InOwningSection)
 		: FIntegralCurveKeyAreaBase(InCurve, InOwningSection)
 	{ }
 
+	/**
+	* Creates a new key area for editing integral curves whose value can be overridden externally.
+	*
+	* @param InCurve The integral curve which has the integral keys.
+	* @param InExternalValue An attribute which can provide an external value for this key area.  External values are
+	*        useful for things like property tracks where the property value can change without changing the animation
+	*        and we want to be able to key and update using the new property value.
+	* @param InOwningSection The section which owns the curve which is being displayed and edited by this area.
+	*/
+	FIntegralKeyArea(FIntegralCurve& InCurve, TAttribute<TOptional<IntegralType>> InExternalValue, UMovieSceneSection* InOwningSection)
+		: FIntegralCurveKeyAreaBase(InCurve, InOwningSection)
+		, ExternalValue(InExternalValue)
+	{ }
+
 	virtual void CopyKeys(FMovieSceneClipboardBuilder& ClipboardBuilder, const TFunctionRef<bool(FKeyHandle, const IKeyArea&)>& KeyMask) const override
 	{
 		const UMovieSceneSection* Section = const_cast<FIntegralKeyArea*>(this)->GetOwningSection();
@@ -134,41 +154,32 @@ public:
 		});
 	}
 
-public:
-
-	void ClearIntermediateValue()
-	{
-		IntermediateValue.Reset();
-	}
-
-	void SetIntermediateValue(IntegralType InIntermediateValue)
-	{
-		IntermediateValue = InIntermediateValue;
-	}
-
 protected:
 
+	virtual IntegralType ConvertCurveValueToIntegralType(int32 CurveValue) const = 0;
+
 	virtual void EvaluateAndAddKey(float Time, float TimeToCopyFrom, FKeyHandle& CurrentKey) override
 	{
-		int32 Value = Curve.Evaluate(Time);
-
-		if (TimeToCopyFrom != FLT_MAX)
+		IntegralType Value;
+		if (ExternalValue.IsSet() && ExternalValue.Get().IsSet() && TimeToCopyFrom == FLT_MAX)
 		{
-			Value = Curve.Evaluate(TimeToCopyFrom);
+			Value = ExternalValue.Get().GetValue();
 		}
-		else if (IntermediateValue.IsSet())
+		else
 		{
-			Value = IntermediateValue.GetValue();
+			float EvalTime = TimeToCopyFrom != FLT_MAX ? Time : TimeToCopyFrom;
+			IntegralType DefaultValue = 0;
+			Value = ConvertCurveValueToIntegralType(Curve.Evaluate(EvalTime, DefaultValue));
 		}
 
 		Curve.AddKey(Time, Value, CurrentKey);
 	}
 
-	virtual void SetKeyValue(float Time) override
+	virtual void UpdateKeyWithExternalValue(float Time) override
 	{
-		if (IntermediateValue.IsSet())
+		if (ExternalValue.IsSet() && ExternalValue.Get().IsSet())
 		{
-			int32 Value = IntermediateValue.GetValue();
+			IntegralType Value = ExternalValue.Get().GetValue();
 
 			Curve.UpdateOrAddKey(Time, Value);
 		}
@@ -176,5 +187,5 @@ protected:
 
 protected:
 
-	TOptional<IntegralType> IntermediateValue;
+	TAttribute<TOptional<IntegralType>> ExternalValue;
 };
diff --git a/Engine/Source/Editor/MovieSceneTools/Public/KeyframeTrackEditor.h b/Engine/Source/Editor/MovieSceneTools/Public/KeyframeTrackEditor.h
--- a/Engine/Source/Editor/MovieSceneTools/Public/KeyframeTrackEditor.h
+++ b/Engine/Source/Editor/MovieSceneTools/Public/KeyframeTrackEditor.h
@@ -40,6 +40,15 @@ public:
 		return Type == TrackType::StaticClass();
 	}
 
+
+	virtual void BuildTrackContextMenu(FMenuBuilder& MenuBuilder, UMovieSceneTrack* Track) override
+	{
+		MenuBuilder.AddSubMenu(
+			NSLOCTEXT("KeyframeTrackEditor", "TrackDefaults", "Track Defaults"),
+			NSLOCTEXT("KeyframeTrackEditor", "TrackDefaultsTooltip", "Track default value operations."),
+			FNewMenuDelegate::CreateSP(this, &FKeyframeTrackEditor::AddTrackDefaultsItems, Track));
+	}
+
 protected:
 
 	/*
@@ -55,29 +64,31 @@ protected:
 	 * @param PropertyName The name of the property to add keys for.
 	 * @param OnInitializeNewTrack A delegate which allows for custom initialization for new tracks.  This is called after the 
 	 *        track is created, but before any sections or keys have been added.
-	 * @param OnSetIntermediateValue A delegate which is called when the key params specify that keying should only happen when
-	 * 		  auto-keying is on, but auto-keying is off.  This allows handling the changed but not keyed intermediate value.
 	 * @return Whether or not a handle guid or track was created. Note this does not return true if keys were added or modified.
 	 */
 	bool AddKeysToObjects(
 		TArray<UObject*> ObjectsToKey, float KeyTime,
 		const TArray<KeyDataType>& NewKeys, const TArray<KeyDataType>& DefaultKeys,
-		FKeyParams KeyParams, TSubclassOf<UMovieSceneTrack> TrackClass, FName PropertyName,
-		TFunction<void(TrackType*)> OnInitializeNewTrack,
-		TFunction<void(TrackType*)> OnSetIntermediateValue)
+		ESequencerKeyMode KeyMode, TSubclassOf<UMovieSceneTrack> TrackClass, FName PropertyName,
+		TFunction<void(TrackType*)> OnInitializeNewTrack)
 	{
 		bool bHandleCreated = false;
 		bool bTrackCreated = false;
 
+		bool bCreateHandle =
+			(KeyMode == ESequencerKeyMode::AutoKey && GetSequencer()->GetAutoKeyMode() == EAutoKeyMode::KeyAll) ||
+			KeyMode == ESequencerKeyMode::ManualKey ||
+			KeyMode == ESequencerKeyMode::ManualKeyForced;
+
 		for ( UObject* Object : ObjectsToKey )
 		{
-			FFindOrCreateHandleResult HandleResult = FindOrCreateHandleToObject( Object, KeyParams.bCreateHandleIfMissing );
+			FFindOrCreateHandleResult HandleResult = FindOrCreateHandleToObject( Object, bCreateHandle );
 			FGuid ObjectHandle = HandleResult.Handle;
 			bHandleCreated = HandleResult.bWasCreated;
 
 			if ( ObjectHandle.IsValid() )
 			{
-				bTrackCreated |= AddKeysToHandle( ObjectHandle, KeyTime, NewKeys, DefaultKeys, KeyParams, TrackClass, PropertyName, OnInitializeNewTrack, OnSetIntermediateValue );
+				bTrackCreated |= AddKeysToHandle( ObjectHandle, KeyTime, NewKeys, DefaultKeys, KeyMode, TrackClass, PropertyName, OnInitializeNewTrack );
 			}
 		}
 		return bHandleCreated || bTrackCreated;
@@ -86,6 +97,29 @@ protected:
 
 private:
 
+	virtual void ClearDefaults( UMovieSceneTrack* Track )
+	{
+		const FScopedTransaction Transaction(NSLOCTEXT("KeyframeTrackEditor", "ClearTrackDefaultsTransaction", "Clear track defaults"));
+		for (UMovieSceneSection* Section : Track->GetAllSections())
+		{
+			Section->Modify();
+			IKeyframeSection<KeyDataType>* KeyframeSection = CastChecked<SectionType>(Section);
+			KeyframeSection->ClearDefaults();
+		}
+		GetSequencer()->NotifyMovieSceneDataChanged(EMovieSceneDataChangeType::TrackValueChanged);
+	}
+
+	void AddTrackDefaultsItems( FMenuBuilder& MenuBuilder, UMovieSceneTrack* Track )
+	{
+		MenuBuilder.AddMenuEntry(
+			NSLOCTEXT( "KeyframeTrackEditor", "ClearDefaults", "Clear Defaults" ),
+			NSLOCTEXT( "KeyframeTrackEditor", "ClearDefaultsToolTip", "Clear the current default values for this track." ),
+			FSlateIcon(),
+			FUIAction( FExecuteAction::CreateSP( this, &FKeyframeTrackEditor::ClearDefaults, Track ) ),
+			NAME_None,
+			EUserInterfaceActionType::Button );
+	}
+
 	/*
 	 * Adds keys to the specified guid.  This may also add tracks and sections depending on the options specified. 
 	 *
@@ -99,22 +133,24 @@ private:
 	 * @param PropertyName The name of the property to add keys for.
 	 * @param OnInitializeNewTrack A delegate which allows for custom initialization for new tracks.  This is called after the 
 	 *        track is created, but before any sections or keys have been added.
-	 * @param OnSetIntermediateValue A delegate which is called when the key params specify that keying should only happen when
-	 * 		  auto-keying is on, but auto-keying is off.  This allows handling the changed but not keyed intermediate value.
 	 * @return Whether or not a track was created. Note this does not return true if keys were added or modified.
 	*/
 	bool AddKeysToHandle(
 		FGuid ObjectHandle, float KeyTime,
 		const TArray<KeyDataType>& NewKeys, const TArray<KeyDataType>& DefaultKeys,
-		FKeyParams KeyParams, TSubclassOf<UMovieSceneTrack> TrackClass, FName PropertyName,
-		TFunction<void(TrackType*)> OnInitializeNewTrack,
-		TFunction<void(TrackType*)> OnSetIntermediateValue)
+		ESequencerKeyMode KeyMode, TSubclassOf<UMovieSceneTrack> TrackClass, FName PropertyName,
+		TFunction<void(TrackType*)> OnInitializeNewTrack)
 	{
 		bool bTrackCreated = false;
 		bool bSectionCreated = false;
 
+		bool bCreateTrack =
+			(KeyMode == ESequencerKeyMode::AutoKey && GetSequencer()->GetAutoKeyMode() == EAutoKeyMode::KeyAll) ||
+			KeyMode == ESequencerKeyMode::ManualKey ||
+			KeyMode == ESequencerKeyMode::ManualKeyForced;
+
 		// Try to find an existing Track, and if one doesn't exist check the key params and create one if requested.
-		FFindOrCreateTrackResult TrackResult = FindOrCreateTrackForObject( ObjectHandle, TrackClass, PropertyName, KeyParams.bCreateTrackIfMissing );
+		FFindOrCreateTrackResult TrackResult = FindOrCreateTrackForObject( ObjectHandle, TrackClass, PropertyName, bCreateTrack );
 		TrackType* Track = CastChecked<TrackType>( TrackResult.Track, ECastCheckedType::NullAllowed );
 		bTrackCreated = TrackResult.bWasCreated;
 
@@ -128,7 +164,7 @@ private:
 
 		if ( Track )
 		{
-			bSectionCreated = AddKeysToTrack( Track, KeyTime, NewKeys, DefaultKeys, KeyParams, bTrackCreated, OnSetIntermediateValue );
+			bSectionCreated = AddKeysToTrack( Track, KeyTime, NewKeys, DefaultKeys, KeyMode, bTrackCreated );
 		}
 
 		return bTrackCreated || bSectionCreated;
@@ -138,46 +174,56 @@ private:
 	bool AddKeysToTrack(
 		TrackType* Track, float KeyTime,
 		const TArray<KeyDataType>& NewKeys, const TArray<KeyDataType>& DefaultKeys,
-		FKeyParams KeyParams, bool bNewTrack,
-		TFunction<void(TrackType*)> OnSetIntermediateValue)
+		ESequencerKeyMode KeyMode, bool bNewTrack)
 	{
 		bool bSectionCreated = false;
+		bool bInfiniteKeyAreas = GetSequencer()->GetInfiniteKeyAreas();
 
-		bool bSettingIntermediateValue = KeyParams.bCreateKeyOnlyWhenAutoKeying && GetSequencer()->GetAutoKeyMode() == EAutoKeyMode::KeyNone;
-		if ( bSettingIntermediateValue )
-		{
-			checkf(OnSetIntermediateValue, TEXT("A valid OnSetIntermediateValue delegate must be provided for consistent keyframing behavior."));
-			OnSetIntermediateValue(Track);
-		}
-		else
+		if ( KeyMode != ESequencerKeyMode::AutoKey || GetSequencer()->GetAutoKeyMode() != EAutoKeyMode::KeyNone )
 		{
 			EMovieSceneKeyInterpolation InterpolationMode = GetSequencer()->GetKeyInterpolation();
-			bool bInfiniteKeyAreas = GetSequencer()->GetInfiniteKeyAreas();
 
-			for ( const KeyDataType& NewKey : NewKeys )
+			bool bKeyEvenIfUnchanged =
+				KeyMode == ESequencerKeyMode::ManualKeyForced ||
+				GetSequencer()->GetKeyAllEnabled();
+
+			bool bKeyEvenIfEmpty =
+				(KeyMode == ESequencerKeyMode::AutoKey && GetSequencer()->GetAutoKeyMode() == EAutoKeyMode::KeyAll) ||
+				KeyMode == ESequencerKeyMode::ManualKeyForced;
+
+			for (const KeyDataType& NewKey : NewKeys)
 			{
-				// Only modify the track if the new key will create new data, or if creating keys has been forced.
-				if ( NewKeyIsNewData( Track, KeyTime, NewKey ) || KeyParams.bCreateKeyIfUnchanged )
+				if ( NewKeyIsNewData( Track, KeyTime, NewKey ) || bKeyEvenIfUnchanged )
 				{
-					if ( HasKeys( Track, NewKey ) || KeyParams.bCreateKeyIfEmpty )
+					if ( HasKeys( Track, NewKey ) || bKeyEvenIfEmpty )
 					{
 						bSectionCreated |= AddKey( Track, KeyTime, NewKey, InterpolationMode, bInfiniteKeyAreas );
 					}
-					else
-					{
-						bSectionCreated |= SetDefault( Track, KeyTime, NewKey, bInfiniteKeyAreas );
-					}
 				}
 			}
-			// If a new track was added, set defaults for all of the default keys.
-			if( bNewTrack )
+		}
+
+		if (GetSequencer()->GetAutoSetTrackDefaults())
+		{
+			for (const KeyDataType& NewKey : NewKeys)
 			{
-				for ( const KeyDataType& DefaultKey : DefaultKeys )
-				{
-					bSectionCreated |= SetDefault( Track, KeyTime, DefaultKey, bInfiniteKeyAreas  );
-				}
+				bSectionCreated |= SetDefault(Track, KeyTime, NewKey, bInfiniteKeyAreas);
 			}
+
+			for (const KeyDataType& DefaultKey : DefaultKeys)
+			{
+				bSectionCreated |= SetDefault(Track, KeyTime, DefaultKey, bInfiniteKeyAreas);
+			}
+		}
+
+		// If a new track was created but no keys or defaults were set, make sure a new section is created too to allow the user to edit it.
+		if ( bNewTrack && bSectionCreated == false)
+		{
+			Track->Modify();
+			UMovieSceneSection* NewSection = Track->FindOrAddSection(KeyTime, bSectionCreated);
+			NewSection->SetIsInfinite(bInfiniteKeyAreas);
 		}
+
 		return bSectionCreated;
 	}
 
diff --git a/Engine/Source/Editor/MovieSceneTools/Public/MovieSceneToolHelpers.h b/Engine/Source/Editor/MovieSceneTools/Public/MovieSceneToolHelpers.h
--- a/Engine/Source/Editor/MovieSceneTools/Public/MovieSceneToolHelpers.h
+++ b/Engine/Source/Editor/MovieSceneTools/Public/MovieSceneToolHelpers.h
@@ -97,7 +97,7 @@ public:
 	 * @param OnSelectionChanged Delegate fired when selection is changed
 	 * @return The new widget
 	 */
-	static TSharedRef<SWidget> MakeEnumComboBox(const UEnum* Enum, TAttribute<int32> CurrentValue, FOnEnumSelectionChanged OnSelectionChanged, TAttribute<TOptional<uint8>> InIntermediateValue);
+	static TSharedRef<SWidget> MakeEnumComboBox(const UEnum* Enum, TAttribute<int32> CurrentValue, FOnEnumSelectionChanged OnSelectionChanged);
 
 
 	/**
diff --git a/Engine/Source/Editor/MovieSceneTools/Public/PropertySection.h b/Engine/Source/Editor/MovieSceneTools/Public/PropertySection.h
--- a/Engine/Source/Editor/MovieSceneTools/Public/PropertySection.h
+++ b/Engine/Source/Editor/MovieSceneTools/Public/PropertySection.h
@@ -4,67 +4,87 @@
 
 #include "MovieSceneTrack.h"
 #include "MovieSceneSection.h"
+#include "MovieSceneCommonHelpers.h"
 #include "ISequencerSection.h"
 
+class FTrackInstancePropertyBindings;
 
 /**
  * A generic implementation for displaying simple property sections
  */
-class FPropertySection
+class MOVIESCENETOOLS_API FPropertySection
 	: public ISequencerSection
 {
 public:
 
 	/**
-	 * Create and initialize a new instance.
-	 *
-	 * @param The section's human readable display name.
-	 */
-	FPropertySection(UMovieSceneSection& InSectionObject, const FText& InDisplayName)
-		: DisplayName(InDisplayName)
-		, SectionObject(InSectionObject)
-	{ }
+	* Creates a new property section for editing non-property sections.
+	*
+	* @param InSectionObject The section object which is being displayed and edited.
+	* @param InDisplayName A display name for the section being displayed and edited.
+	* @TODO Make another section which is a good base class for non-property sections.
+	*/
+	FPropertySection(UMovieSceneSection& InSectionObject, const FText& InDisplayName);
+
+	/**
+	* Creates a new property section.
+	*
+	* @param InSequencer The sequencer which is controlling this property section.
+	* @param InObjectBinding The object binding for the object which owns the property that this section is animating.
+	* @param InPropertyName The name of the property which is animated by this section.
+	* @param InPropertyPath A string representing the path to the property which is animated by this section.
+	* @param InSectionObject The section object which is being displayed and edited.
+	* @param InDisplayName A display name for the section being displayed and edited.
+	*/
+	FPropertySection(ISequencer* InSequencer, FGuid InObjectBinding, FName InPropertyName, const FString& InPropertyPath, UMovieSceneSection& InSectionObject, const FText& InDisplayName);
 
 public:
 
 	// ISequencerSection interface
 
-	virtual UMovieSceneSection* GetSectionObject() override
-	{
-		return &SectionObject;
-	}
+	virtual UMovieSceneSection* GetSectionObject() override;
+	virtual FText GetDisplayName() const override;
+	virtual FText GetSectionTitle() const override;
+	virtual int32 OnPaintSection(FSequencerSectionPainter& Painter) const override;
 
-	virtual FText GetDisplayName() const override
-	{
-		return DisplayName;
-	}
-	
-	virtual FText GetSectionTitle() const override
-	{
-		return FText::GetEmpty();
-	}
+protected:
 
-	virtual void GenerateSectionLayout( class ISectionLayoutBuilder& LayoutBuilder ) const override { }
+	/** Gets the sequencer which is controlling this section .*/
+	ISequencer* GetSequencer() const { return Sequencer; }
 
-	virtual int32 OnPaintSection( FSequencerSectionPainter& Painter ) const override 
-	{
-		return Painter.PaintSectionBackground();
-	}
+	/** Gets the property being animated by this section.  Returns nullptr if this section was not constructed with the necessary
+		data to get the property value, or if the runtime object or property can not be found. */
+	UProperty* GetProperty() const;
 
-	void SetIntermediateValueForTrack( UMovieSceneTrack* Track, FPropertyChangedParams PropertyChangedParams )
+	/** Gets the current value of the property which is being animated.  The optional return value will be unset if
+		this section was not constructed with the necessary data, or is the runtime object or property can not be found. */
+	template<typename ValueType>
+	TOptional<ValueType> GetPropertyValue() const
 	{
-		if ( SectionObject.GetOuter() == Track )
+		UObject* RuntimeObject = GetRuntimeObjectAndUpdatePropertyBindings();
+
+		// Bool property values are stored in a bit field so using a straight cast of the pointer to get their value does not
+		// work.  Instead use the actual property to get the correct value.
+		const UBoolProperty* BoolProperty = Cast<const UBoolProperty>(GetProperty());
+		if (BoolProperty)
 		{
-			SetIntermediateValue(PropertyChangedParams);
+			uint32* ValuePtr = BoolProperty->ContainerPtrToValuePtr<uint32>(RuntimeObject);
+			bool BoolPropertyValue = BoolProperty->GetPropertyValue(ValuePtr);
+			void *PropertyValue = &BoolPropertyValue;
+			return TOptional<ValueType>(*((ValueType*)PropertyValue));
 		}
+
+		return RuntimeObject != nullptr
+			? TOptional<ValueType>(PropertyBindings->GetCurrentValue<ValueType>(RuntimeObject))
+			: TOptional<ValueType>();
 	}
 
-	/** Sets the intermediate value for this section interface.  Intermediate values are used to display property values which
-		have changed, but have not been keyed yet. */
-	virtual void SetIntermediateValue(FPropertyChangedParams PropertyChangedParams) = 0;
+	/** Returns true when this section was constructed with the data necessary to query for the current property value. */
+	bool CanGetPropertyValue() const;
 
-	/** Clears any previously set intermediate values. */
-	virtual void ClearIntermediateValue() = 0;
+	/** Gets the current runtime object for the object binding which was used to construct this section.  Returns nullptr
+		if no valid object can be found, or if an object binding was not supplied. */
+	UObject* GetRuntimeObjectAndUpdatePropertyBindings() const;
 
 protected:
 
@@ -73,5 +93,20 @@ protected:
 
 	/** The section we are visualizing */
 	UMovieSceneSection& SectionObject;
+
+private:
+
+	/** The sequencer which is controlling this section. */
+	ISequencer* Sequencer;
+
+	/** An object binding for the object which owns the property being animated by this section. */
+	FGuid ObjectBinding;
+
+	/** An object which is used to retrieve the value of a property based on it's name and path. */
+	mutable TSharedPtr<FTrackInstancePropertyBindings> PropertyBindings;
+
+	/** Caches a weak pointer to the most recent runtime object retrieved by this section.  This cache
+		is maintained so that the property bindings can be updated any time this object changes. */
+	mutable TWeakObjectPtr<UObject> RuntimeObjectCache;
 };
 
diff --git a/Engine/Source/Editor/MovieSceneTools/Public/PropertyTrackEditor.h b/Engine/Source/Editor/MovieSceneTools/Public/PropertyTrackEditor.h
--- a/Engine/Source/Editor/MovieSceneTools/Public/PropertyTrackEditor.h
+++ b/Engine/Source/Editor/MovieSceneTools/Public/PropertyTrackEditor.h
@@ -11,9 +11,6 @@ class IPropertyHandle;
 class FPropertyChangedParams;
 
 
-DECLARE_MULTICAST_DELEGATE_TwoParams(FOnSetIntermediateValueFromPropertyChange, UMovieSceneTrack*, FPropertyChangedParams)
-
-
 /**
 * Tools for animatable property types such as floats ands vectors
 */
@@ -110,19 +107,8 @@ public:
 		return Type == TrackType::StaticClass();
 	}
 
-	virtual TSharedRef<ISequencerSection> MakeSectionInterface( class UMovieSceneSection& SectionObject, UMovieSceneTrack& Track ) override
-	{
-		TSharedRef<FPropertySection> PropertySection = MakePropertySectionInterface( SectionObject, Track );
-		OnSetIntermediateValueFromPropertyChange.AddSP( PropertySection, &FPropertySection::SetIntermediateValueForTrack );
-		FMovieSceneTrackEditor::GetSequencer()->OnGlobalTimeChanged().AddSP( PropertySection, &FPropertySection::ClearIntermediateValue );
-		return PropertySection;
-	}
-
 protected:
 
-	/** Creates a property section for a section object. */
-	virtual TSharedRef<FPropertySection> MakePropertySectionInterface( class UMovieSceneSection& SectionObject, UMovieSceneTrack& Track ) = 0;
-
 	/**
 	* Generates keys based on the new value from the property property change parameters.
 	*
@@ -204,14 +190,6 @@ private:
 			TrackClass = TrackType::StaticClass();
 		}
 
-		if ( FMovieSceneTrackEditor::GetSequencer()->GetAutoKeyMode() == EAutoKeyMode::KeyAll )
-		{
-			// TODO: Setting these here is a bit sketchy but necessary right now without a larger refactor.
-			PropertyChangedParams.KeyParams.bCreateHandleIfMissing = true;
-			PropertyChangedParams.KeyParams.bCreateTrackIfMissing = true;
-			PropertyChangedParams.KeyParams.bCreateKeyIfEmpty = true;
-		}
-
 		// If the track class has been customized for this property then it's possible this track editor doesn't support it, 
 		// also check for track editors which should only be used for customization.
 		if ( SupportsType( TrackClass ) && ( ForCustomizedUseOnly() == false || *CustomizedClass != nullptr) )
@@ -221,11 +199,10 @@ private:
 				KeyTime,
 				NewKeysForPropertyChange,
 				DefaultKeysForPropertyChange,
-				PropertyChangedParams.KeyParams,
+				PropertyChangedParams.KeyMode,
 				TrackClass,
 				PropertyChangedParams.PropertyPath.Last()->GetFName(),
-				[&](TrackType* NewTrack) { InitializeNewTrack(NewTrack, PropertyChangedParams); },
-				[&](TrackType* NewTrack) { SetIntermediateValueFromPropertyChange(NewTrack, PropertyChangedParams); }
+				[&](TrackType* NewTrack) { InitializeNewTrack(NewTrack, PropertyChangedParams); }
 			);
 		}
 		else
@@ -234,17 +211,8 @@ private:
 		}
 	}
 
-private:
-
-	void SetIntermediateValueFromPropertyChange(TrackType* Track, FPropertyChangedParams PropertyChangedParams)
-	{
-		OnSetIntermediateValueFromPropertyChange.Broadcast( Track, PropertyChangedParams );
-	}
-
 private:
 
 	/** An array of property type names which are being watched for changes. */
 	TArray<FName> WatchedPropertyTypeNames;
-
-	FOnSetIntermediateValueFromPropertyChange OnSetIntermediateValueFromPropertyChange;
-};
\ No newline at end of file
+};
diff --git a/Engine/Source/Editor/MovieSceneTools/Public/StringCurveKeyArea.h b/Engine/Source/Editor/MovieSceneTools/Public/StringCurveKeyArea.h
--- a/Engine/Source/Editor/MovieSceneTools/Public/StringCurveKeyArea.h
+++ b/Engine/Source/Editor/MovieSceneTools/Public/StringCurveKeyArea.h
@@ -18,23 +18,35 @@ class MOVIESCENETOOLS_API FStringCurveKeyArea
 {
 public:
 
+	/**
+	* Creates a new key area for editing string curves.
+	*
+	* @param InCurve The string curve which has the string keys.
+	* @param InOwningSection The section which owns the curve which is being displayed and edited by this area.
+	* @param InColor An optional color which is used to draw the background of this key area.
+	*/
 	FStringCurveKeyArea(FStringCurve* InCurve, UMovieSceneSection* InOwningSection, TOptional<FLinearColor> InColor = TOptional<FLinearColor>())
 		: Color(InColor)
 		, Curve(InCurve)
 		, OwningSection(InOwningSection)
 	{ }
 
-public:
-
-	void ClearIntermediateValue()
-	{
-		IntermediateValue.Reset();
-	}
-
-	void SetIntermediateValue(FString InIntermediateValue)
-	{
-		IntermediateValue = InIntermediateValue;
-	}
+	/**
+	* Creates a new key area for editing string curves whose value can be overridden externally.
+	*
+	* @param InCurve The string curve which has the string keys.
+	* @param InExternalValue An attribute which can provide an external value for this key area.  External values are
+	*        useful for things like property tracks where the property value can change without changing the animation
+	*        and we want to be able to key and update using the new property value.
+	* @param InOwningSection The section which owns the curve which is being displayed and edited by this area.
+	* @param InColor An optional color which is used to draw the background of this key area.
+	*/
+	FStringCurveKeyArea(FStringCurve* InCurve, TAttribute<TOptional<FString>> InExternalValue, UMovieSceneSection* InOwningSection, TOptional<FLinearColor> InColor = TOptional<FLinearColor>())
+		: Color(InColor)
+		, Curve(InCurve)
+		, OwningSection(InOwningSection)
+		, ExternalValue(InExternalValue)
+	{ }
 
 public:
 
@@ -63,10 +75,6 @@ public:
 	virtual void CopyKeys(FMovieSceneClipboardBuilder& ClipboardBuilder, const TFunctionRef<bool(FKeyHandle, const IKeyArea&)>& KeyMask) const override;
 	virtual void PasteKeys(const FMovieSceneClipboardKeyTrack& KeyTrack, const FMovieSceneClipboardEnvironment& SrcEnvironment, const FSequencerPasteEnvironment& DstEnvironment) override;
 
-protected:
-
-	void OnValueChanged(FString InValue);
-
 private:
 
 	/** The key area's color. */
@@ -78,5 +86,6 @@ private:
 	/** The section that owns this key area. */
 	UMovieSceneSection* OwningSection;
 
-	TOptional<FString> IntermediateValue;
+	/** An attribute which allows the value for this key area to be overridden externally. */
+	TAttribute<TOptional<FString>> ExternalValue;
 };
diff --git a/Engine/Source/Editor/NiagaraEditor/Private/NiagaraSequencer.h b/Engine/Source/Editor/NiagaraEditor/Private/NiagaraSequencer.h
--- a/Engine/Source/Editor/NiagaraEditor/Private/NiagaraSequencer.h
+++ b/Engine/Source/Editor/NiagaraEditor/Private/NiagaraSequencer.h
@@ -241,7 +241,7 @@ public:
 		return false;
 	}
 
-	virtual TSharedRef<ISequencerSection> MakeSectionInterface(UMovieSceneSection& SectionObject, UMovieSceneTrack& Track)
+	virtual TSharedRef<ISequencerSection> MakeSectionInterface(UMovieSceneSection& SectionObject, UMovieSceneTrack& Track, FGuid ObjectBinding) override
 	{
 		INiagaraSequencerSection *Sec = new INiagaraSequencerSection(SectionObject);
 		return MakeShareable(Sec);
diff --git a/Engine/Source/Editor/SequenceRecorder/Private/SequenceRecorder.cpp b/Engine/Source/Editor/SequenceRecorder/Private/SequenceRecorder.cpp
--- a/Engine/Source/Editor/SequenceRecorder/Private/SequenceRecorder.cpp
+++ b/Engine/Source/Editor/SequenceRecorder/Private/SequenceRecorder.cpp
@@ -304,7 +304,7 @@ void FSequenceRecorder::Tick(float DeltaSeconds)
 
 void FSequenceRecorder::DrawDebug(UCanvas* InCanvas, APlayerController* InPlayerController)
 {
-	const float NumFrames = 4.0f;
+	const float NumFrames = 9.0f;
 	const bool bCountingDown = CurrentDelay > 0.0f && CurrentDelay < NumFrames;
 
 	if(bCountingDown)
@@ -317,6 +317,8 @@ void FSequenceRecorder::DrawDebug(UCanvas* InCanvas, APlayerController* InPlayer
 		InCanvas->GetCenter(Center.X, Center.Y);
 		FVector2D IconPosition = Center - HalfIconSize;
 
+		InCanvas->SetDrawColor(FColor::White);
+
 		FCanvasIcon Icon = UCanvas::MakeIcon(CountdownTexture.Get(), FMath::FloorToFloat(NumFrames - CurrentDelay) * IconSize.X, 0.0f, IconSize.X, IconSize.Y);
 		InCanvas->DrawIcon(Icon, IconPosition.X, IconPosition.Y);
 		
@@ -459,7 +461,7 @@ bool FSequenceRecorder::StartRecording(const FOnRecordingStarted& OnRecordingSta
 	if(Settings->RecordingDelay > 0.0f)
 	{
 		CurrentDelay = Settings->RecordingDelay;
-		
+
 		UE_LOG(LogAnimation, Display, TEXT("Starting sequence recording with delay of %f seconds"), CurrentDelay);
 
 		return QueuedRecordings.Num() > 0;
@@ -603,6 +605,40 @@ bool FSequenceRecorder::StartRecordingInternal(UWorld* World)
 
 		UE_LOG(LogAnimation, Display, TEXT("Started recording sequence %s"), *LevelSequence->GetPathName());
 
+		// If we created an audio recorder at the start of the count down, then start recording
+		// Create the audio recorder now before the count down finishes
+		if (Settings->RecordAudio != EAudioRecordingMode::None)
+		{
+			if (LevelSequence)
+			{
+				FDirectoryPath AudioDirectory;
+				AudioDirectory.Path = Settings->SequenceRecordingBasePath.Path;
+				if (Settings->AudioSubDirectory.Len())
+				{
+					AudioDirectory.Path /= Settings->AudioSubDirectory;
+				}
+
+				ISequenceRecorder& Recorder = FModuleManager::Get().LoadModuleChecked<ISequenceRecorder>("SequenceRecorder");
+
+				FAudioRecorderSettings AudioSettings;
+				AudioSettings.Directory = AudioDirectory;
+				AudioSettings.AssetName = FText::Format(LOCTEXT("AudioFormatStr", "{0}_Audio"), FText::FromString(LevelSequence->GetName())).ToString();
+				AudioSettings.RecordingDurationSec = Settings->SequenceLength;
+				AudioSettings.GainDB = Settings->AudioGain;
+				AudioSettings.InputBufferSize = Settings->AudioInputBufferSize;
+
+				AudioRecorder = Recorder.CreateAudioRecorder();
+				if (AudioRecorder)
+				{
+					AudioRecorder->Start(AudioSettings);
+				}
+			}
+			else
+			{
+				UE_LOG(LogAnimation, Display, TEXT("'Create Level Sequence' needs to be enabled for audio recording"));
+			}
+		}
+
 		OnRecordingStartedDelegate.ExecuteIfBound(CurrentSequence.Get());
 		return true;
 	}
@@ -642,9 +678,53 @@ bool FSequenceRecorder::StopRecording()
 		}
 	}
 
-	FScopedSlowTask SlowTask((float)(QueuedRecordings.Num() + DeadRecordings.Num()), LOCTEXT("ProcessingRecording", "Processing Recording"));
+	// 1 step for the audio processing
+	static const uint8 NumAdditionalSteps = 1;
+
+	FScopedSlowTask SlowTask((float)(QueuedRecordings.Num() + DeadRecordings.Num() + NumAdditionalSteps), LOCTEXT("ProcessingRecording", "Processing Recording"));
 	SlowTask.MakeDialog(false, true);
 
+	// Process audio first so it doesn't record while we're processing the other captured state
+	ULevelSequence* LevelSequence = CurrentSequence.Get();
+
+	SlowTask.EnterProgressFrame(1.f, LOCTEXT("ProcessingAudio", "Processing Audio"));
+	if (AudioRecorder && LevelSequence)
+	{
+		USoundWave* RecordedAudio = AudioRecorder->Stop();
+		AudioRecorder.Reset();
+
+		if (RecordedAudio)
+		{
+			// Add a new master audio track to the level sequence
+			
+			UMovieScene* MovieScene = LevelSequence->GetMovieScene();
+
+			UMovieSceneAudioTrack* AudioTrack = MovieScene->FindMasterTrack<UMovieSceneAudioTrack>();
+			if (!AudioTrack)
+			{
+				AudioTrack = MovieScene->AddMasterTrack<UMovieSceneAudioTrack>();
+				AudioTrack->SetDisplayName(LOCTEXT("DefaultAudioTrackName", "Recorded Audio"));
+			}
+
+			int32 RowIndex = -1;
+			for (UMovieSceneSection* Section : AudioTrack->GetAllSections())
+			{
+				RowIndex = FMath::Max(RowIndex, Section->GetRowIndex());
+			}
+
+			UMovieSceneAudioSection* NewAudioSection = NewObject<UMovieSceneAudioSection>(AudioTrack, UMovieSceneAudioSection::StaticClass());
+
+			NewAudioSection->SetRowIndex(RowIndex + 1);
+			NewAudioSection->SetSound(RecordedAudio);
+			NewAudioSection->SetStartTime(0);
+			NewAudioSection->SetEndTime(RecordedAudio->GetDuration());
+
+			AudioTrack->AddSection(*NewAudioSection);
+		}
+	}
+
+
+
 	CurrentDelay = 0.0f;
 
 	// remove spawn delegates
@@ -681,10 +761,8 @@ bool FSequenceRecorder::StopRecording()
 
 	if(Settings->bCreateLevelSequence)
 	{
-		if(CurrentSequence.IsValid())
+		if(LevelSequence)
 		{
-			ULevelSequence* LevelSequence = CurrentSequence.Get();
-
 			// Stop referencing the sequence so we are listed as 'not recording'
 			CurrentSequence = nullptr;
 
diff --git a/Engine/Source/Editor/SequenceRecorder/Private/SequenceRecorder.h b/Engine/Source/Editor/SequenceRecorder/Private/SequenceRecorder.h
--- a/Engine/Source/Editor/SequenceRecorder/Private/SequenceRecorder.h
+++ b/Engine/Source/Editor/SequenceRecorder/Private/SequenceRecorder.h
@@ -165,6 +165,9 @@ private:
 	/** Built-in transform recorder factory */
 	FMovieScene3DTransformSectionRecorderFactory TransformSectionRecorderFactory;
 
+	/** Audio recorder */
+	TUniquePtr<ISequenceAudioRecorder> AudioRecorder;
+
 	/** Built in multi-property recorder */
 	FMovieSceneMultiPropertyRecorderFactory MultiPropertySectionRecorder;
 };
diff --git a/Engine/Source/Editor/SequenceRecorder/Private/SequenceRecorderModule.cpp b/Engine/Source/Editor/SequenceRecorder/Private/SequenceRecorderModule.cpp
--- a/Engine/Source/Editor/SequenceRecorder/Private/SequenceRecorderModule.cpp
+++ b/Engine/Source/Editor/SequenceRecorder/Private/SequenceRecorderModule.cpp
@@ -11,6 +11,7 @@
 #include "NotificationManager.h"
 #include "SNotificationList.h"
 #include "ActorRecordingDetailsCustomization.h"
+#include "SequenceRecorderDetailsCustomization.h"
 #include "PropertiesToRecordForClassDetailsCustomization.h"
 
 #define LOCTEXT_NAMESPACE "SequenceRecorder"
@@ -97,6 +98,7 @@ class FSequenceRecorderModule : public ISequenceRecorder, private FSelfRegisteri
 			// register details customization
 			FPropertyEditorModule& PropertyModule = FModuleManager::LoadModuleChecked<FPropertyEditorModule>("PropertyEditor");
 			PropertyModule.RegisterCustomClassLayout(UActorRecording::StaticClass()->GetFName(), FOnGetDetailCustomizationInstance::CreateStatic(&FActorRecordingDetailsCustomization::MakeInstance));
+			PropertyModule.RegisterCustomClassLayout(USequenceRecorderSettings::StaticClass()->GetFName(), FOnGetDetailCustomizationInstance::CreateStatic(&FSequenceRecorderDetailsCustomization::MakeInstance));
 			PropertyModule.RegisterCustomPropertyTypeLayout(FPropertiesToRecordForClass::StaticStruct()->GetFName(), FOnGetPropertyTypeCustomizationInstance::CreateStatic(&FPropertiesToRecordForClassDetailsCustomization::MakeInstance));
 		}
 #endif
@@ -533,6 +535,34 @@ class FSequenceRecorderModule : public ISequenceRecorder, private FSelfRegisteri
 		return FGuid();
 	}
 
+	virtual FDelegateHandle RegisterAudioRecorder(const TFunction<TUniquePtr<ISequenceAudioRecorder>()>& FactoryFunction) override
+	{
+		ensureMsgf(!AudioFactory, TEXT("Audio recorder already registered."));
+
+		AudioFactory = FactoryFunction;
+		AudioFactoryHandle = FDelegateHandle(FDelegateHandle::GenerateNewHandle);
+		return AudioFactoryHandle;
+	}
+
+	virtual void UnregisterAudioRecorder(FDelegateHandle Handle) override
+	{
+		if (Handle == AudioFactoryHandle)
+		{
+			AudioFactory.Unset();
+			AudioFactoryHandle = FDelegateHandle();
+		}
+	}
+
+	virtual bool HasAudioRecorder() const override
+	{
+		return AudioFactoryHandle.IsValid();
+	}
+
+	virtual TUniquePtr<ISequenceAudioRecorder> CreateAudioRecorder() const
+	{
+		return AudioFactory ? AudioFactory() : TUniquePtr<ISequenceAudioRecorder>();
+	}
+
 	static void TickSequenceRecorder(float DeltaSeconds)
 	{
 		if (!IsRunningDedicatedServer() && !IsRunningCommandlet())
@@ -608,6 +638,10 @@ class FSequenceRecorderModule : public ISequenceRecorder, private FSelfRegisteri
 	FDelegateHandle PostEditorTickHandle;
 
 	FDelegateHandle DrawDebugDelegateHandle;
+
+	TFunction<TUniquePtr<ISequenceAudioRecorder>()> AudioFactory;
+
+	FDelegateHandle AudioFactoryHandle;
 };
 
 IMPLEMENT_MODULE( FSequenceRecorderModule, SequenceRecorder )
diff --git a/Engine/Source/Editor/SequenceRecorder/Private/SequenceRecorderSettings.cpp b/Engine/Source/Editor/SequenceRecorder/Private/SequenceRecorderSettings.cpp
--- a/Engine/Source/Editor/SequenceRecorder/Private/SequenceRecorderSettings.cpp
+++ b/Engine/Source/Editor/SequenceRecorder/Private/SequenceRecorderSettings.cpp
@@ -14,6 +14,9 @@ USequenceRecorderSettings::USequenceRecorderSettings(const FObjectInitializer& O
 	RecordingDelay = 4.0f;
 	SequenceName = TEXT("RecordedSequence");
 	AnimationSubDirectory = TEXT("Animations");
+	AudioSubDirectory = TEXT("Audio");
+	AudioGain = 0.0f;
+	AudioInputBufferSize = 4048;
 	SequenceRecordingBasePath.Path = TEXT("/Game/Cinematics/Sequences");
 	bRecordNearbySpawnedActors = true;
 	NearbyActorRecordingProximity = 5000.0f;
diff --git a/Engine/Source/Editor/SequenceRecorder/Public/ISequenceAudioRecorder.h b/Engine/Source/Editor/SequenceRecorder/Public/ISequenceAudioRecorder.h
new file mode 100644
--- /dev/null
+++ b/Engine/Source/Editor/SequenceRecorder/Public/ISequenceAudioRecorder.h
@@ -0,0 +1,37 @@
+// Copyright 1998-2016 Epic Games, Inc. All Rights Reserved.
+
+#pragma once
+
+#include "Engine/EngineTypes.h"
+
+class USoundWave;
+
+struct FAudioRecorderSettings
+{
+	/** Directory to create the asset within (empty for transient package) */
+	FDirectoryPath Directory;
+	FString AssetName;
+	float RecordingDurationSec;
+	float GainDB;
+	int32 InputBufferSize;
+
+	FAudioRecorderSettings()
+		: RecordingDurationSec(-1.0f)
+		, GainDB(0.0f)
+		, InputBufferSize(256)
+	{}
+};
+
+class ISequenceAudioRecorder
+{
+public:
+
+	/** Virtual destructor */
+	virtual ~ISequenceAudioRecorder() {}
+	
+	/** Start recording audio data */
+	virtual void Start(const FAudioRecorderSettings& Settings) = 0;
+
+	/** Stop recording audio data */
+	virtual USoundWave* Stop() = 0;
+};
\ No newline at end of file
diff --git a/Engine/Source/Editor/SequenceRecorder/Public/ISequenceRecorder.h b/Engine/Source/Editor/SequenceRecorder/Public/ISequenceRecorder.h
--- a/Engine/Source/Editor/SequenceRecorder/Public/ISequenceRecorder.h
+++ b/Engine/Source/Editor/SequenceRecorder/Public/ISequenceRecorder.h
@@ -3,6 +3,8 @@
 #pragma once
 
 #include "ModuleInterface.h"
+#include "ISequenceAudioRecorder.h"
+#include "UniquePtr.h"
 #include "Containers/ArrayView.h"
 
 DECLARE_DELEGATE_OneParam(FOnRecordingStarted, class UMovieSceneSequence* /*Sequence*/);
@@ -74,4 +76,30 @@ public:
 	 * @return the spawnable Guid
 	 */
 	virtual FGuid GetRecordingGuid(AActor* Actor) const = 0;
+
+	/**
+	 * Register a function that will return a new audio capturer for the specified parameters
+	 * @param	FactoryFunction		Function used to generate a new audio recorder
+	 * @return A handle to be passed to UnregisterAudioRecorder to unregister the recorder
+	 */
+	virtual FDelegateHandle RegisterAudioRecorder(const TFunction<TUniquePtr<ISequenceAudioRecorder>()>& FactoryFunction) = 0;
+
+	/**
+	 * Unregister a previously registered audio recorder factory function
+	 * @param	RegisteredHandle	The handle returned from RegisterAudioRecorder
+	 */
+	virtual void UnregisterAudioRecorder(FDelegateHandle RegisteredHandle) = 0;
+
+	/**
+	 * Check whether we have an audio recorder registered or not
+	 * @return true if we have an audio recorder registered, false otherwise
+	 */
+	virtual bool HasAudioRecorder() const = 0;
+
+	/**
+	 * Attempt to create an audio recorder
+	 * @param	Settings	Settings for the audio recorder
+	 * @return A valid ptr to an audio recorder or null
+	 */
+	virtual TUniquePtr<ISequenceAudioRecorder> CreateAudioRecorder() const = 0;
 };
\ No newline at end of file
diff --git a/Engine/Source/Editor/SequenceRecorder/Public/SequenceRecorderSettings.h b/Engine/Source/Editor/SequenceRecorder/Public/SequenceRecorderSettings.h
--- a/Engine/Source/Editor/SequenceRecorder/Public/SequenceRecorderSettings.h
+++ b/Engine/Source/Editor/SequenceRecorder/Public/SequenceRecorderSettings.h
@@ -6,6 +6,14 @@
 #include "SequenceRecorderActorFilter.h"
 #include "SequenceRecorderSettings.generated.h"
 
+/** Enum denoting if (and how) to record audio */
+UENUM()
+enum class EAudioRecordingMode : uint8
+{
+	None 			UMETA(DisplayName="Don't Record Audio"),
+	AudioTrack		UMETA(DisplayName="Into Audio Track"),
+};
+
 USTRUCT()
 struct FPropertiesToRecordForClass
 {
@@ -62,7 +70,7 @@ public:
 	float SequenceLength;
 
 	/** Delay that we will use before starting recording */
-	UPROPERTY(Config, EditAnywhere, Category = "Sequence Recording")
+	UPROPERTY(Config, EditAnywhere, Category = "Sequence Recording", meta = (ClampMax="9.0", UIMax = "9.0"))
 	float RecordingDelay;
 
 	/** The base name of the sequence to record to. This name will also be used to auto-generate any assets created by this recording. */
@@ -77,6 +85,22 @@ public:
 	UPROPERTY(Config, EditAnywhere, AdvancedDisplay, Category = "Sequence Recording")
 	FString AnimationSubDirectory;
 
+	/** The name of the subdirectory audio will be placed in. Leave this empty to place into the same directory as the sequence base path */
+	UPROPERTY(Config, EditAnywhere, AdvancedDisplay, Category = "Sequence Recording")
+	FString AudioSubDirectory;
+
+	/** Whether to record audio alongside animation or not */
+	UPROPERTY(Config, EditAnywhere, Category = "Sequence Recording")
+	EAudioRecordingMode RecordAudio;
+
+	/** Gain in decibels to apply to recorded audio */
+	UPROPERTY(Config, EditAnywhere, Category = "Sequence Recording")
+	float AudioGain;
+
+	/** The buffer size to use on mic input callbacks. Larger sizes increase latency but reduce chances of buffer overruns (pops and discontinuities). */
+	UPROPERTY(Config, EditAnywhere, Category = "Sequence Recording")
+	int32 AudioInputBufferSize;
+
 	/** Whether to record nearby spawned actors. */
 	UPROPERTY(Config, EditAnywhere, Category = "Sequence Recording")
 	bool bRecordNearbySpawnedActors;
diff --git a/Engine/Source/Editor/SequenceRecorder/SequenceRecorder.Build.cs b/Engine/Source/Editor/SequenceRecorder/SequenceRecorder.Build.cs
--- a/Engine/Source/Editor/SequenceRecorder/SequenceRecorder.Build.cs
+++ b/Engine/Source/Editor/SequenceRecorder/SequenceRecorder.Build.cs
@@ -54,6 +54,12 @@ namespace UnrealBuildTool.Rules
 				{
 				}
 				);
+
+			if (Target.Platform == UnrealTargetPlatform.Win64 || Target.Platform == UnrealTargetPlatform.Win32)
+			{
+				// Add __WINDOWS_WASAPI__ so that RtAudio compiles with WASAPI
+				Definitions.Add("__WINDOWS_WASAPI__");
+			}
 		}
 	}
 }
\ No newline at end of file
diff --git a/Engine/Source/Editor/Sequencer/Private/DisplayNodes/SequencerObjectBindingNode.cpp b/Engine/Source/Editor/Sequencer/Private/DisplayNodes/SequencerObjectBindingNode.cpp
--- a/Engine/Source/Editor/Sequencer/Private/DisplayNodes/SequencerObjectBindingNode.cpp
+++ b/Engine/Source/Editor/Sequencer/Private/DisplayNodes/SequencerObjectBindingNode.cpp
@@ -633,13 +633,7 @@ void FSequencerObjectBindingNode::HandlePropertyMenuItemExecute(TArray<UProperty
 		}
 	}
 
-	FKeyPropertyParams KeyPropertyParams(KeyableBoundObjects, PropertyPath);
-	{
-		KeyPropertyParams.KeyParams.bCreateTrackIfMissing = true;
-		KeyPropertyParams.KeyParams.bCreateHandleIfMissing = true;
-		KeyPropertyParams.KeyParams.bCreateKeyIfUnchanged = true;
-		KeyPropertyParams.KeyParams.bCreateKeyIfEmpty = true;
-	}
+	FKeyPropertyParams KeyPropertyParams(KeyableBoundObjects, PropertyPath, ESequencerKeyMode::ManualKeyForced);
 
 	Sequencer.KeyProperty(KeyPropertyParams);
 }
diff --git a/Engine/Source/Editor/Sequencer/Private/IntegralKeyDetailsCustomization.cpp b/Engine/Source/Editor/Sequencer/Private/IntegralKeyDetailsCustomization.cpp
--- a/Engine/Source/Editor/Sequencer/Private/IntegralKeyDetailsCustomization.cpp
+++ b/Engine/Source/Editor/Sequencer/Private/IntegralKeyDetailsCustomization.cpp
@@ -60,8 +60,7 @@ void FIntegralKeyDetailsCustomization::CustomizeDetails(IDetailLayoutBuilder& De
 					FOnEnumSelectionChanged::CreateLambda([ValueProperty](int32 Selection, ESelectInfo::Type SelectionType)
 					{
 						ValueProperty->SetValue(Selection);
-					}),
-					TAttribute<TOptional<uint8>>());
+					}));
 			}
 			else
 			{
diff --git a/Engine/Source/Editor/Sequencer/Private/KeyPropertyParams.cpp b/Engine/Source/Editor/Sequencer/Private/KeyPropertyParams.cpp
--- a/Engine/Source/Editor/Sequencer/Private/KeyPropertyParams.cpp
+++ b/Engine/Source/Editor/Sequencer/Private/KeyPropertyParams.cpp
@@ -3,8 +3,10 @@
 #include "SequencerPrivatePCH.h"
 #include "PropertyHandle.h"
 
-void PropertyHandleToPropertyPath(const UClass* OwnerClass, const IPropertyHandle& InPropertyHandle, TArray<UProperty*>& PropertyPath)
+TArray<UProperty*> PropertyHandleToPropertyPath(const UClass* OwnerClass, const IPropertyHandle& InPropertyHandle)
 {
+	TArray<UProperty*> PropertyPath;
+
 	PropertyPath.Add(InPropertyHandle.GetProperty());
 	TSharedPtr<IPropertyHandle> CurrentHandle = InPropertyHandle.GetParentHandle();
 	while (CurrentHandle.IsValid() && CurrentHandle->GetProperty() != nullptr)
@@ -12,40 +14,43 @@ void PropertyHandleToPropertyPath(const UClass* OwnerClass, const IPropertyHandl
 		PropertyPath.Insert(CurrentHandle->GetProperty(), 0);
 		CurrentHandle = CurrentHandle->GetParentHandle();
 	}
-}
 
-FCanKeyPropertyParams::FCanKeyPropertyParams()
-{
-	ObjectClass = nullptr;
+	return PropertyPath;
 }
 
 FCanKeyPropertyParams::FCanKeyPropertyParams(UClass* InObjectClass, TArray<UProperty*> InPropertyPath)
+	: ObjectClass(InObjectClass)
+	, PropertyPath(InPropertyPath)
 {
-	ObjectClass = InObjectClass;
-	PropertyPath = InPropertyPath;
 }
 
 FCanKeyPropertyParams::FCanKeyPropertyParams(UClass* InObjectClass, const IPropertyHandle& InPropertyHandle)
+	: ObjectClass(InObjectClass)
+	, PropertyPath(PropertyHandleToPropertyPath(InObjectClass, InPropertyHandle))
 {
-	ObjectClass = InObjectClass;
-	PropertyHandleToPropertyPath(ObjectClass, InPropertyHandle, PropertyPath);
 }
 
-FKeyPropertyParams::FKeyPropertyParams()
+FKeyPropertyParams::FKeyPropertyParams(TArray<UObject*> InObjectsToKey, TArray<UProperty*> InPropertyPath, ESequencerKeyMode InKeyMode)
+	: ObjectsToKey(InObjectsToKey)
+	, PropertyPath(InPropertyPath)
+	, KeyMode(InKeyMode)
+
 {
 }
 
-FKeyPropertyParams::FKeyPropertyParams(TArray<UObject*> InObjectsToKey, TArray<UProperty*> InPropertyPath)
+FKeyPropertyParams::FKeyPropertyParams(TArray<UObject*> InObjectsToKey, const IPropertyHandle& InPropertyHandle, ESequencerKeyMode InKeyMode)
+	: ObjectsToKey(InObjectsToKey)
+	, PropertyPath(PropertyHandleToPropertyPath(InObjectsToKey[0]->GetClass(), InPropertyHandle))
+	, KeyMode(InKeyMode)
 {
-	ObjectsToKey = InObjectsToKey;
-	PropertyPath = InPropertyPath;
 }
 
-FKeyPropertyParams::FKeyPropertyParams(TArray<UObject*> InObjectsToKey, const IPropertyHandle& InPropertyHandle)
+FPropertyChangedParams::FPropertyChangedParams(TArray<UObject*> InObjectsThatChanged, TArray<UProperty*> InPropertyPath, FName InStructPropertyNameToKey, ESequencerKeyMode InKeyMode)
+	: ObjectsThatChanged(InObjectsThatChanged)
+	, PropertyPath(InPropertyPath)
+	, StructPropertyNameToKey(InStructPropertyNameToKey)
+	, KeyMode(InKeyMode)
 {
-	ObjectsToKey = InObjectsToKey;
-	PropertyHandleToPropertyPath(InObjectsToKey[0]->GetClass(), InPropertyHandle, PropertyPath);
-	KeyParams.bCreateKeyOnlyWhenAutoKeying = true;
 }
 
 FString FPropertyChangedParams::GetPropertyPathString() const
diff --git a/Engine/Source/Editor/Sequencer/Private/SKeyNavigationButtons.h b/Engine/Source/Editor/Sequencer/Private/SKeyNavigationButtons.h
--- a/Engine/Source/Editor/Sequencer/Private/SKeyNavigationButtons.h
+++ b/Engine/Source/Editor/Sequencer/Private/SKeyNavigationButtons.h
@@ -208,7 +208,7 @@ public:
 			return FReply::Unhandled();
 		}
 
-		FScopedTransaction Transaction(LOCTEXT("AddKeys", "Add keys at current time"));
+		FScopedTransaction Transaction(LOCTEXT("AddKeys", "Add Keys at Current Time"));
 		for (TSharedPtr<IKeyArea> KeyArea : KeyAreas)
 		{
 			UMovieSceneSection* OwningSection = KeyArea->GetOwningSection();
diff --git a/Engine/Source/Editor/Sequencer/Private/SSequencerSection.cpp b/Engine/Source/Editor/Sequencer/Private/SSequencerSection.cpp
--- a/Engine/Source/Editor/Sequencer/Private/SSequencerSection.cpp
+++ b/Engine/Source/Editor/Sequencer/Private/SSequencerSection.cpp
@@ -252,7 +252,6 @@ FSequencerSelectedKey SSequencerSection::CreateKeyUnderMouse( const FVector2D& M
 				KeyTime = TimeToPixelConverter.PixelToTime(LocalSpaceMousePosition.X);
 			}
 
-			FScopedTransaction CreateKeyTransaction(NSLOCTEXT("Sequencer", "CreateKey_Transaction", "Create Key"));
 			if (Section.TryModify())
 			{
 				// If the pressed key exists, offset the new key and look for it in the newly laid out key areas
@@ -288,7 +287,7 @@ FSequencerSelectedKey SSequencerSection::CreateKeyUnderMouse( const FVector2D& M
 				}
 				else
 				{
-					KeyArea->AddKeyUnique(KeyTime, GetSequencer().GetKeyInterpolation());
+					KeyArea->AddKeyUnique(KeyTime, GetSequencer().GetKeyInterpolation(), KeyTime);
 					Layout = FSectionLayout(*ParentSectionArea, SectionIndex);
 						
 					return GetKeyUnderMouse(MousePosition, AllottedGeometry);
@@ -853,6 +852,8 @@ FReply SSequencerSection::OnMouseButtonDown( const FGeometry& MyGeometry, const
 
 	if (MouseEvent.GetEffectingButton() == EKeys::MiddleMouseButton)
 	{
+		GEditor->BeginTransaction(NSLOCTEXT("Sequencer", "CreateKey_Transaction", "Create Key"));
+
 		// Generate a key and set it as the PressedKey
 		FSequencerSelectedKey NewKey = CreateKeyUnderMouse(MouseEvent.GetScreenSpacePosition(), MyGeometry, HoveredKey);
 
@@ -922,7 +923,17 @@ FReply SSequencerSection::OnMouseMove( const FGeometry& MyGeometry, const FPoint
 
 	return FReply::Unhandled();
 }
+	
+FReply SSequencerSection::OnMouseButtonUp(const FGeometry& MyGeometry, const FPointerEvent& MouseEvent)
+{
+	if (MouseEvent.GetEffectingButton() == EKeys::MiddleMouseButton)
+	{
+		GEditor->EndTransaction();
 
+		return FReply::Handled();
+	}
+	return FReply::Unhandled();
+}
 
 void SSequencerSection::OnMouseLeave( const FPointerEvent& MouseEvent )
 {
diff --git a/Engine/Source/Editor/Sequencer/Private/SSequencerSection.h b/Engine/Source/Editor/Sequencer/Private/SSequencerSection.h
--- a/Engine/Source/Editor/Sequencer/Private/SSequencerSection.h
+++ b/Engine/Source/Editor/Sequencer/Private/SSequencerSection.h
@@ -66,6 +66,7 @@ private:
 	virtual FReply OnMouseButtonDown( const FGeometry& MyGeometry, const FPointerEvent& MouseEvent ) override;
 	virtual FReply OnMouseButtonDoubleClick( const FGeometry& MyGeometry, const FPointerEvent& MouseEvent ) override;
 	virtual FReply OnMouseMove( const FGeometry& MyGeometry, const FPointerEvent& MouseEvent ) override;
+	virtual FReply OnMouseButtonUp(const FGeometry& MyGeometry, const FPointerEvent& MouseEvent) override;
 	virtual void OnMouseLeave( const FPointerEvent& MouseEvent ) override;
 
 	/**
diff --git a/Engine/Source/Editor/Sequencer/Private/Sequencer.cpp b/Engine/Source/Editor/Sequencer/Private/Sequencer.cpp
--- a/Engine/Source/Editor/Sequencer/Private/Sequencer.cpp
+++ b/Engine/Source/Editor/Sequencer/Private/Sequencer.cpp
@@ -73,6 +73,7 @@
 #include "CameraRig_Rail.h"
 #include "CameraRig_Crane.h"
 #include "Components/SplineComponent.h"
+#include "AudioDevice.h"
 #include "MainFrame.h"
 #include "DesktopPlatformModule.h"
 #include "FbxExporter.h"
@@ -157,6 +158,9 @@ void FSequencer::InitSequencer(const FSequencerInitParams& InitParams, const TSh
 
 		SettingsModule->RegisterSettings("Editor", "ContentEditors", SettingsSectionName, SettingsDisplayName, SettingsDescription, Settings);
 	}
+	
+	Settings->GetOnLockPlaybackToAudioClockChanged().AddSP(this, &FSequencer::ResetTimingManager);
+	ResetTimingManager(Settings->ShouldLockPlaybackToAudioClock());
 
 	ToolkitHost = InitParams.ToolkitHost;
 
@@ -396,6 +400,8 @@ void FSequencer::Tick(float InDeltaTime)
 		SetPlaybackStatus(StoredPlaybackState);
 	}
 
+	FTimeAndDelta TimeAndDelta = TimingManager->AdjustTime(GetGlobalTime(), InDeltaTime, PlayRate);
+
 	static const float AutoScrollFactor = 0.1f;
 
 	// Animate the autoscroll offset if it's set
@@ -433,26 +439,16 @@ void FSequencer::Tick(float InDeltaTime)
 		}
 	}
 
-	float TimeDilation = 1.f;
-	UWorld* PlaybackContext = Cast<UWorld>(GetPlaybackContext());
-	if (PlaybackContext)
-	{
-		TimeDilation = PlaybackContext->GetWorldSettings()->MatineeTimeDilation;
-	}
-
-	// calculate new global time
-	float NewTime = GetGlobalTime() + InDeltaTime * TimeDilation * PlayRate;
-
 	if (PlaybackState == EMovieScenePlayerStatus::Playing ||
 		PlaybackState == EMovieScenePlayerStatus::Recording)
 	{
-		SetGlobalTimeLooped(NewTime);
+		SetGlobalTimeLooped(TimeAndDelta.Time);
 	}
 
 	// Tick all the tools we own as well
 	for (int32 EditorIndex = 0; EditorIndex < TrackEditors.Num(); ++EditorIndex)
 	{
-		TrackEditors[EditorIndex]->Tick(InDeltaTime * PlayRate);
+		TrackEditors[EditorIndex]->Tick(TimeAndDelta.Delta * PlayRate);
 	}
 
 	if (!IsInSilentMode())
@@ -1053,6 +1049,10 @@ void FSequencer::NotifyMovieSceneDataChanged( EMovieSceneDataChangeType DataChan
 		bNeedTreeRefresh = false;
 		SetPlaybackStatus( StoredPlaybackState );
 	}
+	else if (DataChangeType == EMovieSceneDataChangeType::TrackValueChangedRefreshImmediately)
+	{
+		UpdateRuntimeInstances();
+	}
 	else
 	{
 		if ( DataChangeType != EMovieSceneDataChangeType::TrackValueChanged )
@@ -1367,6 +1367,12 @@ void FSequencer::SetInfiniteKeyAreas(bool bInfiniteKeyAreas)
 }
 
 
+bool FSequencer::GetAutoSetTrackDefaults() const
+{
+	return Settings->GetAutoSetTrackDefaults();
+}
+
+
 bool FSequencer::IsRecordingLive() const 
 {
 	return PlaybackState == EMovieScenePlayerStatus::Recording && GIsPlayInEditorWorld;
@@ -1930,6 +1936,8 @@ void FSequencer::SetPlaybackStatus(EMovieScenePlayerStatus::Type InPlaybackStatu
 		PlayRate = 1.f;
 		ShuttleMultiplier = 0;
 	}
+
+	TimingManager->Update(PlaybackState, GetGlobalTime());
 }
 
 
@@ -2775,6 +2783,7 @@ void FSequencer::SetGlobalTimeLooped(float InTime)
 				InTime >= FocusedSequence->GetMovieScene()->GetPlaybackRange().GetUpperBoundValue())
 			{
 				InTime = PlayRate > 0 ? FocusedSequence->GetMovieScene()->GetPlaybackRange().GetLowerBoundValue() : FocusedSequence->GetMovieScene()->GetPlaybackRange().GetUpperBoundValue();
+				TimingManager->OnStartPlaying(InTime);
 				bRestarted = true;
 			}
 		}
@@ -3667,14 +3676,7 @@ void FSequencer::OnNewActorsDropped(const TArray<UObject*>& DroppedObjects, cons
 					UProperty* CurrentPositionOnRailProperty = RailActor->GetClass()->FindPropertyByName(TEXT("CurrentPositionOnRail"));
 					PropertyPath.Add(CurrentPositionOnRailProperty);
 
-					FKeyPropertyParams KeyPropertyParams(TArrayBuilder<UObject*>().Add(RailActor), PropertyPath);
-					{
-						KeyPropertyParams.KeyParams.bCreateTrackIfMissing = true;
-						KeyPropertyParams.KeyParams.bCreateHandleIfMissing = true;
-						KeyPropertyParams.KeyParams.bCreateKeyIfUnchanged = false;
-						KeyPropertyParams.KeyParams.bCreateKeyIfEmpty = true;
-						KeyPropertyParams.KeyParams.bCreateKeyOnlyWhenAutoKeying = false;
-					}
+					FKeyPropertyParams KeyPropertyParams(TArrayBuilder<UObject*>().Add(RailActor), PropertyPath, ESequencerKeyMode::ManualKeyForced);
 
 					float OriginalTime = GetGlobalTime();
 
@@ -3739,6 +3741,8 @@ void FSequencer::OnPreBeginPIE(bool bIsSimulating)
 {
 	RootMovieSceneSequenceInstance->RestoreState(*this);
 	PrePossessionViewTargets.Reset();
+
+	bNeedInstanceRefresh = true;
 }
 
 
@@ -5364,6 +5368,7 @@ void FSequencer::SetKeyTime(const bool bUseFrames)
 
 void FSequencer::OnSetKeyTimeTextCommitted(const FText& InText, ETextCommit::Type CommitInfo, const bool bUseFrames)
 {
+	bool bAnythingChanged = false;
 	CloseEntryPopupMenu();
 	if (CommitInfo == ETextCommit::OnEnter)
 	{
@@ -5385,10 +5390,25 @@ void FSequencer::OnSetKeyTimeTextCommitted(const FText& InText, ETextCommit::Typ
 				if (Key.Section->TryModify())
 				{
 					Key.KeyArea->SetKeyTime(Key.KeyHandle.GetValue(), NewKeyTime);
+					bAnythingChanged = true;
+
+					if( NewKeyTime > Key.Section->GetEndTime() )
+					{
+						Key.Section->SetEndTime( NewKeyTime );
+					}
+					else if( NewKeyTime < Key.Section->GetStartTime() )
+					{
+						Key.Section->SetStartTime( NewKeyTime );
+					}
 				}
 			}
 		}
 	}
+
+	if (bAnythingChanged)
+	{
+		NotifyMovieSceneDataChanged(EMovieSceneDataChangeType::TrackValueChanged);
+	}
 }
 
 
@@ -6775,5 +6795,17 @@ void FSequencer::BuildObjectBindingEditButtons(TSharedPtr<SHorizontalBox> EditBo
 	}
 }
 
+void FSequencer::ResetTimingManager(bool bInShouldLockToAudioClock)
+{
+	if (bInShouldLockToAudioClock && GEngine->GetMainAudioDevice())
+	{
+		TimingManager.Reset(new FSequencerAudioClockTimer);
+	}
+	else
+	{
+		TimingManager.Reset(new FSequencerDefaultTimingManager);
+	}
+	TimingManager->Update(PlaybackState, GetGlobalTime());
+}
 
 #undef LOCTEXT_NAMESPACE
diff --git a/Engine/Source/Editor/Sequencer/Private/Sequencer.h b/Engine/Source/Editor/Sequencer/Private/Sequencer.h
--- a/Engine/Source/Editor/Sequencer/Private/Sequencer.h
+++ b/Engine/Source/Editor/Sequencer/Private/Sequencer.h
@@ -8,6 +8,7 @@
 #include "MovieScenePossessable.h"
 #include "SequencerLabelManager.h"
 #include "LevelEditor.h"
+#include "SequencerTimingManager.h"
 
 class ACineCameraActor;
 class FMenuBuilder;
@@ -547,6 +548,7 @@ public:
 	virtual void SetKeyInterpolation(EMovieSceneKeyInterpolation) override;
 	virtual bool GetInfiniteKeyAreas() const override;
 	virtual void SetInfiniteKeyAreas(bool bInfiniteKeyAreas) override;
+	virtual bool GetAutoSetTrackDefaults() const override;
 	virtual bool IsRecordingLive() const override;
 	virtual float GetCurrentLocalTime(UMovieSceneSequence& InMovieSceneSequence) override;
 	virtual float GetGlobalTime() const override;
@@ -853,6 +855,10 @@ protected:
 	void PossessPIEViewports(UObject* CameraObject, UObject* UnlockIfCameraObject, bool bJumpCut);
 
 private:
+
+	/** Reset the timing manager to default, or audio clock locked */
+	void ResetTimingManager(bool bUseAudioClock);
+
 	/** Performs any post-tick rendering work needed when moving through scenes */
 	void PostTickRenderStateFixup();
 
@@ -996,6 +1002,9 @@ private:
 	/** The maximum tick rate prior to playing (used for overriding delta time during playback). */
 	double OldMaxTickRate;
 
+	/** Timing manager that can adjust playback times */
+	TUniquePtr<FSequencerTimingManager> TimingManager;
+	
 	struct FCachedViewTarget
 	{
 		/** The player controller we're possessing */
diff --git a/Engine/Source/Editor/Sequencer/Private/SequencerDetailKeyframeHandler.cpp b/Engine/Source/Editor/Sequencer/Private/SequencerDetailKeyframeHandler.cpp
--- a/Engine/Source/Editor/Sequencer/Private/SequencerDetailKeyframeHandler.cpp
+++ b/Engine/Source/Editor/Sequencer/Private/SequencerDetailKeyframeHandler.cpp
@@ -23,12 +23,7 @@ void FSequencerDetailKeyframeHandler::OnKeyPropertyClicked(const IPropertyHandle
 	TArray<UObject*> Objects;
 	KeyedPropertyHandle.GetOuterObjects( Objects );
 
-	FKeyPropertyParams KeyPropertyParams(Objects, KeyedPropertyHandle);
-	KeyPropertyParams.KeyParams.bCreateHandleIfMissing = true;
-	KeyPropertyParams.KeyParams.bCreateTrackIfMissing = true;
-	KeyPropertyParams.KeyParams.bCreateKeyIfUnchanged = true;
-	KeyPropertyParams.KeyParams.bCreateKeyIfEmpty = true;
-	KeyPropertyParams.KeyParams.bCreateKeyOnlyWhenAutoKeying = false;
+	FKeyPropertyParams KeyPropertyParams(Objects, KeyedPropertyHandle, ESequencerKeyMode::ManualKeyForced);
 
 	Sequencer.Pin()->KeyProperty(KeyPropertyParams);
 }
diff --git a/Engine/Source/Editor/Sequencer/Private/SequencerNodeTree.cpp b/Engine/Source/Editor/Sequencer/Private/SequencerNodeTree.cpp
--- a/Engine/Source/Editor/Sequencer/Private/SequencerNodeTree.cpp
+++ b/Engine/Source/Editor/Sequencer/Private/SequencerNodeTree.cpp
@@ -256,8 +256,15 @@ void FSequencerNodeTree::MakeSectionInterfaces( UMovieSceneTrack& Track, TShared
 
 	for (int32 SectionIndex = 0; SectionIndex < MovieSceneSections.Num(); ++SectionIndex )
 	{
+		FGuid ObjectBinding;
+		if (SectionAreaNode->GetParent().IsValid() && SectionAreaNode->GetParent()->GetType() == ESequencerNode::Object)
+		{
+			TSharedPtr<FSequencerObjectBindingNode> ParentObjectNode = StaticCastSharedPtr<FSequencerObjectBindingNode>(SectionAreaNode->GetParent());
+			ObjectBinding = ParentObjectNode->GetObjectBinding();
+		}
+
 		UMovieSceneSection* SectionObject = MovieSceneSections[SectionIndex];
-		TSharedRef<ISequencerSection> Section = Editor->MakeSectionInterface( *SectionObject, Track );
+		TSharedRef<ISequencerSection> Section = Editor->MakeSectionInterface( *SectionObject, Track, ObjectBinding );
 
 		// Ask the section to generate it's inner layout
 		FSequencerSectionLayoutBuilder Builder( SectionAreaNode );
diff --git a/Engine/Source/Editor/Sequencer/Private/SequencerObjectChangeListener.cpp b/Engine/Source/Editor/Sequencer/Private/SequencerObjectChangeListener.cpp
--- a/Engine/Source/Editor/Sequencer/Private/SequencerObjectChangeListener.cpp
+++ b/Engine/Source/Editor/Sequencer/Private/SequencerObjectChangeListener.cpp
@@ -26,7 +26,7 @@ FSequencerObjectChangeListener::~FSequencerObjectChangeListener()
 
 void FSequencerObjectChangeListener::OnPropertyChanged(const TArray<UObject*>& ChangedObjects, const IPropertyHandle& PropertyHandle) const
 {
-	BroadcastPropertyChanged(FKeyPropertyParams(ChangedObjects, PropertyHandle));
+	BroadcastPropertyChanged(FKeyPropertyParams(ChangedObjects, PropertyHandle, ESequencerKeyMode::AutoKey));
 
 	for (UObject* Object : ChangedObjects)
 	{
@@ -50,9 +50,20 @@ void FSequencerObjectChangeListener::BroadcastPropertyChanged( FKeyPropertyParam
 	TArray<UObject*> KeyableObjects;
 	for (auto ObjectToKey : KeyPropertyParams.ObjectsToKey)
 	{
-		if (CanKeyProperty(FCanKeyPropertyParams(ObjectToKey->GetClass(), KeyPropertyParams.PropertyPath)))
+		if (KeyPropertyParams.PropertyPath.Num() > 0)
 		{
-			KeyableObjects.Add(ObjectToKey);
+			for (TFieldIterator<UProperty> PropertyIterator(ObjectToKey->GetClass()); PropertyIterator; ++PropertyIterator)
+			{
+				UProperty* Property = *PropertyIterator;
+				if (Property == KeyPropertyParams.PropertyPath[0])
+				{
+					if (CanKeyProperty(FCanKeyPropertyParams(ObjectToKey->GetClass(), KeyPropertyParams.PropertyPath)))
+					{
+						KeyableObjects.Add(ObjectToKey);
+						break;
+					}
+				}
+			}
 		}
 	}
 
@@ -68,29 +79,29 @@ void FSequencerObjectChangeListener::BroadcastPropertyChanged( FKeyPropertyParam
 		ParentStructProperty = Cast<const UStructProperty>(KeyPropertyParams.PropertyPath[KeyPropertyParams.PropertyPath.Num() - 2]);
 	}
 
-	FPropertyChangedParams Params;
-	Params.KeyParams = KeyPropertyParams.KeyParams;
-	Params.ObjectsThatChanged = KeyableObjects;
-	Params.StructPropertyNameToKey = NAME_None;
+	TArray<UProperty*> PropertyPath;
+	FName StructPropertyNameToKey = NAME_None;
 
 	bool bFoundAndBroadcastedDelegate = false;
 	if (ParentStructProperty)
 	{
-		Params.PropertyPath.Append(KeyPropertyParams.PropertyPath.GetData(), KeyPropertyParams.PropertyPath.Num() - 1);
+		PropertyPath.Append(KeyPropertyParams.PropertyPath.GetData(), KeyPropertyParams.PropertyPath.Num() - 1);
 
 		// If the property parent is a struct, see if this property parent can be keyed. (e.g R,G,B,A for a color)
 		FOnAnimatablePropertyChanged Delegate = ClassToPropertyChangedMap.FindRef(ParentStructProperty->Struct->GetFName());
 		if (Delegate.IsBound())
 		{
-			Params.StructPropertyNameToKey = KeyPropertyParams.PropertyPath.Last()->GetFName();
+			StructPropertyNameToKey = KeyPropertyParams.PropertyPath.Last()->GetFName();
 			bFoundAndBroadcastedDelegate = true;
+			FPropertyChangedParams Params(KeyableObjects, PropertyPath, StructPropertyNameToKey, KeyPropertyParams.KeyMode);
 			Delegate.Broadcast(Params);
 		}
 	}
 
 	if (!bFoundAndBroadcastedDelegate)
 	{
-		Params.PropertyPath = KeyPropertyParams.PropertyPath;
+		PropertyPath = KeyPropertyParams.PropertyPath;
+		FPropertyChangedParams Params(KeyableObjects, PropertyPath, StructPropertyNameToKey, KeyPropertyParams.KeyMode);
 		if (StructProperty)
 		{
 			ClassToPropertyChangedMap.FindRef(StructProperty->Struct->GetFName()).Broadcast(Params);
diff --git a/Engine/Source/Editor/Sequencer/Private/SequencerSettings.cpp b/Engine/Source/Editor/Sequencer/Private/SequencerSettings.cpp
--- a/Engine/Source/Editor/Sequencer/Private/SequencerSettings.cpp
+++ b/Engine/Source/Editor/Sequencer/Private/SequencerSettings.cpp
@@ -9,6 +9,7 @@ USequencerSettings::USequencerSettings( const FObjectInitializer& ObjectInitiali
 	bKeyAllEnabled = false;
 	bKeyInterpPropertiesOnly = false;
 	KeyInterpolation = EMovieSceneKeyInterpolation::Auto;
+	bAutoSetTrackDefaults = false;
 	SpawnPosition = SSP_Origin;
 	bCreateSpawnableCameras = true;
 	bShowFrameNumbers = true;
@@ -38,9 +39,21 @@ USequencerSettings::USequencerSettings( const FObjectInitializer& ObjectInitiali
 	bInfiniteKeyAreas = false;
 	bShowChannelColors = false;
 	bShowViewportTransportControls = true;
+	bLockPlaybackToAudioClock = false;
 	bAllowPossessionOfPIEViewports = false;
 }
 
+void USequencerSettings::PostEditChangeProperty(FPropertyChangedEvent& PropertyChangedEvent)
+{
+	FName PropertyName = (PropertyChangedEvent.Property != nullptr) ? PropertyChangedEvent.Property->GetFName() : NAME_None;
+	if (PropertyName == GET_MEMBER_NAME_CHECKED(USequencerSettings, bLockPlaybackToAudioClock))
+	{
+		OnLockPlaybackToAudioClockChanged.Broadcast(bLockPlaybackToAudioClock);
+	}
+
+	Super::PostEditChangeProperty(PropertyChangedEvent);
+}
+
 EAutoKeyMode USequencerSettings::GetAutoKeyMode() const
 {
 	return AutoKeyMode;
@@ -529,6 +542,35 @@ void USequencerSettings::SetAllowPossessionOfPIEViewports(bool bInAllowPossessio
 	}
 }
 
+bool USequencerSettings::ShouldLockPlaybackToAudioClock() const
+{
+	return bLockPlaybackToAudioClock;
+}
+
+void USequencerSettings::SetLockPlaybackToAudioClock(bool bInLockPlaybackToAudioClock)
+{
+	if (bLockPlaybackToAudioClock != bInLockPlaybackToAudioClock)
+	{
+		bLockPlaybackToAudioClock = bInLockPlaybackToAudioClock;
+		OnLockPlaybackToAudioClockChanged.Broadcast(bLockPlaybackToAudioClock);
+		SaveConfig();
+	}
+}
+
+bool USequencerSettings::GetAutoSetTrackDefaults() const
+{
+	return bAutoSetTrackDefaults;
+}
+
+void USequencerSettings::SetAutoSetTrackDefaults(bool bInAutoSetTrackDefaults)
+{
+	if (bInAutoSetTrackDefaults != bAutoSetTrackDefaults)
+	{
+		bAutoSetTrackDefaults = bInAutoSetTrackDefaults;
+		SaveConfig();
+	}
+}
+
 USequencerSettings::FOnTimeSnapIntervalChanged& USequencerSettings::GetOnTimeSnapIntervalChanged()
 {
 	return OnTimeSnapIntervalChanged;
diff --git a/Engine/Source/Editor/Sequencer/Private/SequencerTimingManager.cpp b/Engine/Source/Editor/Sequencer/Private/SequencerTimingManager.cpp
new file mode 100644
--- /dev/null
+++ b/Engine/Source/Editor/Sequencer/Private/SequencerTimingManager.cpp
@@ -0,0 +1,52 @@
+// Copyright 1998-2016 Epic Games, Inc. All Rights Reserved.
+
+#include "SequencerPrivatePCH.h"
+#include "SequencerTimingManager.h"
+#include "AudioDevice.h"
+
+FTimeAndDelta FSequencerDefaultTimingManager::AdjustTime(float InCurrentTime, float InDelta, float InPlayRate)
+{
+	FTimeAndDelta RetVal;
+	RetVal.Delta = InDelta;
+	RetVal.Time = InCurrentTime + InDelta * GWorld->GetWorldSettings()->MatineeTimeDilation * InPlayRate;
+
+	return RetVal;
+}
+
+void FSequencerAudioClockTimer::OnStartPlaying(float InStartTime)
+{
+	PlaybackStartAudioClock = GEngine->GetMainAudioDevice()->GetAudioClock();
+	PlaybackStartTime = InStartTime;
+	bIsPlaying = true;
+}
+
+void FSequencerAudioClockTimer::OnStopPlaying(float InStopTime)
+{
+	bIsPlaying = false;
+}
+
+FTimeAndDelta FSequencerAudioClockTimer::AdjustTime(float InCurrentTime, float InDelta, float InPlayRate)
+{
+	if (!bIsPlaying)
+	{
+		// If we're not playing, we just use the default impl
+		return FSequencerDefaultTimingManager::AdjustTime(InCurrentTime, InDelta, InPlayRate);
+	}
+
+	const double Now = GEngine->GetMainAudioDevice()->GetAudioClock();
+	const double AbsoluteAudioClockDelta = Now - PlaybackStartAudioClock;
+
+	FTimeAndDelta RetVal;
+	if (LastAudioClock.IsSet())
+	{
+		RetVal.Delta = float(Now - LastAudioClock.GetValue());
+	}
+	else
+	{
+		RetVal.Delta = AbsoluteAudioClockDelta;
+	}
+	RetVal.Time = PlaybackStartTime + float(AbsoluteAudioClockDelta);
+
+	LastAudioClock = Now;
+	return RetVal;
+}
\ No newline at end of file
diff --git a/Engine/Source/Editor/Sequencer/Private/SequencerTimingManager.h b/Engine/Source/Editor/Sequencer/Private/SequencerTimingManager.h
new file mode 100644
--- /dev/null
+++ b/Engine/Source/Editor/Sequencer/Private/SequencerTimingManager.h
@@ -0,0 +1,68 @@
+// Copyright 1998-2016 Epic Games, Inc. All Rights Reserved.
+
+#pragma once
+
+struct FTimeAndDelta
+{
+	float Time;
+	float Delta;
+};
+
+struct FSequencerTimingManager
+{
+	virtual ~FSequencerTimingManager(){}
+
+public:
+	void Update(EMovieScenePlayerStatus::Type InStatus, float InCurrentTime)
+	{
+		if (Status.IsSet() && InStatus != Status.GetValue())
+		{
+			OnStopPlaying(InCurrentTime);
+			Status.Reset();
+		}
+
+		if (!Status.IsSet() && (InStatus == EMovieScenePlayerStatus::Playing || InStatus == EMovieScenePlayerStatus::Recording))
+		{
+			Status = InStatus;
+			OnStartPlaying(InCurrentTime);
+		}
+	}
+
+public:
+
+	virtual void OnStartPlaying(float InStartTime) {}
+	virtual void OnStopPlaying(float InStopTime) {}
+	virtual FTimeAndDelta AdjustTime(float InCurrentTime, float InDelta, float InPlayRate) = 0;
+
+private:
+	TOptional<EMovieScenePlayerStatus::Type> Status;
+};
+
+/** Default timing manager that is forced to the audio device clock */
+struct FSequencerDefaultTimingManager : FSequencerTimingManager
+{
+	virtual FTimeAndDelta AdjustTime(float InCurrentTime, float InDelta, float InPlayRate) override;
+};
+
+/**
+ * Playback timing manager that is forced to the audio device clock for accurate syncing of animation to audio.
+ * @note Does not respect slomo tracks, since they accrue innacuracies each delta.
+ */
+struct FSequencerAudioClockTimer : FSequencerDefaultTimingManager
+{
+	FSequencerAudioClockTimer()
+	{
+		PlaybackStartTime = PlaybackStartAudioClock = 0;
+		bIsPlaying = false;
+	}
+
+	double PlaybackStartTime;
+	double PlaybackStartAudioClock;
+	TOptional<double> LastAudioClock;
+	bool bIsPlaying;
+	
+	virtual void OnStartPlaying(float InStartTime) override;
+	virtual void OnStopPlaying(float InStopTime) override;
+
+	virtual FTimeAndDelta AdjustTime(float InCurrentTime, float InDelta, float InPlayRate) override;
+};
\ No newline at end of file
diff --git a/Engine/Source/Editor/Sequencer/Private/Tools/SequencerEditTool_Movement.cpp b/Engine/Source/Editor/Sequencer/Private/Tools/SequencerEditTool_Movement.cpp
--- a/Engine/Source/Editor/Sequencer/Private/Tools/SequencerEditTool_Movement.cpp
+++ b/Engine/Source/Editor/Sequencer/Private/Tools/SequencerEditTool_Movement.cpp
@@ -173,7 +173,7 @@ TSharedPtr<ISequencerEditToolDragOperation> FSequencerEditTool_Movement::CreateD
 			}
 
 			// @todo sequencer: Make this a customizable UI command modifier?
-			if (MouseEvent.IsAltDown())
+			if (MouseEvent.IsAltDown() || MouseEvent.GetEffectingButton() == EKeys::MiddleMouseButton)
 			{
 				return MakeShareable( new FDuplicateKeys( Sequencer, Selection.GetSelectedKeys() ) );
 			}
@@ -204,6 +204,11 @@ FReply FSequencerEditTool_Movement::OnMouseButtonUp(SWidget& OwnerWidget, const
 		DragOperation->OnEndDrag(MouseEvent, MyGeometry.AbsoluteToLocal(MouseEvent.GetScreenSpacePosition()), SequencerWidget.Pin()->GetVirtualTrackArea());
 		DragOperation = nullptr;
 
+		if (MouseEvent.GetEffectingButton() == EKeys::MiddleMouseButton)
+		{
+			GEditor->EndTransaction();
+		}
+
 		// Only return handled if we actually started a drag
 		return FReply::Handled().ReleaseMouseCapture();
 	}
diff --git a/Engine/Source/Editor/Sequencer/Public/ISequencer.h b/Engine/Source/Editor/Sequencer/Public/ISequencer.h
--- a/Engine/Source/Editor/Sequencer/Public/ISequencer.h
+++ b/Engine/Source/Editor/Sequencer/Public/ISequencer.h
@@ -75,6 +75,8 @@ enum class EMovieSceneDataChangeType
 {
 	/** Data owned by a track has been modified such as adding or removing keys, or changing their values. */
 	TrackValueChanged,
+	/** Data owned by a track has been modified such as adding or removing keys, or changing their values. Refresh immediately. */
+	TrackValueChangedRefreshImmediately,
 	/** The structure of the movie scene has changed by adding folders, object bindings, tracks, or sections. */
 	MovieSceneStructureItemAdded,
 	/** The structure of the movie scene has changed by removing folders, object bindings, tracks, or sections. */
@@ -201,6 +203,9 @@ public:
 	/** Set infinite key area default */
 	virtual void SetInfiniteKeyAreas(bool bInfiniteKeyAreas) = 0;
 
+	/** Gets whether or not property track defaults will be automatically set when adding tracks. */
+	virtual bool GetAutoSetTrackDefaults() const = 0;
+
 	/** @return Returns whether sequencer is currently recording live data from simulated actors */
 	virtual bool IsRecordingLive() const = 0;
 
diff --git a/Engine/Source/Editor/Sequencer/Public/ISequencerTrackEditor.h b/Engine/Source/Editor/Sequencer/Public/ISequencerTrackEditor.h
--- a/Engine/Source/Editor/Sequencer/Public/ISequencerTrackEditor.h
+++ b/Engine/Source/Editor/Sequencer/Public/ISequencerTrackEditor.h
@@ -101,8 +101,9 @@ public:
 	 *
 	 * @param SectionObject The section to make UI for.
 	 * @param Track The track that owns the section.
+	 * @param ObjectBinding the object binding for the track that owns the section, if there is one.
 	 */
-	virtual TSharedRef<ISequencerSection> MakeSectionInterface(UMovieSceneSection& SectionObject, UMovieSceneTrack& Track) = 0;
+	virtual TSharedRef<ISequencerSection> MakeSectionInterface(UMovieSceneSection& SectionObject, UMovieSceneTrack& Track, FGuid ObjectBinding) = 0;
 
 	/** Gets an icon brush for this track editor */
 	virtual const FSlateBrush* GetIconBrush() const { return nullptr; }
diff --git a/Engine/Source/Editor/Sequencer/Public/KeyPropertyParams.h b/Engine/Source/Editor/Sequencer/Public/KeyPropertyParams.h
--- a/Engine/Source/Editor/Sequencer/Public/KeyPropertyParams.h
+++ b/Engine/Source/Editor/Sequencer/Public/KeyPropertyParams.h
@@ -13,8 +13,6 @@ class UProperty;
  */
 struct SEQUENCER_API FCanKeyPropertyParams
 {
-	FCanKeyPropertyParams();
-
 	/**
 	 * Creates new can key property parameters.
 	 * @param InObjectClass the class of the object which has the property to be keyed.
@@ -32,9 +30,10 @@ struct SEQUENCER_API FCanKeyPropertyParams
 
 	/** The class of the object which has the property to be keyed. */
 	const UClass* ObjectClass;
+
 	/** An array of UProperty objects which represents a path of properties to get from the root object to the
 	property to be keyed. */
-	TArray<UProperty*> PropertyPath;
+	const TArray<UProperty*> PropertyPath;
 };
 
 /**
@@ -42,30 +41,30 @@ struct SEQUENCER_API FCanKeyPropertyParams
  */
 struct SEQUENCER_API FKeyPropertyParams
 {
-	FKeyPropertyParams();
-
 	/**
 	* Creates new key property parameters for a manually triggered property change.
 	* @param InObjectsToKey an array of the objects who's property will be keyed.
 	* @param InPropertyPath an array of UProperty objects which represents a path of properties to get from
 	*        the root object to the property to be keyed.
 	*/
-	FKeyPropertyParams(TArray<UObject*> InObjectsToKey, TArray<UProperty*> InPropertyPath);
+	FKeyPropertyParams(TArray<UObject*> InObjectsToKey, TArray<UProperty*> InPropertyPath, ESequencerKeyMode InKeyMode);
 
 	/**
 	* Creates new key property parameters from an actual property change notification with a property handle.
 	* @param InObjectsToKey an array of the objects who's property will be keyed.
 	* @param InPropertyHandle a handle to the property to be keyed.
 	*/
-	FKeyPropertyParams(TArray<UObject*> InObjectsToKey, const IPropertyHandle& InPropertyHandle);
+	FKeyPropertyParams(TArray<UObject*> InObjectsToKey, const IPropertyHandle& InPropertyHandle, ESequencerKeyMode InKeyMode);
 
 	/** An array of the objects who's property will be keyed. */
-	TArray<UObject*> ObjectsToKey;
+	const TArray<UObject*> ObjectsToKey;
+
 	/** An array of UProperty objects which represents a path of properties to get from the root object to the
 	property to be keyed. */
-	TArray<UProperty*> PropertyPath;
+	const TArray<UProperty*> PropertyPath;
+
 	/** Keyframing params */
-	FKeyParams KeyParams;
+	const ESequencerKeyMode KeyMode;
 };
 
 /**
@@ -74,6 +73,8 @@ struct SEQUENCER_API FKeyPropertyParams
 class SEQUENCER_API FPropertyChangedParams
 {
 public:
+	FPropertyChangedParams(TArray<UObject*> InObjectsThatChanged, TArray<UProperty*> InPropertyPath, FName InStructPropertyNameToKey, ESequencerKeyMode InKeyMode);
+
 	/**
 	 * Gets the value of the property that changed.
 	 */
@@ -108,13 +109,16 @@ public:
 	FString GetPropertyPathString() const;
 
 	/** An array of the objects that changed. */
-	TArray<UObject*> ObjectsThatChanged;
+	const TArray<UObject*> ObjectsThatChanged;
+
 	/** An array of UProperty objects which represents a path of properties to get from the root object to the 
 	property that changed. */
-	TArray<UProperty*> PropertyPath;
+	const TArray<UProperty*> PropertyPath;
+
 	/** Represents the FName of an inner property which should be keyed for a struct property.  If all inner 
 	properties should be keyed, this will be FName::None. */
-	FName StructPropertyNameToKey;
+	const FName StructPropertyNameToKey;
+
 	/** Keyframing params */
-	FKeyParams KeyParams;
+	const ESequencerKeyMode KeyMode;
 };
\ No newline at end of file
diff --git a/Engine/Source/Editor/Sequencer/Public/MovieSceneTrackEditor.h b/Engine/Source/Editor/Sequencer/Public/MovieSceneTrackEditor.h
--- a/Engine/Source/Editor/Sequencer/Public/MovieSceneTrackEditor.h
+++ b/Engine/Source/Editor/Sequencer/Public/MovieSceneTrackEditor.h
@@ -221,7 +221,6 @@ public:
 		return Sequencer.Pin()->IsRecordingLive() || Sequencer.Pin()->GetAutoKeyMode() != EAutoKeyMode::KeyNone;
 	}
 
-	virtual TSharedRef<ISequencerSection> MakeSectionInterface(class UMovieSceneSection& SectionObject, UMovieSceneTrack& Track) = 0;
 	virtual void OnInitialize() override { };
 	virtual void OnRelease() override { };
 
diff --git a/Engine/Source/Editor/Sequencer/Public/SequencerSettings.h b/Engine/Source/Editor/Sequencer/Public/SequencerSettings.h
--- a/Engine/Source/Editor/Sequencer/Public/SequencerSettings.h
+++ b/Engine/Source/Editor/Sequencer/Public/SequencerSettings.h
@@ -66,6 +66,9 @@ public:
 	GENERATED_UCLASS_BODY()
 
 	DECLARE_MULTICAST_DELEGATE( FOnTimeSnapIntervalChanged );
+	DECLARE_MULTICAST_DELEGATE_OneParam( FOnLockPlaybackToAudioClockChanged, bool );
+
+	virtual void PostEditChangeProperty(FPropertyChangedEvent& PropertyChangedEvent) override;
 
 	/** Gets the current auto-key mode. */
 	EAutoKeyMode GetAutoKeyMode() const;
@@ -237,12 +240,25 @@ public:
 	/** Toggle whether to allow possession of PIE viewports */
 	void SetAllowPossessionOfPIEViewports(bool bInAllowPossessionOfPIEViewports);
 
+	/** Gets whether or not track defaults will be automatically set when modifying tracks. */
+	bool GetAutoSetTrackDefaults() const;
+	/** Sets whether or not track defaults will be automatically set when modifying tracks. */
+	void SetAutoSetTrackDefaults(bool bInAutoSetTrackDefaults);
+
 	/** Snaps a time value in seconds to the currently selected interval. */
 	float SnapTimeToInterval(float InTimeValue) const;
 
 	/** Gets the multicast delegate which is run whenever the time snap interval is changed. */
 	FOnTimeSnapIntervalChanged& GetOnTimeSnapIntervalChanged();
 
+	/** @return true we're locking playback to the audio clock */
+	bool ShouldLockPlaybackToAudioClock() const;
+	/** Toggle whether to lock playback to the audio clock */
+	void SetLockPlaybackToAudioClock(bool bInLockPlaybackToAudioClock);
+
+	/** Gets the multicast delegate which is invoked whenevcer the bLockPlaybackToAudioClock setting is changed. */
+	FOnLockPlaybackToAudioClockChanged& GetOnLockPlaybackToAudioClockChanged() { return OnLockPlaybackToAudioClockChanged; }
+
 protected:
 
 	/** Enable or disable autokeying. */
@@ -261,6 +277,10 @@ protected:
 	UPROPERTY( config, EditAnywhere, Category=Keyframing )
 	TEnumAsByte<EMovieSceneKeyInterpolation> KeyInterpolation;
 
+	/** Whether or not track defaults will be automatically set when modifying tracks. */
+	UPROPERTY( config, EditAnywhere, Category=Keyframing, meta = (ToolTip = "When setting keys on properties and transforms automatically update the track default values used when there are no keys."))
+	bool bAutoSetTrackDefaults;
+
 	/** The default location of a spawnable when it is first dragged into the viewport from the content browser. */
 	UPROPERTY( config, EditAnywhere, Category=General )
 	TEnumAsByte<ESequencerSpawnPosition> SpawnPosition;
@@ -376,9 +396,15 @@ protected:
 	UPROPERTY( config )
 	bool bShowViewportTransportControls;
 
+	/** When enabled, sequencer playback will be locked to the engine's audio clock */
+	UPROPERTY(config, EditAnywhere, Category=Playback)
+	bool bLockPlaybackToAudioClock;
+
 	/** When enabled, sequencer is able to possess viewports that represent PIE worlds */
 	UPROPERTY(config, EditAnywhere, Category=General)
 	bool bAllowPossessionOfPIEViewports;
 
+	FOnLockPlaybackToAudioClockChanged OnLockPlaybackToAudioClockChanged;
+
 	FOnTimeSnapIntervalChanged OnTimeSnapIntervalChanged;
 };
diff --git a/Engine/Source/Editor/UMGEditor/Private/Animation/MarginTrackEditor.cpp b/Engine/Source/Editor/UMGEditor/Private/Animation/MarginTrackEditor.cpp
--- a/Engine/Source/Editor/UMGEditor/Private/Animation/MarginTrackEditor.cpp
+++ b/Engine/Source/Editor/UMGEditor/Private/Animation/MarginTrackEditor.cpp
@@ -20,47 +20,63 @@ class FMarginPropertySection
 {
 public:
 
-	FMarginPropertySection( UMovieSceneSection& InSectionObject, const FText& SectionName )
-		: FPropertySection(InSectionObject, SectionName) {}
+	FMarginPropertySection(ISequencer* InSequencer, FGuid InObjectBinding, FName InPropertyName, const FString& InPropertyPath, UMovieSceneSection& InSectionObject, const FText& SectionName)
+		: FPropertySection(InSequencer, InObjectBinding, InPropertyName, InPropertyPath, InSectionObject, SectionName)
+	{
+	}
 
 	virtual void GenerateSectionLayout( class ISectionLayoutBuilder& LayoutBuilder ) const override
 	{
 		UMovieSceneMarginSection* MarginSection = Cast<UMovieSceneMarginSection>(&SectionObject);
 
-		LeftKeyArea = MakeShareable(new FFloatCurveKeyArea(&MarginSection->GetLeftCurve(), MarginSection));
-		TopKeyArea = MakeShareable(new FFloatCurveKeyArea(&MarginSection->GetTopCurve(), MarginSection));
-		RightKeyArea = MakeShareable(new FFloatCurveKeyArea(&MarginSection->GetRightCurve(), MarginSection));
-		BottomKeyArea = MakeShareable(new FFloatCurveKeyArea( &MarginSection->GetBottomCurve(), MarginSection));
+		TAttribute<TOptional<float>> MarginLeftExternalValue = TAttribute<TOptional<float>>::Create(
+			TAttribute<TOptional<float>>::FGetter::CreateRaw(this, &FMarginPropertySection::GetMarginLeftValue));
+		TSharedRef<FFloatCurveKeyArea> LeftKeyArea = MakeShareable(new FFloatCurveKeyArea(&MarginSection->GetLeftCurve(), MarginLeftExternalValue, MarginSection));
+
+		TAttribute<TOptional<float>> MarginTopExternalValue = TAttribute<TOptional<float>>::Create(
+			TAttribute<TOptional<float>>::FGetter::CreateRaw(this, &FMarginPropertySection::GetMarginTopValue));
+		TSharedRef<FFloatCurveKeyArea> TopKeyArea = MakeShareable(new FFloatCurveKeyArea(&MarginSection->GetTopCurve(), MarginTopExternalValue, MarginSection));
+
+		TAttribute<TOptional<float>> MarginRightExternalValue = TAttribute<TOptional<float>>::Create(
+			TAttribute<TOptional<float>>::FGetter::CreateRaw(this, &FMarginPropertySection::GetMarginRightValue));
+		TSharedRef<FFloatCurveKeyArea> RightKeyArea = MakeShareable(new FFloatCurveKeyArea(&MarginSection->GetRightCurve(), MarginRightExternalValue, MarginSection));
 
-		LayoutBuilder.AddKeyArea("Left", NSLOCTEXT("FMarginPropertySection", "MarginLeft", "Left"), LeftKeyArea.ToSharedRef());
-		LayoutBuilder.AddKeyArea("Top", NSLOCTEXT("FMarginPropertySection", "MarginTop", "Top"), TopKeyArea.ToSharedRef());
-		LayoutBuilder.AddKeyArea("Right", NSLOCTEXT("FMarginPropertySection", "MarginRight", "Right"), RightKeyArea.ToSharedRef());
-		LayoutBuilder.AddKeyArea("Bottom", NSLOCTEXT("FMarginPropertySection", "MarginBottom", "Bottom"), BottomKeyArea.ToSharedRef());
+		TAttribute<TOptional<float>> MarginBottomExternalValue = TAttribute<TOptional<float>>::Create(
+			TAttribute<TOptional<float>>::FGetter::CreateRaw(this, &FMarginPropertySection::GetMarginBottomValue));
+		TSharedRef<FFloatCurveKeyArea> BottomKeyArea = MakeShareable(new FFloatCurveKeyArea( &MarginSection->GetBottomCurve(), MarginBottomExternalValue, MarginSection));
+
+		LayoutBuilder.AddKeyArea("Left", NSLOCTEXT("FMarginPropertySection", "MarginLeft", "Left"), LeftKeyArea);
+		LayoutBuilder.AddKeyArea("Top", NSLOCTEXT("FMarginPropertySection", "MarginTop", "Top"), TopKeyArea);
+		LayoutBuilder.AddKeyArea("Right", NSLOCTEXT("FMarginPropertySection", "MarginRight", "Right"), RightKeyArea);
+		LayoutBuilder.AddKeyArea("Bottom", NSLOCTEXT("FMarginPropertySection", "MarginBottom", "Bottom"), BottomKeyArea);
 	}
 
-	virtual void SetIntermediateValue( FPropertyChangedParams PropertyChangedParams ) override
+
+private:
+
+	TOptional<float> GetMarginLeftValue() const
 	{
-		FMargin Margin = PropertyChangedParams.GetPropertyValue<FMargin>();
-		LeftKeyArea->SetIntermediateValue( Margin.Left );
-		TopKeyArea->SetIntermediateValue( Margin.Top );
-		RightKeyArea->SetIntermediateValue( Margin.Right );
-		BottomKeyArea->SetIntermediateValue( Margin.Bottom );
+		TOptional<FMargin> Margin = GetPropertyValue<FMargin>();
+		return Margin.IsSet() ? TOptional<float>(Margin.GetValue().Left) : TOptional<float>();
 	}
 
+	TOptional<float> GetMarginTopValue() const
+	{
+		TOptional<FMargin> Margin = GetPropertyValue<FMargin>();
+		return Margin.IsSet() ? TOptional<float>(Margin.GetValue().Top) : TOptional<float>();
+	}
 
-	virtual void ClearIntermediateValue() override
+	TOptional<float> GetMarginRightValue() const
 	{
-		LeftKeyArea->ClearIntermediateValue();
-		TopKeyArea->ClearIntermediateValue();
-		RightKeyArea->ClearIntermediateValue();
-		BottomKeyArea->ClearIntermediateValue();
+		TOptional<FMargin> Margin = GetPropertyValue<FMargin>();
+		return Margin.IsSet() ? TOptional<float>(Margin.GetValue().Right) : TOptional<float>();
 	}
 
-private:
-	mutable TSharedPtr<FFloatCurveKeyArea> LeftKeyArea;
-	mutable TSharedPtr<FFloatCurveKeyArea> TopKeyArea;
-	mutable TSharedPtr<FFloatCurveKeyArea> RightKeyArea;
-	mutable TSharedPtr<FFloatCurveKeyArea> BottomKeyArea;
+	TOptional<float> GetMarginBottomValue() const
+	{
+		TOptional<FMargin> Margin = GetPropertyValue<FMargin>();
+		return Margin.IsSet() ? TOptional<float>(Margin.GetValue().Bottom) : TOptional<float>();
+	}
 };
 
 
@@ -70,12 +86,11 @@ TSharedRef<ISequencerTrackEditor> FMarginTrackEditor::CreateTrackEditor( TShared
 }
 
 
-TSharedRef<FPropertySection> FMarginTrackEditor::MakePropertySectionInterface( UMovieSceneSection& SectionObject, UMovieSceneTrack& Track )
+TSharedRef<ISequencerSection> FMarginTrackEditor::MakeSectionInterface( UMovieSceneSection& SectionObject, UMovieSceneTrack& Track, FGuid ObjectBinding )
 {
-	check( SupportsType( SectionObject.GetOuter()->GetClass() ) );
-
-	UClass* SectionClass = SectionObject.GetOuter()->GetClass();
-	return MakeShareable(new FMarginPropertySection(SectionObject, Track.GetDisplayName()));
+	UMovieScenePropertyTrack* PropertyTrack = Cast<UMovieScenePropertyTrack>(&Track);
+	checkf(PropertyTrack != nullptr, TEXT("Incompatible track in FMarginTrackEditor"));
+	return MakeShareable(new FMarginPropertySection(GetSequencer().Get(), ObjectBinding, PropertyTrack->GetPropertyName(), PropertyTrack->GetPropertyPath(), SectionObject, Track.GetDisplayName()));
 }
 
 
diff --git a/Engine/Source/Editor/UMGEditor/Private/Animation/MarginTrackEditor.h b/Engine/Source/Editor/UMGEditor/Private/Animation/MarginTrackEditor.h
--- a/Engine/Source/Editor/UMGEditor/Private/Animation/MarginTrackEditor.h
+++ b/Engine/Source/Editor/UMGEditor/Private/Animation/MarginTrackEditor.h
@@ -34,11 +34,14 @@ public:
 	 */
 	static TSharedRef<ISequencerTrackEditor> CreateTrackEditor( TSharedRef<ISequencer> OwningSequencer );
 
+	//~ ISequencerTrackEditor Interface
+
+	virtual TSharedRef<ISequencerSection> MakeSectionInterface(UMovieSceneSection& SectionObject, UMovieSceneTrack& Track, FGuid ObjectBinding) override;
+
 protected:
 
-	// FPropertyTrackEditor Interface
+	//~ FPropertyTrackEditor Interface
 
-	virtual TSharedRef<FPropertySection> MakePropertySectionInterface( UMovieSceneSection& SectionObject, UMovieSceneTrack& Track ) override;
 	virtual void GenerateKeysFromPropertyChanged( const FPropertyChangedParams& PropertyChangedParams, TArray<FMarginKey>& NewGeneratedKeys, TArray<FMarginKey>& DefaultGeneratedKeys) override;
 
 private:
