Sequencer: Remove stray pragma once, causing a non unity compile error on mac

diff --git a/Engine/Source/Editor/MovieSceneTools/Private/PropertySection.cpp b/Engine/Source/Editor/MovieSceneTools/Private/PropertySection.cpp
--- a/Engine/Source/Editor/MovieSceneTools/Private/PropertySection.cpp
+++ b/Engine/Source/Editor/MovieSceneTools/Private/PropertySection.cpp
@@ -1,7 +1,5 @@
 // Copyright 1998-2016 Epic Games, Inc. All Rights Reserved.
 
-#pragma once
-
 #include "MovieSceneToolsPrivatePCH.h"
 #include "PropertySection.h"
 
