Copying //UE4/Dev-Sequencer to //UE4/Dev-Main (Source: //UE4/Dev-Sequencer @ 3178529)

diff --git a/Engine/Source/Editor/MovieSceneTools/Private/Sections/CinematicShotSection.cpp b/Engine/Source/Editor/MovieSceneTools/Private/Sections/CinematicShotSection.cpp
--- a/Engine/Source/Editor/MovieSceneTools/Private/Sections/CinematicShotSection.cpp
+++ b/Engine/Source/Editor/MovieSceneTools/Private/Sections/CinematicShotSection.cpp
@@ -22,8 +22,8 @@ FCinematicShotSection::FCinematicSectionCache::FCinematicSectionCache(UMovieScen
 {
 	if (Section)
 	{
-		ActualStartTime = Section->GetStartTime() - Section->StartOffset;
-		TimeScale = Section->TimeScale;
+		ActualStartTime = Section->GetStartTime() - Section->Parameters.StartOffset;
+		TimeScale = Section->Parameters.TimeScale;
 	}
 }
 
@@ -37,10 +37,6 @@ FCinematicShotSection::FCinematicShotSection(TSharedPtr<ISequencer> InSequencer,
 	, InitialStartTimeDuringResize(0.f)
 	, ThumbnailCacheData(&SectionObject)
 {
-	if (InSequencer->HasSequenceInstanceForSection(InSection))
-	{
-		SequenceInstance = InSequencer->GetSequenceInstanceForSection(InSection);
-	}
 }
 
 
@@ -65,7 +61,7 @@ void FCinematicShotSection::SetSingleTime(float GlobalTime)
 
 void FCinematicShotSection::BeginResizeSection()
 {
-	InitialStartOffsetDuringResize = SectionObject.StartOffset;
+	InitialStartOffsetDuringResize = SectionObject.Parameters.StartOffset;
 	InitialStartTimeDuringResize = SectionObject.GetStartTime();
 }
 
@@ -74,13 +70,13 @@ void FCinematicShotSection::ResizeSection(ESequencerSectionResizeMode ResizeMode
 	// Adjust the start offset when resizing from the beginning
 	if (ResizeMode == SSRM_LeadingEdge)
 	{
-		float StartOffset = (ResizeTime - InitialStartTimeDuringResize) / SectionObject.TimeScale;
+		float StartOffset = (ResizeTime - InitialStartTimeDuringResize) / SectionObject.Parameters.TimeScale;
 		StartOffset += InitialStartOffsetDuringResize;
 
 		// Ensure start offset is not less than 0
 		StartOffset = FMath::Max(StartOffset, 0.f);
 
-		SectionObject.StartOffset = StartOffset;
+		SectionObject.Parameters.StartOffset = StartOffset;
 	}
 
 	FThumbnailSection::ResizeSection(ResizeMode, ResizeTime);
@@ -115,12 +111,12 @@ int32 FCinematicShotSection::OnPaintSection(FSequencerSectionPainter& InPainter)
 
 	InPainter.LayerId = InPainter.PaintSectionBackground();
 
-	FVector2D SectionSize = InPainter.SectionGeometry.GetLocalSize();
+	FVector2D LocalSectionSize = InPainter.SectionGeometry.GetLocalSize();
 
 	FSlateDrawElement::MakeBox(
 		InPainter.DrawElements,
 		InPainter.LayerId++,
-		InPainter.SectionGeometry.ToPaintGeometry(FVector2D(SectionSize.X-2.f, 7.f), FSlateLayoutTransform(FVector2D(1.f, 4.f))),
+		InPainter.SectionGeometry.ToPaintGeometry(FVector2D(LocalSectionSize.X-2.f, 7.f), FSlateLayoutTransform(FVector2D(1.f, 4.f))),
 		FilmBorder,
 		InPainter.SectionClippingRect.InsetBy(FMargin(1.f)),
 		InPainter.bParentEnabled ? ESlateDrawEffect::None : ESlateDrawEffect::DisabledEffect
@@ -129,15 +125,110 @@ int32 FCinematicShotSection::OnPaintSection(FSequencerSectionPainter& InPainter)
 	FSlateDrawElement::MakeBox(
 		InPainter.DrawElements,
 		InPainter.LayerId++,
-		InPainter.SectionGeometry.ToPaintGeometry(FVector2D(SectionSize.X-2.f, 7.f), FSlateLayoutTransform(FVector2D(1.f, SectionSize.Y - 11.f))),
+		InPainter.SectionGeometry.ToPaintGeometry(FVector2D(LocalSectionSize.X-2.f, 7.f), FSlateLayoutTransform(FVector2D(1.f, LocalSectionSize.Y - 11.f))),
 		FilmBorder,
 		InPainter.SectionClippingRect.InsetBy(FMargin(1.f)),
 		InPainter.bParentEnabled ? ESlateDrawEffect::None : ESlateDrawEffect::DisabledEffect
 	);
 
-	if (SequenceInstance.IsValid())
+	const float SectionSize = SectionObject.GetTimeSize();
+
+	if (SectionSize <= 0.0f)
+	{
+		return InPainter.LayerId;
+	}
+
+	FThumbnailSection::OnPaintSection(InPainter);
+
+	const float DrawScale = InPainter.SectionGeometry.Size.X / SectionSize;
+		
+	UMovieScene* MovieScene = nullptr;
+	FFloatRange PlaybackRange;
+	if(SectionObject.GetSequence() != nullptr)
+	{
+		MovieScene = SectionObject.GetSequence()->GetMovieScene();
+		PlaybackRange = MovieScene->GetPlaybackRange();
+	}
+	else
+	{
+		UMovieSceneTrack* MovieSceneTrack = CastChecked<UMovieSceneTrack>(SectionObject.GetOuter());
+		MovieScene = CastChecked<UMovieScene>(MovieSceneTrack->GetOuter());
+		PlaybackRange = MovieScene->GetPlaybackRange();
+	}
+
+	// add box for the working size
+	const float StartOffset = SectionObject.Parameters.TimeScale * SectionObject.Parameters.StartOffset;
+	const float WorkingStart = -SectionObject.Parameters.TimeScale * PlaybackRange.GetLowerBoundValue() - StartOffset;
+	const float WorkingSize = SectionObject.Parameters.TimeScale * (MovieScene != nullptr ? MovieScene->GetEditorData().WorkingRange.Size<float>() : 1.0f);
+
+	// add dark tint for left out-of-bounds & working range
+	if (StartOffset < 0.0f)
+	{
+		FSlateDrawElement::MakeBox(
+			InPainter.DrawElements,
+			InPainter.LayerId++,
+			InPainter.SectionGeometry.ToPaintGeometry(
+				FVector2D(0.0f, 0.f),
+				FVector2D(-StartOffset * DrawScale, InPainter.SectionGeometry.Size.Y)
+			),
+			FEditorStyle::GetBrush("WhiteBrush"),
+			InPainter.SectionClippingRect,
+			ESlateDrawEffect::None,
+			FLinearColor::Black.CopyWithNewOpacity(0.5f)
+		);
+	}
+
+	// add green line for playback start
+	if (StartOffset < 0)
 	{
-		return FThumbnailSection::OnPaintSection(InPainter);
+		FSlateDrawElement::MakeBox(
+			InPainter.DrawElements,
+			InPainter.LayerId++,
+			InPainter.SectionGeometry.ToPaintGeometry(
+				FVector2D(-StartOffset * DrawScale, 0.f),
+				FVector2D(1.0f, InPainter.SectionGeometry.Size.Y)
+			),
+			FEditorStyle::GetBrush("WhiteBrush"),
+			InPainter.SectionClippingRect,
+			ESlateDrawEffect::None,
+			FColor(32, 128, 32)	// 120, 75, 50 (HSV)
+		);
+	}
+
+	// add dark tint for right out-of-bounds & working range
+	const float PlaybackEnd = SectionObject.Parameters.TimeScale * PlaybackRange.Size<float>() - StartOffset;
+
+	if (PlaybackEnd < SectionSize)
+	{
+		FSlateDrawElement::MakeBox(
+			InPainter.DrawElements,
+			InPainter.LayerId++,
+			InPainter.SectionGeometry.ToPaintGeometry(
+				FVector2D(PlaybackEnd * DrawScale, 0.f),
+				FVector2D((SectionSize - PlaybackEnd) * DrawScale, InPainter.SectionGeometry.Size.Y)
+			),
+			FEditorStyle::GetBrush("WhiteBrush"),
+			InPainter.SectionClippingRect,
+			ESlateDrawEffect::None,
+			FLinearColor::Black.CopyWithNewOpacity(0.5f)
+		);
+	}
+
+	// add red line for playback end
+	if (PlaybackEnd <= SectionSize)
+	{
+		FSlateDrawElement::MakeBox(
+			InPainter.DrawElements,
+			InPainter.LayerId++,
+			InPainter.SectionGeometry.ToPaintGeometry(
+				FVector2D(PlaybackEnd * DrawScale, 0.f),
+				FVector2D(1.0f, InPainter.SectionGeometry.Size.Y)
+			),
+			FEditorStyle::GetBrush("WhiteBrush"),
+			InPainter.SectionClippingRect,
+			ESlateDrawEffect::None,
+			FColor(128, 32, 32)	// 0, 75, 50 (HSV)
+		);
 	}
 
 	return InPainter.LayerId;
@@ -149,20 +240,17 @@ void FCinematicShotSection::BuildSectionContextMenu(FMenuBuilder& MenuBuilder, c
 
 	MenuBuilder.BeginSection(NAME_None, LOCTEXT("ShotMenuText", "Shot"));
 	{
-		if (SequenceInstance.IsValid())
-		{
-			MenuBuilder.AddSubMenu(
-				LOCTEXT("TakesMenu", "Takes"),
-				LOCTEXT("TakesMenuTooltip", "Shot takes"),
-				FNewMenuDelegate::CreateLambda([=](FMenuBuilder& InMenuBuilder){ AddTakesMenu(InMenuBuilder); }));
-
-			MenuBuilder.AddMenuEntry(
-				LOCTEXT("NewTake", "New Take"),
-				FText::Format(LOCTEXT("NewTakeTooltip", "Create a new take for {0}"), SectionObject.GetShotDisplayName()),
-				FSlateIcon(),
-				FUIAction(FExecuteAction::CreateSP(CinematicShotTrackEditor.Pin().ToSharedRef(), &FCinematicShotTrackEditor::NewTake, &SectionObject))
-			);
-		}
+		MenuBuilder.AddSubMenu(
+			LOCTEXT("TakesMenu", "Takes"),
+			LOCTEXT("TakesMenuTooltip", "Shot takes"),
+			FNewMenuDelegate::CreateLambda([=](FMenuBuilder& InMenuBuilder){ AddTakesMenu(InMenuBuilder); }));
+
+		MenuBuilder.AddMenuEntry(
+			LOCTEXT("NewTake", "New Take"),
+			FText::Format(LOCTEXT("NewTakeTooltip", "Create a new take for {0}"), SectionObject.GetShotDisplayName()),
+			FSlateIcon(),
+			FUIAction(FExecuteAction::CreateSP(CinematicShotTrackEditor.Pin().ToSharedRef(), &FCinematicShotTrackEditor::NewTake, &SectionObject))
+		);
 
 		MenuBuilder.AddMenuEntry(
 			LOCTEXT("InsertNewShot", "Insert Shot"),
@@ -171,32 +259,26 @@ void FCinematicShotSection::BuildSectionContextMenu(FMenuBuilder& MenuBuilder, c
 			FUIAction(FExecuteAction::CreateSP(CinematicShotTrackEditor.Pin().ToSharedRef(), &FCinematicShotTrackEditor::InsertShot, &SectionObject))
 		);
 
-		if (SequenceInstance.IsValid())
-		{
-			MenuBuilder.AddMenuEntry(
-				LOCTEXT("DuplicateShot", "Duplicate Shot"),
-				FText::Format(LOCTEXT("DuplicateShotTooltip", "Duplicate {0} to create a new shot"), SectionObject.GetShotDisplayName()),
-				FSlateIcon(),
-				FUIAction(FExecuteAction::CreateSP(CinematicShotTrackEditor.Pin().ToSharedRef(), &FCinematicShotTrackEditor::DuplicateShot, &SectionObject))
-			);
-
-			MenuBuilder.AddMenuEntry(
-				LOCTEXT("RenderShot", "Render Shot"),
-				FText::Format(LOCTEXT("RenderShotTooltip", "Render shot movie"), SectionObject.GetShotDisplayName()),
-				FSlateIcon(),
-				FUIAction(FExecuteAction::CreateSP(CinematicShotTrackEditor.Pin().ToSharedRef(), &FCinematicShotTrackEditor::RenderShot, &SectionObject))
-			);
-
-			/*
-			//@todo
-			MenuBuilder.AddMenuEntry(
-				LOCTEXT("RenameShot", "Rename Shot"),
-				FText::Format(LOCTEXT("RenameShotTooltip", "Rename {0}"), SectionObject.GetShotDisplayName()),
-				FSlateIcon(),
-				FUIAction(FExecuteAction::CreateSP(CinematicShotTrackEditor.Pin().ToSharedRef(), &FCinematicShotTrackEditor::RenameShot, &SectionObject))
-			);
-			*/
-		}
+		MenuBuilder.AddMenuEntry(
+			LOCTEXT("DuplicateShot", "Duplicate Shot"),
+			FText::Format(LOCTEXT("DuplicateShotTooltip", "Duplicate {0} to create a new shot"), SectionObject.GetShotDisplayName()),
+			FSlateIcon(),
+			FUIAction(FExecuteAction::CreateSP(CinematicShotTrackEditor.Pin().ToSharedRef(), &FCinematicShotTrackEditor::DuplicateShot, &SectionObject))
+		);
+
+		MenuBuilder.AddMenuEntry(
+			LOCTEXT("RenderShot", "Render Shot"),
+			FText::Format(LOCTEXT("RenderShotTooltip", "Render shot movie"), SectionObject.GetShotDisplayName()),
+			FSlateIcon(),
+			FUIAction(FExecuteAction::CreateSP(CinematicShotTrackEditor.Pin().ToSharedRef(), &FCinematicShotTrackEditor::RenderShot, &SectionObject))
+		);
+
+		MenuBuilder.AddMenuEntry(
+			LOCTEXT("RenameShot", "Rename Shot"),
+			FText::Format(LOCTEXT("RenameShotTooltip", "Rename {0}"), SectionObject.GetShotDisplayName()),
+			FSlateIcon(),
+			FUIAction(FExecuteAction::CreateSP(this, &FCinematicShotSection::EnterRename))
+		);
 	}
 	MenuBuilder.EndSection();
 }
@@ -248,12 +330,9 @@ FReply FCinematicShotSection::OnSectionDoubleClicked(const FGeometry& SectionGeo
 {
 	if( MouseEvent.GetEffectingButton() == EKeys::LeftMouseButton )
 	{
-		if (SequenceInstance.IsValid())
+		if (SectionObject.GetSequence())
 		{
-			if (SectionObject.GetSequence())
-			{
-				Sequencer.Pin()->FocusSequenceInstance(SectionObject);
-			}
+			Sequencer.Pin()->FocusSequenceInstance(SectionObject);
 		}
 	}
 
diff --git a/Engine/Source/Editor/MovieSceneTools/Private/Sections/CinematicShotSection.h b/Engine/Source/Editor/MovieSceneTools/Private/Sections/CinematicShotSection.h
--- a/Engine/Source/Editor/MovieSceneTools/Private/Sections/CinematicShotSection.h
+++ b/Engine/Source/Editor/MovieSceneTools/Private/Sections/CinematicShotSection.h
@@ -12,7 +12,6 @@ class FCinematicShotTrackEditor;
 class IMenu;
 class ISectionLayoutBuilder;
 class UMovieSceneCinematicShotSection;
-class FMovieSceneSequenceInstance;
 
 
 /**
@@ -45,7 +44,6 @@ public:
 
 	// FThumbnail interface
 	virtual void SetSingleTime(float GlobalTime) override;
-	virtual bool CanRename() const override { return true; }
 	virtual FText HandleThumbnailTextBlockText() const override;
 	virtual void HandleThumbnailTextBlockTextCommitted(const FText& NewThumbnailName, ETextCommit::Type CommitType) override;
 
@@ -59,9 +57,6 @@ private:
 	/** The section we are visualizing */
 	UMovieSceneCinematicShotSection& SectionObject;
 
-	/** The instance that this section is part of */
-	TWeakPtr<FMovieSceneSequenceInstance> SequenceInstance;
-
 	/** Sequencer interface */
 	TWeakPtr<ISequencer> Sequencer;
 
diff --git a/Engine/Source/Editor/MovieSceneTools/Private/Sections/ColorPropertySection.cpp b/Engine/Source/Editor/MovieSceneTools/Private/Sections/ColorPropertySection.cpp
--- a/Engine/Source/Editor/MovieSceneTools/Private/Sections/ColorPropertySection.cpp
+++ b/Engine/Source/Editor/MovieSceneTools/Private/Sections/ColorPropertySection.cpp
@@ -179,7 +179,15 @@ void FColorPropertySection::ConsolidateColorCurves( TArray< TKeyValuePair<float,
 	// @todo Sequencer Optimize - This another O(n^2) loop, since Eval is O(n)!
 	for ( int32 i = 0; i < TimesWithKeys.Num(); ++i )
 	{
-		OutColorKeys.Add( TKeyValuePair<float, FLinearColor>( TimesWithKeys[i], Section->Eval( TimesWithKeys[i], DefaultColor ) ) );
+		float Time = TimesWithKeys[i];
+
+		FLinearColor Color(
+			Section->GetRedCurve().Eval(Time, DefaultColor.R),
+			Section->GetGreenCurve().Eval(Time, DefaultColor.G),
+			Section->GetBlueCurve().Eval(Time, DefaultColor.B),
+			Section->GetAlphaCurve().Eval(Time, DefaultColor.A)
+			);
+		OutColorKeys.Add( TKeyValuePair<float, FLinearColor>( TimesWithKeys[i], Color ) );
 	}
 }
 
diff --git a/Engine/Source/Editor/MovieSceneTools/Private/Sections/EventTrackSection.cpp b/Engine/Source/Editor/MovieSceneTools/Private/Sections/EventTrackSection.cpp
--- a/Engine/Source/Editor/MovieSceneTools/Private/Sections/EventTrackSection.cpp
+++ b/Engine/Source/Editor/MovieSceneTools/Private/Sections/EventTrackSection.cpp
@@ -3,7 +3,7 @@
 #include "MovieSceneToolsPrivatePCH.h"
 #include "EventTrackSection.h"
 #include "MovieSceneEventSection.h"
-#include "NameCurveKeyArea.h"
+#include "GenericKeyArea.h"
 
 
 #define LOCTEXT_NAMESPACE "FEventTrackSection"
@@ -44,9 +44,10 @@ FText FEventTrackSection::GetSectionTitle() const
 }
 
 
-void FEventTrackSection::GenerateSectionLayout(class ISectionLayoutBuilder& LayoutBuilder) const
+void FEventTrackSection::GenerateSectionLayout(ISectionLayoutBuilder& LayoutBuilder) const
 {
-	LayoutBuilder.SetSectionAsKeyArea(MakeShareable(new FNameCurveKeyArea(Section->GetEventCurve(), Section)));
+	auto KeyArea = MakeShared<TGenericKeyArea<FEventPayload, float>>(Section->GetCurveInterface(), Section);
+	LayoutBuilder.SetSectionAsKeyArea(KeyArea);
 }
 
 
diff --git a/Engine/Source/Editor/MovieSceneTools/Private/Sections/IntegerPropertySection.cpp b/Engine/Source/Editor/MovieSceneTools/Private/Sections/IntegerPropertySection.cpp
new file mode 100644
--- /dev/null
+++ b/Engine/Source/Editor/MovieSceneTools/Private/Sections/IntegerPropertySection.cpp
@@ -0,0 +1,23 @@
+// Copyright 1998-2016 Epic Games, Inc. All Rights Reserved.
+
+#include "MovieSceneToolsPrivatePCH.h"
+#include "IntegerKeyArea.h"
+#include "IntegerPropertySection.h"
+#include "MovieSceneIntegerSection.h"
+#include "MovieScenePropertyTrack.h"
+
+void FIntegerPropertySection::GenerateSectionLayout(class ISectionLayoutBuilder& LayoutBuilder) const
+{
+	UMovieSceneIntegerSection* IntegerSection = Cast<UMovieSceneIntegerSection>(&SectionObject);
+	TAttribute<TOptional<int32>> ExternalValue;
+	if (CanGetPropertyValue())
+	{
+		ExternalValue.Bind(TAttribute<TOptional<int32>>::FGetter::CreateLambda([&]
+		{
+			return GetPropertyValue<int32>();
+		}));
+	}
+	TSharedRef<FIntegerKeyArea> KeyArea = MakeShareable(
+		new FIntegerKeyArea(IntegerSection->GetCurve(), ExternalValue, IntegerSection));
+	LayoutBuilder.SetSectionAsKeyArea(KeyArea);
+}
diff --git a/Engine/Source/Editor/MovieSceneTools/Private/Sections/IntegerPropertySection.h b/Engine/Source/Editor/MovieSceneTools/Private/Sections/IntegerPropertySection.h
new file mode 100644
--- /dev/null
+++ b/Engine/Source/Editor/MovieSceneTools/Private/Sections/IntegerPropertySection.h
@@ -0,0 +1,51 @@
+// Copyright 1998-2016 Epic Games, Inc. All Rights Reserved.
+
+#pragma once
+
+#include "PropertySection.h"
+
+
+class UMovieSceneSection;
+
+
+/**
+ * An implementation of integer property sections
+ */
+class FIntegerPropertySection
+	: public FPropertySection
+{
+public:
+
+	/**
+	* Creates a new integer property section for editing non-property integer sections.
+	*
+	* @param InSectionObject The section object which is being displayed and edited.
+	* @param InDisplayName A display name for the section being displayed and edited.
+	* @TODO Make another section which is a good base class for non-property bool sections.
+	*/
+	FIntegerPropertySection(UMovieSceneSection& InSectionObject, const FText& InDisplayName)
+		: FPropertySection(InSectionObject, InDisplayName)
+	{
+	}
+
+	/**
+	* Creates a new integer property section.
+	*
+	* @param InSequencer The sequencer which is controlling this property section.
+	* @param InObjectBinding The object binding for the object which owns the property that this section is animating.
+	* @param InPropertyName The name of the property which is animated by this section.
+	* @param InPropertyPath A string representing the path to the property which is animated by this section.
+	* @param InSectionObject The section object which is being displayed and edited.
+	* @param InDisplayName A display name for the section being displayed and edited.
+	*/
+	FIntegerPropertySection(ISequencer* InSequencer, FGuid InObjectBinding, FName InPropertyName, const FString& InPropertyPath, UMovieSceneSection& InSectionObject, const FText& InDisplayName)
+		: FPropertySection(InSequencer, InObjectBinding, InPropertyName, InPropertyPath, InSectionObject, InDisplayName)
+	{
+	}
+
+public:
+
+	// FPropertySection interface
+
+	virtual void GenerateSectionLayout(class ISectionLayoutBuilder& LayoutBuilder) const override;
+};
diff --git a/Engine/Source/Editor/MovieSceneTools/Private/Sections/ThumbnailSection.cpp b/Engine/Source/Editor/MovieSceneTools/Private/Sections/ThumbnailSection.cpp
--- a/Engine/Source/Editor/MovieSceneTools/Private/Sections/ThumbnailSection.cpp
+++ b/Engine/Source/Editor/MovieSceneTools/Private/Sections/ThumbnailSection.cpp
@@ -62,7 +62,8 @@ void FThumbnailSection::PreDraw(FTrackEditorThumbnail& Thumbnail, FLevelEditorVi
 
 		SavedPlaybackStatus = Sequencer->GetPlaybackStatus();
 		Sequencer->SetPlaybackStatus(EMovieScenePlayerStatus::Jumping);
-		Sequencer->SetGlobalTimeDirectly(Thumbnail.GetEvalPosition());
+		Sequencer->SetLocalTimeDirectly(Thumbnail.GetEvalPosition());
+		Sequencer->ForceEvaluate();
 
 		ViewportClient.SetActorLock(Camera);
 		Sequencer->OnCameraCut().Remove(Handle);
@@ -97,7 +98,7 @@ TSharedRef<SWidget> FThumbnailSection::GenerateSectionWidget()
 		.VAlign(VAlign_Top)
 		.Padding(GetContentPadding())
 		[
-			SNew(SInlineEditableTextBlock)
+			SAssignNew(NameWidget, SInlineEditableTextBlock)
 				.ToolTipText(CanRename() ? LOCTEXT("RenameThumbnail", "Click or hit F2 to rename") : FText::GetEmpty())
 				.Text(this, &FThumbnailSection::HandleThumbnailTextBlockText)
 				.ShadowOffset(FVector2D(1,1))
@@ -106,6 +107,15 @@ TSharedRef<SWidget> FThumbnailSection::GenerateSectionWidget()
 		];
 }
 
+void FThumbnailSection::EnterRename()
+{
+	if (NameWidget.IsValid())
+	{
+		NameWidget->SetReadOnly(false);
+		NameWidget->EnterEditingMode();
+		NameWidget->SetReadOnly(!CanRename());
+	}
+}
 
 void FThumbnailSection::BuildSectionContextMenu(FMenuBuilder& MenuBuilder, const FGuid& ObjectBinding)
 {
@@ -118,7 +128,7 @@ void FThumbnailSection::BuildSectionContextMenu(FMenuBuilder& MenuBuilder, const
 
 				TSharedPtr<ISequencer> Sequencer = SequencerPtr.Pin();
 
-				FText CurrentTime = FText::FromString(Sequencer->GetZeroPadNumericTypeInterface()->ToString(Sequencer->GetGlobalTime()));
+				FText CurrentTime = FText::FromString(Sequencer->GetZeroPadNumericTypeInterface()->ToString(Sequencer->GetLocalTime()));
 
 				InMenuBuilder.BeginSection(NAME_None, LOCTEXT("ThisSectionText", "This Section"));
 				{
@@ -134,7 +144,7 @@ void FThumbnailSection::BuildSectionContextMenu(FMenuBuilder& MenuBuilder, const
 						FSlateIcon(),
 						FUIAction(
 							FExecuteAction::CreateLambda([=]{
-								SetSingleTime(Sequencer->GetGlobalTime());
+								SetSingleTime(Sequencer->GetLocalTime());
 								GetMutableDefault<UMovieSceneUserThumbnailSettings>()->bDrawSingleThumbnails = true;
 								GetMutableDefault<UMovieSceneUserThumbnailSettings>()->SaveConfig();
 							})
diff --git a/Engine/Source/Editor/MovieSceneTools/Private/Sections/ThumbnailSection.h b/Engine/Source/Editor/MovieSceneTools/Private/Sections/ThumbnailSection.h
--- a/Engine/Source/Editor/MovieSceneTools/Private/Sections/ThumbnailSection.h
+++ b/Engine/Source/Editor/MovieSceneTools/Private/Sections/ThumbnailSection.h
@@ -10,6 +10,7 @@ class FTrackEditorThumbnailPool;
 class IMenu;
 class ISectionLayoutBuilder;
 class SWidget;
+class SInlineEditableTextBlock;
 class UMovieSceneSection;
 
 
@@ -43,6 +44,9 @@ public:
 		return SequencerPtr.Pin()->GetSequencerWidget();
 	}
 
+	/** Enter rename mode for the section */
+	void EnterRename();
+
 	/** Get whether the text is renameable */
 	virtual bool CanRename() const { return false; }
 
@@ -98,6 +102,9 @@ protected:
 	/** Saved playback status. Used for restoring state when rendering thumbnails */
 	EMovieScenePlayerStatus::Type SavedPlaybackStatus;
 	
+	/** Rename widget */
+	TSharedPtr<SInlineEditableTextBlock> NameWidget;
+
 	/** Fade brush. */
 	const FSlateBrush* WhiteBrush;
 };
diff --git a/Engine/Source/Editor/MovieSceneTools/Private/Sections/VisibilityPropertySection.cpp b/Engine/Source/Editor/MovieSceneTools/Private/Sections/VisibilityPropertySection.cpp
--- a/Engine/Source/Editor/MovieSceneTools/Private/Sections/VisibilityPropertySection.cpp
+++ b/Engine/Source/Editor/MovieSceneTools/Private/Sections/VisibilityPropertySection.cpp
@@ -3,20 +3,23 @@
 #include "MovieSceneToolsPrivatePCH.h"
 #include "BoolKeyArea.h"
 #include "VisibilityPropertySection.h"
-#include "MovieSceneVisibilitySection.h"
+#include "MovieSceneBoolSection.h"
 
 
 void FVisibilityPropertySection::GenerateSectionLayout(class ISectionLayoutBuilder& LayoutBuilder) const
 {
-	UMovieSceneVisibilitySection* VisibilitySection = Cast<UMovieSceneVisibilitySection>(&SectionObject);
+	UMovieSceneBoolSection* BoolSection = Cast<UMovieSceneBoolSection>( &SectionObject );
 	TAttribute<TOptional<bool>> ExternalValue;
-	ExternalValue.Bind(TAttribute<TOptional<bool>>::FGetter::CreateLambda([&]
+	if (CanGetPropertyValue())
 	{
-		TOptional<bool> BoolValue = GetPropertyValue<bool>();
-		return BoolValue.IsSet()
-			? TOptional<bool>(!BoolValue.GetValue())
-			: TOptional<bool>();
-	}));
-	TSharedRef<FBoolKeyArea> KeyArea = MakeShareable(new FBoolKeyArea(VisibilitySection->GetCurve(), ExternalValue, VisibilitySection));
-	LayoutBuilder.SetSectionAsKeyArea(KeyArea);
+		ExternalValue.Bind(TAttribute<TOptional<bool>>::FGetter::CreateLambda([&]
+		{
+			TOptional<bool> BoolValue = GetPropertyValue<bool>();
+			return BoolValue.IsSet()
+				? TOptional<bool>(!BoolValue.GetValue())
+				: TOptional<bool>();
+		}));
+	}
+	TSharedRef<FBoolKeyArea> KeyArea = MakeShareable( new FBoolKeyArea( BoolSection->GetCurve(), ExternalValue, BoolSection ) );
+	LayoutBuilder.SetSectionAsKeyArea( KeyArea );
 }
\ No newline at end of file
diff --git a/Engine/Source/Editor/MovieSceneTools/Private/TrackEditors/AttachTrackEditor.cpp b/Engine/Source/Editor/MovieSceneTools/Private/TrackEditors/AttachTrackEditor.cpp
--- a/Engine/Source/Editor/MovieSceneTools/Private/TrackEditors/AttachTrackEditor.cpp
+++ b/Engine/Source/Editor/MovieSceneTools/Private/TrackEditors/AttachTrackEditor.cpp
@@ -117,13 +117,30 @@ void F3DAttachTrackEditor::BuildObjectBindingTrackMenu(FMenuBuilder& MenuBuilder
 bool F3DAttachTrackEditor::IsActorPickable(const AActor* const ParentActor, FGuid ObjectBinding, UMovieSceneSection* InSection)
 {
 	// Can't pick the object that this track binds
-	TArray<TWeakObjectPtr<UObject>> OutObjects;
-	GetSequencer()->GetRuntimeObjects( GetSequencer()->GetFocusedMovieSceneSequenceInstance(), ObjectBinding, OutObjects);
-	if (OutObjects.Contains(ParentActor))
+	TArrayView<TWeakObjectPtr<>> Objects = GetSequencer()->FindObjectsInCurrentSequence(ObjectBinding);
+	if (Objects.Contains(ParentActor))
 	{
 		return false;
 	}
 
+	for (auto Object : Objects)
+	{
+		if (Object.IsValid())
+		{
+			AActor* ChildActor = Cast<AActor>(Object.Get());
+			if (ChildActor)
+			{
+				USceneComponent* ChildRoot = ChildActor->GetRootComponent();
+				USceneComponent* ParentRoot = ParentActor->GetDefaultAttachComponent();
+
+				if (!ChildRoot || !ParentRoot || ParentRoot->IsAttachedTo(ChildRoot))
+				{
+					return false;
+				}
+			}
+		}
+	}
+
 	if (ParentActor->IsListedInSceneOutliner() &&
 		!FActorEditorUtils::IsABuilderBrush(ParentActor) &&
 		!ParentActor->IsA( AWorldSettings::StaticClass() ) &&
@@ -153,8 +170,11 @@ void F3DAttachTrackEditor::ActorSocketPicked(const FName SocketName, USceneCompo
 	}
 	else if (ObjectGuid.IsValid())
 	{
-		TArray<TWeakObjectPtr<UObject>> OutObjects;
-		GetSequencer()->GetRuntimeObjects( GetSequencer()->GetFocusedMovieSceneSequenceInstance(), ObjectGuid, OutObjects);
+		TArray<TWeakObjectPtr<>> OutObjects;
+		for (TWeakObjectPtr<> Object : GetSequencer()->FindObjectsInCurrentSequence(ObjectGuid))
+		{
+			OutObjects.Add(Object);
+		}
 
 		AnimatablePropertyChanged( FOnKeyProperty::CreateRaw( this, &F3DAttachTrackEditor::AddKeyInternal, OutObjects, SocketName, Component ? Component->GetFName() : NAME_None, ParentActor) );
 	}
diff --git a/Engine/Source/Editor/MovieSceneTools/Private/TrackEditors/AudioTrackEditor.cpp b/Engine/Source/Editor/MovieSceneTools/Private/TrackEditors/AudioTrackEditor.cpp
--- a/Engine/Source/Editor/MovieSceneTools/Private/TrackEditors/AudioTrackEditor.cpp
+++ b/Engine/Source/Editor/MovieSceneTools/Private/TrackEditors/AudioTrackEditor.cpp
@@ -855,8 +855,11 @@ bool FAudioTrackEditor::HandleAssetAdded(UObject* Asset, const FGuid& TargetObje
 		
 		if (TargetObjectGuid.IsValid())
 		{
-			TArray<TWeakObjectPtr<UObject>> OutObjects;
-			GetSequencer()->GetRuntimeObjects(GetSequencer()->GetFocusedMovieSceneSequenceInstance(), TargetObjectGuid, OutObjects);
+			TArray<TWeakObjectPtr<>> OutObjects;
+			for (TWeakObjectPtr<> Object : GetSequencer()->FindObjectsInCurrentSequence(TargetObjectGuid))
+			{
+				OutObjects.Add(Object);
+			}
 
 			AnimatablePropertyChanged( FOnKeyProperty::CreateRaw(this, &FAudioTrackEditor::AddNewAttachedSound, Sound, OutObjects));
 		}
@@ -996,7 +999,7 @@ void FAudioTrackEditor::OnAudioAssetSelected(const FAssetData& AssetData, UMovie
 			auto AudioTrack = Cast<UMovieSceneAudioTrack>(Track);
 			AudioTrack->Modify();
 
-			float KeyTime = GetSequencer()->GetGlobalTime();
+			float KeyTime = GetSequencer()->GetLocalTime();
 			AudioTrack->AddNewSound( NewSound, KeyTime );
 
 			GetSequencer()->NotifyMovieSceneDataChanged( EMovieSceneDataChangeType::MovieSceneStructureItemAdded );
diff --git a/Engine/Source/Editor/MovieSceneTools/Private/TrackEditors/CameraAnimTrackEditor.cpp b/Engine/Source/Editor/MovieSceneTools/Private/TrackEditors/CameraAnimTrackEditor.cpp
--- a/Engine/Source/Editor/MovieSceneTools/Private/TrackEditors/CameraAnimTrackEditor.cpp
+++ b/Engine/Source/Editor/MovieSceneTools/Private/TrackEditors/CameraAnimTrackEditor.cpp
@@ -37,7 +37,7 @@ public:
 	virtual FText GetSectionTitle() const override 
 	{ 
 		UMovieSceneCameraAnimSection const* const AnimSection = Cast<UMovieSceneCameraAnimSection>(&Section);
-		UCameraAnim const* const Anim = AnimSection ? AnimSection->GetCameraAnim() : nullptr;
+		UCameraAnim const* const Anim = AnimSection ? AnimSection->AnimData.CameraAnim : nullptr;
 		if (Anim)
 		{
 			return FText::FromString(Anim->GetName());
@@ -120,9 +120,12 @@ bool FCameraAnimTrackEditor::HandleAssetAdded(UObject* Asset, const FGuid& Targe
 	{
 		if (TargetObjectGuid.IsValid())
 		{
-			TArray<TWeakObjectPtr<UObject>> OutObjects;
+			TArray<TWeakObjectPtr<>> OutObjects;
+			for (TWeakObjectPtr<> Object : GetSequencer()->FindObjectsInCurrentSequence(TargetObjectGuid))
+			{
+				OutObjects.Add(Object);
+			}
 
-			GetSequencer()->GetRuntimeObjects(GetSequencer()->GetFocusedMovieSceneSequenceInstance(), TargetObjectGuid, OutObjects);
 			AnimatablePropertyChanged(FOnKeyProperty::CreateRaw(this, &FCameraAnimTrackEditor::AddKeyInternal, OutObjects, CameraAnim));
 
 			return true;
@@ -216,9 +219,11 @@ void FCameraAnimTrackEditor::OnCameraAnimAssetSelected(const FAssetData& AssetDa
 	{
 		UCameraAnim* const CameraAnim = CastChecked<UCameraAnim>(AssetData.GetAsset());
 
-		TArray<TWeakObjectPtr<UObject>> OutObjects;
-
-		GetSequencer()->GetRuntimeObjects(GetSequencer()->GetFocusedMovieSceneSequenceInstance(), ObjectBinding, OutObjects);
+		TArray<TWeakObjectPtr<>> OutObjects;
+		for (TWeakObjectPtr<> Object : GetSequencer()->FindObjectsInCurrentSequence(ObjectBinding))
+		{
+			OutObjects.Add(Object);
+		}
 		AnimatablePropertyChanged(FOnKeyProperty::CreateRaw(this, &FCameraAnimTrackEditor::AddKeyInternal, OutObjects, CameraAnim));
 	}
 }
@@ -258,13 +263,10 @@ bool FCameraAnimTrackEditor::AddKeyInternal(float KeyTime, const TArray<TWeakObj
 
 UCameraComponent* FCameraAnimTrackEditor::AcquireCameraComponentFromObjectGuid(const FGuid& Guid)
 {
-	TArray<TWeakObjectPtr<UObject>> OutObjects;
-	GetSequencer()->GetRuntimeObjects(GetSequencer()->GetFocusedMovieSceneSequenceInstance(), Guid, OutObjects);
-
 	USkeleton* Skeleton = nullptr;
-	for (int32 i = 0; i < OutObjects.Num(); ++i)
+	for (TWeakObjectPtr<> WeakObject : GetSequencer()->FindObjectsInCurrentSequence(Guid))
 	{
-		UObject* const Obj = OutObjects[i].Get();
+		UObject* const Obj = WeakObject.Get();
 	
 		if (AActor* const Actor = Cast<AActor>(Obj))
 		{
diff --git a/Engine/Source/Editor/MovieSceneTools/Private/TrackEditors/CameraCutTrackEditor.cpp b/Engine/Source/Editor/MovieSceneTools/Private/TrackEditors/CameraCutTrackEditor.cpp
--- a/Engine/Source/Editor/MovieSceneTools/Private/TrackEditors/CameraCutTrackEditor.cpp
+++ b/Engine/Source/Editor/MovieSceneTools/Private/TrackEditors/CameraCutTrackEditor.cpp
@@ -128,11 +128,11 @@ void FCameraCutTrackEditor::Tick(float DeltaTime)
 	{
 		SequencerPin->EnterSilentMode();
 
-		float SavedTime = SequencerPin->GetGlobalTime();
+		float SavedTime = SequencerPin->GetLocalTime();
 
 		if (DeltaTime > 0.f && ThumbnailPool->DrawThumbnails())
 		{
-			SequencerPin->SetGlobalTimeDirectly(SavedTime);
+			SequencerPin->SetLocalTimeDirectly(SavedTime);
 		}
 
 		SequencerPin->ExitSilentMode();
@@ -301,7 +301,7 @@ void FCameraCutTrackEditor::OnLockCameraClicked(ECheckBoxState CheckBoxState)
 		GetSequencer()->SetPerspectiveViewportCameraCutEnabled(false);
 	}
 
-	GetSequencer()->SetGlobalTime(GetSequencer()->GetGlobalTime());
+	GetSequencer()->ForceEvaluate();
 }
 
 FText FCameraCutTrackEditor::GetLockCameraToolTip() const
diff --git a/Engine/Source/Editor/MovieSceneTools/Private/TrackEditors/CameraShakeTrackEditor.cpp b/Engine/Source/Editor/MovieSceneTools/Private/TrackEditors/CameraShakeTrackEditor.cpp
--- a/Engine/Source/Editor/MovieSceneTools/Private/TrackEditors/CameraShakeTrackEditor.cpp
+++ b/Engine/Source/Editor/MovieSceneTools/Private/TrackEditors/CameraShakeTrackEditor.cpp
@@ -40,7 +40,7 @@ public:
 	virtual FText GetSectionTitle() const override
 	{
 		UMovieSceneCameraShakeSection const* const ShakeSection = Cast<UMovieSceneCameraShakeSection>(&Section);
-		UClass const* const Shake = ShakeSection ? ShakeSection->GetCameraShakeClass() : nullptr;
+		UClass const* const Shake = ShakeSection ? ShakeSection->ShakeData.ShakeClass : nullptr;
 		if (Shake)
 		{
 			return FText::FromString(Shake->GetName());
@@ -125,8 +125,11 @@ bool FCameraShakeTrackEditor::HandleAssetAdded(UObject* Asset, const FGuid& Targ
 		{
 			TSubclassOf<UCameraShake> const ShakeClass = *(SelectedObject->GeneratedClass);
 
-			TArray<TWeakObjectPtr<UObject>> OutObjects;
-			GetSequencer()->GetRuntimeObjects(GetSequencer()->GetFocusedMovieSceneSequenceInstance(), TargetObjectGuid, OutObjects);
+			TArray<TWeakObjectPtr<>> OutObjects;
+			for (TWeakObjectPtr<> Object : GetSequencer()->FindObjectsInCurrentSequence(TargetObjectGuid))
+			{
+				OutObjects.Add(Object);
+			}
 			AnimatablePropertyChanged(FOnKeyProperty::CreateRaw(this, &FCameraShakeTrackEditor::AddKeyInternal, OutObjects, ShakeClass));
 
 			return true;
@@ -227,9 +230,11 @@ void FCameraShakeTrackEditor::OnCameraShakeAssetSelected(const FAssetData& Asset
 	{
 		TSubclassOf<UCameraShake> const ShakeClass = *(SelectedObject->GeneratedClass);
 
-		TArray<TWeakObjectPtr<UObject>> OutObjects;
-
-		GetSequencer()->GetRuntimeObjects(GetSequencer()->GetFocusedMovieSceneSequenceInstance(), ObjectBinding, OutObjects);
+		TArray<TWeakObjectPtr<>> OutObjects;
+		for (TWeakObjectPtr<> Object : GetSequencer()->FindObjectsInCurrentSequence(ObjectBinding))
+		{
+			OutObjects.Add(Object);
+		}
 		AnimatablePropertyChanged(FOnKeyProperty::CreateRaw(this, &FCameraShakeTrackEditor::AddKeyInternal, OutObjects, ShakeClass));
 	}
 }
@@ -267,13 +272,10 @@ bool FCameraShakeTrackEditor::AddKeyInternal(float KeyTime, const TArray<TWeakOb
 
 UCameraComponent* FCameraShakeTrackEditor::AcquireCameraComponentFromObjectGuid(const FGuid& Guid)
 {
-	TArray<TWeakObjectPtr<UObject>> OutObjects;
-	GetSequencer()->GetRuntimeObjects(GetSequencer()->GetFocusedMovieSceneSequenceInstance(), Guid, OutObjects);
-
 	USkeleton* Skeleton = nullptr;
-	for (int32 i = 0; i < OutObjects.Num(); ++i)
+	for (TWeakObjectPtr<> WeakObject : GetSequencer()->FindObjectsInCurrentSequence(Guid))
 	{
-		UObject* const Obj = OutObjects[i].Get();
+		UObject* const Obj = WeakObject.Get();
 		if (AActor* const Actor = Cast<AActor>(Obj))
 		{
 			UCameraComponent* const CameraComp = MovieSceneHelpers::CameraComponentFromActor(Actor);
diff --git a/Engine/Source/Editor/MovieSceneTools/Private/TrackEditors/CinematicShotTrackEditor.cpp b/Engine/Source/Editor/MovieSceneTools/Private/TrackEditors/CinematicShotTrackEditor.cpp
--- a/Engine/Source/Editor/MovieSceneTools/Private/TrackEditors/CinematicShotTrackEditor.cpp
+++ b/Engine/Source/Editor/MovieSceneTools/Private/TrackEditors/CinematicShotTrackEditor.cpp
@@ -181,7 +181,7 @@ void FCinematicShotTrackEditor::Tick(float DeltaTime)
 
 		if (DeltaTime > 0.f && ThumbnailPool->DrawThumbnails())
 		{
-			SequencerPin->SetGlobalTimeDirectly(SavedTime);
+			SequencerPin->SetGlobalTime(SavedTime);
 		}
 
 		SequencerPin->ExitSilentMode();
@@ -276,7 +276,7 @@ void FCinematicShotTrackEditor::InsertShotAtCurrentTime()
 {
 	const FScopedTransaction Transaction(LOCTEXT("InsertShot_Transaction", "Insert Shot"));
 
-	float NewShotStartTime = GetSequencer()->GetGlobalTime();
+	float NewShotStartTime = GetSequencer()->GetLocalTime();
 
 	UMovieSceneCinematicShotTrack* CinematicShotTrack = FindOrCreateCinematicShotTrack();
 	FString NewShotName = MovieSceneToolHelpers::GenerateNewShotName(CinematicShotTrack->GetAllSections(), NewShotStartTime);
@@ -309,7 +309,7 @@ void FCinematicShotTrackEditor::InsertFillerAtCurrentTime()
 
 	const FScopedTransaction Transaction(LOCTEXT("InsertFiller_Transaction", "Insert Filler"));
 
-	float NewShotStartTime = GetSequencer()->GetGlobalTime();
+	float NewShotStartTime = GetSequencer()->GetLocalTime();
 
 	UMovieSceneCinematicShotTrack* CinematicShotTrack = FindOrCreateCinematicShotTrack();
 
@@ -342,9 +342,9 @@ void FCinematicShotTrackEditor::DuplicateShot(UMovieSceneCinematicShotSection* S
 	if (NewShot)
 	{
 		NewShot->SetEndTime(NewShot->GetStartTime() + Section->GetEndTime() - Section->GetStartTime());
-		NewShot->StartOffset = Section->StartOffset;
-		NewShot->TimeScale = Section->TimeScale;
-		NewShot->PrerollTime = Section->PrerollTime;
+		NewShot->Parameters.StartOffset = Section->Parameters.StartOffset;
+		NewShot->Parameters.TimeScale = Section->Parameters.TimeScale;
+		NewShot->Parameters.PrerollTime = Section->Parameters.PrerollTime;
 
 		GetSequencer()->NotifyMovieSceneDataChanged( EMovieSceneDataChangeType::MovieSceneStructureItemAdded );
 	}
@@ -386,9 +386,9 @@ void FCinematicShotTrackEditor::NewTake(UMovieSceneCinematicShotSection* Section
 
 		float NewShotStartTime = Section->GetStartTime();
 		float NewShotEndTime = Section->GetEndTime();
-		float NewShotStartOffset = Section->StartOffset;
-		float NewShotTimeScale = Section->TimeScale;
-		float NewShotPrerollTime = Section->PrerollTime;
+		float NewShotStartOffset = Section->Parameters.StartOffset;
+		float NewShotTimeScale = Section->Parameters.TimeScale;
+		float NewShotPrerollTime = Section->Parameters.PrerollTime;
 
 		const bool bInsertShot = false;
 		UMovieSceneSubSection* NewShot = CreateShotInternal(NewShotName, NewShotStartTime, Section, bInsertShot);
@@ -398,11 +398,14 @@ void FCinematicShotTrackEditor::NewTake(UMovieSceneCinematicShotSection* Section
 			UMovieSceneCinematicShotTrack* CinematicShotTrack = FindOrCreateCinematicShotTrack();
 			CinematicShotTrack->RemoveSection(*Section);
 
-			NewShot->SetStartTime(NewShotStartTime);
-			NewShot->SetEndTime(NewShotEndTime);
-			NewShot->StartOffset = NewShotStartOffset;
-			NewShot->TimeScale = NewShotTimeScale;
-			NewShot->PrerollTime = NewShotPrerollTime;
+			if (NewShot != nullptr)
+			{
+				NewShot->SetStartTime(NewShotStartTime);
+				NewShot->SetEndTime(NewShotEndTime);
+				NewShot->Parameters.StartOffset = NewShotStartOffset;
+				NewShot->Parameters.TimeScale = NewShotTimeScale;
+				NewShot->Parameters.PrerollTime = NewShotPrerollTime;
+			}
 
 			GetSequencer()->NotifyMovieSceneDataChanged( EMovieSceneDataChangeType::MovieSceneStructureItemsChanged );
 		}
@@ -424,9 +427,9 @@ void FCinematicShotTrackEditor::SwitchTake(UMovieSceneCinematicShotSection* Sect
 		
 		float NewShotStartTime = Section->GetStartTime();
 		float NewShotEndTime = Section->GetEndTime();
-		float NewShotStartOffset = Section->StartOffset;
-		float NewShotTimeScale = Section->TimeScale;
-		float NewShotPrerollTime = Section->PrerollTime;
+		float NewShotStartOffset = Section->Parameters.StartOffset;
+		float NewShotTimeScale = Section->Parameters.TimeScale;
+		float NewShotPrerollTime = Section->Parameters.PrerollTime;
 
 		const bool bInsertSequence = false;
 		const float Duration = NewShotEndTime - NewShotStartTime;
@@ -438,12 +441,12 @@ void FCinematicShotTrackEditor::SwitchTake(UMovieSceneCinematicShotSection* Sect
 
 			NewShot->SetStartTime(NewShotStartTime);
 			NewShot->SetEndTime(NewShotEndTime);
-			NewShot->StartOffset = NewShotStartOffset;
-			NewShot->TimeScale = NewShotTimeScale;
-			NewShot->PrerollTime = NewShotPrerollTime;
-
-			GetSequencer()->NotifyMovieSceneDataChanged( EMovieSceneDataChangeType::MovieSceneStructureItemsChanged );
+			NewShot->Parameters.StartOffset = NewShotStartOffset;
+			NewShot->Parameters.TimeScale = NewShotTimeScale;
+			NewShot->Parameters.PrerollTime = NewShotPrerollTime;
 		}
+
+		GetSequencer()->NotifyMovieSceneDataChanged( EMovieSceneDataChangeType::MovieSceneStructureItemsChanged );
 	}
 
 }
@@ -600,7 +603,7 @@ void FCinematicShotTrackEditor::OnLockShotsClicked(ECheckBoxState CheckBoxState)
 		GetSequencer()->SetPerspectiveViewportCameraCutEnabled(false);
 	}
 
-	GetSequencer()->SetGlobalTime(GetSequencer()->GetGlobalTime());
+	GetSequencer()->ForceEvaluate();
 }
 
 
diff --git a/Engine/Source/Editor/MovieSceneTools/Private/TrackEditors/MaterialTrackEditor.cpp b/Engine/Source/Editor/MovieSceneTools/Private/TrackEditors/MaterialTrackEditor.cpp
--- a/Engine/Source/Editor/MovieSceneTools/Private/TrackEditors/MaterialTrackEditor.cpp
+++ b/Engine/Source/Editor/MovieSceneTools/Private/TrackEditors/MaterialTrackEditor.cpp
@@ -121,8 +121,7 @@ UMaterial* FMaterialTrackEditor::GetMaterialForTrack( FGuid ObjectBinding, UMovi
 
 void FMaterialTrackEditor::AddScalarParameter( FGuid ObjectBinding, UMovieSceneMaterialTrack* MaterialTrack, FName ParameterName )
 {
-	UMovieSceneSequence* MovieSceneSequence = GetMovieSceneSequence();
-	float KeyTime = GetTimeForKey( MovieSceneSequence );
+	float KeyTime = GetTimeForKey();
 
 	UMaterial* Material = GetMaterialForTrack(ObjectBinding, MaterialTrack);
 	if (Material != nullptr)
@@ -139,8 +138,7 @@ void FMaterialTrackEditor::AddScalarParameter( FGuid ObjectBinding, UMovieSceneM
 
 void FMaterialTrackEditor::AddColorParameter( FGuid ObjectBinding, UMovieSceneMaterialTrack* MaterialTrack, FName ParameterName )
 {
-	UMovieSceneSequence* MovieSceneSequence = GetMovieSceneSequence();
-	float KeyTime = GetTimeForKey( MovieSceneSequence );
+	float KeyTime = GetTimeForKey();
 
 	UMaterial* Material = GetMaterialForTrack( ObjectBinding, MaterialTrack );
 	if ( Material != nullptr )
diff --git a/Engine/Source/Editor/MovieSceneTools/Private/TrackEditors/ParticleParameterTrackEditor.cpp b/Engine/Source/Editor/MovieSceneTools/Private/TrackEditors/ParticleParameterTrackEditor.cpp
--- a/Engine/Source/Editor/MovieSceneTools/Private/TrackEditors/ParticleParameterTrackEditor.cpp
+++ b/Engine/Source/Editor/MovieSceneTools/Private/TrackEditors/ParticleParameterTrackEditor.cpp
@@ -139,20 +139,24 @@ void FParticleParameterTrackEditor::AddParticleParameterTrack( FGuid ObjectBindi
 
 void FParticleParameterTrackEditor::AddScalarParameter( FGuid ObjectBinding, UMovieSceneParticleParameterTrack* ParticleParameterTrack, FName ParameterName )
 {
-	UMovieSceneSequence* MovieSceneSequence = GetMovieSceneSequence();
-	float KeyTime = GetTimeForKey( MovieSceneSequence );
+	float KeyTime = GetTimeForKey();
 
-	UObject* Object = GetSequencer()->GetFocusedMovieSceneSequenceInstance()->FindObject( ObjectBinding, *GetSequencer() );
-	AEmitter* Emitter = Cast<AEmitter>(Object);
-	if ( Emitter != nullptr )
+	const FScopedTransaction Transaction( LOCTEXT( "AddScalarParameter", "Add scalar parameter" ) );
+
+	for (TWeakObjectPtr<> Object : GetSequencer()->FindObjectsInCurrentSequence(ObjectBinding))
 	{
+		AEmitter* Emitter = Cast<AEmitter>(Object.Get());
+		if ( !Emitter )
+		{
+			continue;
+		}
+
 		UParticleSystemComponent* ParticleSystemComponent = Emitter->GetParticleSystemComponent();
 		if ( ParticleSystemComponent != nullptr )
 		{
 			float Value;
 			ParticleSystemComponent->GetFloatParameter( ParameterName, Value );
 
-			const FScopedTransaction Transaction( LOCTEXT( "AddScalarParameter", "Add scalar parameter" ) );
 			ParticleParameterTrack->Modify();
 			ParticleParameterTrack->AddScalarParameterKey( ParameterName, KeyTime, Value );
 		}
@@ -163,20 +167,23 @@ void FParticleParameterTrackEditor::AddScalarParameter( FGuid ObjectBinding, UMo
 
 void FParticleParameterTrackEditor::AddVectorParameter( FGuid ObjectBinding, UMovieSceneParticleParameterTrack* ParticleParameterTrack, FName ParameterName )
 {
-	UMovieSceneSequence* MovieSceneSequence = GetMovieSceneSequence();
-	float KeyTime = GetTimeForKey( MovieSceneSequence );
+	float KeyTime = GetTimeForKey();
 
-	UObject* Object = GetSequencer()->GetFocusedMovieSceneSequenceInstance()->FindObject( ObjectBinding, *GetSequencer() );
-	AEmitter* Emitter = Cast<AEmitter>( Object );
-	if ( Emitter != nullptr )
+	const FScopedTransaction Transaction( LOCTEXT( "AddVectorParameter", "Add vector parameter" ) );
+	for (TWeakObjectPtr<> Object : GetSequencer()->FindObjectsInCurrentSequence(ObjectBinding))
 	{
+		AEmitter* Emitter = Cast<AEmitter>(Object.Get());
+		if ( !Emitter )
+		{
+			continue;
+		}
+		
 		UParticleSystemComponent* ParticleSystemComponent = Emitter->GetParticleSystemComponent();
 		if ( ParticleSystemComponent != nullptr )
 		{
 			FVector Value;
 			ParticleSystemComponent->GetVectorParameter( ParameterName, Value );
 
-			const FScopedTransaction Transaction( LOCTEXT( "AddVectorParameter", "Add vector parameter" ) );
 			ParticleParameterTrack->Modify();
 			ParticleParameterTrack->AddVectorParameterKey( ParameterName, KeyTime, Value );
 		}
@@ -187,20 +194,24 @@ void FParticleParameterTrackEditor::AddVectorParameter( FGuid ObjectBinding, UMo
 
 void FParticleParameterTrackEditor::AddColorParameter( FGuid ObjectBinding, UMovieSceneParticleParameterTrack* ParticleParameterTrack, FName ParameterName )
 {
-	UMovieSceneSequence* MovieSceneSequence = GetMovieSceneSequence();
-	float KeyTime = GetTimeForKey( MovieSceneSequence );
+	float KeyTime = GetTimeForKey();
 
-	UObject* Object = GetSequencer()->GetFocusedMovieSceneSequenceInstance()->FindObject( ObjectBinding, *GetSequencer() );
-	AEmitter* Emitter = Cast<AEmitter>( Object );
-	if ( Emitter != nullptr )
+	const FScopedTransaction Transaction( LOCTEXT( "AddColorParameter", "Add color parameter" ) );
+
+	for (TWeakObjectPtr<> Object : GetSequencer()->FindObjectsInCurrentSequence(ObjectBinding))
 	{
+		AEmitter* Emitter = Cast<AEmitter>(Object.Get());
+		if (!Emitter)
+		{
+			continue;
+		}
+
 		UParticleSystemComponent* ParticleSystemComponent = Emitter->GetParticleSystemComponent();
 		if ( ParticleSystemComponent != nullptr )
 		{
 			FLinearColor Value;
 			ParticleSystemComponent->GetColorParameter( ParameterName, Value );
 
-			const FScopedTransaction Transaction( LOCTEXT( "AddColorParameter", "Add color parameter" ) );
 			ParticleParameterTrack->Modify();
 			ParticleParameterTrack->AddColorParameterKey( ParameterName, KeyTime, Value );
 		}
diff --git a/Engine/Source/Editor/MovieSceneTools/Private/TrackEditors/ParticleTrackEditor.cpp b/Engine/Source/Editor/MovieSceneTools/Private/TrackEditors/ParticleTrackEditor.cpp
--- a/Engine/Source/Editor/MovieSceneTools/Private/TrackEditors/ParticleTrackEditor.cpp
+++ b/Engine/Source/Editor/MovieSceneTools/Private/TrackEditors/ParticleTrackEditor.cpp
@@ -113,20 +113,22 @@ int32 FParticleSection::OnPaintSection( FSequencerSectionPainter& InPainter ) co
 
 			if ( ObjectHandle.IsValid() )
 			{
-				UObject* BoundObject = OwningSequencer->GetFocusedMovieSceneSequenceInstance()->FindObject( ObjectHandle, *OwningSequencer );
-				UParticleSystemComponent* ParticleSystemComponent = Cast<UParticleSystemComponent>(BoundObject);
-				if(AEmitter* ParticleSystemActor = Cast<AEmitter>(BoundObject))
+				for (TWeakObjectPtr<> Object : OwningSequencer->FindObjectsInCurrentSequence(ObjectHandle))
 				{
-					ParticleSystemComponent = ParticleSystemActor->GetParticleSystemComponent();
-				}
+					UParticleSystemComponent* ParticleSystemComponent = Cast<UParticleSystemComponent>(Object.Get());
+					if(AEmitter* ParticleSystemActor = Cast<AEmitter>(Object.Get()))
+					{
+						ParticleSystemComponent = ParticleSystemActor->GetParticleSystemComponent();
+					}
 
-				if ( ParticleSystemComponent != nullptr && ParticleSystemComponent->Template != nullptr )
-				{
-					for ( UParticleEmitter* Emitter : ParticleSystemComponent->Template->Emitters )
+					if ( ParticleSystemComponent != nullptr && ParticleSystemComponent->Template != nullptr )
 					{
-						UParticleModuleRequired* RequiredModule = Emitter->GetLODLevel( 0 )->RequiredModule;
-						bIsLooping |= RequiredModule->EmitterLoops == 0;
-						LastEmitterEndTime = FMath::Max( LastEmitterEndTime, RequiredModule->EmitterDelay + RequiredModule->EmitterDuration );
+						for ( UParticleEmitter* Emitter : ParticleSystemComponent->Template->Emitters )
+						{
+							UParticleModuleRequired* RequiredModule = Emitter->GetLODLevel( 0 )->RequiredModule;
+							bIsLooping |= RequiredModule->EmitterLoops == 0;
+							LastEmitterEndTime = FMath::Max( LastEmitterEndTime, RequiredModule->EmitterDelay + RequiredModule->EmitterDuration );
+						}
 					}
 				}
 			}
diff --git a/Engine/Source/Editor/MovieSceneTools/Private/TrackEditors/PathTrackEditor.cpp b/Engine/Source/Editor/MovieSceneTools/Private/TrackEditors/PathTrackEditor.cpp
--- a/Engine/Source/Editor/MovieSceneTools/Private/TrackEditors/PathTrackEditor.cpp
+++ b/Engine/Source/Editor/MovieSceneTools/Private/TrackEditors/PathTrackEditor.cpp
@@ -115,9 +115,7 @@ void F3DPathTrackEditor::BuildObjectBindingTrackMenu(FMenuBuilder& MenuBuilder,
 bool F3DPathTrackEditor::IsActorPickable(const AActor* const ParentActor, FGuid ObjectBinding, UMovieSceneSection* InSection)
 {
 	// Can't pick the object that this track binds
-	TArray<TWeakObjectPtr<UObject>> OutObjects;
-	GetSequencer()->GetRuntimeObjects( GetSequencer()->GetFocusedMovieSceneSequenceInstance(), ObjectBinding, OutObjects);
-	if (OutObjects.Contains(ParentActor))
+	if (GetSequencer()->FindObjectsInCurrentSequence(ObjectBinding).Contains(ParentActor))
 	{
 		return false;
 	}
@@ -130,9 +128,7 @@ bool F3DPathTrackEditor::IsActorPickable(const AActor* const ParentActor, FGuid
 
 		if (ConstraintId.IsValid())
 		{
-			TArray<TWeakObjectPtr<UObject>> ConstraintObjects;
-			GetSequencer()->GetRuntimeObjects( GetSequencer()->GetFocusedMovieSceneSequenceInstance(), ConstraintId, ConstraintObjects);
-			if (ConstraintObjects.Contains(ParentActor))
+			if (GetSequencer()->FindObjectsInCurrentSequence(ConstraintId).Contains(ParentActor))
 			{
 				return false;
 			}
@@ -171,8 +167,11 @@ void F3DPathTrackEditor::ActorSocketPicked(const FName SocketName, USceneCompone
 	}
 	else if (ObjectGuid.IsValid())
 	{
-		TArray<TWeakObjectPtr<UObject>> OutObjects;
-		GetSequencer()->GetRuntimeObjects( GetSequencer()->GetFocusedMovieSceneSequenceInstance(), ObjectGuid, OutObjects);
+		TArray<TWeakObjectPtr<>> OutObjects;
+		for (TWeakObjectPtr<> Object : GetSequencer()->FindObjectsInCurrentSequence(ObjectGuid))
+		{
+			OutObjects.Add(Object);
+		}
 		AnimatablePropertyChanged( FOnKeyProperty::CreateRaw( this, &F3DPathTrackEditor::AddKeyInternal, OutObjects, ParentActor) );
 	}
 }
diff --git a/Engine/Source/Editor/MovieSceneTools/Private/TrackEditors/PropertyTrackEditors/BytePropertyTrackEditor.cpp b/Engine/Source/Editor/MovieSceneTools/Private/TrackEditors/PropertyTrackEditors/BytePropertyTrackEditor.cpp
--- a/Engine/Source/Editor/MovieSceneTools/Private/TrackEditors/PropertyTrackEditors/BytePropertyTrackEditor.cpp
+++ b/Engine/Source/Editor/MovieSceneTools/Private/TrackEditors/PropertyTrackEditors/BytePropertyTrackEditor.cpp
@@ -24,12 +24,16 @@ TSharedRef<ISequencerSection> FBytePropertyTrackEditor::MakeSectionInterface(UMo
 
 UEnum* GetEnumForByteTrack(TSharedPtr<ISequencer> Sequencer, const FGuid& OwnerObjectHandle, FName PropertyName, UMovieSceneByteTrack* ByteTrack)
 {
-	
-	UObject* RuntimeObject = Sequencer->GetFocusedMovieSceneSequenceInstance()->FindObject(OwnerObjectHandle, *Sequencer);
 	TSet<UEnum*> PropertyEnums;
 
-	if (RuntimeObject != nullptr)
+	for (TWeakObjectPtr<> WeakObject : Sequencer->FindObjectsInCurrentSequence(OwnerObjectHandle))
 	{
+		UObject* RuntimeObject = WeakObject.Get();
+		if (!RuntimeObject)
+		{
+			continue;
+		}
+
 		UProperty* Property = RuntimeObject->GetClass()->FindPropertyByName(PropertyName);
 		if (Property != nullptr)
 		{
diff --git a/Engine/Source/Editor/MovieSceneTools/Private/TrackEditors/PropertyTrackEditors/IntegerPropertyTrackEditor.cpp b/Engine/Source/Editor/MovieSceneTools/Private/TrackEditors/PropertyTrackEditors/IntegerPropertyTrackEditor.cpp
new file mode 100644
--- /dev/null
+++ b/Engine/Source/Editor/MovieSceneTools/Private/TrackEditors/PropertyTrackEditors/IntegerPropertyTrackEditor.cpp
@@ -0,0 +1,31 @@
+// Copyright 1998-2016 Epic Games, Inc. All Rights Reserved.
+
+#include "MovieSceneToolsPrivatePCH.h"
+#include "IntegerPropertyTrackEditor.h"
+#include "IntegerPropertySection.h"
+
+
+TSharedRef<ISequencerTrackEditor> FIntegerPropertyTrackEditor::CreateTrackEditor( TSharedRef<ISequencer> OwningSequencer )
+{
+	return MakeShareable(new FIntegerPropertyTrackEditor(OwningSequencer));
+}
+
+
+TSharedRef<ISequencerSection> FIntegerPropertyTrackEditor::MakeSectionInterface(UMovieSceneSection& SectionObject, UMovieSceneTrack& Track, FGuid ObjectBinding)
+{
+	UMovieScenePropertyTrack* PropertyTrack = Cast<UMovieScenePropertyTrack>(&Track);
+	if (PropertyTrack != nullptr)
+	{
+		return MakeShareable(new FIntegerPropertySection(GetSequencer().Get(), ObjectBinding, PropertyTrack->GetPropertyName(), PropertyTrack->GetPropertyPath(), SectionObject, Track.GetDisplayName()));
+	}
+	else
+	{
+		return MakeShareable(new FIntegerPropertySection(SectionObject, Track.GetDisplayName()));
+	}
+}
+
+
+void FIntegerPropertyTrackEditor::GenerateKeysFromPropertyChanged( const FPropertyChangedParams& PropertyChangedParams, TArray<int32>& NewGeneratedKeys, TArray<int32>& DefaultGeneratedKeys )
+{
+	NewGeneratedKeys.Add( PropertyChangedParams.GetPropertyValue<int32>() );
+}
diff --git a/Engine/Source/Editor/MovieSceneTools/Private/TrackEditors/PropertyTrackEditors/IntegerPropertyTrackEditor.h b/Engine/Source/Editor/MovieSceneTools/Private/TrackEditors/PropertyTrackEditors/IntegerPropertyTrackEditor.h
new file mode 100644
--- /dev/null
+++ b/Engine/Source/Editor/MovieSceneTools/Private/TrackEditors/PropertyTrackEditors/IntegerPropertyTrackEditor.h
@@ -0,0 +1,47 @@
+// Copyright 1998-2016 Epic Games, Inc. All Rights Reserved.
+
+#pragma once
+
+#include "MovieSceneIntegerTrack.h"
+#include "MovieSceneIntegerSection.h"
+#include "PropertyTrackEditor.h"
+
+
+class ISequencer;
+
+
+/**
+ * A property track editor for integers.
+ */
+class FIntegerPropertyTrackEditor
+	: public FPropertyTrackEditor<UMovieSceneIntegerTrack, UMovieSceneIntegerSection, int32>
+{
+public:
+
+	/**
+	 * Constructor.
+	 *
+	 * @param InSequencer The sequencer instance to be used by this tool.
+	 */
+	FIntegerPropertyTrackEditor(TSharedRef<ISequencer> InSequencer)
+		: FPropertyTrackEditor(InSequencer, NAME_IntProperty)
+	{ }
+
+	/**
+	 * Creates an instance of this class (called by a sequencer).
+	 *
+	 * @param OwningSequencer The sequencer instance to be used by this tool
+	 * @return The new instance of this class
+	 */
+	static TSharedRef<ISequencerTrackEditor> CreateTrackEditor(TSharedRef<ISequencer> OwningSequencer);
+
+public:
+
+	//~ ISequencerTrackEditor interface
+
+	virtual TSharedRef<ISequencerSection> MakeSectionInterface(UMovieSceneSection& SectionObject, UMovieSceneTrack& Track, FGuid ObjectBinding) override;
+
+protected:
+
+	virtual void GenerateKeysFromPropertyChanged(const FPropertyChangedParams& PropertyChangedParams, TArray<int32>& NewGeneratedKeys, TArray<int32>& DefaultGeneratedKeys) override;
+};
diff --git a/Engine/Source/Editor/MovieSceneTools/Private/TrackEditors/PropertyTrackEditors/VisibilityPropertyTrackEditor.h b/Engine/Source/Editor/MovieSceneTools/Private/TrackEditors/PropertyTrackEditors/VisibilityPropertyTrackEditor.h
--- a/Engine/Source/Editor/MovieSceneTools/Private/TrackEditors/PropertyTrackEditors/VisibilityPropertyTrackEditor.h
+++ b/Engine/Source/Editor/MovieSceneTools/Private/TrackEditors/PropertyTrackEditors/VisibilityPropertyTrackEditor.h
@@ -2,7 +2,7 @@
 
 #pragma once
 
-#include "MovieSceneVisibilitySection.h"
+#include "MovieSceneBoolSection.h"
 #include "MovieSceneVisibilityTrack.h"
 #include "PropertyTrackEditor.h"
 
@@ -15,14 +15,14 @@ class ISequencer;
  * to use a UMovieSceneVisibilityTrack through metadata.
  */
 class FVisibilityPropertyTrackEditor
-	: public FPropertyTrackEditor<UMovieSceneVisibilityTrack, UMovieSceneVisibilitySection, bool>
+	: public FPropertyTrackEditor<UMovieSceneVisibilityTrack, UMovieSceneBoolSection, bool>
 {
 public:
 
 	/** Constructor. */
 	FVisibilityPropertyTrackEditor(TSharedRef<ISequencer> InSequencer)
 		// Don't supply any property type names to watch since the FBoolPropertyTrackEditor is already watching for bool property changes.
-		: FPropertyTrackEditor<UMovieSceneVisibilityTrack, UMovieSceneVisibilitySection, bool>(InSequencer, NAME_BoolProperty)
+		: FPropertyTrackEditor<UMovieSceneVisibilityTrack, UMovieSceneBoolSection, bool>(InSequencer, NAME_BoolProperty)
 	{ }
 
 	/**
diff --git a/Engine/Source/Editor/MovieSceneTools/Private/TrackEditors/SkeletalAnimationTrackEditor.cpp b/Engine/Source/Editor/MovieSceneTools/Private/TrackEditors/SkeletalAnimationTrackEditor.cpp
--- a/Engine/Source/Editor/MovieSceneTools/Private/TrackEditors/SkeletalAnimationTrackEditor.cpp
+++ b/Engine/Source/Editor/MovieSceneTools/Private/TrackEditors/SkeletalAnimationTrackEditor.cpp
@@ -55,9 +55,9 @@ FText FSkeletalAnimationSection::GetDisplayName() const
 
 FText FSkeletalAnimationSection::GetSectionTitle() const
 {
-	if (Section.GetAnimSequence() != nullptr)
+	if (Section.Params.Animation != nullptr)
 	{
-		return FText::FromString( Section.GetAnimSequence()->GetName() );
+		return FText::FromString( Section.Params.Animation->GetName() );
 	}
 	return LOCTEXT("NoAnimationSection", "No Animation");
 }
@@ -79,9 +79,9 @@ int32 FSkeletalAnimationSection::OnPaintSection( FSequencerSectionPainter& Paint
 
 	// Add lines where the animation starts and ends/loops
 	float CurrentTime = Section.GetStartTime();
-	float AnimPlayRate = FMath::IsNearlyZero(Section.GetPlayRate()) ? 1.0f : Section.GetPlayRate();
-	float SeqLength = (Section.GetSequenceLength() - (Section.GetStartOffset() + Section.GetEndOffset())) / AnimPlayRate;
-	while (CurrentTime < Section.GetEndTime() && !FMath::IsNearlyZero(Section.GetDuration()) && SeqLength > 0)
+	float AnimPlayRate = FMath::IsNearlyZero(Section.Params.PlayRate) ? 1.0f : Section.Params.PlayRate;
+	float SeqLength = (Section.Params.GetSequenceLength() - (Section.Params.StartOffset + Section.Params.EndOffset)) / AnimPlayRate;
+	while (CurrentTime < Section.GetEndTime() && !FMath::IsNearlyZero(Section.Params.GetDuration()) && SeqLength > 0)
 	{
 		if (CurrentTime > Section.GetStartTime())
 		{
@@ -108,7 +108,7 @@ int32 FSkeletalAnimationSection::OnPaintSection( FSequencerSectionPainter& Paint
 
 void FSkeletalAnimationSection::BeginResizeSection()
 {
-	InitialStartOffsetDuringResize = Section.GetStartOffset();
+	InitialStartOffsetDuringResize = Section.Params.StartOffset;
 	InitialStartTimeDuringResize = Section.GetStartTime();
 }
 
@@ -117,18 +117,18 @@ void FSkeletalAnimationSection::ResizeSection(ESequencerSectionResizeMode Resize
 	// Adjust the start offset when resizing from the beginning
 	if (ResizeMode == SSRM_LeadingEdge)
 	{
-		float StartOffset = (ResizeTime - InitialStartTimeDuringResize) * Section.GetPlayRate();
+		float StartOffset = (ResizeTime - InitialStartTimeDuringResize) * Section.Params.PlayRate;
 		StartOffset += InitialStartOffsetDuringResize;
 
 		// Ensure start offset is not less than 0 and adjust ResizeTime
 		if (StartOffset < 0)
 		{
-			ResizeTime = ResizeTime - (StartOffset / Section.GetPlayRate());
+			ResizeTime = ResizeTime - (StartOffset / Section.Params.PlayRate);
 
 			StartOffset = 0.f;
 		}
 
-		Section.SetStartOffset(StartOffset);
+		Section.Params.StartOffset = StartOffset;
 	}
 
 	ISequencerSection::ResizeSection(ResizeMode, ResizeTime);
diff --git a/Engine/Source/Editor/MovieSceneTools/Private/TrackEditors/SubTrackEditor.cpp b/Engine/Source/Editor/MovieSceneTools/Private/TrackEditors/SubTrackEditor.cpp
--- a/Engine/Source/Editor/MovieSceneTools/Private/TrackEditors/SubTrackEditor.cpp
+++ b/Engine/Source/Editor/MovieSceneTools/Private/TrackEditors/SubTrackEditor.cpp
@@ -34,10 +34,6 @@ public:
 		, SectionObject(*CastChecked<UMovieSceneSubSection>(&InSection))
 		, Sequencer(InSequencer)
 	{
-		if (InSequencer->HasSequenceInstanceForSection(InSection))
-		{
-			SequenceInstance = InSequencer->GetSequenceInstanceForSection(InSection);
-		}
 	}
 
 public:
@@ -135,9 +131,9 @@ public:
 		}
 
 		// add box for the working size
-		const float StartOffset = SectionObject.TimeScale * SectionObject.StartOffset;
-		const float WorkingStart = -SectionObject.TimeScale * PlaybackRange.GetLowerBoundValue() - StartOffset;
-		const float WorkingSize = SectionObject.TimeScale * (MovieScene != nullptr ? MovieScene->GetEditorData().WorkingRange.Size<float>() : 1.0f);
+		const float StartOffset = SectionObject.Parameters.TimeScale * SectionObject.Parameters.StartOffset;
+		const float WorkingStart = -SectionObject.Parameters.TimeScale * PlaybackRange.GetLowerBoundValue() - StartOffset;
+		const float WorkingSize = SectionObject.Parameters.TimeScale * (MovieScene != nullptr ? MovieScene->GetEditorData().WorkingRange.Size<float>() : 1.0f);
 		
 		if(UMovieSceneSubSection::GetRecordingSection() == &SectionObject)
 		{
@@ -181,7 +177,7 @@ public:
 		}
 
 		// add green line for playback start
-		if (StartOffset <= 0)
+		if (StartOffset < 0)
 		{
 			FSlateDrawElement::MakeBox(
 				InPainter.DrawElements,
@@ -197,8 +193,8 @@ public:
 			);
 		}
 
-		// add dark tint for left out-of-bounds & working range
-		const float PlaybackEnd = SectionObject.TimeScale * PlaybackRange.Size<float>() - StartOffset;
+		// add dark tint for right out-of-bounds & working range
+		const float PlaybackEnd = SectionObject.Parameters.TimeScale * PlaybackRange.Size<float>() - StartOffset;
 
 		if (PlaybackEnd < SectionSize)
 		{
@@ -290,9 +286,6 @@ private:
 	/** The section we are visualizing */
 	UMovieSceneSubSection& SectionObject;
 
-	/** The instance that this section is part of */
-	TWeakPtr<FMovieSceneSequenceInstance> SequenceInstance;
-
 	/** Sequencer interface */
 	TWeakPtr<ISequencer> Sequencer;
 };
diff --git a/Engine/Source/Editor/MovieSceneTools/Private/TrackEditors/TransformTrackEditor.cpp b/Engine/Source/Editor/MovieSceneTools/Private/TrackEditors/TransformTrackEditor.cpp
--- a/Engine/Source/Editor/MovieSceneTools/Private/TrackEditors/TransformTrackEditor.cpp
+++ b/Engine/Source/Editor/MovieSceneTools/Private/TrackEditors/TransformTrackEditor.cpp
@@ -192,8 +192,7 @@ private:
 	{
 		if (Sequencer != nullptr)
 		{
-			TArray<TWeakObjectPtr<UObject>> RuntimeObjects;
-			Sequencer->GetRuntimeObjects(Sequencer->GetFocusedMovieSceneSequenceInstance(), ObjectBinding, RuntimeObjects);
+			TArrayView<TWeakObjectPtr<UObject>> RuntimeObjects = Sequencer->FindBoundObjects(ObjectBinding, Sequencer->GetFocusedTemplateID());
 			if (RuntimeObjects.Num() == 1)
 			{
 				return MovieSceneHelpers::SceneComponentFromRuntimeObject(RuntimeObjects[0].Get());
@@ -403,8 +402,7 @@ TSharedRef<ISequencerSection> F3DTransformTrackEditor::MakeSectionInterface(UMov
 
 void F3DTransformTrackEditor::OnPreTransformChanged( UObject& InObject )
 {
-	UMovieSceneSequence* MovieSceneSequence = GetMovieSceneSequence();
-	float AutoKeyTime = GetTimeForKey( MovieSceneSequence );
+	float AutoKeyTime = GetTimeForKey();
 
 	if( IsAllowedToAutoKey() )
 	{
@@ -522,10 +520,7 @@ void F3DTransformTrackEditor::BuildObjectBindingEditButtons(TSharedPtr<SHorizont
 {
 	bool bHasCameraComponent = false;
 
-	TArray<TWeakObjectPtr<UObject>> OutObjects;
-	GetSequencer()->GetRuntimeObjects( GetSequencer()->GetFocusedMovieSceneSequenceInstance(), ObjectGuid, OutObjects);
-
-	for (auto Object : OutObjects)
+	for (auto Object : GetSequencer()->FindObjectsInCurrentSequence(ObjectGuid))
 	{
 		AActor* Actor = Cast<AActor>(Object.Get());
 		if (Actor != nullptr)
@@ -593,10 +588,7 @@ bool F3DTransformTrackEditor::CanAddTransformTrackForActorHandle( FGuid ObjectBi
 EVisibility
 F3DTransformTrackEditor::IsCameraVisible(FGuid ObjectGuid) const
 {
-	TArray<TWeakObjectPtr<UObject>> OutObjects;
-	GetSequencer()->GetRuntimeObjects( GetSequencer()->GetFocusedMovieSceneSequenceInstance(), ObjectGuid, OutObjects);
-
-	for (auto Object : OutObjects)
+	for (auto Object : GetSequencer()->FindObjectsInCurrentSequence(ObjectGuid))
 	{
 		AActor* Actor = Cast<AActor>(Object.Get());
 		
@@ -611,12 +603,9 @@ F3DTransformTrackEditor::IsCameraVisible(FGuid ObjectGuid) const
 
 ECheckBoxState F3DTransformTrackEditor::IsCameraLocked(FGuid ObjectGuid) const
 {
-	TArray<TWeakObjectPtr<UObject>> OutObjects;
-	GetSequencer()->GetRuntimeObjects( GetSequencer()->GetFocusedMovieSceneSequenceInstance(), ObjectGuid, OutObjects);
-
 	TWeakObjectPtr<AActor> CameraActor;
 
-	for (auto Object : OutObjects)
+	for (auto Object : GetSequencer()->FindObjectsInCurrentSequence(ObjectGuid))
 	{
 		AActor* Actor = Cast<AActor>(Object.Get());
 		
@@ -670,12 +659,9 @@ ECheckBoxState F3DTransformTrackEditor::IsCameraLocked(FGuid ObjectGuid) const
 
 void F3DTransformTrackEditor::OnLockCameraClicked(ECheckBoxState CheckBoxState, FGuid ObjectGuid)
 {
-	TArray<TWeakObjectPtr<UObject>> OutObjects;
-	GetSequencer()->GetRuntimeObjects( GetSequencer()->GetFocusedMovieSceneSequenceInstance(), ObjectGuid, OutObjects);
-
 	TWeakObjectPtr<AActor> CameraActor;
 
-	for (auto Object : OutObjects)
+	for (auto Object : GetSequencer()->FindObjectsInCurrentSequence(ObjectGuid))
 	{
 		AActor* Actor = Cast<AActor>(Object.Get());
 
@@ -741,12 +727,9 @@ void F3DTransformTrackEditor::OnLockCameraClicked(ECheckBoxState CheckBoxState,
 
 FText F3DTransformTrackEditor::GetLockCameraToolTip(FGuid ObjectGuid) const
 {
-	TArray<TWeakObjectPtr<UObject>> OutObjects;
-	GetSequencer()->GetRuntimeObjects( GetSequencer()->GetFocusedMovieSceneSequenceInstance(), ObjectGuid, OutObjects);
-
 	TWeakObjectPtr<AActor> CameraActor;
 
-	for (auto Object : OutObjects)
+	for (auto Object : GetSequencer()->FindObjectsInCurrentSequence(ObjectGuid))
 	{
 		AActor* Actor = Cast<AActor>(Object.Get());
 
@@ -816,10 +799,7 @@ void F3DTransformTrackEditor::GetTransformKeys( const FTransformData& LastTransf
 
 void F3DTransformTrackEditor::AddTransformKeysForHandle( FGuid ObjectHandle, EKey3DTransformChannel::Type ChannelToKey, ESequencerKeyMode KeyMode )
 {
-	TArray<TWeakObjectPtr<UObject>> OutObjects;
-	GetSequencer()->GetRuntimeObjects( GetSequencer()->GetFocusedMovieSceneSequenceInstance(), ObjectHandle, OutObjects );
-
-	for ( TWeakObjectPtr<UObject> Object : OutObjects )
+	for ( TWeakObjectPtr<UObject> Object : GetSequencer()->FindObjectsInCurrentSequence(ObjectHandle) )
 	{
 		AddTransformKeysForObject(Object.Get(), ChannelToKey, KeyMode);
 	}
diff --git a/Engine/Source/Editor/MovieSceneTools/Public/FloatCurveKeyArea.h b/Engine/Source/Editor/MovieSceneTools/Public/FloatCurveKeyArea.h
--- a/Engine/Source/Editor/MovieSceneTools/Public/FloatCurveKeyArea.h
+++ b/Engine/Source/Editor/MovieSceneTools/Public/FloatCurveKeyArea.h
@@ -71,7 +71,7 @@ public:
 	virtual bool CanSetExtrapolationMode() const override;
 	virtual void SetKeyInterpMode(FKeyHandle KeyHandle, ERichCurveInterpMode InterpMode) override;
 	virtual void SetKeyTangentMode(FKeyHandle KeyHandle, ERichCurveTangentMode TangentMode) override;
-	virtual void SetKeyTime(FKeyHandle KeyHandle, float NewKeyTime) const override;
+	virtual void SetKeyTime(FKeyHandle KeyHandle, float NewKeyTime) override;
 	virtual void CopyKeys(FMovieSceneClipboardBuilder& ClipboardBuilder, const TFunctionRef<bool(FKeyHandle, const IKeyArea&)>& KeyMask) const override;
 	virtual void PasteKeys(const FMovieSceneClipboardKeyTrack& KeyTrack, const FMovieSceneClipboardEnvironment& SrcEnvironment, const FSequencerPasteEnvironment& DstEnvironment) override;
 
diff --git a/Engine/Source/Editor/MovieSceneTools/Public/GenericKeyArea.h b/Engine/Source/Editor/MovieSceneTools/Public/GenericKeyArea.h
new file mode 100644
--- /dev/null
+++ b/Engine/Source/Editor/MovieSceneTools/Public/GenericKeyArea.h
@@ -0,0 +1,263 @@
+// Copyright 1998-2016 Epic Games, Inc. All Rights Reserved.
+
+#pragma once
+
+#include "NamedKeyArea.h"
+#include "ClipboardTypes.h"
+#include "DetailLayoutBuilder.h"
+#include "DetailCategoryBuilder.h"
+#include "DetailWidgetRow.h"
+#include "Curves/CurveInterface.h"
+
+#include "GenericKeyArea.generated.h"
+
+class UMovieSceneSection;
+
+/** Interface for editing a key value/time */
+struct ICurveEditInterface
+{
+	virtual ~ICurveEditInterface() {}
+
+	/** Extend the specified (empty) details customization with the specified key handle */
+	virtual void Extend(FKeyHandle KeyHandle, IDetailLayoutBuilder& DetailBuilder) = 0;
+};
+
+/** Largely redundant structure that is used to forward detail customizations for keys */
+USTRUCT()
+struct FMovieSceneCurveInterfaceKeyEditStruct
+{
+	GENERATED_BODY()
+
+	UPROPERTY()
+	bool bBenignPropertyToEnsurePropertyCustomization;
+
+	/** The key handle to edit */
+	FKeyHandle KeyHandle;
+
+	/** The curve interface to edit with */
+	ICurveEditInterface* EditInterface;
+};
+
+/**
+ * A generic key area that utilizes TCurveInterface for interaction
+ */
+template<typename KeyValueType, typename TimeType>
+class TGenericKeyArea
+	: public FNamedKeyArea
+	, public ICurveEditInterface
+{
+public:
+
+	/**
+	 * Create and initialize a new instance.
+	 *
+	 * @param InCurve The curve to assign to this key area.
+	 * @param InOwningSection The section that owns this key area.
+	 */
+	TGenericKeyArea(TArray<TimeType>* KeyTimesParam, TArray<KeyValueType>* KeyValuesParam, UMovieSceneSection* InOwningSection)
+		: CurveInterface(KeyTimesParam, KeyValuesParam)
+		, OwningSection(InOwningSection)
+	{ }
+
+	TGenericKeyArea(TCurveInterface<KeyValueType, TimeType>&& InCurveInterface, UMovieSceneSection* InOwningSection)
+		: CurveInterface(InCurveInterface)
+		, OwningSection(InOwningSection)
+	{ }
+
+public:
+
+	//~ IKeyArea interface
+
+	virtual TArray<FKeyHandle> AddKeyUnique(TimeType Time, EMovieSceneKeyInterpolation InKeyInterpolation, TimeType TimeToCopyFrom) override
+	{
+		ModifySection();
+
+		TArray<FKeyHandle> AddedKeyHandles;
+		TOptional<FKeyHandle> CurrentKey = CurveInterface.FindKey(
+			[=](TimeType ExistingTime)
+			{
+				return FMath::IsNearlyEqual(Time, ExistingTime, KINDA_SMALL_NUMBER);
+			}
+		);
+
+		if (CurrentKey.IsSet())
+		{
+			return AddedKeyHandles;
+		}
+
+		ExtendSectionBounds(Time);
+
+		FKeyHandle Handle = CurveInterface.AddKey(Time);
+		AddedKeyHandles.Add(Handle);
+
+		return AddedKeyHandles;
+	}
+
+	virtual TOptional<FKeyHandle> DuplicateKey(FKeyHandle KeyToDuplicate) override
+	{
+		ModifySection();
+
+		TOptional<TKeyFrameProxy<KeyValueType, TimeType>> Key = CurveInterface.GetKey(KeyToDuplicate);
+		if (Key.IsSet())
+		{
+			return CurveInterface.AddKeyValue(Key->Time, CopyTemp(Key->Value));
+		}
+
+		return TOptional<FKeyHandle>();
+	}
+
+	virtual void DeleteKey(FKeyHandle KeyHandle) override
+	{
+		ModifySection();
+
+		CurveInterface.RemoveKey(KeyHandle);
+	}
+
+	virtual TimeType GetKeyTime(FKeyHandle KeyHandle) const override
+	{
+		TOptional<TimeType> Time = CurveInterface.GetKeyTime(KeyHandle);
+		return Time.Get(TNumericLimits<TimeType>::Lowest());
+	}
+
+	virtual void SetKeyTime(FKeyHandle KeyHandle, TimeType NewKeyTime) override
+	{
+		ModifySection();
+
+		CurveInterface.SetKeyTime(KeyHandle, NewKeyTime);
+		ExtendSectionBounds(NewKeyTime);
+	}
+
+	virtual FKeyHandle MoveKey(FKeyHandle KeyHandle, TimeType DeltaPosition) override
+	{
+		TOptional<TimeType> Time = CurveInterface.GetKeyTime(KeyHandle);
+		if (Time.IsSet())
+		{
+			ModifySection();
+			
+			CurveInterface.SetKeyTime(KeyHandle, Time.GetValue() + DeltaPosition);
+		}
+		return KeyHandle;
+	}
+
+	virtual UMovieSceneSection* GetOwningSection() override
+	{
+		return OwningSection.Get();
+	}
+
+	virtual TArray<FKeyHandle> GetUnsortedKeyHandles() const override
+	{
+		TKeyTimeIterator<TimeType> Iterator = CurveInterface.IterateKeys();
+
+		TArray<FKeyHandle> OutKeyHandles;
+		OutKeyHandles.Reserve(Iterator.GetEndIndex() - Iterator.GetStartIndex());
+
+		for (auto It = begin(Iterator); It != end(Iterator); ++It)
+		{
+			OutKeyHandles.Add(It.GetKeyHandle());
+		}
+
+		return OutKeyHandles;
+	}
+
+	virtual TSharedPtr<FStructOnScope> GetKeyStruct(FKeyHandle KeyHandle) override
+	{
+		auto Struct = MakeShared<FStructOnScope>(FMovieSceneCurveInterfaceKeyEditStruct::StaticStruct());
+
+		auto* StructPtr = reinterpret_cast<FMovieSceneCurveInterfaceKeyEditStruct*>(Struct->GetStructMemory());
+		StructPtr->KeyHandle = KeyHandle;
+		StructPtr->EditInterface = this;
+
+		return Struct;
+	}
+
+	virtual bool CanCreateKeyEditor() override																{ return false; }
+	virtual TSharedRef<SWidget> CreateKeyEditor(ISequencer* Sequencer) override								{ return SNullWidget::NullWidget; }
+	virtual TOptional<FLinearColor> GetColor() override														{ return TOptional<FLinearColor>(); }
+	virtual ERichCurveExtrapolation GetExtrapolationMode(bool bPreInfinity) const override					{ return RCCE_None; }
+	virtual ERichCurveInterpMode GetKeyInterpMode(FKeyHandle Keyhandle) const override						{ return RCIM_None; }
+	virtual ERichCurveTangentMode GetKeyTangentMode(FKeyHandle KeyHandle) const override					{ return RCTM_None; }
+	virtual void SetExtrapolationMode(ERichCurveExtrapolation ExtrapMode, bool bPreInfinity) override		{}
+	virtual void SetKeyInterpMode(FKeyHandle KeyHandle, ERichCurveInterpMode InterpMode) override			{}
+	virtual void SetKeyTangentMode(FKeyHandle KeyHandle, ERichCurveTangentMode TangentMode) override		{}
+	virtual FRichCurve* GetRichCurve() override																{ return nullptr; }
+
+	
+	virtual void CopyKeys(FMovieSceneClipboardBuilder& ClipboardBuilder, const TFunctionRef<bool(FKeyHandle, const IKeyArea&)>& KeyMask) const override
+	{
+
+	}
+	virtual void PasteKeys(const FMovieSceneClipboardKeyTrack& KeyTrack, const FMovieSceneClipboardEnvironment& SrcEnvironment, const FSequencerPasteEnvironment& DstEnvironment) override
+	{
+
+	}
+
+private:
+
+	virtual void Extend(FKeyHandle KeyHandle, IDetailLayoutBuilder& DetailBuilder) override
+	{
+		TOptional<TKeyFrameProxy<KeyValueType, TimeType>> Key = CurveInterface.GetKey(KeyHandle);
+		if (Key.IsSet())
+		{
+			IDetailCategoryBuilder& GeneralCategory = DetailBuilder.EditCategory("General");
+
+			auto OnValueChanged = [=](TimeType InTime){ this->SetKeyTime(KeyHandle, InTime); };
+
+			FText TimeText = NSLOCTEXT("GenericKeyArea", "TimeParameter", "Time");
+			FText TimeTooltipText = NSLOCTEXT("GenericKeyArea", "TimeParameter_ToolTip", "The time of this key");
+			GeneralCategory.AddCustomRow(TimeText, false)
+				.NameContent()
+				[
+					SNew(STextBlock)
+					.Text(TimeText)
+					.Font(DetailBuilder.GetDetailFont())
+					.ToolTipText(TimeTooltipText)
+				]
+				.ValueContent()
+				[
+					SNew(SSpinBox<TimeType>)
+					.Value_Lambda([=]{ return this->CurveInterface.GetKeyTime(KeyHandle).Get(TimeType()); })
+					.OnValueChanged_Lambda(OnValueChanged)
+					.OnValueCommitted_Lambda([=](TimeType InTime, ETextCommit::Type){ OnValueChanged(InTime); })
+					.ToolTipText(TimeTooltipText)
+				];
+
+			TSharedRef<FStructOnScope> KeyValue = MakeShared<FStructOnScope>(KeyValueType::StaticStruct(), (uint8*)&Key->Value);
+			GeneralCategory.AddExternalProperties(KeyValue);
+		}
+	}
+
+	void ExtendSectionBounds(TimeType IncludeTime)
+	{
+		UMovieSceneSection* Section = OwningSection.Get();
+		if (!Section)
+		{
+			return;
+		}
+
+		if (Section->GetStartTime() > IncludeTime)
+		{
+			Section->SetStartTime(IncludeTime);
+		}
+
+		if (Section->GetEndTime() < IncludeTime)
+		{
+			Section->SetEndTime(IncludeTime);
+		}
+	}
+
+	void ModifySection()
+	{
+		if (UMovieSceneSection* Section = OwningSection.Get())
+		{
+			Section->Modify();
+		}
+	}
+
+protected:
+
+	/** The curve managed by this area. */
+	TCurveInterface<KeyValueType, TimeType> CurveInterface;
+
+	/** The section that owns this area. */
+	TWeakObjectPtr<UMovieSceneSection> OwningSection;
+};
diff --git a/Engine/Source/Editor/MovieSceneTools/Public/IntegerKeyArea.h b/Engine/Source/Editor/MovieSceneTools/Public/IntegerKeyArea.h
new file mode 100644
--- /dev/null
+++ b/Engine/Source/Editor/MovieSceneTools/Public/IntegerKeyArea.h
@@ -0,0 +1,52 @@
+// Copyright 1998-2016 Epic Games, Inc. All Rights Reserved.
+
+#pragma once
+
+#include "IntegralKeyArea.h"
+
+
+class UMovieSceneSection;
+
+
+/**
+ * A key area for displaying and editing integer curves.
+ */
+class FIntegerKeyArea
+	: public FIntegralKeyArea<int32>
+{
+public:
+
+	/**
+	* Creates a new key area for editing integer curves.
+	*
+	* @param InCurve The integral curve which has the integer keys.
+	* @param InOwningSection The section which owns the curve which is being displayed and edited by this area.
+	*/
+	FIntegerKeyArea(FIntegralCurve& InCurve, UMovieSceneSection* InOwningSection)
+		: FIntegralKeyArea(InCurve, InOwningSection)
+	{ }
+
+	/**
+	* Creates a new key area for editing integer curves whose value can be overridden externally.
+	*
+	* @param InCurve The integral curve which has the integer keys.
+	* @param InExternalValue An attribute which can provide an external value for this key area.  External values are
+	*        useful for things like property tracks where the property value can change without changing the animation
+	*        and we want to be able to key and update using the new property value.
+	* @param InOwningSection The section which owns the curve which is being displayed and edited by this area.
+	*/
+	FIntegerKeyArea(FIntegralCurve& InCurve, TAttribute<TOptional<int32>> ExternalValue, UMovieSceneSection* InOwningSection)
+		: FIntegralKeyArea(InCurve, ExternalValue, InOwningSection)
+	{ }
+
+public:
+
+	//~ FIntegralKeyArea interface
+
+	virtual bool CanCreateKeyEditor() override;
+	virtual TSharedRef<SWidget> CreateKeyEditor(ISequencer* Sequencer) override;
+
+protected:
+
+	virtual int32 ConvertCurveValueToIntegralType(int32 CurveValue) const { return CurveValue; }
+};
diff --git a/Engine/Source/Editor/MovieSceneTools/Public/IntegralKeyArea.h b/Engine/Source/Editor/MovieSceneTools/Public/IntegralKeyArea.h
--- a/Engine/Source/Editor/MovieSceneTools/Public/IntegralKeyArea.h
+++ b/Engine/Source/Editor/MovieSceneTools/Public/IntegralKeyArea.h
@@ -45,7 +45,7 @@ public:
 	virtual void SetExtrapolationMode(ERichCurveExtrapolation ExtrapMode, bool bPreInfinity) override;
 	virtual void SetKeyInterpMode(FKeyHandle KeyHandle, ERichCurveInterpMode InterpMode) override;
 	virtual void SetKeyTangentMode(FKeyHandle KeyHandle, ERichCurveTangentMode TangentMode) override;
-	virtual void SetKeyTime(FKeyHandle KeyHandle, float NewKeyTime) const override;
+	virtual void SetKeyTime(FKeyHandle KeyHandle, float NewKeyTime) override;
 
 protected:
 
diff --git a/Engine/Source/Editor/MovieSceneTools/Public/KeyframeTrackEditor.h b/Engine/Source/Editor/MovieSceneTools/Public/KeyframeTrackEditor.h
--- a/Engine/Source/Editor/MovieSceneTools/Public/KeyframeTrackEditor.h
+++ b/Engine/Source/Editor/MovieSceneTools/Public/KeyframeTrackEditor.h
@@ -268,21 +268,26 @@ private:
 	bool SetDefault( TrackType* Track, float Time, const KeyDataType& KeyData, bool bInfiniteKeyAreas )
 	{
 		bool bSectionAdded = false;
-		Track->Modify();
 		const TArray<UMovieSceneSection*>& Sections = Track->GetAllSections();
 		if ( Sections.Num() > 0)
 		{
 			for ( UMovieSceneSection* Section : Sections )
 			{
 				IKeyframeSection<KeyDataType>* KeyframeSection = CastChecked<SectionType>( Section );
-				KeyframeSection->SetDefault( KeyData );
+				if (!KeyframeSection->HasKeys(KeyData))
+				{
+					KeyframeSection->SetDefault( KeyData );
+				}
 			}
 		}
 		else
 		{
 			UMovieSceneSection* NewSection = Track->FindOrAddSection( Time, bSectionAdded );
 			IKeyframeSection<KeyDataType>* KeyframeSection = CastChecked<SectionType>( NewSection );
-			KeyframeSection->SetDefault( KeyData );
+			if (!KeyframeSection->HasKeys(KeyData))
+			{
+				KeyframeSection->SetDefault( KeyData );
+			}
 
 			if (bSectionAdded)
 			{
diff --git a/Engine/Source/Editor/MovieSceneTools/Public/MatineeImportTools.h b/Engine/Source/Editor/MovieSceneTools/Public/MatineeImportTools.h
--- a/Engine/Source/Editor/MovieSceneTools/Public/MatineeImportTools.h
+++ b/Engine/Source/Editor/MovieSceneTools/Public/MatineeImportTools.h
@@ -75,7 +75,7 @@ public:
 	static bool CopyInterpFadeTrack( UInterpTrackFade* MatineeFadeTrack, UMovieSceneFadeTrack* FadeTrack );
 
 	/** Copies keys from a matinee director track to a sequencer camera cut track. */
-	static bool CopyInterpDirectorTrack( UInterpTrackDirector* DirectorTrack, UMovieSceneCameraCutTrack* CameraCutTrack, AMatineeActor* MatineeActor, UMovieSceneSequence* MovieSceneSequence );
+	static bool CopyInterpDirectorTrack( UInterpTrackDirector* DirectorTrack, UMovieSceneCameraCutTrack* CameraCutTrack, AMatineeActor* MatineeActor, IMovieScenePlayer& Player );
 
 	/** Copies keys from a matinee event track to a sequencer event track. */
 	static bool CopyInterpEventTrack( UInterpTrackEvent* MatineeEventTrack, UMovieSceneEventTrack* EventTrack );
diff --git a/Engine/Source/Editor/MovieSceneTools/Public/MovieSceneToolHelpers.h b/Engine/Source/Editor/MovieSceneTools/Public/MovieSceneToolHelpers.h
--- a/Engine/Source/Editor/MovieSceneTools/Public/MovieSceneToolHelpers.h
+++ b/Engine/Source/Editor/MovieSceneTools/Public/MovieSceneToolHelpers.h
@@ -128,5 +128,5 @@ public:
 	 * @param InObjectBindingNameMap The object binding to name map to map import fbx animation onto
 	 * @return Whether the import was successful
 	 */
-	static bool ImportFBX(UMovieScene* InMovieScene, FMovieSceneSequenceInstance& InSequence, ISequencer& InSequencer, const TMap<FGuid, FString>& InObjectBindingNameMap);
+	static bool ImportFBX(UMovieScene* InMovieScene, ISequencer& InSequencer, const TMap<FGuid, FString>& InObjectBindingNameMap);
 };
diff --git a/Engine/Source/Editor/MovieSceneTools/Public/PropertySection.h b/Engine/Source/Editor/MovieSceneTools/Public/PropertySection.h
--- a/Engine/Source/Editor/MovieSceneTools/Public/PropertySection.h
+++ b/Engine/Source/Editor/MovieSceneTools/Public/PropertySection.h
@@ -61,35 +61,27 @@ protected:
 	template<typename ValueType>
 	TOptional<ValueType> GetPropertyValue() const
 	{
-		UObject* RuntimeObject = GetRuntimeObjectAndUpdatePropertyBindings();
-
-		if (RuntimeObject == nullptr)
+		if (!PropertyBindings.IsValid())
 		{
 			return TOptional<ValueType>();
 		}
 
-		// Bool property values are stored in a bit field so using a straight cast of the pointer to get their value does not
-		// work.  Instead use the actual property to get the correct value.
-		const UBoolProperty* BoolProperty = Cast<const UBoolProperty>(GetProperty());
-		if (BoolProperty)
 		{
-			uint32* ValuePtr = BoolProperty->ContainerPtrToValuePtr<uint32>(RuntimeObject);
-			bool BoolPropertyValue = BoolProperty->GetPropertyValue(ValuePtr);
-			void *PropertyValue = &BoolPropertyValue;
-			return TOptional<ValueType>(*((ValueType*)PropertyValue));
+			for (const TWeakObjectPtr<UObject>& WeakObject : Sequencer->FindBoundObjects(ObjectBinding, Sequencer->GetFocusedTemplateID()))
+			{
+				if (UObject* Object = WeakObject.Get())
+				{
+					return PropertyBindings->GetCurrentValue<ValueType>(*Object);
+				}
+			}
 		}
 
-		return TOptional<ValueType>(PropertyBindings->GetCurrentValue<ValueType>(RuntimeObject));
+		return TOptional<ValueType>();
 	}
 
 	/** Returns true when this section was constructed with the data necessary to query for the current property value. */
 	bool CanGetPropertyValue() const;
 
-private:
-	/** Gets the current runtime object for the object binding which was used to construct this section.  Returns nullptr
-		if no valid object can be found, or if an object binding was not supplied. */
-	UObject* GetRuntimeObjectAndUpdatePropertyBindings() const;
-
 protected:
 
 	/** Display name of the section */
@@ -108,9 +100,5 @@ private:
 
 	/** An object which is used to retrieve the value of a property based on it's name and path. */
 	mutable TSharedPtr<FTrackInstancePropertyBindings> PropertyBindings;
-
-	/** Caches a weak pointer to the most recent runtime object retrieved by this section.  This cache
-		is maintained so that the property bindings can be updated any time this object changes. */
-	mutable TWeakObjectPtr<UObject> RuntimeObjectCache;
 };
 
diff --git a/Engine/Source/Editor/NiagaraEditor/Private/NiagaraSequencer.cpp b/Engine/Source/Editor/NiagaraEditor/Private/NiagaraSequencer.cpp
--- a/Engine/Source/Editor/NiagaraEditor/Private/NiagaraSequencer.cpp
+++ b/Engine/Source/Editor/NiagaraEditor/Private/NiagaraSequencer.cpp
@@ -17,21 +17,21 @@ UNiagaraMovieSceneSection::UNiagaraMovieSceneSection(const FObjectInitializer& O
 }
 
 
-/*
- *  This is called when the user edits a track in the effect editor timeline
- */
-void INiagaraTrackInstance::RefreshInstance(const TArray<TWeakObjectPtr<UObject>>& RuntimeObjects, IMovieScenePlayer& Player, FMovieSceneSequenceInstance& SequenceInstance)
-{
-	// every track should have exactly one section, always
-	check(Track->GetAllSections().Num() == 1)
-
-	UNiagaraMovieSceneSection *Section = Cast<UNiagaraMovieSceneSection>( Track->GetAllSections()[0] );
+// /*
+//  *  This is called when the user edits a track in the effect editor timeline
+//  */
+// void INiagaraTrackInstance::RefreshInstance(const TArray<TWeakObjectPtr<UObject>>& RuntimeObjects, IMovieScenePlayer& Player)
+// {
+// 	// every track should have exactly one section, always
+// 	check(Track->GetAllSections().Num() == 1)
+
+// 	UNiagaraMovieSceneSection *Section = Cast<UNiagaraMovieSceneSection>( Track->GetAllSections()[0] );
 	
-	check(Section);
+// 	check(Section);
 
-	Track->GetEmitter()->GetProperties()->StartTime = Section->GetStartTime();
-	Track->GetEmitter()->GetProperties()->EndTime = Section->GetEndTime();
-}
+// 	Track->GetEmitter()->GetProperties()->StartTime = Section->GetStartTime();
+// 	Track->GetEmitter()->GetProperties()->EndTime = Section->GetEndTime();
+// }
 
 
 #undef LOCTEXT_NAMESPACE
diff --git a/Engine/Source/Editor/NiagaraEditor/Private/NiagaraSequencer.h b/Engine/Source/Editor/NiagaraEditor/Private/NiagaraSequencer.h
--- a/Engine/Source/Editor/NiagaraEditor/Private/NiagaraSequencer.h
+++ b/Engine/Source/Editor/NiagaraEditor/Private/NiagaraSequencer.h
@@ -115,28 +115,28 @@ private:
 /**
 *	One instance of a UEmitterMovieSceneTrack
 */
-class INiagaraTrackInstance
-	: public IMovieSceneTrackInstance
-{
-public:
-	INiagaraTrackInstance(UEmitterMovieSceneTrack *InTrack)
-		: Track(InTrack)
-	{
-	}
-
-	virtual void Update(EMovieSceneUpdateData& UpdateData, const TArray<TWeakObjectPtr<UObject>>& RuntimeObjects, IMovieScenePlayer& Player, FMovieSceneSequenceInstance& SequenceInstance)
-	{
-	}
-
-	virtual void RefreshInstance(const TArray<TWeakObjectPtr<UObject>>& RuntimeObjects, IMovieScenePlayer& Player, FMovieSceneSequenceInstance& SequenceInstance);
-	virtual void ClearInstance(IMovieScenePlayer& Player, FMovieSceneSequenceInstance& SequenceInstance) override {}
-	virtual void SaveState(const TArray<TWeakObjectPtr<UObject>>& RuntimeObjects, IMovieScenePlayer& Player, FMovieSceneSequenceInstance& SequenceInstance) override {}
-	virtual void RestoreState(const TArray<TWeakObjectPtr<UObject>>& RuntimeObjects, IMovieScenePlayer& Player, FMovieSceneSequenceInstance& SequenceInstance) override {}
-
-private:
-	TSharedPtr<FNiagaraSimulation> Emitter;
-	UEmitterMovieSceneTrack *Track;
-};
+// class INiagaraTrackInstance
+// 	: public IMovieSceneTrackInstance
+// {
+// public:
+// 	INiagaraTrackInstance(UEmitterMovieSceneTrack *InTrack)
+// 		: Track(InTrack)
+// 	{
+// 	}
+
+// 	virtual void Update(EMovieSceneUpdateData& UpdateData, const TArray<TWeakObjectPtr<UObject>>& RuntimeObjects, IMovieScenePlayer& Player)
+// 	{
+// 	}
+
+// 	virtual void RefreshInstance(const TArray<TWeakObjectPtr<UObject>>& RuntimeObjects, IMovieScenePlayer& Player, FMovieSceneSequenceInstance& SequenceInstance);
+// 	virtual void ClearInstance(IMovieScenePlayer& Player, FMovieSceneSequenceInstance& SequenceInstance) override {}
+// 	virtual void SaveState(const TArray<TWeakObjectPtr<UObject>>& RuntimeObjects, IMovieScenePlayer& Player, FMovieSceneSequenceInstance& SequenceInstance) override {}
+// 	virtual void RestoreState(const TArray<TWeakObjectPtr<UObject>>& RuntimeObjects, IMovieScenePlayer& Player, FMovieSceneSequenceInstance& SequenceInstance) override {}
+
+// private:
+// 	TSharedPtr<FNiagaraSimulation> Emitter;
+// 	UEmitterMovieSceneTrack *Track;
+// };
 
 
 /**
@@ -196,11 +196,6 @@ public:
 		return false;
 	}
 
-	virtual TSharedPtr<IMovieSceneTrackInstance> CreateInstance()
-	{
-		return MakeShareable(new INiagaraTrackInstance(this));
-	}
-
 	virtual const TArray<UMovieSceneSection*>& GetAllSections() const override
 	{
 		return Sections;
diff --git a/Engine/Source/Editor/PropertyEditor/Private/CustomChildBuilder.cpp b/Engine/Source/Editor/PropertyEditor/Private/CustomChildBuilder.cpp
--- a/Engine/Source/Editor/PropertyEditor/Private/CustomChildBuilder.cpp
+++ b/Engine/Source/Editor/PropertyEditor/Private/CustomChildBuilder.cpp
@@ -55,6 +55,11 @@ IDetailPropertyRow& FCustomChildrenBuilder::AddChildProperty( TSharedRef<IProper
 	return *NewCustomization.PropertyRow;
 }
 
+TArray<TSharedPtr<IPropertyHandle>> FCustomChildrenBuilder::AddChildStructure( TSharedRef<FStructOnScope> ChildStructure )
+{
+	return ParentCategory.Pin()->AddExternalProperties(ChildStructure);
+}
+
 class SStandaloneCustomStructValue : public SCompoundWidget, public IPropertyTypeCustomizationUtils
 {
 public:
diff --git a/Engine/Source/Editor/PropertyEditor/Private/CustomChildBuilder.h b/Engine/Source/Editor/PropertyEditor/Private/CustomChildBuilder.h
--- a/Engine/Source/Editor/PropertyEditor/Private/CustomChildBuilder.h
+++ b/Engine/Source/Editor/PropertyEditor/Private/CustomChildBuilder.h
@@ -13,6 +13,7 @@ public:
 	virtual IDetailGroup& AddChildGroup( FName GroupName, const FText& LocalizedDisplayName ) override;
 	virtual FDetailWidgetRow& AddChildContent( const FText& SearchString ) override;
 	virtual IDetailPropertyRow& AddChildProperty( TSharedRef<IPropertyHandle> PropertyHandle ) override;
+	virtual TArray<TSharedPtr<IPropertyHandle>> AddChildStructure( TSharedRef<FStructOnScope> ChildStructure ) override;
 	virtual TSharedRef<SWidget> GenerateStructValueWidget( TSharedRef<IPropertyHandle> StructPropertyHandle ) override;
 
 	const TArray< FDetailLayoutCustomization >& GetChildCustomizations() const { return ChildCustomizations; }
diff --git a/Engine/Source/Editor/PropertyEditor/Private/DetailCategoryBuilderImpl.cpp b/Engine/Source/Editor/PropertyEditor/Private/DetailCategoryBuilderImpl.cpp
--- a/Engine/Source/Editor/PropertyEditor/Private/DetailCategoryBuilderImpl.cpp
+++ b/Engine/Source/Editor/PropertyEditor/Private/DetailCategoryBuilderImpl.cpp
@@ -381,6 +381,49 @@ IDetailPropertyRow* FDetailCategoryImpl::AddExternalProperty( TSharedPtr<FStruct
 	return nullptr;
 }
 
+TArray<TSharedPtr<IPropertyHandle>> FDetailCategoryImpl::AddExternalProperties( TSharedRef<FStructOnScope> StructData, EPropertyLocation::Type Location )
+{
+	TSharedPtr<FStructurePropertyNode> RootPropertyNode( new FStructurePropertyNode );
+	RootPropertyNode->SetStructure(StructData);
+
+	FPropertyNodeInitParams InitParams;
+	InitParams.ParentNode = nullptr;
+	InitParams.Property = nullptr;
+	InitParams.ArrayOffset = 0;
+	InitParams.ArrayIndex = INDEX_NONE;
+	InitParams.bAllowChildren = true;
+	InitParams.bForceHiddenPropertyVisibility = FPropertySettings::Get().ShowHiddenProperties();
+	InitParams.bCreateCategoryNodes = false;
+
+	RootPropertyNode->InitNode(InitParams);
+
+	TArray<TSharedPtr<IPropertyHandle>> Handles;
+
+	FDetailLayoutBuilderImpl& DetailLayoutBuilderRef = GetParentLayoutImpl();
+
+	const bool bForAdvanced = Location == EPropertyLocation::Advanced;
+	if( RootPropertyNode.IsValid() )
+	{
+		RootPropertyNode->RebuildChildren();
+		DetailLayoutBuilder.Pin()->AddExternalRootPropertyNode( RootPropertyNode.ToSharedRef() );
+
+		for(int32 ChildIdx = 0; ChildIdx < RootPropertyNode->GetNumChildNodes(); ++ChildIdx)
+		{
+			TSharedPtr< FPropertyNode > PropertyNode = RootPropertyNode->GetChildNode(ChildIdx);
+			if(UProperty* Property = PropertyNode->GetProperty())
+			{
+				FDetailLayoutCustomization NewCustomization;
+				NewCustomization.PropertyRow = MakeShared<FDetailPropertyRow>( PropertyNode, AsShared(), RootPropertyNode );
+				AddCustomLayout( NewCustomization, bForAdvanced);
+
+				Handles.Add(DetailLayoutBuilderRef.GetPropertyHandle(PropertyNode));
+			}
+		}
+	}
+
+	return Handles;
+}
+
 void FDetailCategoryImpl::AddPropertyNode( TSharedRef<FPropertyNode> PropertyNode, FName InstanceName )
 {
 	FDetailLayoutCustomization NewCustomization;
diff --git a/Engine/Source/Editor/PropertyEditor/Private/DetailCategoryBuilderImpl.h b/Engine/Source/Editor/PropertyEditor/Private/DetailCategoryBuilderImpl.h
--- a/Engine/Source/Editor/PropertyEditor/Private/DetailCategoryBuilderImpl.h
+++ b/Engine/Source/Editor/PropertyEditor/Private/DetailCategoryBuilderImpl.h
@@ -138,6 +138,7 @@ public:
 	virtual IDetailPropertyRow& AddProperty( TSharedPtr<IPropertyHandle> PropertyHandle, EPropertyLocation::Type Location = EPropertyLocation::Default ) override;
 	virtual IDetailPropertyRow* AddExternalProperty( const TArray<UObject*>& Objects, FName PropertyName, EPropertyLocation::Type Location = EPropertyLocation::Default ) override;
 	virtual IDetailPropertyRow* AddExternalProperty( TSharedPtr<FStructOnScope> StructData, FName PropertyName, EPropertyLocation::Type Location = EPropertyLocation::Default ) override;
+	virtual TArray<TSharedPtr<IPropertyHandle>> AddExternalProperties( TSharedRef<FStructOnScope> StructData, EPropertyLocation::Type Location = EPropertyLocation::Default ) override;
 	virtual IDetailLayoutBuilder& GetParentLayout() const override { return *DetailLayoutBuilder.Pin(); }
 	virtual FDetailWidgetRow& AddCustomRow( const FText& FilterString, bool bForAdvanced = false ) override;
 	virtual void AddCustomBuilder( TSharedRef<IDetailCustomNodeBuilder> InCustomBuilder, bool bForAdvanced = false ) override;
diff --git a/Engine/Source/Editor/PropertyEditor/Private/DetailLayoutBuilderImpl.cpp b/Engine/Source/Editor/PropertyEditor/Private/DetailLayoutBuilderImpl.cpp
--- a/Engine/Source/Editor/PropertyEditor/Private/DetailLayoutBuilderImpl.cpp
+++ b/Engine/Source/Editor/PropertyEditor/Private/DetailLayoutBuilderImpl.cpp
@@ -4,6 +4,7 @@
 #include "PropertyNode.h"
 #include "PropertyHandleImpl.h"
 #include "PropertyEditorHelpers.h"
+#include "StructurePropertyNode.h"
 #include "IPropertyUtilities.h"
 #include "DetailMultiTopLevelObjectRootNode.h"
 #include "IDetailRootObjectCustomization.h"
@@ -517,6 +518,45 @@ void FDetailLayoutBuilderImpl::GetObjectsBeingCustomized( TArray< TWeakObjectPtr
 	}
 }
 
+void FDetailLayoutBuilderImpl::GetStructsBeingCustomized( TArray< TSharedPtr<FStructOnScope> >& OutStructs ) const
+{
+	OutStructs.Reset();
+
+	FStructurePropertyNode* RootStructNode = RootNode.IsValid() ? RootNode.Pin()->AsStructureNode() : nullptr;
+
+	// The class to find properties in defaults to the class currently being customized
+	FName ClassName = CurrentCustomizationClass ? CurrentCustomizationClass->GetFName() : NAME_None;
+
+	if(ClassName != NAME_None && CurrentCustomizationVariableName != NAME_None)
+	{
+		// If this fails there are no properties associated with the class name provided
+		FClassInstanceToPropertyMap* ClassInstanceToPropertyMapPtr = PropertyMap.Find(ClassName);
+
+		if(ClassInstanceToPropertyMapPtr)
+		{
+			FClassInstanceToPropertyMap& ClassInstanceToPropertyMap = *ClassInstanceToPropertyMapPtr;
+
+			FPropertyNodeMap* PropertyNodeMapPtr = ClassInstanceToPropertyMap.Find(CurrentCustomizationVariableName);
+
+			if(PropertyNodeMapPtr)
+			{
+				FPropertyNodeMap& PropertyNodeMap = *PropertyNodeMapPtr;
+				FComplexPropertyNode* ParentComplexProperty = PropertyNodeMap.ParentProperty ? PropertyNodeMap.ParentProperty->AsComplexNode() : nullptr;
+				FStructurePropertyNode* StructureNode = ParentComplexProperty ? ParentComplexProperty->AsStructureNode() : nullptr;
+				if(StructureNode)
+				{
+					OutStructs.Add(StructureNode->GetStructData());
+				}
+			}
+		}
+	}
+
+	if(RootStructNode)
+	{
+		OutStructs.Add(RootStructNode->GetStructData());
+	}
+}
+
 const TSharedRef< IPropertyUtilities >& FDetailLayoutBuilderImpl::GetPropertyUtilities() const
 {
 	return PropertyDetailsUtilities;
diff --git a/Engine/Source/Editor/PropertyEditor/Private/DetailLayoutBuilderImpl.h b/Engine/Source/Editor/PropertyEditor/Private/DetailLayoutBuilderImpl.h
--- a/Engine/Source/Editor/PropertyEditor/Private/DetailLayoutBuilderImpl.h
+++ b/Engine/Source/Editor/PropertyEditor/Private/DetailLayoutBuilderImpl.h
@@ -15,6 +15,7 @@ public:
 	/** IDetailLayoutBuilder Interface */
 	virtual const IDetailsView& GetDetailsView() const override;
 	virtual void GetObjectsBeingCustomized( TArray< TWeakObjectPtr<UObject> >& OutObjects ) const override;
+	virtual void GetStructsBeingCustomized( TArray< TSharedPtr<FStructOnScope> >& OutStructs ) const override;
 	virtual IDetailCategoryBuilder& EditCategory( FName CategoryName, const FText& NewLocalizedDisplayName = FText::GetEmpty(), ECategoryPriority::Type CategoryType = ECategoryPriority::Default ) override;
 	virtual TSharedRef<IPropertyHandle> GetProperty( const FName PropertyPath, const UClass* ClassOutermost, FName InInstanceName ) override;
 	virtual FName GetTopLevelProperty() override;
diff --git a/Engine/Source/Editor/PropertyEditor/Private/SDetailsViewBase.cpp b/Engine/Source/Editor/PropertyEditor/Private/SDetailsViewBase.cpp
--- a/Engine/Source/Editor/PropertyEditor/Private/SDetailsViewBase.cpp
+++ b/Engine/Source/Editor/PropertyEditor/Private/SDetailsViewBase.cpp
@@ -980,14 +980,10 @@ void SDetailsViewBase::QueryCustomDetailLayout(FDetailLayoutData& LayoutData)
 		}
 	}
 
-	// Query extra base classes
+	// Query extra base classes and structs
 	for (auto ParentIt = ParentClassesToQuery.CreateConstIterator(); ParentIt; ++ParentIt)
 	{
-		UClass* ParentClass = Cast<UClass>(*ParentIt);
-		if (ParentClass)
-		{
-			QueryLayoutForClass(LayoutData, ParentClass);
-		}
+		QueryLayoutForClass(LayoutData, *ParentIt);
 	}
 }
 
diff --git a/Engine/Source/Editor/PropertyEditor/Private/SStructureDetailsView.cpp b/Engine/Source/Editor/PropertyEditor/Private/SStructureDetailsView.cpp
--- a/Engine/Source/Editor/PropertyEditor/Private/SStructureDetailsView.cpp
+++ b/Engine/Source/Editor/PropertyEditor/Private/SStructureDetailsView.cpp
@@ -52,7 +52,8 @@ void SStructureDetailsView::Construct(const FArguments& InArgs)
 
 	TSharedRef<SScrollBar> ExternalScrollbar = 
 		SNew(SScrollBar)
-		.AlwaysShowScrollbar( true );
+		.AlwaysShowScrollbar( DetailsViewArgs.bShowScrollBar )
+		.Visibility(DetailsViewArgs.bShowScrollBar ? EVisibility::Visible : EVisibility::Collapsed);
 
 		FMenuBuilder DetailViewOptions( true, NULL );
 
diff --git a/Engine/Source/Editor/PropertyEditor/Private/StructurePropertyNode.h b/Engine/Source/Editor/PropertyEditor/Private/StructurePropertyNode.h
--- a/Engine/Source/Editor/PropertyEditor/Private/StructurePropertyNode.h
+++ b/Engine/Source/Editor/PropertyEditor/Private/StructurePropertyNode.h
@@ -34,6 +34,11 @@ public:
 		return StructData.IsValid() && StructData->IsValid();
 	}
 
+	TSharedPtr<FStructOnScope> GetStructData() const
+	{
+		return StructData;
+	}
+
 	bool GetReadAddressUncached(FPropertyNode& InPropertyNode, FReadAddressListData& OutAddresses) const override
 	{
 		if (!HasValidStructData())
diff --git a/Engine/Source/Editor/PropertyEditor/Private/UserInterface/PropertyEditor/SPropertyMenuAssetPicker.cpp b/Engine/Source/Editor/PropertyEditor/Private/UserInterface/PropertyEditor/SPropertyMenuAssetPicker.cpp
--- a/Engine/Source/Editor/PropertyEditor/Private/UserInterface/PropertyEditor/SPropertyMenuAssetPicker.cpp
+++ b/Engine/Source/Editor/PropertyEditor/Private/UserInterface/PropertyEditor/SPropertyMenuAssetPicker.cpp
@@ -22,7 +22,10 @@ void SPropertyMenuAssetPicker::Construct( const FArguments& InArgs )
 	OnSet = InArgs._OnSet;
 	OnClose = InArgs._OnClose;
 
-	FMenuBuilder MenuBuilder(true, NULL);
+	const bool bInShouldCloseWindowAfterMenuSelection = true;
+	const bool bCloseSelfOnly = true;
+	
+	FMenuBuilder MenuBuilder(bInShouldCloseWindowAfterMenuSelection, nullptr, nullptr, bCloseSelfOnly);
 
 	if (NewAssetFactories.Num() > 0)
 	{
