Sequencer: Fix non unity build.

diff --git a/Engine/Source/Editor/MovieSceneTools/Public/GenericKeyArea.h b/Engine/Source/Editor/MovieSceneTools/Public/GenericKeyArea.h
--- a/Engine/Source/Editor/MovieSceneTools/Public/GenericKeyArea.h
+++ b/Engine/Source/Editor/MovieSceneTools/Public/GenericKeyArea.h
@@ -8,6 +8,7 @@
 #include "DetailCategoryBuilder.h"
 #include "DetailWidgetRow.h"
 #include "Curves/CurveInterface.h"
+#include "MovieSceneSection.h"
 
 #include "GenericKeyArea.generated.h"
 
