Copying //UE4/Dev-Sequencer to //UE4/Dev-Main (Source: //UE4/Dev-Sequencer @ 3237992)

diff --git a/Engine/Source/Editor/LevelEditor/Private/SActorDetails.cpp b/Engine/Source/Editor/LevelEditor/Private/SActorDetails.cpp
--- a/Engine/Source/Editor/LevelEditor/Private/SActorDetails.cpp
+++ b/Engine/Source/Editor/LevelEditor/Private/SActorDetails.cpp
@@ -80,7 +80,7 @@ public:
 	}
 };
 
-void SActorDetails::Construct(const FArguments& InArgs, const FName TabIdentifier, TSharedPtr<FUICommandList> InCommandList)
+void SActorDetails::Construct(const FArguments& InArgs, const FName TabIdentifier, TSharedPtr<FUICommandList> InCommandList, TSharedPtr<FTabManager> InTabManager)
 {
 	bSelectionGuard = false;
 	bShowingRootActorNodeSelected = false;
@@ -102,6 +102,7 @@ void SActorDetails::Construct(const FArguments& InArgs, const FName TabIdentifie
 	DetailsViewArgs.bCustomFilterAreaLocation = true;
 	DetailsViewArgs.DefaultsOnlyVisibility = FDetailsViewArgs::EEditDefaultsOnlyNodeVisibility::Hide;
 	DetailsViewArgs.HostCommandList = InCommandList;
+	DetailsViewArgs.HostTabManager = InTabManager;
 	DetailsView = PropPlugin.CreateDetailView(DetailsViewArgs);
 
 	auto IsPropertyVisible = [](const FPropertyAndParent& PropertyAndParent)
diff --git a/Engine/Source/Editor/LevelEditor/Private/SActorDetails.h b/Engine/Source/Editor/LevelEditor/Private/SActorDetails.h
--- a/Engine/Source/Editor/LevelEditor/Private/SActorDetails.h
+++ b/Engine/Source/Editor/LevelEditor/Private/SActorDetails.h
@@ -9,8 +9,10 @@
 #include "Framework/Text/SlateHyperlinkRun.h"
 #include "EditorUndoClient.h"
 
+
 class AActor;
 class FSCSEditorTreeNode;
+class FTabManager;
 class FUICommandList;
 class IDetailsView;
 class SBox;
@@ -18,6 +20,7 @@ class SSCSEditor;
 class SSplitter;
 class UBlueprint;
 
+
 /**
  * Wraps a details panel customized for viewing actors
  */
@@ -27,7 +30,7 @@ public:
 	SLATE_BEGIN_ARGS(SActorDetails) {}
 	SLATE_END_ARGS()
 
-	void Construct(const FArguments& InArgs, const FName TabIdentifier, TSharedPtr<FUICommandList> InCommandList);
+	void Construct(const FArguments& InArgs, const FName TabIdentifier, TSharedPtr<FUICommandList> InCommandList, TSharedPtr<FTabManager> InTabManager);
 	~SActorDetails();
 
 	/**
diff --git a/Engine/Source/Editor/LevelEditor/Private/SLevelEditor.cpp b/Engine/Source/Editor/LevelEditor/Private/SLevelEditor.cpp
--- a/Engine/Source/Editor/LevelEditor/Private/SLevelEditor.cpp
+++ b/Engine/Source/Editor/LevelEditor/Private/SLevelEditor.cpp
@@ -46,6 +46,7 @@
 #include "Widgets/Docking/SDockTab.h"
 #include "SActorDetails.h"
 #include "GameFramework/WorldSettings.h"
+#include "LayoutExtender.h"
 #include "HierarchicalLODOutlinerModule.h"
 
 
@@ -1108,6 +1109,8 @@ TSharedRef<SWidget> SLevelEditor::RestoreContentArea( const TSharedRef<SDockTab>
 
 		FTabSpawnerEntry& BuildAndSubmitEntry = LevelEditorTabManager->RegisterTabSpawner(LevelEditorBuildAndSubmitTab, FOnSpawnTab::CreateSP<SLevelEditor, FName, FString>(this, &SLevelEditor::SpawnLevelEditorTab, LevelEditorBuildAndSubmitTab, FString()));
 		BuildAndSubmitEntry.SetAutoGenerateMenuEntry(false);
+
+		LevelEditorModule.OnRegisterTabs().Broadcast(LevelEditorTabManager);
 	}
 
 	// Rebuild the editor mode commands and their tab spawners before we restore the layout,
@@ -1188,6 +1191,9 @@ TSharedRef<SWidget> SLevelEditor::RestoreContentArea( const TSharedRef<SDockTab>
 			
 		));
 	
+	FLayoutExtender LayoutExtender;
+	LevelEditorModule.OnRegisterLayoutExtensions().Broadcast(LayoutExtender);
+	Layout->ProcessExtensions(LayoutExtender);
 
 	return LevelEditorTabManager->RestoreFrom( Layout, OwnerWindow ).ToSharedRef();
 }
@@ -1489,7 +1495,7 @@ void SLevelEditor::AddStandaloneLevelViewport( const TSharedRef<SLevelViewport>&
 
 TSharedRef<SWidget> SLevelEditor::CreateActorDetails( const FName TabIdentifier )
 {
-	TSharedRef<SActorDetails> ActorDetails = SNew( SActorDetails, TabIdentifier, LevelEditorCommands );
+	TSharedRef<SActorDetails> ActorDetails = SNew( SActorDetails, TabIdentifier, LevelEditorCommands, GetTabManager() );
 
 	// Immediately update it (otherwise it will appear empty)
 	{
diff --git a/Engine/Source/Editor/LevelEditor/Public/LevelEditor.h b/Engine/Source/Editor/LevelEditor/Public/LevelEditor.h
--- a/Engine/Source/Editor/LevelEditor/Public/LevelEditor.h
+++ b/Engine/Source/Editor/LevelEditor/Public/LevelEditor.h
@@ -237,6 +237,12 @@ public:
 	virtual TSharedPtr<FExtensibilityManager> GetModeBarExtensibilityManager() {return ModeBarExtensibilityManager;}
 	virtual TSharedPtr<FExtensibilityManager> GetNotificationBarExtensibilityManager() {return NotificationBarExtensibilityManager;}
 
+	DECLARE_EVENT_OneParam(ILevelEditor, FOnRegisterTabs, TSharedPtr<FTabManager>);
+	FOnRegisterTabs& OnRegisterTabs() { return RegisterTabs; }
+
+	DECLARE_EVENT_OneParam(ILevelEditor, FOnRegisterLayoutExtensions, FLayoutExtender&);
+	FOnRegisterLayoutExtensions& OnRegisterLayoutExtensions() { return RegisterLayoutExtensions; }
+
 	/** Called when a new map is loaded */
 	DECLARE_EVENT( FLevelEditorModule, FNotificationBarChanged );
 	virtual FNotificationBarChanged& OnNotificationBarChanged() { return NotificationBarChangedEvent; }
@@ -382,6 +388,12 @@ private:
 	/** Map of named viewport types to factory functions */
 	TMap<FName, FViewportTypeDefinition> CustomViewports;
 
+	/** Register layout extensions event */
+	FOnRegisterLayoutExtensions RegisterLayoutExtensions;
+
+	/** Register external tabs event */
+	FOnRegisterTabs RegisterTabs;
+
 	/** Array of delegates that are used to check if the specified objects should be editable on the details panel */
 	TArray<FAreObjectsEditable> AreObjectsEditableDelegates;
 };
diff --git a/Engine/Source/Editor/MovieSceneTools/MovieSceneTools.Build.cs b/Engine/Source/Editor/MovieSceneTools/MovieSceneTools.Build.cs
--- a/Engine/Source/Editor/MovieSceneTools/MovieSceneTools.Build.cs
+++ b/Engine/Source/Editor/MovieSceneTools/MovieSceneTools.Build.cs
@@ -38,11 +38,13 @@ public class MovieSceneTools : ModuleRules
                 "MovieSceneCapture",
 				"MovieSceneTracks",
 				"BlueprintGraph",
+				"KismetCompiler",
                 "ContentBrowser",
 				"Slate",
 				"SlateCore",
                 "EditorStyle",
 				"EditorWidgets",
+				"PropertyEditor",
 				"RenderCore",
 				"RHI",
 				"ShaderCore",
@@ -57,7 +59,6 @@ public class MovieSceneTools : ModuleRules
 				"Sequencer",
                 "Settings",
 				"SceneOutliner",
-				"PropertyEditor",
                 "MainFrame",
                 "UnrealEd",
             }
@@ -68,7 +69,6 @@ public class MovieSceneTools : ModuleRules
                 "AssetRegistry",
 				"AssetTools",
 				"SceneOutliner",
-				"PropertyEditor",
                 "MainFrame",
 			}
 		);
diff --git a/Engine/Source/Editor/MovieSceneTools/Private/CurveKeyAreas/FloatCurveKeyArea.cpp b/Engine/Source/Editor/MovieSceneTools/Private/CurveKeyAreas/FloatCurveKeyArea.cpp
--- a/Engine/Source/Editor/MovieSceneTools/Private/CurveKeyAreas/FloatCurveKeyArea.cpp
+++ b/Engine/Source/Editor/MovieSceneTools/Private/CurveKeyAreas/FloatCurveKeyArea.cpp
@@ -188,6 +188,13 @@ TArray<FKeyHandle> FFloatCurveKeyArea::GetUnsortedKeyHandles() const
 }
 
 
+FKeyHandle FFloatCurveKeyArea::DilateKey(FKeyHandle KeyHandle, float Scale, float Origin)
+{
+	float NewKeyTime = Curve->GetKeyTime(KeyHandle);
+	NewKeyTime = (NewKeyTime - Origin) * Scale + Origin;
+	return Curve->SetKeyTime(KeyHandle, NewKeyTime);
+}
+
 FKeyHandle FFloatCurveKeyArea::MoveKey(FKeyHandle KeyHandle, float DeltaPosition)
 {
 	return Curve->SetKeyTime(KeyHandle, Curve->GetKeyTime(KeyHandle) + DeltaPosition);
diff --git a/Engine/Source/Editor/MovieSceneTools/Private/CurveKeyAreas/IntegralKeyArea.cpp b/Engine/Source/Editor/MovieSceneTools/Private/CurveKeyAreas/IntegralKeyArea.cpp
--- a/Engine/Source/Editor/MovieSceneTools/Private/CurveKeyAreas/IntegralKeyArea.cpp
+++ b/Engine/Source/Editor/MovieSceneTools/Private/CurveKeyAreas/IntegralKeyArea.cpp
@@ -111,6 +111,14 @@ TArray<FKeyHandle> FIntegralCurveKeyAreaBase::GetUnsortedKeyHandles() const
 }
 
 
+FKeyHandle FIntegralCurveKeyAreaBase::DilateKey(FKeyHandle KeyHandle, float Scale, float Origin)
+{
+	float NewKeyTime = Curve.GetKeyTime(KeyHandle);
+	NewKeyTime = (NewKeyTime - Origin) * Scale + Origin;
+	return Curve.SetKeyTime(KeyHandle, NewKeyTime);
+}
+
+
 FKeyHandle FIntegralCurveKeyAreaBase::MoveKey(FKeyHandle KeyHandle, float DeltaPosition)
 {
 	return Curve.SetKeyTime(KeyHandle, Curve.GetKeyTime(KeyHandle) + DeltaPosition);
diff --git a/Engine/Source/Editor/MovieSceneTools/Private/CurveKeyAreas/NameCurveKeyArea.cpp b/Engine/Source/Editor/MovieSceneTools/Private/CurveKeyAreas/NameCurveKeyArea.cpp
--- a/Engine/Source/Editor/MovieSceneTools/Private/CurveKeyAreas/NameCurveKeyArea.cpp
+++ b/Engine/Source/Editor/MovieSceneTools/Private/CurveKeyAreas/NameCurveKeyArea.cpp
@@ -125,6 +125,14 @@ TArray<FKeyHandle> FNameCurveKeyArea::GetUnsortedKeyHandles() const
 }
 
 
+FKeyHandle FNameCurveKeyArea::DilateKey( FKeyHandle KeyHandle, float Scale, float Origin )
+{
+	float NewKeyTime = Curve.GetKeyTime(KeyHandle);
+	NewKeyTime = (NewKeyTime - Origin) * Scale + Origin;
+	return Curve.SetKeyTime(KeyHandle, NewKeyTime);
+}
+	
+
 FKeyHandle FNameCurveKeyArea::MoveKey( FKeyHandle KeyHandle, float DeltaPosition )
 {
 	return Curve.SetKeyTime(KeyHandle, Curve.GetKeyTime(KeyHandle) + DeltaPosition);
diff --git a/Engine/Source/Editor/MovieSceneTools/Private/CurveKeyAreas/StringCurveKeyArea.cpp b/Engine/Source/Editor/MovieSceneTools/Private/CurveKeyAreas/StringCurveKeyArea.cpp
--- a/Engine/Source/Editor/MovieSceneTools/Private/CurveKeyAreas/StringCurveKeyArea.cpp
+++ b/Engine/Source/Editor/MovieSceneTools/Private/CurveKeyAreas/StringCurveKeyArea.cpp
@@ -154,6 +154,14 @@ TArray<FKeyHandle> FStringCurveKeyArea::GetUnsortedKeyHandles() const
 }
 
 
+FKeyHandle FStringCurveKeyArea::DilateKey(FKeyHandle KeyHandle, float Scale, float Origin)
+{
+	float NewKeyTime = Curve->GetKeyTime(KeyHandle);
+	NewKeyTime = (NewKeyTime - Origin) * Scale + Origin;
+	return Curve->SetKeyTime(KeyHandle, NewKeyTime);
+}
+
+
 FKeyHandle FStringCurveKeyArea::MoveKey(FKeyHandle KeyHandle, float DeltaPosition)
 {
 	return Curve->SetKeyTime(KeyHandle, Curve->GetKeyTime(KeyHandle) + DeltaPosition);
diff --git a/Engine/Source/Editor/MovieSceneTools/Private/CurveKeyEditors/SBoolCurveKeyEditor.cpp b/Engine/Source/Editor/MovieSceneTools/Private/CurveKeyEditors/SBoolCurveKeyEditor.cpp
--- a/Engine/Source/Editor/MovieSceneTools/Private/CurveKeyEditors/SBoolCurveKeyEditor.cpp
+++ b/Engine/Source/Editor/MovieSceneTools/Private/CurveKeyEditors/SBoolCurveKeyEditor.cpp
@@ -66,13 +66,16 @@ void SBoolCurveKeyEditor::OnCheckStateChanged(ECheckBoxState NewCheckboxState)
 				Curve->AddKey(CurrentTime, NewValue, CurrentKeyHandle);
 			}
 
-			if (OwningSection->GetStartTime() > CurrentTime)
+			if (Curve->GetNumKeys() != 0)
 			{
-				OwningSection->SetStartTime(CurrentTime);
-			}
-			if (OwningSection->GetEndTime() < CurrentTime)
-			{
-				OwningSection->SetEndTime(CurrentTime);
+				if (OwningSection->GetStartTime() > CurrentTime)
+				{
+					OwningSection->SetStartTime(CurrentTime);
+				}
+				if (OwningSection->GetEndTime() < CurrentTime)
+				{
+					OwningSection->SetEndTime(CurrentTime);
+				}
 			}
 		}
 
diff --git a/Engine/Source/Editor/MovieSceneTools/Private/K2Node_GetSequenceBindings.cpp b/Engine/Source/Editor/MovieSceneTools/Private/K2Node_GetSequenceBindings.cpp
new file mode 100644
--- /dev/null
+++ b/Engine/Source/Editor/MovieSceneTools/Private/K2Node_GetSequenceBindings.cpp
@@ -0,0 +1,248 @@
+// Copyright 1998-2016 Epic Games, Inc. All Rights Reserved.
+
+#include "K2Node_GetSequenceBindings.h"
+#include "BlueprintEditorUtils.h"
+#include "KismetCompiler.h"
+#include "BlueprintNodeSpawner.h"
+#include "EditorCategoryUtils.h"
+#include "BlueprintActionDatabaseRegistrar.h"
+#include "Framework/Application/SlateApplication.h"
+#include "PropertyCustomizationHelpers.h"
+#include "MovieScene.h"
+#include "MovieSceneSequence.h"
+#include "MultiBox/MultiBoxBuilder.h"
+
+
+static const FString SequencePinName(TEXT("Sequence"));
+
+#define LOCTEXT_NAMESPACE "UK2Node_GetSequenceBindings"
+
+
+class FKCHandler_GetSequenceBindings : public FNodeHandlingFunctor
+{
+public:
+	FKCHandler_GetSequenceBindings(FKismetCompilerContext& InCompilerContext)
+		: FNodeHandlingFunctor(InCompilerContext)
+	{
+	}
+
+	virtual void RegisterNets(FKismetFunctionContext& Context, UEdGraphNode* Node) override
+	{
+		UK2Node_GetSequenceBindings* BindingsNode = CastChecked<UK2Node_GetSequenceBindings>(Node);
+		
+		for (UEdGraphPin* Pin : BindingsNode->GetAllPins())
+		{
+			if (Pin->Direction == EGPD_Output && Pin->LinkedTo.Num())
+			{
+				FBPTerminal* Term = RegisterLiteral(Context, Pin);
+				TOptional<FGuid> Guid = BindingsNode->GetGuidFromPin(Pin);
+				if (!Guid.IsSet())
+				{
+					Context.MessageLog.Warning(TEXT("Invalid Object Binding ID (%%) for node %%."), Pin, Node);
+				}
+				else
+				{
+					FMovieSceneObjectBindingPtr Value;
+					Value.Guid = Guid.GetValue();
+					FMovieSceneObjectBindingPtr::StaticStruct()->ExportText(Term->Name, &Value, nullptr, nullptr, 0, nullptr);
+				}
+			}
+		}
+	}
+};
+
+void UK2Node_GetSequenceBindings::PostEditChangeProperty(FPropertyChangedEvent& PropertyChangedEvent)
+{
+	Super::PostEditChangeProperty(PropertyChangedEvent);
+
+	UpdatePins();
+}
+
+TOptional<FGuid> UK2Node_GetSequenceBindings::GetGuidFromPin(UEdGraphPin* Pin) const
+{
+	const FGetSequenceBindingGuidMapping* Mapping = BindingGuids.FindByPredicate(
+		[=](const FGetSequenceBindingGuidMapping& InMapping)
+		{
+			return InMapping.PinName == Pin->PinName;
+		}
+	);
+
+	return Mapping ? Mapping->Guid : TOptional<FGuid>();
+}
+
+void UK2Node_GetSequenceBindings::NotifyPinConnectionListChanged(UEdGraphPin* Pin)
+{
+	// Ensure pins always have a default value
+	TOptional<FGuid> FoundGuid = GetGuidFromPin(Pin);
+	if (FoundGuid.IsSet())
+	{
+		FMovieSceneObjectBindingPtr Value;
+		Value.Guid = FoundGuid.GetValue();
+		FMovieSceneObjectBindingPtr::StaticStruct()->ExportText(Pin->DefaultValue, &Value, nullptr, nullptr, 0, nullptr);
+	}
+
+	Super::NotifyPinConnectionListChanged(Pin);
+}
+
+void UK2Node_GetSequenceBindings::PinDefaultValueChanged(UEdGraphPin* Pin)
+{
+	if (Pin && Pin->Direction == EGPD_Input)
+	{
+		UpdatePins();
+	}
+
+	Super::PinDefaultValueChanged(Pin);
+}
+
+FNodeHandlingFunctor* UK2Node_GetSequenceBindings::CreateNodeHandler(FKismetCompilerContext& CompilerContext) const
+{
+	return new FKCHandler_GetSequenceBindings(CompilerContext);
+}
+
+FText UK2Node_GetSequenceBindings::GetNodeTitle(ENodeTitleType::Type TitleType) const
+{
+	return FText::Format(LOCTEXT("NodeTitle", "Get Sequence Bindings ({0})"), Sequence ? FText::FromName(Sequence->GetFName()) : LOCTEXT("NoSequence", "No Sequence"));
+}
+
+void UK2Node_GetSequenceBindings::ReallocatePinsDuringReconstruction(TArray<UEdGraphPin*>& OldPins)
+{
+	Super::ReallocatePinsDuringReconstruction(OldPins);
+
+	UpdatePins();
+}
+
+FText UK2Node_GetSequenceBindings::GetTooltipText() const
+{
+	return LOCTEXT("NodeTooltip", "Access all the binding IDs for the specified sequence");
+}
+
+FSlateIcon UK2Node_GetSequenceBindings::GetIconAndTint(FLinearColor& OutColor) const
+{
+	static FSlateIcon Icon("EditorStyle", "GraphEditor.GetSequenceBindings");
+	return Icon;
+}
+
+void UK2Node_GetSequenceBindings::UpdatePins()
+{
+	PreloadObject(Sequence);
+
+	UMovieScene* MovieScene = Sequence ? Sequence->GetMovieScene() : nullptr;
+	PreloadObject(MovieScene);
+
+	for (int32 Index = Pins.Num() - 1; Index >= 0; --Index)
+	{
+		UEdGraphPin* Pin = Pins[Index];
+		RemovePin(Pin);
+	}
+
+	BindingGuids.Reset();
+
+	// Generate all new pins
+	if (MovieScene)
+	{
+		int32 PossessableCount = MovieScene->GetPossessableCount();
+		for (int32 Index = 0; Index < PossessableCount; ++Index)
+		{
+			const FMovieScenePossessable& Possessable = MovieScene->GetPossessable(Index);
+			if (!Sequence->CanRebindPossessable(Possessable))
+			{
+				continue;
+			}
+
+			FString GuidString = Possessable.GetGuid().ToString();
+			UEdGraphPin* NewPin = CreatePin(EGPD_Output, UEdGraphSchema_K2::PC_Struct, TEXT(""), FMovieSceneObjectBindingPtr::StaticStruct(), false, false, GuidString);
+			NewPin->PinFriendlyName = MovieScene->GetObjectDisplayName(Possessable.GetGuid());
+			NewPin->PersistentGuid = Possessable.GetGuid();
+
+			FMovieSceneObjectBindingPtr Value;
+			Value.Guid = Possessable.GetGuid();
+			FMovieSceneObjectBindingPtr::StaticStruct()->ExportText(NewPin->DefaultValue, &Value, nullptr, nullptr, 0, nullptr);
+
+			BindingGuids.Add(FGetSequenceBindingGuidMapping{ NewPin->PinName, Possessable.GetGuid() });
+		}
+	}
+
+	const bool bIsCompiling = GetBlueprint()->bBeingCompiled;
+	if( !bIsCompiling )
+	{
+		FBlueprintEditorUtils::MarkBlueprintAsStructurallyModified(GetBlueprint());
+	}
+}
+
+void UK2Node_GetSequenceBindings::GetContextMenuActions(const FGraphNodeContextMenuBuilder& Context) const
+{
+	Super::GetContextMenuActions(Context);
+
+	if (!Context.bIsDebugging)
+	{
+		Context.MenuBuilder->BeginSection("K2NodeGetSequenceBindings", LOCTEXT("ThisNodeHeader", "This Node"));
+		if (!Context.Pin)
+		{
+			auto SubMenu = [=](FMenuBuilder& SubMenuBuilder)
+				{
+					TArray<const UClass*> AllowedClasses({ UMovieSceneSequence::StaticClass() });
+
+					TSharedRef<SWidget> MenuContent = PropertyCustomizationHelpers::MakeAssetPickerWithMenu(
+						FAssetData(Sequence),
+						true /* bAllowClear */,
+						false /* bSearchInBlueprint */, 
+						AllowedClasses,
+						PropertyCustomizationHelpers::GetNewAssetFactoriesForClasses(AllowedClasses),
+						FOnShouldFilterAsset(),
+						FOnAssetSelected::CreateUObject(this, &UK2Node_GetSequenceBindings::SetSequence),
+						FSimpleDelegate());
+					
+					SubMenuBuilder.AddWidget(MenuContent, FText::GetEmpty(), false);
+				};
+
+			Context.MenuBuilder->AddSubMenu(
+				LOCTEXT("SetSequence_Text", "Sequence"),
+				LOCTEXT("SetSequence_ToolTip", "Sets the sequence to extract bindings from"),
+				FNewMenuDelegate::CreateLambda(SubMenu)
+				);
+
+			Context.MenuBuilder->AddMenuEntry(
+				LOCTEXT("Refresh_Text", "Refresh"),
+				LOCTEXT("Refresh_ToolTip", "Refresh this node's bindings"),
+				FSlateIcon(),
+				FUIAction(
+						FExecuteAction::CreateUObject(this, &UK2Node_GetSequenceBindings::ReconstructNode)
+					)
+				);
+		}
+
+		Context.MenuBuilder->EndSection();
+	}
+}
+
+void UK2Node_GetSequenceBindings::SetSequence(const FAssetData& InAssetData)
+{
+	FSlateApplication::Get().DismissAllMenus();
+	Sequence = Cast<UMovieSceneSequence>(InAssetData.GetAsset());
+	UpdatePins();
+}
+
+void UK2Node_GetSequenceBindings::GetMenuActions(FBlueprintActionDatabaseRegistrar& ActionRegistrar) const
+{
+	UClass* ActionKey = GetClass();
+	if (ActionRegistrar.IsOpenForRegistration(ActionKey))
+	{
+		UBlueprintNodeSpawner* NodeSpawner = UBlueprintNodeSpawner::Create(GetClass());
+		check(NodeSpawner != nullptr);
+
+		ActionRegistrar.AddBlueprintAction(ActionKey, NodeSpawner);
+	}
+}
+
+FText UK2Node_GetSequenceBindings::GetMenuCategory() const
+{
+	static FNodeTextCache CachedCategory;
+	if (CachedCategory.IsOutOfDate(this))
+	{
+		// FText::Format() is slow, so we cache this to save on performance
+		CachedCategory.SetCachedText(FEditorCategoryUtils::BuildCategoryString(FCommonEditorCategory::Utilities, LOCTEXT("ActionMenuCategory", "Sequence")), this);
+	}
+	return CachedCategory;
+}
+
+#undef LOCTEXT_NAMESPACE
diff --git a/Engine/Source/Editor/MovieSceneTools/Private/K2Node_GetSequenceBindings.h b/Engine/Source/Editor/MovieSceneTools/Private/K2Node_GetSequenceBindings.h
new file mode 100644
--- /dev/null
+++ b/Engine/Source/Editor/MovieSceneTools/Private/K2Node_GetSequenceBindings.h
@@ -0,0 +1,67 @@
+// Copyright 1998-2016 Epic Games, Inc. All Rights Reserved.
+
+#pragma once
+
+#include "CoreMinimal.h"
+#include "K2Node.h"
+#include "UObject/ObjectMacros.h"
+
+#include "K2Node_GetSequenceBindings.generated.h"
+
+
+class FAssetData;
+class UEdGraphPin;
+class UMovieSceneSequence;
+
+
+USTRUCT()
+struct FGetSequenceBindingGuidMapping
+{
+	GENERATED_BODY()
+
+	UPROPERTY()
+	FString PinName;
+
+	UPROPERTY()
+	FGuid Guid;
+};
+
+
+UCLASS(MinimalAPI)
+class UK2Node_GetSequenceBindings
+	: public UK2Node
+{
+public:
+	GENERATED_BODY()
+
+	UPROPERTY()
+	TArray<FGetSequenceBindingGuidMapping> BindingGuids;
+
+	UPROPERTY(EditAnywhere, Category="Sequence")
+	UMovieSceneSequence* Sequence;
+
+	TOptional<FGuid> GetGuidFromPin(UEdGraphPin* Pin) const;
+
+public:
+	virtual FText GetNodeTitle(ENodeTitleType::Type TitleType) const override;
+	virtual FText GetTooltipText() const override;
+	virtual FSlateIcon GetIconAndTint(FLinearColor& OutColor) const override;
+	virtual bool ShouldShowNodeProperties() const override { return true; }
+	virtual bool IsNodePure() const override { return true; }
+	virtual void NotifyPinConnectionListChanged(UEdGraphPin* Pin) override;
+	virtual void PinDefaultValueChanged(UEdGraphPin* Pin) override;
+	virtual void GetContextMenuActions(const FGraphNodeContextMenuBuilder& Context) const override;
+	virtual class FNodeHandlingFunctor* CreateNodeHandler(class FKismetCompilerContext& CompilerContext) const override;
+	virtual void GetMenuActions(FBlueprintActionDatabaseRegistrar& ActionRegistrar) const override;
+	virtual FText GetMenuCategory() const override;
+	virtual void ReallocatePinsDuringReconstruction(TArray<UEdGraphPin*>& OldPins);
+
+#if WITH_EDITOR
+	virtual void PostEditChangeProperty(FPropertyChangedEvent& PropertyChangedEvent) override;
+#endif
+
+private:
+
+	void SetSequence(const FAssetData& InAssetData);
+	void UpdatePins();
+};
diff --git a/Engine/Source/Editor/MovieSceneTools/Private/MatineeImportTools.cpp b/Engine/Source/Editor/MovieSceneTools/Private/MatineeImportTools.cpp
--- a/Engine/Source/Editor/MovieSceneTools/Private/MatineeImportTools.cpp
+++ b/Engine/Source/Editor/MovieSceneTools/Private/MatineeImportTools.cpp
@@ -5,6 +5,9 @@
 #include "Tracks/MovieSceneAudioTrack.h"
 #include "ScopedTransaction.h"
 #include "MovieSceneCommonHelpers.h"
+#include "IMovieScenePlayer.h"
+
+#include "Curves/CurveInterface.h"
 
 #include "Matinee/MatineeActor.h"
 #include "Matinee/InterpData.h"
@@ -18,8 +21,11 @@
 #include "Matinee/InterpTrackFade.h"
 #include "Matinee/InterpTrackDirector.h"
 #include "Matinee/InterpTrackEvent.h"
+#include "Matinee/InterpTrackVectorProp.h"
 #include "Matinee/InterpTrackVisibility.h"
 
+#include "Tracks/MovieSceneBoolTrack.h"
+#include "Tracks/MovieSceneFloatTrack.h"
 #include "Tracks/MovieSceneColorTrack.h"
 #include "Tracks/MovieScene3DTransformTrack.h"
 #include "Tracks/MovieSceneParticleTrack.h"
@@ -28,6 +34,8 @@
 #include "Tracks/MovieSceneCameraCutTrack.h"
 #include "Tracks/MovieSceneEventTrack.h"
 #include "Tracks/MovieSceneVisibilityTrack.h"
+#include "Tracks/MovieSceneAudioTrack.h"
+#include "Tracks/MovieSceneVectorTrack.h"
 
 #include "Sections/MovieSceneColorSection.h"
 #include "Sections/MovieSceneBoolSection.h"
@@ -37,8 +45,8 @@
 #include "Sections/MovieSceneAudioSection.h"
 #include "Sections/MovieSceneFadeSection.h"
 #include "Sections/MovieSceneCameraCutSection.h"
-#include "Curves/CurveInterface.h"
 #include "Sections/MovieSceneEventSection.h"
+#include "Sections/MovieSceneVectorSection.h"
 
 
 #include "Animation/AnimSequence.h"
@@ -195,6 +203,54 @@ bool FMatineeImportTools::CopyInterpFloatTrack( UInterpTrackFloatBase* MatineeFl
 	return bSectionCreated;
 }
 
+bool FMatineeImportTools::CopyInterpVectorTrack( UInterpTrackVectorProp* MatineeVectorTrack, UMovieSceneVectorTrack* VectorTrack )
+{
+	const FScopedTransaction Transaction( NSLOCTEXT( "Sequencer", "PasteMatineeVectorTrack", "Paste Matinee Vector Track" ) );
+	bool bSectionCreated = false;
+
+	VectorTrack->Modify();
+
+	float KeyTime = MatineeVectorTrack->GetKeyframeTime( 0 );
+	UMovieSceneVectorSection* Section = Cast<UMovieSceneVectorSection>( MovieSceneHelpers::FindSectionAtTime( VectorTrack->GetAllSections(), KeyTime ) );
+	if ( Section == nullptr )
+	{
+		Section = Cast<UMovieSceneVectorSection>( VectorTrack->CreateNewSection() );
+		VectorTrack->AddSection( *Section );
+		Section->SetIsInfinite(true);
+		bSectionCreated = true;
+	}
+	if (Section->TryModify())
+	{
+		float SectionMin = Section->GetStartTime();
+		float SectionMax = Section->GetEndTime();
+
+		if (Section->GetChannelsUsed() == 3)
+		{
+			FRichCurve& XCurve = Section->GetCurve(0);
+			FRichCurve& YCurve = Section->GetCurve(1);
+			FRichCurve& ZCurve = Section->GetCurve(2);
+
+			for ( const auto& Point : MatineeVectorTrack->VectorTrack.Points )
+			{
+				FMatineeImportTools::SetOrAddKey( XCurve, Point.InVal, Point.OutVal.X, Point.ArriveTangent.X, Point.LeaveTangent.X, Point.InterpMode );
+				FMatineeImportTools::SetOrAddKey( YCurve, Point.InVal, Point.OutVal.Y, Point.ArriveTangent.Y, Point.LeaveTangent.Y, Point.InterpMode );
+				FMatineeImportTools::SetOrAddKey( ZCurve, Point.InVal, Point.OutVal.Z, Point.ArriveTangent.Z, Point.LeaveTangent.Z, Point.InterpMode );
+				SectionMin = FMath::Min( SectionMin, Point.InVal );
+				SectionMax = FMath::Max( SectionMax, Point.InVal );
+			}
+			
+			CleanupCurveKeys(XCurve);
+			CleanupCurveKeys(YCurve);
+			CleanupCurveKeys(ZCurve);
+		}
+
+		Section->SetStartTime( SectionMin );
+		Section->SetEndTime( SectionMax );
+	}
+
+	return bSectionCreated;
+}
+
 
 bool FMatineeImportTools::CopyInterpColorTrack( UInterpTrackColorProp* ColorPropTrack, UMovieSceneColorTrack* ColorTrack )
 {
@@ -300,7 +356,7 @@ bool FMatineeImportTools::CopyInterpLinearColorTrack( UInterpTrackLinearColorPro
 	return bSectionCreated;
 }
 
-bool FMatineeImportTools::CopyInterpMoveTrack( UInterpTrackMove* MoveTrack, UMovieScene3DTransformTrack* TransformTrack )
+bool FMatineeImportTools::CopyInterpMoveTrack( UInterpTrackMove* MoveTrack, UMovieScene3DTransformTrack* TransformTrack, const FVector& DefaultScale )
 {
 	const FScopedTransaction Transaction( NSLOCTEXT( "Sequencer", "PasteMatineeMoveTrack", "Paste Matinee Move Track" ) );
 	bool bSectionCreated = false;
@@ -312,9 +368,9 @@ bool FMatineeImportTools::CopyInterpMoveTrack( UInterpTrackMove* MoveTrack, UMov
 	if ( Section == nullptr )
 	{
 		Section = Cast<UMovieScene3DTransformSection>( TransformTrack->CreateNewSection() );
-		Section->GetScaleCurve(EAxis::X).SetDefaultValue(1);
-		Section->GetScaleCurve(EAxis::Y).SetDefaultValue(1);
-		Section->GetScaleCurve(EAxis::Z).SetDefaultValue(1);
+		Section->GetScaleCurve(EAxis::X).SetDefaultValue(DefaultScale.X);
+		Section->GetScaleCurve(EAxis::Y).SetDefaultValue(DefaultScale.Y);
+		Section->GetScaleCurve(EAxis::Z).SetDefaultValue(DefaultScale.Z);
 		TransformTrack->AddSection( *Section );
 		Section->SetIsInfinite(true);
 		bSectionCreated = true;
@@ -528,9 +584,8 @@ bool FMatineeImportTools::CopyInterpSoundTrack( UInterpTrackSound* MatineeSoundT
 
 		UMovieSceneAudioSection* NewAudioSection = Cast<UMovieSceneAudioSection>(AudioTrack->GetAllSections().Last());
 		NewAudioSection->SetRowIndex( MaxSectionRowIndex + 1 );
-		// The audio dilation is inverted before it's set as the pitch, so we need to invert the pitch here to get the same result.
-		NewAudioSection->SetAudioDilationFactor( 1 / SoundTrackKey.Pitch );
-		NewAudioSection->SetAudioVolume( SoundTrackKey.Volume );
+		NewAudioSection->GetPitchMultiplierCurve().SetDefaultValue( SoundTrackKey.Pitch );
+		NewAudioSection->GetSoundVolumeCurve().SetDefaultValue( SoundTrackKey.Volume );
 
 		AudioTrack->AddSection( *NewAudioSection );
 		bSectionCreated = true;
diff --git a/Engine/Source/Editor/MovieSceneTools/Private/TrackEditors/ActorPickerTrackEditor.cpp b/Engine/Source/Editor/MovieSceneTools/Private/TrackEditors/ActorPickerTrackEditor.cpp
--- a/Engine/Source/Editor/MovieSceneTools/Private/TrackEditors/ActorPickerTrackEditor.cpp
+++ b/Engine/Source/Editor/MovieSceneTools/Private/TrackEditors/ActorPickerTrackEditor.cpp
@@ -10,6 +10,8 @@
 #include "Framework/Application/SlateApplication.h"
 #include "Widgets/Images/SImage.h"
 #include "Widgets/Layout/SBox.h"
+#include "Widgets/Views/STableRow.h"
+#include "Widgets/Views/SListView.h"
 #include "Widgets/Input/SButton.h"
 #include "EditorStyleSet.h"
 #include "Editor/UnrealEdEngine.h"
@@ -118,6 +120,101 @@ void FActorPickerTrackEditor::ShowActorSubMenu(FMenuBuilder& MenuBuilder, FGuid
 	MenuBuilder.AddWidget(MenuWidget, FText::GetEmpty(), false);
 }
 
+class SComponentChooserPopup : public SCompoundWidget
+{
+public:
+	DECLARE_DELEGATE_OneParam( FOnComponentChosen, FString );
+
+	SLATE_BEGIN_ARGS( SComponentChooserPopup )
+		: _Actor(NULL)
+		{}
+
+		/** An actor with components */
+		SLATE_ARGUMENT( AActor*, Actor )
+
+		/** Called when the text is chosen. */
+		SLATE_EVENT( FOnComponentChosen, OnComponentChosen )
+
+	SLATE_END_ARGS()
+
+	/** Delegate to call when component is selected */
+	FOnComponentChosen OnComponentChosen;
+
+	/** List of tag names selected in the tag containers*/
+	TArray< TSharedPtr<FString> > ComponentNames;
+
+private:
+	TSharedRef<ITableRow> MakeListViewWidget(TSharedPtr<FString> InItem, const TSharedRef<STableViewBase>& OwnerTable)
+	{
+		return SNew( STableRow< TSharedPtr<FString> >, OwnerTable )
+				[
+					SNew(STextBlock) .Text( FText::FromString(*InItem.Get()) )
+				];
+	}
+
+	void OnComponentSelected(TSharedPtr<FString> InItem, ESelectInfo::Type InSelectInfo)
+	{
+		FSlateApplication::Get().DismissAllMenus();
+
+		if(OnComponentChosen.IsBound())
+		{
+			OnComponentChosen.Execute(*InItem.Get());
+		}
+	}
+
+public:
+	void Construct( const FArguments& InArgs )
+	{
+		OnComponentChosen = InArgs._OnComponentChosen;
+		AActor* Actor = InArgs._Actor;
+
+		TInlineComponentArray<USceneComponent*> Components(Actor);
+
+		ComponentNames.Empty();
+		for(USceneComponent* Component : Components)
+		{
+			if (Component->HasAnySockets())
+			{
+				ComponentNames.Add(MakeShareable(new FString(Component->GetName())));
+			}
+		}
+
+		// Then make widget
+		this->ChildSlot
+		[
+			SNew(SBorder)
+			.BorderImage(FEditorStyle::GetBrush(TEXT("Menu.Background")))
+			.Padding(5)
+			.Content()
+			[
+				SNew(SVerticalBox)
+				+SVerticalBox::Slot()
+				.AutoHeight()
+				.Padding(0.0f, 1.0f)
+				[
+					SNew(STextBlock)
+					.Font( FEditorStyle::GetFontStyle(TEXT("SocketChooser.TitleFont")) )
+					.Text( NSLOCTEXT("ComponentChooser", "ChooseComponentLabel", "Choose Component") )
+				]
+				+SVerticalBox::Slot()
+				.AutoHeight()
+				.MaxHeight(512)
+				[
+					SNew(SBox)
+					.WidthOverride(256)
+					.Content()
+					[
+						SNew(SListView< TSharedPtr<FString> >)
+						.ListItemsSource( &ComponentNames)
+						.OnGenerateRow( this, &SComponentChooserPopup::MakeListViewWidget )
+						.OnSelectionChanged( this, &SComponentChooserPopup::OnComponentSelected )
+					]
+				]
+			]
+		];
+	}
+};
+
 void FActorPickerTrackEditor::ActorPicked(AActor* ParentActor, FGuid ObjectGuid, UMovieSceneSection* Section)
 {
 	TArray<USceneComponent*> ComponentsWithSockets;
@@ -134,49 +231,24 @@ void FActorPickerTrackEditor::ActorPicked(AActor* ParentActor, FGuid ObjectGuid,
 		}
 	}
 
-	// Show socket chooser if we have sockets to select
-	if (ComponentsWithSockets.Num() > 0)
+	if (ComponentsWithSockets.Num() == 0)
 	{
-		FLevelEditorModule& LevelEditorModule = FModuleManager::GetModuleChecked<FLevelEditorModule>( "LevelEditor");
-		TSharedPtr< ILevelEditor > LevelEditor = LevelEditorModule.GetFirstLevelEditor();
+		FSlateApplication::Get().DismissAllMenus();
+		ActorSocketPicked( NAME_None, nullptr, ParentActor, ObjectGuid, Section );
+		return;
+	}
 
-		TSharedPtr<SWidget> MenuWidget;
+	FLevelEditorModule& LevelEditorModule = FModuleManager::GetModuleChecked<FLevelEditorModule>( "LevelEditor");
+	TSharedPtr< ILevelEditor > LevelEditor = LevelEditorModule.GetFirstLevelEditor();
 
-		if(ComponentsWithSockets.Num() > 1)
-		{
-			const bool bInShouldCloseWindowAfterMenuSelection = true;
-			FMenuBuilder MenuBuilder(bInShouldCloseWindowAfterMenuSelection, TSharedPtr<const FUICommandList>());
+	TSharedPtr<SWidget> MenuWidget;
 
-			MenuBuilder.BeginSection(TEXT("ComponentsWithSockets"), LOCTEXT("ComponentsWithSockets", "Components With Sockets"));
-
-			for(USceneComponent* Component : ComponentsWithSockets)
-			{
-				auto NewMenuDelegate = [&](FMenuBuilder& ComponentMenuBuilder)
-				{
-					ComponentMenuBuilder.AddWidget(
-						SNew(SSocketChooserPopup)
-						.SceneComponent(Component)
-						.OnSocketChosen(this, &FActorPickerTrackEditor::ActorSocketPicked, Component, ParentActor, ObjectGuid, Section),
-						FText());
-				};
-
-				MenuBuilder.AddSubMenu(
-					FText::FromName(Component->GetFName()),
-					FText::Format(LOCTEXT("ComponentSubMenuToolTip", "View sockets for component {0}"), FText::FromName(Component->GetFName())),
-					FNewMenuDelegate::CreateLambda(NewMenuDelegate)
-					);
-			}
-		
-			MenuBuilder.EndSection();
-			MenuWidget = MenuBuilder.MakeWidget();
-		}
-		else
-		{
-			MenuWidget = 
-				SNew(SSocketChooserPopup)
-				.SceneComponent(ComponentsWithSockets[0])
-				.OnSocketChosen(this, &FActorPickerTrackEditor::ActorSocketPicked, ComponentsWithSockets[0], ParentActor, ObjectGuid, Section);		
-		}
+	if (ComponentsWithSockets.Num() > 1)
+	{			
+		MenuWidget = 
+			SNew(SComponentChooserPopup)
+			.Actor(ParentActor)
+			.OnComponentChosen(this, &FActorPickerTrackEditor::ActorComponentPicked, ParentActor, ObjectGuid, Section);		
 
 		// Create as context menu
 		FSlateApplication::Get().PushMenu(
@@ -189,9 +261,48 @@ void FActorPickerTrackEditor::ActorPicked(AActor* ParentActor, FGuid ObjectGuid,
 	}
 	else
 	{
-		FSlateApplication::Get().DismissAllMenus();
-		ActorSocketPicked( NAME_None, nullptr, ParentActor, ObjectGuid, Section );
+		ActorComponentPicked(ComponentsWithSockets[0]->GetName(), ParentActor, ObjectGuid, Section);
 	}
 }
 
+
+void FActorPickerTrackEditor::ActorComponentPicked(FString ComponentName, AActor* ParentActor, FGuid ObjectGuid, UMovieSceneSection* Section)
+{
+	USceneComponent* ComponentWithSockets = nullptr;
+	TInlineComponentArray<USceneComponent*> Components(ParentActor);
+
+	for(USceneComponent* Component : Components)
+	{
+		if (Component->GetName() == ComponentName)
+		{
+			ComponentWithSockets = Component;
+			break;
+		}
+	}
+
+	if (ComponentWithSockets == nullptr)
+	{
+		return;
+	}
+
+	FLevelEditorModule& LevelEditorModule = FModuleManager::GetModuleChecked<FLevelEditorModule>( "LevelEditor");
+	TSharedPtr< ILevelEditor > LevelEditor = LevelEditorModule.GetFirstLevelEditor();
+
+	TSharedPtr<SWidget> MenuWidget;
+
+	MenuWidget = 
+		SNew(SSocketChooserPopup)
+		.SceneComponent(ComponentWithSockets)
+		.OnSocketChosen(this, &FActorPickerTrackEditor::ActorSocketPicked, ComponentWithSockets, ParentActor, ObjectGuid, Section);		
+
+	// Create as context menu
+	FSlateApplication::Get().PushMenu(
+		LevelEditor.ToSharedRef(),
+		FWidgetPath(),
+		MenuWidget.ToSharedRef(),
+		FSlateApplication::Get().GetCursorPos(),
+		FPopupTransitionEffect( FPopupTransitionEffect::ContextMenu )
+		);
+}
+
 #undef LOCTEXT_NAMESPACE
diff --git a/Engine/Source/Editor/MovieSceneTools/Private/TrackEditors/ActorPickerTrackEditor.h b/Engine/Source/Editor/MovieSceneTools/Private/TrackEditors/ActorPickerTrackEditor.h
--- a/Engine/Source/Editor/MovieSceneTools/Private/TrackEditors/ActorPickerTrackEditor.h
+++ b/Engine/Source/Editor/MovieSceneTools/Private/TrackEditors/ActorPickerTrackEditor.h
@@ -23,7 +23,7 @@ public:
 	/** Is this actor pickable? */
 	virtual bool IsActorPickable( const AActor* const ParentActor, FGuid ObjectBinding, UMovieSceneSection* InSection) { return false; }
 
-	/** Actor socked was picked */
+	/** Actor socket was picked */
 	virtual void ActorSocketPicked(const FName SocketName, USceneComponent* Component, AActor* ParentActor, FGuid ObjectBinding, UMovieSceneSection* Section) {}
 
 	/** Show a sub menu of the pickable actors */
@@ -32,6 +32,9 @@ public:
 	/** Actor was picked */
 	void ActorPicked(AActor* ParentActor, FGuid ObjectBinding, UMovieSceneSection* Section);
 
+	/** Actor component was picked */
+	void ActorComponentPicked(FString ComponentName, AActor* ParentActor, FGuid ObjectBinding, UMovieSceneSection* Section);
+
 private: 
 
 	/** Interactively pick an actor from the viewport */
diff --git a/Engine/Source/Editor/MovieSceneTools/Private/TrackEditors/AudioTrackEditor.cpp b/Engine/Source/Editor/MovieSceneTools/Private/TrackEditors/AudioTrackEditor.cpp
--- a/Engine/Source/Editor/MovieSceneTools/Private/TrackEditors/AudioTrackEditor.cpp
+++ b/Engine/Source/Editor/MovieSceneTools/Private/TrackEditors/AudioTrackEditor.cpp
@@ -31,6 +31,8 @@
 #include "AssetRegistryModule.h"
 #include "MatineeImportTools.h"
 #include "Matinee/InterpTrackSound.h"
+#include "FloatCurveKeyArea.h"
+#include "ISectionLayoutBuilder.h"
 
 
 #define LOCTEXT_NAMESPACE "FAudioTrackEditor"
@@ -39,7 +41,7 @@
 namespace AnimatableAudioEditorConstants
 {
 	// @todo Sequencer Allow this to be customizable
-	const uint32 AudioTrackHeight = 50;
+	const uint32 WaveformHeight = 50;
 
 	// Optimization - maximum samples per pixel this sound allows
 	const uint32 MaxSamplesPerPixel = 60;
@@ -70,7 +72,7 @@ USoundWave* DeriveSoundWave(UMovieSceneAudioSection* AudioSection)
 		const TArray<USoundNode*>& AllNodes = SoundCue->AllNodes;
 		for (int32 Index = 0; Index < AllNodes.Num() && SoundWave == nullptr; ++Index)
 		{
-			if (AllNodes[Index]->IsA<USoundNodeWavePlayer>())
+			if (AllNodes[Index] && AllNodes[Index]->IsA<USoundNodeWavePlayer>())
 			{
 				SoundWave = Cast<USoundNodeWavePlayer>(AllNodes[Index])->GetSoundWave();
 			}
@@ -122,7 +124,7 @@ struct FSplineSegment
 
 
 /**
- * The audio thumnail, which holds a texture which it can pass back to a viewport to render
+ * The audio thumbnail, which holds a texture which it can pass back to a viewport to render
  */
 class FAudioThumbnail
 	: public ISlateViewport
@@ -231,7 +233,7 @@ FAudioThumbnail::~FAudioThumbnail()
 }
 
 
-FIntPoint FAudioThumbnail::GetSize() const {return FIntPoint(TextureSize, AnimatableAudioEditorConstants::AudioTrackHeight);}
+FIntPoint FAudioThumbnail::GetSize() const {return FIntPoint(TextureSize, AnimatableAudioEditorConstants::WaveformHeight);}
 FSlateShaderResource* FAudioThumbnail::GetViewportRenderTargetTexture() const {return Texture;}
 bool FAudioThumbnail::RequiresVsync() const {return false;}
 
@@ -325,8 +327,8 @@ void FAudioThumbnail::GenerateWaveformPreview(TArray<uint8>& OutData, TRange<flo
 
 	// @todo Sequencer This fixes looping drawing by pretending we are only dealing with a SoundWave
 	TRange<float> AudioTrueRange = TRange<float>(
-		AudioSection->GetAudioStartTime(),
-		AudioSection->GetAudioStartTime() + DeriveUnloopedDuration(AudioSection) * AudioSection->GetAudioDilationFactor());
+		AudioSection->GetStartTime() - AudioSection->GetStartOffset(),
+		AudioSection->GetStartTime() - AudioSection->GetStartOffset() + DeriveUnloopedDuration(AudioSection) * AudioSection->GetPitchMultiplierCurve().GetDefaultValue());
 	float TrueRangeSize = AudioTrueRange.Size<float>();
 
 	float DrawRangeSize = DrawRange.Size<float>();
@@ -627,7 +629,6 @@ FAudioSection::~FAudioSection()
 {
 }
 
-
 UMovieSceneSection* FAudioSection::GetSectionObject()
 { 
 	return &Section;
@@ -653,9 +654,19 @@ FText FAudioSection::GetSectionTitle() const
 
 float FAudioSection::GetSectionHeight() const
 {
-	return (float)AnimatableAudioEditorConstants::AudioTrackHeight;
+	return (float)AnimatableAudioEditorConstants::WaveformHeight;
 }
 
+void FAudioSection::GenerateSectionLayout( class ISectionLayoutBuilder& LayoutBuilder ) const
+{
+	UMovieSceneAudioSection* AudioSection = Cast<UMovieSceneAudioSection>( &Section );
+
+	SoundVolumeArea = MakeShareable( new FFloatCurveKeyArea( &AudioSection->GetSoundVolumeCurve(), AudioSection ) ) ;
+	PitchMultiplierArea = MakeShareable( new FFloatCurveKeyArea( &AudioSection->GetPitchMultiplierCurve(), AudioSection ) );
+
+	LayoutBuilder.AddKeyArea( "Volume", NSLOCTEXT( "FAudioSection", "SoundVolumeArea", "Volume" ), SoundVolumeArea.ToSharedRef() );
+	LayoutBuilder.AddKeyArea( "Pitch Multiplier", NSLOCTEXT( "FAudioSection", "PitchMultiplierArea", "Pitch Multiplier" ), PitchMultiplierArea.ToSharedRef() );
+}
 
 int32 FAudioSection::OnPaintSection( FSequencerSectionPainter& Painter ) const
 {
@@ -667,7 +678,7 @@ int32 FAudioSection::OnPaintSection( FSequencerSectionPainter& Painter ) const
 		FSlateDrawElement::MakeViewport(
 			Painter.DrawElements,
 			++LayerId,
-			Painter.SectionGeometry.ToPaintGeometry(FVector2D(StoredXOffset, 0), FVector2D(StoredXSize, Painter.SectionGeometry.Size.Y)),
+			Painter.SectionGeometry.ToPaintGeometry(FVector2D(StoredXOffset, 0), FVector2D(StoredXSize, (float)AnimatableAudioEditorConstants::WaveformHeight + 8.f)),
 			WaveformThumbnail,
 			Painter.SectionClippingRect,
 			(Painter.bParentEnabled ? ESlateDrawEffect::None : ESlateDrawEffect::DisabledEffect) | ESlateDrawEffect::NoGamma,
@@ -713,7 +724,8 @@ void FAudioSection::Tick( const FGeometry& AllottedGeometry, const FGeometry& Pa
 			!FMath::IsNearlyEqual(DrawRange.GetUpperBoundValue(), StoredDrawRange.GetUpperBoundValue()) ||
 			XOffset != StoredXOffset || XSize != StoredXSize || Track->GetColorTint() != StoredColor ||
 			StoredSoundWave != SoundWave ||
-			StoredShowIntensity != AudioSection->ShouldShowIntensity())
+			StoredShowIntensity != AudioSection->ShouldShowIntensity() ||
+			StoredStartOffset != AudioSection->GetStartOffset())
 		{
 			float DisplayScale = XSize / DrawRange.Size<float>();
 
@@ -739,13 +751,16 @@ void FAudioSection::Tick( const FGeometry& AllottedGeometry, const FGeometry& Pa
 
 void FAudioSection::RegenerateWaveforms(TRange<float> DrawRange, int32 XOffset, int32 XSize, const FColor& ColorTint, float DisplayScale)
 {
+	UMovieSceneAudioSection* AudioSection = Cast<UMovieSceneAudioSection>(&Section);
+
 	StoredDrawRange = DrawRange;
 	StoredXOffset = XOffset;
 	StoredXSize = XSize;
 	StoredColor = ColorTint;
-	StoredShowIntensity =  Cast<UMovieSceneAudioSection>(&Section)->ShouldShowIntensity();
+	StoredShowIntensity = AudioSection->ShouldShowIntensity();
+	StoredStartOffset = AudioSection->GetStartOffset();
 
-	if (DrawRange.IsDegenerate() || DrawRange.IsEmpty() || Cast<UMovieSceneAudioSection>(&Section)->GetSound() == NULL)
+	if (DrawRange.IsDegenerate() || DrawRange.IsEmpty() || AudioSection->GetSound() == NULL)
 	{
 		WaveformThumbnail.Reset();
 	}
diff --git a/Engine/Source/Editor/MovieSceneTools/Private/TrackEditors/AudioTrackEditor.h b/Engine/Source/Editor/MovieSceneTools/Private/TrackEditors/AudioTrackEditor.h
--- a/Engine/Source/Editor/MovieSceneTools/Private/TrackEditors/AudioTrackEditor.h
+++ b/Engine/Source/Editor/MovieSceneTools/Private/TrackEditors/AudioTrackEditor.h
@@ -14,6 +14,7 @@
 
 class FAssetData;
 class FAudioThumbnail;
+class FFloatCurveKeyArea;
 class FMenuBuilder;
 class FSequencerSectionPainter;
 class USoundWave;
@@ -100,10 +101,10 @@ public:
 	virtual FText GetDisplayName() const override;
 	virtual FText GetSectionTitle() const override;
 	virtual float GetSectionHeight() const override;
-	virtual void GenerateSectionLayout( class ISectionLayoutBuilder& LayoutBuilder) const override { }
+	virtual void GenerateSectionLayout( class ISectionLayoutBuilder& LayoutBuilder) const override;
 	virtual int32 OnPaintSection(FSequencerSectionPainter& Painter) const override;
 	virtual void Tick(const FGeometry& AllottedGeometry, const FGeometry& ParentGeometry, const double InCurrentTime, const float InDeltaTime) override;
-	
+
 private:
 
 	/* Re-creates the texture used to preview the waveform. */
@@ -113,13 +114,18 @@ private:
 
 	/** The section we are visualizing. */
 	UMovieSceneSection& Section;
-	
+
 	/** The waveform thumbnail render object. */
 	TSharedPtr<class FAudioThumbnail> WaveformThumbnail;
 
+	/** Sound volume/pitch key areas. */
+	mutable TSharedPtr<FFloatCurveKeyArea> SoundVolumeArea;
+	mutable TSharedPtr<FFloatCurveKeyArea> PitchMultiplierArea;
+
 	/** Stored data about the waveform to determine when it is invalidated. */
 	TRange<float> StoredDrawRange;
 	bool StoredShowIntensity;
+	float StoredStartOffset;
 	int32 StoredXOffset;
 	int32 StoredXSize;
 	FColor StoredColor;
diff --git a/Engine/Source/Editor/MovieSceneTools/Private/TrackEditors/CameraCutTrackEditor.cpp b/Engine/Source/Editor/MovieSceneTools/Private/TrackEditors/CameraCutTrackEditor.cpp
--- a/Engine/Source/Editor/MovieSceneTools/Private/TrackEditors/CameraCutTrackEditor.cpp
+++ b/Engine/Source/Editor/MovieSceneTools/Private/TrackEditors/CameraCutTrackEditor.cpp
@@ -23,6 +23,39 @@
 #define LOCTEXT_NAMESPACE "FCameraCutTrackEditor"
 
 
+class FCameraCutTrackCommands
+	: public TCommands<FCameraCutTrackCommands>
+{
+public:
+
+	FCameraCutTrackCommands()
+		: TCommands<FCameraCutTrackCommands>
+	(
+		"CameraCutTrack",
+		NSLOCTEXT("Contexts", "CameraCutTrack", "CameraCutTrack"),
+		NAME_None, // "MainFrame" // @todo Fix this crash
+		FEditorStyle::GetStyleSetName() // Icon Style Set
+	)
+		, BindingCount(0)
+	{ }
+		
+	/** Toggle the camera lock */
+	TSharedPtr< FUICommandInfo > ToggleLockCamera;
+
+	/**
+	 * Initialize commands
+	 */
+	virtual void RegisterCommands() override;
+
+	mutable uint32 BindingCount;
+};
+
+
+void FCameraCutTrackCommands::RegisterCommands()
+{
+	UI_COMMAND( ToggleLockCamera, "Toggle Lock Camera", "Toggle locking the viewport to the camera cut track.", EUserInterfaceActionType::Button, FInputChord(EModifierKey::Control, EKeys::L) );
+}
+
 /* FCameraCutTrackEditor structors
  *****************************************************************************/
 
@@ -30,8 +63,20 @@ FCameraCutTrackEditor::FCameraCutTrackEditor(TSharedRef<ISequencer> InSequencer)
 	: FMovieSceneTrackEditor(InSequencer) 
 {
 	ThumbnailPool = MakeShareable(new FTrackEditorThumbnailPool(InSequencer));
+
+	FCameraCutTrackCommands::Register();
 }
 
+void FCameraCutTrackEditor::OnRelease()
+{
+	const FCameraCutTrackCommands& Commands = FCameraCutTrackCommands::Get();
+	Commands.BindingCount--;
+	
+	if (Commands.BindingCount < 1)
+	{
+		FCameraCutTrackCommands::Unregister();
+	}
+}
 
 TSharedRef<ISequencerTrackEditor> FCameraCutTrackEditor::CreateTrackEditor(TSharedRef<ISequencer> InSequencer)
 {
@@ -42,6 +87,17 @@ TSharedRef<ISequencerTrackEditor> FCameraCutTrackEditor::CreateTrackEditor(TShar
 /* ISequencerTrackEditor interface
  *****************************************************************************/
 
+void FCameraCutTrackEditor::BindCommands(TSharedRef<FUICommandList> SequencerCommandBindings)
+{
+	const FCameraCutTrackCommands& Commands = FCameraCutTrackCommands::Get();
+
+	SequencerCommandBindings->MapAction(
+		Commands.ToggleLockCamera,
+		FExecuteAction::CreateSP( this, &FCameraCutTrackEditor::ToggleLockCamera) );
+
+	Commands.BindingCount++;
+}
+
 void FCameraCutTrackEditor::BuildAddTrackMenu(FMenuBuilder& MenuBuilder)
 {
 	UMovieSceneSequence* RootMovieSceneSequence = GetSequencer()->GetRootMovieSceneSequence();
@@ -301,6 +357,11 @@ void FCameraCutTrackEditor::OnLockCameraClicked(ECheckBoxState CheckBoxState)
 	GetSequencer()->ForceEvaluate();
 }
 
+void FCameraCutTrackEditor::ToggleLockCamera()
+{
+	OnLockCameraClicked(IsCameraLocked() == ECheckBoxState::Checked ?  ECheckBoxState::Unchecked :  ECheckBoxState::Checked);
+}
+
 FText FCameraCutTrackEditor::GetLockCameraToolTip() const
 {
 	return IsCameraLocked() == ECheckBoxState::Checked ?
diff --git a/Engine/Source/Editor/MovieSceneTools/Private/TrackEditors/CameraCutTrackEditor.h b/Engine/Source/Editor/MovieSceneTools/Private/TrackEditors/CameraCutTrackEditor.h
--- a/Engine/Source/Editor/MovieSceneTools/Private/TrackEditors/CameraCutTrackEditor.h
+++ b/Engine/Source/Editor/MovieSceneTools/Private/TrackEditors/CameraCutTrackEditor.h
@@ -48,9 +48,11 @@ public:
 
 	// ISequencerTrackEditor interface
 
+	virtual void BindCommands(TSharedRef<FUICommandList> SequencerCommandBindings) override;
 	virtual void BuildAddTrackMenu(FMenuBuilder& MenuBuilder) override;
 	virtual TSharedPtr<SWidget> BuildOutlinerEditWidget(const FGuid& ObjectBinding, UMovieSceneTrack* Track, const FBuildEditWidgetParams& Params) override;
 	virtual TSharedRef<ISequencerSection> MakeSectionInterface(UMovieSceneSection& SectionObject, UMovieSceneTrack& Track, FGuid ObjectBinding) override;
+	virtual void OnRelease() override;
 	virtual bool SupportsType(TSubclassOf<UMovieSceneTrack> Type ) const override;
 	virtual void Tick(float DeltaTime) override;
 	virtual const FSlateBrush* GetIconBrush() const override;
@@ -88,6 +90,9 @@ private:
 	/** Delegate for locked camera button */
 	void OnLockCameraClicked(ECheckBoxState CheckBoxState);
 
+	/** Toggle the state of the camera lock */
+	void ToggleLockCamera();
+
 	/** Delegate for camera button lock tooltip */
 	FText GetLockCameraToolTip() const; 
 
diff --git a/Engine/Source/Editor/MovieSceneTools/Private/TrackEditors/MaterialTrackEditor.cpp b/Engine/Source/Editor/MovieSceneTools/Private/TrackEditors/MaterialTrackEditor.cpp
--- a/Engine/Source/Editor/MovieSceneTools/Private/TrackEditors/MaterialTrackEditor.cpp
+++ b/Engine/Source/Editor/MovieSceneTools/Private/TrackEditors/MaterialTrackEditor.cpp
@@ -126,7 +126,7 @@ void FMaterialTrackEditor::AddScalarParameter( FGuid ObjectBinding, UMovieSceneM
 {
 	float KeyTime = GetTimeForKey();
 
-	UMaterial* Material = GetMaterialForTrack(ObjectBinding, MaterialTrack);
+	UMaterialInterface* Material = GetMaterialInterfaceForTrack(ObjectBinding, MaterialTrack);
 	if (Material != nullptr)
 	{
 		const FScopedTransaction Transaction( LOCTEXT( "AddScalarParameter", "Add scalar parameter" ) );
@@ -143,7 +143,7 @@ void FMaterialTrackEditor::AddColorParameter( FGuid ObjectBinding, UMovieSceneMa
 {
 	float KeyTime = GetTimeForKey();
 
-	UMaterial* Material = GetMaterialForTrack( ObjectBinding, MaterialTrack );
+	UMaterialInterface* Material = GetMaterialInterfaceForTrack( ObjectBinding, MaterialTrack );
 	if ( Material != nullptr )
 	{
 		const FScopedTransaction Transaction( LOCTEXT( "AddVectorParameter", "Add vector parameter" ) );
diff --git a/Engine/Source/Editor/MovieSceneTools/Private/TrackEditors/PropertyTrackEditors/VectorPropertyTrackEditor.cpp b/Engine/Source/Editor/MovieSceneTools/Private/TrackEditors/PropertyTrackEditors/VectorPropertyTrackEditor.cpp
--- a/Engine/Source/Editor/MovieSceneTools/Private/TrackEditors/PropertyTrackEditors/VectorPropertyTrackEditor.cpp
+++ b/Engine/Source/Editor/MovieSceneTools/Private/TrackEditors/PropertyTrackEditors/VectorPropertyTrackEditor.cpp
@@ -2,6 +2,12 @@
 
 #include "TrackEditors/PropertyTrackEditors/VectorPropertyTrackEditor.h"
 #include "Sections/VectorPropertySection.h"
+#include "MovieSceneVectorTrack.h"
+#include "MatineeImportTools.h"
+#include "Matinee/InterpTrackVectorProp.h"
+#include "UnrealEdGlobals.h"
+#include "Classes/Editor/UnrealEdEngine.h"
+
 
 FName FVectorPropertyTrackEditor::XName( "X" );
 FName FVectorPropertyTrackEditor::YName( "Y" );
@@ -96,3 +102,36 @@ void FVectorPropertyTrackEditor::InitializeNewTrack( UMovieSceneVectorTrack* New
 		NewTrack->SetNumChannelsUsed( 4 );
 	}
 }
+
+void CopyInterpVectorTrack(TSharedRef<ISequencer> Sequencer, UInterpTrackVectorProp* MatineeVectorTrack, UMovieSceneVectorTrack* VectorTrack)
+{
+	if (FMatineeImportTools::CopyInterpVectorTrack(MatineeVectorTrack, VectorTrack))
+	{
+		Sequencer.Get().NotifyMovieSceneDataChanged( EMovieSceneDataChangeType::MovieSceneStructureItemAdded );
+	}
+}
+
+void FVectorPropertyTrackEditor::BuildTrackContextMenu( FMenuBuilder& MenuBuilder, UMovieSceneTrack* Track )
+{
+	UInterpTrackVectorProp* MatineeVectorTrack = nullptr;
+	for ( UObject* CopyPasteObject : GUnrealEd->MatineeCopyPasteBuffer )
+	{
+		MatineeVectorTrack = Cast<UInterpTrackVectorProp>( CopyPasteObject );
+		if ( MatineeVectorTrack != nullptr )
+		{
+			break;
+		}
+	}
+	UMovieSceneVectorTrack* VectorTrack = Cast<UMovieSceneVectorTrack>( Track );
+	MenuBuilder.AddMenuEntry(
+		NSLOCTEXT( "Sequencer", "PasteMatineeVectorTrack", "Paste Matinee Vector Track" ),
+		NSLOCTEXT( "Sequencer", "PasteMatineeVectorTrackTooltip", "Pastes keys from a Matinee vector track into this track." ),
+		FSlateIcon(),
+		FUIAction(
+			FExecuteAction::CreateStatic( &CopyInterpVectorTrack, GetSequencer().ToSharedRef(), MatineeVectorTrack, VectorTrack ),
+			FCanExecuteAction::CreateLambda( [=]()->bool { return MatineeVectorTrack != nullptr && MatineeVectorTrack->GetNumKeys() > 0 && VectorTrack != nullptr && VectorTrack->GetNumChannelsUsed() == 3; } ) ) );
+
+	MenuBuilder.AddMenuSeparator();
+	FKeyframeTrackEditor::BuildTrackContextMenu(MenuBuilder, Track);
+}
+
diff --git a/Engine/Source/Editor/MovieSceneTools/Private/TrackEditors/PropertyTrackEditors/VectorPropertyTrackEditor.h b/Engine/Source/Editor/MovieSceneTools/Private/TrackEditors/PropertyTrackEditors/VectorPropertyTrackEditor.h
--- a/Engine/Source/Editor/MovieSceneTools/Private/TrackEditors/PropertyTrackEditors/VectorPropertyTrackEditor.h
+++ b/Engine/Source/Editor/MovieSceneTools/Private/TrackEditors/PropertyTrackEditors/VectorPropertyTrackEditor.h
@@ -43,6 +43,10 @@ public:
 
 protected:
 
+	//~ FMovieSceneTrackEditor interface
+
+	virtual void BuildTrackContextMenu(FMenuBuilder& MenuBuilder, UMovieSceneTrack* Track) override;
+
 	//~ FPropertyTrackEditor interface
 
 	virtual void GenerateKeysFromPropertyChanged(const FPropertyChangedParams& PropertyChangedParams, TArray<FVectorKey>& NewGeneratedKeys, TArray<FVectorKey>& DefaultGeneratedKeys) override;
diff --git a/Engine/Source/Editor/MovieSceneTools/Private/TrackEditors/SkeletalAnimationTrackEditor.cpp b/Engine/Source/Editor/MovieSceneTools/Private/TrackEditors/SkeletalAnimationTrackEditor.cpp
--- a/Engine/Source/Editor/MovieSceneTools/Private/TrackEditors/SkeletalAnimationTrackEditor.cpp
+++ b/Engine/Source/Editor/MovieSceneTools/Private/TrackEditors/SkeletalAnimationTrackEditor.cpp
@@ -25,6 +25,10 @@
 #include "MatineeImportTools.h"
 #include "Matinee/InterpTrackAnimControl.h"
 #include "SequencerUtilities.h"
+#include "FloatCurveKeyArea.h"
+#include "ISectionLayoutBuilder.h"
+#include "Classes/Animation/AnimMontage.h"
+#include "Classes/Animation/AnimSequence.h"
 
 
 namespace SkeletalAnimationEditorConstants
@@ -70,6 +74,13 @@ float FSkeletalAnimationSection::GetSectionHeight() const
 	return (float)SkeletalAnimationEditorConstants::AnimationTrackHeight;
 }
 
+void FSkeletalAnimationSection::GenerateSectionLayout( class ISectionLayoutBuilder& LayoutBuilder ) const
+{
+	WeightArea = MakeShareable( new FFloatCurveKeyArea( &Section.Params.Weight, &Section ) );
+
+	LayoutBuilder.AddKeyArea( "Weight", NSLOCTEXT( "FSkeletalAnimationSection", "WeightArea", "Weight" ), WeightArea.ToSharedRef() );
+}
+
 
 int32 FSkeletalAnimationSection::OnPaintSection( FSequencerSectionPainter& Painter ) const
 {
@@ -83,7 +94,7 @@ int32 FSkeletalAnimationSection::OnPaintSection( FSequencerSectionPainter& Paint
 	float CurrentTime = Section.GetStartTime();
 	float AnimPlayRate = FMath::IsNearlyZero(Section.Params.PlayRate) ? 1.0f : Section.Params.PlayRate;
 	float SeqLength = (Section.Params.GetSequenceLength() - (Section.Params.StartOffset + Section.Params.EndOffset)) / AnimPlayRate;
-	while (CurrentTime < Section.GetEndTime() && !FMath::IsNearlyZero(Section.Params.GetDuration()) && SeqLength > 0)
+	while (CurrentTime < Section.GetEndTime() && !FMath::IsNearlyZero(SeqLength, KINDA_SMALL_NUMBER) && SeqLength > 0)
 	{
 		if (CurrentTime > Section.GetStartTime())
 		{
@@ -182,7 +193,7 @@ void FSkeletalAnimationTrackEditor::AddKey(const FGuid& ObjectGuid)
 				FSlateApplication::Get().PushMenu(
 					Parent.ToSharedRef(),
 					FWidgetPath(),
-					BuildAnimationSubMenu(ObjectGuid, Skeleton),
+					BuildAnimationSubMenu(ObjectGuid, Skeleton, nullptr),
 					FSlateApplication::Get().GetCursorPos(),
 					FPopupTransitionEffect(FPopupTransitionEffect::TypeInPopup)
 					);
@@ -206,10 +217,11 @@ bool FSkeletalAnimationTrackEditor::HandleAssetAdded(UObject* Asset, const FGuid
 
 			if (Skeleton && Skeleton == AnimSequence->GetSkeleton())
 			{
+				UObject* Object = SequencerPtr->FindSpawnedObjectOrTemplate(TargetObjectGuid);
 				
+				UMovieSceneTrack* Track = nullptr;
 
-				UObject* Object = SequencerPtr->FindSpawnedObjectOrTemplate(TargetObjectGuid);
-				AnimatablePropertyChanged(FOnKeyProperty::CreateRaw(this, &FSkeletalAnimationTrackEditor::AddKeyInternal, Object, AnimSequence));
+				AnimatablePropertyChanged(FOnKeyProperty::CreateRaw(this, &FSkeletalAnimationTrackEditor::AddKeyInternal, Object, AnimSequence, Track));
 
 				return true;
 			}
@@ -238,31 +250,52 @@ void FSkeletalAnimationTrackEditor::BuildObjectBindingTrackMenu(FMenuBuilder& Me
 
 			if (AssetDataList.Num())
 			{
+				UMovieSceneTrack* Track = nullptr;
+
 				MenuBuilder.AddSubMenu(
 					LOCTEXT("AddAnimation", "Animation"), NSLOCTEXT("Sequencer", "AddAnimationTooltip", "Adds an animation track."),
-					FNewMenuDelegate::CreateRaw(this, &FSkeletalAnimationTrackEditor::AddAnimationSubMenu, ObjectBinding, Skeleton)
+					FNewMenuDelegate::CreateRaw(this, &FSkeletalAnimationTrackEditor::AddAnimationSubMenu, ObjectBinding, Skeleton, Track)
 				);
 			}
 		}
 	}
 }
 
-TSharedRef<SWidget> FSkeletalAnimationTrackEditor::BuildAnimationSubMenu(FGuid ObjectBinding, USkeleton* Skeleton)
+TSharedRef<SWidget> FSkeletalAnimationTrackEditor::BuildAnimationSubMenu(FGuid ObjectBinding, USkeleton* Skeleton, UMovieSceneTrack* Track)
 {
 	FMenuBuilder MenuBuilder(true, nullptr);
-
-	AddAnimationSubMenu(MenuBuilder, ObjectBinding, Skeleton);
+	
+	AddAnimationSubMenu(MenuBuilder, ObjectBinding, Skeleton, Track);
 
 	return MenuBuilder.MakeWidget();
 }
 
-void FSkeletalAnimationTrackEditor::AddAnimationSubMenu(FMenuBuilder& MenuBuilder, FGuid ObjectBinding, USkeleton* Skeleton)
+bool FSkeletalAnimationTrackEditor::ShouldFilterAsset(const FAssetData& AssetData)
+{
+	// we don't want montage
+	if (AssetData.AssetClass == UAnimMontage::StaticClass()->GetFName())
+	{
+		return true;
+	}
+
+	const FString EnumString = AssetData.GetTagValueRef<FString>(GET_MEMBER_NAME_CHECKED(UAnimSequence, AdditiveAnimType));
+	if (EnumString.IsEmpty())
+	{
+		return false;
+	}
+
+	UEnum* AdditiveTypeEnum = FindObject<UEnum>(ANY_PACKAGE, TEXT("EAdditiveAnimationType"), true);
+	return ((EAdditiveAnimationType)AdditiveTypeEnum->GetValueByName(*EnumString) == AAT_RotationOffsetMeshSpace);
+}
+
+void FSkeletalAnimationTrackEditor::AddAnimationSubMenu(FMenuBuilder& MenuBuilder, FGuid ObjectBinding, USkeleton* Skeleton, UMovieSceneTrack* Track)
 {
 	FAssetPickerConfig AssetPickerConfig;
 	{
-		AssetPickerConfig.OnAssetSelected = FOnAssetSelected::CreateRaw( this, &FSkeletalAnimationTrackEditor::OnAnimationAssetSelected, ObjectBinding);
+		AssetPickerConfig.OnAssetSelected = FOnAssetSelected::CreateRaw( this, &FSkeletalAnimationTrackEditor::OnAnimationAssetSelected, ObjectBinding, Track);
 		AssetPickerConfig.bAllowNullSelection = false;
 		AssetPickerConfig.InitialAssetViewType = EAssetViewType::List;
+		AssetPickerConfig.OnShouldFilterAsset = FOnShouldFilterAsset::CreateRaw(this, &FSkeletalAnimationTrackEditor::ShouldFilterAsset);
 		AssetPickerConfig.Filter.bRecursiveClasses = true;
 		AssetPickerConfig.Filter.ClassNames.Add(UAnimSequenceBase::StaticClass()->GetFName());
 		AssetPickerConfig.Filter.TagsAndValues.Add(TEXT("Skeleton"), FAssetData(Skeleton).GetExportTextName());
@@ -281,7 +314,7 @@ void FSkeletalAnimationTrackEditor::AddAnimationSubMenu(FMenuBuilder& MenuBuilde
 }
 
 
-void FSkeletalAnimationTrackEditor::OnAnimationAssetSelected(const FAssetData& AssetData, FGuid ObjectBinding)
+void FSkeletalAnimationTrackEditor::OnAnimationAssetSelected(const FAssetData& AssetData, FGuid ObjectBinding, UMovieSceneTrack* Track)
 {
 	FSlateApplication::Get().DismissAllMenus();
 
@@ -293,12 +326,12 @@ void FSkeletalAnimationTrackEditor::OnAnimationAssetSelected(const FAssetData& A
 		UAnimSequenceBase* AnimSequence = CastChecked<UAnimSequenceBase>(AssetData.GetAsset());
 
 		UObject* Object = SequencerPtr->FindSpawnedObjectOrTemplate(ObjectBinding);
-		AnimatablePropertyChanged( FOnKeyProperty::CreateRaw( this, &FSkeletalAnimationTrackEditor::AddKeyInternal, Object, AnimSequence) );
+		AnimatablePropertyChanged( FOnKeyProperty::CreateRaw( this, &FSkeletalAnimationTrackEditor::AddKeyInternal, Object, AnimSequence, Track) );
 	}
 }
 
 
-bool FSkeletalAnimationTrackEditor::AddKeyInternal( float KeyTime, UObject* Object, class UAnimSequenceBase* AnimSequence )
+bool FSkeletalAnimationTrackEditor::AddKeyInternal( float KeyTime, UObject* Object, class UAnimSequenceBase* AnimSequence, UMovieSceneTrack* Track )
 {
 	bool bHandleCreated = false;
 	bool bTrackCreated = false;
@@ -309,9 +342,11 @@ bool FSkeletalAnimationTrackEditor::AddKeyInternal( float KeyTime, UObject* Obje
 	bHandleCreated |= HandleResult.bWasCreated;
 	if (ObjectHandle.IsValid())
 	{
-		FFindOrCreateTrackResult TrackResult = FindOrCreateTrackForObject(ObjectHandle, UMovieSceneSkeletalAnimationTrack::StaticClass());
-		UMovieSceneTrack* Track = TrackResult.Track;
-		bTrackCreated |= TrackResult.bWasCreated;
+		if (!Track)
+		{
+			Track = AddTrack(GetSequencer()->GetFocusedMovieSceneSequence()->GetMovieScene(), ObjectHandle, UMovieSceneSkeletalAnimationTrack::StaticClass(), NAME_None);
+			bTrackCreated = true;
+		}
 
 		if (ensure(Track))
 		{
@@ -401,7 +436,7 @@ TSharedPtr<SWidget> FSkeletalAnimationTrackEditor::BuildOutlinerEditWidget(const
 		.AutoWidth()
 		.VAlign(VAlign_Center)
 		[
-			FSequencerUtilities::MakeAddButton(LOCTEXT("AnimationText", "Animation"), FOnGetContent::CreateSP(this, &FSkeletalAnimationTrackEditor::BuildAnimationSubMenu, ObjectBinding, Skeleton), Params.NodeIsHovered)
+			FSequencerUtilities::MakeAddButton(LOCTEXT("AnimationText", "Animation"), FOnGetContent::CreateSP(this, &FSkeletalAnimationTrackEditor::BuildAnimationSubMenu, ObjectBinding, Skeleton, Track), Params.NodeIsHovered)
 		];
 	}
 
diff --git a/Engine/Source/Editor/MovieSceneTools/Private/TrackEditors/SkeletalAnimationTrackEditor.h b/Engine/Source/Editor/MovieSceneTools/Private/TrackEditors/SkeletalAnimationTrackEditor.h
--- a/Engine/Source/Editor/MovieSceneTools/Private/TrackEditors/SkeletalAnimationTrackEditor.h
+++ b/Engine/Source/Editor/MovieSceneTools/Private/TrackEditors/SkeletalAnimationTrackEditor.h
@@ -13,6 +13,7 @@
 #include "MovieSceneTrackEditor.h"
 
 class FAssetData;
+class FFloatCurveKeyArea;
 class FMenuBuilder;
 class FSequencerSectionPainter;
 class UMovieSceneSkeletalAnimationSection;
@@ -57,21 +58,24 @@ public:
 private:
 
 	/** Animation sub menu */
-	TSharedRef<SWidget> BuildAnimationSubMenu(FGuid ObjectBinding, USkeleton* Skeleton);
-	void AddAnimationSubMenu(FMenuBuilder& MenuBuilder, FGuid ObjectBinding, USkeleton* Skeleton);
+	TSharedRef<SWidget> BuildAnimationSubMenu(FGuid ObjectBinding, USkeleton* Skeleton, UMovieSceneTrack* Track);
+	void AddAnimationSubMenu(FMenuBuilder& MenuBuilder, FGuid ObjectBinding, USkeleton* Skeleton, UMovieSceneTrack* Track);
+
+	/** Animation sub menu filter function */
+	bool ShouldFilterAsset(const FAssetData& AssetData);
 
 	/** Animation asset selected */
-	void OnAnimationAssetSelected(const FAssetData& AssetData, FGuid ObjectBinding);
+	void OnAnimationAssetSelected(const FAssetData& AssetData, FGuid ObjectBinding, UMovieSceneTrack* Track);
 
 	/** Delegate for AnimatablePropertyChanged in AddKey */
-	bool AddKeyInternal(float KeyTime, UObject* Object, class UAnimSequenceBase* AnimSequence);
+	bool AddKeyInternal(float KeyTime, UObject* Object, class UAnimSequenceBase* AnimSequence, UMovieSceneTrack* Track);
 
 	/** Gets a skeleton from an object guid in the movie scene */
 	class USkeleton* AcquireSkeletonFromObjectGuid(const FGuid& Guid);
 };
 
 
-/** Class for animation sections, handles drawing of all waveform previews */
+/** Class for animation sections */
 class FSkeletalAnimationSection
 	: public ISequencerSection
 	, public TSharedFromThis<FSkeletalAnimationSection>
@@ -92,7 +96,7 @@ public:
 	virtual FText GetDisplayName() const override;
 	virtual FText GetSectionTitle() const override;
 	virtual float GetSectionHeight() const override;
-	virtual void GenerateSectionLayout( class ISectionLayoutBuilder& LayoutBuilder ) const override {}
+	virtual void GenerateSectionLayout( class ISectionLayoutBuilder& LayoutBuilder ) const override;
 	virtual int32 OnPaintSection( FSequencerSectionPainter& Painter ) const override;
 	virtual void BeginResizeSection() override;
 	virtual void ResizeSection(ESequencerSectionResizeMode ResizeMode, float ResizeTime) override;
@@ -102,6 +106,9 @@ private:
 	/** The section we are visualizing */
 	UMovieSceneSkeletalAnimationSection& Section;
 
+	/** Weight key areas. */
+	mutable TSharedPtr<FFloatCurveKeyArea> WeightArea;
+
 	/** Cached start offset value valid only during resize */
 	float InitialStartOffsetDuringResize;
 	
diff --git a/Engine/Source/Editor/MovieSceneTools/Private/TrackEditors/TransformTrackEditor.cpp b/Engine/Source/Editor/MovieSceneTools/Private/TrackEditors/TransformTrackEditor.cpp
--- a/Engine/Source/Editor/MovieSceneTools/Private/TrackEditors/TransformTrackEditor.cpp
+++ b/Engine/Source/Editor/MovieSceneTools/Private/TrackEditors/TransformTrackEditor.cpp
@@ -27,6 +27,23 @@
 
 #define LOCTEXT_NAMESPACE "MovieScene_TransformTrack"
 
+void GetActorAndSceneComponentFromObject( UObject* Object, AActor*& OutActor, USceneComponent*& OutSceneComponent )
+{
+	OutActor = Cast<AActor>( Object );
+	if ( OutActor != nullptr && OutActor->GetRootComponent() )
+	{
+		OutSceneComponent = OutActor->GetRootComponent();
+	}
+	else
+	{
+		// If the object wasn't an actor attempt to get it directly as a scene component and then get the actor from there.
+		OutSceneComponent = Cast<USceneComponent>( Object );
+		if ( OutSceneComponent != nullptr )
+		{
+			OutActor = Cast<AActor>( OutSceneComponent->GetOuter() );
+		}
+	}
+}
 
 /**
  * Class that draws a transform section in the sequencer
@@ -198,9 +215,12 @@ private:
 		if (Sequencer != nullptr)
 		{
 			TArrayView<TWeakObjectPtr<UObject>> RuntimeObjects = Sequencer->FindBoundObjects(ObjectBinding, Sequencer->GetFocusedTemplateID());
-			if (RuntimeObjects.Num() == 1)
+			if (RuntimeObjects.Num() == 1 && RuntimeObjects[0].IsValid())
 			{
-				return MovieSceneHelpers::SceneComponentFromRuntimeObject(RuntimeObjects[0].Get());
+				AActor* Actor = nullptr;
+				USceneComponent* SceneComponentThatChanged = nullptr;
+				GetActorAndSceneComponentFromObject(RuntimeObjects[0].Get(), Actor, SceneComponentThatChanged);
+				return SceneComponentThatChanged;
 			}
 		}
 		return nullptr;
@@ -230,6 +250,7 @@ public:
 		NAME_None, // "MainFrame" // @todo Fix this crash
 		FEditorStyle::GetStyleSetName() // Icon Style Set
 	)
+		, BindingCount(0)
 	{ }
 		
 	/** Sets a transform key at the current time for the selected actor */
@@ -248,6 +269,8 @@ public:
 	 * Initialize commands
 	 */
 	virtual void RegisterCommands() override;
+
+	mutable uint32 BindingCount;
 };
 
 
@@ -289,7 +312,13 @@ void F3DTransformTrackEditor::OnRelease()
 	GEditor->OnBeginCameraMovement().RemoveAll( this );
 	GEditor->OnEndCameraMovement().RemoveAll( this );
 
-	F3DTransformTrackCommands::Unregister();
+	const F3DTransformTrackCommands& Commands = F3DTransformTrackCommands::Get();
+	Commands.BindingCount--;
+	
+	if (Commands.BindingCount < 1)
+	{
+		F3DTransformTrackCommands::Unregister();
+	}
 }
 
 
@@ -411,23 +440,15 @@ void F3DTransformTrackEditor::OnPreTransformChanged( UObject& InObject )
 
 	if( IsAllowedToAutoKey() )
 	{
-		USceneComponent* SceneComponent = nullptr;
-		AActor* Actor = Cast<AActor>( &InObject );
-		if( Actor && Actor->GetRootComponent() )
-		{
-			SceneComponent = Actor->GetRootComponent();
-		}
-		else
-		{
-			// If the object wasn't an actor attempt to get it directly as a scene component 
-			SceneComponent = Cast<USceneComponent>( &InObject );
-		}
+		AActor* Actor = nullptr;
+		USceneComponent* SceneComponentThatChanged = nullptr;
+		GetActorAndSceneComponentFromObject(&InObject, Actor, SceneComponentThatChanged);
 
-		if( SceneComponent )
+		if( SceneComponentThatChanged )
 		{
 			// Cache off the existing transform so we can detect which components have changed
 			// and keys only when something has changed
-			FTransformData Transform( SceneComponent );
+			FTransformData Transform( SceneComponentThatChanged );
 
 			ObjectToExistingTransform.Add( &InObject, Transform );
 		}
@@ -470,7 +491,7 @@ void F3DTransformTrackEditor::OnTransformChanged( UObject& InObject )
 
 		const bool bUnwindRotation = GetSequencer()->IsRecordingLive();
 
-		AddTransformKeys(Actor, ExistingTransform, NewTransformData, EKey3DTransformChannel::All, bUnwindRotation, ESequencerKeyMode::AutoKey);
+		AddTransformKeys(&InObject, ExistingTransform, NewTransformData, EKey3DTransformChannel::All, bUnwindRotation, ESequencerKeyMode::AutoKey);
 	}
 }
 
@@ -488,11 +509,23 @@ void F3DTransformTrackEditor::OnAddTransformKeysForSelectedObjects( EKey3DTransf
 		}
 	}
 
-	USelection* CurrentSelection = GEditor->GetSelectedActors();
-	TArray<UObject*> SelectedActors;
-	CurrentSelection->GetSelectedObjects( AActor::StaticClass(), SelectedActors );
+	TArray<UObject*> SelectedObjects;
+	for (FSelectedEditableComponentIterator It(GEditor->GetSelectedEditableComponentIterator()); It; ++It)
+	{
+		USceneComponent* SceneComponent = Cast<USceneComponent>(*It);
+		if (SceneComponent)
+		{
+			SelectedObjects.Add(SceneComponent);
+		}
+	}
+	
+	if (SelectedObjects.Num() == 0)
+	{
+		USelection* CurrentSelection = GEditor->GetSelectedActors();
+		CurrentSelection->GetSelectedObjects( AActor::StaticClass(), SelectedObjects );
+	}
 
-	for (TArray<UObject*>::TIterator It(SelectedActors); It; ++It)
+	for (TArray<UObject*>::TIterator It(SelectedObjects); It; ++It)
 	{
 		AddTransformKeysForObject(*It, Channel, ESequencerKeyMode::ManualKeyForced);
 	}
@@ -518,6 +551,8 @@ void F3DTransformTrackEditor::BindCommands(TSharedRef<FUICommandList> SequencerC
 	SequencerCommandBindings->MapAction(
 		Commands.AddScaleKey,
 		FExecuteAction::CreateSP( this, &F3DTransformTrackEditor::OnAddTransformKeysForSelectedObjects, EKey3DTransformChannel::Scale ) );
+
+	Commands.BindingCount++;
 }
 
 
@@ -754,24 +789,6 @@ FText F3DTransformTrackEditor::GetLockCameraToolTip(FGuid ObjectGuid) const
 	return FText();
 }
 
-void F3DTransformTrackEditor::GetActorAndSceneComponentFromObject( UObject* Object, AActor*& OutActor, USceneComponent*& OutSceneComponent )
-{
-	OutActor = Cast<AActor>( Object );
-	if ( OutActor != nullptr && OutActor->GetRootComponent() )
-	{
-		OutSceneComponent = OutActor->GetRootComponent();
-	}
-	else
-	{
-		// If the object wasn't an actor attempt to get it directly as a scene component and then get the actor from there.
-		OutSceneComponent = Cast<USceneComponent>( Object );
-		if ( OutSceneComponent != nullptr )
-		{
-			OutActor = Cast<AActor>( OutSceneComponent->GetOuter() );
-		}
-	}
-}
-
 void GetKeysForVector( bool LastVectorIsValid, const FVector& LastVector, const FVector& CurrentVector, EKey3DTransformChannel::Type VectorChannel, EKey3DTransformChannel::Type ChannelsToKey, bool bUnwindRotation, TArray<FTransformKey>& OutNewKeys, TArray<FTransformKey>& OutDefaultKeys)
 {
 	bool bKeyChannel = ChannelsToKey == EKey3DTransformChannel::All || ChannelsToKey == VectorChannel;
diff --git a/Engine/Source/Editor/MovieSceneTools/Private/TrackEditors/TransformTrackEditor.h b/Engine/Source/Editor/MovieSceneTools/Private/TrackEditors/TransformTrackEditor.h
--- a/Engine/Source/Editor/MovieSceneTools/Private/TrackEditors/TransformTrackEditor.h
+++ b/Engine/Source/Editor/MovieSceneTools/Private/TrackEditors/TransformTrackEditor.h
@@ -88,9 +88,6 @@ private:
 	/** Delegate for camera button lock tooltip */
 	FText GetLockCameraToolTip(FGuid ObjectGuid) const; 
 
-	/** Tries to get an actor and root scene component from an object which can be either an actor or a component. */
-	void GetActorAndSceneComponentFromObject(UObject* Object, AActor*& OutActor, USceneComponent*& OutSceneComponent);
-
 	/** Generates transform keys based on the last transform, the current transform, and other options. 
 		One transform key is generated for each individual key to be added to the section. */
 	void GetTransformKeys( const FTransformData& LastTransform, const FTransformData& CurrentTransform, EKey3DTransformChannel::Type ChannelsToKey, bool bUnwindRotation, TArray<FTransformKey>& OutNewKeys, TArray<FTransformKey>& OutDefaultKeys );
diff --git a/Engine/Source/Editor/MovieSceneTools/Public/ClipboardTypes.h b/Engine/Source/Editor/MovieSceneTools/Public/ClipboardTypes.h
--- a/Engine/Source/Editor/MovieSceneTools/Public/ClipboardTypes.h
+++ b/Engine/Source/Editor/MovieSceneTools/Public/ClipboardTypes.h
@@ -8,6 +8,7 @@
 
 #if WITH_EDITOR
 
+#include "MovieSceneEventSection.h"
 #include "MovieSceneClipboard.h"
 
 namespace MovieSceneClipboard
@@ -47,6 +48,10 @@ namespace MovieSceneClipboard
 		static FName Name("Guid");
 		return Name;
 	}
+	template<> inline FName GetKeyTypeName<FEventPayload>()
+	{
+		return "EventPayload";
+	}
 }
 
 #endif
diff --git a/Engine/Source/Editor/MovieSceneTools/Public/FloatCurveKeyArea.h b/Engine/Source/Editor/MovieSceneTools/Public/FloatCurveKeyArea.h
--- a/Engine/Source/Editor/MovieSceneTools/Public/FloatCurveKeyArea.h
+++ b/Engine/Source/Editor/MovieSceneTools/Public/FloatCurveKeyArea.h
@@ -73,6 +73,7 @@ public:
 	virtual UMovieSceneSection* GetOwningSection() override;
 	virtual FRichCurve* GetRichCurve() override;
 	virtual TArray<FKeyHandle> GetUnsortedKeyHandles() const override;
+	virtual FKeyHandle DilateKey(FKeyHandle KeyHandle, float Scale, float Origin) override;
 	virtual FKeyHandle MoveKey(FKeyHandle KeyHandle, float DeltaPosition) override;
 	virtual void SetExtrapolationMode(ERichCurveExtrapolation ExtrapMode, bool bPreInfinity) override;
 	virtual bool CanSetExtrapolationMode() const override;
diff --git a/Engine/Source/Editor/MovieSceneTools/Public/GenericKeyArea.h b/Engine/Source/Editor/MovieSceneTools/Public/GenericKeyArea.h
--- a/Engine/Source/Editor/MovieSceneTools/Public/GenericKeyArea.h
+++ b/Engine/Source/Editor/MovieSceneTools/Public/GenericKeyArea.h
@@ -18,6 +18,7 @@
 #include "Curves/CurveInterface.h"
 #include "Widgets/Text/STextBlock.h"
 #include "Widgets/Input/SSpinBox.h"
+#include "SequencerClipboardReconciler.h"
 #include "GenericKeyArea.generated.h"
 
 class ISequencer;
@@ -137,6 +138,20 @@ public:
 		ExtendSectionBounds(NewKeyTime);
 	}
 
+	virtual FKeyHandle DilateKey(FKeyHandle KeyHandle, TimeType Scale, TimeType Origin) override
+	{
+		TOptional<TimeType> Time = CurveInterface.GetKeyTime(KeyHandle);
+		if (Time.IsSet())
+		{
+			ModifySection();
+			
+			float NewKeyTime = Time.GetValue();
+			NewKeyTime = (NewKeyTime - Origin) * Scale + Origin;
+			CurveInterface.SetKeyTime(KeyHandle, NewKeyTime);
+		}
+		return KeyHandle;
+	}
+
 	virtual FKeyHandle MoveKey(FKeyHandle KeyHandle, TimeType DeltaPosition) override
 	{
 		TOptional<TimeType> Time = CurveInterface.GetKeyTime(KeyHandle);
@@ -194,11 +209,59 @@ public:
 	
 	virtual void CopyKeys(FMovieSceneClipboardBuilder& ClipboardBuilder, const TFunctionRef<bool(FKeyHandle, const IKeyArea&)>& KeyMask) const override
 	{
+		const UMovieSceneSection* Section = OwningSection.Get();
+		UMovieSceneTrack* Track = Section ? Section->GetTypedOuter<UMovieSceneTrack>() : nullptr;
+		if (Track)
+		{
+			FMovieSceneClipboardKeyTrack* KeyTrack = nullptr;
+			TKeyIterator<const KeyValueType, TimeType> Iterator = CurveInterface.IterateKeysAndValues();
 
+			for (auto It = begin(Iterator); It != end(Iterator); ++It)
+			{
+				if (KeyMask(It.GetKeyHandle(), *this))
+				{
+					if (!KeyTrack)
+					{
+						KeyTrack = &ClipboardBuilder.FindOrAddKeyTrack<KeyValueType>(GetName(), *Track);
+					}
+
+					KeyTrack->AddKey(It->Time, It->Value);
+				}
+			}
+		}
 	}
+
 	virtual void PasteKeys(const FMovieSceneClipboardKeyTrack& KeyTrack, const FMovieSceneClipboardEnvironment& SrcEnvironment, const FSequencerPasteEnvironment& DstEnvironment) override
 	{
+		float PasteAt = DstEnvironment.CardinalTime;
+
+		KeyTrack.IterateKeys([&](const FMovieSceneClipboardKey& Key){
+			UMovieSceneSection* Section = GetOwningSection();
+			if (!Section)
+			{
+				return true;
+			}
+
+			if (Section->TryModify())
+			{
+				float Time = PasteAt + Key.GetTime();
+				if (Section->GetStartTime() > Time)
+				{
+					Section->SetStartTime(Time);
+				}
+				if (Section->GetEndTime() < Time)
+				{
+					Section->SetEndTime(Time);
+				}
+
+				KeyValueType NewKeyValue = Key.GetValue<KeyValueType>();
+
+				FKeyHandle KeyHandle = CurveInterface.UpdateOrAddKey(Time, NewKeyValue, KINDA_SMALL_NUMBER);
+				DstEnvironment.ReportPastedKey(KeyHandle, *this);
+			}
 
+			return true;
+		});
 	}
 
 private:
diff --git a/Engine/Source/Editor/MovieSceneTools/Public/IntegralKeyArea.h b/Engine/Source/Editor/MovieSceneTools/Public/IntegralKeyArea.h
--- a/Engine/Source/Editor/MovieSceneTools/Public/IntegralKeyArea.h
+++ b/Engine/Source/Editor/MovieSceneTools/Public/IntegralKeyArea.h
@@ -43,6 +43,7 @@ public:
 	virtual UMovieSceneSection* GetOwningSection() override;
 	virtual FRichCurve* GetRichCurve() override;
 	virtual TArray<FKeyHandle> GetUnsortedKeyHandles() const override;
+	virtual FKeyHandle DilateKey(FKeyHandle KeyHandle, float Scale, float Origin) override;
 	virtual FKeyHandle MoveKey(FKeyHandle KeyHandle, float DeltaPosition) override;
 	virtual void SetExtrapolationMode(ERichCurveExtrapolation ExtrapMode, bool bPreInfinity) override;
 	virtual void SetKeyInterpMode(FKeyHandle KeyHandle, ERichCurveInterpMode InterpMode) override;
diff --git a/Engine/Source/Editor/MovieSceneTools/Public/MatineeImportTools.h b/Engine/Source/Editor/MovieSceneTools/Public/MatineeImportTools.h
--- a/Engine/Source/Editor/MovieSceneTools/Public/MatineeImportTools.h
+++ b/Engine/Source/Editor/MovieSceneTools/Public/MatineeImportTools.h
@@ -17,6 +17,7 @@ class UInterpTrackFloatBase;
 class UInterpTrackLinearColorProp;
 class UInterpTrackMove;
 class UInterpTrackSound;
+class UInterpTrackVectorProp;
 class UInterpTrackVisibility;
 class UMovieScene3DTransformTrack;
 class UMovieSceneAudioTrack;
@@ -28,6 +29,7 @@ class UMovieSceneFadeTrack;
 class UMovieSceneFloatTrack;
 class UMovieSceneParticleTrack;
 class UMovieSceneSkeletalAnimationTrack;
+class UMovieSceneVectorTrack;
 class UMovieSceneVisibilityTrack;
 
 class MOVIESCENETOOLS_API FMatineeImportTools
@@ -52,8 +54,11 @@ public:
 	/** Copies keys from a matinee float track to a sequencer float track. */
 	static bool CopyInterpFloatTrack( UInterpTrackFloatBase* MatineeFloatTrack, UMovieSceneFloatTrack* FloatTrack );
 
+	/** Copies keys from a matinee vector track to a sequencer vector track. */
+	static bool CopyInterpVectorTrack( UInterpTrackVectorProp* MatineeVectorTrack, UMovieSceneVectorTrack* VectorTrack );
+
 	/** Copies keys from a matinee move track to a sequencer transform track. */
-	static bool CopyInterpMoveTrack( UInterpTrackMove* MoveTrack, UMovieScene3DTransformTrack* TransformTrack );
+	static bool CopyInterpMoveTrack( UInterpTrackMove* MoveTrack, UMovieScene3DTransformTrack* TransformTrack, const FVector& DefaultScale = FVector(1.f) );
 
 	/** Copies keys from a matinee color track to a sequencer color track. */
 	static bool CopyInterpColorTrack( UInterpTrackColorProp* ColorPropTrack, UMovieSceneColorTrack* ColorTrack );
diff --git a/Engine/Source/Editor/NiagaraEditor/Private/NiagaraSequencer.h b/Engine/Source/Editor/NiagaraEditor/Private/NiagaraSequencer.h
--- a/Engine/Source/Editor/NiagaraEditor/Private/NiagaraSequencer.h
+++ b/Engine/Source/Editor/NiagaraEditor/Private/NiagaraSequencer.h
@@ -215,7 +215,7 @@ private:
 
 	TSharedPtr<FNiagaraSimulation> Emitter;
 
-	UPROPERTY()
+	UPROPERTY(Instanced)
 	TArray<UMovieSceneSection*> Sections;
 };
 
diff --git a/Engine/Source/Editor/PropertyEditor/Private/PropertyHandleImpl.cpp b/Engine/Source/Editor/PropertyEditor/Private/PropertyHandleImpl.cpp
--- a/Engine/Source/Editor/PropertyEditor/Private/PropertyHandleImpl.cpp
+++ b/Engine/Source/Editor/PropertyEditor/Private/PropertyHandleImpl.cpp
@@ -17,6 +17,7 @@
 #include "ScopedTransaction.h"
 #include "Kismet2/KismetEditorUtilities.h"
 #include "Engine/Selection.h"
+#include "ItemPropertyNode.h"
 
 #include "Framework/Notifications/NotificationManager.h"
 #include "Widgets/Notifications/SNotificationList.h"
@@ -2573,6 +2574,72 @@ void FPropertyHandleBase::SetIgnoreValidation(bool bInIgnore)
 	}
 }
 
+TArray<TSharedPtr<IPropertyHandle>> FPropertyHandleBase::AddChildStructure( TSharedRef<FStructOnScope> InStruct )
+{
+	TArray<TSharedPtr<IPropertyHandle>> PropertyHandles;
+
+	TSharedPtr<FPropertyNode> PropertyNode = Implementation->GetPropertyNode();
+	if( !PropertyNode.IsValid() )
+	{
+		return PropertyHandles;
+	}
+
+	TSharedPtr<FStructurePropertyNode> StructPropertyNode( new FStructurePropertyNode );
+	StructPropertyNode->SetStructure(InStruct);
+
+	FPropertyNodeInitParams RootInitParams;
+	RootInitParams.ParentNode = PropertyNode;
+	RootInitParams.Property = nullptr;
+	RootInitParams.ArrayOffset = 0;
+	RootInitParams.ArrayIndex = INDEX_NONE;
+	RootInitParams.bAllowChildren = true;
+	RootInitParams.bForceHiddenPropertyVisibility = FPropertySettings::Get().ShowHiddenProperties();
+	RootInitParams.bCreateCategoryNodes = false;
+
+	StructPropertyNode->InitNode(RootInitParams);
+
+	const bool bShouldShowHiddenProperties = !!PropertyNode->HasNodeFlags(EPropertyNodeFlags::ShouldShowHiddenProperties);
+	const bool bShouldShowDisableEditOnInstance = !!PropertyNode->HasNodeFlags(EPropertyNodeFlags::ShouldShowDisableEditOnInstance);
+
+	for (TFieldIterator<UProperty> It(InStruct->GetStruct()); It; ++It)
+	{
+		UProperty* StructMember = *It;
+		if (!StructMember)
+		{
+			continue;
+		}
+
+		static const FName Name_InlineEditConditionToggle("InlineEditConditionToggle");
+		const bool bOnlyShowAsInlineEditCondition = StructMember->HasMetaData(Name_InlineEditConditionToggle);
+		const bool bShowIfEditableProperty = StructMember->HasAnyPropertyFlags(CPF_Edit);
+		const bool bShowIfDisableEditOnInstance = !StructMember->HasAnyPropertyFlags(CPF_DisableEditOnInstance) || bShouldShowDisableEditOnInstance;
+
+		if (bShouldShowHiddenProperties || (bShowIfEditableProperty && !bOnlyShowAsInlineEditCondition && bShowIfDisableEditOnInstance))
+		{
+			TSharedRef<FItemPropertyNode> NewItemNode(new FItemPropertyNode);
+
+			FPropertyNodeInitParams InitParams;
+			InitParams.ParentNode = StructPropertyNode;
+			InitParams.Property = StructMember;
+			InitParams.ArrayOffset = 0;
+			InitParams.ArrayIndex = INDEX_NONE;
+			InitParams.bAllowChildren = true;
+			InitParams.bForceHiddenPropertyVisibility = bShouldShowHiddenProperties;
+			InitParams.bCreateDisableEditOnInstanceNodes = bShouldShowDisableEditOnInstance;
+			InitParams.bCreateCategoryNodes = false;
+
+			NewItemNode->InitNode(InitParams);
+			StructPropertyNode->AddChildNode(NewItemNode);
+
+			PropertyHandles.Add(PropertyEditorHelpers::GetPropertyHandle(NewItemNode, Implementation->GetNotifyHook(), Implementation->GetPropertyUtilities()));
+		}
+	}
+
+	PropertyNode->AddChildNode(StructPropertyNode);
+
+	return PropertyHandles;
+}
+
 /** Implements common property value functions */
 #define IMPLEMENT_PROPERTY_VALUE( ClassName ) \
 	ClassName::ClassName( TSharedRef<FPropertyNode> PropertyNode, FNotifyHook* NotifyHook, TSharedPtr<IPropertyUtilities> PropertyUtilities ) \
diff --git a/Engine/Source/Editor/PropertyEditor/Private/PropertyHandleImpl.h b/Engine/Source/Editor/PropertyEditor/Private/PropertyHandleImpl.h
--- a/Engine/Source/Editor/PropertyEditor/Private/PropertyHandleImpl.h
+++ b/Engine/Source/Editor/PropertyEditor/Private/PropertyHandleImpl.h
@@ -478,6 +478,7 @@ public:
 	virtual bool IsRestricted(const FString& Value, TArray<FText>& OutReasons) const override;
 	virtual bool GenerateRestrictionToolTip(const FString& Value, FText& OutTooltip) const override;
 	virtual void SetIgnoreValidation(bool bInIgnore) override;
+	virtual TArray<TSharedPtr<IPropertyHandle>> AddChildStructure( TSharedRef<FStructOnScope> ChildStructure ) override;
 
 	TSharedPtr<FPropertyNode> GetPropertyNode() const;
 protected:
diff --git a/Engine/Source/Editor/PropertyEditor/Private/PropertyNode.cpp b/Engine/Source/Editor/PropertyEditor/Private/PropertyNode.cpp
--- a/Engine/Source/Editor/PropertyEditor/Private/PropertyNode.cpp
+++ b/Engine/Source/Editor/PropertyEditor/Private/PropertyNode.cpp
@@ -313,13 +313,39 @@ const FComplexPropertyNode* FPropertyNode::FindComplexParent() const
 class FObjectPropertyNode* FPropertyNode::FindObjectItemParent()
 {
 	auto ComplexParent = FindComplexParent();
-	return ComplexParent ? ComplexParent->AsObjectNode() : NULL;
+	if (!ComplexParent)
+	{
+		return nullptr;
+	}
+
+	if (FObjectPropertyNode* ObjectNode = ComplexParent->AsObjectNode())
+	{
+		return ObjectNode;
+	}
+	else if (FPropertyNode* ParentNodePtr = ComplexParent->GetParentNode())
+	{
+		return ParentNodePtr->FindObjectItemParent();
+	}
+	return nullptr;
 }
 
 const class FObjectPropertyNode* FPropertyNode::FindObjectItemParent() const
 {
 	const auto ComplexParent = FindComplexParent();
-	return ComplexParent ? ComplexParent->AsObjectNode() : NULL;
+	if (!ComplexParent)
+	{
+		return nullptr;
+	}
+
+	if (const FObjectPropertyNode* ObjectNode = ComplexParent->AsObjectNode())
+	{
+		return ObjectNode;
+	}
+	else if (const FPropertyNode* ParentNodePtr = ComplexParent->GetParentNode())
+	{
+		return ParentNodePtr->FindObjectItemParent();
+	}
+	return nullptr;
 }
 
 /**
diff --git a/Engine/Source/Editor/PropertyEditor/Private/SDetailsViewBase.cpp b/Engine/Source/Editor/PropertyEditor/Private/SDetailsViewBase.cpp
--- a/Engine/Source/Editor/PropertyEditor/Private/SDetailsViewBase.cpp
+++ b/Engine/Source/Editor/PropertyEditor/Private/SDetailsViewBase.cpp
@@ -833,6 +833,17 @@ TSharedPtr<class FUICommandList> SDetailsViewBase::GetHostCommandList() const
 	return DetailsViewArgs.HostCommandList;
 }
 
+TSharedPtr<FTabManager> SDetailsViewBase::GetHostTabManager() const
+{
+	return DetailsViewArgs.HostTabManager;
+}
+
+void SDetailsViewBase::SetHostTabManager(TSharedPtr<FTabManager> InTabManager)
+{
+	DetailsViewArgs.HostTabManager = InTabManager;
+}
+
+
 /** 
  * Hides or shows properties based on the passed in filter text
  * 
diff --git a/Engine/Source/Editor/PropertyEditor/Private/SDetailsViewBase.h b/Engine/Source/Editor/PropertyEditor/Private/SDetailsViewBase.h
--- a/Engine/Source/Editor/PropertyEditor/Private/SDetailsViewBase.h
+++ b/Engine/Source/Editor/PropertyEditor/Private/SDetailsViewBase.h
@@ -355,6 +355,8 @@ protected:
 	virtual TSharedPtr<SWidget> GetNameAreaWidget() override;
 	virtual TSharedPtr<SWidget> GetFilterAreaWidget() override;
 	virtual TSharedPtr<class FUICommandList> GetHostCommandList() const override;
+	virtual TSharedPtr<FTabManager> GetHostTabManager() const override;
+	virtual void SetHostTabManager(TSharedPtr<FTabManager> InTabManager) override;
 
 	/** 
 	 * Hides or shows properties based on the passed in filter text
diff --git a/Engine/Source/Editor/PropertyEditor/Private/SStructureDetailsView.cpp b/Engine/Source/Editor/PropertyEditor/Private/SStructureDetailsView.cpp
--- a/Engine/Source/Editor/PropertyEditor/Private/SStructureDetailsView.cpp
+++ b/Engine/Source/Editor/PropertyEditor/Private/SStructureDetailsView.cpp
@@ -226,6 +226,11 @@ void SStructureDetailsView::ForceRefresh()
 	SetStructureData(StructData);
 }
 
+void SStructureDetailsView::AddExternalRootPropertyNode(TSharedRef<FPropertyNode> ExternalRootNode)
+{
+	ExternalRootPropertyNodes.Add( ExternalRootNode );
+}
+
 void SStructureDetailsView::ClearSearch()
 {
 	CurrentFilter.FilterStrings.Empty();
diff --git a/Engine/Source/Editor/PropertyEditor/Private/SStructureDetailsView.h b/Engine/Source/Editor/PropertyEditor/Private/SStructureDetailsView.h
--- a/Engine/Source/Editor/PropertyEditor/Private/SStructureDetailsView.h
+++ b/Engine/Source/Editor/PropertyEditor/Private/SStructureDetailsView.h
@@ -94,7 +94,7 @@ public:
 
 	virtual void ForceRefresh() override;
 	virtual void MoveScrollOffset(int32 DeltaOffset) override {}
-	virtual void AddExternalRootPropertyNode(TSharedRef<FPropertyNode> ExternalRootNode) override {}
+	virtual void AddExternalRootPropertyNode(TSharedRef<FPropertyNode> ExternalRootNode) override;
 	virtual void ClearSearch() override;
 public:
 
diff --git a/Engine/Source/Editor/PropertyEditor/Private/UserInterface/PropertyEditor/SPropertyEditorAsset.cpp b/Engine/Source/Editor/PropertyEditor/Private/UserInterface/PropertyEditor/SPropertyEditorAsset.cpp
--- a/Engine/Source/Editor/PropertyEditor/Private/UserInterface/PropertyEditor/SPropertyEditorAsset.cpp
+++ b/Engine/Source/Editor/PropertyEditor/Private/UserInterface/PropertyEditor/SPropertyEditorAsset.cpp
@@ -246,7 +246,7 @@ void SPropertyEditorAsset::Construct( const FArguments& InArgs, const TSharedPtr
 			// Go through all the found objects and see if any are a CDO, we can't set an actor in a CDO default.
 			for (UObject* Obj : ObjectList)
 			{
-				if (Obj->HasAllFlags(RF_ClassDefaultObject))
+				if (Obj->IsTemplate())
 				{
 					IsEnabledAttribute.Set(false);
 					TooltipAttribute.Set(LOCTEXT("VariableHasDisableEditOnTemplateTooltip", "Editing this value in a Class Default Object is not allowed"));
diff --git a/Engine/Source/Editor/PropertyEditor/Public/IDetailsView.h b/Engine/Source/Editor/PropertyEditor/Public/IDetailsView.h
--- a/Engine/Source/Editor/PropertyEditor/Public/IDetailsView.h
+++ b/Engine/Source/Editor/PropertyEditor/Public/IDetailsView.h
@@ -7,13 +7,16 @@
 #include "PropertyEditorDelegates.h"
 #include "Framework/Commands/UICommandList.h"
 
+
 class AActor;
 class FNotifyHook;
 class FPropertyPath;
+class FTabManager;
 class IDetailKeyframeHandler;
 class IDetailPropertyExtensionHandler;
 class IDetailRootObjectCustomization;
 
+
 /**
  * Init params for a details view widget
  */
@@ -45,6 +48,8 @@ struct FDetailsViewArgs
 	EEditDefaultsOnlyNodeVisibility DefaultsOnlyVisibility;
 	/** The command list from the host of the details view, allowing child widgets to bind actions with a bound chord */
 	TSharedPtr<class FUICommandList> HostCommandList;
+	/** The tab manager from the host of the details view, allowing child widgets to spawn tabs */
+	TSharedPtr<FTabManager> HostTabManager;
 
 	/** Identifier for this details view; NAME_None if this view is anonymous */
 	FName ViewIdentifier;
@@ -305,4 +310,10 @@ public:
 
 	/** Returns the command list of the hosting toolkit (can be nullptr if the widget that contains the details panel didn't route a command list in) */
 	virtual TSharedPtr<class FUICommandList> GetHostCommandList() const = 0;
+
+	/** Returns the tab manager of the hosting toolkit (can be nullptr if the details panel is not hosted within a tab) */
+	virtual TSharedPtr<FTabManager> GetHostTabManager() const = 0;
+
+	/** Sets the tab manager of the hosting toolkit (can be nullptr if the details panel is not hosted within a tab) */
+	virtual void SetHostTabManager(TSharedPtr<FTabManager> InTabManager) = 0;
 };
diff --git a/Engine/Source/Editor/PropertyEditor/Public/PropertyHandle.h b/Engine/Source/Editor/PropertyEditor/Public/PropertyHandle.h
--- a/Engine/Source/Editor/PropertyEditor/Public/PropertyHandle.h
+++ b/Engine/Source/Editor/PropertyEditor/Public/PropertyHandle.h
@@ -528,6 +528,14 @@ public:
 	  * This should only ever be set for extremely large arrays or other costly validation checks where validation is handled by the customizer
 	  */
 	virtual void SetIgnoreValidation(bool bInIgnore) = 0;
+	
+	/**
+	 * Adds a child structure
+	 * 
+	 * @param ChildStructure	The structure to add
+	 * @return An array of interfaces to the properties that were added
+	 */
+	virtual TArray<TSharedPtr<IPropertyHandle>> AddChildStructure( TSharedRef<FStructOnScope> ChildStructure ) = 0;
 };
 
 /**
diff --git a/Engine/Source/Editor/SequenceRecorder/Private/ActorRecording.cpp b/Engine/Source/Editor/SequenceRecorder/Private/ActorRecording.cpp
--- a/Engine/Source/Editor/SequenceRecorder/Private/ActorRecording.cpp
+++ b/Engine/Source/Editor/SequenceRecorder/Private/ActorRecording.cpp
@@ -675,12 +675,20 @@ void UActorRecording::StartRecordingNewComponents(ULevelSequence* CurrentSequenc
 
 			for (USceneComponent* SceneComponent : NewComponents)
 			{
-				// new component, so we need to add this to our BP if it didn't come from SCS
+				// new component, so we need to add this to our BP if its not native
 				FName NewName;
-				if (SceneComponent->CreationMethod != EComponentCreationMethod::SimpleConstructionScript)
+				if (SceneComponent->CreationMethod != EComponentCreationMethod::Native)
 				{
 					// Give this component a unique name within its parent
-					NewName = *FString::Printf(TEXT("Dynamic%s"), *SceneComponent->GetFName().GetPlainNameString());
+					if (SceneComponent->CreationMethod != EComponentCreationMethod::SimpleConstructionScript)
+					{
+						NewName = *FString::Printf(TEXT("Dynamic%s"), *SceneComponent->GetFName().GetPlainNameString());
+					}
+					else
+					{
+						NewName = SceneComponent->GetFName();
+					}					
+					
 					NewName.SetNumber(1);
 					while (FindObjectFast<UObject>(ObjectTemplate, NewName))
 					{
diff --git a/Engine/Source/Editor/SequenceRecorder/Private/Sections/MovieSceneAnimationSectionRecorder.cpp b/Engine/Source/Editor/SequenceRecorder/Private/Sections/MovieSceneAnimationSectionRecorder.cpp
--- a/Engine/Source/Editor/SequenceRecorder/Private/Sections/MovieSceneAnimationSectionRecorder.cpp
+++ b/Engine/Source/Editor/SequenceRecorder/Private/Sections/MovieSceneAnimationSectionRecorder.cpp
@@ -117,14 +117,7 @@ void FMovieSceneAnimationSectionRecorder::FinalizeSection()
 		if (AnimationSettings.bRemoveRootAnimation)
 		{
 			// enable root motion on the animation
-			AnimSequence->bEnableRootMotion = true;
-			AnimSequence->RootMotionRootLock = ERootMotionRootLock::Zero;
-		}
-		else
-		{
-			// enable root motion on the animation
-			AnimSequence->bEnableRootMotion = false;
-			AnimSequence->RootMotionRootLock = ERootMotionRootLock::RefPose;
+			AnimSequence->bForceRootLock = true;
 		}
 	}
 
diff --git a/Engine/Source/Editor/SequenceRecorder/Private/SequenceRecorderSettings.cpp b/Engine/Source/Editor/SequenceRecorder/Private/SequenceRecorderSettings.cpp
--- a/Engine/Source/Editor/SequenceRecorder/Private/SequenceRecorderSettings.cpp
+++ b/Engine/Source/Editor/SequenceRecorder/Private/SequenceRecorderSettings.cpp
@@ -25,7 +25,7 @@ USequenceRecorderSettings::USequenceRecorderSettings(const FObjectInitializer& O
 	bRecordNearbySpawnedActors = true;
 	NearbyActorRecordingProximity = 5000.0f;
 	bRecordWorldSettingsActor = true;
-	bReduceKeys = false;
+	bReduceKeys = true;
 
 	ClassesAndPropertiesToRecord.Add(FPropertiesToRecordForClass(USkeletalMeshComponent::StaticClass()));
 	ClassesAndPropertiesToRecord.Add(FPropertiesToRecordForClass(UStaticMeshComponent::StaticClass()));
diff --git a/Engine/Source/Editor/Sequencer/Private/DisplayNodes/SequencerDisplayNode.cpp b/Engine/Source/Editor/Sequencer/Private/DisplayNodes/SequencerDisplayNode.cpp
--- a/Engine/Source/Editor/Sequencer/Private/DisplayNodes/SequencerDisplayNode.cpp
+++ b/Engine/Source/Editor/Sequencer/Private/DisplayNodes/SequencerDisplayNode.cpp
@@ -464,6 +464,9 @@ void FSequencerDisplayNode::BuildContextMenu(FMenuBuilder& MenuBuilder)
 {
 	TSharedRef<FSequencerDisplayNode> ThisNode = SharedThis(this);
 
+	bool bIsReadOnly = !GetSequencer().IsReadOnly();
+	FCanExecuteAction CanExecute = FCanExecuteAction::CreateLambda([bIsReadOnly]{ return bIsReadOnly; });
+
 	MenuBuilder.BeginSection("Edit", LOCTEXT("EditContextMenuSectionName", "Edit"));
 	{
 		MenuBuilder.AddMenuEntry(
@@ -472,7 +475,7 @@ void FSequencerDisplayNode::BuildContextMenu(FMenuBuilder& MenuBuilder)
 			FSlateIcon(),
 			FUIAction(
 				FExecuteAction::CreateSP(&GetSequencer(), &FSequencer::ToggleNodeActive),
-				FCanExecuteAction(),
+				CanExecute,
 				FIsActionChecked::CreateSP(&GetSequencer(), &FSequencer::IsNodeActive)
 			),
 			NAME_None,
@@ -485,7 +488,7 @@ void FSequencerDisplayNode::BuildContextMenu(FMenuBuilder& MenuBuilder)
 			FSlateIcon(),
 			FUIAction(
 				FExecuteAction::CreateSP(&GetSequencer(), &FSequencer::ToggleNodeLocked),
-				FCanExecuteAction(),
+				CanExecute,
 				FIsActionChecked::CreateSP(&GetSequencer(), &FSequencer::IsNodeLocked)
 			),
 			NAME_None,
@@ -504,7 +507,7 @@ void FSequencerDisplayNode::BuildContextMenu(FMenuBuilder& MenuBuilder)
 			LOCTEXT("DeleteNode", "Delete"),
 			LOCTEXT("DeleteNodeTooltip", "Delete this or selected tracks"),
 			FSlateIcon(),
-			FUIAction(FExecuteAction::CreateSP(&GetSequencer(), &FSequencer::DeleteNode, ThisNode))
+			FUIAction(FExecuteAction::CreateSP(&GetSequencer(), &FSequencer::DeleteNode, ThisNode), CanExecute)
 		);
 
 		MenuBuilder.AddMenuEntry(
diff --git a/Engine/Source/Editor/Sequencer/Private/DisplayNodes/SequencerObjectBindingNode.cpp b/Engine/Source/Editor/Sequencer/Private/DisplayNodes/SequencerObjectBindingNode.cpp
--- a/Engine/Source/Editor/Sequencer/Private/DisplayNodes/SequencerObjectBindingNode.cpp
+++ b/Engine/Source/Editor/Sequencer/Private/DisplayNodes/SequencerObjectBindingNode.cpp
@@ -169,17 +169,17 @@ void FSequencerObjectBindingNode::BuildContextMenu(FMenuBuilder& MenuBuilder)
 			FUIAction(
 				FExecuteAction::CreateLambda([=]{ GetSequencer().ExportFBX(); })
 			));
+	}
 
-		MenuBuilder.BeginSection("Organize", LOCTEXT("OrganizeContextMenuSectionName", "Organize"));
-		{
-			MenuBuilder.AddSubMenu(
-				LOCTEXT("LabelsSubMenuText", "Labels"),
-				LOCTEXT("LabelsSubMenuTip", "Add or remove labels on this track"),
-				FNewMenuDelegate::CreateSP(this, &FSequencerObjectBindingNode::HandleLabelsSubMenuCreate)
-			);
-		}
-		MenuBuilder.EndSection();
+	MenuBuilder.BeginSection("Organize", LOCTEXT("OrganizeContextMenuSectionName", "Organize"));
+	{
+		MenuBuilder.AddSubMenu(
+			LOCTEXT("LabelsSubMenuText", "Labels"),
+			LOCTEXT("LabelsSubMenuTip", "Add or remove labels on this track"),
+			FNewMenuDelegate::CreateSP(this, &FSequencerObjectBindingNode::HandleLabelsSubMenuCreate)
+		);
 	}
+	MenuBuilder.EndSection();
 
 	FSequencerDisplayNode::BuildContextMenu(MenuBuilder);
 }
@@ -260,6 +260,11 @@ bool FSequencerObjectBindingNode::CanRenameNode() const
 
 TSharedRef<SWidget> FSequencerObjectBindingNode::GetCustomOutlinerContent()
 {
+	if (GetSequencer().IsReadOnly())
+	{
+		return SNullWidget::NullWidget;
+	}
+	
 	// Create a container edit box
 	TSharedRef<SHorizontalBox> BoxPanel = SNew(SHorizontalBox)
 		+ SHorizontalBox::Slot()
@@ -440,7 +445,6 @@ const UClass* FSequencerObjectBindingNode::GetClassForObjectBinding() const
 	
 	// should exist, but also shouldn't be both a spawnable and a possessable
 	check((Spawnable != nullptr) ^ (Possessable != nullptr));
-	check((nullptr == Spawnable) || (nullptr != Spawnable->GetObjectTemplate())  );
 	const UClass* ObjectClass = Spawnable ? Spawnable->GetObjectTemplate()->GetClass() : Possessable->GetPossessedObjectClass();
 
 	return ObjectClass;
diff --git a/Engine/Source/Editor/Sequencer/Private/DisplayNodes/SequencerSectionKeyAreaNode.cpp b/Engine/Source/Editor/Sequencer/Private/DisplayNodes/SequencerSectionKeyAreaNode.cpp
--- a/Engine/Source/Editor/Sequencer/Private/DisplayNodes/SequencerSectionKeyAreaNode.cpp
+++ b/Engine/Source/Editor/Sequencer/Private/DisplayNodes/SequencerSectionKeyAreaNode.cpp
@@ -46,6 +46,7 @@ TSharedRef<SWidget> FSequencerSectionKeyAreaNode::GetCustomOutlinerContent()
 				.VAlign(VAlign_Center)
 				[
 					SNew(SBox)
+					.IsEnabled(!GetSequencer().IsReadOnly())
 					.WidthOverride(100)
 					.HAlign(HAlign_Left)
 					[
