Copying //UE4/Dev-Sequencer to //UE4/Dev-Main (Source: //UE4/Dev-Sequencer @ 3800792)

diff --git a/Engine/Plugins/Messaging/UdpMessaging/Source/UdpMessaging/Private/Transport/UdpMessageProcessor.cpp b/Engine/Plugins/Messaging/UdpMessaging/Source/UdpMessaging/Private/Transport/UdpMessageProcessor.cpp
--- a/Engine/Plugins/Messaging/UdpMessaging/Source/UdpMessaging/Private/Transport/UdpMessageProcessor.cpp
+++ b/Engine/Plugins/Messaging/UdpMessaging/Source/UdpMessaging/Private/Transport/UdpMessageProcessor.cpp
@@ -37,7 +37,11 @@ FUdpMessageProcessor::FUdpMessageProcessor(FSocket& InSocket, const FGuid& InNod
 	, SocketSender(nullptr)
 	, Stopping(false)
 {
-	WorkEvent = FPlatformProcess::GetSynchEventFromPool();
+	WorkEvent = MakeShareable(FPlatformProcess::GetSynchEventFromPool(), [](FEvent* EventToDelete)
+	{
+		FPlatformProcess::ReturnSynchEventToPool(EventToDelete);
+	});
+
 	Thread = FRunnableThread::Create(this, TEXT("FUdpMessageProcessor"), 128 * 1024, TPri_AboveNormal, FPlatformAffinity::GetPoolThreadMask());
 
 	const UUdpMessagingSettings& Settings = *GetDefault<UUdpMessagingSettings>();
@@ -76,10 +80,6 @@ FUdpMessageProcessor::~FUdpMessageProcessor()
 	}
 
 	KnownNodes.Empty();
-
-	// clean up 
-	FPlatformProcess::ReturnSynchEventToPool(WorkEvent);
-	WorkEvent = nullptr;
 }
 
 
@@ -101,14 +101,7 @@ bool FUdpMessageProcessor::EnqueueInboundSegment(const TSharedPtr<FArrayReader,
 
 bool FUdpMessageProcessor::EnqueueOutboundMessage(const TSharedRef<FUdpSerializedMessage, ESPMode::ThreadSafe>& SerializedMessage, const FGuid& Recipient)
 {
-	if (!OutboundMessages.Enqueue(FOutboundMessage(SerializedMessage, Recipient)))
-	{
-		return false;
-	}
-
-	SerializedMessage->OnStateChanged().BindRaw(this, &FUdpMessageProcessor::HandleSerializedMessageStateChanged);
-
-	return true;
+	return OutboundMessages.Enqueue(FOutboundMessage(SerializedMessage, Recipient));
 }
 
 
@@ -538,12 +531,3 @@ void FUdpMessageProcessor::UpdateStaticNodes()
 		UpdateSegmenters(StaticNodePair.Value);
 	}
 }
-
-
-/* FUdpMessageProcessor callbacks
- *****************************************************************************/
-
-void FUdpMessageProcessor::HandleSerializedMessageStateChanged()
-{
-	WorkEvent->Trigger();
-}
diff --git a/Engine/Plugins/Messaging/UdpMessaging/Source/UdpMessaging/Private/Transport/UdpMessageProcessor.h b/Engine/Plugins/Messaging/UdpMessaging/Source/UdpMessaging/Private/Transport/UdpMessageProcessor.h
--- a/Engine/Plugins/Messaging/UdpMessaging/Source/UdpMessaging/Private/Transport/UdpMessageProcessor.h
+++ b/Engine/Plugins/Messaging/UdpMessaging/Source/UdpMessaging/Private/Transport/UdpMessageProcessor.h
@@ -149,6 +149,16 @@ public:
 	 */
 	bool EnqueueOutboundMessage(const TSharedRef<FUdpSerializedMessage, ESPMode::ThreadSafe>& SerializedMessage, const FGuid& Recipient);
 
+	/**
+	 * Get the event used to signal the message processor that work is available.
+	 *
+	 * @return The event.
+	 */
+	const TSharedPtr<FEvent, ESPMode::ThreadSafe>& GetWorkEvent() const
+	{
+		return WorkEvent;
+	}
+
 public:
 
 	// @todo gmp: remove the need for this typedef
@@ -317,11 +327,6 @@ protected:
 	/** Updates all static remote nodes. */
 	void UpdateStaticNodes();
 
-private:
-
-	/** Handles message data state changes. */
-	void HandleSerializedMessageStateChanged();
-
 private:
 
 	/** Holds the queue of outbound messages. */
@@ -366,7 +371,7 @@ private:
 	FRunnableThread* Thread;
 
 	/** Holds an event signaling that inbound messages need to be processed. */
-	FEvent* WorkEvent;
+	TSharedPtr<FEvent, ESPMode::ThreadSafe> WorkEvent;
 
 private:
 
diff --git a/Engine/Plugins/Messaging/UdpMessaging/Source/UdpMessaging/Private/Transport/UdpMessageTransport.cpp b/Engine/Plugins/Messaging/UdpMessaging/Source/UdpMessaging/Private/Transport/UdpMessageTransport.cpp
--- a/Engine/Plugins/Messaging/UdpMessaging/Source/UdpMessaging/Private/Transport/UdpMessageTransport.cpp
+++ b/Engine/Plugins/Messaging/UdpMessaging/Source/UdpMessaging/Private/Transport/UdpMessageTransport.cpp
@@ -180,7 +180,7 @@ bool FUdpMessageTransport::TransportMessage(const TSharedRef<IMessageContext, ES
 		return false;
 	}
 
-	TSharedRef<FUdpSerializedMessage, ESPMode::ThreadSafe> SerializedMessage = MakeShareable(new FUdpSerializedMessage());
+	TSharedRef<FUdpSerializedMessage, ESPMode::ThreadSafe> SerializedMessage = MakeShared<FUdpSerializedMessage, ESPMode::ThreadSafe>();
 
 	if (Recipients.Num() == 0)
 	{
@@ -196,7 +196,7 @@ bool FUdpMessageTransport::TransportMessage(const TSharedRef<IMessageContext, ES
 		}
 	}
 
-	TGraphTask<FUdpSerializeMessageTask>::CreateTask().ConstructAndDispatchWhenReady(Context, SerializedMessage);
+	TGraphTask<FUdpSerializeMessageTask>::CreateTask().ConstructAndDispatchWhenReady(Context, SerializedMessage, MessageProcessor->GetWorkEvent());
 
 	return true;
 }
diff --git a/Engine/Plugins/Messaging/UdpMessaging/Source/UdpMessaging/Private/Transport/UdpSerializeMessageTask.cpp b/Engine/Plugins/Messaging/UdpMessaging/Source/UdpMessaging/Private/Transport/UdpSerializeMessageTask.cpp
--- a/Engine/Plugins/Messaging/UdpMessaging/Source/UdpMessaging/Private/Transport/UdpSerializeMessageTask.cpp
+++ b/Engine/Plugins/Messaging/UdpMessaging/Source/UdpMessaging/Private/Transport/UdpSerializeMessageTask.cpp
@@ -3,6 +3,7 @@
 #include "Transport/UdpSerializeMessageTask.h"
 
 #include "Backends/JsonStructSerializerBackend.h"
+#include "HAL/Event.h"
 #include "IMessageContext.h"
 #include "StructSerializer.h"
 
@@ -61,6 +62,14 @@ void FUdpSerializeMessageTask::DoTask(ENamedThreads::Type CurrentThread, const F
 	{
 		SerializedMessage->UpdateState(EUdpSerializedMessageState::Invalid);
 	}
+
+	// signal task completion
+	TSharedPtr<FEvent, ESPMode::ThreadSafe> CompletionEvent = CompletionEventPtr.Pin();
+
+	if (CompletionEvent.IsValid())
+	{
+		CompletionEvent->Trigger();
+	}
 }
 
 
diff --git a/Engine/Plugins/Messaging/UdpMessaging/Source/UdpMessaging/Private/Transport/UdpSerializeMessageTask.h b/Engine/Plugins/Messaging/UdpMessaging/Source/UdpMessaging/Private/Transport/UdpSerializeMessageTask.h
--- a/Engine/Plugins/Messaging/UdpMessaging/Source/UdpMessaging/Private/Transport/UdpSerializeMessageTask.h
+++ b/Engine/Plugins/Messaging/UdpMessaging/Source/UdpMessaging/Private/Transport/UdpSerializeMessageTask.h
@@ -6,6 +6,7 @@
 #include "Stats/Stats.h"
 #include "Templates/SharedPointer.h"
 
+class FEvent;
 class FUdpSerializedMessage;
 class IMessageContext;
 
@@ -23,8 +24,9 @@ public:
 	 * @param InMessageContext The context of the message to serialize.
 	 * @param InSerializedMessage Will hold the serialized message data.
 	 */
-	FUdpSerializeMessageTask(TSharedRef<IMessageContext, ESPMode::ThreadSafe> InMessageContext, TSharedRef<FUdpSerializedMessage, ESPMode::ThreadSafe> InSerializedMessage)
-		: MessageContext(InMessageContext)
+	FUdpSerializeMessageTask(TSharedRef<IMessageContext, ESPMode::ThreadSafe> InMessageContext, TSharedRef<FUdpSerializedMessage, ESPMode::ThreadSafe> InSerializedMessage, TSharedPtr<FEvent, ESPMode::ThreadSafe> InCompletionEvent)
+		: CompletionEventPtr(InCompletionEvent)
+		, MessageContext(InMessageContext)
 		, SerializedMessage(InSerializedMessage)
 	{ }
 
@@ -61,9 +63,12 @@ public:
 
 private:
 
-	/** Holds the context of the message to serialize. */
+	/** An event signaling that the message was serialized. */
+	TWeakPtr<FEvent, ESPMode::ThreadSafe> CompletionEventPtr;
+
+	/** The context of the message to serialize. */
 	TSharedRef<IMessageContext, ESPMode::ThreadSafe> MessageContext;
 
-	/** Holds a reference to the serialized message data. */
+	/** A reference to the serialized message data. */
 	TSharedRef<FUdpSerializedMessage, ESPMode::ThreadSafe> SerializedMessage;
 };
diff --git a/Engine/Plugins/Messaging/UdpMessaging/Source/UdpMessaging/Private/Transport/UdpSerializedMessage.h b/Engine/Plugins/Messaging/UdpMessaging/Source/UdpMessaging/Private/Transport/UdpSerializedMessage.h
--- a/Engine/Plugins/Messaging/UdpMessaging/Source/UdpMessaging/Private/Transport/UdpSerializedMessage.h
+++ b/Engine/Plugins/Messaging/UdpMessaging/Source/UdpMessaging/Private/Transport/UdpSerializedMessage.h
@@ -71,34 +71,22 @@ public:
 	 * Gets the state of the message data.
 	 *
 	 * @return Message data state.
-	 * @see GetData, OnStateChanged, UpdateState
+	 * @see GetData, UpdateState
 	 */
 	EUdpSerializedMessageState GetState() const
 	{
 		return State;
 	}
 
-	/**
-	 * Returns a delegate that is executed when the message data's state changed.
-	 *
-	 * @return The delegate.
-	 * @see GetState, UpdateState
-	 */
-	FSimpleDelegate& OnStateChanged()
-	{
-		return StateChangedDelegate;
-	}
-
 	/**
 	 * Updates the state of this message data.
 	 *
 	 * @param InState The state to set.
-	 * @see GetState, OnStateChanged
+	 * @see GetState
 	 */
 	void UpdateState(EUdpSerializedMessageState InState)
 	{
 		State = InState;
-		StateChangedDelegate.ExecuteIfBound();
 	}
 
 private:
@@ -108,9 +96,4 @@ private:
 
 	/** Holds the message data state. */
 	EUdpSerializedMessageState State;
-
-private:
-
-	/** Holds a delegate that is invoked when the data's state changed. */
-	FSimpleDelegate StateChangedDelegate;
 };
diff --git a/Engine/Plugins/MovieScene/ActorSequence/Source/ActorSequence/Private/ActorSequence.cpp b/Engine/Plugins/MovieScene/ActorSequence/Source/ActorSequence/Private/ActorSequence.cpp
--- a/Engine/Plugins/MovieScene/ActorSequence/Source/ActorSequence/Private/ActorSequence.cpp
+++ b/Engine/Plugins/MovieScene/ActorSequence/Source/ActorSequence/Private/ActorSequence.cpp
@@ -84,6 +84,11 @@ void UActorSequence::BindPossessableObject(const FGuid& ObjectId, UObject& Posse
 
 bool UActorSequence::CanPossessObject(UObject& Object, UObject* InPlaybackContext) const
 {
+	if (InPlaybackContext == nullptr)
+	{
+		return false;
+	}
+
 	AActor* ActorContext = CastChecked<AActor>(InPlaybackContext);
 
 	if (AActor* Actor = Cast<AActor>(&Object))
diff --git a/Engine/Plugins/MovieScene/ActorSequenceEditor/Source/ActorSequenceEditor/Private/ActorSequenceComponentCustomization.cpp b/Engine/Plugins/MovieScene/ActorSequenceEditor/Source/ActorSequenceEditor/Private/ActorSequenceComponentCustomization.cpp
--- a/Engine/Plugins/MovieScene/ActorSequenceEditor/Source/ActorSequenceEditor/Private/ActorSequenceComponentCustomization.cpp
+++ b/Engine/Plugins/MovieScene/ActorSequenceEditor/Source/ActorSequenceEditor/Private/ActorSequenceComponentCustomization.cpp
@@ -133,7 +133,7 @@ void FActorSequenceComponentCustomization::CustomizeDetails(IDetailLayoutBuilder
 	}
 
 	// Only display an inline editor for non-blueprint sequences
-	if (!GetActorSequence()->GetParentBlueprint() && !bIsExternalTabAlreadyOpened)
+	if (GetActorSequence() && !GetActorSequence()->GetParentBlueprint() && !bIsExternalTabAlreadyOpened)
 	{
 		Category.AddCustomRow(FText())
 		.WholeRowContent()
diff --git a/Engine/Plugins/MovieScene/LevelSequenceEditor/Source/LevelSequenceEditor/Private/CinematicViewport/SCinematicLevelViewport.cpp b/Engine/Plugins/MovieScene/LevelSequenceEditor/Source/LevelSequenceEditor/Private/CinematicViewport/SCinematicLevelViewport.cpp
--- a/Engine/Plugins/MovieScene/LevelSequenceEditor/Source/LevelSequenceEditor/Private/CinematicViewport/SCinematicLevelViewport.cpp
+++ b/Engine/Plugins/MovieScene/LevelSequenceEditor/Source/LevelSequenceEditor/Private/CinematicViewport/SCinematicLevelViewport.cpp
@@ -651,7 +651,7 @@ void SCinematicLevelViewport::Tick(const FGeometry& AllottedGeometry, const doub
 		UMovieSceneCinematicShotSection* CinematicShotSection = Cast<UMovieSceneCinematicShotSection>(SubSection);
 		if (CinematicShotSection)
 		{
-			UIData.ShotName = CinematicShotSection->GetShotDisplayName();
+			UIData.ShotName = FText::FromString(CinematicShotSection->GetShotDisplayName());
 		}
 		else if (SubSection->GetSequence() != nullptr)
 		{
diff --git a/Engine/Plugins/MovieScene/LevelSequenceEditor/Source/LevelSequenceEditor/Private/LevelSequenceEditorCommands.cpp b/Engine/Plugins/MovieScene/LevelSequenceEditor/Source/LevelSequenceEditor/Private/LevelSequenceEditorCommands.cpp
--- a/Engine/Plugins/MovieScene/LevelSequenceEditor/Source/LevelSequenceEditor/Private/LevelSequenceEditorCommands.cpp
+++ b/Engine/Plugins/MovieScene/LevelSequenceEditor/Source/LevelSequenceEditor/Private/LevelSequenceEditorCommands.cpp
@@ -13,7 +13,7 @@ void FLevelSequenceEditorCommands::RegisterCommands()
 {
 	UI_COMMAND(CreateNewLevelSequenceInLevel, "Add Level Sequence", "Create a new level sequence asset, and place an instance of it in this level", EUserInterfaceActionType::Button, FInputChord());
 	UI_COMMAND(CreateNewMasterSequenceInLevel, "Add Master Sequence", "Create a new master sequence asset, and place an instance of it in this level", EUserInterfaceActionType::Button, FInputChord());
-	UI_COMMAND(ToggleCinematicViewportCommand, "Cinematic Viewport", "A viewport layout tailored to cinematic preview", EUserInterfaceActionType::RadioButton, FInputChord(EKeys::C, EModifierKey::Shift));
+	UI_COMMAND(ToggleCinematicViewportCommand, "Cinematic Viewport", "A viewport layout tailored to cinematic preview", EUserInterfaceActionType::RadioButton, FInputChord());
 }
 
 #undef LOCTEXT_NAMESPACE
diff --git a/Engine/Plugins/MovieScene/LevelSequenceEditor/Source/LevelSequenceEditor/Private/LevelSequenceEditorToolkit.cpp b/Engine/Plugins/MovieScene/LevelSequenceEditor/Source/LevelSequenceEditor/Private/LevelSequenceEditorToolkit.cpp
--- a/Engine/Plugins/MovieScene/LevelSequenceEditor/Source/LevelSequenceEditor/Private/LevelSequenceEditorToolkit.cpp
+++ b/Engine/Plugins/MovieScene/LevelSequenceEditor/Source/LevelSequenceEditor/Private/LevelSequenceEditorToolkit.cpp
@@ -464,8 +464,6 @@ void FLevelSequenceEditorToolkit::AddDefaultTracksForActor(AActor& Actor, const
 
 					NewSection->SetIsInfinite(GetSequencer()->GetInfiniteKeyAreas());
 				}
-
-				Sequencer->UpdateRuntimeInstances();
 			}
 		}
 
@@ -578,8 +576,6 @@ void FLevelSequenceEditorToolkit::AddDefaultTracksForActor(AActor& Actor, const
 			FKeyPropertyParams KeyPropertyParams(TArrayBuilder<UObject*>().Add(PropertyOwner), *PropertyPath, ESequencerKeyMode::ManualKey);
 
 			Sequencer->KeyProperty(KeyPropertyParams);
-
-			Sequencer->UpdateRuntimeInstances();
 		}
 	}
 }
@@ -645,7 +641,6 @@ void FLevelSequenceEditorToolkit::HandleMapChanged(class UWorld* NewWorld, EMapC
 	if( ( MapChangeType == EMapChangeType::LoadMap || MapChangeType == EMapChangeType::NewMap || MapChangeType == EMapChangeType::TearDownWorld) )
 	{
 		Sequencer->GetSpawnRegister().CleanUp(*Sequencer);
-		Sequencer->UpdateRuntimeInstances();
 	}
 }
 
@@ -657,7 +652,6 @@ void FLevelSequenceEditorToolkit::AddShot(UMovieSceneCinematicShotTrack* ShotTra
 	UMovieSceneSubSection* ShotSubSection = ShotTrack->AddSequence(ShotSequence, ShotStartTime, ShotEndTime-ShotStartTime);
 
 	// Focus on the new shot
-	GetSequencer()->UpdateRuntimeInstances();
 	GetSequencer()->ForceEvaluate();
 	GetSequencer()->FocusSequenceInstance(*ShotSubSection);
 
@@ -727,7 +721,6 @@ void FLevelSequenceEditorToolkit::AddShot(UMovieSceneCinematicShotTrack* ShotTra
 		if (bCreateSpawnableCamera)
 		{
 			CameraGuid = GetSequencer()->MakeNewSpawnable(*NewCamera);
-			GetSequencer()->UpdateRuntimeInstances();
 			UObject* SpawnedCamera = GetSequencer()->FindSpawnedObjectOrTemplate(CameraGuid);
 			if (SpawnedCamera)
 			{
diff --git a/Engine/Plugins/MovieScene/LevelSequenceEditor/Source/LevelSequenceEditor/Private/Misc/LevelSequenceEditorActorSpawner.cpp b/Engine/Plugins/MovieScene/LevelSequenceEditor/Source/LevelSequenceEditor/Private/Misc/LevelSequenceEditorActorSpawner.cpp
--- a/Engine/Plugins/MovieScene/LevelSequenceEditor/Source/LevelSequenceEditor/Private/Misc/LevelSequenceEditorActorSpawner.cpp
+++ b/Engine/Plugins/MovieScene/LevelSequenceEditor/Source/LevelSequenceEditor/Private/Misc/LevelSequenceEditorActorSpawner.cpp
@@ -26,26 +26,17 @@ TSharedRef<IMovieSceneObjectSpawner> FLevelSequenceEditorActorSpawner::CreateObj
 
 #if WITH_EDITOR
 
-TValueOrError<FNewSpawnable, FText> FLevelSequenceEditorActorSpawner::CreateNewSpawnableType(UObject& SourceObject, UMovieScene& OwnerMovieScene)
+TValueOrError<FNewSpawnable, FText> FLevelSequenceEditorActorSpawner::CreateNewSpawnableType(UObject& SourceObject, UMovieScene& OwnerMovieScene, UActorFactory* ActorFactory)
 {
 	FNewSpawnable NewSpawnable(nullptr, FName::NameToDisplayString(SourceObject.GetName(), false));
 
 	const FName TemplateName = MakeUniqueObjectName(&OwnerMovieScene, UObject::StaticClass(), SourceObject.GetFName());
 
-	// First off, deal with creating a spawnable from a class
-	if (UClass* InClass = Cast<UClass>(&SourceObject))
-	{
-		if (!InClass->IsChildOf(AActor::StaticClass()))
-		{
-			FText ErrorText = FText::Format(LOCTEXT("NotAnActorClass", "Unable to add spawnable for class of type '{0}' since it is not a valid actor class."), FText::FromString(InClass->GetName()));
-			return MakeError(ErrorText);
-		}
-
-		NewSpawnable.ObjectTemplate = NewObject<UObject>(&OwnerMovieScene, InClass, TemplateName);
-	}
+	FText ErrorText;
 
+	// First off, deal with creating a spawnable from a class
 	// Deal with creating a spawnable from an instance of an actor
-	else if (AActor* Actor = Cast<AActor>(&SourceObject))
+	if (AActor* Actor = Cast<AActor>(&SourceObject))
 	{
 		AActor* SpawnedActor = Cast<AActor>(StaticDuplicateObject(Actor, &OwnerMovieScene, TemplateName, RF_AllFlags & ~RF_Transactional));
 		SpawnedActor->bIsEditorPreviewActor = false;
@@ -62,24 +53,21 @@ TValueOrError<FNewSpawnable, FText> FLevelSequenceEditorActorSpawner::CreateNewS
 	// At this point we have to assume it's an asset
 	else
 	{
-		// @todo sequencer: Add support for forcing specific factories for an asset?
-		UActorFactory* FactoryToUse = FActorFactoryAssetProxy::GetFactoryForAssetObject(&SourceObject);
+		UActorFactory* FactoryToUse = ActorFactory ? ActorFactory : FActorFactoryAssetProxy::GetFactoryForAssetObject(&SourceObject);
 		if (!FactoryToUse)
 		{
-			FText ErrorText = FText::Format(LOCTEXT("CouldNotFindFactory", "Unable to create spawnable from  asset '{0}' - no valid factory could be found."), FText::FromString(SourceObject.GetName()));
-			return MakeError(ErrorText);
+			ErrorText = FText::Format(LOCTEXT("CouldNotFindFactory", "Unable to create spawnable from  asset '{0}' - no valid factory could be found."), FText::FromString(SourceObject.GetName()));
 		}
 
-		FText ErrorText;
 		if (!FactoryToUse->CanCreateActorFrom(FAssetData(&SourceObject), ErrorText))
 		{
 			if (!ErrorText.IsEmpty())
 			{
-				return MakeError(FText::Format(LOCTEXT("CannotCreateActorFromAsset_Ex", "Unable to create spawnable from  asset '{0}'. {1}."), FText::FromString(SourceObject.GetName()), ErrorText));
+				ErrorText = FText::Format(LOCTEXT("CannotCreateActorFromAsset_Ex", "Unable to create spawnable from  asset '{0}'. {1}."), FText::FromString(SourceObject.GetName()), ErrorText);
 			}
 			else
 			{
-				return MakeError(FText::Format(LOCTEXT("CannotCreateActorFromAsset", "Unable to create spawnable from  asset '{0}'."), FText::FromString(SourceObject.GetName())));
+				ErrorText = FText::Format(LOCTEXT("CannotCreateActorFromAsset", "Unable to create spawnable from  asset '{0}'."), FText::FromString(SourceObject.GetName()));
 			}
 		}
 
@@ -94,8 +82,26 @@ TValueOrError<FNewSpawnable, FText> FLevelSequenceEditorActorSpawner::CreateNewS
 
 	if (!NewSpawnable.ObjectTemplate || !NewSpawnable.ObjectTemplate->IsA<AActor>())
 	{
-		FText ErrorText = FText::Format(LOCTEXT("UnknownClassError", "Unable to create a new spawnable object from {0}."), FText::FromString(SourceObject.GetName()));
-		return MakeError(ErrorText);
+		if (UClass* InClass = Cast<UClass>(&SourceObject))
+		{
+			if (!InClass->IsChildOf(AActor::StaticClass()))
+			{
+				ErrorText = FText::Format(LOCTEXT("NotAnActorClass", "Unable to add spawnable for class of type '{0}' since it is not a valid actor class."), FText::FromString(InClass->GetName()));
+				return MakeError(ErrorText);
+			}
+
+			NewSpawnable.ObjectTemplate = NewObject<UObject>(&OwnerMovieScene, InClass, TemplateName);
+		}
+
+		if (!NewSpawnable.ObjectTemplate || !NewSpawnable.ObjectTemplate->IsA<AActor>())
+		{
+			if (ErrorText.IsEmpty())
+			{
+				ErrorText = FText::Format(LOCTEXT("UnknownClassError", "Unable to create a new spawnable object from {0}."), FText::FromString(SourceObject.GetName()));
+			}
+
+			return MakeError(ErrorText);
+		}
 	}
 
 	return MakeValue(NewSpawnable);
diff --git a/Engine/Plugins/MovieScene/LevelSequenceEditor/Source/LevelSequenceEditor/Private/Misc/LevelSequenceEditorActorSpawner.h b/Engine/Plugins/MovieScene/LevelSequenceEditor/Source/LevelSequenceEditor/Private/Misc/LevelSequenceEditorActorSpawner.h
--- a/Engine/Plugins/MovieScene/LevelSequenceEditor/Source/LevelSequenceEditor/Private/Misc/LevelSequenceEditorActorSpawner.h
+++ b/Engine/Plugins/MovieScene/LevelSequenceEditor/Source/LevelSequenceEditor/Private/Misc/LevelSequenceEditorActorSpawner.h
@@ -13,7 +13,7 @@ public:
 	// IMovieSceneObjectSpawner interface
 	virtual bool IsEditor() const override { return true; }
 #if WITH_EDITOR
-	virtual TValueOrError<FNewSpawnable, FText> CreateNewSpawnableType(UObject& SourceObject, UMovieScene& OwnerMovieScene) override;
+	virtual TValueOrError<FNewSpawnable, FText> CreateNewSpawnableType(UObject& SourceObject, UMovieScene& OwnerMovieScene, UActorFactory* ActorFactory = nullptr) override;
 	virtual void SetupDefaultsForSpawnable(UObject* SpawnedObject, const FGuid& Guid, const FTransformData& TransformData, TSharedRef<ISequencer> Sequencer, USequencerSettings* Settings) override;
 	virtual bool CanSetupDefaultsForSpawnable(UObject* SpawnedObject) const override;
 #endif	
diff --git a/Engine/Plugins/MovieScene/LevelSequenceEditor/Source/LevelSequenceEditor/Private/Misc/LevelSequenceEditorSpawnRegister.cpp b/Engine/Plugins/MovieScene/LevelSequenceEditor/Source/LevelSequenceEditor/Private/Misc/LevelSequenceEditorSpawnRegister.cpp
--- a/Engine/Plugins/MovieScene/LevelSequenceEditor/Source/LevelSequenceEditor/Private/Misc/LevelSequenceEditorSpawnRegister.cpp
+++ b/Engine/Plugins/MovieScene/LevelSequenceEditor/Source/LevelSequenceEditor/Private/Misc/LevelSequenceEditorSpawnRegister.cpp
@@ -92,11 +92,7 @@ void FLevelSequenceEditorSpawnRegister::PreDestroyObject(UObject& Object, const
 	
 	TSharedPtr<ISequencer> Sequencer = WeakSequencer.Pin();
 
-	// We only save default state for the currently focussed movie scene sequence instance
-	if (Sequencer.IsValid() && Sequencer->GetFocusedTemplateID() == TemplateID)
-	{
-		SaveDefaultSpawnableState(BindingId);
-	}
+	SaveDefaultSpawnableState(BindingId, TemplateID);
 
 	// Cache its selection state
 	AActor* Actor = Cast<AActor>(&Object);
@@ -122,8 +118,7 @@ void FLevelSequenceEditorSpawnRegister::PreDestroyObject(UObject& Object, const
 
 void FLevelSequenceEditorSpawnRegister::SaveDefaultSpawnableState(FMovieSceneSpawnable& Spawnable, FMovieSceneSequenceIDRef TemplateID, IMovieScenePlayer& Player)
 {
-	const FMovieSceneEvaluationTemplateInstance* Template = Player.GetEvaluationTemplate().GetInstance(TemplateID);
-	UMovieSceneSequence* Sequence = Template ? Template->Sequence.Get() : nullptr;
+	UMovieSceneSequence* Sequence = Player.GetEvaluationTemplate().GetSequence(TemplateID);
 
 	UObject* Object = FindSpawnedObject(Spawnable.GetGuid(), TemplateID);
 	if (!Object || !Sequence)
@@ -152,7 +147,7 @@ void FLevelSequenceEditorSpawnRegister::SaveDefaultSpawnableState(FMovieSceneSpa
 	Spawnable.CopyObjectTemplate(*Object, *Sequence);
 }
 
-void FLevelSequenceEditorSpawnRegister::SaveDefaultSpawnableState(const FGuid& BindingId)
+void FLevelSequenceEditorSpawnRegister::SaveDefaultSpawnableState(const FGuid& BindingId, FMovieSceneSequenceIDRef TemplateID)
 {
 	TSharedPtr<ISequencer> Sequencer = WeakSequencer.Pin();
 	if (!Sequencer.IsValid())
@@ -160,9 +155,8 @@ void FLevelSequenceEditorSpawnRegister::SaveDefaultSpawnableState(const FGuid& B
 		return;
 	}
 
-	UMovieSceneSequence* Sequence = Sequencer->GetFocusedMovieSceneSequence();
+	UMovieSceneSequence* Sequence = Sequencer->GetEvaluationTemplate().GetSequence(TemplateID);
 	UMovieScene* MovieScene = Sequence ? Sequence->GetMovieScene() : nullptr;
-	FMovieSceneSequenceID TemplateID = Sequencer->GetFocusedTemplateID();
 
 	FMovieSceneSpawnable* Spawnable = MovieScene->FindSpawnable(BindingId);
 
@@ -186,7 +180,7 @@ void FLevelSequenceEditorSpawnRegister::OnPreSaveMovieScene(ISequencer& InSequen
 	for (int32 SpawnableIndex = 0; SpawnableIndex < MovieScene->GetSpawnableCount(); ++SpawnableIndex)
 	{
 		FMovieSceneSpawnable& Spawnable = MovieScene->GetSpawnable(SpawnableIndex);
-		SaveDefaultSpawnableState(Spawnable.GetGuid());
+		SaveDefaultSpawnableState(Spawnable.GetGuid(), InSequencer.GetFocusedTemplateID());
 	}
 }
 
@@ -300,11 +294,11 @@ void FLevelSequenceEditorSpawnRegister::OnObjectsReplaced(const TMap<UObject*, U
 
 #if WITH_EDITOR
 
-TValueOrError<FNewSpawnable, FText> FLevelSequenceEditorSpawnRegister::CreateNewSpawnableType(UObject& SourceObject, UMovieScene& OwnerMovieScene)
+TValueOrError<FNewSpawnable, FText> FLevelSequenceEditorSpawnRegister::CreateNewSpawnableType(UObject& SourceObject, UMovieScene& OwnerMovieScene, UActorFactory* ActorFactory)
 {
 	for (TSharedPtr<IMovieSceneObjectSpawner> MovieSceneObjectSpawner : MovieSceneObjectSpawners)
 	{
-		TValueOrError<FNewSpawnable, FText> Result = MovieSceneObjectSpawner->CreateNewSpawnableType(SourceObject, OwnerMovieScene);
+		TValueOrError<FNewSpawnable, FText> Result = MovieSceneObjectSpawner->CreateNewSpawnableType(SourceObject, OwnerMovieScene, ActorFactory);
 		if (Result.IsValid())
 		{
 			return Result;
diff --git a/Engine/Plugins/MovieScene/LevelSequenceEditor/Source/LevelSequenceEditor/Private/Misc/LevelSequenceEditorSpawnRegister.h b/Engine/Plugins/MovieScene/LevelSequenceEditor/Source/LevelSequenceEditor/Private/Misc/LevelSequenceEditorSpawnRegister.h
--- a/Engine/Plugins/MovieScene/LevelSequenceEditor/Source/LevelSequenceEditor/Private/Misc/LevelSequenceEditorSpawnRegister.h
+++ b/Engine/Plugins/MovieScene/LevelSequenceEditor/Source/LevelSequenceEditor/Private/Misc/LevelSequenceEditorSpawnRegister.h
@@ -39,7 +39,7 @@ public:
 	virtual void PreDestroyObject(UObject& Object, const FGuid& BindingId, FMovieSceneSequenceIDRef TemplateID) override;
 	virtual void SaveDefaultSpawnableState(FMovieSceneSpawnable& Spawnable, FMovieSceneSequenceIDRef TemplateID, IMovieScenePlayer& Player) override;
 #if WITH_EDITOR
-	virtual TValueOrError<FNewSpawnable, FText> CreateNewSpawnableType(UObject& SourceObject, UMovieScene& OwnerMovieScene) override;
+	virtual TValueOrError<FNewSpawnable, FText> CreateNewSpawnableType(UObject& SourceObject, UMovieScene& OwnerMovieScene, UActorFactory* ActorFactory = nullptr) override;
 	virtual void SetupDefaultsForSpawnable(UObject* SpawnedObject, const FGuid& Guid, const FTransformData& TransformData, TSharedRef<ISequencer> Sequencer, USequencerSettings* Settings) override;
 	virtual void HandleConvertPossessableToSpawnable(UObject* OldObject, IMovieScenePlayer& Player, FTransformData& OutTransformData) override;
 	virtual bool CanConvertSpawnableToPossessable(FMovieSceneSpawnable& Spawnable) const override;
@@ -57,7 +57,7 @@ private:
 	void OnSequenceInstanceActivated(FMovieSceneSequenceIDRef ActiveInstance);
 
 	/** Saves the default state for the specified spawnable, if an instance for it currently exists */
-	void SaveDefaultSpawnableState(const FGuid& BindingId);
+	void SaveDefaultSpawnableState(const FGuid& BindingId, FMovieSceneSequenceIDRef TemplateID);
 
 	/** Check whether the specified objects are editable on the details panel. Called from the level editor */
 	bool AreObjectsEditable(const TArray<TWeakObjectPtr<UObject>>& InObjects) const;
diff --git a/Engine/Source/Developer/BlueprintCompilerCppBackend/Private/BlueprintCompilerCppBackendUMG.cpp b/Engine/Source/Developer/BlueprintCompilerCppBackend/Private/BlueprintCompilerCppBackendUMG.cpp
--- a/Engine/Source/Developer/BlueprintCompilerCppBackend/Private/BlueprintCompilerCppBackendUMG.cpp
+++ b/Engine/Source/Developer/BlueprintCompilerCppBackend/Private/BlueprintCompilerCppBackendUMG.cpp
@@ -11,6 +11,7 @@
 #include "Evaluation/MovieSceneSegment.h"
 #include "Evaluation/MovieSceneTrackImplementation.h"
 #include "Evaluation/MovieSceneEvalTemplate.h"
+#include "Compilation/MovieSceneCompiler.h"
 
 void FBackendHelperUMG::WidgetFunctionsInHeader(FEmitterLocalContext& Context)
 {
@@ -44,7 +45,8 @@ void FBackendHelperUMG::CreateClassSubobjects(FEmitterLocalContext& Context, boo
 			ensure(Anim->GetOuter() == Context.GetCurrentlyGeneratedClass());
 
 			// We need the same regeneration like for cooking. See UMovieSceneSequence::Serialize
-			Anim->EvaluationTemplate.Regenerate(Anim->TemplateParameters);
+			FMovieSceneSequencePrecompiledTemplateStore Store;
+			FMovieSceneCompiler::Compile(*Anim, Store);
 
 			FEmitDefaultValueHelper::HandleClassSubobject(Context, Anim, FEmitterLocalContext::EClassSubobjectList::MiscConvertedSubobjects, bCreate, bInitialize);
 		}
diff --git a/Engine/Source/Editor/DetailCustomizations/Private/LevelSequenceActorDetails.cpp b/Engine/Source/Editor/DetailCustomizations/Private/LevelSequenceActorDetails.cpp
--- a/Engine/Source/Editor/DetailCustomizations/Private/LevelSequenceActorDetails.cpp
+++ b/Engine/Source/Editor/DetailCustomizations/Private/LevelSequenceActorDetails.cpp
@@ -11,6 +11,7 @@
 #include "DetailCategoryBuilder.h"
 #include "IDetailsView.h"
 #include "LevelSequenceActor.h"
+#include "Algo/Transform.h"
 
 #define LOCTEXT_NAMESPACE "LevelSequenceActorDetails"
 
@@ -20,6 +21,30 @@ TSharedRef<IDetailCustomization> FLevelSequenceActorDetails::MakeInstance()
 	return MakeShareable( new FLevelSequenceActorDetails );
 }
 
+void AddAllSubObjectProperties(TArray<UObject*>& SubObjects, IDetailCategoryBuilder& Category, TAttribute<EVisibility> Visibility = TAttribute<EVisibility>(EVisibility::Visible))
+{
+	SubObjects.Remove(nullptr);
+	if (!SubObjects.Num())
+	{
+		return;
+	}
+
+	for (const UProperty* TestProperty : TFieldRange<UProperty>(SubObjects[0]->GetClass()))
+	{
+		if (TestProperty->HasAnyPropertyFlags(CPF_Edit))
+		{
+			const bool bAdvancedDisplay = TestProperty->HasAnyPropertyFlags(CPF_AdvancedDisplay);
+			const EPropertyLocation::Type PropertyLocation = bAdvancedDisplay ? EPropertyLocation::Advanced : EPropertyLocation::Common;
+
+			IDetailPropertyRow* NewRow = Category.AddExternalObjectProperty(SubObjects, TestProperty->GetFName(), PropertyLocation);
+			if (NewRow)
+			{
+				NewRow->Visibility(Visibility);
+			}
+		}
+	}
+}
+
 void FLevelSequenceActorDetails::CustomizeDetails( IDetailLayoutBuilder& DetailLayout )
 {
 	const TArray< TWeakObjectPtr<UObject> >& SelectedObjects = DetailLayout.GetSelectedObjects();
@@ -37,10 +62,25 @@ void FLevelSequenceActorDetails::CustomizeDetails( IDetailLayoutBuilder& DetailL
 			}
 		}
 	}
+
+	TArray<ALevelSequenceActor*> LevelSequenceActors;
+	{
+		TArray<TWeakObjectPtr<>> ObjectPtrs;
+		DetailLayout.GetObjectsBeingCustomized(ObjectPtrs);
+
+		for (TWeakObjectPtr<> WeakObj : ObjectPtrs)
+		{
+			if (auto* Actor = Cast<ALevelSequenceActor>(WeakObj.Get()))
+			{
+				LevelSequenceActors.Add(Actor);
+			}
+		}
+	}
 	
-	//DetailLayout.EditCategory( "LevelSequenceActor", NSLOCTEXT("LevelSequenceActorDetails", "LevelSequenceActor", "Level Sequence Actor"), ECategoryPriority::Important )
-	DetailLayout.EditCategory( "General", NSLOCTEXT("GeneralDetails", "General", "General"), ECategoryPriority::Important )
-	.AddCustomRow( NSLOCTEXT("LevelSequenceActorDetails", "OpenLevelSequence", "Open Level Sequence") )
+
+	IDetailCategoryBuilder& GeneralCategory = DetailLayout.EditCategory( "General", NSLOCTEXT("GeneralDetails", "General", "General"), ECategoryPriority::Important );
+
+	GeneralCategory.AddCustomRow( NSLOCTEXT("LevelSequenceActorDetails", "OpenLevelSequence", "Open Level Sequence") )
 	[
 		SNew(SHorizontalBox)
 		+SHorizontalBox::Slot()
@@ -56,6 +96,43 @@ void FLevelSequenceActorDetails::CustomizeDetails( IDetailLayoutBuilder& DetailL
 			.Text( NSLOCTEXT("LevelSequenceActorDetails", "OpenLevelSequence", "Open Level Sequence") )
 		]
 	];
+
+	TArray<UObject*> SubObjects;
+
+	IDetailCategoryBuilder& BurnInOptionsCategory = DetailLayout.EditCategory( "BurnInOptions", LOCTEXT("BurnInOptions", "Burn In Options") ).InitiallyCollapsed(false);
+	{
+		SubObjects.Reset();
+		Algo::Transform(LevelSequenceActors, SubObjects, &ALevelSequenceActor::BurnInOptions);
+
+		AddAllSubObjectProperties(SubObjects, BurnInOptionsCategory);
+	}
+
+	IDetailCategoryBuilder& BindingOverridesCategory = DetailLayout.EditCategory( "BindingOverrides", LOCTEXT("BindingOverrides", "Binding Overrides") ).InitiallyCollapsed(false);
+	{
+		SubObjects.Reset();
+		Algo::Transform(LevelSequenceActors, SubObjects, &ALevelSequenceActor::BindingOverrides);
+		AddAllSubObjectProperties(SubObjects, BindingOverridesCategory);
+	}
+
+	IDetailCategoryBuilder& InstanceDataCategory = DetailLayout.EditCategory( "InstanceData", LOCTEXT("InstanceData", "Instance Data") ).InitiallyCollapsed(false);
+	{
+		TSharedRef<IPropertyHandle> UseInstanceData = DetailLayout.GetProperty(GET_MEMBER_NAME_CHECKED(ALevelSequenceActor, bOverrideInstanceData));
+
+		DetailLayout.HideProperty(UseInstanceData);
+		InstanceDataCategory.AddProperty(UseInstanceData);
+
+		SubObjects.Reset();
+		Algo::Transform(LevelSequenceActors, SubObjects, &ALevelSequenceActor::DefaultInstanceData);
+
+		auto IsVisible = [UseInstanceData]() -> EVisibility
+		{
+			bool bValue = false;
+			return UseInstanceData->GetValue(bValue) == FPropertyAccess::Success && bValue ? EVisibility::Visible : EVisibility::Collapsed;
+		};
+
+		TAttribute<EVisibility> Visibility = TAttribute<EVisibility>::Create(TAttribute<EVisibility>::FGetter::CreateLambda(IsVisible));
+		AddAllSubObjectProperties(SubObjects, InstanceDataCategory, Visibility);
+	}
 }
 
 
diff --git a/Engine/Source/Editor/DetailCustomizations/Private/MovieSceneSequencePlaybackSettingsCustomization.cpp b/Engine/Source/Editor/DetailCustomizations/Private/MovieSceneSequencePlaybackSettingsCustomization.cpp
--- a/Engine/Source/Editor/DetailCustomizations/Private/MovieSceneSequencePlaybackSettingsCustomization.cpp
+++ b/Engine/Source/Editor/DetailCustomizations/Private/MovieSceneSequencePlaybackSettingsCustomization.cpp
@@ -139,6 +139,7 @@ void FMovieSceneSequencePlaybackSettingsCustomization::CustomizeChildren(TShared
 	ChildBuilder.AddProperty(StructPropertyHandle->GetChildHandle(GET_MEMBER_NAME_CHECKED(FMovieSceneSequencePlaybackSettings, bDisableLookAtInput)).ToSharedRef());
 	ChildBuilder.AddProperty(StructPropertyHandle->GetChildHandle(GET_MEMBER_NAME_CHECKED(FMovieSceneSequencePlaybackSettings, bHidePlayer)).ToSharedRef());
 	ChildBuilder.AddProperty(StructPropertyHandle->GetChildHandle(GET_MEMBER_NAME_CHECKED(FMovieSceneSequencePlaybackSettings, bHideHud)).ToSharedRef());
+	ChildBuilder.AddProperty(StructPropertyHandle->GetChildHandle(GET_MEMBER_NAME_CHECKED(FMovieSceneSequencePlaybackSettings, bDisableCameraCuts)).ToSharedRef());
 }
 
 
diff --git a/Engine/Source/Editor/EditorStyle/Private/SlateEditorStyle.cpp b/Engine/Source/Editor/EditorStyle/Private/SlateEditorStyle.cpp
--- a/Engine/Source/Editor/EditorStyle/Private/SlateEditorStyle.cpp
+++ b/Engine/Source/Editor/EditorStyle/Private/SlateEditorStyle.cpp
@@ -1977,6 +1977,8 @@ void FSlateEditorStyle::FStyle::SetupGeneralStyles()
 		Set( "SequenceRecorder.Common.RemoveRecording", new IMAGE_BRUSH( "SequenceRecorder/icon_RemoveRecording_40x", Icon40x40 ) );
 		Set( "SequenceRecorder.Common.RemoveAllRecordings.Small", new IMAGE_BRUSH( "SequenceRecorder/icon_RemoveRecording_40x", Icon20x20 ) );
 		Set( "SequenceRecorder.Common.RemoveAllRecordings", new IMAGE_BRUSH( "SequenceRecorder/icon_RemoveRecording_40x", Icon40x40 ) );
+		Set( "SequenceRecorder.Common.RecordingActive", new IMAGE_BRUSH( "Common/SmallCheckBox_Checked", Icon14x14 ) );
+		Set( "SequenceRecorder.Common.RecordingInactive", new IMAGE_BRUSH( "Common/SmallCheckBox", Icon14x14 ) );
 	}
 
 	// Foliage Edit Mode
diff --git a/Engine/Source/Editor/LevelEditor/Private/LevelViewportActions.cpp b/Engine/Source/Editor/LevelEditor/Private/LevelViewportActions.cpp
--- a/Engine/Source/Editor/LevelEditor/Private/LevelViewportActions.cpp
+++ b/Engine/Source/Editor/LevelEditor/Private/LevelViewportActions.cpp
@@ -52,7 +52,7 @@ void FLevelViewportCommands::RegisterCommands()
 	UI_COMMAND( ViewportConfig_FourPanesBottom, "Layout Four Panes (one bottom, three top)", "Changes the viewport arrangement to four panes, one on the bottom, three on the top", EUserInterfaceActionType::ToggleButton, FInputChord() );
 	UI_COMMAND( ViewportConfig_FourPanes2x2, "Layout Four Panes (2x2)", "Changes the viewport arrangement to four panes, in a 2x2 grid", EUserInterfaceActionType::ToggleButton, FInputChord() );
 
-	UI_COMMAND( SetDefaultViewportType, "Default Viewport", "Reconfigures this viewport to the default arrangement", EUserInterfaceActionType::RadioButton, FInputChord(EModifierKey::Shift, EKeys::D) );
+	UI_COMMAND( SetDefaultViewportType, "Default Viewport", "Reconfigures this viewport to the default arrangement", EUserInterfaceActionType::RadioButton, FInputChord() );
 
 	UI_COMMAND( ToggleViewportToolbar, "Show Toolbar", "Defines whether a toolbar should be displayed on this viewport", EUserInterfaceActionType::ToggleButton, FInputChord(EModifierKey::Control | EModifierKey::Shift, EKeys::T) );
 
diff --git a/Engine/Source/Editor/MovieSceneTools/Private/K2Node_GetSequenceBinding.cpp b/Engine/Source/Editor/MovieSceneTools/Private/K2Node_GetSequenceBinding.cpp
--- a/Engine/Source/Editor/MovieSceneTools/Private/K2Node_GetSequenceBinding.cpp
+++ b/Engine/Source/Editor/MovieSceneTools/Private/K2Node_GetSequenceBinding.cpp
@@ -21,6 +21,7 @@
 #include "SImage.h"
 #include "SComboBox.h"
 #include "Editor.h"
+#include "Compilation/MovieSceneCompiler.h"
 
 static const FName OutputPinName(TEXT("Output"));
 static const FName SequencePinName(TEXT("Sequence"));
@@ -141,17 +142,54 @@ UMovieScene* UK2Node_GetSequenceBinding::GetObjectMovieScene() const
 		}
 		else
 		{
-			// Ensure the template is up to date
-			FMovieSceneTrackCompilationParams Params = Sequence->TemplateParameters;
-			// Indicate that the template was generated during a blueprint compile so it can be regenerated with the full data when necessary
-			Params.bDuringBlueprintCompile = true;
-			// Regenerate the template
-			Sequence->EvaluationTemplate.Regenerate(Params);
+			bool bHierarchyIsValid = true;
+
+			// Ensure the hierarchy is valid (ie, the user hasn't changed a sub sequence for something else)
+			FMovieSceneSequenceID CurrentSequenceID = SequenceID;
+			while (CurrentSequenceID != MovieSceneSequenceID::Root)
+			{
+				const FMovieSceneSubSequenceData* SubData = SequenceHierarchyCache.FindSubData(CurrentSequenceID);
+				const UMovieSceneSequence* SubSequence = SubData ? SubData->GetSequence() : nullptr;
+
+				if (!SubSequence || SubSequence->GetSignature() != SequenceSignatureCache.FindRef(CurrentSequenceID))
+				{
+					bHierarchyIsValid = false;
+					break;
+				}
+			}
+
+			// If it's not valid, it needs recompiling
+			if (!bHierarchyIsValid)
+			{
+				// Recompile the hierarchy
+				SequenceSignatureCache.Reset();
+				SequenceHierarchyCache = FMovieSceneSequenceHierarchy();
+
+				FMovieSceneCompiler::CompileHierarchy(*Sequence, SequenceHierarchyCache);
+
+				TArray<FMovieSceneSequenceID> AllSequenceIDs;
+				while (AllSequenceIDs.Num())
+				{
+					int32 NumSequenceIDs = AllSequenceIDs.Num();
+					for (int32 Index = 0; Index < NumSequenceIDs; ++Index)
+					{
+						FMovieSceneSequenceID ThisSequenceID = AllSequenceIDs[Index];
+						const FMovieSceneSubSequenceData* SubData = SequenceHierarchyCache.FindSubData(ThisSequenceID);
+						const UMovieSceneSequence* SubSequence = SubData ? SubData->GetSequence() : nullptr;
+						if (ensure(SubSequence))
+						{
+							SequenceSignatureCache.Add(ThisSequenceID, SubSequence->GetSignature());
+						}
+					}
+
+					AllSequenceIDs.RemoveAt(0, NumSequenceIDs);
+				}
+			}
 
 			UMovieScene* MovieScene = nullptr;
-			if (const FMovieSceneSubSequenceData* SubData = Sequence->EvaluationTemplate.Hierarchy.FindSubData(SequenceID))
+			if (const FMovieSceneSubSequenceData* SubData = SequenceHierarchyCache.FindSubData(SequenceID))
 			{
-				UMovieSceneSequence* SubSequence = SubData->Sequence;
+				UMovieSceneSequence* SubSequence = SubData->GetSequence();
 				return SubSequence ? SubSequence->GetMovieScene() : nullptr;
 			}
 		}
diff --git a/Engine/Source/Editor/MovieSceneTools/Private/MatineeImportTools.cpp b/Engine/Source/Editor/MovieSceneTools/Private/MatineeImportTools.cpp
--- a/Engine/Source/Editor/MovieSceneTools/Private/MatineeImportTools.cpp
+++ b/Engine/Source/Editor/MovieSceneTools/Private/MatineeImportTools.cpp
@@ -655,7 +655,8 @@ bool FMatineeImportTools::CopyInterpDirectorTrack( UInterpTrackDirector* Directo
 				FGuid CameraHandle = Player.FindObjectId(*CameraActor, MovieSceneSequenceID::Root);
 				if (CameraHandle.IsValid())
 				{
-					CameraCutTrack->AddNewCameraCut(CameraHandle, TrackCut.Time);
+					FMovieSceneObjectBindingID CameraBindingID(CameraHandle, MovieSceneSequenceID::Root);
+					CameraCutTrack->AddNewCameraCut(CameraBindingID, TrackCut.Time);
 					bCutsAdded = true;
 				}
 			}
diff --git a/Engine/Source/Editor/MovieSceneTools/Private/MovieSceneObjectBindingIDCustomization.cpp b/Engine/Source/Editor/MovieSceneTools/Private/MovieSceneObjectBindingIDCustomization.cpp
--- a/Engine/Source/Editor/MovieSceneTools/Private/MovieSceneObjectBindingIDCustomization.cpp
+++ b/Engine/Source/Editor/MovieSceneTools/Private/MovieSceneObjectBindingIDCustomization.cpp
@@ -6,15 +6,38 @@
 #include "MovieSceneSequence.h"
 #include "MovieScene.h"
 #include "PropertyHandle.h"
+#include "PropertyHandle.h"
 #include "DetailWidgetRow.h"
 #include "IDetailChildrenBuilder.h"
 #include "SDropTarget.h"
+#include "ISequencer.h"
 #include "SComboButton.h"
 #include "STextBlock.h"
+#include "ScopedTransaction.h"
 #include "SequencerObjectBindingDragDropOp.h"
 
 #define LOCTEXT_NAMESPACE "MovieSceneObjectBindingIDCustomization"
 
+void FMovieSceneObjectBindingIDCustomization::BindTo(TSharedRef<ISequencer> OuterSequencer)
+{
+	OuterSequencer->OnInitializeDetailsPanel().AddStatic(
+		[](TSharedRef<IDetailsView> DetailsView, TSharedRef<ISequencer> InSequencer)
+		{
+			TWeakPtr<ISequencer> WeakSequencer = InSequencer;
+
+			FOnGetPropertyTypeCustomizationInstance BindingIDCustomizationFactory = FOnGetPropertyTypeCustomizationInstance::CreateLambda(
+				[WeakSequencer]
+				{
+					return MakeShared<FMovieSceneObjectBindingIDCustomization>(WeakSequencer.Pin()->GetFocusedTemplateID(), WeakSequencer);
+				}
+			);
+
+			// Register an object binding ID customization that can use the current sequencer interface
+			DetailsView->RegisterInstancedCustomPropertyTypeLayout("MovieSceneObjectBindingID", BindingIDCustomizationFactory);
+		}
+	);
+}
+
 void FMovieSceneObjectBindingIDCustomization::CustomizeHeader(TSharedRef<IPropertyHandle> PropertyHandle, FDetailWidgetRow& HeaderRow, IPropertyTypeCustomizationUtils& CustomizationUtils)
 {
 	StructProperty = PropertyHandle;
@@ -98,13 +121,27 @@ FMovieSceneObjectBindingID FMovieSceneObjectBindingIDCustomization::GetCurrentVa
 
 void FMovieSceneObjectBindingIDCustomization::SetCurrentValue(const FMovieSceneObjectBindingID& InObjectBinding)
 {
+	FScopedTransaction Transaction(LOCTEXT("SetBinding", "Set Binding"));
+
+	StructProperty->NotifyPreChange();
+
+	TArray<UObject*> Objects;
+	StructProperty->GetOuterObjects(Objects);
+	for (UObject* Object : Objects)
+	{
+		Object->Modify();
+	}
+
 	TArray<void*> Ptrs;
 	StructProperty->AccessRawData(Ptrs);
 
-	if (ensure(Ptrs.Num() == 1))
+	for (void* Ptr : Ptrs)
 	{
-		*static_cast<FMovieSceneObjectBindingID*>(Ptrs[0]) = InObjectBinding;
+		*static_cast<FMovieSceneObjectBindingID*>(Ptr) = InObjectBinding;
 	}
+	
+	StructProperty->NotifyPostChange();
+	StructProperty->NotifyFinishedChangingProperties();
 }
 
 #undef LOCTEXT_NAMESPACE
diff --git a/Engine/Source/Editor/MovieSceneTools/Private/MovieSceneObjectBindingIDPicker.cpp b/Engine/Source/Editor/MovieSceneTools/Private/MovieSceneObjectBindingIDPicker.cpp
--- a/Engine/Source/Editor/MovieSceneTools/Private/MovieSceneObjectBindingIDPicker.cpp
+++ b/Engine/Source/Editor/MovieSceneTools/Private/MovieSceneObjectBindingIDPicker.cpp
@@ -19,6 +19,7 @@
 #include "ISequencer.h"
 #include "MovieSceneEvaluationTemplateInstance.h"
 #include "MovieSceneSequenceHierarchy.h"
+#include "SlateApplication.h"
 
 #define LOCTEXT_NAMESPACE "MovieSceneObjectBindingIDPicker"
 
@@ -90,6 +91,8 @@ struct FSequenceBindingTree
 	 */
 	void Build(UMovieSceneSequence* InSequence, FObjectKey InActiveSequence, FMovieSceneSequenceID InActiveSequenceID)
 	{
+		bIsEmpty = true;
+
 		// Reset state
 		ActiveSequenceID = InActiveSequenceID;
 		ActiveSequence = InActiveSequence;
@@ -149,6 +152,11 @@ struct FSequenceBindingTree
 		return Hierarchy.FindRef(BindingID);
 	}
 
+	bool IsEmpty() const
+	{
+		return bIsEmpty;
+	}
+
 private:
 
 	/** Recursive sort helper for a sequence binding node */
@@ -214,7 +222,7 @@ private:
 						FMovieSceneObjectBindingID CurrentID(FGuid(), SequenceIDStack.GetCurrent());
 
 						UMovieSceneCinematicShotSection* ShotSection = Cast<UMovieSceneCinematicShotSection>(Section);
-						FText DisplayString = ShotSection ? ShotSection->GetShotDisplayName() : FText::FromName(SubSection->GetFName());
+						FText DisplayString = ShotSection ? FText::FromString(ShotSection->GetShotDisplayName()) : FText::FromName(SubSequence->GetFName());
 						FSlateIcon Icon(FEditorStyle::GetStyleSetName(), ShotSection ? "Sequencer.Tracks.CinematicShot" : "Sequencer.Tracks.Sub");
 						
 						TSharedRef<FSequenceBindingNode> NewNode = MakeShared<FSequenceBindingNode>(DisplayString, CurrentID, Icon);
@@ -248,6 +256,8 @@ private:
 			EnsureParent(FGuid(), MovieScene, CurrentSequenceID)->AddChild(NewNode);
 			ensure(!Hierarchy.Contains(ID));
 			Hierarchy.Add(ID, NewNode);
+
+			bIsEmpty = false;
 		}
 
 		// Add all possessables
@@ -265,6 +275,8 @@ private:
 				EnsureParent(Possessable.GetParent(), MovieScene, CurrentSequenceID)->AddChild(NewNode);
 				ensure(!Hierarchy.Contains(ID));
 				Hierarchy.Add(ID, NewNode);
+
+				bIsEmpty = false;
 			}
 		}
 	}
@@ -310,6 +322,7 @@ private:
 		
 		ensure(!Hierarchy.Contains(ParentPtr));
 		Hierarchy.Add(ParentPtr, NewNode);
+		bIsEmpty = false;
 
 		EnsureParent(AddToGuid, InMovieScene, SequenceID)->AddChild(NewNode);
 
@@ -328,8 +341,15 @@ private:
 	TSharedPtr<FSequenceBindingNode> TopLevelNode;
 	/** Map of hierarchical information */
 	TMap<FMovieSceneObjectBindingID, TSharedPtr<FSequenceBindingNode>> Hierarchy;
+	/** Whether the tree is considered empty */
+	bool bIsEmpty;
 };
 
+bool FMovieSceneObjectBindingIDPicker::IsEmpty() const
+{
+	return !DataTree.IsValid() || DataTree->IsEmpty();
+}
+
 void FMovieSceneObjectBindingIDPicker::Initialize()
 {
 	if (!DataTree.IsValid())
@@ -381,7 +401,8 @@ void FMovieSceneObjectBindingIDPicker::OnGetMenuContent(FMenuBuilder& MenuBuilde
 					FText(),
 					FNewMenuDelegate::CreateRaw(this, &FMovieSceneObjectBindingIDPicker::OnGetMenuContent, Child),
 					false,
-					Child->Icon
+					Child->Icon,
+					false
 					);
 			}
 		}
@@ -411,9 +432,17 @@ TSharedRef<SWidget> FMovieSceneObjectBindingIDPicker::GetPickerMenu()
 	FMenuBuilder MenuBuilder(true, nullptr, nullptr, true);
 
 	Initialize();
-	OnGetMenuContent(MenuBuilder, DataTree->GetRootNode());
+	GetPickerMenu(MenuBuilder);
+
+	// Hold onto the menu widget so we can destroy it manually
+	TSharedRef<SWidget> MenuWidget = MenuBuilder.MakeWidget();
+	DismissWidget = MenuWidget;
+	return MenuWidget;
+}
 
-	return MenuBuilder.MakeWidget();
+void FMovieSceneObjectBindingIDPicker::GetPickerMenu(FMenuBuilder& MenuBuilder)
+{
+	OnGetMenuContent(MenuBuilder, DataTree->GetRootNode());
 }
 
 TSharedRef<SWidget> FMovieSceneObjectBindingIDPicker::GetCurrentItemWidget(TSharedRef<STextBlock> TextContent)
@@ -455,6 +484,12 @@ void FMovieSceneObjectBindingIDPicker::SetBindingId(FMovieSceneObjectBindingID I
 {
 	SetRemappedCurrentValue(InBindingId);
 	UpdateCachedData();
+
+	TSharedPtr<SWidget> MenuWidget = DismissWidget.Pin();
+	if (MenuWidget.IsValid())
+	{
+		FSlateApplication::Get().DismissMenuByWidget(MenuWidget.ToSharedRef());
+	}
 }
 
 void FMovieSceneObjectBindingIDPicker::UpdateCachedData()
diff --git a/Engine/Source/Editor/MovieSceneTools/Private/MovieSceneToolHelpers.cpp b/Engine/Source/Editor/MovieSceneTools/Private/MovieSceneToolHelpers.cpp
--- a/Engine/Source/Editor/MovieSceneTools/Private/MovieSceneToolHelpers.cpp
+++ b/Engine/Source/Editor/MovieSceneTools/Private/MovieSceneToolHelpers.cpp
@@ -38,6 +38,13 @@
 #include "Widgets/Input/SButton.h"
 #include "Editor.h"
 #include "LevelEditorViewport.h"
+#include "AssetToolsModule.h"
+#include "Camera/CameraAnim.h"
+#include "Matinee/InterpGroup.h"
+#include "Matinee/InterpGroupInst.h"
+#include "Matinee/InterpTrackMove.h"
+#include "Matinee/InterpTrackMoveAxis.h"
+#include "Matinee/InterpTrackInstMove.h"
 
 /* MovieSceneToolHelpers
  *****************************************************************************/
@@ -282,7 +289,7 @@ FString MovieSceneToolHelpers::GenerateNewShotName(const TArray<UMovieSceneSecti
 		uint32 NextShotNumber = ProjectSettings->FirstShotNumber;
 		uint32 NextTakeNumber = ProjectSettings->FirstTakeNumber;
 
-		if (ParseShotName(NextShot->GetShotDisplayName().ToString(), NextShotPrefix, NextShotNumber, NextTakeNumber))
+		if (ParseShotName(NextShot->GetShotDisplayName(), NextShotPrefix, NextShotNumber, NextTakeNumber))
 		{
 			uint32 NewShotNumber = NextShotNumber + ProjectSettings->ShotIncrement;
 			return ComposeShotName(NextShotPrefix, NewShotNumber, ProjectSettings->FirstTakeNumber);
@@ -299,8 +306,8 @@ FString MovieSceneToolHelpers::GenerateNewShotName(const TArray<UMovieSceneSecti
 		uint32 NextShotNumber = ProjectSettings->FirstShotNumber;
 		uint32 NextTakeNumber = ProjectSettings->FirstTakeNumber;
 
-		if (ParseShotName(BeforeShot->GetShotDisplayName().ToString(), BeforeShotPrefix, BeforeShotNumber, BeforeTakeNumber) &&
-			ParseShotName(NextShot->GetShotDisplayName().ToString(), NextShotPrefix, NextShotNumber, NextTakeNumber))
+		if (ParseShotName(BeforeShot->GetShotDisplayName(), BeforeShotPrefix, BeforeShotNumber, BeforeTakeNumber) &&
+			ParseShotName(NextShot->GetShotDisplayName(), NextShotPrefix, NextShotNumber, NextTakeNumber))
 		{
 			if (BeforeShotNumber < NextShotNumber)
 			{
@@ -331,7 +338,7 @@ void MovieSceneToolHelpers::GatherTakes(const UMovieSceneSection* Section, TArra
 	uint32 ShotNumber = INDEX_NONE;
 	CurrentTakeNumber = INDEX_NONE;
 
-	if (ParseShotName(Shot->GetShotDisplayName().ToString(), ShotPrefix, ShotNumber, CurrentTakeNumber))
+	if (ParseShotName(Shot->GetShotDisplayName(), ShotPrefix, ShotNumber, CurrentTakeNumber))
 	{
 		// Gather up all level sequence assets
 		FAssetRegistryModule& AssetRegistryModule = FModuleManager::LoadModuleChecked<FAssetRegistryModule>(TEXT("AssetRegistry"));
@@ -377,7 +384,7 @@ UObject* MovieSceneToolHelpers::GetTake(const UMovieSceneSection* Section, uint3
 	uint32 ShotNumber = INDEX_NONE;
 	uint32 TakeNumberDummy = INDEX_NONE;
 
-	if (ParseShotName(Shot->GetShotDisplayName().ToString(), ShotPrefix, ShotNumber, TakeNumberDummy))
+	if (ParseShotName(Shot->GetShotDisplayName(), ShotPrefix, ShotNumber, TakeNumberDummy))
 	{
 		// Gather up all level sequence assets
 		FAssetRegistryModule& AssetRegistryModule = FModuleManager::LoadModuleChecked<FAssetRegistryModule>(TEXT("AssetRegistry"));
@@ -411,6 +418,31 @@ UObject* MovieSceneToolHelpers::GetTake(const UMovieSceneSection* Section, uint3
 	return nullptr;
 }
 
+int32 MovieSceneToolHelpers::FindAvailableRowIndex(UMovieSceneTrack* InTrack, UMovieSceneSection* InSection)
+{
+	for (int32 RowIndex = 0; RowIndex <= InTrack->GetMaxRowIndex(); ++RowIndex)
+	{
+		bool bFoundIntersect = false;
+		for (auto Section : InTrack->GetAllSections())
+		{
+			TRange<float> InRange(InSection->GetStartTime(), InSection->GetEndTime());
+			TRange<float> Range(Section->GetStartTime(), Section->GetEndTime());
+
+			if (Section != InSection && Section->GetRowIndex() == RowIndex && Range.Overlaps(InRange))
+			{
+				bFoundIntersect = true;
+				break;
+			}
+		}
+		if (!bFoundIntersect)
+		{
+			return RowIndex;
+		}
+	}
+
+	return InTrack->GetMaxRowIndex() + 1;
+}
+
 class SEnumCombobox : public SComboBox<TSharedPtr<int32>>
 {
 public:
@@ -605,6 +637,7 @@ bool MovieSceneToolHelpers::ShowExportEDLDialog(const UMovieScene* InMovieScene,
 bool ImportFBXProperty(FString NodeName, FString AnimatedPropertyName, FGuid ObjectBinding, UnFbx::FFbxCurvesAPI& CurveAPI, UMovieScene* InMovieScene, ISequencer& InSequencer)
 {
 	const UMovieSceneToolsProjectSettings* ProjectSettings = GetDefault<UMovieSceneToolsProjectSettings>();
+	const UMovieSceneUserImportFBXSettings* ImportFBXSettings = GetDefault<UMovieSceneUserImportFBXSettings>();
 
 	TArrayView<TWeakObjectPtr<>> BoundObjects = InSequencer.FindBoundObjects(ObjectBinding, InSequencer.GetFocusedTemplateID());
 
@@ -695,7 +728,11 @@ bool ImportFBXProperty(FString NodeName, FString AnimatedPropertyName, FGuid Obj
 					FMatineeImportTools::SetOrAddKey(FloatCurve, CurveHandle.Points[KeyIndex].InVal, CurveHandle.Points[KeyIndex].OutVal, ArriveTangent, LeaveTangent, CurveHandle.Points[KeyIndex].InterpMode);
 				}
 
-				FloatCurve.RemoveRedundantKeys(KINDA_SMALL_NUMBER);
+				if (ImportFBXSettings->bReduceKeys)
+				{
+					FloatCurve.RemoveRedundantKeys(ImportFBXSettings->ReduceKeysTolerance);
+				}
+				
 				FloatCurve.AutoSetTangents();
 
 				FloatSection->SetStartTime(MinTime);
@@ -709,6 +746,8 @@ bool ImportFBXProperty(FString NodeName, FString AnimatedPropertyName, FGuid Obj
 
 bool ImportFBXTransform(FString NodeName, FGuid ObjectBinding, UnFbx::FFbxCurvesAPI& CurveAPI, UMovieScene* InMovieScene)
 {
+	const UMovieSceneUserImportFBXSettings* ImportFBXSettings = GetDefault<UMovieSceneUserImportFBXSettings>();
+
 	// Look for transforms explicitly
 	FInterpCurveFloat Translation[3];
 	FInterpCurveFloat EulerRotation[3];
@@ -826,8 +865,12 @@ bool ImportFBXTransform(FString NodeName, FGuid ObjectBinding, UnFbx::FFbxCurves
 
 					FMatineeImportTools::SetOrAddKey(*ChannelCurve, CurveFloat->Points[KeyIndex].InVal, CurveFloat->Points[KeyIndex].OutVal, ArriveTangent, LeaveTangent, CurveFloat->Points[KeyIndex].InterpMode);
 				}
-
-				ChannelCurve->RemoveRedundantKeys(KINDA_SMALL_NUMBER);
+				
+				if (ImportFBXSettings->bReduceKeys)
+				{
+					ChannelCurve->RemoveRedundantKeys(ImportFBXSettings->ReduceKeysTolerance);
+				}
+				
 				ChannelCurve->AutoSetTangents();
 			}
 		}
@@ -966,9 +1009,16 @@ void ImportFBXCamera(UnFbx::FFbxImporter* FbxImporter, UMovieScene* InMovieScene
 
 		// Add any unmatched cameras
 		UWorld* World = GCurrentLevelEditingViewportClient ? GCurrentLevelEditingViewportClient->GetWorld() : nullptr;
+
+		// If there are new cameras, clear the object binding map so that we're only assigning values to the newly created cameras
+		if (UnmatchedCameras.Num() != 0)
+		{
+			InObjectBindingMap.Reset();
+		}
+
 		for (auto UnmatchedCamera : UnmatchedCameras)
 		{
-			FString CameraName = FString(UnmatchedCamera->GetName());
+			FString CameraName = FString(ANSI_TO_TCHAR(UnmatchedCamera->GetName()));
 
 			FActorSpawnParameters SpawnParams;
 			SpawnParams.Name = *CameraName;
@@ -982,8 +1032,6 @@ void ImportFBXCamera(UnFbx::FFbxImporter* FbxImporter, UMovieScene* InMovieScene
 			InObjectBindingMap.Add(NewCameraGuids[0]);
 			InObjectBindingMap[NewCameraGuids[0]] = CameraName;
 		}
-
-		// refresh?
 	}
 
 
@@ -1276,3 +1324,111 @@ bool MovieSceneToolHelpers::ImportFBX(UMovieScene* InMovieScene, ISequencer& InS
 
 	return true;
 }
+
+
+EInterpCurveMode MovieSceneToolHelpers::RichCurveInterpolationToMatineeInterpolation( ERichCurveInterpMode InterpMode )
+{
+	switch ( InterpMode )
+	{
+	case ERichCurveInterpMode::RCIM_Constant:
+		return CIM_Constant;
+	case ERichCurveInterpMode::RCIM_Cubic:
+		return CIM_CurveAuto;
+	case ERichCurveInterpMode::RCIM_Linear:
+		return CIM_Linear;
+	default:
+		return CIM_CurveAuto;
+	}
+}
+
+void MovieSceneToolHelpers::CopyRichCurveToMoveAxis(const FRichCurve& RichCurve, UInterpTrackMoveAxis* MoveAxis)
+{
+	static FName LookupName(NAME_None);
+	int32 KeyIndex = 0;
+	for ( auto KeyIt( RichCurve.GetKeyIterator() ); KeyIt; ++KeyIt, ++KeyIndex )
+	{
+		MoveAxis->FloatTrack.AddPoint(KeyIt->Time, KeyIt->Value);
+		MoveAxis->FloatTrack.Points[KeyIndex].InterpMode = RichCurveInterpolationToMatineeInterpolation(KeyIt->InterpMode);
+		MoveAxis->FloatTrack.Points[KeyIndex].ArriveTangent = KeyIt->ArriveTangent;
+		MoveAxis->FloatTrack.Points[KeyIndex].LeaveTangent = KeyIt->LeaveTangent;
+
+		MoveAxis->LookupTrack.AddPoint(KeyIt->Time, LookupName);
+	}
+}
+
+UObject* MovieSceneToolHelpers::ExportToCameraAnim(UMovieScene* InMovieScene, FGuid& InObjectBinding)
+{
+	// Create a new camera anim
+	IAssetTools& AssetTools = FModuleManager::GetModuleChecked<FAssetToolsModule>("AssetTools").Get();
+
+	UObject* NewAsset = nullptr;
+
+	// Attempt to create a new asset
+	for (TObjectIterator<UClass> It ; It ; ++It)
+	{
+		UClass* CurrentClass = *It;
+		if (CurrentClass->IsChildOf(UFactory::StaticClass()) && !(CurrentClass->HasAnyClassFlags(CLASS_Abstract)))
+		{
+			UFactory* Factory = Cast<UFactory>(CurrentClass->GetDefaultObject());
+			if (Factory->CanCreateNew() && Factory->ImportPriority >= 0 && Factory->SupportedClass == UCameraAnim::StaticClass())
+			{
+				NewAsset = AssetTools.CreateAssetWithDialog(UCameraAnim::StaticClass(), Factory);
+				break;
+			}
+		}
+	}
+
+	if (!NewAsset)
+	{
+		return NewAsset;
+	}
+
+	static FName Transform("Transform");
+	UMovieScene3DTransformTrack* TransformTrack = InMovieScene->FindTrack<UMovieScene3DTransformTrack>(InObjectBinding, Transform); 
+	if (TransformTrack)
+	{
+		UCameraAnim* CameraAnim = CastChecked<UCameraAnim>(NewAsset);
+		UInterpGroup* CameraInterpGroup = CameraAnim->CameraInterpGroup;
+		CameraAnim->bRelativeToInitialTransform=false;
+
+		UInterpGroupInst* CameraInst = NewObject<UInterpGroupInst>(CameraAnim, NAME_None, RF_Transactional);
+		CameraInst->InitGroupInst(CameraInterpGroup, nullptr);
+
+		UInterpTrackMove* MovementTrack = NewObject<UInterpTrackMove>(CameraInterpGroup, NAME_None, RF_Transactional);
+		CameraInterpGroup->InterpTracks.Add(MovementTrack);
+		
+		UInterpTrackInstMove* MovementTrackInst = NewObject<UInterpTrackInstMove>(CameraInst, NAME_None, RF_Transactional);
+		CameraInst->TrackInst.Add(MovementTrackInst);
+		MovementTrackInst->InitTrackInst(MovementTrack);
+			
+		MovementTrack->CreateSubTracks(false);
+
+		UInterpTrackMoveAxis* MoveAxies[6];
+		for( int32 SubTrackIndex = 0; SubTrackIndex < 6; ++SubTrackIndex )
+		{
+			MoveAxies[ SubTrackIndex ] = Cast<UInterpTrackMoveAxis>( MovementTrack->SubTracks[ SubTrackIndex ] );
+		}
+
+		TArray<UMovieSceneSection*> Sections = TransformTrack->GetAllSections();
+
+		if (Sections.Num())
+		{
+			if (Sections.Num() > 1)
+			{
+				UE_LOG(LogMovieScene, Error, TEXT("Export to Camera Anim: Failed to export, multiple sections (%d) are not supported"), Sections.Num());
+			}
+			else
+			{
+				UMovieScene3DTransformSection* TransformSection = Cast<UMovieScene3DTransformSection>(Sections[0]);
+				CopyRichCurveToMoveAxis(TransformSection->GetTranslationCurve(EAxis::X), MoveAxies[AXIS_TranslationX]);
+				CopyRichCurveToMoveAxis(TransformSection->GetTranslationCurve(EAxis::Y), MoveAxies[AXIS_TranslationY]);
+				CopyRichCurveToMoveAxis(TransformSection->GetTranslationCurve(EAxis::Z), MoveAxies[AXIS_TranslationZ]);
+				CopyRichCurveToMoveAxis(TransformSection->GetRotationCurve(EAxis::X), MoveAxies[AXIS_RotationX]);
+				CopyRichCurveToMoveAxis(TransformSection->GetRotationCurve(EAxis::Y), MoveAxies[AXIS_RotationY]);
+				CopyRichCurveToMoveAxis(TransformSection->GetRotationCurve(EAxis::Z), MoveAxies[AXIS_RotationZ]);
+			}
+		}
+	}
+
+	return NewAsset;
+}
\ No newline at end of file
diff --git a/Engine/Source/Editor/MovieSceneTools/Private/MovieSceneToolsModule.cpp b/Engine/Source/Editor/MovieSceneTools/Private/MovieSceneToolsModule.cpp
--- a/Engine/Source/Editor/MovieSceneTools/Private/MovieSceneToolsModule.cpp
+++ b/Engine/Source/Editor/MovieSceneTools/Private/MovieSceneToolsModule.cpp
@@ -136,6 +136,7 @@ public:
 		SequencerModule.UnRegisterTrackEditor( FloatPropertyTrackCreateEditorHandle );
 		SequencerModule.UnRegisterTrackEditor( IntegerPropertyTrackCreateEditorHandle );
 		SequencerModule.UnRegisterTrackEditor( VectorPropertyTrackCreateEditorHandle );
+		SequencerModule.UnRegisterTrackEditor( TransformPropertyTrackCreateEditorHandle );
 		SequencerModule.UnRegisterTrackEditor( VisibilityPropertyTrackCreateEditorHandle );
 		SequencerModule.UnRegisterTrackEditor( ActorReferencePropertyTrackCreateEditorHandle );
 		SequencerModule.UnRegisterTrackEditor( StringPropertyTrackCreateEditorHandle );
@@ -156,6 +157,7 @@ public:
 		SequencerModule.UnRegisterTrackEditor( ComponentMaterialTrackCreateEditorHandle );
 		SequencerModule.UnRegisterTrackEditor( FadeTrackCreateEditorHandle );
 		SequencerModule.UnRegisterTrackEditor( SpawnTrackCreateEditorHandle );
+		SequencerModule.UnRegisterTrackEditor( LevelVisibilityTrackCreateEditorHandle );
 		SequencerModule.UnRegisterTrackEditor( CameraAnimTrackCreateEditorHandle );
 		SequencerModule.UnRegisterTrackEditor( CameraShakeTrackCreateEditorHandle );
 		SequencerModule.UnRegisterTrackEditor( MPCTrackCreateEditorHandle );
diff --git a/Engine/Source/Editor/MovieSceneTools/Private/MovieSceneToolsUserSettings.cpp b/Engine/Source/Editor/MovieSceneTools/Private/MovieSceneToolsUserSettings.cpp
--- a/Engine/Source/Editor/MovieSceneTools/Private/MovieSceneToolsUserSettings.cpp
+++ b/Engine/Source/Editor/MovieSceneTools/Private/MovieSceneToolsUserSettings.cpp
@@ -33,4 +33,6 @@ UMovieSceneUserImportFBXSettings::UMovieSceneUserImportFBXSettings(const FObject
 {
 	bForceFrontXAxis = false;
 	bCreateCameras = true;
+	bReduceKeys = true;
+	ReduceKeysTolerance = 0.001f;
 }
diff --git a/Engine/Source/Editor/MovieSceneTools/Private/Sections/CameraCutSection.cpp b/Engine/Source/Editor/MovieSceneTools/Private/Sections/CameraCutSection.cpp
--- a/Engine/Source/Editor/MovieSceneTools/Private/Sections/CameraCutSection.cpp
+++ b/Engine/Source/Editor/MovieSceneTools/Private/Sections/CameraCutSection.cpp
@@ -12,6 +12,7 @@
 #include "ScopedTransaction.h"
 #include "MovieSceneSequence.h"
 #include "MovieSceneCommonHelpers.h"
+#include "Evaluation/MovieSceneEvaluationTemplateInstance.h"
 #include "EditorStyleSet.h"
 #include "EngineUtils.h"
 
@@ -105,7 +106,7 @@ void FCameraCutSection::BuildSectionContextMenu(FMenuBuilder& MenuBuilder, const
 
 			MenuBuilder.AddMenuEntry(
 				FText::Format(LOCTEXT("SetCameraMenuEntryTextFormat", "{0}"), ActorLabel),
-				FText::Format(LOCTEXT("SetCameraMenuEntryTooltipFormat", "Assign {0} to this camera cut"), ActorLabel),
+				FText::Format(LOCTEXT("SetCameraMenuEntryTooltipFormat", "Assign {0} to this camera cut"), FText::FromString(EachCamera->GetPathName())),
 				FSlateIcon(),
 				FUIAction(FExecuteAction::CreateRaw(this, &FCameraCutSection::HandleSetCameraMenuEntryExecute, EachCamera))
 			);
@@ -125,7 +126,15 @@ const AActor* FCameraCutSection::GetCameraForFrame(float Time) const
 
 	if (CameraCutSection && Sequencer.IsValid())
 	{
-		for (TWeakObjectPtr<>& Object : Sequencer->FindBoundObjects(CameraCutSection->GetCameraGuid(), Sequencer->GetFocusedTemplateID()))
+		FMovieSceneSequenceID SequenceID = Sequencer->GetFocusedTemplateID();
+		if (CameraCutSection->GetCameraBindingID().GetSequenceID().IsValid())
+		{
+			// Ensure that this ID is resolvable from the root, based on the current local sequence ID
+			FMovieSceneObjectBindingID RootBindingID = CameraCutSection->GetCameraBindingID().ResolveLocalToRoot(SequenceID, Sequencer->GetEvaluationTemplate().GetHierarchy());
+			SequenceID = RootBindingID.GetSequenceID();
+		}
+
+		for (TWeakObjectPtr<>& Object : Sequencer->FindBoundObjects(CameraCutSection->GetCameraBindingID().GetGuid(), SequenceID))
 		{
 			if (AActor* Actor = Cast<AActor>(Object.Get()))
 			{
@@ -133,7 +142,7 @@ const AActor* FCameraCutSection::GetCameraForFrame(float Time) const
 			}
 		}
 
-		FMovieSceneSpawnable* Spawnable = Sequencer->GetFocusedMovieSceneSequence()->GetMovieScene()->FindSpawnable(CameraCutSection->GetCameraGuid());
+		FMovieSceneSpawnable* Spawnable = Sequencer->GetFocusedMovieSceneSequence()->GetMovieScene()->FindSpawnable(CameraCutSection->GetCameraBindingID().GetGuid());
 		if (Spawnable)
 		{
 			return Cast<AActor>(Spawnable->GetObjectTemplate());
diff --git a/Engine/Source/Editor/MovieSceneTools/Private/Sections/CinematicShotSection.cpp b/Engine/Source/Editor/MovieSceneTools/Private/Sections/CinematicShotSection.cpp
--- a/Engine/Source/Editor/MovieSceneTools/Private/Sections/CinematicShotSection.cpp
+++ b/Engine/Source/Editor/MovieSceneTools/Private/Sections/CinematicShotSection.cpp
@@ -264,7 +264,7 @@ void FCinematicShotSection::BuildSectionContextMenu(FMenuBuilder& MenuBuilder, c
 
 		MenuBuilder.AddMenuEntry(
 			LOCTEXT("NewTake", "New Take"),
-			FText::Format(LOCTEXT("NewTakeTooltip", "Create a new take for {0}"), SectionObject.GetShotDisplayName()),
+			FText::Format(LOCTEXT("NewTakeTooltip", "Create a new take for {0}"), FText::FromString(SectionObject.GetShotDisplayName())),
 			FSlateIcon(),
 			FUIAction(FExecuteAction::CreateSP(CinematicShotTrackEditor.Pin().ToSharedRef(), &FCinematicShotTrackEditor::NewTake, &SectionObject))
 		);
@@ -278,21 +278,21 @@ void FCinematicShotSection::BuildSectionContextMenu(FMenuBuilder& MenuBuilder, c
 
 		MenuBuilder.AddMenuEntry(
 			LOCTEXT("DuplicateShot", "Duplicate Shot"),
-			FText::Format(LOCTEXT("DuplicateShotTooltip", "Duplicate {0} to create a new shot"), SectionObject.GetShotDisplayName()),
+			FText::Format(LOCTEXT("DuplicateShotTooltip", "Duplicate {0} to create a new shot"), FText::FromString(SectionObject.GetShotDisplayName())),
 			FSlateIcon(),
 			FUIAction(FExecuteAction::CreateSP(CinematicShotTrackEditor.Pin().ToSharedRef(), &FCinematicShotTrackEditor::DuplicateShot, &SectionObject))
 		);
 
 		MenuBuilder.AddMenuEntry(
 			LOCTEXT("RenderShot", "Render Shot"),
-			FText::Format(LOCTEXT("RenderShotTooltip", "Render shot movie"), SectionObject.GetShotDisplayName()),
+			FText::Format(LOCTEXT("RenderShotTooltip", "Render shot movie"), FText::FromString(SectionObject.GetShotDisplayName())),
 			FSlateIcon(),
 			FUIAction(FExecuteAction::CreateSP(CinematicShotTrackEditor.Pin().ToSharedRef(), &FCinematicShotTrackEditor::RenderShot, &SectionObject))
 		);
 
 		MenuBuilder.AddMenuEntry(
 			LOCTEXT("RenameShot", "Rename Shot"),
-			FText::Format(LOCTEXT("RenameShotTooltip", "Rename {0}"), SectionObject.GetShotDisplayName()),
+			FText::Format(LOCTEXT("RenameShotTooltip", "Rename {0}"), FText::FromString(SectionObject.GetShotDisplayName())),
 			FSlateIcon(),
 			FUIAction(FExecuteAction::CreateSP(this, &FCinematicShotSection::EnterRename))
 		);
@@ -322,7 +322,7 @@ void FCinematicShotSection::AddTakesMenu(FMenuBuilder& MenuBuilder)
 
 FText FCinematicShotSection::HandleThumbnailTextBlockText() const
 {
-	return SectionObject.GetShotDisplayName();
+	return FText::FromString(SectionObject.GetShotDisplayName());
 }
 
 
@@ -334,7 +334,7 @@ void FCinematicShotSection::HandleThumbnailTextBlockTextCommitted(const FText& N
 
 		const FScopedTransaction Transaction(LOCTEXT("SetShotName", "Set Shot Name"));
 	
-		SectionObject.SetShotDisplayName(NewShotName);
+		SectionObject.SetShotDisplayName(NewShotName.ToString());
 	}
 }
 
diff --git a/Engine/Source/Editor/MovieSceneTools/Private/Sections/TransformPropertySection.cpp b/Engine/Source/Editor/MovieSceneTools/Private/Sections/TransformPropertySection.cpp
--- a/Engine/Source/Editor/MovieSceneTools/Private/Sections/TransformPropertySection.cpp
+++ b/Engine/Source/Editor/MovieSceneTools/Private/Sections/TransformPropertySection.cpp
@@ -146,6 +146,15 @@ void FTransformSection::BuildSectionContextMenu(FMenuBuilder& MenuBuilder, const
 						TransformSection->SetMask(TransformSection->GetMask().GetChannels() | ChannelsToToggle);
 					}
 
+					// Restore pre-animated state for the bound objects so that inactive channels will return to their default values.
+					for (TWeakObjectPtr<> WeakObject : SequencerPtr->FindBoundObjects(InObjectBinding, SequencerPtr->GetFocusedTemplateID()))
+					{
+						if (UObject* Object = WeakObject.Get())
+						{
+							SequencerPtr->RestorePreAnimatedState();
+						}
+					}
+
 					SequencerPtr->NotifyMovieSceneDataChanged(EMovieSceneDataChangeType::MovieSceneStructureItemsChanged);
 				}
 			),
@@ -242,6 +251,22 @@ TOptional<FTransform> FTransformSection::GetCurrentValue() const
 	return TOptional<FTransform>();
 }
 
+TOptional<FRotator> FTransformSection::GetCurrentRotator() const
+{
+	TSharedPtr<ISequencer> Sequencer = WeakSequencer.Pin();
+	if (ensure(PropertyBindings.IsSet()) && Sequencer.IsValid())
+	{
+		for (TWeakObjectPtr<> WeakObject : Sequencer->FindBoundObjects(ObjectBinding, Sequencer->GetFocusedTemplateID()))
+		{
+			if (UObject* Object = WeakObject.Get())
+			{
+				return PropertyBindings->GetCurrentValue<FTransform>(*Object).GetRotation().Rotator();
+			}
+		}
+	}
+	return TOptional<FRotator>();
+}
+
 TOptional<float> FTransformSection::GetTranslationValue(EAxis::Type Axis) const
 {
 	FTransform Transform = GetCurrentValue().Get(FTransform::Identity);
@@ -256,12 +281,19 @@ TOptional<float> FTransformSection::GetTranslationValue(EAxis::Type Axis) const
 
 TOptional<float> FTransformSection::GetRotationValue(EAxis::Type Axis) const
 {
-	FTransform Transform = GetCurrentValue().Get(FTransform::Identity);
+	TOptional<FRotator> Rotator = GetCurrentRotator();
+
+	if (!Rotator.IsSet())
+	{
+		FTransform Transform = GetCurrentValue().Get(FTransform::Identity);
+		Rotator = Transform.GetRotation().Rotator();
+	}
+
 	switch (Axis)
 	{
-	case EAxis::X:		return Transform.GetRotation().Rotator().Roll;
-	case EAxis::Y:		return Transform.GetRotation().Rotator().Pitch;
-	case EAxis::Z:		return Transform.GetRotation().Rotator().Yaw;
+	case EAxis::X:		return Rotator.GetValue().Roll;
+	case EAxis::Y:		return Rotator.GetValue().Pitch;
+	case EAxis::Z:		return Rotator.GetValue().Yaw;
 	}
 	return TOptional<float>();
 }
diff --git a/Engine/Source/Editor/MovieSceneTools/Private/Sections/TransformPropertySection.h b/Engine/Source/Editor/MovieSceneTools/Private/Sections/TransformPropertySection.h
--- a/Engine/Source/Editor/MovieSceneTools/Private/Sections/TransformPropertySection.h
+++ b/Engine/Source/Editor/MovieSceneTools/Private/Sections/TransformPropertySection.h
@@ -38,6 +38,7 @@ public:
 protected:
 
 	virtual TOptional<FTransform> GetCurrentValue() const;
+	virtual TOptional<FRotator> GetCurrentRotator() const;
 
 protected:
 
diff --git a/Engine/Source/Editor/MovieSceneTools/Private/TrackEditors/ActorPickerTrackEditor.cpp b/Engine/Source/Editor/MovieSceneTools/Private/TrackEditors/ActorPickerTrackEditor.cpp
--- a/Engine/Source/Editor/MovieSceneTools/Private/TrackEditors/ActorPickerTrackEditor.cpp
+++ b/Engine/Source/Editor/MovieSceneTools/Private/TrackEditors/ActorPickerTrackEditor.cpp
@@ -21,6 +21,8 @@
 #include "SceneOutlinerModule.h"
 #include "Private/SSocketChooser.h"
 #include "LevelEditor.h"
+#include "MovieSceneObjectBindingIDPicker.h"
+#include "MovieSceneToolHelpers.h"
 
 #define LOCTEXT_NAMESPACE "FActorPickerTrackEditor"
 
@@ -56,68 +58,105 @@ void FActorPickerTrackEditor::ShowActorSubMenu(FMenuBuilder& MenuBuilder, FGuid
 		}
 	};
 
-	using namespace SceneOutliner;
-
-	SceneOutliner::FInitializationOptions InitOptions;
+	auto CreateNewBinding = 
+		[this, ObjectBinding, Section](FMenuBuilder& SubMenuBuilder)
 	{
-		InitOptions.Mode = ESceneOutlinerMode::ActorPicker;			
-		InitOptions.bShowHeaderRow = false;
-		InitOptions.bFocusSearchBoxWhenOpened = true;
-		InitOptions.bShowTransient = true;
-		InitOptions.bShowCreateNewFolder = false;
-		// Only want the actor label column
-		InitOptions.ColumnMap.Add(FBuiltInColumnTypes::Label(), FColumnInfo(EColumnVisibility::Visible, 0));
-
-		// Only display Actors that we can attach too
-		InitOptions.Filters->AddFilterPredicate( SceneOutliner::FActorFilterPredicate::CreateSP(this, &FActorPickerTrackEditor::IsActorPickable, ObjectBinding, Section) );
-	}		
-
-	// Actor selector to allow the user to choose a parent actor
-	FSceneOutlinerModule& SceneOutlinerModule = FModuleManager::LoadModuleChecked<FSceneOutlinerModule>( "SceneOutliner" );
-
-	TSharedRef< SWidget > MenuWidget = 
-		SNew(SHorizontalBox)
-
-		+SHorizontalBox::Slot()
-		.AutoWidth()
-		[
-			SNew(SBox)
-			.MaxDesiredHeight(400.0f)
-			.WidthOverride(300.0f)
+		using namespace SceneOutliner;
+
+		SceneOutliner::FInitializationOptions InitOptions;
+		{
+			InitOptions.Mode = ESceneOutlinerMode::ActorPicker;			
+			InitOptions.bShowHeaderRow = false;
+			InitOptions.bFocusSearchBoxWhenOpened = true;
+			InitOptions.bShowTransient = true;
+			InitOptions.bShowCreateNewFolder = false;
+			// Only want the actor label column
+			InitOptions.ColumnMap.Add(FBuiltInColumnTypes::Label(), FColumnInfo(EColumnVisibility::Visible, 0));
+
+			// Only display Actors that we can attach too
+			InitOptions.Filters->AddFilterPredicate( SceneOutliner::FActorFilterPredicate::CreateSP(this, &FActorPickerTrackEditor::IsActorPickable, ObjectBinding, Section) );
+		}		
+
+		// Actor selector to allow the user to choose a parent actor
+		FSceneOutlinerModule& SceneOutlinerModule = FModuleManager::LoadModuleChecked<FSceneOutlinerModule>( "SceneOutliner" );
+
+		TSharedRef< SWidget > MenuWidget = 
+			SNew(SHorizontalBox)
+
+			+SHorizontalBox::Slot()
+			.AutoWidth()
 			[
-				SceneOutlinerModule.CreateSceneOutliner(
-					InitOptions,
-					FOnActorPicked::CreateSP(this, &FActorPickerTrackEditor::ActorPicked, ObjectBinding, Section )
-					)
+				SNew(SBox)
+				.MaxDesiredHeight(400.0f)
+				.WidthOverride(300.0f)
+				[
+					SceneOutlinerModule.CreateSceneOutliner(
+						InitOptions,
+						FOnActorPicked::CreateSP(this, &FActorPickerTrackEditor::ActorPicked, ObjectBinding, Section )
+						)
+				]
 			]
-		]
 	
-		+SHorizontalBox::Slot()
-		.VAlign(VAlign_Top)
-		.AutoWidth()
-		[
-			SNew(SVerticalBox)
-
-			+SVerticalBox::Slot()
-			.AutoHeight()
-			.Padding(4.0f, 0.0f, 0.0f, 0.0f)
+			+SHorizontalBox::Slot()
+			.VAlign(VAlign_Top)
+			.AutoWidth()
 			[
-				SNew(SButton)
-				.ToolTipText( LOCTEXT( "PickButtonLabel", "Pick a parent actor to attach to") )
-				.ButtonStyle(FEditorStyle::Get(), "HoverHintOnly")
-				.OnClicked(FOnClicked::CreateStatic(&Local::OnInteractiveActorPickerClicked, this, ObjectBinding, Section))
-				.ContentPadding(4.0f)
-				.ForegroundColor(FSlateColor::UseForeground())
-				.IsFocusable(false)
+				SNew(SVerticalBox)
+
+				+SVerticalBox::Slot()
+				.AutoHeight()
+				.Padding(4.0f, 0.0f, 0.0f, 0.0f)
 				[
-					SNew(SImage)
-					.Image(FEditorStyle::GetBrush("PropertyWindow.Button_PickActorInteractive"))
-					.ColorAndOpacity(FSlateColor::UseForeground())
+					SNew(SButton)
+					.ToolTipText( LOCTEXT( "PickButtonLabel", "Pick a parent actor to attach to") )
+					.ButtonStyle(FEditorStyle::Get(), "HoverHintOnly")
+					.OnClicked(FOnClicked::CreateStatic(&Local::OnInteractiveActorPickerClicked, this, ObjectBinding, Section))
+					.ContentPadding(4.0f)
+					.ForegroundColor(FSlateColor::UseForeground())
+					.IsFocusable(false)
+					[
+						SNew(SImage)
+						.Image(FEditorStyle::GetBrush("PropertyWindow.Button_PickActorInteractive"))
+						.ColorAndOpacity(FSlateColor::UseForeground())
+					]
 				]
-			]
-		];
+			];
+
+		SubMenuBuilder.AddWidget(MenuWidget, FText::GetEmpty(), false);
+	};
+
+	TSharedPtr<ISequencer> SequencerPtr = GetSequencer();
+
+	// Always recreate the binding picker to ensure we have the correct sequence ID
+	BindingIDPicker = MakeShared<FTrackEditorBindingIDPicker>(SequencerPtr->GetFocusedTemplateID(), SequencerPtr);
+	BindingIDPicker->OnBindingPicked().AddRaw(this, &FActorPickerTrackEditor::ExistingBindingPicked, ObjectBinding);
+
+	FText ExistingBindingText = LOCTEXT("ExistingBinding", "Existing Binding");
+	FText NewBindingText = LOCTEXT("NewBinding", "New Binding");
 
-	MenuBuilder.AddWidget(MenuWidget, FText::GetEmpty(), false);
+	const bool bHasExistingBindings = !BindingIDPicker->IsEmpty();
+	if (bHasExistingBindings)
+	{
+		MenuBuilder.AddSubMenu(
+			NewBindingText,
+			LOCTEXT("NewBinding_Tip", "Add a new section by creating a new binding to an object in the world."),
+			FNewMenuDelegate::CreateLambda(CreateNewBinding)
+		);
+
+		MenuBuilder.BeginSection(NAME_None, ExistingBindingText);
+		{
+			BindingIDPicker->GetPickerMenu(MenuBuilder);
+		}
+		MenuBuilder.EndSection();
+	}
+	else
+	{
+		MenuBuilder.BeginSection(NAME_None, NewBindingText);
+		{
+			CreateNewBinding(MenuBuilder);
+		}
+		MenuBuilder.EndSection();
+	}
 }
 
 class SComponentChooserPopup : public SCompoundWidget
@@ -216,11 +255,39 @@ public:
 };
 
 void FActorPickerTrackEditor::ActorPicked(AActor* ParentActor, FGuid ObjectGuid, UMovieSceneSection* Section)
+{
+	ActorPickerIDPicked(FActorPickerID(ParentActor, FMovieSceneObjectBindingID()), ObjectGuid, Section);
+}
+
+
+void FActorPickerTrackEditor::ExistingBindingPicked(FMovieSceneObjectBindingID ExistingBindingID, FGuid ObjectBinding)
+{
+	TSharedPtr<ISequencer> SequencerPtr = GetSequencer();
+
+	TArrayView<TWeakObjectPtr<UObject>> RuntimeObjects = SequencerPtr->FindBoundObjects(ExistingBindingID.GetGuid(), ExistingBindingID.GetSequenceID());
+	for (auto RuntimeObject : RuntimeObjects)
+	{
+		if (RuntimeObject.IsValid())
+		{
+			AActor* Actor = Cast<AActor>(RuntimeObject.Get());
+			if (Actor)
+			{
+				ActorPickerIDPicked(FActorPickerID(Actor, ExistingBindingID), ObjectBinding, nullptr);
+				return;
+			}
+		}
+	}
+
+	ActorPickerIDPicked(FActorPickerID(nullptr, ExistingBindingID), ObjectBinding, nullptr);
+}
+
+
+void FActorPickerTrackEditor::ActorPickerIDPicked(FActorPickerID ActorPickerID, FGuid ObjectGuid, UMovieSceneSection* Section)
 {
 	TArray<USceneComponent*> ComponentsWithSockets;
-	if (ParentActor != NULL)
+	if (ActorPickerID.ActorPicked.IsValid())
 	{
-		TInlineComponentArray<USceneComponent*> Components(ParentActor);
+		TInlineComponentArray<USceneComponent*> Components(ActorPickerID.ActorPicked.Get());
 
 		for(USceneComponent* Component : Components)
 		{
@@ -234,7 +301,7 @@ void FActorPickerTrackEditor::ActorPicked(AActor* ParentActor, FGuid ObjectGuid,
 	if (ComponentsWithSockets.Num() == 0)
 	{
 		FSlateApplication::Get().DismissAllMenus();
-		ActorSocketPicked( NAME_None, nullptr, ParentActor, ObjectGuid, Section );
+		ActorSocketPicked( NAME_None, nullptr, ActorPickerID, ObjectGuid, Section );
 		return;
 	}
 
@@ -247,8 +314,8 @@ void FActorPickerTrackEditor::ActorPicked(AActor* ParentActor, FGuid ObjectGuid,
 	{			
 		MenuWidget = 
 			SNew(SComponentChooserPopup)
-			.Actor(ParentActor)
-			.OnComponentChosen(this, &FActorPickerTrackEditor::ActorComponentPicked, ParentActor, ObjectGuid, Section);		
+			.Actor(ActorPickerID.ActorPicked.Get())
+			.OnComponentChosen(this, &FActorPickerTrackEditor::ActorComponentPicked, ActorPickerID, ObjectGuid, Section);		
 
 		// Create as context menu
 		FSlateApplication::Get().PushMenu(
@@ -261,22 +328,25 @@ void FActorPickerTrackEditor::ActorPicked(AActor* ParentActor, FGuid ObjectGuid,
 	}
 	else
 	{
-		ActorComponentPicked(ComponentsWithSockets[0]->GetName(), ParentActor, ObjectGuid, Section);
+		ActorComponentPicked(ComponentsWithSockets[0]->GetName(), ActorPickerID, ObjectGuid, Section);
 	}
 }
 
 
-void FActorPickerTrackEditor::ActorComponentPicked(FString ComponentName, AActor* ParentActor, FGuid ObjectGuid, UMovieSceneSection* Section)
+void FActorPickerTrackEditor::ActorComponentPicked(FString ComponentName, FActorPickerID ActorPickerID, FGuid ObjectGuid, UMovieSceneSection* Section)
 {
 	USceneComponent* ComponentWithSockets = nullptr;
-	TInlineComponentArray<USceneComponent*> Components(ParentActor);
-
-	for(USceneComponent* Component : Components)
+	if (ActorPickerID.ActorPicked.IsValid())
 	{
-		if (Component->GetName() == ComponentName)
+		TInlineComponentArray<USceneComponent*> Components(ActorPickerID.ActorPicked.Get());
+	
+		for(USceneComponent* Component : Components)
 		{
-			ComponentWithSockets = Component;
-			break;
+			if (Component->GetName() == ComponentName)
+			{
+				ComponentWithSockets = Component;
+				break;
+			}
 		}
 	}
 
@@ -293,7 +363,7 @@ void FActorPickerTrackEditor::ActorComponentPicked(FString ComponentName, AActor
 	MenuWidget = 
 		SNew(SSocketChooserPopup)
 		.SceneComponent(ComponentWithSockets)
-		.OnSocketChosen(this, &FActorPickerTrackEditor::ActorSocketPicked, ComponentWithSockets, ParentActor, ObjectGuid, Section);		
+		.OnSocketChosen(this, &FActorPickerTrackEditor::ActorSocketPicked, ComponentWithSockets, ActorPickerID, ObjectGuid, Section);		
 
 	// Create as context menu
 	FSlateApplication::Get().PushMenu(
diff --git a/Engine/Source/Editor/MovieSceneTools/Private/TrackEditors/ActorPickerTrackEditor.h b/Engine/Source/Editor/MovieSceneTools/Private/TrackEditors/ActorPickerTrackEditor.h
--- a/Engine/Source/Editor/MovieSceneTools/Private/TrackEditors/ActorPickerTrackEditor.h
+++ b/Engine/Source/Editor/MovieSceneTools/Private/TrackEditors/ActorPickerTrackEditor.h
@@ -10,6 +10,8 @@
 class AActor;
 class FMenuBuilder;
 class USceneComponent;
+class FTrackEditorBindingIDPicker;
+struct FMovieSceneObjectBindingID;
 
 /**
  * Track editor actor picker
@@ -18,26 +20,42 @@ class FActorPickerTrackEditor : public FMovieSceneTrackEditor
 {
 public:
 
+	struct FActorPickerID
+	{
+		FActorPickerID(AActor* InActorPicked, FMovieSceneObjectBindingID InExistingBindingID) : ActorPicked(InActorPicked), ExistingBindingID(InExistingBindingID) {}
+
+		TWeakObjectPtr<AActor> ActorPicked;
+		FMovieSceneObjectBindingID ExistingBindingID;
+	};
+
 	FActorPickerTrackEditor(TSharedRef<ISequencer> InSequencer);
 
 	/** Is this actor pickable? */
 	virtual bool IsActorPickable( const AActor* const ParentActor, FGuid ObjectBinding, UMovieSceneSection* InSection) { return false; }
 
 	/** Actor socket was picked */
-	virtual void ActorSocketPicked(const FName SocketName, USceneComponent* Component, AActor* ParentActor, FGuid ObjectBinding, UMovieSceneSection* Section) {}
+	virtual void ActorSocketPicked(const FName SocketName, USceneComponent* Component, FActorPickerID ActorPickerID, FGuid ObjectBinding, UMovieSceneSection* Section) {}
 
 	/** Show a sub menu of the pickable actors */
 	void ShowActorSubMenu(FMenuBuilder& MenuBuilder, FGuid ObjectBinding, UMovieSceneSection* Section);
 
+private: 
+
 	/** Actor was picked */
 	void ActorPicked(AActor* ParentActor, FGuid ObjectBinding, UMovieSceneSection* Section);
 
-	/** Actor component was picked */
-	void ActorComponentPicked(FString ComponentName, AActor* ParentActor, FGuid ObjectBinding, UMovieSceneSection* Section);
+	/** Existing binding was picked */
+	void ExistingBindingPicked(FMovieSceneObjectBindingID ExistingBindingID, FGuid ObjectBinding);
 
-private: 
+	/** Actor and/or existing binding was picked */
+	void ActorPickerIDPicked(FActorPickerID, FGuid ObjectBinding, UMovieSceneSection* Section);
+
+	/** Actor component was picked */
+	void ActorComponentPicked(FString ComponentName, FActorPickerID ActorPickerID, FGuid ObjectBinding, UMovieSceneSection* Section);
 
 	/** Interactively pick an actor from the viewport */
 	void PickActorInteractive(FGuid ObjectBinding, UMovieSceneSection* Section);
 
+	/** A binding ID picker that allows us to create a new section from an existing binding */
+	TSharedPtr<FTrackEditorBindingIDPicker> BindingIDPicker;
 };
diff --git a/Engine/Source/Editor/MovieSceneTools/Private/TrackEditors/AttachTrackEditor.cpp b/Engine/Source/Editor/MovieSceneTools/Private/TrackEditors/AttachTrackEditor.cpp
--- a/Engine/Source/Editor/MovieSceneTools/Private/TrackEditors/AttachTrackEditor.cpp
+++ b/Engine/Source/Editor/MovieSceneTools/Private/TrackEditors/AttachTrackEditor.cpp
@@ -8,6 +8,8 @@
 #include "Tracks/MovieScene3DAttachTrack.h"
 #include "Sections/MovieScene3DAttachSection.h"
 #include "ActorEditorUtils.h"
+#include "MovieSceneObjectBindingIDPicker.h"
+#include "MovieSceneToolHelpers.h"
 
 
 #define LOCTEXT_NAMESPACE "F3DAttachTrackEditor"
@@ -39,7 +41,7 @@ public:
 			TSharedPtr<ISequencer> Sequencer = AttachTrackEditor->GetSequencer();
 			if (Sequencer.IsValid())
 			{
-				TArrayView<TWeakObjectPtr<UObject>> RuntimeObjects = Sequencer->FindBoundObjects(AttachSection->GetConstraintId(), Sequencer->GetFocusedTemplateID());
+				TArrayView<TWeakObjectPtr<UObject>> RuntimeObjects = Sequencer->FindBoundObjects(AttachSection->GetConstraintBindingID().GetGuid(), AttachSection->GetConstraintBindingID().GetSequenceID());
 				if (RuntimeObjects.Num() == 1 && RuntimeObjects[0].IsValid())
 				{
 					if (AActor* Actor = Cast<AActor>(RuntimeObjects[0].Get()))
@@ -167,21 +169,33 @@ bool F3DAttachTrackEditor::IsActorPickable(const AActor* const ParentActor, FGui
 }
 
 
-void F3DAttachTrackEditor::ActorSocketPicked(const FName SocketName, USceneComponent* Component, AActor* ParentActor, FGuid ObjectGuid, UMovieSceneSection* Section)
+void F3DAttachTrackEditor::ActorSocketPicked(const FName SocketName, USceneComponent* Component, FActorPickerID ActorPickerID, FGuid ObjectGuid, UMovieSceneSection* Section)
 {
 	if (Section != nullptr)
 	{
 		const FScopedTransaction Transaction(LOCTEXT("UndoSetAttach", "Set Attach"));
 
 		UMovieScene3DAttachSection* AttachSection = (UMovieScene3DAttachSection*)(Section);
-		FGuid ActorId = FindOrCreateHandleToObject(ParentActor).Handle;
 
-		if (ActorId.IsValid())
+		FMovieSceneObjectBindingID ConstraintBindingID;
+
+		if (ActorPickerID.ExistingBindingID.IsValid())
+		{
+			ConstraintBindingID = ActorPickerID.ExistingBindingID;
+		}
+		else if (ActorPickerID.ActorPicked.IsValid())
+		{
+			FGuid ParentActorId = FindOrCreateHandleToObject(ActorPickerID.ActorPicked.Get()).Handle;
+			ConstraintBindingID = FMovieSceneObjectBindingID(ParentActorId, MovieSceneSequenceID::Root);
+		}
+
+		if (ConstraintBindingID.IsValid())
 		{
-			AttachSection->SetConstraintId(ActorId);
-			AttachSection->AttachSocketName = SocketName;
-			AttachSection->AttachComponentName = Component ? Component->GetFName() : NAME_None;
+			AttachSection->SetConstraintBindingID(ConstraintBindingID);
 		}
+
+		AttachSection->AttachSocketName = SocketName;			
+		AttachSection->AttachComponentName = Component ? Component->GetFName() : NAME_None;
 	}
 	else if (ObjectGuid.IsValid())
 	{
@@ -191,24 +205,29 @@ void F3DAttachTrackEditor::ActorSocketPicked(const FName SocketName, USceneCompo
 			OutObjects.Add(Object);
 		}
 
-		AnimatablePropertyChanged( FOnKeyProperty::CreateRaw( this, &F3DAttachTrackEditor::AddKeyInternal, OutObjects, SocketName, Component ? Component->GetFName() : NAME_None, ParentActor) );
+		AnimatablePropertyChanged( FOnKeyProperty::CreateRaw( this, &F3DAttachTrackEditor::AddKeyInternal, OutObjects, SocketName, Component ? Component->GetFName() : NAME_None, ActorPickerID) );
 	}
 }
 
-FKeyPropertyResult F3DAttachTrackEditor::AddKeyInternal( float KeyTime, const TArray<TWeakObjectPtr<UObject>> Objects, const FName SocketName, const FName ComponentName, AActor* ParentActor)
+FKeyPropertyResult F3DAttachTrackEditor::AddKeyInternal( float KeyTime, const TArray<TWeakObjectPtr<UObject>> Objects, const FName SocketName, const FName ComponentName, FActorPickerID ActorPickerID)
 {
 	FKeyPropertyResult KeyPropertyResult;
 
-	FGuid ParentActorId;
+	FMovieSceneObjectBindingID ConstraintBindingID;
 
-	if (ParentActor != nullptr)
+	if (ActorPickerID.ExistingBindingID.IsValid())
 	{
-		FFindOrCreateHandleResult HandleResult = FindOrCreateHandleToObject(ParentActor);
-		ParentActorId = HandleResult.Handle;
+		ConstraintBindingID = ActorPickerID.ExistingBindingID;
+	}
+	else if (ActorPickerID.ActorPicked.IsValid())
+	{
+		FFindOrCreateHandleResult HandleResult = FindOrCreateHandleToObject(ActorPickerID.ActorPicked.Get());
+		FGuid ParentActorId = HandleResult.Handle;
 		KeyPropertyResult.bHandleCreated |= HandleResult.bWasCreated;
+		ConstraintBindingID = FMovieSceneObjectBindingID(ParentActorId, MovieSceneSequenceID::Root);
 	}
 
-	if (!ParentActorId.IsValid())
+	if (!ConstraintBindingID.IsValid())
 	{
 		return KeyPropertyResult;
 	}
@@ -244,7 +263,7 @@ FKeyPropertyResult F3DAttachTrackEditor::AddKeyInternal( float KeyTime, const TA
 					}
 				}
 
-				Cast<UMovieScene3DAttachTrack>(Track)->AddConstraint( KeyTime, AttachEndTime, SocketName, ComponentName, ParentActorId );
+				Cast<UMovieScene3DAttachTrack>(Track)->AddConstraint( KeyTime, AttachEndTime, SocketName, ComponentName, ConstraintBindingID);
 				KeyPropertyResult.bTrackModified = true;
 			}
 		}
@@ -253,5 +272,4 @@ FKeyPropertyResult F3DAttachTrackEditor::AddKeyInternal( float KeyTime, const TA
 	return KeyPropertyResult;
 }
 
-
 #undef LOCTEXT_NAMESPACE
diff --git a/Engine/Source/Editor/MovieSceneTools/Private/TrackEditors/AttachTrackEditor.h b/Engine/Source/Editor/MovieSceneTools/Private/TrackEditors/AttachTrackEditor.h
--- a/Engine/Source/Editor/MovieSceneTools/Private/TrackEditors/AttachTrackEditor.h
+++ b/Engine/Source/Editor/MovieSceneTools/Private/TrackEditors/AttachTrackEditor.h
@@ -51,10 +51,10 @@ public:
 
 	// FTrackEditorActorPicker
 	virtual bool IsActorPickable( const AActor* const ParentActor, FGuid ObjectBinding, UMovieSceneSection* InSection ) override;
-	virtual void ActorSocketPicked(const FName SocketName, USceneComponent* Component, AActor* ParentActor, FGuid ObjectBinding, UMovieSceneSection* Section) override;
+	virtual void ActorSocketPicked(const FName SocketName, USceneComponent* Component, FActorPickerID ActorPickerID, FGuid ObjectBinding, UMovieSceneSection* Section) override;
 
 private:
 
 	/** Delegate for AnimatablePropertyChanged in AddKey */
-	FKeyPropertyResult AddKeyInternal(float KeyTime, const TArray<TWeakObjectPtr<UObject>> Objects, const FName SocketName, const FName ComponentName, AActor* ParentActor);
+	FKeyPropertyResult AddKeyInternal(float KeyTime, const TArray<TWeakObjectPtr<UObject>> Objects, const FName SocketName, const FName ComponentName, FActorPickerID ActorPickerID);
 };
diff --git a/Engine/Source/Editor/MovieSceneTools/Private/TrackEditors/CameraCutTrackEditor.cpp b/Engine/Source/Editor/MovieSceneTools/Private/TrackEditors/CameraCutTrackEditor.cpp
--- a/Engine/Source/Editor/MovieSceneTools/Private/TrackEditors/CameraCutTrackEditor.cpp
+++ b/Engine/Source/Editor/MovieSceneTools/Private/TrackEditors/CameraCutTrackEditor.cpp
@@ -19,6 +19,8 @@
 #include "SceneOutlinerPublicTypes.h"
 #include "SceneOutlinerModule.h"
 #include "TrackEditorThumbnail/TrackEditorThumbnailPool.h"
+#include "MovieSceneObjectBindingIDPicker.h"
+#include "MovieSceneToolHelpers.h"
 #include "DragAndDrop/ActorDragDropGraphEdOp.h"
 
 #define LOCTEXT_NAMESPACE "FCameraCutTrackEditor"
@@ -57,6 +59,7 @@ void FCameraCutTrackCommands::RegisterCommands()
 	UI_COMMAND( ToggleLockCamera, "Toggle Lock Camera", "Toggle locking the viewport to the camera cut track.", EUserInterfaceActionType::Button, FInputChord(EModifierKey::Control, EKeys::L) );
 }
 
+
 /* FCameraCutTrackEditor structors
  *****************************************************************************/
 
@@ -280,7 +283,7 @@ FKeyPropertyResult FCameraCutTrackEditor::AddKeyInternal( float KeyTime, const F
 	UMovieSceneCameraCutTrack* CameraCutTrack = FindOrCreateCameraCutTrack();
 	const TArray<UMovieSceneSection*>& AllSections = CameraCutTrack->GetAllSections();
 
-	CameraCutTrack->AddNewCameraCut(ObjectGuid, KeyTime);
+	CameraCutTrack->AddNewCameraCut(FMovieSceneObjectBindingID(ObjectGuid, MovieSceneSequenceID::Root), KeyTime);
 	KeyPropertyResult.bTrackModified = true;
 
 	return KeyPropertyResult;
@@ -341,47 +344,100 @@ TSharedRef<SWidget> FCameraCutTrackEditor::HandleAddCameraCutComboButtonGetMenuC
 {
 	FMenuBuilder MenuBuilder(true, nullptr);
 
-	using namespace SceneOutliner;
+	auto CreateNewCamera =
+		[this](FMenuBuilder& SubMenuBuilder)
+		{
+			using namespace SceneOutliner;
 
-	SceneOutliner::FInitializationOptions InitOptions;
+			SceneOutliner::FInitializationOptions InitOptions;
+			{
+				InitOptions.Mode = ESceneOutlinerMode::ActorPicker;
+				InitOptions.bShowHeaderRow = false;
+				InitOptions.bFocusSearchBoxWhenOpened = true;
+				InitOptions.bShowTransient = true;
+				InitOptions.bShowCreateNewFolder = false;
+				// Only want the actor label column
+				InitOptions.ColumnMap.Add(FBuiltInColumnTypes::Label(), FColumnInfo(EColumnVisibility::Visible, 0));
+
+				// Only display Actors that we can attach too
+				InitOptions.Filters->AddFilterPredicate( SceneOutliner::FActorFilterPredicate::CreateRaw(this, &FCameraCutTrackEditor::IsCameraPickable) );
+			}		
+
+			// Actor selector to allow the user to choose a parent actor
+			FSceneOutlinerModule& SceneOutlinerModule = FModuleManager::LoadModuleChecked<FSceneOutlinerModule>( "SceneOutliner" );
+
+			TSharedRef< SWidget > MenuWidget = 
+				SNew(SHorizontalBox)
+
+				+SHorizontalBox::Slot()
+				.AutoWidth()
+				[
+					SNew(SBox)
+					.MaxDesiredHeight(400.0f)
+					.WidthOverride(300.0f)
+					[
+						SceneOutlinerModule.CreateSceneOutliner(
+							InitOptions,
+							FOnActorPicked::CreateSP(this, &FCameraCutTrackEditor::HandleAddCameraCutComboButtonMenuEntryExecute )
+							)
+					]
+				];
+			SubMenuBuilder.AddWidget(MenuWidget, FText::GetEmpty(), false);
+		};
+
+	TSharedPtr<ISequencer> SequencerPtr = GetSequencer();
+
+	// Always recreate the binding picker to ensure we have the correct sequence ID
+	BindingIDPicker = MakeShared<FTrackEditorBindingIDPicker>(SequencerPtr->GetFocusedTemplateID(), SequencerPtr);
+	BindingIDPicker->OnBindingPicked().AddRaw(this, &FCameraCutTrackEditor::CreateNewSectionFromBinding);
+
+	FText ExistingBindingText = LOCTEXT("ExistingBinding", "Existing Binding");
+	FText NewBindingText = LOCTEXT("NewBinding", "New Binding");
+
+	const bool bHasExistingBindings = !BindingIDPicker->IsEmpty();
+	if (bHasExistingBindings)
 	{
-		InitOptions.Mode = ESceneOutlinerMode::ActorPicker;
-		InitOptions.bShowHeaderRow = false;
-		InitOptions.bFocusSearchBoxWhenOpened = true;
-		InitOptions.bShowTransient = true;
-		InitOptions.bShowCreateNewFolder = false;
-		// Only want the actor label column
-		InitOptions.ColumnMap.Add(FBuiltInColumnTypes::Label(), FColumnInfo(EColumnVisibility::Visible, 0));
-
-		// Only display Actors that we can attach too
-		InitOptions.Filters->AddFilterPredicate( SceneOutliner::FActorFilterPredicate::CreateSP(this, &FCameraCutTrackEditor::IsCameraPickable) );
-	}		
-
-	// Actor selector to allow the user to choose a parent actor
-	FSceneOutlinerModule& SceneOutlinerModule = FModuleManager::LoadModuleChecked<FSceneOutlinerModule>( "SceneOutliner" );
-
-	TSharedRef< SWidget > MenuWidget = 
-		SNew(SHorizontalBox)
-
-		+SHorizontalBox::Slot()
-		.AutoWidth()
-		[
-			SNew(SBox)
-			.MaxDesiredHeight(400.0f)
-			.WidthOverride(300.0f)
-			[
-				SceneOutlinerModule.CreateSceneOutliner(
-					InitOptions,
-					FOnActorPicked::CreateSP(this, &FCameraCutTrackEditor::HandleAddCameraCutComboButtonMenuEntryExecute )
-					)
-			]
-		];
-
-	MenuBuilder.AddWidget(MenuWidget, FText::GetEmpty(), false);
+		MenuBuilder.AddSubMenu(
+			NewBindingText,
+			LOCTEXT("NewBinding_Tip", "Add a new camera cut by creating a new binding to an object in the world."),
+			FNewMenuDelegate::CreateLambda(CreateNewCamera)
+		);
+
+		MenuBuilder.BeginSection(NAME_None, ExistingBindingText);
+		{
+			BindingIDPicker->GetPickerMenu(MenuBuilder);
+		}
+		MenuBuilder.EndSection();
+	}
+	else
+	{
+		MenuBuilder.BeginSection(NAME_None, NewBindingText);
+		{
+			CreateNewCamera(MenuBuilder);
+		}
+		MenuBuilder.EndSection();
+	}
+	
 	return MenuBuilder.MakeWidget();
 }
 
 
+void FCameraCutTrackEditor::CreateNewSectionFromBinding(FMovieSceneObjectBindingID InBindingID)
+{
+	auto CreateNewSection = [this, InBindingID](float KeyTime)
+	{
+		FKeyPropertyResult KeyPropertyResult;
+
+		FindOrCreateCameraCutTrack()->AddNewCameraCut(InBindingID, KeyTime);
+		KeyPropertyResult.bTrackModified = true;
+
+		return KeyPropertyResult;
+	};
+
+	AnimatablePropertyChanged(FOnKeyProperty::CreateLambda(CreateNewSection));
+}
+
+
 void FCameraCutTrackEditor::HandleAddCameraCutComboButtonMenuEntryExecute(AActor* Camera)
 {
 	FGuid ObjectGuid = FindOrCreateHandleToObject(Camera).Handle;
diff --git a/Engine/Source/Editor/MovieSceneTools/Private/TrackEditors/CameraCutTrackEditor.h b/Engine/Source/Editor/MovieSceneTools/Private/TrackEditors/CameraCutTrackEditor.h
--- a/Engine/Source/Editor/MovieSceneTools/Private/TrackEditors/CameraCutTrackEditor.h
+++ b/Engine/Source/Editor/MovieSceneTools/Private/TrackEditors/CameraCutTrackEditor.h
@@ -11,12 +11,16 @@
 #include "ISequencerSection.h"
 #include "ISequencerTrackEditor.h"
 #include "MovieSceneTrackEditor.h"
+#include "MovieSceneToolHelpers.h"
 
 class AActor;
 class FMenuBuilder;
 class FTrackEditorThumbnailPool;
 class UFactory;
 class UMovieSceneCameraCutTrack;
+class FCameraCutTrackEditor;
+class FTrackEditorBindingIDPicker;
+struct FMovieSceneObjectBindingID;
 
 /**
  * Tools for camera cut tracks.
@@ -87,6 +91,9 @@ private:
 	/** Callback for executing a menu entry in the "Add Camera Cut" combo button. */
 	void HandleAddCameraCutComboButtonMenuEntryExecute(AActor* Camera);
 
+	/** Called to create a new section for the specified binding ID. */
+	void CreateNewSectionFromBinding(FMovieSceneObjectBindingID InBindingID);
+
 	/** Delegate for camera button lock state */
 	ECheckBoxState IsCameraLocked() const; 
 
@@ -103,4 +110,7 @@ private:
 
 	/** The Thumbnail pool which draws all the viewport thumbnails for the camera cut track. */
 	TSharedPtr<FTrackEditorThumbnailPool> ThumbnailPool;
+
+	/** A binding ID picker that allows us to create a new section from an existing binding */
+	TSharedPtr<FTrackEditorBindingIDPicker> BindingIDPicker;
 };
