Copying //UE4/Dev-Sequencer to //UE4/Dev-Main (Source: //UE4/Dev-Sequencer @ 3946692)

diff --git a/Engine/Source/Editor/MovieSceneCaptureDialog/Private/MovieSceneCaptureDialogModule.cpp b/Engine/Source/Editor/MovieSceneCaptureDialog/Private/MovieSceneCaptureDialogModule.cpp
--- a/Engine/Source/Editor/MovieSceneCaptureDialog/Private/MovieSceneCaptureDialogModule.cpp
+++ b/Engine/Source/Editor/MovieSceneCaptureDialog/Private/MovieSceneCaptureDialogModule.cpp
@@ -18,6 +18,8 @@
 #include "Widgets/SWindow.h"
 #include "Widgets/Text/STextBlock.h"
 #include "Framework/Docking/TabManager.h"
+#include "Widgets/Input/NumericTypeInterface.h"
+#include "FrameNumberDetailsCustomization.h"
 
 #include "HAL/IConsoleManager.h"
 #include "HAL/PlatformProcess.h"
@@ -59,6 +61,7 @@
 #include "ErrorCodes.h"
 
 #include "GameFramework/WorldSettings.h"
+#include "FrameNumberNumericInterface.h"
 
 #define LOCTEXT_NAMESPACE "MovieSceneCaptureDialog"
 
@@ -71,6 +74,7 @@ class SRenderMovieSceneSettings : public SCompoundWidget, public FGCObject
 	SLATE_BEGIN_ARGS(SRenderMovieSceneSettings) : _InitialObject(nullptr) {}
 		SLATE_EVENT(FOnStartCapture, OnStartCapture)
 		SLATE_ARGUMENT(UMovieSceneCapture*, InitialObject)
+		SLATE_ARGUMENT(TSharedPtr<INumericTypeInterface<double>>, NumericTypeInterface)
 	SLATE_END_ARGS()
 
 	void Construct(const FArguments& InArgs)
@@ -84,6 +88,7 @@ class SRenderMovieSceneSettings : public SCompoundWidget, public FGCObject
 		DetailsViewArgs.ViewIdentifier = "RenderMovieScene";
 
 		DetailView = PropertyEditor.CreateDetailView(DetailsViewArgs);
+		DetailView->RegisterInstancedCustomPropertyTypeLayout("FrameNumber", FOnGetPropertyTypeCustomizationInstance::CreateStatic(&FFrameNumberDetailsCustomization::MakeInstance, InArgs._NumericTypeInterface));
 
 		OnStartCapture = InArgs._OnStartCapture;
 
@@ -545,7 +550,7 @@ private:
 
 					FVector2D PreviewWindowSize(Settings.Resolution.ResX, Settings.Resolution.ResY);
 
-					// Keep scaling down the window size while we're bigger than half the destop width/height
+					// Keep scaling down the window size while we're bigger than half the desktop width/height
 					{
 						FDisplayMetrics DisplayMetrics;
 						FSlateApplication::Get().GetDisplayMetrics(DisplayMetrics);
@@ -654,7 +659,7 @@ class FMovieSceneCaptureDialogModule : public IMovieSceneCaptureDialogModule
 		return Pinned.IsValid() ? Pinned->GetWorld() : nullptr;
 	}
 
-	virtual void OpenDialog(const TSharedRef<FTabManager>& TabManager, UMovieSceneCapture* CaptureObject) override
+	virtual void OpenDialog(const TSharedRef<FTabManager>& TabManager, UMovieSceneCapture* CaptureObject, TSharedPtr<INumericTypeInterface<double>> InNumericTypeInterface) override
 	{
 		// Ensure the session services module is loaded otherwise we won't necessarily receive status updates from the movie capture session
 		FModuleManager::Get().LoadModuleChecked<ISessionServicesModule>("SessionServices").GetSessionManager();
@@ -688,6 +693,7 @@ class FMovieSceneCaptureDialogModule : public IMovieSceneCaptureDialogModule
 		ExistingWindow->SetContent(
 			SNew(SRenderMovieSceneSettings)
 			.InitialObject(CaptureObject)
+			.NumericTypeInterface(InNumericTypeInterface)
 			.OnStartCapture_Raw(this, &FMovieSceneCaptureDialogModule::OnStartCapture)
 		);
 
diff --git a/Engine/Source/Editor/MovieSceneCaptureDialog/Public/MovieSceneCaptureDialogModule.h b/Engine/Source/Editor/MovieSceneCaptureDialog/Public/MovieSceneCaptureDialogModule.h
--- a/Engine/Source/Editor/MovieSceneCaptureDialog/Public/MovieSceneCaptureDialogModule.h
+++ b/Engine/Source/Editor/MovieSceneCaptureDialog/Public/MovieSceneCaptureDialogModule.h
@@ -5,6 +5,7 @@
 #include "CoreMinimal.h"
 #include "Modules/ModuleInterface.h"
 #include "Modules/ModuleManager.h"
+#include "Widgets/Input/NumericTypeInterface.h"
 
 class UMovieSceneCapture;
 
@@ -16,7 +17,7 @@ public:
 		static const FName ModuleName(TEXT("MovieSceneCaptureDialog"));
 		return FModuleManager::LoadModuleChecked<IMovieSceneCaptureDialogModule>(ModuleName);
 	}
-	virtual void OpenDialog(const TSharedRef<class FTabManager>& TabManager, UMovieSceneCapture* CaptureObject) = 0;
+	virtual void OpenDialog(const TSharedRef<class FTabManager>& TabManager, UMovieSceneCapture* CaptureObject, TSharedPtr<INumericTypeInterface<double>> InNumericTypeInterface) = 0;
 
 	/** Get the world we're currently recording from, if an in process record is happening */
 	virtual UWorld* GetCurrentlyRecordingWorld() = 0;
diff --git a/Engine/Source/Editor/MovieSceneTools/Private/Channels/BuiltInChannelEditors.cpp b/Engine/Source/Editor/MovieSceneTools/Private/Channels/BuiltInChannelEditors.cpp
new file mode 100644
--- /dev/null
+++ b/Engine/Source/Editor/MovieSceneTools/Private/Channels/BuiltInChannelEditors.cpp
@@ -0,0 +1,665 @@
+// Copyright 1998-2018 Epic Games, Inc. All Rights Reserved.
+
+#include "BuiltInChannelEditors.h"
+#include "ISequencerChannelInterface.h"
+#include "Widgets/SNullWidget.h"
+#include "ISequencer.h"
+#include "MovieSceneCommonHelpers.h"
+#include "GameFramework/Actor.h"
+#include "EditorStyleSet.h"
+#include "CurveKeyEditors/SNumericKeyEditor.h"
+#include "CurveKeyEditors/SBoolCurveKeyEditor.h"
+#include "CurveKeyEditors/SStringCurveKeyEditor.h"
+#include "CurveKeyEditors/SEnumKeyEditor.h"
+#include "UObject/StructOnScope.h"
+#include "KeyDrawParams.h"
+#include "Framework/MultiBox/MultiBoxBuilder.h"
+#include "Channels/MovieSceneChannelProxy.h"
+#include "Channels/MovieSceneChannelEditorData.h"
+#include "FloatChannelCurveModel.h"
+
+#define LOCTEXT_NAMESPACE "BuiltInChannelEditors"
+
+
+FKeyHandle AddOrUpdateKey(FMovieSceneFloatChannel* Channel, const TMovieSceneExternalValue<float>& ExternalValue, FFrameNumber InTime, ISequencer& Sequencer, const FGuid& InObjectBindingID, FTrackInstancePropertyBindings* PropertyBindings)
+{
+	TOptional<float> Value;
+
+	// Add a key for the current value of the valid first object we can find
+	if (ExternalValue.OnGetExternalValue && InObjectBindingID.IsValid())
+	{
+		for (TWeakObjectPtr<> WeakObject : Sequencer.FindBoundObjects(InObjectBindingID, Sequencer.GetFocusedTemplateID()))
+		{
+			if (UObject* Object = WeakObject.Get())
+			{
+				Value = ExternalValue.OnGetExternalValue(*Object, PropertyBindings);
+				break;
+			}
+		}
+	}
+
+	float NewValue = Channel->GetDefault().Get(0.f);
+	if (Value.IsSet())
+	{
+		NewValue = Value.GetValue();
+	}
+	else
+	{
+		Channel->Evaluate(InTime, NewValue);
+	}
+
+	using namespace MovieScene;
+	return AddKeyToChannel(Channel, InTime, NewValue, Sequencer.GetKeyInterpolation());
+}
+
+FKeyHandle AddOrUpdateKey(FMovieSceneActorReferenceData* Channel, FFrameNumber InTime, ISequencer& Sequencer, const FGuid& InObjectBindingID, FTrackInstancePropertyBindings* PropertyBindings)
+{
+	AActor* CurrentActor = nullptr;
+
+	if (PropertyBindings && InObjectBindingID.IsValid())
+	{
+		for (TWeakObjectPtr<> WeakObject : Sequencer.FindBoundObjects(InObjectBindingID, Sequencer.GetFocusedTemplateID()))
+		{
+			if (UObject* Object = WeakObject.Get())
+			{
+				CurrentActor = PropertyBindings->GetCurrentValue<AActor*>(*Object);
+				break;
+			}
+		}
+	}
+
+	FGuid ThisGuid = CurrentActor ? Sequencer.FindObjectId(*CurrentActor, Sequencer.GetFocusedTemplateID()) : FGuid();
+
+	FMovieSceneObjectBindingID NewValue(ThisGuid, MovieSceneSequenceID::Root, EMovieSceneObjectBindingSpace::Local);
+	int32 NewIndex = Channel->GetInterface().AddKey(InTime, NewValue);
+	return Channel->GetInterface().GetHandle(NewIndex);
+}
+
+bool CanCreateKeyEditor(FMovieSceneBoolChannel*    Channel)
+{
+	return true;
+}
+bool CanCreateKeyEditor(FMovieSceneByteChannel*    Channel)
+{
+	return true;
+}
+bool CanCreateKeyEditor(FMovieSceneIntegerChannel* Channel)
+{
+	return true;
+}
+bool CanCreateKeyEditor(FMovieSceneFloatChannel*   Channel)
+{
+	return true;
+}
+bool CanCreateKeyEditor(FMovieSceneStringChannel*  Channel)
+{
+	return true;
+}
+
+TSharedRef<SWidget> CreateKeyEditor(FMovieSceneBoolChannel*    Channel, const TMovieSceneExternalValue<bool>& ExternalValue, UMovieSceneSection* Section, const FGuid& InObjectBindingID, TWeakPtr<FTrackInstancePropertyBindings> PropertyBindings, TWeakPtr<ISequencer> InSequencer)
+{
+	TSequencerKeyEditor<FMovieSceneBoolChannel, bool> KeyEditor(
+		InObjectBindingID, Section->GetChannelProxy().MakeHandle(Channel),
+		Section, InSequencer, PropertyBindings, ExternalValue.OnGetExternalValue
+		);
+
+	return SNew(SBoolCurveKeyEditor, KeyEditor);
+}
+
+
+TSharedRef<SWidget> CreateKeyEditor(FMovieSceneIntegerChannel* Channel, const TMovieSceneExternalValue<int32>& ExternalValue, UMovieSceneSection* Section, const FGuid& InObjectBindingID, TWeakPtr<FTrackInstancePropertyBindings> PropertyBindings, TWeakPtr<ISequencer> InSequencer)
+{
+	TSequencerKeyEditor<FMovieSceneIntegerChannel, int32> KeyEditor(
+		InObjectBindingID, Section->GetChannelProxy().MakeHandle(Channel),
+		Section, InSequencer, PropertyBindings, ExternalValue.OnGetExternalValue
+		);
+
+	typedef SNumericKeyEditor<FMovieSceneIntegerChannel, int32> KeyEditorType;
+	return SNew(KeyEditorType, KeyEditor);
+}
+
+
+TSharedRef<SWidget> CreateKeyEditor(FMovieSceneFloatChannel*   Channel, const TMovieSceneExternalValue<float>& ExternalValue, UMovieSceneSection* Section, const FGuid& InObjectBindingID, TWeakPtr<FTrackInstancePropertyBindings> PropertyBindings, TWeakPtr<ISequencer> InSequencer)
+{
+	TSequencerKeyEditor<FMovieSceneFloatChannel, float> KeyEditor(
+		InObjectBindingID, Section->GetChannelProxy().MakeHandle(Channel),
+		Section, InSequencer, PropertyBindings, ExternalValue.OnGetExternalValue
+		);
+
+	typedef SNumericKeyEditor<FMovieSceneFloatChannel, float> KeyEditorType;
+	return SNew(KeyEditorType, KeyEditor);
+}
+
+
+TSharedRef<SWidget> CreateKeyEditor(FMovieSceneStringChannel*  Channel, const TMovieSceneExternalValue<FString>& ExternalValue, UMovieSceneSection* Section, const FGuid& InObjectBindingID, TWeakPtr<FTrackInstancePropertyBindings> PropertyBindings, TWeakPtr<ISequencer> InSequencer)
+{
+	TSequencerKeyEditor<FMovieSceneStringChannel, FString> KeyEditor(
+		InObjectBindingID, Section->GetChannelProxy().MakeHandle(Channel),
+		Section, InSequencer, PropertyBindings, ExternalValue.OnGetExternalValue
+		);
+
+	return SNew(SStringCurveKeyEditor, KeyEditor);
+}
+
+
+TSharedRef<SWidget> CreateKeyEditor(FMovieSceneByteChannel*    Channel, const TMovieSceneExternalValue<uint8>& ExternalValue, UMovieSceneSection* Section, const FGuid& InObjectBindingID, TWeakPtr<FTrackInstancePropertyBindings> PropertyBindings, TWeakPtr<ISequencer> InSequencer)
+{
+	TSequencerKeyEditor<FMovieSceneByteChannel, uint8> KeyEditor(
+		InObjectBindingID, Section->GetChannelProxy().MakeHandle(Channel),
+		Section, InSequencer, PropertyBindings, ExternalValue.OnGetExternalValue
+		);
+
+	if (UEnum* Enum = Channel->GetEnum())
+	{
+		return SNew(SEnumCurveKeyEditor, KeyEditor, Enum);
+	}
+	else
+	{
+		typedef SNumericKeyEditor<FMovieSceneByteChannel, uint8> KeyEditorType;
+		return SNew(KeyEditorType, KeyEditor);
+	}
+}
+
+TSharedPtr<FStructOnScope> GetKeyStruct(TMovieSceneChannelHandle<FMovieSceneBoolChannel>     ChannelHandle, FKeyHandle InHandle)
+{
+	return CreateKeyStruct<FMovieSceneBoolKeyStruct>(ChannelHandle, InHandle);
+}
+TSharedPtr<FStructOnScope> GetKeyStruct(TMovieSceneChannelHandle<FMovieSceneByteChannel>     ChannelHandle, FKeyHandle InHandle)
+{
+	return CreateKeyStruct<FMovieSceneByteKeyStruct>(ChannelHandle, InHandle);
+}
+TSharedPtr<FStructOnScope> GetKeyStruct(TMovieSceneChannelHandle<FMovieSceneIntegerChannel>  ChannelHandle, FKeyHandle InHandle)
+{
+	return CreateKeyStruct<FMovieSceneIntegerKeyStruct>(ChannelHandle, InHandle);
+}
+TSharedPtr<FStructOnScope> GetKeyStruct(TMovieSceneChannelHandle<FMovieSceneStringChannel>   ChannelHandle, FKeyHandle InHandle)
+{
+	return CreateKeyStruct<FMovieSceneStringKeyStruct>(ChannelHandle, InHandle);
+}
+TSharedPtr<FStructOnScope> GetKeyStruct(TMovieSceneChannelHandle<FMovieSceneParticleChannel> ChannelHandle, FKeyHandle InHandle)
+{
+	FMovieSceneParticleChannel* Channel = ChannelHandle.Get();
+	if (!Channel)
+	{
+		return nullptr;
+	}
+
+	auto ChannelInterface = Channel->GetInterface();
+	const int32 KeyIndex = ChannelInterface.GetIndex(InHandle);
+
+	if (KeyIndex == INDEX_NONE)
+	{
+		return nullptr;
+	}
+
+	TSharedPtr<FStructOnScope> KeyStruct = MakeShared<FStructOnScope>(FMovieSceneParticleKeyStruct::StaticStruct());
+	FMovieSceneParticleKeyStruct* Struct = reinterpret_cast<FMovieSceneParticleKeyStruct*>(KeyStruct->GetStructMemory());
+
+	Struct->Time  = ChannelInterface.GetTimes()[KeyIndex];
+	Struct->Value = (EParticleKey)ChannelInterface.GetValues()[KeyIndex];
+
+	Struct->KeyStructInterop.Add(FMovieSceneChannelValueHelper(ChannelHandle, &Struct->Value, MakeTuple(InHandle, Struct->Time)));
+	return KeyStruct;
+}
+
+void DrawKeys(FMovieSceneFloatChannel* Channel, TArrayView<const FKeyHandle> InKeyHandles, TArrayView<FKeyDrawParams> OutKeyDrawParams)
+{
+	static const FName CircleKeyBrushName("Sequencer.KeyCircle");
+	static const FName DiamondKeyBrushName("Sequencer.KeyDiamond");
+	static const FName SquareKeyBrushName("Sequencer.KeySquare");
+	static const FName TriangleKeyBrushName("Sequencer.KeyTriangle");
+
+	const FSlateBrush* CircleKeyBrush = FEditorStyle::GetBrush(CircleKeyBrushName);
+	const FSlateBrush* DiamondKeyBrush = FEditorStyle::GetBrush(DiamondKeyBrushName);
+	const FSlateBrush* SquareKeyBrush = FEditorStyle::GetBrush(SquareKeyBrushName);
+	const FSlateBrush* TriangleKeyBrush = FEditorStyle::GetBrush(TriangleKeyBrushName);
+
+	TMovieSceneChannel<FMovieSceneFloatValue> ChannelInterface = Channel->GetInterface();
+	TArrayView<const FMovieSceneFloatValue> Values = ChannelInterface.GetValues();
+
+	FKeyDrawParams TempParams;
+	TempParams.BorderBrush = TempParams.FillBrush = DiamondKeyBrush;
+
+	for (int32 Index = 0; Index < InKeyHandles.Num(); ++Index)
+	{
+		FKeyHandle Handle = InKeyHandles[Index];
+
+		const int32 KeyIndex = ChannelInterface.GetIndex(Handle);
+
+		ERichCurveInterpMode InterpMode   = KeyIndex == INDEX_NONE ? RCIM_None : Values[KeyIndex].InterpMode.GetValue();
+		ERichCurveTangentMode TangentMode = KeyIndex == INDEX_NONE ? RCTM_None : Values[KeyIndex].TangentMode.GetValue();
+
+		TempParams.FillOffset = FVector2D(0.f, 0.f);
+
+		switch (InterpMode)
+		{
+		case RCIM_Linear:
+			TempParams.BorderBrush = TempParams.FillBrush = TriangleKeyBrush;
+			TempParams.FillTint = FLinearColor(0.0f, 0.617f, 0.449f, 1.0f); // blueish green
+			TempParams.FillOffset = FVector2D(0.0f, 1.0f);
+			break;
+
+		case RCIM_Constant:
+			TempParams.BorderBrush = TempParams.FillBrush = SquareKeyBrush;
+			TempParams.FillTint = FLinearColor(0.0f, 0.445f, 0.695f, 1.0f); // blue
+			break;
+
+		case RCIM_Cubic:
+			TempParams.BorderBrush = TempParams.FillBrush = CircleKeyBrush;
+
+			switch (TangentMode)
+			{
+			case RCTM_Auto:  TempParams.FillTint = FLinearColor(0.972f, 0.2f, 0.2f, 1.0f);     break; // vermillion
+			case RCTM_Break: TempParams.FillTint = FLinearColor(0.336f, 0.703f, 0.5f, 0.91f);  break; // sky blue
+			case RCTM_User:  TempParams.FillTint = FLinearColor(0.797f, 0.473f, 0.5f, 0.652f); break; // reddish purple
+			default:         TempParams.FillTint = FLinearColor(0.75f, 0.75f, 0.75f, 1.0f);    break; // light gray
+			}
+			break;
+
+		default:
+			TempParams.BorderBrush = TempParams.FillBrush = DiamondKeyBrush;
+			TempParams.FillTint   = FLinearColor(1.0f, 1.0f, 1.0f, 1.0f); // white
+			break;
+		}
+
+		OutKeyDrawParams[Index] = TempParams;
+	}
+}
+
+void DrawKeys(FMovieSceneParticleChannel* Channel, TArrayView<const FKeyHandle> InKeyHandles, TArrayView<FKeyDrawParams> OutKeyDrawParams)
+{
+	static const FName KeyLeftBrushName("Sequencer.KeyLeft");
+	static const FName KeyRightBrushName("Sequencer.KeyRight");
+	static const FName KeyDiamondBrushName("Sequencer.KeyDiamond");
+
+	const FSlateBrush* LeftKeyBrush = FEditorStyle::GetBrush(KeyLeftBrushName);
+	const FSlateBrush* RightKeyBrush = FEditorStyle::GetBrush(KeyRightBrushName);
+	const FSlateBrush* DiamondBrush = FEditorStyle::GetBrush(KeyDiamondBrushName);
+
+	TMovieSceneChannel<uint8> ChannelInterface = Channel->GetInterface();
+
+	for (int32 Index = 0; Index < InKeyHandles.Num(); ++Index)
+	{
+		FKeyHandle Handle = InKeyHandles[Index];
+
+		FKeyDrawParams Params;
+		Params.BorderBrush = Params.FillBrush = DiamondBrush;
+
+		const int32 KeyIndex = ChannelInterface.GetIndex(Handle);
+		if ( KeyIndex != INDEX_NONE )
+		{
+			const EParticleKey Value = (EParticleKey)ChannelInterface.GetValues()[KeyIndex];
+			if ( Value == EParticleKey::Activate )
+			{
+				Params.BorderBrush = Params.FillBrush = LeftKeyBrush;
+				Params.FillOffset = FVector2D(-1.0f, 1.0f);
+			}
+			else if ( Value == EParticleKey::Deactivate )
+			{
+				Params.BorderBrush = Params.FillBrush = RightKeyBrush;
+				Params.FillOffset = FVector2D(1.0f, 1.0f);
+			}
+		}
+
+		OutKeyDrawParams[Index] = Params;
+	}
+}
+
+
+struct FFloatChannelKeyMenuExtension : FExtender, TSharedFromThis<FFloatChannelKeyMenuExtension>
+{
+	FFloatChannelKeyMenuExtension(TWeakPtr<ISequencer> InSequencer, TArray<TChannelAndHandles<FMovieSceneFloatChannel>>&& InChannels)
+		: WeakSequencer(InSequencer)
+		, ChannelAndHandles(MoveTemp(InChannels))
+	{}
+
+	void ExtendMenu(FMenuBuilder& MenuBuilder)
+	{
+		ISequencer* SequencerPtr = WeakSequencer.Pin().Get();
+		if (!SequencerPtr)
+		{
+			return;
+		}
+
+		TSharedRef<FFloatChannelKeyMenuExtension> SharedThis = AsShared();
+
+		MenuBuilder.BeginSection("SequencerInterpolation", LOCTEXT("KeyInterpolationMenu", "Key Interpolation"));
+		{
+			MenuBuilder.AddMenuEntry(
+				LOCTEXT("SetKeyInterpolationAuto", "Cubic (Auto)"),
+				LOCTEXT("SetKeyInterpolationAutoTooltip", "Set key interpolation to auto"),
+				FSlateIcon(FEditorStyle::GetStyleSetName(), "Sequencer.IconKeyAuto"),
+				FUIAction(
+					FExecuteAction::CreateLambda([SharedThis]{ SharedThis->SetInterpTangentMode(RCIM_Cubic, RCTM_Auto); }),
+					FCanExecuteAction(),
+					FIsActionChecked::CreateLambda([SharedThis]{ return SharedThis->IsInterpTangentModeSelected(RCIM_Cubic, RCTM_Auto); }) ),
+				NAME_None,
+				EUserInterfaceActionType::ToggleButton
+			);
+
+			MenuBuilder.AddMenuEntry(
+				LOCTEXT("SetKeyInterpolationUser", "Cubic (User)"),
+				LOCTEXT("SetKeyInterpolationUserTooltip", "Set key interpolation to user"),
+				FSlateIcon(FEditorStyle::GetStyleSetName(), "Sequencer.IconKeyUser"),
+				FUIAction(
+					FExecuteAction::CreateLambda([SharedThis]{ SharedThis->SetInterpTangentMode(RCIM_Cubic, RCTM_User); }),
+					FCanExecuteAction(),
+					FIsActionChecked::CreateLambda([SharedThis]{ return SharedThis->IsInterpTangentModeSelected(RCIM_Cubic, RCTM_User); }) ),
+				NAME_None,
+				EUserInterfaceActionType::ToggleButton
+			);
+
+			MenuBuilder.AddMenuEntry(
+				LOCTEXT("SetKeyInterpolationBreak", "Cubic (Break)"),
+				LOCTEXT("SetKeyInterpolationBreakTooltip", "Set key interpolation to break"),
+				FSlateIcon(FEditorStyle::GetStyleSetName(), "Sequencer.IconKeyBreak"),
+				FUIAction(
+					FExecuteAction::CreateLambda([SharedThis]{ SharedThis->SetInterpTangentMode(RCIM_Cubic, RCTM_Break); }),
+					FCanExecuteAction(),
+					FIsActionChecked::CreateLambda([SharedThis]{ return SharedThis->IsInterpTangentModeSelected(RCIM_Cubic, RCTM_Break); }) ),
+				NAME_None,
+				EUserInterfaceActionType::ToggleButton
+			);
+
+			MenuBuilder.AddMenuEntry(
+				LOCTEXT("SetKeyInterpolationLinear", "Linear"),
+				LOCTEXT("SetKeyInterpolationLinearTooltip", "Set key interpolation to linear"),
+				FSlateIcon(FEditorStyle::GetStyleSetName(), "Sequencer.IconKeyLinear"),
+				FUIAction(
+					FExecuteAction::CreateLambda([SharedThis]{ SharedThis->SetInterpTangentMode(RCIM_Linear, RCTM_Auto); }),
+					FCanExecuteAction(),
+					FIsActionChecked::CreateLambda([SharedThis]{ return SharedThis->IsInterpTangentModeSelected(RCIM_Linear, RCTM_Auto); }) ),
+				NAME_None,
+				EUserInterfaceActionType::ToggleButton
+			);
+
+			MenuBuilder.AddMenuEntry(
+				LOCTEXT("SetKeyInterpolationConstant", "Constant"),
+				LOCTEXT("SetKeyInterpolationConstantTooltip", "Set key interpolation to constant"),
+				FSlateIcon(FEditorStyle::GetStyleSetName(), "Sequencer.IconKeyConstant"),
+				FUIAction(
+					FExecuteAction::CreateLambda([SharedThis]{ SharedThis->SetInterpTangentMode(RCIM_Constant, RCTM_Auto); }),
+					FCanExecuteAction(),
+					FIsActionChecked::CreateLambda([SharedThis]{ return SharedThis->IsInterpTangentModeSelected(RCIM_Constant, RCTM_Auto); }) ),
+				NAME_None,
+				EUserInterfaceActionType::ToggleButton
+			);
+		}
+		MenuBuilder.EndSection(); // SequencerInterpolation
+	}
+
+	void SetInterpTangentMode(ERichCurveInterpMode InterpMode, ERichCurveTangentMode TangentMode)
+	{
+		FScopedTransaction SetInterpTangentModeTransaction(NSLOCTEXT("Sequencer", "SetInterpTangentMode_Transaction", "Set Interpolation and Tangent Mode"));
+		bool bAnythingChanged = false;
+
+		for (const TChannelAndHandles<FMovieSceneFloatChannel>& Channel : ChannelAndHandles)
+		{
+			UMovieSceneSection* Section = Channel.Section.Get();
+			FMovieSceneFloatChannel* ChannelPtr = Channel.Channel.Get();
+
+			if (Section && ChannelPtr)
+			{
+				Section->Modify();
+
+				TMovieSceneChannel<FMovieSceneFloatValue> ChannelInterface = ChannelPtr->GetInterface();
+				TArrayView<FMovieSceneFloatValue> Values = ChannelInterface.GetValues();
+
+				for (FKeyHandle Handle : Channel.Handles)
+				{
+					const int32 KeyIndex = ChannelInterface.GetIndex(Handle);
+					if (KeyIndex != INDEX_NONE)
+					{
+						Values[KeyIndex].InterpMode = InterpMode;
+						Values[KeyIndex].TangentMode = TangentMode;
+						bAnythingChanged = true;
+					}
+				}
+
+				ChannelPtr->AutoSetTangents();
+			}
+		}
+
+		if (bAnythingChanged)
+		{
+			if (ISequencer* Sequencer = WeakSequencer.Pin().Get())
+			{
+				Sequencer->NotifyMovieSceneDataChanged( EMovieSceneDataChangeType::TrackValueChanged );
+			}
+		}
+	}
+
+	bool IsInterpTangentModeSelected(ERichCurveInterpMode InterpMode, ERichCurveTangentMode TangentMode) const
+	{
+		for (const TChannelAndHandles<FMovieSceneFloatChannel>& Channel : ChannelAndHandles)
+		{
+			FMovieSceneFloatChannel* ChannelPtr = Channel.Channel.Get();
+			if (ChannelPtr)
+			{
+				TMovieSceneChannel<FMovieSceneFloatValue> ChannelInterface = ChannelPtr->GetInterface();
+				TArrayView<FMovieSceneFloatValue> Values = ChannelInterface.GetValues();
+
+				for (FKeyHandle Handle : Channel.Handles)
+				{
+					int32 KeyIndex = ChannelInterface.GetIndex(Handle);
+					if (KeyIndex == INDEX_NONE || Values[KeyIndex].InterpMode != InterpMode || Values[KeyIndex].TangentMode != TangentMode)
+					{
+						return false;
+					}
+				}
+			}
+		}
+		return true;
+	}
+
+private:
+
+	/** Hidden AsShared() methods to prevent CreateSP delegate use since this extender disappears with its menu. */
+	using TSharedFromThis::AsShared;
+
+	TWeakPtr<ISequencer> WeakSequencer;
+	TArray<TChannelAndHandles<FMovieSceneFloatChannel>> ChannelAndHandles;
+};
+
+
+struct FFloatChannelSectionMenuExtension : FExtender, TSharedFromThis<FFloatChannelSectionMenuExtension>
+{
+	FFloatChannelSectionMenuExtension(TWeakPtr<ISequencer> InSequencer, TArray<TMovieSceneChannelHandle<FMovieSceneFloatChannel>>&& InChannels, TArrayView<UMovieSceneSection* const> InSections)
+		: WeakSequencer(InSequencer)
+		, Channels(MoveTemp(InChannels))
+	{
+		Sections.Reserve(InSections.Num());
+		for (UMovieSceneSection* Section : InSections)
+		{
+			Sections.Add(Section);
+		}
+	}
+
+	void ExtendMenu(FMenuBuilder& MenuBuilder)
+	{
+		ISequencer* SequencerPtr = WeakSequencer.Pin().Get();
+		if (!SequencerPtr)
+		{
+			return;
+		}
+
+		TSharedRef<FFloatChannelSectionMenuExtension> SharedThis = AsShared();
+
+		MenuBuilder.AddSubMenu(
+			LOCTEXT("SetPreInfinityExtrap", "Pre-Infinity"),
+			LOCTEXT("SetPreInfinityExtrapTooltip", "Set pre-infinity extrapolation"),
+			FNewMenuDelegate::CreateLambda([SharedThis](FMenuBuilder& SubMenuBuilder){ SharedThis->AddExtrapolationMenu(SubMenuBuilder, true); })
+			);
+
+		MenuBuilder.AddSubMenu(
+			LOCTEXT("SetPostInfinityExtrap", "Post-Infinity"),
+			LOCTEXT("SetPostInfinityExtrapTooltip", "Set post-infinity extrapolation"),
+			FNewMenuDelegate::CreateLambda([SharedThis](FMenuBuilder& SubMenuBuilder){ SharedThis->AddExtrapolationMenu(SubMenuBuilder, false); })
+			);
+	}
+
+	void AddExtrapolationMenu(FMenuBuilder& MenuBuilder, bool bPreInfinity)
+	{
+		TSharedRef<FFloatChannelSectionMenuExtension> SharedThis = AsShared();
+
+		MenuBuilder.AddMenuEntry(
+			LOCTEXT("SetExtrapCycle", "Cycle"),
+			LOCTEXT("SetExtrapCycleTooltip", "Set extrapolation cycle"),
+			FSlateIcon(),
+			FUIAction(
+				FExecuteAction::CreateLambda([SharedThis, bPreInfinity]{ SharedThis->SetExtrapolationMode(RCCE_Cycle, bPreInfinity); }),
+				FCanExecuteAction(),
+				FIsActionChecked::CreateLambda([SharedThis, bPreInfinity]{ return SharedThis->IsExtrapolationModeSelected(RCCE_Cycle, bPreInfinity); })
+			),
+			NAME_None,
+			EUserInterfaceActionType::RadioButton
+		);
+
+		MenuBuilder.AddMenuEntry(
+			LOCTEXT("SetExtrapCycleWithOffset", "Cycle with Offset"),
+			LOCTEXT("SetExtrapCycleWithOffsetTooltip", "Set extrapolation cycle with offset"),
+			FSlateIcon(),
+			FUIAction(
+				FExecuteAction::CreateLambda([SharedThis, bPreInfinity]{ SharedThis->SetExtrapolationMode(RCCE_CycleWithOffset, bPreInfinity); }),
+				FCanExecuteAction(),
+				FIsActionChecked::CreateLambda([SharedThis, bPreInfinity]{ return SharedThis->IsExtrapolationModeSelected(RCCE_CycleWithOffset, bPreInfinity); })
+			),
+			NAME_None,
+			EUserInterfaceActionType::RadioButton
+		);
+
+		MenuBuilder.AddMenuEntry(
+			LOCTEXT("SetExtrapOscillate", "Oscillate"),
+			LOCTEXT("SetExtrapOscillateTooltip", "Set extrapolation oscillate"),
+			FSlateIcon(),
+			FUIAction(
+				FExecuteAction::CreateLambda([SharedThis, bPreInfinity]{ SharedThis->SetExtrapolationMode(RCCE_Oscillate, bPreInfinity); }),
+				FCanExecuteAction(),
+				FIsActionChecked::CreateLambda([SharedThis, bPreInfinity]{ return SharedThis->IsExtrapolationModeSelected(RCCE_Oscillate, bPreInfinity); })
+			),
+			NAME_None,
+			EUserInterfaceActionType::RadioButton
+		);
+
+		MenuBuilder.AddMenuEntry(
+			LOCTEXT("SetExtrapLinear", "Linear"),
+			LOCTEXT("SetExtrapLinearTooltip", "Set extrapolation linear"),
+			FSlateIcon(),
+			FUIAction(
+				FExecuteAction::CreateLambda([SharedThis, bPreInfinity]{ SharedThis->SetExtrapolationMode(RCCE_Linear, bPreInfinity); }),
+				FCanExecuteAction(),
+				FIsActionChecked::CreateLambda([SharedThis, bPreInfinity]{ return SharedThis->IsExtrapolationModeSelected(RCCE_Linear, bPreInfinity); })
+			),
+			NAME_None,
+			EUserInterfaceActionType::RadioButton
+		);
+
+		MenuBuilder.AddMenuEntry(
+			LOCTEXT("SetExtrapConstant", "Constant"),
+			LOCTEXT("SetExtrapConstantTooltip", "Set extrapolation constant"),
+			FSlateIcon(),
+			FUIAction(
+				FExecuteAction::CreateLambda([SharedThis, bPreInfinity]{ SharedThis->SetExtrapolationMode(RCCE_Constant, bPreInfinity); }),
+				FCanExecuteAction(),
+				FIsActionChecked::CreateLambda([SharedThis, bPreInfinity]{ return SharedThis->IsExtrapolationModeSelected(RCCE_Constant, bPreInfinity); })
+			),
+			NAME_None,
+			EUserInterfaceActionType::RadioButton
+		);
+	}
+
+	void SetExtrapolationMode(ERichCurveExtrapolation ExtrapMode, bool bPreInfinity)
+	{
+		FScopedTransaction Transaction(LOCTEXT("SetExtrapolationMode_Transaction", "Set Extrapolation Mode"));
+
+		bool bAnythingChanged = false;
+
+		// Modify all sections
+		for (TWeakObjectPtr<UMovieSceneSection> WeakSection : Sections)
+		{
+			if (UMovieSceneSection* Section = WeakSection.Get())
+			{
+				Section->Modify();
+			}
+		}
+
+		// Apply to all channels
+		for (const TMovieSceneChannelHandle<FMovieSceneFloatChannel>& Handle : Channels)
+		{
+			FMovieSceneFloatChannel* Channel = Handle.Get();
+
+			if (Channel)
+			{
+				TEnumAsByte<ERichCurveExtrapolation>& DestExtrap = bPreInfinity ? Channel->PreInfinityExtrap : Channel->PostInfinityExtrap;
+				DestExtrap = ExtrapMode;
+				bAnythingChanged = true;
+			}
+		}
+
+		if (bAnythingChanged)
+		{
+			if (ISequencer* Sequencer = WeakSequencer.Pin().Get())
+			{
+				Sequencer->NotifyMovieSceneDataChanged( EMovieSceneDataChangeType::TrackValueChanged );
+			}
+		}
+		else
+		{
+			Transaction.Cancel();
+		}
+	}
+
+
+	bool IsExtrapolationModeSelected(ERichCurveExtrapolation ExtrapMode, bool bPreInfinity) const
+	{
+		for (const TMovieSceneChannelHandle<FMovieSceneFloatChannel>& Handle : Channels)
+		{
+			FMovieSceneFloatChannel* Channel = Handle.Get();
+
+			if (Channel)
+			{
+				ERichCurveExtrapolation SourceExtrap = bPreInfinity ? Channel->PreInfinityExtrap : Channel->PostInfinityExtrap;
+				if (SourceExtrap != ExtrapMode)
+				{
+					return false;
+				}
+			}
+		}
+
+		return true;
+	}
+
+private:
+
+	/** Hidden AsShared() methods to prevent CreateSP delegate use since this extender disappears with its menu. */
+	using TSharedFromThis::AsShared;
+
+	TWeakPtr<ISequencer> WeakSequencer;
+	TArray<TMovieSceneChannelHandle<FMovieSceneFloatChannel>> Channels;
+	TArray<TWeakObjectPtr<UMovieSceneSection>> Sections;
+};
+
+void ExtendSectionMenu(FMenuBuilder& OuterMenuBuilder, TArray<TMovieSceneChannelHandle<FMovieSceneFloatChannel>>&& Channels, TArrayView<UMovieSceneSection* const> Sections, TWeakPtr<ISequencer> InSequencer)
+{
+	TSharedRef<FFloatChannelSectionMenuExtension> Extension = MakeShared<FFloatChannelSectionMenuExtension>(InSequencer, MoveTemp(Channels), Sections);
+
+	Extension->AddMenuExtension("SequencerSections", EExtensionHook::First, nullptr, FMenuExtensionDelegate::CreateLambda([Extension](FMenuBuilder& MenuBuilder) { Extension->ExtendMenu(MenuBuilder); }));
+
+	OuterMenuBuilder.PushExtender(Extension);
+}
+
+void ExtendKeyMenu(FMenuBuilder& OuterMenuBuilder, TArray<TChannelAndHandles<FMovieSceneFloatChannel>>&& Channels, TWeakPtr<ISequencer> InSequencer)
+{
+	TSharedRef<FFloatChannelKeyMenuExtension> Extension = MakeShared<FFloatChannelKeyMenuExtension>(InSequencer, MoveTemp(Channels));
+
+	Extension->AddMenuExtension("SequencerKeyEdit", EExtensionHook::After, nullptr, FMenuExtensionDelegate::CreateLambda([Extension](FMenuBuilder& MenuBuilder) { Extension->ExtendMenu(MenuBuilder); }));
+
+	OuterMenuBuilder.PushExtender(Extension);
+}
+
+TUniquePtr<FCurveModel> CreateCurveEditorModel(TMovieSceneChannelHandle<FMovieSceneFloatChannel> FloatChannel, UMovieSceneSection* OwningSection, TSharedRef<ISequencer> InSequencer)
+{
+	return MakeUnique<FFloatChannelCurveModel>(FloatChannel, OwningSection, InSequencer);
+}
+
+#undef LOCTEXT_NAMESPACE
\ No newline at end of file
diff --git a/Engine/Source/Editor/MovieSceneTools/Private/Channels/BuiltInChannelEditors.h b/Engine/Source/Editor/MovieSceneTools/Private/Channels/BuiltInChannelEditors.h
new file mode 100644
--- /dev/null
+++ b/Engine/Source/Editor/MovieSceneTools/Private/Channels/BuiltInChannelEditors.h
@@ -0,0 +1,124 @@
+// Copyright 1998-2018 Epic Games, Inc. All Rights Reserved.
+
+#pragma once
+
+#include "CoreTypes.h"
+#include "Templates/SharedPointer.h"
+#include "Containers/Array.h"
+#include "UObject/ObjectMacros.h"
+#include "UObject/Class.h"
+
+#include "MovieSceneKeyStruct.h"
+#include "ISequencerChannelInterface.h"
+#include "Channels/MovieSceneChannelHandle.h"
+
+#include "Channels/MovieSceneBoolChannel.h"
+#include "Channels/MovieSceneByteChannel.h"
+#include "Channels/MovieSceneFloatChannel.h"
+#include "Channels/MovieSceneIntegerChannel.h"
+#include "Sections/MovieSceneStringSection.h"
+#include "Sections/MovieSceneParticleSection.h"
+#include "Sections/MovieSceneActorReferenceSection.h"
+
+#include "BuiltInChannelEditors.generated.h"
+
+struct FKeyHandle;
+struct FKeyDrawParams;
+
+class SWidget;
+class ISequencer;
+class FMenuBuilder;
+class FStructOnScope;
+class ISectionLayoutBuilder;
+
+template<typename> class  TArrayView;
+
+/** Overrides for adding or updating a key for non-standard channels */
+FKeyHandle AddOrUpdateKey(FMovieSceneFloatChannel* Channel, const TMovieSceneExternalValue<float>& EditorData, FFrameNumber InTime, ISequencer& Sequencer, const FGuid& InObjectBindingID, FTrackInstancePropertyBindings* PropertyBindings);
+FKeyHandle AddOrUpdateKey(FMovieSceneActorReferenceData* Channel, FFrameNumber InTime, ISequencer& Sequencer, const FGuid& InObjectBindingID, FTrackInstancePropertyBindings* PropertyBindings);
+
+/** Key editor overrides */
+bool CanCreateKeyEditor(FMovieSceneBoolChannel*    Channel);
+bool CanCreateKeyEditor(FMovieSceneByteChannel*    Channel);
+bool CanCreateKeyEditor(FMovieSceneIntegerChannel* Channel);
+bool CanCreateKeyEditor(FMovieSceneFloatChannel*   Channel);
+bool CanCreateKeyEditor(FMovieSceneStringChannel*  Channel);
+
+TSharedRef<SWidget> CreateKeyEditor(FMovieSceneBoolChannel*    Channel, const TMovieSceneExternalValue<bool>&    EditorData, UMovieSceneSection* Section, const FGuid& InObjectBindingID, TWeakPtr<FTrackInstancePropertyBindings> PropertyBindings, TWeakPtr<ISequencer> InSequencer);
+TSharedRef<SWidget> CreateKeyEditor(FMovieSceneByteChannel*    Channel, const TMovieSceneExternalValue<uint8>&   EditorData, UMovieSceneSection* Section, const FGuid& InObjectBindingID, TWeakPtr<FTrackInstancePropertyBindings> PropertyBindings, TWeakPtr<ISequencer> InSequencer);
+TSharedRef<SWidget> CreateKeyEditor(FMovieSceneIntegerChannel* Channel, const TMovieSceneExternalValue<int32>&   EditorData, UMovieSceneSection* Section, const FGuid& InObjectBindingID, TWeakPtr<FTrackInstancePropertyBindings> PropertyBindings, TWeakPtr<ISequencer> InSequencer);
+TSharedRef<SWidget> CreateKeyEditor(FMovieSceneFloatChannel*   Channel, const TMovieSceneExternalValue<float>&   EditorData, UMovieSceneSection* Section, const FGuid& InObjectBindingID, TWeakPtr<FTrackInstancePropertyBindings> PropertyBindings, TWeakPtr<ISequencer> InSequencer);
+TSharedRef<SWidget> CreateKeyEditor(FMovieSceneStringChannel*  Channel, const TMovieSceneExternalValue<FString>& EditorData, UMovieSceneSection* Section, const FGuid& InObjectBindingID, TWeakPtr<FTrackInstancePropertyBindings> PropertyBindings, TWeakPtr<ISequencer> InSequencer);
+
+/** KeyStruct overrides */
+TSharedPtr<FStructOnScope> GetKeyStruct(TMovieSceneChannelHandle<FMovieSceneBoolChannel>     Channel, FKeyHandle InHandle);
+TSharedPtr<FStructOnScope> GetKeyStruct(TMovieSceneChannelHandle<FMovieSceneByteChannel>     Channel, FKeyHandle InHandle);
+TSharedPtr<FStructOnScope> GetKeyStruct(TMovieSceneChannelHandle<FMovieSceneIntegerChannel>  Channel, FKeyHandle InHandle);
+TSharedPtr<FStructOnScope> GetKeyStruct(TMovieSceneChannelHandle<FMovieSceneStringChannel>   Channel, FKeyHandle InHandle);
+TSharedPtr<FStructOnScope> GetKeyStruct(TMovieSceneChannelHandle<FMovieSceneParticleChannel> Channel, FKeyHandle InHandle);
+
+/** Key drawing overrides */
+void DrawKeys(FMovieSceneFloatChannel*    Channel, TArrayView<const FKeyHandle> InKeyHandles, TArrayView<FKeyDrawParams> OutKeyDrawParams);
+void DrawKeys(FMovieSceneParticleChannel* Channel, TArrayView<const FKeyHandle> InKeyHandles, TArrayView<FKeyDrawParams> OutKeyDrawParams);
+
+/** Context menu overrides */
+void ExtendSectionMenu(FMenuBuilder& OuterMenuBuilder, TArray<TMovieSceneChannelHandle<FMovieSceneFloatChannel>>&& Channels, TArrayView<UMovieSceneSection* const> Sections, TWeakPtr<ISequencer> InSequencer);
+void ExtendKeyMenu(FMenuBuilder& OuterMenuBuilder, TArray<TChannelAndHandles<FMovieSceneFloatChannel>>&& Channels, TWeakPtr<ISequencer> InSequencer);
+
+/** Curve editor models */
+TUniquePtr<FCurveModel> CreateCurveEditorModel(TMovieSceneChannelHandle<FMovieSceneFloatChannel> FloatChannel, UMovieSceneSection* OwningSection, TSharedRef<ISequencer> InSequencer);
+
+USTRUCT()
+struct FMovieSceneIntegerKeyStruct : public FMovieSceneKeyTimeStruct
+{
+	GENERATED_BODY();
+
+	UPROPERTY(EditAnywhere, Category="Key")
+	int32 Value;
+};
+template<> struct TStructOpsTypeTraits<FMovieSceneIntegerKeyStruct> : public TStructOpsTypeTraitsBase2<FMovieSceneIntegerKeyStruct> { enum { WithCopy = false }; };
+
+
+USTRUCT()
+struct FMovieSceneByteKeyStruct : public FMovieSceneKeyTimeStruct
+{
+	GENERATED_BODY();
+
+	UPROPERTY(EditAnywhere, Category="Key")
+	uint8 Value;
+};
+template<> struct TStructOpsTypeTraits<FMovieSceneByteKeyStruct> : public TStructOpsTypeTraitsBase2<FMovieSceneByteKeyStruct> { enum { WithCopy = false }; };
+
+
+USTRUCT()
+struct FMovieSceneBoolKeyStruct : public FMovieSceneKeyTimeStruct
+{
+	GENERATED_BODY();
+
+	UPROPERTY(EditAnywhere, Category="Key")
+	bool Value;
+};
+template<> struct TStructOpsTypeTraits<FMovieSceneBoolKeyStruct> : public TStructOpsTypeTraitsBase2<FMovieSceneBoolKeyStruct> { enum { WithCopy = false }; };
+
+
+USTRUCT()
+struct FMovieSceneStringKeyStruct : public FMovieSceneKeyTimeStruct
+{
+	GENERATED_BODY();
+
+	UPROPERTY(EditAnywhere, Category="Key")
+	FString Value;
+};
+template<> struct TStructOpsTypeTraits<FMovieSceneStringKeyStruct> : public TStructOpsTypeTraitsBase2<FMovieSceneStringKeyStruct> { enum { WithCopy = false }; };
+
+
+USTRUCT()
+struct FMovieSceneParticleKeyStruct : public FMovieSceneKeyTimeStruct
+{
+	GENERATED_BODY();
+
+	UPROPERTY(EditAnywhere, Category="Key")
+	EParticleKey Value;
+};
+template<> struct TStructOpsTypeTraits<FMovieSceneParticleKeyStruct> : public TStructOpsTypeTraitsBase2<FMovieSceneParticleKeyStruct> { enum { WithCopy = false }; };
+
diff --git a/Engine/Source/Editor/MovieSceneTools/Private/Channels/FloatChannelCurveModel.cpp b/Engine/Source/Editor/MovieSceneTools/Private/Channels/FloatChannelCurveModel.cpp
new file mode 100644
--- /dev/null
+++ b/Engine/Source/Editor/MovieSceneTools/Private/Channels/FloatChannelCurveModel.cpp
@@ -0,0 +1,453 @@
+// Copyright 1998-2018 Epic Games, Inc. All Rights Reserved.
+
+#include "FloatChannelCurveModel.h"
+#include "Math/Vector2D.h"
+#include "HAL/PlatformMath.h"
+#include "Channels/MovieSceneFloatChannel.h"
+#include "MovieSceneSection.h"
+#include "MovieScene.h"
+#include "CurveDrawInfo.h"
+#include "CurveDataAbstraction.h"
+#include "CurveEditor.h"
+#include "CurveEditorScreenSpace.h"
+#include "CurveEditorSnapMetrics.h"
+#include "EditorStyleSet.h"
+#include "BuiltInChannelEditors.h"
+#include "SKeyEditInterface.h"
+#include "SequencerChannelTraits.h"
+#include "ISequencer.h"
+
+FFloatChannelCurveModel::FFloatChannelCurveModel(TMovieSceneChannelHandle<FMovieSceneFloatChannel> InChannel, UMovieSceneSection* OwningSection, TWeakPtr<ISequencer> InWeakSequencer)
+{
+	ChannelHandle = InChannel;
+	WeakSection = OwningSection;
+	WeakSequencer = InWeakSequencer;
+}
+
+const void* FFloatChannelCurveModel::GetCurve() const
+{
+	return ChannelHandle.Get();
+}
+
+void FFloatChannelCurveModel::Modify()
+{
+	if (UMovieSceneSection* Section = WeakSection.Get())
+	{
+		Section->Modify();
+	}
+}
+
+double FFloatChannelCurveModel::GetInputDisplayOffset() const
+{
+	UMovieSceneSection* Section = WeakSection.Get();
+	return Section ? ( FFrameTime(0, 0.5f) / Section->GetTypedOuter<UMovieScene>()->GetFrameResolution() ) : 0.0;
+}
+
+void FFloatChannelCurveModel::AddKeys(TArrayView<const FKeyPosition> InKeyPositions, TArrayView<const FKeyAttributes> InKeyAttributes, TArrayView<TOptional<FKeyHandle>>* OutKeyHandles)
+{
+	check(InKeyPositions.Num() == InKeyAttributes.Num() && (!OutKeyHandles || OutKeyHandles->Num() == InKeyPositions.Num()));
+
+	FMovieSceneFloatChannel* Channel = ChannelHandle.Get();
+	UMovieSceneSection*      Section = WeakSection.Get();
+	if (Channel && Section)
+	{
+		Section->Modify();
+
+		TMovieSceneChannel<FMovieSceneFloatValue> ChannelInterface = Channel->GetInterface();
+		FFrameRate FrameResolution = Section->GetTypedOuter<UMovieScene>()->GetFrameResolution();
+
+		for (int32 Index = 0; Index < InKeyPositions.Num(); ++Index)
+		{
+			FKeyPosition   Position   = InKeyPositions[Index];
+			FKeyAttributes Attributes = InKeyAttributes[Index];
+
+			FFrameNumber Time = (Position.InputValue * FrameResolution).RoundToFrame();
+
+			FMovieSceneFloatValue Value(Position.OutputValue);
+
+			if (Attributes.HasInterpMode())    { Value.InterpMode            = Attributes.GetInterpMode();    }
+			if (Attributes.HasTangentMode())   { Value.TangentMode           = Attributes.GetTangentMode();   }
+			if (Attributes.HasArriveTangent()) { Value.Tangent.ArriveTangent = Attributes.GetArriveTangent(); }
+			if (Attributes.HasLeaveTangent())  { Value.Tangent.LeaveTangent  = Attributes.GetLeaveTangent();  }
+
+			int32 KeyIndex = ChannelInterface.AddKey(Time, Value);
+			if (OutKeyHandles)
+			{
+				(*OutKeyHandles)[Index] = ChannelInterface.GetHandle(KeyIndex);
+			}
+		}
+
+		Channel->AutoSetTangents();
+	}
+}
+
+bool FFloatChannelCurveModel::Evaluate(double Time, double& OutValue) const
+{
+	FMovieSceneFloatChannel* Channel = ChannelHandle.Get();
+	UMovieSceneSection*      Section = WeakSection.Get();
+
+	if (Channel && Section)
+	{
+		FFrameRate FrameResolution = Section->GetTypedOuter<UMovieScene>()->GetFrameResolution();
+
+		float ThisValue = 0.f;
+		if (Channel->Evaluate(Time * FrameResolution, ThisValue))
+		{
+			OutValue = ThisValue;
+			return true;
+		}
+	}
+
+	return false;
+}
+
+void FFloatChannelCurveModel::RemoveKeys(TArrayView<const FKeyHandle> InKeys)
+{
+	FMovieSceneFloatChannel* Channel = ChannelHandle.Get();
+	UMovieSceneSection*      Section = WeakSection.Get();
+	if (Channel && Section)
+	{
+		Section->Modify();
+
+		TMovieSceneChannel<FMovieSceneFloatValue> ChannelInterface = Channel->GetInterface();
+
+		for (FKeyHandle Handle : InKeys)
+		{
+			int32 KeyIndex = ChannelInterface.GetIndex(Handle);
+			if (KeyIndex != INDEX_NONE)
+			{
+				ChannelInterface.RemoveKey(KeyIndex);
+			}
+		}
+	}
+}
+
+void FFloatChannelCurveModel::DrawCurve(const FCurveEditor& CurveEditor, TArray<TTuple<double, double>>& InterpolatingPoints) const
+{
+	FMovieSceneFloatChannel* Channel = ChannelHandle.Get();
+	UMovieSceneSection*      Section = WeakSection.Get();
+
+	if (Channel && Section)
+	{
+		FCurveEditorScreenSpace ScreenSpace = CurveEditor.GetScreenSpace();
+
+		FFrameRate   FrameResolution  = Section->GetTypedOuter<UMovieScene>()->GetFrameResolution();
+
+		const double DisplayOffset    = GetInputDisplayOffset();
+		const double StartTimeSeconds = ScreenSpace.GetInputMin() - DisplayOffset;
+		const double EndTimeSeconds   = ScreenSpace.GetInputMax() - DisplayOffset;
+		const double TimeThreshold    = FMath::Max(0.0001, 1.0 / ScreenSpace.PixelsPerInput());
+		const double ValueThreshold   = FMath::Max(0.0001, 1.0 / ScreenSpace.PixelsPerOutput());
+
+		Channel->PopulateCurvePoints(StartTimeSeconds, EndTimeSeconds, TimeThreshold, ValueThreshold, FrameResolution, InterpolatingPoints);
+	}
+}
+
+void FFloatChannelCurveModel::GetKeys(const FCurveEditor& CurveEditor, double MinTime, double MaxTime, double MinValue, double MaxValue, TArray<FKeyHandle>& OutKeyHandles) const
+{
+	FMovieSceneFloatChannel* Channel = ChannelHandle.Get();
+	UMovieSceneSection*      Section = WeakSection.Get();
+
+	if (Channel && Section)
+	{
+		FFrameRate FrameResolution = Section->GetTypedOuter<UMovieScene>()->GetFrameResolution();
+
+		TMovieSceneChannel<FMovieSceneFloatValue> ChannelInterface = Channel->GetInterface();
+		TArrayView<const FFrameNumber>            Times  = ChannelInterface.GetTimes();
+		TArrayView<const FMovieSceneFloatValue>   Values = ChannelInterface.GetValues();
+
+		const FFrameNumber StartFrame = MinTime <= MIN_int32 ? MIN_int32 : (MinTime * FrameResolution).CeilToFrame();
+		const FFrameNumber EndFrame   = MaxTime >= MAX_int32 ? MAX_int32 : (MaxTime * FrameResolution).FloorToFrame();
+
+		const int32 StartingIndex = Algo::LowerBound(Times, StartFrame);
+		const int32 EndingIndex   = Algo::UpperBound(Times, EndFrame);
+
+		for (int32 KeyIndex = StartingIndex; KeyIndex < EndingIndex; ++KeyIndex)
+		{
+			if (Values[KeyIndex].Value >= MinValue && Values[KeyIndex].Value <= MaxValue)
+			{
+				OutKeyHandles.Add(ChannelInterface.GetHandle(KeyIndex));
+			}
+		}
+	}
+}
+
+void FFloatChannelCurveModel::GetKeyDrawInfo(ECurvePointType PointType, FKeyDrawInfo& OutDrawInfo) const
+{
+	switch (PointType)
+	{
+	case ECurvePointType::ArriveTangent:
+	case ECurvePointType::LeaveTangent:
+		OutDrawInfo.Brush = FEditorStyle::GetBrush("Sequencer.TangentHandle");
+		OutDrawInfo.ScreenSize = FVector2D(7, 7);
+		break;
+	default:
+		OutDrawInfo.Brush = FEditorStyle::GetBrush("CurveEd.CurveKey");
+		OutDrawInfo.ScreenSize = FVector2D(11, 11);
+		break;
+	}
+}
+
+void FFloatChannelCurveModel::GetKeyPositions(TArrayView<const FKeyHandle> InKeys, TArrayView<FKeyPosition> OutKeyPositions) const
+{
+	FMovieSceneFloatChannel* Channel = ChannelHandle.Get();
+	UMovieSceneSection*      Section = WeakSection.Get();
+
+	if (Channel && Section)
+	{
+		FFrameRate FrameResolution = Section->GetTypedOuter<UMovieScene>()->GetFrameResolution();
+
+		TMovieSceneChannel<FMovieSceneFloatValue> ChannelInterface = Channel->GetInterface();
+		TArrayView<const FFrameNumber>            Times            = ChannelInterface.GetTimes();
+		TArrayView<const FMovieSceneFloatValue>   Values           = ChannelInterface.GetValues();
+
+		for (int32 Index = 0; Index < InKeys.Num(); ++Index)
+		{
+			int32 KeyIndex = ChannelInterface.GetIndex(InKeys[Index]);
+			if (KeyIndex != INDEX_NONE)
+			{
+				OutKeyPositions[Index].InputValue  = Times[KeyIndex] / FrameResolution;
+				OutKeyPositions[Index].OutputValue = Values[KeyIndex].Value;
+			}
+		}
+	}
+}
+
+void FFloatChannelCurveModel::SetKeyPositions(TArrayView<const FKeyHandle> InKeys, TArrayView<const FKeyPosition> InKeyPositions)
+{
+	FMovieSceneFloatChannel* Channel = ChannelHandle.Get();
+	UMovieSceneSection*      Section = WeakSection.Get();
+
+	if (Channel && Section)
+	{
+		Section->MarkAsChanged();
+
+		FFrameRate FrameResolution = Section->GetTypedOuter<UMovieScene>()->GetFrameResolution();
+
+		TMovieSceneChannel<FMovieSceneFloatValue> ChannelInterface = Channel->GetInterface();
+		for (int32 Index = 0; Index < InKeys.Num(); ++Index)
+		{
+			int32 KeyIndex = ChannelInterface.GetIndex(InKeys[Index]);
+			if (KeyIndex != INDEX_NONE)
+			{
+				FFrameTime NewTime = InKeyPositions[Index].InputValue * FrameResolution;
+
+				KeyIndex = ChannelInterface.MoveKey(KeyIndex, NewTime.FloorToFrame());
+				ChannelInterface.GetValues()[KeyIndex].Value = InKeyPositions[Index].OutputValue;
+			}
+		}
+
+		Channel->AutoSetTangents();
+	}
+}
+
+void FFloatChannelCurveModel::GetKeyAttributes(TArrayView<const FKeyHandle> InKeys, TArrayView<FKeyAttributes> OutAttributes) const
+{
+	FMovieSceneFloatChannel* Channel = ChannelHandle.Get();
+	UMovieSceneSection*      Section = WeakSection.Get();
+	if (Channel && Section)
+	{
+		TMovieSceneChannel<FMovieSceneFloatValue> ChannelInterface = Channel->GetInterface();
+		TArrayView<const FFrameNumber>    Times  = ChannelInterface.GetTimes();
+		TArrayView<FMovieSceneFloatValue> Values = ChannelInterface.GetValues();
+
+		float TimeInterval = Section->GetTypedOuter<UMovieScene>()->GetFrameResolution().AsInterval();
+
+		for (int32 Index = 0; Index < InKeys.Num(); ++Index)
+		{
+			const int32 KeyIndex = ChannelInterface.GetIndex(InKeys[Index]);
+			if (KeyIndex != INDEX_NONE)
+			{
+				const FMovieSceneFloatValue& KeyValue    = Values[KeyIndex];
+				FKeyAttributes&              Attributes  = OutAttributes[Index];
+
+				Attributes.SetInterpMode(KeyValue.InterpMode);
+
+				if (KeyValue.InterpMode != RCIM_Constant && KeyValue.InterpMode != RCIM_Linear)
+				{
+					Attributes.SetTangentMode(KeyValue.TangentMode);
+					if (KeyIndex != 0)
+					{
+						Attributes.SetArriveTangent(KeyValue.Tangent.ArriveTangent / TimeInterval);
+					}
+
+					if (KeyIndex != Times.Num()-1)
+					{
+						Attributes.SetLeaveTangent(KeyValue.Tangent.LeaveTangent / TimeInterval);
+					}
+				}
+			}
+		}
+	}
+}
+
+void FFloatChannelCurveModel::SetKeyAttributes(TArrayView<const FKeyHandle> InKeys, TArrayView<const FKeyAttributes> InAttributes)
+{
+	FMovieSceneFloatChannel* Channel = ChannelHandle.Get();
+	UMovieSceneSection*      Section = WeakSection.Get();
+	if (Channel && Section)
+	{
+		bool bAutoSetTangents = false;
+		Section->MarkAsChanged();
+
+		TMovieSceneChannel<FMovieSceneFloatValue> ChannelInterface = Channel->GetInterface();
+		TArrayView<FMovieSceneFloatValue> Values = ChannelInterface.GetValues();
+
+		float TimeInterval = Section->GetTypedOuter<UMovieScene>()->GetFrameResolution().AsInterval();
+
+		for (int32 Index = 0; Index < InKeys.Num(); ++Index)
+		{
+			const int32 KeyIndex = ChannelInterface.GetIndex(InKeys[Index]);
+			if (KeyIndex != INDEX_NONE)
+			{
+				const FKeyAttributes&  Attributes = InAttributes[Index];
+				FMovieSceneFloatValue& KeyValue   = Values[KeyIndex];
+
+				if (Attributes.HasInterpMode())    { KeyValue.InterpMode  = Attributes.GetInterpMode();  bAutoSetTangents = true; }
+				if (Attributes.HasTangentMode())   { KeyValue.TangentMode = Attributes.GetTangentMode(); bAutoSetTangents = true; }
+
+				if (Attributes.HasArriveTangent())
+				{
+					if (KeyValue.TangentMode == RCTM_Auto)
+					{
+						KeyValue.TangentMode = RCTM_User;
+					}
+
+					KeyValue.Tangent.ArriveTangent = Attributes.GetArriveTangent() * TimeInterval;
+					if (KeyValue.InterpMode == RCIM_Cubic && KeyValue.TangentMode != RCTM_Break)
+					{
+						KeyValue.Tangent.LeaveTangent = KeyValue.Tangent.ArriveTangent;
+					}
+				}
+
+				if (Attributes.HasLeaveTangent())
+				{
+					if (KeyValue.TangentMode == RCTM_Auto)
+					{
+						KeyValue.TangentMode = RCTM_User;
+					}
+
+					KeyValue.Tangent.LeaveTangent = Attributes.GetLeaveTangent() * TimeInterval;
+					if (KeyValue.InterpMode == RCIM_Cubic && KeyValue.TangentMode != RCTM_Break)
+					{
+						KeyValue.Tangent.ArriveTangent = KeyValue.Tangent.LeaveTangent;
+					}
+				}
+			}
+		}
+
+		if (bAutoSetTangents)
+		{
+			Channel->AutoSetTangents();
+		}
+	}
+}
+
+void FFloatChannelCurveModel::GetCurveAttributes(FCurveAttributes& OutCurveAttributes) const
+{
+	FMovieSceneFloatChannel* Channel = ChannelHandle.Get();
+	if (Channel)
+	{
+		OutCurveAttributes.SetPreExtrapolation(Channel->PreInfinityExtrap);
+		OutCurveAttributes.SetPostExtrapolation(Channel->PostInfinityExtrap);
+	}
+}
+
+void FFloatChannelCurveModel::SetCurveAttributes(const FCurveAttributes& InCurveAttributes)
+{
+	FMovieSceneFloatChannel* Channel = ChannelHandle.Get();
+	UMovieSceneSection*      Section = WeakSection.Get();
+	if (Channel && Section)
+	{
+		Section->MarkAsChanged();
+
+		if (InCurveAttributes.HasPreExtrapolation())
+		{
+			Channel->PreInfinityExtrap = InCurveAttributes.GetPreExtrapolation();
+		}
+
+		if (InCurveAttributes.HasPostExtrapolation())
+		{
+			Channel->PostInfinityExtrap = InCurveAttributes.GetPostExtrapolation();
+		}
+	}
+}
+
+class SCurveEditorKeyEditInterface : public SCompoundWidget
+{
+public:
+	SLATE_BEGIN_ARGS(SCurveEditorKeyEditInterface){}
+	SLATE_END_ARGS()
+
+	void Construct(
+		const FArguments& InArgs,
+		TSharedRef<ISequencer> InSequencer,
+		FCurveModelID InCurveID,
+		TWeakPtr<FCurveEditor> InWeakCurveEditor,
+		TWeakObjectPtr<UMovieSceneSection> InWeakSection,
+		TMovieSceneChannelHandle<FMovieSceneFloatChannel> InChannelHandle
+		)
+	{
+		CachedSelectionSerialNumber = 0;
+
+		CurveID         = InCurveID;
+		WeakCurveEditor = InWeakCurveEditor;
+		WeakSection     = InWeakSection;
+		ChannelHandle   = InChannelHandle;
+
+		ChildSlot
+		[
+			SAssignNew(Interface, SKeyEditInterface, InSequencer)
+			.EditData(this, &SCurveEditorKeyEditInterface::UpdateAndRetrieveEditData)
+		];
+	}
+
+	virtual void Tick( const FGeometry& AllottedGeometry, const double InCurrentTime, const float InDeltaTime )
+	{
+		TSharedPtr<FCurveEditor> CurveEditor = WeakCurveEditor.Pin();
+		if (CurveEditor.IsValid() && CurveEditor->Selection.GetSerialNumber() != CachedSelectionSerialNumber)
+		{
+			Interface->Initialize();
+			CachedSelectionSerialNumber = CurveEditor->Selection.GetSerialNumber();
+		}
+	}
+
+private:
+
+	FKeyEditData UpdateAndRetrieveEditData() const
+	{
+		FKeyEditData EditData;
+
+		TSharedPtr<FCurveEditor> CurveEditor = WeakCurveEditor.Pin();
+		if (CurveEditor.IsValid())
+		{
+			const FKeyHandleSet* SelectedKeys = CurveEditor->Selection.FindForCurve(CurveID.GetValue());
+			if (SelectedKeys && SelectedKeys->Num() == 1)
+			{
+				using namespace Sequencer;
+				EditData.KeyStruct = GetKeyStruct(ChannelHandle, SelectedKeys->AsArray()[0]);
+				EditData.OwningSection = WeakSection;
+			}
+		}
+
+		return EditData;
+	}
+
+	TOptional<FCurveModelID> CurveID;
+	TWeakPtr<FCurveEditor> WeakCurveEditor;
+	TWeakObjectPtr<UMovieSceneSection> WeakSection;
+	TMovieSceneChannelHandle<FMovieSceneFloatChannel> ChannelHandle;
+	uint32 CachedSelectionSerialNumber;
+	TSharedPtr<SKeyEditInterface> Interface;
+};
+
+TSharedPtr<SWidget> FFloatChannelCurveModel::CreateEditUI(TSharedPtr<FCurveEditor> InCurveEditor, FCurveModelID ThisCurveID)
+{
+	TSharedPtr<ISequencer> Sequencer = WeakSequencer.Pin();
+	if (Sequencer.IsValid())
+	{
+		return SNew(SCurveEditorKeyEditInterface, Sequencer.ToSharedRef(), ThisCurveID, TWeakPtr<FCurveEditor>(InCurveEditor), WeakSection, ChannelHandle);
+	}
+
+	return nullptr;
+}
\ No newline at end of file
diff --git a/Engine/Source/Editor/MovieSceneTools/Private/Channels/FloatChannelCurveModel.h b/Engine/Source/Editor/MovieSceneTools/Private/Channels/FloatChannelCurveModel.h
new file mode 100644
--- /dev/null
+++ b/Engine/Source/Editor/MovieSceneTools/Private/Channels/FloatChannelCurveModel.h
@@ -0,0 +1,51 @@
+// Copyright 1998-2018 Epic Games, Inc. All Rights Reserved.
+
+#pragma once
+
+#include "CoreTypes.h"
+#include "UObject/WeakObjectPtr.h"
+#include "UObject/WeakObjectPtrTemplates.h"
+#include "CurveModel.h"
+#include "Channels/MovieSceneChannelHandle.h"
+
+struct FMovieSceneFloatChannel;
+class UMovieSceneSection;
+class FCurveEditor;
+class ISequencer;
+
+class FFloatChannelCurveModel : public FCurveModel
+{
+public:
+	FFloatChannelCurveModel(TMovieSceneChannelHandle<FMovieSceneFloatChannel> InChannel, UMovieSceneSection* InOwningSection, TWeakPtr<ISequencer> InWeakSequencer);
+
+	virtual const void* GetCurve() const override;
+
+	virtual void Modify() override;
+
+	virtual void DrawCurve(const FCurveEditor& CurveEditor, TArray<TTuple<double, double>>& InterpolatingPoints) const override;
+	virtual void GetKeys(const FCurveEditor& CurveEditor, double MinTime, double MaxTime, double MinValue, double MaxValue, TArray<FKeyHandle>& OutKeyHandles) const override;
+	virtual void GetKeyDrawInfo(ECurvePointType PointType, FKeyDrawInfo& OutDrawInfo) const override;
+
+	virtual void GetKeyPositions(TArrayView<const FKeyHandle> InKeys, TArrayView<FKeyPosition> OutKeyPositions) const override;
+	virtual void SetKeyPositions(TArrayView<const FKeyHandle> InKeys, TArrayView<const FKeyPosition> InKeyPositions) override;
+
+	virtual void GetKeyAttributes(TArrayView<const FKeyHandle> InKeys, TArrayView<FKeyAttributes> OutAttributes) const override;
+	virtual void SetKeyAttributes(TArrayView<const FKeyHandle> InKeys, TArrayView<const FKeyAttributes> InAttributes) override;
+
+	virtual void GetCurveAttributes(FCurveAttributes& OutCurveAttributes) const override;
+	virtual void SetCurveAttributes(const FCurveAttributes& InCurveAttributes) override;
+
+	virtual bool Evaluate(double ProspectiveTime, double& OutValue) const override;
+	virtual void AddKeys(TArrayView<const FKeyPosition> InKeyPositions, TArrayView<const FKeyAttributes> InAttributes, TArrayView<TOptional<FKeyHandle>>* OutKeyHandles) override;
+	virtual void RemoveKeys(TArrayView<const FKeyHandle> InKeys) override;
+
+	virtual double GetInputDisplayOffset() const override;
+
+	virtual TSharedPtr<SWidget> CreateEditUI(TSharedPtr<FCurveEditor> InCurveEditor, FCurveModelID ThisCurveID) override;
+
+private:
+
+	TMovieSceneChannelHandle<FMovieSceneFloatChannel> ChannelHandle;
+	TWeakObjectPtr<UMovieSceneSection> WeakSection;
+	TWeakPtr<ISequencer> WeakSequencer;
+};
\ No newline at end of file
diff --git a/Engine/Source/Editor/MovieSceneTools/Private/CurveKeyEditors/SBoolCurveKeyEditor.cpp b/Engine/Source/Editor/MovieSceneTools/Private/CurveKeyEditors/SBoolCurveKeyEditor.cpp
--- a/Engine/Source/Editor/MovieSceneTools/Private/CurveKeyEditors/SBoolCurveKeyEditor.cpp
+++ b/Engine/Source/Editor/MovieSceneTools/Private/CurveKeyEditors/SBoolCurveKeyEditor.cpp
@@ -1,21 +1,16 @@
 // Copyright 1998-2018 Epic Games, Inc. All Rights Reserved.
 
 #include "CurveKeyEditors/SBoolCurveKeyEditor.h"
-#include "Styling/SlateTypes.h"
 #include "Widgets/Input/SCheckBox.h"
-#include "Curves/KeyHandle.h"
 #include "ISequencer.h"
 #include "ScopedTransaction.h"
-#include "Curves/IntegralCurve.h"
+#include "Channels/MovieSceneBoolChannel.h"
 
 #define LOCTEXT_NAMESPACE "BoolCurveKeyEditor"
 
-void SBoolCurveKeyEditor::Construct(const FArguments& InArgs)
+void SBoolCurveKeyEditor::Construct(const FArguments& InArgs, const TSequencerKeyEditor<FMovieSceneBoolChannel, bool>& InKeyEditor)
 {
-	Sequencer = InArgs._Sequencer;
-	OwningSection = InArgs._OwningSection;
-	Curve = InArgs._Curve;
-	ExternalValue = InArgs._ExternalValue;
+	KeyEditor = InKeyEditor;
 
 	ChildSlot
 	[
@@ -27,67 +22,15 @@ void SBoolCurveKeyEditor::Construct(const FArguments& InArgs)
 
 ECheckBoxState SBoolCurveKeyEditor::IsChecked() const
 {
-	bool bCurrentValue;
-	if (ExternalValue.IsSet() && ExternalValue.Get().IsSet())
-	{
-		bCurrentValue = ExternalValue.Get().GetValue();
-	}
-	else
-	{
-		float CurrentTime = Sequencer->GetLocalTime();
-		bool DefaultValue = false;
-		bCurrentValue = Curve->Evaluate(CurrentTime, DefaultValue) != 0;
-	}
-	
-	return bCurrentValue ? ECheckBoxState::Checked : ECheckBoxState::Unchecked;
+	return KeyEditor.GetCurrentValue() ? ECheckBoxState::Checked : ECheckBoxState::Unchecked;
 }
 
 void SBoolCurveKeyEditor::OnCheckStateChanged(ECheckBoxState NewCheckboxState)
 {
 	FScopedTransaction Transaction(LOCTEXT("SetBoolKey", "Set Bool Key Value"));
-	OwningSection->SetFlags(RF_Transactional);
-	if (OwningSection->TryModify())
-	{
-		float CurrentTime = Sequencer->GetLocalTime();
-		bool bAutoSetTrackDefaults = Sequencer->GetAutoSetTrackDefaults();
-		int32 NewValue = NewCheckboxState == ECheckBoxState::Checked ? 1 : 0;
 
-		FKeyHandle CurrentKeyHandle = Curve->FindKey(CurrentTime);
-		if (Curve->IsKeyHandleValid(CurrentKeyHandle))
-		{
-			Curve->SetKeyValue(CurrentKeyHandle, NewValue);
-		}
-		else
-		{
-			if (Curve->GetNumKeys() != 0 || bAutoSetTrackDefaults == false)
-			{
-				// When auto setting track defaults are disabled, add a key even when it's empty so that the changed
-				// value is saved and is propagated to the property.
-				Curve->AddKey(CurrentTime, NewValue, CurrentKeyHandle);
-			}
-
-			if (Curve->GetNumKeys() != 0)
-			{
-				if (OwningSection->GetStartTime() > CurrentTime)
-				{
-					OwningSection->SetStartTime(CurrentTime);
-				}
-				if (OwningSection->GetEndTime() < CurrentTime)
-				{
-					OwningSection->SetEndTime(CurrentTime);
-				}
-			}
-		}
-
-		// Always update the default value when auto-set default values is enabled so that the last changes
-		// are always saved to the track.
-		if (bAutoSetTrackDefaults)
-		{
-			Curve->SetDefaultValue(NewValue);
-		}
-
-		Sequencer->NotifyMovieSceneDataChanged( EMovieSceneDataChangeType::TrackValueChangedRefreshImmediately );
-	}
+	const bool bNewValue = NewCheckboxState == ECheckBoxState::Checked;
+	KeyEditor.SetValueWithNotify(bNewValue, EMovieSceneDataChangeType::TrackValueChangedRefreshImmediately);
 }
 
 #undef LOCTEXT_NAMESPACE
diff --git a/Engine/Source/Editor/MovieSceneTools/Private/CurveKeyEditors/SBoolCurveKeyEditor.h b/Engine/Source/Editor/MovieSceneTools/Private/CurveKeyEditors/SBoolCurveKeyEditor.h
--- a/Engine/Source/Editor/MovieSceneTools/Private/CurveKeyEditors/SBoolCurveKeyEditor.h
+++ b/Engine/Source/Editor/MovieSceneTools/Private/CurveKeyEditors/SBoolCurveKeyEditor.h
@@ -2,14 +2,11 @@
 
 #pragma once
 
-#include "CoreMinimal.h"
-#include "Misc/Attribute.h"
 #include "Widgets/DeclarativeSyntaxSupport.h"
 #include "Widgets/SCompoundWidget.h"
+#include "SequencerKeyEditor.h"
 
-class ISequencer;
-class UMovieSceneSection;
-struct FIntegralCurve;
+struct FMovieSceneBoolChannel;
 enum class ECheckBoxState : uint8;
 
 /**
@@ -19,31 +16,14 @@ class SBoolCurveKeyEditor : public SCompoundWidget
 {
 public:
 	SLATE_BEGIN_ARGS(SBoolCurveKeyEditor) {}
-
-		/** The sequencer which is controlling this key editor. */
-		SLATE_ARGUMENT(ISequencer*, Sequencer)
-
-		/** The section that owns the data edited by this key editor. */
-		SLATE_ARGUMENT(UMovieSceneSection*, OwningSection)
-
-		/** The curve being edited by this curve editor. */
-		SLATE_ARGUMENT(FIntegralCurve*, Curve)
-
-		/** Allows the value displayed and edited by this key editor to be supplied from an external source.  This
-			is useful for curves on property tracks who's property value can change without changing the animation. */
-		SLATE_ATTRIBUTE(TOptional<bool>, ExternalValue)
-
 	SLATE_END_ARGS();
 
-	void Construct(const FArguments& InArgs);
+	void Construct(const FArguments& InArgs, const TSequencerKeyEditor<FMovieSceneBoolChannel, bool>& InKeyEditor);
 
 private:
 
 	ECheckBoxState IsChecked() const;
 	void OnCheckStateChanged(ECheckBoxState);
 
-	ISequencer* Sequencer;
-	UMovieSceneSection* OwningSection;
-	FIntegralCurve* Curve;
-	TAttribute<TOptional<bool>> ExternalValue;
+	TSequencerKeyEditor<FMovieSceneBoolChannel, bool> KeyEditor;
 };
diff --git a/Engine/Source/Editor/MovieSceneTools/Private/CurveKeyEditors/SNumericKeyEditor.h b/Engine/Source/Editor/MovieSceneTools/Private/CurveKeyEditors/SNumericKeyEditor.h
new file mode 100644
--- /dev/null
+++ b/Engine/Source/Editor/MovieSceneTools/Private/CurveKeyEditors/SNumericKeyEditor.h
@@ -0,0 +1,83 @@
+// Copyright 1998-2018 Epic Games, Inc. All Rights Reserved.
+
+#pragma once
+
+#include "Widgets/DeclarativeSyntaxSupport.h"
+#include "Widgets/SCompoundWidget.h"
+#include "ScopedTransaction.h"
+#include "Styling/SlateTypes.h"
+#include "Editor.h"
+#include "SequencerKeyEditor.h"
+#include "Widgets/Input/SSpinBox.h"
+#include "EditorStyleSet.h"
+#include "SequencerKeyEditor.h"
+
+#define LOCTEXT_NAMESPACE "NumericKeyEditor"
+
+/**
+ * A widget for editing a curve representing integer keys.
+ */
+template<typename ChannelType, typename NumericType>
+class SNumericKeyEditor : public SCompoundWidget
+{
+public:
+	SLATE_BEGIN_ARGS(SNumericKeyEditor){}
+	SLATE_END_ARGS()
+
+	void Construct(const FArguments& InArgs, const TSequencerKeyEditor<ChannelType, NumericType>& InKeyEditor)
+	{
+		KeyEditor = InKeyEditor;
+		ChildSlot
+		[
+			SNew(SSpinBox<NumericType>)
+			.Style(&FEditorStyle::GetWidgetStyle<FSpinBoxStyle>("Sequencer.HyperlinkSpinBox"))
+			.Font(FEditorStyle::GetFontStyle("Sequencer.AnimationOutliner.RegularFont"))
+			.MinValue(TOptional<NumericType>())
+			.MaxValue(TOptional<NumericType>())
+			.MaxSliderValue(TOptional<NumericType>())
+			.MinSliderValue(TOptional<NumericType>())
+			.Value_Raw(&KeyEditor, &decltype(KeyEditor)::GetCurrentValue)
+			.OnValueChanged(this, &SNumericKeyEditor::OnValueChanged)
+			.OnValueCommitted(this, &SNumericKeyEditor::OnValueCommitted)
+			.OnBeginSliderMovement(this, &SNumericKeyEditor::OnBeginSliderMovement)
+			.OnEndSliderMovement(this, &SNumericKeyEditor::OnEndSliderMovement)
+			.ClearKeyboardFocusOnCommit(true)
+		];
+	}
+
+private:
+
+	void OnBeginSliderMovement()
+	{
+		GEditor->BeginTransaction(LOCTEXT("SetNumericKey", "Set Key Value"));
+	}
+
+	void OnEndSliderMovement(NumericType Value)
+	{
+		if (GEditor->IsTransactionActive())
+		{
+			KeyEditor.SetValue(Value);
+			GEditor->EndTransaction();
+		}
+	}
+
+	void OnValueChanged(NumericType Value)
+	{
+		KeyEditor.SetValueWithNotify(Value, EMovieSceneDataChangeType::TrackValueChanged);
+	}
+
+	void OnValueCommitted(NumericType Value, ETextCommit::Type CommitInfo)
+	{
+		if (CommitInfo == ETextCommit::OnEnter || CommitInfo == ETextCommit::OnUserMovedFocus)
+		{
+			const FScopedTransaction Transaction( LOCTEXT("SetNumericKey", "Set Key Value") );
+			KeyEditor.SetValueWithNotify(Value, EMovieSceneDataChangeType::TrackValueChangedRefreshImmediately);
+		}
+	}
+
+private:
+
+	TSequencerKeyEditor<ChannelType, NumericType> KeyEditor;
+};
+
+#undef LOCTEXT_NAMESPACE
diff --git a/Engine/Source/Editor/MovieSceneTools/Private/CurveKeyEditors/SStringCurveKeyEditor.cpp b/Engine/Source/Editor/MovieSceneTools/Private/CurveKeyEditors/SStringCurveKeyEditor.cpp
--- a/Engine/Source/Editor/MovieSceneTools/Private/CurveKeyEditors/SStringCurveKeyEditor.cpp
+++ b/Engine/Source/Editor/MovieSceneTools/Private/CurveKeyEditors/SStringCurveKeyEditor.cpp
@@ -2,21 +2,15 @@
 
 #include "CurveKeyEditors/SStringCurveKeyEditor.h"
 #include "Widgets/Input/SEditableText.h"
-#include "Curves/KeyHandle.h"
-#include "ISequencer.h"
 #include "ScopedTransaction.h"
-#include "Curves/StringCurve.h"
+#include "Sections/MovieSceneStringSection.h"
 
 #define LOCTEXT_NAMESPACE "StringCurveKeyEditor"
 
-void SStringCurveKeyEditor::Construct(const FArguments& InArgs)
+void SStringCurveKeyEditor::Construct(const FArguments& InArgs, const TSequencerKeyEditor<FMovieSceneStringChannel, FString>& InKeyEditor)
 {
-	Sequencer = InArgs._Sequencer;
-	OwningSection = InArgs._OwningSection;
-	Curve = InArgs._Curve;
-	ExternalValue = InArgs._ExternalValue;
+	KeyEditor = InKeyEditor;
 
-	float CurrentTime = Sequencer->GetLocalTime();
 	ChildSlot
 	[
 		SNew(SEditableText)
@@ -28,62 +22,13 @@ void SStringCurveKeyEditor::Construct(const FArguments& InArgs)
 
 FText SStringCurveKeyEditor::GetText() const
 {
-	if (ExternalValue.IsSet() && ExternalValue.Get().IsSet())
-	{
-		return FText::FromString(ExternalValue.Get().GetValue());
-	}
-	float CurrentTime = Sequencer->GetLocalTime();
-	FString DefaultValue;
-	FString CurrentValue = Curve->Eval(CurrentTime, DefaultValue);
-	
-	return FText::FromString(CurrentValue);
+	return FText::FromString(KeyEditor.GetCurrentValue());
 }
 
 void SStringCurveKeyEditor::OnTextCommitted(const FText& InText, ETextCommit::Type InCommitType)
 {
 	FScopedTransaction Transaction(LOCTEXT("SetStringKey", "Set String Key Value"));
-	OwningSection->SetFlags(RF_Transactional);
-	if (OwningSection->TryModify())
-	{
-		float CurrentTime = Sequencer->GetLocalTime();
-		bool bAutoSetTrackDefaults = Sequencer->GetAutoSetTrackDefaults();
-
-		FKeyHandle CurrentKeyHandle = Curve->FindKey(CurrentTime);
-		if (Curve->IsKeyHandleValid(CurrentKeyHandle))
-		{
-			Curve->SetKeyValue(CurrentKeyHandle, InText.ToString());
-		}
-		else
-		{
-			if (Curve->GetNumKeys() != 0 || bAutoSetTrackDefaults == false)
-			{
-				// When auto setting track defaults are disabled, add a key even when it's empty so that the changed
-				// value is saved and is propagated to the property.
-				Curve->AddKey(CurrentTime, InText.ToString(), CurrentKeyHandle);
-			}
-
-			if (Curve->GetNumKeys() != 0)
-			{
-				if (OwningSection->GetStartTime() > CurrentTime)
-				{
-					OwningSection->SetStartTime(CurrentTime);
-				}
-				if (OwningSection->GetEndTime() < CurrentTime)
-				{
-					OwningSection->SetEndTime(CurrentTime);
-				}
-			}
-		}
-
-		// Always update the default value when auto-set default values is enabled so that the last changes
-		// are always saved to the track.
-		if (bAutoSetTrackDefaults)
-		{
-			Curve->SetDefaultValue(InText.ToString());
-		}
-
-		Sequencer->NotifyMovieSceneDataChanged( EMovieSceneDataChangeType::TrackValueChangedRefreshImmediately );
-	}
+	KeyEditor.SetValueWithNotify(InText.ToString(), EMovieSceneDataChangeType::TrackValueChangedRefreshImmediately);
 }
 
 #undef LOCTEXT_NAMESPACE
diff --git a/Engine/Source/Editor/MovieSceneTools/Private/CurveKeyEditors/SStringCurveKeyEditor.h b/Engine/Source/Editor/MovieSceneTools/Private/CurveKeyEditors/SStringCurveKeyEditor.h
--- a/Engine/Source/Editor/MovieSceneTools/Private/CurveKeyEditors/SStringCurveKeyEditor.h
+++ b/Engine/Source/Editor/MovieSceneTools/Private/CurveKeyEditors/SStringCurveKeyEditor.h
@@ -2,14 +2,11 @@
 
 #pragma once
 
-#include "CoreMinimal.h"
-#include "Misc/Attribute.h"
 #include "Widgets/DeclarativeSyntaxSupport.h"
 #include "Widgets/SCompoundWidget.h"
+#include "SequencerKeyEditor.h"
 
-class ISequencer;
-class UMovieSceneSection;
-struct FStringCurve;
+struct FMovieSceneStringChannel;
 
 /**
  * A widget for editing a curve representing string keys.
@@ -18,31 +15,14 @@ class SStringCurveKeyEditor : public SCompoundWidget
 {
 public:
 	SLATE_BEGIN_ARGS(SStringCurveKeyEditor) {}
-
-		/** The sequencer which is controlling this key editor. */
-		SLATE_ARGUMENT(ISequencer*, Sequencer)
-
-		/** The section that owns the data edited by this key editor. */
-		SLATE_ARGUMENT(UMovieSceneSection*, OwningSection)
-
-		/** The curve being edited by this curve editor. */
-		SLATE_ARGUMENT(FStringCurve*, Curve)
-
-		/** Allows the value displayed and edited by this key editor to be supplied from an external source.  This
-		is useful for curves on property tracks who's property value can change without changing the animation. */
-		SLATE_ATTRIBUTE(TOptional<FString>, ExternalValue)
-
 	SLATE_END_ARGS();
 
-	void Construct(const FArguments& InArgs);
+	void Construct(const FArguments& InArgs, const TSequencerKeyEditor<FMovieSceneStringChannel, FString>& InKeyEditor);
 
 private:
 
 	FText GetText() const;
 	void OnTextCommitted(const FText& InText, ETextCommit::Type InCommitType);
 
-	ISequencer* Sequencer;
-	UMovieSceneSection* OwningSection;
-	FStringCurve* Curve;
-	TAttribute<TOptional<FString>> ExternalValue;
+	TSequencerKeyEditor<FMovieSceneStringChannel, FString> KeyEditor;
 };
