Copying //UE4/Dev-Sequencer to //UE4/Dev-Main (Source: //UE4/Dev-Sequencer @ 3946692)

diff --git a/Engine/Source/Editor/MovieSceneTools/Private/MatineeImportTools.cpp b/Engine/Source/Editor/MovieSceneTools/Private/MatineeImportTools.cpp
--- a/Engine/Source/Editor/MovieSceneTools/Private/MatineeImportTools.cpp
+++ b/Engine/Source/Editor/MovieSceneTools/Private/MatineeImportTools.cpp
@@ -6,8 +6,8 @@
 #include "ScopedTransaction.h"
 #include "MovieSceneCommonHelpers.h"
 #include "IMovieScenePlayer.h"
-
-#include "Curves/CurveInterface.h"
+#include "Channels/MovieSceneChannelProxy.h"
+#include "MovieScene.h"
 
 #include "Matinee/MatineeActor.h"
 #include "Matinee/InterpData.h"
@@ -87,14 +87,16 @@ ERichCurveTangentMode FMatineeImportTools::MatineeInterpolationToRichCurveTangen
 	}
 }
 
-void CleanupCurveKeys(FRichCurve& InCurve)
+void CleanupCurveKeys(FMovieSceneFloatChannel* InChannel)
 {
-	InCurve.RemoveRedundantKeys(KINDA_SMALL_NUMBER);
-	InCurve.AutoSetTangents();
+	FKeyDataOptimizationParams Params;
+	Params.bAutoSetInterpolation = true;
+
+	MovieScene::Optimize(InChannel, Params);
 }
 
 
-bool FMatineeImportTools::TryConvertMatineeToggleToOutParticleKey( ETrackToggleAction ToggleAction, EParticleKey::Type& OutParticleKey )
+bool FMatineeImportTools::TryConvertMatineeToggleToOutParticleKey( ETrackToggleAction ToggleAction, EParticleKey& OutParticleKey )
 {
 	switch ( ToggleAction )
 	{
@@ -112,20 +114,18 @@ bool FMatineeImportTools::TryConvertMatineeToggleToOutParticleKey( ETrackToggleA
 }
 
 
-void FMatineeImportTools::SetOrAddKey( FRichCurve& Curve, float Time, float Value, float ArriveTangent, float LeaveTangent, EInterpCurveMode MatineeInterpMode )
+void FMatineeImportTools::SetOrAddKey(TMovieSceneChannel<FMovieSceneFloatValue>& ChannelInterface, FFrameNumber Time, float Value, float ArriveTangent, float LeaveTangent, EInterpCurveMode MatineeInterpMode )
 {
-	FKeyHandle KeyHandle = Curve.FindKey(Time);
-
-	if (!Curve.IsKeyHandleValid(KeyHandle))
+	if (ChannelInterface.FindKey(Time) == INDEX_NONE)
 	{
-		KeyHandle = Curve.AddKey( Time, Value, false);
-	}
+		FMovieSceneFloatValue NewKey(Value);
 
-	FRichCurveKey& Key = Curve.GetKey( KeyHandle );
-	Key.ArriveTangent = ArriveTangent;
-	Key.LeaveTangent = LeaveTangent;
-	Key.InterpMode = MatineeInterpolationToRichCurveInterpolation( MatineeInterpMode );
-	Key.TangentMode = MatineeInterpolationToRichCurveTangent( MatineeInterpMode );
+		NewKey.InterpMode = MatineeInterpolationToRichCurveInterpolation( MatineeInterpMode );
+		NewKey.TangentMode = MatineeInterpolationToRichCurveTangent( MatineeInterpMode );
+		NewKey.Tangent.ArriveTangent = ArriveTangent;
+		NewKey.Tangent.LeaveTangent = LeaveTangent;
+		ChannelInterface.AddKey( Time, NewKey );
+	}
 }
 
 
@@ -136,30 +136,35 @@ bool FMatineeImportTools::CopyInterpBoolTrack( UInterpTrackBoolProp* MatineeBool
 
 	BoolTrack->Modify();
 
-	float KeyTime = MatineeBoolTrack->GetKeyframeTime( 0 );
-	UMovieSceneBoolSection* Section = Cast<UMovieSceneBoolSection>( MovieSceneHelpers::FindSectionAtTime( BoolTrack->GetAllSections(), KeyTime ) );
+	FFrameRate   FrameRate    = BoolTrack->GetTypedOuter<UMovieScene>()->GetFrameResolution();
+	FFrameNumber FirstKeyTime = (MatineeBoolTrack->GetKeyframeTime( 0 ) * FrameRate).RoundToFrame();
+
+	UMovieSceneBoolSection* Section = Cast<UMovieSceneBoolSection>( MovieSceneHelpers::FindSectionAtTime( BoolTrack->GetAllSections(), FirstKeyTime ) );
 	if ( Section == nullptr )
 	{
 		Section = Cast<UMovieSceneBoolSection>( BoolTrack->CreateNewSection() );
 		BoolTrack->AddSection( *Section );
-		Section->SetIsInfinite(true);
+		Section->SetRange(TRange<FFrameNumber>::All());
 		bSectionCreated = true;
 	}
 	if (Section->TryModify())
 	{
-		float SectionMin = Section->GetStartTime();
-		float SectionMax = Section->GetEndTime();
+		TRange<FFrameNumber> KeyRange = TRange<FFrameNumber>::Empty();
 
-		FIntegralCurve& BoolCurve = Section->GetCurve();
+		TMovieSceneChannel<bool> ChannelInterface = Section->GetChannelProxy().GetChannel<FMovieSceneBoolChannel>(0)->GetInterface();
 		for ( const auto& Point : MatineeBoolTrack->BoolTrack )
 		{
-			BoolCurve.UpdateOrAddKey(Point.Time, Point.Value);
-			SectionMin = FMath::Min( SectionMin, Point.Time );
-			SectionMax = FMath::Max( SectionMax, Point.Time );
+			FFrameNumber KeyTime = (Point.Time * FrameRate).RoundToFrame();
+
+			ChannelInterface.UpdateOrAddKey(KeyTime, Point.Value);
+
+			KeyRange = TRange<FFrameNumber>::Hull(KeyRange, TRange<FFrameNumber>(KeyTime));
 		}
 
-		Section->SetStartTime( SectionMin );
-		Section->SetEndTime( SectionMax );
+		if (!KeyRange.IsEmpty())
+		{
+			Section->SetRange( KeyRange );
+		}
 	}
 
 	return bSectionCreated;
@@ -172,32 +177,39 @@ bool FMatineeImportTools::CopyInterpFloatTrack( UInterpTrackFloatBase* MatineeFl
 
 	FloatTrack->Modify();
 
-	float KeyTime = MatineeFloatTrack->GetKeyframeTime( 0 );
-	UMovieSceneFloatSection* Section = Cast<UMovieSceneFloatSection>( MovieSceneHelpers::FindSectionAtTime( FloatTrack->GetAllSections(), KeyTime ) );
+	FFrameRate   FrameRate    = FloatTrack->GetTypedOuter<UMovieScene>()->GetFrameResolution();
+	FFrameNumber FirstKeyTime = (MatineeFloatTrack->GetKeyframeTime( 0 ) * FrameRate).RoundToFrame();;
+
+	UMovieSceneFloatSection* Section = Cast<UMovieSceneFloatSection>( MovieSceneHelpers::FindSectionAtTime( FloatTrack->GetAllSections(), FirstKeyTime ) );
 	if ( Section == nullptr )
 	{
 		Section = Cast<UMovieSceneFloatSection>( FloatTrack->CreateNewSection() );
 		FloatTrack->AddSection( *Section );
-		Section->SetIsInfinite(true);
+		Section->SetRange(TRange<FFrameNumber>::All());
 		bSectionCreated = true;
 	}
 	if (Section->TryModify())
 	{
-		float SectionMin = Section->GetStartTime();
-		float SectionMax = Section->GetEndTime();
+		TRange<FFrameNumber> KeyRange = TRange<FFrameNumber>::Empty();
+
+		FMovieSceneFloatChannel* Channel = Section->GetChannelProxy().GetChannel<FMovieSceneFloatChannel>(0);
+		TMovieSceneChannel<FMovieSceneFloatValue> ChannelInterface = Channel->GetInterface();
 
-		FRichCurve& FloatCurve = Section->GetFloatCurve();
 		for ( const auto& Point : MatineeFloatTrack->FloatTrack.Points )
 		{
-			FMatineeImportTools::SetOrAddKey( FloatCurve, Point.InVal, Point.OutVal, Point.ArriveTangent, Point.LeaveTangent, Point.InterpMode );
-			SectionMin = FMath::Min( SectionMin, Point.InVal );
-			SectionMax = FMath::Max( SectionMax, Point.InVal );
+			FFrameNumber KeyTime = (Point.InVal * FrameRate).RoundToFrame();
+
+			FMatineeImportTools::SetOrAddKey( ChannelInterface, KeyTime, Point.OutVal, Point.ArriveTangent, Point.LeaveTangent, Point.InterpMode );
+
+			KeyRange = TRange<FFrameNumber>::Hull(KeyRange, TRange<FFrameNumber>(KeyTime));
 		}
 
-		CleanupCurveKeys(FloatCurve);
+		CleanupCurveKeys(Channel);
 
-		Section->SetStartTime( SectionMin );
-		Section->SetEndTime( SectionMax );
+		if (!KeyRange.IsEmpty())
+		{
+			Section->SetRange( KeyRange );
+		}
 	}
 
 	return bSectionCreated;
@@ -210,42 +222,46 @@ bool FMatineeImportTools::CopyInterpVectorTrack( UInterpTrackVectorProp* Matinee
 
 	VectorTrack->Modify();
 
-	float KeyTime = MatineeVectorTrack->GetKeyframeTime( 0 );
-	UMovieSceneVectorSection* Section = Cast<UMovieSceneVectorSection>( MovieSceneHelpers::FindSectionAtTime( VectorTrack->GetAllSections(), KeyTime ) );
+	FFrameRate   FrameRate    = VectorTrack->GetTypedOuter<UMovieScene>()->GetFrameResolution();
+	FFrameNumber FirstKeyTime = (MatineeVectorTrack->GetKeyframeTime( 0 ) * FrameRate).RoundToFrame();
+
+	UMovieSceneVectorSection* Section = Cast<UMovieSceneVectorSection>( MovieSceneHelpers::FindSectionAtTime( VectorTrack->GetAllSections(), FirstKeyTime ) );
 	if ( Section == nullptr )
 	{
 		Section = Cast<UMovieSceneVectorSection>( VectorTrack->CreateNewSection() );
 		VectorTrack->AddSection( *Section );
-		Section->SetIsInfinite(true);
+		Section->SetRange(TRange<FFrameNumber>::All());
 		bSectionCreated = true;
 	}
 	if (Section->TryModify())
 	{
-		float SectionMin = Section->GetStartTime();
-		float SectionMax = Section->GetEndTime();
+		TRange<FFrameNumber> KeyRange = TRange<FFrameNumber>::Empty();
 
 		if (Section->GetChannelsUsed() == 3)
 		{
-			FRichCurve& XCurve = Section->GetCurve(0);
-			FRichCurve& YCurve = Section->GetCurve(1);
-			FRichCurve& ZCurve = Section->GetCurve(2);
+			TArrayView<FMovieSceneFloatChannel*> Channels = Section->GetChannelProxy().GetChannels<FMovieSceneFloatChannel>();
+			TMovieSceneChannel<FMovieSceneFloatValue> ChannelInterfaces[3] = { Channels[0]->GetInterface(), Channels[1]->GetInterface(), Channels[2]->GetInterface() };
 
 			for ( const auto& Point : MatineeVectorTrack->VectorTrack.Points )
 			{
-				FMatineeImportTools::SetOrAddKey( XCurve, Point.InVal, Point.OutVal.X, Point.ArriveTangent.X, Point.LeaveTangent.X, Point.InterpMode );
-				FMatineeImportTools::SetOrAddKey( YCurve, Point.InVal, Point.OutVal.Y, Point.ArriveTangent.Y, Point.LeaveTangent.Y, Point.InterpMode );
-				FMatineeImportTools::SetOrAddKey( ZCurve, Point.InVal, Point.OutVal.Z, Point.ArriveTangent.Z, Point.LeaveTangent.Z, Point.InterpMode );
-				SectionMin = FMath::Min( SectionMin, Point.InVal );
-				SectionMax = FMath::Max( SectionMax, Point.InVal );
+				FFrameNumber KeyTime = (Point.InVal * FrameRate).RoundToFrame();
+
+				FMatineeImportTools::SetOrAddKey( ChannelInterfaces[0], KeyTime, Point.OutVal.X, Point.ArriveTangent.X, Point.LeaveTangent.X, Point.InterpMode );
+				FMatineeImportTools::SetOrAddKey( ChannelInterfaces[1], KeyTime, Point.OutVal.Y, Point.ArriveTangent.Y, Point.LeaveTangent.Y, Point.InterpMode );
+				FMatineeImportTools::SetOrAddKey( ChannelInterfaces[2], KeyTime, Point.OutVal.Z, Point.ArriveTangent.Z, Point.LeaveTangent.Z, Point.InterpMode );
+
+				KeyRange = TRange<FFrameNumber>::Hull(KeyRange, TRange<FFrameNumber>(KeyTime));
 			}
 			
-			CleanupCurveKeys(XCurve);
-			CleanupCurveKeys(YCurve);
-			CleanupCurveKeys(ZCurve);
+			CleanupCurveKeys(Channels[0]);
+			CleanupCurveKeys(Channels[1]);
+			CleanupCurveKeys(Channels[2]);
 		}
 
-		Section->SetStartTime( SectionMin );
-		Section->SetEndTime( SectionMax );
+		if (!KeyRange.IsEmpty())
+		{
+			Section->SetRange( KeyRange );
+		}
 	}
 
 	return bSectionCreated;
@@ -259,44 +275,51 @@ bool FMatineeImportTools::CopyInterpColorTrack( UInterpTrackColorProp* ColorProp
 
 	ColorTrack->Modify();
 
-	float KeyTime = ColorPropTrack->GetKeyframeTime( 0 );
-	UMovieSceneColorSection* Section = Cast<UMovieSceneColorSection>( MovieSceneHelpers::FindSectionAtTime( ColorTrack->GetAllSections(), KeyTime ) );
+	FFrameRate   FrameRate    = ColorTrack->GetTypedOuter<UMovieScene>()->GetFrameResolution();
+	FFrameNumber FirstKeyTime = (ColorPropTrack->GetKeyframeTime( 0 ) * FrameRate).RoundToFrame();
+
+	UMovieSceneColorSection* Section = Cast<UMovieSceneColorSection>( MovieSceneHelpers::FindSectionAtTime( ColorTrack->GetAllSections(), FirstKeyTime ) );
 	if ( Section == nullptr )
 	{
 		Section = Cast<UMovieSceneColorSection>( ColorTrack->CreateNewSection() );
 		ColorTrack->AddSection( *Section );
-		Section->GetRedCurve().SetDefaultValue(0.f);
-		Section->GetGreenCurve().SetDefaultValue(0.f);
-		Section->GetBlueCurve().SetDefaultValue(0.f);
-		Section->GetAlphaCurve().SetDefaultValue(1.f);
-		Section->SetIsInfinite(true);
+
+		TArrayView<FMovieSceneFloatChannel*> FloatChannels = Section->GetChannelProxy().GetChannels<FMovieSceneFloatChannel>();
+		FloatChannels[0]->SetDefault(0.f);
+		FloatChannels[1]->SetDefault(0.f);
+		FloatChannels[2]->SetDefault(0.f);
+		FloatChannels[3]->SetDefault(1.f);
+
+		Section->SetRange(TRange<FFrameNumber>::All());
 		bSectionCreated = true;
 	}
 
 	if (Section->TryModify())
 	{
-		float SectionMin = Section->GetStartTime();
-		float SectionMax = Section->GetEndTime();
+		TArrayView<FMovieSceneFloatChannel*> Channels = Section->GetChannelProxy().GetChannels<FMovieSceneFloatChannel>();
+		TMovieSceneChannel<FMovieSceneFloatValue> ChannelInterfaces[3] = { Channels[0]->GetInterface(), Channels[1]->GetInterface(), Channels[2]->GetInterface() };
 
-		FRichCurve& RedCurve = Section->GetRedCurve();
-		FRichCurve& GreenCurve = Section->GetGreenCurve();
-		FRichCurve& BlueCurve = Section->GetBlueCurve();
+		TRange<FFrameNumber> KeyRange = TRange<FFrameNumber>::Empty();
 
-		for ( const auto& Point : ColorPropTrack->VectorTrack.Points )
+		for ( const FInterpCurvePoint<FVector>& Point : ColorPropTrack->VectorTrack.Points )
 		{
-			FMatineeImportTools::SetOrAddKey( RedCurve, Point.InVal, Point.OutVal.X, Point.ArriveTangent.X, Point.LeaveTangent.X, Point.InterpMode );
-			FMatineeImportTools::SetOrAddKey( GreenCurve, Point.InVal, Point.OutVal.Y, Point.ArriveTangent.Y, Point.LeaveTangent.Y, Point.InterpMode );
-			FMatineeImportTools::SetOrAddKey( BlueCurve, Point.InVal, Point.OutVal.Z, Point.ArriveTangent.Z, Point.LeaveTangent.Z, Point.InterpMode );
-			SectionMin = FMath::Min( SectionMin, Point.InVal );
-			SectionMax = FMath::Max( SectionMax, Point.InVal );
+			FFrameNumber KeyTime = (Point.InVal * FrameRate).RoundToFrame();
+
+			FMatineeImportTools::SetOrAddKey( ChannelInterfaces[0], KeyTime, Point.OutVal.X, Point.ArriveTangent.X, Point.LeaveTangent.X, Point.InterpMode );
+			FMatineeImportTools::SetOrAddKey( ChannelInterfaces[1], KeyTime, Point.OutVal.Y, Point.ArriveTangent.Y, Point.LeaveTangent.Y, Point.InterpMode );
+			FMatineeImportTools::SetOrAddKey( ChannelInterfaces[2], KeyTime, Point.OutVal.Z, Point.ArriveTangent.Z, Point.LeaveTangent.Z, Point.InterpMode );
+
+			KeyRange = TRange<FFrameNumber>::Hull(KeyRange, TRange<FFrameNumber>(KeyTime));
 		}
 
-		CleanupCurveKeys(RedCurve);
-		CleanupCurveKeys(GreenCurve);
-		CleanupCurveKeys(BlueCurve);
+		CleanupCurveKeys(Channels[0]);
+		CleanupCurveKeys(Channels[1]);
+		CleanupCurveKeys(Channels[2]);
 
-		Section->SetStartTime( SectionMin );
-		Section->SetEndTime( SectionMax );
+		if (!KeyRange.IsEmpty())
+		{
+			Section->SetRange( KeyRange );
+		}
 	}
 
 	return bSectionCreated;
@@ -310,47 +333,53 @@ bool FMatineeImportTools::CopyInterpLinearColorTrack( UInterpTrackLinearColorPro
 
 	ColorTrack->Modify();
 
-	float KeyTime = LinearColorPropTrack->GetKeyframeTime( 0 );
-	UMovieSceneColorSection* Section = Cast<UMovieSceneColorSection>( MovieSceneHelpers::FindSectionAtTime( ColorTrack->GetAllSections(), KeyTime ) );
+	FFrameRate   FrameRate    = ColorTrack->GetTypedOuter<UMovieScene>()->GetFrameResolution();
+	FFrameNumber FirstKeyTime = (LinearColorPropTrack->GetKeyframeTime( 0 ) * FrameRate).RoundToFrame();
+
+	UMovieSceneColorSection* Section = Cast<UMovieSceneColorSection>( MovieSceneHelpers::FindSectionAtTime( ColorTrack->GetAllSections(), FirstKeyTime ) );
 	if ( Section == nullptr )
 	{
 		Section = Cast<UMovieSceneColorSection>( ColorTrack->CreateNewSection() );
 		ColorTrack->AddSection( *Section );
-		Section->GetRedCurve().SetDefaultValue(0.f);
-		Section->GetGreenCurve().SetDefaultValue(0.f);
-		Section->GetBlueCurve().SetDefaultValue(0.f);
-		Section->GetAlphaCurve().SetDefaultValue(1.f);
-		Section->SetIsInfinite(true);
+
+		TArrayView<FMovieSceneFloatChannel*> FloatChannels = Section->GetChannelProxy().GetChannels<FMovieSceneFloatChannel>();
+		FloatChannels[0]->SetDefault(0.f);
+		FloatChannels[1]->SetDefault(0.f);
+		FloatChannels[2]->SetDefault(0.f);
+		FloatChannels[3]->SetDefault(1.f);
+
+		Section->SetRange(TRange<FFrameNumber>::All());
 		bSectionCreated = true;
 	}
 
 	if (Section->TryModify())
 	{
-		float SectionMin = Section->GetStartTime();
-		float SectionMax = Section->GetEndTime();
+		TRange<FFrameNumber> KeyRange = TRange<FFrameNumber>::Empty();
 
-		FRichCurve& RedCurve = Section->GetRedCurve();
-		FRichCurve& GreenCurve = Section->GetGreenCurve();
-		FRichCurve& BlueCurve = Section->GetBlueCurve();
-		FRichCurve& AlphaCurve = Section->GetAlphaCurve();
+		TArrayView<FMovieSceneFloatChannel*> Channels = Section->GetChannelProxy().GetChannels<FMovieSceneFloatChannel>();
+		TMovieSceneChannel<FMovieSceneFloatValue> ChannelInterfaces[4] = { Channels[0]->GetInterface(), Channels[1]->GetInterface(), Channels[2]->GetInterface(), Channels[3]->GetInterface() };
 
-		for ( const auto& Point : LinearColorPropTrack->LinearColorTrack.Points)
+		for ( const auto& Point : LinearColorPropTrack->LinearColorTrack.Points )
 		{
-			FMatineeImportTools::SetOrAddKey( RedCurve, Point.InVal, Point.OutVal.R, Point.ArriveTangent.R, Point.LeaveTangent.R, Point.InterpMode );
-			FMatineeImportTools::SetOrAddKey( GreenCurve, Point.InVal, Point.OutVal.G, Point.ArriveTangent.G, Point.LeaveTangent.G, Point.InterpMode );
-			FMatineeImportTools::SetOrAddKey( BlueCurve, Point.InVal, Point.OutVal.B, Point.ArriveTangent.B, Point.LeaveTangent.B, Point.InterpMode );
-			FMatineeImportTools::SetOrAddKey( AlphaCurve, Point.InVal, Point.OutVal.A, Point.ArriveTangent.A, Point.LeaveTangent.A, Point.InterpMode );
-			SectionMin = FMath::Min( SectionMin, Point.InVal );
-			SectionMax = FMath::Max( SectionMax, Point.InVal );
+			FFrameNumber KeyTime = (Point.InVal * FrameRate).RoundToFrame();
+
+			FMatineeImportTools::SetOrAddKey( ChannelInterfaces[0], KeyTime, Point.OutVal.R, Point.ArriveTangent.R, Point.LeaveTangent.R, Point.InterpMode );
+			FMatineeImportTools::SetOrAddKey( ChannelInterfaces[1], KeyTime, Point.OutVal.G, Point.ArriveTangent.G, Point.LeaveTangent.G, Point.InterpMode );
+			FMatineeImportTools::SetOrAddKey( ChannelInterfaces[2], KeyTime, Point.OutVal.B, Point.ArriveTangent.B, Point.LeaveTangent.B, Point.InterpMode );
+			FMatineeImportTools::SetOrAddKey( ChannelInterfaces[3], KeyTime, Point.OutVal.A, Point.ArriveTangent.A, Point.LeaveTangent.A, Point.InterpMode );
+
+			KeyRange = TRange<FFrameNumber>::Hull(KeyRange, TRange<FFrameNumber>(KeyTime));
 		}
 
-		CleanupCurveKeys(RedCurve);
-		CleanupCurveKeys(GreenCurve);
-		CleanupCurveKeys(BlueCurve);
-		CleanupCurveKeys(AlphaCurve);
+		CleanupCurveKeys(Channels[0]);
+		CleanupCurveKeys(Channels[1]);
+		CleanupCurveKeys(Channels[2]);
+		CleanupCurveKeys(Channels[3]);
 
-		Section->SetStartTime( SectionMin );
-		Section->SetEndTime( SectionMax );
+		if (!KeyRange.IsEmpty())
+		{
+			Section->SetRange( KeyRange );
+		}
 	}
 
 	return bSectionCreated;
@@ -363,48 +392,54 @@ bool FMatineeImportTools::CopyInterpMoveTrack( UInterpTrackMove* MoveTrack, UMov
 
 	TransformTrack->Modify();
 
-	float KeyTime = MoveTrack->GetKeyframeTime( 0 );
-	UMovieScene3DTransformSection* Section = Cast<UMovieScene3DTransformSection>( MovieSceneHelpers::FindSectionAtTime( TransformTrack->GetAllSections(), KeyTime ) );
+	FFrameRate   FrameRate    = TransformTrack->GetTypedOuter<UMovieScene>()->GetFrameResolution();
+	FFrameNumber FirstKeyTime = (MoveTrack->GetKeyframeTime( 0 ) * FrameRate).RoundToFrame();
+
+	UMovieScene3DTransformSection* Section = Cast<UMovieScene3DTransformSection>( MovieSceneHelpers::FindSectionAtTime( TransformTrack->GetAllSections(), FirstKeyTime ) );
 	if ( Section == nullptr )
 	{
 		Section = Cast<UMovieScene3DTransformSection>( TransformTrack->CreateNewSection() );
-		Section->GetScaleCurve(EAxis::X).SetDefaultValue(DefaultScale.X);
-		Section->GetScaleCurve(EAxis::Y).SetDefaultValue(DefaultScale.Y);
-		Section->GetScaleCurve(EAxis::Z).SetDefaultValue(DefaultScale.Z);
+
 		TransformTrack->AddSection( *Section );
-		Section->SetIsInfinite(true);
+		Section->SetRange(TRange<FFrameNumber>::All());
 		bSectionCreated = true;
+
+		TArrayView<FMovieSceneFloatChannel*> FloatChannels = Section->GetChannelProxy().GetChannels<FMovieSceneFloatChannel>();
+		FloatChannels[6]->SetDefault(DefaultScale.X);
+		FloatChannels[7]->SetDefault(DefaultScale.Y);
+		FloatChannels[8]->SetDefault(DefaultScale.Z);
 	}
 
 	if (Section->TryModify())
 	{
-		float SectionMin = Section->GetStartTime();
-		float SectionMax = Section->GetEndTime();
+		TRange<FFrameNumber> KeyRange = TRange<FFrameNumber>::Empty();
 
-		FRichCurve& TranslationXCurve = Section->GetTranslationCurve( EAxis::X );
-		FRichCurve& TranslationYCurve = Section->GetTranslationCurve( EAxis::Y );
-		FRichCurve& TranslationZCurve = Section->GetTranslationCurve( EAxis::Z );
+		TArrayView<FMovieSceneFloatChannel*> Channels = Section->GetChannelProxy().GetChannels<FMovieSceneFloatChannel>();
+		TMovieSceneChannel<FMovieSceneFloatValue> ChannelInterfaces[6] = {
+			Channels[0]->GetInterface(), Channels[1]->GetInterface(), Channels[2]->GetInterface(),
+			Channels[3]->GetInterface(), Channels[4]->GetInterface(), Channels[5]->GetInterface()
+		};
 
 		for ( const auto& Point : MoveTrack->PosTrack.Points )
 		{
-			FMatineeImportTools::SetOrAddKey( TranslationXCurve, Point.InVal, Point.OutVal.X, Point.ArriveTangent.X, Point.LeaveTangent.X, Point.InterpMode );
-			FMatineeImportTools::SetOrAddKey( TranslationYCurve, Point.InVal, Point.OutVal.Y, Point.ArriveTangent.Y, Point.LeaveTangent.Y, Point.InterpMode );
-			FMatineeImportTools::SetOrAddKey( TranslationZCurve, Point.InVal, Point.OutVal.Z, Point.ArriveTangent.Z, Point.LeaveTangent.Z, Point.InterpMode );
-			SectionMin = FMath::Min( SectionMin, Point.InVal );
-			SectionMax = FMath::Max( SectionMax, Point.InVal );
-		}
+			FFrameNumber KeyTime = (Point.InVal * FrameRate).RoundToFrame();
 
-		FRichCurve& RotationXCurve = Section->GetRotationCurve( EAxis::X );
-		FRichCurve& RotationYCurve = Section->GetRotationCurve( EAxis::Y );
-		FRichCurve& RotationZCurve = Section->GetRotationCurve( EAxis::Z );
+			FMatineeImportTools::SetOrAddKey( ChannelInterfaces[0], KeyTime, Point.OutVal.X, Point.ArriveTangent.X, Point.LeaveTangent.X, Point.InterpMode );
+			FMatineeImportTools::SetOrAddKey( ChannelInterfaces[1], KeyTime, Point.OutVal.Y, Point.ArriveTangent.Y, Point.LeaveTangent.Y, Point.InterpMode );
+			FMatineeImportTools::SetOrAddKey( ChannelInterfaces[2], KeyTime, Point.OutVal.Z, Point.ArriveTangent.Z, Point.LeaveTangent.Z, Point.InterpMode );
+
+			KeyRange = TRange<FFrameNumber>::Hull(KeyRange, TRange<FFrameNumber>(KeyTime));
+		}
 
 		for ( const auto& Point : MoveTrack->EulerTrack.Points )
 		{
-			FMatineeImportTools::SetOrAddKey( RotationXCurve, Point.InVal, Point.OutVal.X, Point.ArriveTangent.X, Point.LeaveTangent.X, Point.InterpMode );
-			FMatineeImportTools::SetOrAddKey( RotationYCurve, Point.InVal, Point.OutVal.Y, Point.ArriveTangent.Y, Point.LeaveTangent.Y, Point.InterpMode );
-			FMatineeImportTools::SetOrAddKey( RotationZCurve, Point.InVal, Point.OutVal.Z, Point.ArriveTangent.Z, Point.LeaveTangent.Z, Point.InterpMode );
-			SectionMin = FMath::Min( SectionMin, Point.InVal );
-			SectionMax = FMath::Max( SectionMax, Point.InVal );
+			FFrameNumber KeyTime = (Point.InVal * FrameRate).RoundToFrame();
+
+			FMatineeImportTools::SetOrAddKey( ChannelInterfaces[3], KeyTime, Point.OutVal.X, Point.ArriveTangent.X, Point.LeaveTangent.X, Point.InterpMode );
+			FMatineeImportTools::SetOrAddKey( ChannelInterfaces[4], KeyTime, Point.OutVal.Y, Point.ArriveTangent.Y, Point.LeaveTangent.Y, Point.InterpMode );
+			FMatineeImportTools::SetOrAddKey( ChannelInterfaces[5], KeyTime, Point.OutVal.Z, Point.ArriveTangent.Z, Point.LeaveTangent.Z, Point.InterpMode );
+
+			KeyRange = TRange<FFrameNumber>::Hull(KeyRange, TRange<FFrameNumber>(KeyTime));
 		}
 
 		for (auto SubTrack : MoveTrack->SubTracks)
@@ -414,57 +449,60 @@ bool FMatineeImportTools::CopyInterpMoveTrack( UInterpTrackMove* MoveTrack, UMov
 				UInterpTrackMoveAxis* MoveSubTrack = Cast<UInterpTrackMoveAxis>(SubTrack);
 				if (MoveSubTrack)
 				{
-					FRichCurve* SubTrackCurve = nullptr;
+					int32 ChannelIndex = INDEX_NONE;
 
 					if (MoveSubTrack->MoveAxis == EInterpMoveAxis::AXIS_TranslationX)
 					{
-						SubTrackCurve = &TranslationXCurve;
+						ChannelIndex = 0;
 					}
 					else if (MoveSubTrack->MoveAxis == EInterpMoveAxis::AXIS_TranslationY)
 					{
-						SubTrackCurve = &TranslationYCurve;
+						ChannelIndex = 1;
 					}
 					else if (MoveSubTrack->MoveAxis == EInterpMoveAxis::AXIS_TranslationZ)
 					{
-						SubTrackCurve = &TranslationZCurve;
+						ChannelIndex = 2;
 					}
 					else if (MoveSubTrack->MoveAxis == EInterpMoveAxis::AXIS_RotationX)
 					{
-						SubTrackCurve = &RotationXCurve;
+						ChannelIndex = 3;
 					}
 					else if (MoveSubTrack->MoveAxis == EInterpMoveAxis::AXIS_RotationY)
 					{
-						SubTrackCurve = &RotationYCurve;
+						ChannelIndex = 4;
 					}
 					else if (MoveSubTrack->MoveAxis == EInterpMoveAxis::AXIS_RotationZ)
 					{
-						SubTrackCurve = &RotationZCurve;
+						ChannelIndex = 5;
 					}
-							
-					if (SubTrackCurve != nullptr)
+
+					if (ChannelIndex != INDEX_NONE)
 					{
 						for (const auto& Point : MoveSubTrack->FloatTrack.Points)
 						{
-							FMatineeImportTools::SetOrAddKey( *SubTrackCurve, Point.InVal, Point.OutVal, Point.ArriveTangent, Point.LeaveTangent, Point.InterpMode );
-							SectionMax = FMath::Max( SectionMax, Point.InVal );
-						}
+							FFrameNumber KeyTime = (Point.InVal * FrameRate).RoundToFrame();
 
-						CleanupCurveKeys(*SubTrackCurve);
+							FMatineeImportTools::SetOrAddKey( ChannelInterfaces[ChannelIndex], KeyTime, Point.OutVal, Point.ArriveTangent, Point.LeaveTangent, Point.InterpMode );
+
+							KeyRange = TRange<FFrameNumber>::Hull(KeyRange, TRange<FFrameNumber>(KeyTime));
+						}
 					}
 				}
 			}
 		}
 
-		CleanupCurveKeys(TranslationXCurve);
-		CleanupCurveKeys(TranslationYCurve);
-		CleanupCurveKeys(TranslationZCurve);
+		CleanupCurveKeys(Channels[0]);
+		CleanupCurveKeys(Channels[1]);
+		CleanupCurveKeys(Channels[2]);
 
-		CleanupCurveKeys(RotationXCurve);
-		CleanupCurveKeys(RotationYCurve);
-		CleanupCurveKeys(RotationZCurve);
+		CleanupCurveKeys(Channels[3]);
+		CleanupCurveKeys(Channels[4]);
+		CleanupCurveKeys(Channels[5]);
 
-		Section->SetStartTime( SectionMin );
-		Section->SetEndTime( SectionMax );
+		if (!KeyRange.IsEmpty())
+		{
+			Section->SetRange( KeyRange );
+		}
 	}
 
 	return bSectionCreated;
@@ -478,8 +516,10 @@ bool FMatineeImportTools::CopyInterpParticleTrack( UInterpTrackToggle* MatineeTo
 
 	ParticleTrack->Modify();
 
-	float KeyTime = MatineeToggleTrack->GetKeyframeTime( 0 );
-	UMovieSceneParticleSection* Section = Cast<UMovieSceneParticleSection>( MovieSceneHelpers::FindSectionAtTime( ParticleTrack->GetAllSections(), KeyTime ) );
+	FFrameRate   FrameRate    = ParticleTrack->GetTypedOuter<UMovieScene>()->GetFrameResolution();
+	FFrameNumber FirstKeyTime = (MatineeToggleTrack->GetKeyframeTime( 0 ) * FrameRate).RoundToFrame();
+
+	UMovieSceneParticleSection* Section = Cast<UMovieSceneParticleSection>( MovieSceneHelpers::FindSectionAtTime( ParticleTrack->GetAllSections(), FirstKeyTime ) );
 	if ( Section == nullptr )
 	{
 		Section = Cast<UMovieSceneParticleSection>( ParticleTrack->CreateNewSection() );
@@ -489,35 +529,43 @@ bool FMatineeImportTools::CopyInterpParticleTrack( UInterpTrackToggle* MatineeTo
 
 	if (Section->TryModify())
 	{
-		float SectionMin = Section->GetStartTime();
-		float SectionMax = Section->GetEndTime();
+		TRange<FFrameNumber> KeyRange = TRange<FFrameNumber>::Empty();
+
+		FMovieSceneParticleChannel* ParticleChannel = Section->GetChannelProxy().GetChannel<FMovieSceneParticleChannel>(0);
+		check(ParticleChannel);
+		TMovieSceneChannel<uint8> ChannelInterface = ParticleChannel->GetInterface();
 
-		FIntegralCurve& ParticleCurve = Section->GetParticleCurve();
 		for ( const auto& Key : MatineeToggleTrack->ToggleTrack )
 		{
-			EParticleKey::Type ParticleKey;
+			FFrameNumber KeyTime = (Key.Time * FrameRate).RoundToFrame();
+
+			EParticleKey ParticleKey;
 			if ( TryConvertMatineeToggleToOutParticleKey( Key.ToggleAction, ParticleKey ) )
 			{
-				ParticleCurve.AddKey( Key.Time, (int32)ParticleKey, ParticleCurve.FindKey( Key.Time ) );
+				ChannelInterface.AddKey( KeyTime, (uint8)ParticleKey );
 			}
-			SectionMin = FMath::Min( SectionMin, Key.Time );
-			SectionMax = FMath::Max( SectionMax, Key.Time );
+
+			KeyRange = TRange<FFrameNumber>::Hull(KeyRange, TRange<FFrameNumber>(KeyTime));
 		}
 
-		Section->SetStartTime( SectionMin );
-		Section->SetEndTime( SectionMax );
+		if (!KeyRange.IsEmpty())
+		{
+			Section->SetRange( KeyRange );
+		}
 	}
 
 	return bSectionCreated;
 }
 
 
-bool FMatineeImportTools::CopyInterpAnimControlTrack( UInterpTrackAnimControl* MatineeAnimControlTrack, UMovieSceneSkeletalAnimationTrack* SkeletalAnimationTrack, float EndPlaybackRange )
+bool FMatineeImportTools::CopyInterpAnimControlTrack( UInterpTrackAnimControl* MatineeAnimControlTrack, UMovieSceneSkeletalAnimationTrack* SkeletalAnimationTrack, FFrameNumber EndPlaybackRange )
 {
 	// @todo - Sequencer - Add support for slot names once they are implemented.
 	const FScopedTransaction Transaction( NSLOCTEXT( "Sequencer", "PasteMatineeAnimTrack", "Paste Matinee Anim Track" ) );
 	bool bSectionCreated = false;
-	
+
+	FFrameRate FrameRate = SkeletalAnimationTrack->GetTypedOuter<UMovieScene>()->GetFrameResolution();
+
 	SkeletalAnimationTrack->Modify();
 	SkeletalAnimationTrack->RemoveAllAnimationData();
 
@@ -534,7 +582,7 @@ bool FMatineeImportTools::CopyInterpAnimControlTrack( UInterpTrackAnimControl* M
 			}
 			else
 			{
-				EndTime = EndPlaybackRange;
+				EndTime = EndPlaybackRange / FrameRate;
 			}
 		}
 		else
@@ -550,8 +598,7 @@ bool FMatineeImportTools::CopyInterpAnimControlTrack( UInterpTrackAnimControl* M
 		}
 
 		UMovieSceneSkeletalAnimationSection* NewSection = Cast<UMovieSceneSkeletalAnimationSection>( SkeletalAnimationTrack->CreateNewSection() );
-		NewSection->SetStartTime( AnimSeq.StartTime );
-		NewSection->SetEndTime( EndTime );
+		NewSection->SetRange( TRange<FFrameNumber>((AnimSeq.StartTime * FrameRate).RoundToFrame(), (EndTime * FrameRate).RoundToFrame() + 1) );
 		NewSection->Params.StartOffset = AnimSeq.AnimStartOffset;
 		NewSection->Params.EndOffset = AnimSeq.AnimEndOffset;
 		NewSection->Params.PlayRate = AnimSeq.AnimPlayRate;
@@ -570,8 +617,10 @@ bool FMatineeImportTools::CopyInterpSoundTrack( UInterpTrackSound* MatineeSoundT
 	const FScopedTransaction Transaction( NSLOCTEXT( "Sequencer", "PasteMatineeSoundTrack", "Paste Matinee Sound Track" ) );
 	bool bSectionCreated = false;
 
+	FFrameRate FrameRate = AudioTrack->GetTypedOuter<UMovieScene>()->GetFrameResolution();
+
 	AudioTrack->Modify();
-	
+
 	int MaxSectionRowIndex = -1;
 	for ( UMovieSceneSection* Section : AudioTrack->GetAllSections() )
 	{
@@ -580,12 +629,14 @@ bool FMatineeImportTools::CopyInterpSoundTrack( UInterpTrackSound* MatineeSoundT
 
 	for ( const FSoundTrackKey& SoundTrackKey : MatineeSoundTrack->Sounds )
 	{
-		AudioTrack->AddNewSound( SoundTrackKey.Sound, SoundTrackKey.Time );
+		AudioTrack->AddNewSound( SoundTrackKey.Sound, (SoundTrackKey.Time * FrameRate).RoundToFrame() );
 
 		UMovieSceneAudioSection* NewAudioSection = Cast<UMovieSceneAudioSection>(AudioTrack->GetAllSections().Last());
 		NewAudioSection->SetRowIndex( MaxSectionRowIndex + 1 );
-		NewAudioSection->GetPitchMultiplierCurve().SetDefaultValue( SoundTrackKey.Pitch );
-		NewAudioSection->GetSoundVolumeCurve().SetDefaultValue( SoundTrackKey.Volume );
+		
+		TArrayView<FMovieSceneFloatChannel*> FloatChannels = NewAudioSection->GetChannelProxy().GetChannels<FMovieSceneFloatChannel>();
+		FloatChannels[0]->SetDefault(SoundTrackKey.Volume);
+		FloatChannels[1]->SetDefault(SoundTrackKey.Pitch);
 
 		AudioTrack->AddSection( *NewAudioSection );
 		bSectionCreated = true;
@@ -600,9 +651,11 @@ bool FMatineeImportTools::CopyInterpFadeTrack( UInterpTrackFade* MatineeFadeTrac
 	bool bSectionCreated = false;
 
 	FadeTrack->Modify();
-	
-	float KeyTime = MatineeFadeTrack->GetKeyframeTime( 0 );
-	UMovieSceneFadeSection* Section = Cast<UMovieSceneFadeSection>( MovieSceneHelpers::FindSectionAtTime( FadeTrack->GetAllSections(), KeyTime ) );
+
+	FFrameRate   FrameRate    = FadeTrack->GetTypedOuter<UMovieScene>()->GetFrameResolution();
+	FFrameNumber FirstKeyTime = (MatineeFadeTrack->GetKeyframeTime( 0 ) * FrameRate).RoundToFrame();
+
+	UMovieSceneFadeSection* Section = Cast<UMovieSceneFadeSection>( MovieSceneHelpers::FindSectionAtTime( FadeTrack->GetAllSections(), FirstKeyTime ) );
 	if ( Section == nullptr )
 	{
 		Section = Cast<UMovieSceneFadeSection>( FadeTrack->CreateNewSection() );
@@ -611,20 +664,25 @@ bool FMatineeImportTools::CopyInterpFadeTrack( UInterpTrackFade* MatineeFadeTrac
 	}
 	if (Section->TryModify())
 	{
-		float SectionMin = Section->GetStartTime();		
-		float SectionMax = Section->GetEndTime();
-		
-		FRichCurve& FloatCurve = Section->GetFloatCurve();
+		TRange<FFrameNumber> KeyRange = TRange<FFrameNumber>::Empty();
+
+		FMovieSceneFloatChannel* FadeChannel = Section->GetChannelProxy().GetChannel<FMovieSceneFloatChannel>(0);
+		check(FadeChannel);
+		TMovieSceneChannel<FMovieSceneFloatValue> FadeInterface = FadeChannel->GetInterface();
 		for ( const auto& Point : MatineeFadeTrack->FloatTrack.Points )
 		{
-			FMatineeImportTools::SetOrAddKey( FloatCurve, Point.InVal, Point.OutVal, Point.ArriveTangent, Point.LeaveTangent, Point.InterpMode );
-			SectionMin = FMath::Min( SectionMin, Point.InVal );
-			SectionMax = FMath::Max( SectionMax, Point.InVal );
+			FFrameNumber KeyTime = (Point.InVal * FrameRate).RoundToFrame();
+
+			FMatineeImportTools::SetOrAddKey( FadeInterface, KeyTime, Point.OutVal, Point.ArriveTangent, Point.LeaveTangent, Point.InterpMode );
+
+			KeyRange = TRange<FFrameNumber>::Hull(KeyRange, TRange<FFrameNumber>(KeyTime));
+		}
+
+		if (!KeyRange.IsEmpty())
+		{
+			Section->SetRange( KeyRange );
 		}
 
-		Section->SetStartTime( SectionMin );
-		Section->SetEndTime( SectionMax );
-	
 		Section->FadeColor = MatineeFadeTrack->FadeColor;
 		Section->bFadeAudio = MatineeFadeTrack->bFadeAudio;
 	}
@@ -637,6 +695,8 @@ bool FMatineeImportTools::CopyInterpDirectorTrack( UInterpTrackDirector* Directo
 	const FScopedTransaction Transaction( NSLOCTEXT( "Sequencer", "PasteMatineeDirectorTrack", "Paste Matinee Director Track" ) );
 	bool bCutsAdded = false;
 
+	FFrameRate FrameRate = CameraCutTrack->GetTypedOuter<UMovieScene>()->GetFrameResolution();
+
 	CameraCutTrack->Modify();
 	
 	for (FDirectorTrackCut TrackCut : DirectorTrack->CutTrack)
@@ -656,7 +716,7 @@ bool FMatineeImportTools::CopyInterpDirectorTrack( UInterpTrackDirector* Directo
 				if (CameraHandle.IsValid())
 				{
 					FMovieSceneObjectBindingID CameraBindingID(CameraHandle, MovieSceneSequenceID::Root);
-					CameraCutTrack->AddNewCameraCut(CameraBindingID, TrackCut.Time);
+					CameraCutTrack->AddNewCameraCut(CameraBindingID, (TrackCut.Time * FrameRate).RoundToFrame());
 					bCutsAdded = true;
 				}
 			}
@@ -675,8 +735,10 @@ bool FMatineeImportTools::CopyInterpEventTrack( UInterpTrackEvent* MatineeEventT
 
 	if (MatineeEventTrack->EventTrack.Num())
 	{
-		float KeyTime = MatineeEventTrack->EventTrack[0].Time;
-		UMovieSceneEventSection* Section = Cast<UMovieSceneEventSection>( MovieSceneHelpers::FindSectionAtTime( EventTrack->GetAllSections(), KeyTime ) );
+		FFrameRate   FrameRate    = EventTrack->GetTypedOuter<UMovieScene>()->GetFrameResolution();
+		FFrameNumber FirstKeyTime = (MatineeEventTrack->EventTrack[0].Time * FrameRate).RoundToFrame();
+
+		UMovieSceneEventSection* Section = Cast<UMovieSceneEventSection>( MovieSceneHelpers::FindSectionAtTime( EventTrack->GetAllSections(), FirstKeyTime ) );
 		if ( Section == nullptr )
 		{
 			Section = Cast<UMovieSceneEventSection>( EventTrack->CreateNewSection() );
@@ -685,19 +747,25 @@ bool FMatineeImportTools::CopyInterpEventTrack( UInterpTrackEvent* MatineeEventT
 		}
 		if (Section->TryModify())
 		{
-			float SectionMin = Section->GetStartTime();
-			float SectionMax = Section->GetEndTime();
+			TRange<FFrameNumber> KeyRange = TRange<FFrameNumber>::Empty();
+
+			FMovieSceneEventSectionData* EventChannel = Section->GetChannelProxy().GetChannel<FMovieSceneEventSectionData>(0);
+			check(EventChannel);
+			TMovieSceneChannel<FEventPayload> ChannelInterface = EventChannel->GetInterface();
 
-			TCurveInterface<FEventPayload, float> CurveInterface = Section->GetCurveInterface();
 			for (FEventTrackKey EventTrackKey : MatineeEventTrack->EventTrack)
 			{
-				CurveInterface.UpdateOrAddKey(EventTrackKey.Time, FEventPayload(EventTrackKey.EventName), KINDA_SMALL_NUMBER);
-				SectionMin = FMath::Min( SectionMin, EventTrackKey.Time );
-				SectionMax = FMath::Max( SectionMax, EventTrackKey.Time );
+				FFrameNumber KeyTime = (EventTrackKey.Time * FrameRate).RoundToFrame();
+
+				ChannelInterface.UpdateOrAddKey(KeyTime, FEventPayload(EventTrackKey.EventName));
+
+				KeyRange = TRange<FFrameNumber>::Hull(KeyRange, TRange<FFrameNumber>(KeyTime));
 			}
 
-			Section->SetStartTime( SectionMin );
-			Section->SetEndTime( SectionMax );
+			if (!KeyRange.IsEmpty())
+			{
+				Section->SetRange( KeyRange );
+			}
 		}
 	}	
 	
@@ -713,8 +781,10 @@ bool FMatineeImportTools::CopyInterpVisibilityTrack( UInterpTrackVisibility* Mat
 
 	if (MatineeVisibilityTrack->VisibilityTrack.Num())
 	{
-		float KeyTime = MatineeVisibilityTrack->VisibilityTrack[0].Time;
-		UMovieSceneBoolSection* Section = Cast<UMovieSceneBoolSection>( MovieSceneHelpers::FindSectionAtTime( VisibilityTrack->GetAllSections(), KeyTime ) );
+		FFrameRate   FrameRate    = VisibilityTrack->GetTypedOuter<UMovieScene>()->GetFrameResolution();
+		FFrameNumber FirstKeyTime = (MatineeVisibilityTrack->GetKeyframeTime( 0 ) * FrameRate).RoundToFrame();
+
+		UMovieSceneBoolSection* Section = Cast<UMovieSceneBoolSection>( MovieSceneHelpers::FindSectionAtTime( VisibilityTrack->GetAllSections(), FirstKeyTime ) );
 		if ( Section == nullptr )
 		{
 			Section = Cast<UMovieSceneBoolSection>( VisibilityTrack->CreateNewSection() );
@@ -723,11 +793,14 @@ bool FMatineeImportTools::CopyInterpVisibilityTrack( UInterpTrackVisibility* Mat
 		}
 		if (Section->TryModify())
 		{
-			float SectionMin = Section->GetStartTime();
-			float SectionMax = Section->GetEndTime();
-			
+			TRange<FFrameNumber> KeyRange = TRange<FFrameNumber>::Empty();
+
 			bool bVisible = true;
-			FIntegralCurve& VisibilityCurve = Section->GetCurve();
+
+			FMovieSceneBoolChannel* VisibilityChannel = Section->GetChannelProxy().GetChannel<FMovieSceneBoolChannel>(0);
+			check(VisibilityChannel);
+			TMovieSceneChannel<bool> ChannelInterface = VisibilityChannel->GetInterface();
+
 			for (FVisibilityTrackKey VisibilityTrackKey : MatineeVisibilityTrack->VisibilityTrack)
 			{
 				if (VisibilityTrackKey.Action == EVisibilityTrackAction::EVTA_Hide)
@@ -743,13 +816,17 @@ bool FMatineeImportTools::CopyInterpVisibilityTrack( UInterpTrackVisibility* Mat
 					bVisible = !bVisible;
 				}
 
-				VisibilityCurve.UpdateOrAddKey(VisibilityTrackKey.Time, bVisible);
-				SectionMin = FMath::Min( SectionMin, VisibilityTrackKey.Time );
-				SectionMax = FMath::Max( SectionMax, VisibilityTrackKey.Time );
+				FFrameNumber KeyTime = (VisibilityTrackKey.Time * FrameRate).RoundToFrame();
+
+				ChannelInterface.UpdateOrAddKey(KeyTime, bVisible);
+
+				KeyRange = TRange<FFrameNumber>::Hull(KeyRange, TRange<FFrameNumber>(KeyTime));
 			}
 
-			Section->SetStartTime( SectionMin );
-			Section->SetEndTime( SectionMax );
+			if (!KeyRange.IsEmpty())
+			{
+				Section->SetRange( KeyRange );
+			}
 		}
 	}	
 	
diff --git a/Engine/Source/Editor/MovieSceneTools/Private/MovieSceneObjectBindingIDCustomization.cpp b/Engine/Source/Editor/MovieSceneTools/Private/MovieSceneObjectBindingIDCustomization.cpp
--- a/Engine/Source/Editor/MovieSceneTools/Private/MovieSceneObjectBindingIDCustomization.cpp
+++ b/Engine/Source/Editor/MovieSceneTools/Private/MovieSceneObjectBindingIDCustomization.cpp
@@ -137,7 +137,7 @@ void FMovieSceneObjectBindingIDCustomization::SetCurrentValue(const FMovieSceneO
 
 	for (void* Ptr : Ptrs)
 	{
-		*static_cast<FMovieSceneObjectBindingID*>(Ptr) = InObjectBinding;
+		*static_cast<FMovieSceneObjectBindingID*>(Ptrs[0]) = InObjectBinding;
 	}
 	
 	StructProperty->NotifyPostChange();
diff --git a/Engine/Source/Editor/MovieSceneTools/Private/MovieSceneToolHelpers.cpp b/Engine/Source/Editor/MovieSceneTools/Private/MovieSceneToolHelpers.cpp
--- a/Engine/Source/Editor/MovieSceneTools/Private/MovieSceneToolHelpers.cpp
+++ b/Engine/Source/Editor/MovieSceneTools/Private/MovieSceneToolHelpers.cpp
@@ -45,11 +45,12 @@
 #include "Matinee/InterpTrackMove.h"
 #include "Matinee/InterpTrackMoveAxis.h"
 #include "Matinee/InterpTrackInstMove.h"
+#include "Channels/MovieSceneChannelProxy.h"
 
 /* MovieSceneToolHelpers
  *****************************************************************************/
 
-void MovieSceneToolHelpers::TrimSection(const TSet<TWeakObjectPtr<UMovieSceneSection>>& Sections, float Time, bool bTrimLeft)
+void MovieSceneToolHelpers::TrimSection(const TSet<TWeakObjectPtr<UMovieSceneSection>>& Sections, FFrameNumber Time, bool bTrimLeft)
 {
 	for (auto Section : Sections)
 	{
@@ -61,7 +62,7 @@ void MovieSceneToolHelpers::TrimSection(const TSet<TWeakObjectPtr<UMovieSceneSec
 }
 
 
-void MovieSceneToolHelpers::SplitSection(const TSet<TWeakObjectPtr<UMovieSceneSection>>& Sections, float Time)
+void MovieSceneToolHelpers::SplitSection(const TSet<TWeakObjectPtr<UMovieSceneSection>>& Sections, FFrameNumber Time)
 {
 	for (auto Section : Sections)
 	{
@@ -97,6 +98,7 @@ bool MovieSceneToolHelpers::ParseShotName(const FString& ShotName, FString& Shot
 	bool bInTakeNumber = false;
 
 	bool bFoundTakeSeparator = false;
+	TOptional<uint32> ParsedTakeNumber;
 	TakeNumber = ProjectSettings->FirstTakeNumber;
 
 	for (int32 CharIndex = 0; CharIndex < ShotName.Len(); ++CharIndex)
@@ -149,7 +151,30 @@ bool MovieSceneToolHelpers::ParseShotName(const FString& ShotName, FString& Shot
 
 	if (FirstTakeNumberIndex != INDEX_NONE)
 	{
-		TakeNumber = FCString::Atoi(*ShotName.Mid(FirstTakeNumberIndex, LastTakeNumberIndex-FirstTakeNumberIndex+1));
+		FString TakeStr = ShotName.Mid(FirstTakeNumberIndex, LastTakeNumberIndex-FirstTakeNumberIndex+1);
+		if (TakeStr.IsNumeric())
+		{
+			ParsedTakeNumber = FCString::Atoi(*TakeStr);
+		}
+	}
+
+	// If take number wasn't found, search backwards to find the first take separator and assume [shot prefix]_[take number]
+	//
+	if (!ParsedTakeNumber.IsSet())
+	{
+		int32 LastSlashPos = ShotName.Find(ProjectSettings->TakeSeparator, ESearchCase::IgnoreCase, ESearchDir::FromEnd);
+		if (LastSlashPos != INDEX_NONE)
+		{
+			ShotPrefix = ShotName.Left(LastSlashPos);
+			ShotNumber = INDEX_NONE; // Nullify the shot number since we only have a shot prefix
+			TakeNumber = FCString::Atoi(*ShotName.RightChop(LastSlashPos+1));
+			return true;
+		}
+	}
+
+	if (ParsedTakeNumber.IsSet())
+	{
+		TakeNumber = ParsedTakeNumber.GetValue();
 	}
 
 	return FirstShotNumberIndex != INDEX_NONE;
@@ -162,12 +187,16 @@ FString MovieSceneToolHelpers::ComposeShotName(const FString& ShotPrefix, uint32
 
 	FString ShotName = ShotPrefix;
 
-	FString ShotFormat = TEXT("%0") + FString::Printf(TEXT("%d"), ProjectSettings->ShotNumDigits) + TEXT("d");
-	FString TakeFormat = TEXT("%0") + FString::Printf(TEXT("%d"), ProjectSettings->TakeNumDigits) + TEXT("d");
-
-	ShotName += FString::Printf(*ShotFormat, ShotNumber);
+	if (ShotNumber != INDEX_NONE)
+	{
+		FString ShotFormat = TEXT("%0") + FString::Printf(TEXT("%d"), ProjectSettings->ShotNumDigits) + TEXT("d");
+		ShotName += FString::Printf(*ShotFormat, ShotNumber);	
+	}
+	
 	if (TakeNumber != INDEX_NONE)
 	{
+		FString TakeFormat = TEXT("%0") + FString::Printf(TEXT("%d"), ProjectSettings->TakeNumDigits) + TEXT("d");
+		
 		ShotName += ProjectSettings->TakeSeparator;
 		ShotName += FString::Printf(*TakeFormat, TakeNumber);
 	}
@@ -198,7 +227,7 @@ FString MovieSceneToolHelpers::GenerateNewShotPath(UMovieScene* SequenceMovieSce
 	UObject* SequenceAsset = SequenceMovieScene->GetOuter();
 	UPackage* SequencePackage = SequenceAsset->GetOutermost();
 	FString SequencePackageName = SequencePackage->GetName(); // ie. /Game/cine/max/master
-	int LastSlashPos = SequencePackageName.Find(TEXT("/"), ESearchCase::IgnoreCase, ESearchDir::FromEnd);
+	int32 LastSlashPos = SequencePackageName.Find(TEXT("/"), ESearchCase::IgnoreCase, ESearchDir::FromEnd);
 	FString SequencePath = SequencePackageName.Left(LastSlashPos);
 
 	FString NewShotPrefix;
@@ -246,29 +275,30 @@ FString MovieSceneToolHelpers::GenerateNewShotPath(UMovieScene* SequenceMovieSce
 }
 
 
-FString MovieSceneToolHelpers::GenerateNewShotName(const TArray<UMovieSceneSection*>& AllSections, float Time)
+FString MovieSceneToolHelpers::GenerateNewShotName(const TArray<UMovieSceneSection*>& AllSections, FFrameNumber Time)
 {
 	const UMovieSceneToolsProjectSettings* ProjectSettings = GetDefault<UMovieSceneToolsProjectSettings>();
 
 	UMovieSceneCinematicShotSection* BeforeShot = nullptr;
 	UMovieSceneCinematicShotSection* NextShot = nullptr;
 
-	float MinEndDiff = FLT_MAX;
-	float MinStartDiff = FLT_MAX;
+	FFrameNumber MinEndDiff = TNumericLimits<int32>::Max();
+	FFrameNumber MinStartDiff = TNumericLimits<int32>::Max();
+
 	for (auto Section : AllSections)
 	{
-		if (Section->GetEndTime() >= Time)
+		if (Section->HasEndFrame() && Section->GetExclusiveEndFrame() >= Time)
 		{
-			float EndDiff = Section->GetEndTime() - Time;
+			FFrameNumber EndDiff = Section->GetExclusiveEndFrame() - Time;
 			if (MinEndDiff > EndDiff)
 			{
 				MinEndDiff = EndDiff;
 				BeforeShot = Cast<UMovieSceneCinematicShotSection>(Section);
 			}
 		}
-		if (Section->GetStartTime() <= Time)
+		if (Section->HasStartFrame() && Section->GetInclusiveStartFrame() <= Time)
 		{
-			float StartDiff = Time - Section->GetStartTime();
+			FFrameNumber StartDiff = Time - Section->GetInclusiveStartFrame();
 			if (MinStartDiff > StartDiff)
 			{
 				MinStartDiff = StartDiff;
@@ -323,14 +353,14 @@ FString MovieSceneToolHelpers::GenerateNewShotName(const TArray<UMovieSceneSecti
 
 void MovieSceneToolHelpers::GatherTakes(const UMovieSceneSection* Section, TArray<uint32>& TakeNumbers, uint32& CurrentTakeNumber)
 {
-	const UMovieSceneCinematicShotSection* Shot = Cast<const UMovieSceneCinematicShotSection>(Section);
+	const UMovieSceneSubSection* SubSection = Cast<const UMovieSceneSubSection>(Section);
 	
-	if (Shot->GetSequence() == nullptr)
+	if (SubSection->GetSequence() == nullptr)
 	{
 		return;
 	}
 
-	FAssetData ShotData(Shot->GetSequence()->GetOuter());
+	FAssetData ShotData(SubSection->GetSequence()->GetOuter());
 
 	FString ShotPackagePath = ShotData.PackagePath.ToString();
 
@@ -338,7 +368,14 @@ void MovieSceneToolHelpers::GatherTakes(const UMovieSceneSection* Section, TArra
 	uint32 ShotNumber = INDEX_NONE;
 	CurrentTakeNumber = INDEX_NONE;
 
-	if (ParseShotName(Shot->GetShotDisplayName(), ShotPrefix, ShotNumber, CurrentTakeNumber))
+	FString SubSectionName = SubSection->GetSequence()->GetName();
+	if (SubSection->IsA<UMovieSceneCinematicShotSection>())
+	{
+		const UMovieSceneCinematicShotSection* ShotSection = Cast<UMovieSceneCinematicShotSection>(SubSection);
+		SubSectionName = ShotSection->GetShotDisplayName();
+	}
+
+	if (ParseShotName(SubSectionName, ShotPrefix, ShotNumber, CurrentTakeNumber))
 	{
 		// Gather up all level sequence assets
 		FAssetRegistryModule& AssetRegistryModule = FModuleManager::LoadModuleChecked<FAssetRegistryModule>(TEXT("AssetRegistry"));
@@ -371,9 +408,9 @@ void MovieSceneToolHelpers::GatherTakes(const UMovieSceneSection* Section, TArra
 
 UObject* MovieSceneToolHelpers::GetTake(const UMovieSceneSection* Section, uint32 TakeNumber)
 {
-	const UMovieSceneCinematicShotSection* Shot = Cast<const UMovieSceneCinematicShotSection>(Section);
+	const UMovieSceneSubSection* SubSection = Cast<const UMovieSceneSubSection>(Section);
 
-	FAssetData ShotData(Shot->GetSequence()->GetOuter());
+	FAssetData ShotData(SubSection->GetSequence()->GetOuter());
 
 	FString ShotPackagePath = ShotData.PackagePath.ToString();
 	int32 ShotLastSlashPos = INDEX_NONE;
@@ -384,7 +421,14 @@ UObject* MovieSceneToolHelpers::GetTake(const UMovieSceneSection* Section, uint3
 	uint32 ShotNumber = INDEX_NONE;
 	uint32 TakeNumberDummy = INDEX_NONE;
 
-	if (ParseShotName(Shot->GetShotDisplayName(), ShotPrefix, ShotNumber, TakeNumberDummy))
+	FString SubSectionName = SubSection->GetSequence()->GetName();
+	if (SubSection->IsA<UMovieSceneCinematicShotSection>())
+	{
+		const UMovieSceneCinematicShotSection* ShotSection = Cast<UMovieSceneCinematicShotSection>(SubSection);
+		SubSectionName = ShotSection->GetShotDisplayName();
+	}
+
+	if (ParseShotName(SubSectionName, ShotPrefix, ShotNumber, TakeNumberDummy))
 	{
 		// Gather up all level sequence assets
 		FAssetRegistryModule& AssetRegistryModule = FModuleManager::LoadModuleChecked<FAssetRegistryModule>(TEXT("AssetRegistry"));
@@ -423,12 +467,15 @@ int32 MovieSceneToolHelpers::FindAvailableRowIndex(UMovieSceneTrack* InTrack, UM
 	for (int32 RowIndex = 0; RowIndex <= InTrack->GetMaxRowIndex(); ++RowIndex)
 	{
 		bool bFoundIntersect = false;
-		for (auto Section : InTrack->GetAllSections())
+		for (UMovieSceneSection* Section : InTrack->GetAllSections())
 		{
-			TRange<float> InRange(InSection->GetStartTime(), InSection->GetEndTime());
-			TRange<float> Range(Section->GetStartTime(), Section->GetEndTime());
+			if (!Section->HasStartFrame() || !Section->HasEndFrame() || InSection->HasStartFrame() || !InSection->HasEndFrame())
+			{
+				bFoundIntersect = true;
+				break;
+			}
 
-			if (Section != InSection && Section->GetRowIndex() == RowIndex && Range.Overlaps(InRange))
+			if (Section != InSection && Section->GetRowIndex() == RowIndex && Section->GetRange().Overlaps(InSection->GetRange()))
 			{
 				bFoundIntersect = true;
 				break;
@@ -548,7 +595,7 @@ TSharedRef<SWidget> MovieSceneToolHelpers::MakeEnumComboBox(const UEnum* InEnum,
 		.OnEnumSelectionChanged(InOnSelectionChanged);
 }
 
-bool MovieSceneToolHelpers::ShowImportEDLDialog(UMovieScene* InMovieScene, float InFrameRate, FString InOpenDirectory)
+bool MovieSceneToolHelpers::ShowImportEDLDialog(UMovieScene* InMovieScene, FFrameRate InFrameRate, FString InOpenDirectory)
 {
 	TArray<FString> OpenFilenames;
 	IDesktopPlatform* DesktopPlatform = FDesktopPlatformModule::Get();
@@ -583,7 +630,7 @@ bool MovieSceneToolHelpers::ShowImportEDLDialog(UMovieScene* InMovieScene, float
 	return MovieSceneCaptureHelpers::ImportEDL(InMovieScene, InFrameRate, OpenFilenames[0]);
 }
 
-bool MovieSceneToolHelpers::ShowExportEDLDialog(const UMovieScene* InMovieScene, float InFrameRate, FString InSaveDirectory, int32 InHandleFrames)
+bool MovieSceneToolHelpers::ShowExportEDLDialog(const UMovieScene* InMovieScene, FFrameRate InFrameRate, FString InSaveDirectory, int32 InHandleFrames)
 {
 	TArray<FString> SaveFilenames;
 	FString SequenceName = InMovieScene->GetOuter()->GetName();
@@ -685,8 +732,10 @@ bool ImportFBXProperty(FString NodeName, FString AnimatedPropertyName, FGuid Obj
 
 			if (FloatTrack)
 			{
+				FFrameRate FrameRate = FloatTrack->GetTypedOuter<UMovieScene>()->GetFrameResolution();
+
 				bool bSectionAdded = false;
-				UMovieSceneFloatSection* FloatSection = Cast<UMovieSceneFloatSection>(FloatTrack->FindOrAddSection(0.f, bSectionAdded));
+				UMovieSceneFloatSection* FloatSection = Cast<UMovieSceneFloatSection>(FloatTrack->FindOrAddSection(0, bSectionAdded));
 				if (!FloatSection)
 				{
 					continue;
@@ -694,25 +743,21 @@ bool ImportFBXProperty(FString NodeName, FString AnimatedPropertyName, FGuid Obj
 
 				if (bSectionAdded)
 				{
-					FloatSection->SetIsInfinite(true);
+					FloatSection->SetRange(TRange<FFrameNumber>::All());
 				}
 
-				float MinTime = FLT_MAX;
-				float MaxTime = -FLT_MAX;
-
 				const int32 ChannelIndex = 0;
 				const int32 CompositeIndex = 0;
 				FInterpCurveFloat CurveHandle;
 				const bool bNegative = false;
 				CurveAPI.GetCurveData(NodeName, AnimatedPropertyName, ChannelIndex, CompositeIndex, CurveHandle, bNegative);
 
-				FRichCurve& FloatCurve = FloatSection->GetFloatCurve();
-				FloatCurve.Reset();
+				FMovieSceneFloatChannel* Channel = FloatSection->GetChannelProxy().GetChannel<FMovieSceneFloatChannel>(0);
+				TMovieSceneChannel<FMovieSceneFloatValue> ChannelInterface = Channel->GetInterface();
+
+				ChannelInterface.Reset();
 				for (int32 KeyIndex = 0; KeyIndex < CurveHandle.Points.Num(); ++KeyIndex)
 				{
-					MinTime = FMath::Min(MinTime, CurveHandle.Points[KeyIndex].InVal);
-					MaxTime = FMath::Max(MaxTime, CurveHandle.Points[KeyIndex].InVal);
-
 					float ArriveTangent = CurveHandle.Points[KeyIndex].ArriveTangent;
 					if (KeyIndex > 0)
 					{
@@ -725,18 +770,22 @@ bool ImportFBXProperty(FString NodeName, FString AnimatedPropertyName, FGuid Obj
 						LeaveTangent = LeaveTangent / (CurveHandle.Points[KeyIndex+1].InVal - CurveHandle.Points[KeyIndex].InVal);
 					}
 
-					FMatineeImportTools::SetOrAddKey(FloatCurve, CurveHandle.Points[KeyIndex].InVal, CurveHandle.Points[KeyIndex].OutVal, ArriveTangent, LeaveTangent, CurveHandle.Points[KeyIndex].InterpMode);
+					FFrameNumber KeyTime = (CurveHandle.Points[KeyIndex].InVal * FrameRate).RoundToFrame();
+					FMatineeImportTools::SetOrAddKey(ChannelInterface, KeyTime, CurveHandle.Points[KeyIndex].OutVal, ArriveTangent, LeaveTangent, CurveHandle.Points[KeyIndex].InterpMode);
 				}
 
 				if (ImportFBXSettings->bReduceKeys)
 				{
-					FloatCurve.RemoveRedundantKeys(ImportFBXSettings->ReduceKeysTolerance);
+					FKeyDataOptimizationParams Params;
+					Params.bAutoSetInterpolation = true;
+					MovieScene::Optimize(Channel, Params);
 				}
-				
-				FloatCurve.AutoSetTangents();
 
-				FloatSection->SetStartTime(MinTime);
-				FloatSection->SetEndTime(MaxTime);
+				TArrayView<const FFrameNumber> AllTimes = ChannelInterface.GetTimes();
+				if (AllTimes.Num() > 0)
+				{
+					FloatSection->SetRange(TRange<FFrameNumber>(AllTimes[0], TRangeBound<FFrameNumber>::Inclusive(AllTimes.Last())));
+				}
 				return true;
 			}
 		}
@@ -744,6 +793,40 @@ bool ImportFBXProperty(FString NodeName, FString AnimatedPropertyName, FGuid Obj
 	return false;
 }
 
+void ImportTransformChannel(const FInterpCurveFloat& Source, FMovieSceneFloatChannel* Dest, FFrameRate DestFrameRate, bool bNegateTangents)
+{
+	TMovieSceneChannel<FMovieSceneFloatValue> ChannelInterface = Dest->GetInterface();
+	ChannelInterface.Reset();
+
+	for (int32 KeyIndex = 0; KeyIndex < Source.Points.Num(); ++KeyIndex)
+	{
+		float ArriveTangent = Source.Points[KeyIndex].ArriveTangent;
+		if (KeyIndex > 0)
+		{
+			ArriveTangent = ArriveTangent / (Source.Points[KeyIndex].InVal - Source.Points[KeyIndex-1].InVal);
+		}
+		
+		float LeaveTangent = Source.Points[KeyIndex].LeaveTangent;
+		if (KeyIndex < Source.Points.Num() - 1)
+		{
+			LeaveTangent = LeaveTangent / (Source.Points[KeyIndex+1].InVal - Source.Points[KeyIndex].InVal);
+		}
+
+		if (bNegateTangents)
+		{
+			ArriveTangent = -ArriveTangent;
+			LeaveTangent = -LeaveTangent;
+		}
+
+		FFrameNumber KeyTime = (Source.Points[KeyIndex].InVal * DestFrameRate).RoundToFrame();
+		FMatineeImportTools::SetOrAddKey(ChannelInterface, KeyTime, Source.Points[KeyIndex].OutVal, ArriveTangent, LeaveTangent, Source.Points[KeyIndex].InterpMode);
+	}
+
+	FKeyDataOptimizationParams Params;
+	Params.bAutoSetInterpolation = true;
+	MovieScene::Optimize(Dest, Params);
+}
+
 bool ImportFBXTransform(FString NodeName, FGuid ObjectBinding, UnFbx::FFbxCurvesAPI& CurveAPI, UMovieScene* InMovieScene)
 {
 	const UMovieSceneUserImportFBXSettings* ImportFBXSettings = GetDefault<UMovieSceneUserImportFBXSettings>();
@@ -763,121 +846,72 @@ bool ImportFBXTransform(FString NodeName, FGuid ObjectBinding, UnFbx::FFbxCurves
 	}
 
 	bool bSectionAdded = false;
-	UMovieScene3DTransformSection* TransformSection = Cast<UMovieScene3DTransformSection>(TransformTrack->FindOrAddSection(0.f, bSectionAdded));
+	UMovieScene3DTransformSection* TransformSection = Cast<UMovieScene3DTransformSection>(TransformTrack->FindOrAddSection(0, bSectionAdded));
 	if (!TransformSection)
 	{
 		return false;
 	}
 
+	FFrameRate FrameRate = TransformSection->GetTypedOuter<UMovieScene>()->GetFrameResolution();
+
 	if (bSectionAdded)
 	{
-		TransformSection->SetIsInfinite(true);
+		TransformSection->SetRange(TRange<FFrameNumber>::All());
 	}
 
-	for (int32 ChannelIndex = 0; ChannelIndex < 3; ++ChannelIndex)
-	{
-		EAxis::Type ChannelAxis = EAxis::X;
-		if (ChannelIndex == 1)
-		{
-			ChannelAxis = EAxis::Y;
-		}
-		else if (ChannelIndex == 2)
-		{
-			ChannelAxis = EAxis::Z;
-		}
-		TransformSection->GetTranslationCurve(ChannelAxis).SetDefaultValue(DefaultTransform.GetLocation()[ChannelIndex]);
-		TransformSection->GetRotationCurve(ChannelAxis).SetDefaultValue(DefaultTransform.GetRotation().Euler()[ChannelIndex]);
-		TransformSection->GetScaleCurve(ChannelAxis).SetDefaultValue(DefaultTransform.GetScale3D()[ChannelIndex]);
-	}
+	FVector Location = DefaultTransform.GetLocation(), Rotation = DefaultTransform.GetRotation().Euler(), Scale3D = DefaultTransform.GetScale3D();
 
-	float MinTime = FLT_MAX;
-	float MaxTime = -FLT_MAX;
+	TArrayView<FMovieSceneFloatChannel*> Channels = TransformSection->GetChannelProxy().GetChannels<FMovieSceneFloatChannel>();
 
-	const int NumCurves = 3; // Trans, Rot, Scale
-	for (int32 CurveIndex = 0; CurveIndex < NumCurves; ++CurveIndex)
-	{
-		for (int32 ChannelIndex = 0; ChannelIndex < 3; ++ChannelIndex)
-		{
-			EAxis::Type ChannelAxis = EAxis::X;
-			if (ChannelIndex == 1)
-			{
-				ChannelAxis = EAxis::Y;
-			}
-			else if (ChannelIndex == 2)
-			{
-				ChannelAxis = EAxis::Z;
-			}
-	
-			FInterpCurveFloat* CurveFloat = nullptr;
-			FRichCurve* ChannelCurve = nullptr;
-			bool bNegative = false;
+	Channels[0]->SetDefault(Location.X);
+	Channels[1]->SetDefault(Location.Y);
+	Channels[2]->SetDefault(Location.Z);
 
-			if (CurveIndex == 0)
-			{
-				CurveFloat = &Translation[ChannelIndex];
-				ChannelCurve = &TransformSection->GetTranslationCurve(ChannelAxis);
-				if (ChannelIndex == 1)
-				{
-					bNegative = true;
-				}
-			}
-			else if (CurveIndex == 1)
-			{
-				CurveFloat = &EulerRotation[ChannelIndex];
-				ChannelCurve = &TransformSection->GetRotationCurve(ChannelAxis);
-				if (ChannelIndex == 1 || ChannelIndex == 2)
-				{
-					bNegative = true;
-				}
-			}
-			else if (CurveIndex == 2)
-			{
-				CurveFloat = &Scale[ChannelIndex];
-				ChannelCurve = &TransformSection->GetScaleCurve(ChannelAxis);
-			}
+	Channels[3]->SetDefault(Rotation.X);
+	Channels[4]->SetDefault(Rotation.Y);
+	Channels[5]->SetDefault(Rotation.Z);
 
-			if (ChannelCurve != nullptr && CurveFloat != nullptr)
-			{
-				ChannelCurve->Reset();
-				
-				for (int32 KeyIndex = 0; KeyIndex < CurveFloat->Points.Num(); ++KeyIndex)
-				{
-					MinTime = FMath::Min(MinTime, CurveFloat->Points[KeyIndex].InVal);
-					MaxTime = FMath::Max(MaxTime, CurveFloat->Points[KeyIndex].InVal);
-					
-					float ArriveTangent = CurveFloat->Points[KeyIndex].ArriveTangent;
-					if (KeyIndex > 0)
-					{
-						ArriveTangent = ArriveTangent / (CurveFloat->Points[KeyIndex].InVal - CurveFloat->Points[KeyIndex-1].InVal);
-					}
-					
-					float LeaveTangent = CurveFloat->Points[KeyIndex].LeaveTangent;
-					if (KeyIndex < CurveFloat->Points.Num() - 1)
-					{
-						LeaveTangent = LeaveTangent / (CurveFloat->Points[KeyIndex+1].InVal - CurveFloat->Points[KeyIndex].InVal);
-					}
+	Channels[6]->SetDefault(Scale3D.X);
+	Channels[7]->SetDefault(Scale3D.Y);
+	Channels[8]->SetDefault(Scale3D.Z);
 
-					if (bNegative)
-					{
-						ArriveTangent = -ArriveTangent;
-						LeaveTangent = -LeaveTangent;
-					}
+	ImportTransformChannel(Translation[0],   Channels[0], FrameRate, false);
+	ImportTransformChannel(Translation[1],   Channels[1], FrameRate, true);
+	ImportTransformChannel(Translation[2],   Channels[2], FrameRate, false);
 
-					FMatineeImportTools::SetOrAddKey(*ChannelCurve, CurveFloat->Points[KeyIndex].InVal, CurveFloat->Points[KeyIndex].OutVal, ArriveTangent, LeaveTangent, CurveFloat->Points[KeyIndex].InterpMode);
-				}
-				
-				if (ImportFBXSettings->bReduceKeys)
-				{
-					ChannelCurve->RemoveRedundantKeys(ImportFBXSettings->ReduceKeysTolerance);
-				}
-				
-				ChannelCurve->AutoSetTangents();
-			}
+	ImportTransformChannel(EulerRotation[0], Channels[3], FrameRate, false);
+	ImportTransformChannel(EulerRotation[1], Channels[4], FrameRate, true);
+	ImportTransformChannel(EulerRotation[2], Channels[5], FrameRate, true);
+
+	ImportTransformChannel(Scale[0],         Channels[6], FrameRate, false);
+	ImportTransformChannel(Scale[1],         Channels[7], FrameRate, false);
+	ImportTransformChannel(Scale[2],         Channels[8], FrameRate, false);
+
+
+	FKeyDataOptimizationParams Params;
+	Params.bAutoSetInterpolation = true;
+
+	TOptional<TRangeBound<FFrameNumber>> MinLower;
+	TOptional<TRangeBound<FFrameNumber>> MaxUpper;
+
+	for (int32 ChannelIndex = 0; ChannelIndex < 9; ++ChannelIndex)
+	{
+		MovieScene::Optimize(Channels[ChannelIndex], Params);
+
+		TArrayView<const FFrameNumber> AllTimes = Channels[ChannelIndex]->GetTimes();
+		if (AllTimes.Num() > 0)
+		{
+			MinLower = TRangeBound<FFrameNumber>::MinLower(AllTimes[0],     MinLower.Get(AllTimes[0]));
+			MaxUpper = TRangeBound<FFrameNumber>::MaxUpper(AllTimes.Last(), MaxUpper.Get(AllTimes.Last()));
 		}
 	}
-		
-	TransformSection->SetStartTime(MinTime);
-	TransformSection->SetEndTime(MaxTime);
+
+	// MinLower and MinUpper are both either set or unset
+	if (MinLower.IsSet())
+	{
+		TransformSection->SetRange(TRange<FFrameNumber>(MinLower.GetValue(), MaxUpper.GetValue()));
+	}
+
 	return true;
 }
 
@@ -1119,7 +1153,7 @@ void ImportFBXCamera(UnFbx::FFbxImporter* FbxImporter, UMovieScene* InMovieScene
 				if (FloatTrack)
 				{
 					bool bSectionAdded = false;
-					UMovieSceneFloatSection* FloatSection = Cast<UMovieSceneFloatSection>(FloatTrack->FindOrAddSection(0.f, bSectionAdded));
+					UMovieSceneFloatSection* FloatSection = Cast<UMovieSceneFloatSection>(FloatTrack->FindOrAddSection(0, bSectionAdded));
 					if (!FloatSection)
 					{
 						continue;
@@ -1127,10 +1161,10 @@ void ImportFBXCamera(UnFbx::FFbxImporter* FbxImporter, UMovieScene* InMovieScene
 
 					if (bSectionAdded)
 					{
-						FloatSection->SetIsInfinite(true);
+						FloatSection->SetRange(TRange<FFrameNumber>::All());
 					}
 
-					FloatSection->SetDefault(FocalLength);
+					FloatSection->GetChannelProxy().GetChannel<FMovieSceneFloatChannel>(0)->SetDefault(FocalLength);
 				}
 			}
 		}
@@ -1341,18 +1375,27 @@ EInterpCurveMode MovieSceneToolHelpers::RichCurveInterpolationToMatineeInterpola
 	}
 }
 
-void MovieSceneToolHelpers::CopyRichCurveToMoveAxis(const FRichCurve& RichCurve, UInterpTrackMoveAxis* MoveAxis)
+void MovieSceneToolHelpers::CopyKeyDataToMoveAxis(const TMovieSceneChannel<FMovieSceneFloatValue>& Channel, UInterpTrackMoveAxis* MoveAxis, FFrameRate InFrameRate)
 {
+	MoveAxis->FloatTrack.Points.Reset();
+
 	static FName LookupName(NAME_None);
-	int32 KeyIndex = 0;
-	for ( auto KeyIt( RichCurve.GetKeyIterator() ); KeyIt; ++KeyIt, ++KeyIndex )
+	
+	TArrayView<const FFrameNumber>          Times  = Channel.GetTimes();
+	TArrayView<const FMovieSceneFloatValue> Values = Channel.GetValues();
+
+	for (int32 KeyIndex = 0; KeyIndex < Times.Num(); ++KeyIndex)
 	{
-		MoveAxis->FloatTrack.AddPoint(KeyIt->Time, KeyIt->Value);
-		MoveAxis->FloatTrack.Points[KeyIndex].InterpMode = RichCurveInterpolationToMatineeInterpolation(KeyIt->InterpMode);
-		MoveAxis->FloatTrack.Points[KeyIndex].ArriveTangent = KeyIt->ArriveTangent;
-		MoveAxis->FloatTrack.Points[KeyIndex].LeaveTangent = KeyIt->LeaveTangent;
+		const float Time = Times[KeyIndex] / InFrameRate;
+		const FMovieSceneFloatValue& Value = Values[KeyIndex];
+
+		const int32 PointIndex = MoveAxis->FloatTrack.AddPoint(Time, Value.Value);
+		MoveAxis->LookupTrack.AddPoint(Time, LookupName);
 
-		MoveAxis->LookupTrack.AddPoint(KeyIt->Time, LookupName);
+		FInterpCurvePoint<float>& Point = MoveAxis->FloatTrack.Points[PointIndex];
+		Point.ArriveTangent = Value.Tangent.ArriveTangent;
+		Point.LeaveTangent = Value.Tangent.LeaveTangent;
+		Point.InterpMode = RichCurveInterpolationToMatineeInterpolation(Value.InterpMode);
 	}
 }
 
@@ -1419,16 +1462,19 @@ UObject* MovieSceneToolHelpers::ExportToCameraAnim(UMovieScene* InMovieScene, FG
 			}
 			else
 			{
+				FFrameRate FrameResolution = InMovieScene->GetFrameResolution();
 				UMovieScene3DTransformSection* TransformSection = Cast<UMovieScene3DTransformSection>(Sections[0]);
-				CopyRichCurveToMoveAxis(TransformSection->GetTranslationCurve(EAxis::X), MoveAxies[AXIS_TranslationX]);
-				CopyRichCurveToMoveAxis(TransformSection->GetTranslationCurve(EAxis::Y), MoveAxies[AXIS_TranslationY]);
-				CopyRichCurveToMoveAxis(TransformSection->GetTranslationCurve(EAxis::Z), MoveAxies[AXIS_TranslationZ]);
-				CopyRichCurveToMoveAxis(TransformSection->GetRotationCurve(EAxis::X), MoveAxies[AXIS_RotationX]);
-				CopyRichCurveToMoveAxis(TransformSection->GetRotationCurve(EAxis::Y), MoveAxies[AXIS_RotationY]);
-				CopyRichCurveToMoveAxis(TransformSection->GetRotationCurve(EAxis::Z), MoveAxies[AXIS_RotationZ]);
+				TArrayView<FMovieSceneFloatChannel*> FloatChannels = TransformSection->GetChannelProxy().GetChannels<FMovieSceneFloatChannel>();
+
+				CopyKeyDataToMoveAxis(FloatChannels[0]->GetInterface(), MoveAxies[AXIS_TranslationX], FrameResolution);
+				CopyKeyDataToMoveAxis(FloatChannels[1]->GetInterface(), MoveAxies[AXIS_TranslationY], FrameResolution);
+				CopyKeyDataToMoveAxis(FloatChannels[2]->GetInterface(), MoveAxies[AXIS_TranslationZ], FrameResolution);
+				CopyKeyDataToMoveAxis(FloatChannels[3]->GetInterface(), MoveAxies[AXIS_RotationX],    FrameResolution);
+				CopyKeyDataToMoveAxis(FloatChannels[4]->GetInterface(), MoveAxies[AXIS_RotationY],    FrameResolution);
+				CopyKeyDataToMoveAxis(FloatChannels[5]->GetInterface(), MoveAxies[AXIS_RotationZ],    FrameResolution);
 			}
 		}
 	}
 
 	return NewAsset;
-}
\ No newline at end of file
+}
diff --git a/Engine/Source/Editor/MovieSceneTools/Private/MovieSceneToolsModule.cpp b/Engine/Source/Editor/MovieSceneTools/Private/MovieSceneToolsModule.cpp
--- a/Engine/Source/Editor/MovieSceneTools/Private/MovieSceneToolsModule.cpp
+++ b/Engine/Source/Editor/MovieSceneTools/Private/MovieSceneToolsModule.cpp
@@ -46,6 +46,11 @@
 #include "IMovieSceneTools.h"
 #include "MovieSceneToolsProjectSettings.h"
 
+#include "ISequencerChannelInterface.h"
+#include "SequencerChannelInterface.h"
+#include "Channels/BuiltInChannelEditors.h"
+#include "Sections/MovieSceneEventSection.h"
+
 
 #define LOCTEXT_NAMESPACE "FMovieSceneToolsModule"
 
@@ -113,6 +118,15 @@ public:
 		PropertyModule.RegisterCustomClassLayout("MovieSceneToolsProjectSettings", FOnGetDetailCustomizationInstance::CreateStatic(&FMovieSceneToolsProjectSettingsCustomization::MakeInstance));
 		PropertyModule.RegisterCustomClassLayout("MovieSceneBuiltInEasingFunction", FOnGetDetailCustomizationInstance::CreateLambda(&MakeShared<FMovieSceneBuiltInEasingFunctionCustomization>));
 		PropertyModule.RegisterCustomPropertyTypeLayout("MovieSceneObjectBindingID", FOnGetPropertyTypeCustomizationInstance::CreateLambda(&MakeShared<FMovieSceneObjectBindingIDCustomization>));
+
+		SequencerModule.RegisterChannelInterface<FMovieSceneBoolChannel>();
+		SequencerModule.RegisterChannelInterface<FMovieSceneByteChannel>();
+		SequencerModule.RegisterChannelInterface<FMovieSceneIntegerChannel>();
+		SequencerModule.RegisterChannelInterface<FMovieSceneFloatChannel>();
+		SequencerModule.RegisterChannelInterface<FMovieSceneStringChannel>();
+		SequencerModule.RegisterChannelInterface<FMovieSceneParticleChannel>();
+		SequencerModule.RegisterChannelInterface<FMovieSceneActorReferenceData>();
+		SequencerModule.RegisterChannelInterface<FMovieSceneEventSectionData>();
 	}
 
 	virtual void ShutdownModule() override
@@ -181,11 +195,11 @@ public:
 		DefineImplicitConversion<uint8, int32>();
 		DefineImplicitConversion<uint8, bool>();
 
-		DefineExplicitConversion<int32, FRichCurveKey>([](const int32& In) -> FRichCurveKey { return FRichCurveKey(0.f, In);	});
-		DefineExplicitConversion<uint8, FRichCurveKey>([](const uint8& In) -> FRichCurveKey { return FRichCurveKey(0.f, In);	});
-		DefineExplicitConversion<FRichCurveKey, int32>([](const FRichCurveKey& In) -> int32 { return In.Value; 					});
-		DefineExplicitConversion<FRichCurveKey, uint8>([](const FRichCurveKey& In) -> uint8 { return In.Value; 					});
-		DefineExplicitConversion<FRichCurveKey, bool>([](const FRichCurveKey& In) -> bool	{ return !!In.Value; 				});
+		DefineExplicitConversion<int32, FMovieSceneFloatValue>([](const int32& In) -> FMovieSceneFloatValue { return FMovieSceneFloatValue(In);	});
+		DefineExplicitConversion<uint8, FMovieSceneFloatValue>([](const uint8& In) -> FMovieSceneFloatValue { return FMovieSceneFloatValue(In);	});
+		DefineExplicitConversion<FMovieSceneFloatValue, int32>([](const FMovieSceneFloatValue& In) -> int32 { return In.Value; 					});
+		DefineExplicitConversion<FMovieSceneFloatValue, uint8>([](const FMovieSceneFloatValue& In) -> uint8 { return In.Value; 					});
+		DefineExplicitConversion<FMovieSceneFloatValue, bool>([](const FMovieSceneFloatValue& In) -> bool	{ return !!In.Value; 				});
 
 		FSequencerClipboardReconciler::AddTrackAlias("Location.X", "R");
 		FSequencerClipboardReconciler::AddTrackAlias("Location.Y", "G");
diff --git a/Engine/Source/Editor/MovieSceneTools/Private/Sections/BoolPropertySection.cpp b/Engine/Source/Editor/MovieSceneTools/Private/Sections/BoolPropertySection.cpp
--- a/Engine/Source/Editor/MovieSceneTools/Private/Sections/BoolPropertySection.cpp
+++ b/Engine/Source/Editor/MovieSceneTools/Private/Sections/BoolPropertySection.cpp
@@ -2,59 +2,48 @@
 
 #include "Sections/BoolPropertySection.h"
 #include "Rendering/DrawElements.h"
-#include "Curves/IntegralCurve.h"
-#include "BoolKeyArea.h"
 #include "SequencerSectionPainter.h"
 #include "ISectionLayoutBuilder.h"
 #include "EditorStyleSet.h"
 #include "CommonMovieSceneTools.h"
 #include "Sections/MovieSceneBoolSection.h"
+#include "Channels/MovieSceneChannelProxy.h"
 
 
-void FBoolPropertySection::GenerateSectionLayout( class ISectionLayoutBuilder& LayoutBuilder ) const
-{
-	UMovieSceneBoolSection* BoolSection = Cast<UMovieSceneBoolSection>( &SectionObject );
-	TAttribute<TOptional<bool>> ExternalValue;
-	if (CanGetPropertyValue())
-	{
-		ExternalValue.Bind(TAttribute<TOptional<bool>>::FGetter::CreateLambda([&]
-		{
-			return GetPropertyValue<bool>();
-		}));
-	}
-	TSharedRef<FBoolKeyArea> KeyArea = MakeShareable( new FBoolKeyArea( BoolSection->GetCurve(), ExternalValue, BoolSection ) );
-	LayoutBuilder.SetSectionAsKeyArea( KeyArea );
-}
-
 int32 FBoolPropertySection::OnPaintSection( FSequencerSectionPainter& Painter ) const
 {
 	// custom drawing for bool curves
-	UMovieSceneBoolSection* BoolSection = Cast<UMovieSceneBoolSection>( &SectionObject );
+	UMovieSceneBoolSection* BoolSection = Cast<UMovieSceneBoolSection>( WeakSection.Get() );
 
 	TArray<float> SectionSwitchTimes;
 
 	const FTimeToPixel& TimeConverter = Painter.GetTimeConverter();
 
 	// Add the start time
-	const float StartTime = TimeConverter.PixelToTime(0.f);
-	const float EndTime = TimeConverter.PixelToTime(Painter.SectionGeometry.GetLocalSize().X);
+	const float StartTime = TimeConverter.PixelToSeconds(0.f);
+	const float EndTime   = TimeConverter.PixelToSeconds(Painter.SectionGeometry.GetLocalSize().X);
 	
 	SectionSwitchTimes.Add(StartTime);
 
-	FIntegralCurve& BoolCurve = BoolSection->GetCurve();
-	for ( TArray<FIntegralKey>::TConstIterator It( BoolCurve.GetKeyIterator() ); It; ++It )
+	int32 LayerId = Painter.PaintSectionBackground();
+
+	FMovieSceneBoolChannel* BoolChannel = BoolSection->GetChannelProxy().GetChannel<FMovieSceneBoolChannel>(0);
+	if (!BoolChannel)
 	{
-		float Time = It->Time;
-		if ( Time > StartTime && Time < EndTime )
+		return LayerId;
+	}
+
+	for ( FFrameNumber Time : BoolChannel->GetTimes() )
+	{
+		double TimeInSeconds = Time / TimeConverter.GetFrameResolution();
+		if ( TimeInSeconds > StartTime && TimeInSeconds < EndTime )
 		{
-			SectionSwitchTimes.Add( Time );
+			SectionSwitchTimes.Add( TimeInSeconds );
 		}
 	}
 
 	SectionSwitchTimes.Add(EndTime);
 
-	int32 LayerId = Painter.PaintSectionBackground();
-
 	const ESlateDrawEffect DrawEffects = Painter.bParentEnabled
 		? ESlateDrawEffect::None
 		: ESlateDrawEffect::DisabledEffect;
@@ -67,11 +56,14 @@ int32 FBoolPropertySection::OnPaintSection( FSequencerSectionPainter& Painter )
 	for ( int32 i = 0; i < SectionSwitchTimes.Num() - 1; ++i )
 	{
 		float ThisTime = SectionSwitchTimes[i];
+
+		bool ValueAtTime = false;
+		BoolChannel->Evaluate(ThisTime * TimeConverter.GetFrameResolution(), ValueAtTime);
+
+		const FColor Color = ValueAtTime ? FColor(0, 255, 0, 125) : FColor(255, 0, 0, 125);
 		
-		const FColor Color = BoolSection->Eval(ThisTime, false) ? FColor(0, 255, 0, 125) : FColor(255, 0, 0, 125);
-		
-		FVector2D StartPos(TimeConverter.TimeToPixel(ThisTime), VerticalOffset);
-		FVector2D Size(TimeConverter.TimeToPixel(SectionSwitchTimes[i+1]) - StartPos.X, Height);
+		FVector2D StartPos(TimeConverter.SecondsToPixel(ThisTime), VerticalOffset);
+		FVector2D Size(TimeConverter.SecondsToPixel(SectionSwitchTimes[i+1]) - StartPos.X, Height);
 
 		FSlateDrawElement::MakeBox(
 			Painter.DrawElements,
diff --git a/Engine/Source/Editor/MovieSceneTools/Private/Sections/BoolPropertySection.h b/Engine/Source/Editor/MovieSceneTools/Private/Sections/BoolPropertySection.h
--- a/Engine/Source/Editor/MovieSceneTools/Private/Sections/BoolPropertySection.h
+++ b/Engine/Source/Editor/MovieSceneTools/Private/Sections/BoolPropertySection.h
@@ -2,51 +2,29 @@
 
 #pragma once
 
-#include "CoreMinimal.h"
-#include "Misc/Guid.h"
-#include "PropertySection.h"
+#include "ISequencerSection.h"
 
 class FSequencerSectionPainter;
 
 /**
-* An implementation of bool property sections.
+* An implementation of bool sections.
 */
 class FBoolPropertySection
-	: public FPropertySection
+	: public FSequencerSection
 {
 public:
 
 	/**
-	* Creates a new bool property section for editing non-property bool sections.
+	* Creates a new bool section for editing bool sections.
 	*
 	* @param InSectionObject The section object which is being displayed and edited.
 	* @param InDisplayName A display name for the section being displayed and edited.
-	* @TODO Make another section which is a good base class for non-property bool sections.
 	*/
-	FBoolPropertySection(UMovieSceneSection& InSectionObject, const FText& InDisplayName)
-		: FPropertySection(InSectionObject, InDisplayName)
-	{
-	}
-
-	/**
-	* Creates a new bool property section.
-	*
-	* @param InSequencer The sequencer which is controlling this property section.
-	* @param InObjectBinding The object binding for the object which owns the property that this section is animating.
-	* @param InPropertyName The name of the property which is animated by this section.
-	* @param InPropertyPath A string representing the path to the property which is animated by this section.
-	* @param InSectionObject The section object which is being displayed and edited.
-	* @param InDisplayName A display name for the section being displayed and edited.
-	*/
-	FBoolPropertySection(ISequencer* InSequencer, FGuid InObjectBinding, FName InPropertyName, const FString& InPropertyPath, UMovieSceneSection& InSectionObject, const FText& InDisplayName)
-		: FPropertySection(InSequencer, InObjectBinding, InPropertyName, InPropertyPath, InSectionObject, InDisplayName)
-	{
-	}
+	FBoolPropertySection(UMovieSceneSection& InSectionObject)
+		: FSequencerSection(InSectionObject)
+	{}
 
 public:
 
-	//~ FPropertySection interface
-
-	virtual void GenerateSectionLayout(class ISectionLayoutBuilder& LayoutBuilder) const override;
 	virtual int32 OnPaintSection(FSequencerSectionPainter& Painter) const override;
 };
diff --git a/Engine/Source/Editor/MovieSceneTools/Private/Sections/CameraCutSection.cpp b/Engine/Source/Editor/MovieSceneTools/Private/Sections/CameraCutSection.cpp
--- a/Engine/Source/Editor/MovieSceneTools/Private/Sections/CameraCutSection.cpp
+++ b/Engine/Source/Editor/MovieSceneTools/Private/Sections/CameraCutSection.cpp
@@ -36,12 +36,13 @@ FCameraCutSection::~FCameraCutSection()
 /* ISequencerSection interface
  *****************************************************************************/
 
-void FCameraCutSection::SetSingleTime(float GlobalTime)
+void FCameraCutSection::SetSingleTime(double GlobalTime)
 {
 	UMovieSceneCameraCutSection* CameraCutSection = Cast<UMovieSceneCameraCutSection>(Section);
-	if (CameraCutSection)
+	if (CameraCutSection && CameraCutSection->HasStartFrame())
 	{
-		CameraCutSection->SetThumbnailReferenceOffset(GlobalTime - CameraCutSection->GetStartTime());
+		double ReferenceOffsetSeconds = CameraCutSection->GetInclusiveStartFrame() / CameraCutSection->GetTypedOuter<UMovieScene>()->GetFrameResolution();
+		CameraCutSection->SetThumbnailReferenceOffset(GlobalTime - ReferenceOffsetSeconds);
 	}
 }
 
@@ -50,13 +51,14 @@ void FCameraCutSection::Tick(const FGeometry& AllottedGeometry, const FGeometry&
 	UMovieSceneCameraCutSection* CameraCutSection = Cast<UMovieSceneCameraCutSection>(Section);
 	if (CameraCutSection)
 	{
-		if (GetDefault<UMovieSceneUserThumbnailSettings>()->bDrawSingleThumbnails)
+		if (GetDefault<UMovieSceneUserThumbnailSettings>()->bDrawSingleThumbnails && CameraCutSection->HasStartFrame())
 		{
-			ThumbnailCache.SetSingleReferenceFrame(CameraCutSection->GetThumbnailReferenceOffset() + CameraCutSection->GetStartTime());
+			double ReferenceOffsetSeconds = CameraCutSection->GetInclusiveStartFrame() / CameraCutSection->GetTypedOuter<UMovieScene>()->GetFrameResolution() + CameraCutSection->GetThumbnailReferenceOffset();
+			ThumbnailCache.SetSingleReferenceFrame(ReferenceOffsetSeconds);
 		}
 		else
 		{
-			ThumbnailCache.SetSingleReferenceFrame(TOptional<float>());
+			ThumbnailCache.SetSingleReferenceFrame(TOptional<double>());
 		}
 	}
 
@@ -69,12 +71,12 @@ void FCameraCutSection::BuildSectionContextMenu(FMenuBuilder& MenuBuilder, const
 
 	UWorld* World = GEditor->GetEditorWorldContext().World();
 
-	if (World == nullptr)
+	if (World == nullptr || !Section->HasStartFrame())
 	{
 		return;
 	}
 
-	const AActor* CameraActor = GetCameraForFrame(Section->GetStartTime());
+	const AActor* CameraActor = GetCameraForFrame(Section->GetInclusiveStartFrame());
 
 	// get list of available cameras
 	TArray<AActor*> AllCameras;
@@ -119,7 +121,7 @@ void FCameraCutSection::BuildSectionContextMenu(FMenuBuilder& MenuBuilder, const
 /* FThumbnailSection interface
  *****************************************************************************/
 
-const AActor* FCameraCutSection::GetCameraForFrame(float Time) const
+const AActor* FCameraCutSection::GetCameraForFrame(FFrameNumber Time) const
 {
 	UMovieSceneCameraCutSection* CameraCutSection = Cast<UMovieSceneCameraCutSection>(Section);
 	TSharedPtr<ISequencer> Sequencer = SequencerPtr.Pin();
@@ -129,9 +131,12 @@ const AActor* FCameraCutSection::GetCameraForFrame(float Time) const
 		FMovieSceneSequenceID SequenceID = Sequencer->GetFocusedTemplateID();
 		if (CameraCutSection->GetCameraBindingID().GetSequenceID().IsValid())
 		{
-			// Ensure that this ID is resolvable from the root, based on the current local sequence ID
-			FMovieSceneObjectBindingID RootBindingID = CameraCutSection->GetCameraBindingID().ResolveLocalToRoot(SequenceID, Sequencer->GetEvaluationTemplate().GetHierarchy());
-			SequenceID = RootBindingID.GetSequenceID();
+			if (const FMovieSceneSubSequenceData* SubData = Sequencer->GetEvaluationTemplate().GetHierarchy().FindSubData(SequenceID))
+			{
+				// Ensure that this ID is resolvable from the root, based on the current local sequence ID
+				FMovieSceneObjectBindingID RootBindingID = CameraCutSection->GetCameraBindingID().ResolveLocalToRoot(SequenceID, Sequencer->GetEvaluationTemplate().GetHierarchy());
+				SequenceID = RootBindingID.GetSequenceID();
+			}
 		}
 
 		for (TWeakObjectPtr<>& Object : Sequencer->FindBoundObjects(CameraCutSection->GetCameraBindingID().GetGuid(), SequenceID))
@@ -172,7 +177,7 @@ int32 FCameraCutSection::OnPaintSection(FSequencerSectionPainter& InPainter) con
 
 FText FCameraCutSection::HandleThumbnailTextBlockText() const
 {
-	const AActor* CameraActor = GetCameraForFrame(Section->GetStartTime());
+	const AActor* CameraActor = Section->HasStartFrame() ? GetCameraForFrame(Section->GetInclusiveStartFrame()) : nullptr;
 	if (CameraActor)
 	{
 		return FText::FromString(CameraActor->GetActorLabel());
diff --git a/Engine/Source/Editor/MovieSceneTools/Private/Sections/CameraCutSection.h b/Engine/Source/Editor/MovieSceneTools/Private/Sections/CameraCutSection.h
--- a/Engine/Source/Editor/MovieSceneTools/Private/Sections/CameraCutSection.h
+++ b/Engine/Source/Editor/MovieSceneTools/Private/Sections/CameraCutSection.h
@@ -35,13 +35,13 @@ public:
 	virtual FMargin GetContentPadding() const override;
 	// FThumbnail interface
 
-	virtual void SetSingleTime(float GlobalTime) override;
+	virtual void SetSingleTime(double GlobalTime) override;
 	virtual FText HandleThumbnailTextBlockText() const override;
 
 private:
 
 	/** Get a representative camera for the given time */
-	const AActor* GetCameraForFrame(float Time) const;
+	const AActor* GetCameraForFrame(FFrameNumber Time) const;
 
 	/** Callback for executing a "Set Camera" menu entry in the context menu. */
 	void HandleSetCameraMenuEntryExecute(AActor* InCamera);
diff --git a/Engine/Source/Editor/MovieSceneTools/Private/Sections/CinematicShotSection.cpp b/Engine/Source/Editor/MovieSceneTools/Private/Sections/CinematicShotSection.cpp
--- a/Engine/Source/Editor/MovieSceneTools/Private/Sections/CinematicShotSection.cpp
+++ b/Engine/Source/Editor/MovieSceneTools/Private/Sections/CinematicShotSection.cpp
@@ -13,6 +13,7 @@
 #include "SequencerSectionPainter.h"
 #include "EditorStyleSet.h"
 #include "MovieSceneToolHelpers.h"
+#include "MovieSceneTimeHelpers.h"
 
 
 #define LOCTEXT_NAMESPACE "FCinematicShotSection"
@@ -22,11 +23,21 @@
  *****************************************************************************/
 
 FCinematicShotSection::FCinematicSectionCache::FCinematicSectionCache(UMovieSceneCinematicShotSection* Section)
-	: ActualStartTime(0.f), TimeScale(0.f)
+	: InnerFrameRate(1, 1)
+	, InnerFrameOffset(0)
+	, SectionStartFrame(0)
+	, TimeScale(1.f)
 {
 	if (Section)
 	{
-		ActualStartTime = Section->GetStartTime() - Section->Parameters.StartOffset;
+		UMovieSceneSequence* InnerSequence = Section->GetSequence();
+		if (InnerSequence)
+		{
+			InnerFrameRate = InnerSequence->GetMovieScene()->GetFrameResolution();
+		}
+
+		InnerFrameOffset = Section->Parameters.GetStartFrameOffset();
+		SectionStartFrame = Section->HasStartFrame() ? Section->GetInclusiveStartFrame() : 0;
 		TimeScale = Section->Parameters.TimeScale;
 	}
 }
@@ -37,8 +48,8 @@ FCinematicShotSection::FCinematicShotSection(TSharedPtr<ISequencer> InSequencer,
 	, SectionObject(*CastChecked<UMovieSceneCinematicShotSection>(&InSection))
 	, Sequencer(InSequencer)
 	, CinematicShotTrackEditor(InCinematicShotTrackEditor)
-	, InitialStartOffsetDuringResize(0.f)
-	, InitialStartTimeDuringResize(0.f)
+	, InitialStartOffsetDuringResize(0)
+	, InitialStartTimeDuringResize(0)
 	, ThumbnailCacheData(&SectionObject)
 {
 	AdditionalDrawEffect = ESlateDrawEffect::NoGamma;
@@ -59,29 +70,33 @@ FMargin FCinematicShotSection::GetContentPadding() const
 	return FMargin(8.f, 15.f);
 }
 
-void FCinematicShotSection::SetSingleTime(float GlobalTime)
+void FCinematicShotSection::SetSingleTime(double GlobalTime)
 {
-	SectionObject.SetThumbnailReferenceOffset(GlobalTime - SectionObject.GetStartTime());
+	double ReferenceOffsetSeconds = SectionObject.HasStartFrame() ? SectionObject.GetInclusiveStartFrame() / SectionObject.GetTypedOuter<UMovieScene>()->GetFrameResolution() : 0;
+	SectionObject.SetThumbnailReferenceOffset(GlobalTime - ReferenceOffsetSeconds);
 }
 
 void FCinematicShotSection::BeginResizeSection()
 {
-	InitialStartOffsetDuringResize = SectionObject.Parameters.StartOffset;
-	InitialStartTimeDuringResize = SectionObject.GetStartTime();
+	InitialStartOffsetDuringResize = SectionObject.Parameters.GetStartFrameOffset();
+	InitialStartTimeDuringResize = SectionObject.HasStartFrame() ? SectionObject.GetInclusiveStartFrame() : 0;
 }
 
-void FCinematicShotSection::ResizeSection(ESequencerSectionResizeMode ResizeMode, float ResizeTime)
+void FCinematicShotSection::ResizeSection(ESequencerSectionResizeMode ResizeMode, FFrameNumber ResizeTime)
 {
+	UMovieSceneSequence* InnerSequence = SectionObject.GetSequence();
+
 	// Adjust the start offset when resizing from the beginning
-	if (ResizeMode == SSRM_LeadingEdge)
+	if (ResizeMode == SSRM_LeadingEdge && InnerSequence)
 	{
-		float StartOffset = (ResizeTime - InitialStartTimeDuringResize) / SectionObject.Parameters.TimeScale;
-		StartOffset += InitialStartOffsetDuringResize;
+		const FFrameRate    OuterFrameRate   = SectionObject.GetTypedOuter<UMovieScene>()->GetFrameResolution();
+		const FFrameRate    InnerFrameRate   = InnerSequence->GetMovieScene()->GetFrameResolution();
+		const FFrameNumber  ResizeDifference = ResizeTime - InitialStartTimeDuringResize;
+		const FFrameTime    InnerFrameTime   = ConvertFrameTime(ResizeDifference, OuterFrameRate, InnerFrameRate);
+		const int32         NewStartOffset   = FFrameTime::FromDecimal(InnerFrameTime.AsDecimal() * SectionObject.Parameters.TimeScale).FrameNumber.Value;
 
 		// Ensure start offset is not less than 0
-		StartOffset = FMath::Max(StartOffset, 0.f);
-
-		SectionObject.Parameters.StartOffset = StartOffset;
+		SectionObject.Parameters.SetStartFrameOffset(FMath::Max(NewStartOffset, 0));
 	}
 
 	FViewportThumbnailSection::ResizeSection(ResizeMode, ResizeTime);
@@ -92,15 +107,21 @@ void FCinematicShotSection::BeginSlipSection()
 	BeginResizeSection();
 }
 
-void FCinematicShotSection::SlipSection(float SlipTime)
+void FCinematicShotSection::SlipSection(double SlipTime)
 {
-	float StartOffset = (SlipTime - InitialStartTimeDuringResize) / SectionObject.Parameters.TimeScale;
-	StartOffset += InitialStartOffsetDuringResize;
+	UMovieSceneSequence* InnerSequence = SectionObject.GetSequence();
 
-	// Ensure start offset is not less than 0
-	StartOffset = FMath::Max(StartOffset, 0.f);
+	// Adjust the start offset when resizing from the beginning
+	if (InnerSequence)
+	{
+		const FFrameRate    OuterFrameRate   = SectionObject.GetTypedOuter<UMovieScene>()->GetFrameResolution();
+		const FFrameRate    InnerFrameRate   = InnerSequence->GetMovieScene()->GetFrameResolution();
+		const double        ResizeDifference = SlipTime - InitialStartTimeDuringResize / OuterFrameRate;
+		const int32         NewStartOffset   = ( (ResizeDifference * SectionObject.Parameters.TimeScale) * InnerFrameRate ).FrameNumber.Value;
 
-	SectionObject.Parameters.StartOffset = StartOffset;
+		// Ensure start offset is not less than 0
+		SectionObject.Parameters.SetStartFrameOffset(FMath::Max(NewStartOffset, 0));
+	}
 
 	FViewportThumbnailSection::SlipSection(SlipTime);
 }
@@ -116,13 +137,14 @@ void FCinematicShotSection::Tick(const FGeometry& AllottedGeometry, const FGeome
 	ThumbnailCacheData = NewCacheData;
 
 	// Update single reference frame settings
-	if (GetDefault<UMovieSceneUserThumbnailSettings>()->bDrawSingleThumbnails)
+	if (GetDefault<UMovieSceneUserThumbnailSettings>()->bDrawSingleThumbnails && SectionObject.HasStartFrame())
 	{
-		ThumbnailCache.SetSingleReferenceFrame(SectionObject.GetStartTime() + SectionObject.GetThumbnailReferenceOffset());
+		double ReferenceTime = SectionObject.GetInclusiveStartFrame() / SectionObject.GetTypedOuter<UMovieScene>()->GetFrameResolution() + SectionObject.GetThumbnailReferenceOffset();
+		ThumbnailCache.SetSingleReferenceFrame(ReferenceTime);
 	}
 	else
 	{
-		ThumbnailCache.SetSingleReferenceFrame(TOptional<float>());
+		ThumbnailCache.SetSingleReferenceFrame(TOptional<double>());
 	}
 
 	FViewportThumbnailSection::Tick(AllottedGeometry, ClippedGeometry, InCurrentTime, InDeltaTime);
@@ -152,60 +174,59 @@ int32 FCinematicShotSection::OnPaintSection(FSequencerSectionPainter& InPainter)
 		InPainter.bParentEnabled ? ESlateDrawEffect::None : ESlateDrawEffect::DisabledEffect
 	);
 
-	const float SectionSize = SectionObject.GetTimeSize();
-
-	if (SectionSize <= 0.0f)
+	TRange<FFrameNumber> SectionRange = SectionObject.GetRange();
+	if (SectionRange.GetLowerBound().IsOpen() || SectionRange.GetUpperBound().IsOpen())
 	{
 		return InPainter.LayerId;
 	}
 
-	FViewportThumbnailSection::OnPaintSection(InPainter);
+	const FFrameNumber SectionStartFrame = SectionObject.GetInclusiveStartFrame();
+	const FFrameNumber SectionEndFrame   = SectionObject.GetExclusiveEndFrame();
+	const int32        SectionSize       = MovieScene::DiscreteSize(SectionRange);
 
-	const float DrawScale = InPainter.SectionGeometry.Size.X / SectionSize;
-		
-	UMovieScene* MovieScene = nullptr;
-	FFloatRange PlaybackRange;
-	if(SectionObject.GetSequence() != nullptr)
+	if (SectionSize <= 0)
 	{
-		MovieScene = SectionObject.GetSequence()->GetMovieScene();
-		PlaybackRange = MovieScene->GetPlaybackRange();
+		return InPainter.LayerId;
 	}
-	else
+
+	FViewportThumbnailSection::OnPaintSection(InPainter);
+
+	UMovieSceneSequence* InnerSequence = SectionObject.GetSequence();
+	if (!InnerSequence)
 	{
-		UMovieSceneTrack* MovieSceneTrack = CastChecked<UMovieSceneTrack>(SectionObject.GetOuter());
-		MovieScene = CastChecked<UMovieScene>(MovieSceneTrack->GetOuter());
-		PlaybackRange = MovieScene->GetPlaybackRange();
+		return InPainter.LayerId;
 	}
 
-	// add box for the working size
-	const float StartOffset = 1.0f/SectionObject.Parameters.TimeScale * SectionObject.Parameters.StartOffset;
-	const float WorkingStart = -1.0f/SectionObject.Parameters.TimeScale * PlaybackRange.GetLowerBoundValue() - StartOffset;
-	const float WorkingSize = 1.0f/SectionObject.Parameters.TimeScale * (MovieScene != nullptr ? MovieScene->GetEditorData().WorkingRange.Size<float>() : 1.0f);
+	const float PixelsPerFrame = InPainter.SectionGeometry.Size.X / float(SectionSize);
+
+	UMovieScene*         MovieScene    = InnerSequence->GetMovieScene();
+	TRange<FFrameNumber> PlaybackRange = MovieScene->GetPlaybackRange();
+
+	FMovieSceneSequenceTransform InnerToOuterTransform = SectionObject.OuterToInnerTransform().Inverse();
 
-	// add dark tint for left out-of-bounds & working range
-	if (StartOffset < 0.0f)
+	const FFrameNumber PlaybackStart = (MovieScene::DiscreteInclusiveLower(PlaybackRange) * InnerToOuterTransform).FloorToFrame();
+	if (SectionRange.Contains(PlaybackStart))
 	{
+		const int32 StartOffset = (PlaybackStart-SectionStartFrame).Value;
+		// add dark tint for left out-of-bounds
 		FSlateDrawElement::MakeBox(
 			InPainter.DrawElements,
 			InPainter.LayerId++,
 			InPainter.SectionGeometry.ToPaintGeometry(
 				FVector2D(0.0f, 0.f),
-				FVector2D(-StartOffset * DrawScale, InPainter.SectionGeometry.Size.Y)
+				FVector2D(StartOffset * PixelsPerFrame, InPainter.SectionGeometry.Size.Y)
 			),
 			FEditorStyle::GetBrush("WhiteBrush"),
 			ESlateDrawEffect::None,
 			FLinearColor::Black.CopyWithNewOpacity(0.5f)
 		);
-	}
 
-	// add green line for playback start
-	if (StartOffset < 0)
-	{
+		// add green line for playback start
 		FSlateDrawElement::MakeBox(
 			InPainter.DrawElements,
 			InPainter.LayerId++,
 			InPainter.SectionGeometry.ToPaintGeometry(
-				FVector2D(-StartOffset * DrawScale, 0.f),
+				FVector2D(StartOffset * PixelsPerFrame, 0.f),
 				FVector2D(1.0f, InPainter.SectionGeometry.Size.Y)
 			),
 			FEditorStyle::GetBrush("WhiteBrush"),
@@ -214,32 +235,30 @@ int32 FCinematicShotSection::OnPaintSection(FSequencerSectionPainter& InPainter)
 		);
 	}
 
-	// add dark tint for right out-of-bounds & working range
-	const float PlaybackEnd = 1.0f/SectionObject.Parameters.TimeScale * PlaybackRange.Size<float>() - StartOffset;
-
-	if (PlaybackEnd < SectionSize)
+	const FFrameNumber PlaybackEnd = (MovieScene::DiscreteExclusiveUpper(PlaybackRange) * InnerToOuterTransform).FloorToFrame();
+	if (SectionRange.Contains(PlaybackEnd))
 	{
+		// add dark tint for right out-of-bounds
+		const int32 EndOffset = (PlaybackEnd-SectionStartFrame).Value;
 		FSlateDrawElement::MakeBox(
 			InPainter.DrawElements,
 			InPainter.LayerId++,
 			InPainter.SectionGeometry.ToPaintGeometry(
-				FVector2D(PlaybackEnd * DrawScale, 0.f),
-				FVector2D((SectionSize - PlaybackEnd) * DrawScale, InPainter.SectionGeometry.Size.Y)
+				FVector2D(EndOffset * PixelsPerFrame, 0.f),
+				FVector2D((SectionSize - EndOffset) * PixelsPerFrame, InPainter.SectionGeometry.Size.Y)
 			),
 			FEditorStyle::GetBrush("WhiteBrush"),
 			ESlateDrawEffect::None,
 			FLinearColor::Black.CopyWithNewOpacity(0.5f)
 		);
-	}
 
-	// add red line for playback end
-	if (PlaybackEnd <= SectionSize)
-	{
+
+		// add red line for playback end
 		FSlateDrawElement::MakeBox(
 			InPainter.DrawElements,
 			InPainter.LayerId++,
 			InPainter.SectionGeometry.ToPaintGeometry(
-				FVector2D(PlaybackEnd * DrawScale, 0.f),
+				FVector2D(EndOffset * PixelsPerFrame, 0.f),
 				FVector2D(1.0f, InPainter.SectionGeometry.Size.Y)
 			),
 			FEditorStyle::GetBrush("WhiteBrush"),
diff --git a/Engine/Source/Editor/MovieSceneTools/Private/Sections/CinematicShotSection.h b/Engine/Source/Editor/MovieSceneTools/Private/Sections/CinematicShotSection.h
--- a/Engine/Source/Editor/MovieSceneTools/Private/Sections/CinematicShotSection.h
+++ b/Engine/Source/Editor/MovieSceneTools/Private/Sections/CinematicShotSection.h
@@ -38,12 +38,12 @@ public:
 	virtual float GetSectionHeight() const override;
 	virtual FMargin GetContentPadding() const override;
 	virtual void BeginResizeSection() override;
-	virtual void ResizeSection(ESequencerSectionResizeMode ResizeMode, float ResizeTime) override;
+	virtual void ResizeSection(ESequencerSectionResizeMode ResizeMode, FFrameNumber ResizeTime) override;
 	virtual void BeginSlipSection() override;
-	virtual void SlipSection(float SlipTime) override;
+	virtual void SlipSection(double SlipTime) override;
 
 	// FThumbnail interface
-	virtual void SetSingleTime(float GlobalTime) override;
+	virtual void SetSingleTime(double GlobalTime) override;
 	virtual FText HandleThumbnailTextBlockText() const override;
 	virtual void HandleThumbnailTextBlockTextCommitted(const FText& NewThumbnailName, ETextCommit::Type CommitType) override;
 
@@ -64,10 +64,10 @@ private:
 	TWeakPtr<FCinematicShotTrackEditor> CinematicShotTrackEditor;
 
 	/** Cached start offset value valid only during resize */
-	float InitialStartOffsetDuringResize;
-	
+	int32 InitialStartOffsetDuringResize;
+
 	/** Cached start time valid only during resize */
-	float InitialStartTimeDuringResize;
+	FFrameNumber InitialStartTimeDuringResize;
 
 	struct FCinematicSectionCache
 	{
@@ -75,11 +75,13 @@ private:
 
 		bool operator!=(const FCinematicSectionCache& RHS) const
 		{
-			return !FMath::IsNearlyEqual(ActualStartTime, RHS.ActualStartTime, 0.001f) || TimeScale != RHS.TimeScale;
+			return InnerFrameRate != RHS.InnerFrameRate || InnerFrameOffset != RHS.InnerFrameOffset || SectionStartFrame != RHS.SectionStartFrame || TimeScale != RHS.TimeScale;
 		}
 
-		float ActualStartTime;
-		float TimeScale;
+		FFrameRate   InnerFrameRate;
+		int32        InnerFrameOffset;
+		FFrameNumber SectionStartFrame;
+		float        TimeScale;
 	};
 
 	/** Cached section thumbnail data */
