Copying //UE4/Dev-Sequencer to //UE4/Dev-Main (Source: //UE4/Dev-Sequencer @ 3946692)

diff --git a/Engine/Source/Editor/MovieSceneTools/Private/Sections/ColorPropertySection.cpp b/Engine/Source/Editor/MovieSceneTools/Private/Sections/ColorPropertySection.cpp
--- a/Engine/Source/Editor/MovieSceneTools/Private/Sections/ColorPropertySection.cpp
+++ b/Engine/Source/Editor/MovieSceneTools/Private/Sections/ColorPropertySection.cpp
@@ -4,58 +4,42 @@
 #include "Rendering/DrawElements.h"
 #include "Sections/MovieSceneColorSection.h"
 #include "SequencerSectionPainter.h"
-#include "FloatCurveKeyArea.h"
 #include "ISectionLayoutBuilder.h"
 #include "EditorStyleSet.h"
 #include "CommonMovieSceneTools.h"
-
-const FName SlateColorName("SlateColor");
-
-void FColorPropertySection::GenerateSectionLayout( class ISectionLayoutBuilder& LayoutBuilder ) const
+#include "Tracks/MovieScenePropertyTrack.h"
+#include "ISequencer.h"
+#include "Channels/MovieSceneChannelProxy.h"
+
+FColorPropertySection::FColorPropertySection(UMovieSceneSection& InSectionObject, const FGuid& InObjectBindingID, TWeakPtr<ISequencer> InSequencer)
+	: FSequencerSection(InSectionObject)
+	, ObjectBindingID(InObjectBindingID)
+	, WeakSequencer(InSequencer)
 {
-	UMovieSceneColorSection* ColorSection = Cast<UMovieSceneColorSection>( &SectionObject );
-
-	TAttribute<TOptional<float>> RedExternalValue = TAttribute<TOptional<float>>::Create(
-		TAttribute<TOptional<float>>::FGetter::CreateRaw(this, &FColorPropertySection::GetColorRedValue));
-	TSharedRef<FFloatCurveKeyArea> RedKeyArea = MakeShareable( new FFloatCurveKeyArea( &ColorSection->GetRedCurve(), RedExternalValue, ColorSection ) ) ;
-
-	TAttribute<TOptional<float>> GreenExternalValue = TAttribute<TOptional<float>>::Create(
-		TAttribute<TOptional<float>>::FGetter::CreateRaw(this, &FColorPropertySection::GetColorGreenValue));
-	TSharedRef<FFloatCurveKeyArea> GreenKeyArea = MakeShareable( new FFloatCurveKeyArea( &ColorSection->GetGreenCurve(), GreenExternalValue, ColorSection ) ) ;
-	
-	TAttribute<TOptional<float>> BlueExternalValue = TAttribute<TOptional<float>>::Create(
-		TAttribute<TOptional<float>>::FGetter::CreateRaw(this, &FColorPropertySection::GetColorBlueValue));
-	TSharedRef<FFloatCurveKeyArea> BlueKeyArea = MakeShareable( new FFloatCurveKeyArea( &ColorSection->GetBlueCurve(), BlueExternalValue, ColorSection ) ) ;
-	
-	TAttribute<TOptional<float>> AlphaExternalValue = TAttribute<TOptional<float>>::Create(
-		TAttribute<TOptional<float>>::FGetter::CreateRaw(this, &FColorPropertySection::GetColorAlphaValue));
-	TSharedRef<FFloatCurveKeyArea> AlphaKeyArea = MakeShareable( new FFloatCurveKeyArea( &ColorSection->GetAlphaCurve(), AlphaExternalValue, ColorSection ) );
-
-	LayoutBuilder.AddKeyArea( "R", NSLOCTEXT( "FColorPropertySection", "RedArea", "Red" ), RedKeyArea );
-	LayoutBuilder.AddKeyArea( "G", NSLOCTEXT( "FColorPropertySection", "GreenArea", "Green" ), GreenKeyArea );
-	LayoutBuilder.AddKeyArea( "B", NSLOCTEXT( "FColorPropertySection", "BlueArea", "Blue" ), BlueKeyArea );
-	LayoutBuilder.AddKeyArea( "A", NSLOCTEXT( "FColorPropertySection", "OpacityArea", "Opacity" ), AlphaKeyArea );
+	UMovieScenePropertyTrack* PropertyTrack = InSectionObject.GetTypedOuter<UMovieScenePropertyTrack>();
+	if (PropertyTrack)
+	{
+		PropertyBindings = FTrackInstancePropertyBindings(PropertyTrack->GetPropertyName(), PropertyTrack->GetPropertyPath());
+	}
 }
 
-
 int32 FColorPropertySection::OnPaintSection( FSequencerSectionPainter& Painter ) const
 {
 	int32 LayerId = Painter.PaintSectionBackground();
 
 	const ESlateDrawEffect DrawEffects = Painter.bParentEnabled ? ESlateDrawEffect::None : ESlateDrawEffect::DisabledEffect;
 
-	const UMovieSceneColorSection* ColorSection = Cast<const UMovieSceneColorSection>( &SectionObject );
+	const UMovieSceneColorSection* ColorSection = Cast<const UMovieSceneColorSection>( WeakSection.Get() );
 
 	const FTimeToPixel& TimeConverter = Painter.GetTimeConverter();
 
-	const float StartTime = TimeConverter.PixelToTime(0.f);
-	const float EndTime = TimeConverter.PixelToTime(Painter.SectionGeometry.GetLocalSize().X);
+	const float StartTime       = TimeConverter.PixelToSeconds(0.f);
+	const float EndTime         = TimeConverter.PixelToSeconds(Painter.SectionGeometry.GetLocalSize().X);
 	const float SectionDuration = EndTime - StartTime;
 
 	FVector2D GradientSize = FVector2D( Painter.SectionGeometry.Size.X - 2.f, (Painter.SectionGeometry.Size.Y / 4) - 3.0f );
 	if ( GradientSize.X >= 1.f )
 	{
-
 		FPaintGeometry PaintGeometry = Painter.SectionGeometry.ToPaintGeometry( FVector2D( 1.f, 1.f ), GradientSize );
 
 		// If we are showing a background pattern and the colors is transparent, draw a checker pattern
@@ -66,22 +50,20 @@ int32 FColorPropertySection::OnPaintSection( FSequencerSectionPainter& Painter )
 			FEditorStyle::GetBrush( "Checker" ),
 			DrawEffects);
 
-		TArray<FSlateGradientStop> GradientStops;
+		TArray< TTuple<float, FLinearColor> > ColorKeys;
+		ConsolidateColorCurves( ColorKeys, ColorSection, TimeConverter );
 
-		TArray< TKeyValuePair<float, FLinearColor> > ColorKeys;
-		ConsolidateColorCurves( ColorKeys, ColorSection );
+		TArray<FSlateGradientStop> GradientStops;
 
-		for ( int32 i = 0; i < ColorKeys.Num(); ++i )
+		for (const TTuple<float, FLinearColor>& ColorStop : ColorKeys)
 		{
-			float Time = ColorKeys[i].Key;
-			FLinearColor Color = ColorKeys[i].Value;
+			const float Time = ColorStop.Get<0>();
 
 			// HACK: The color is converted to SRgb and then reinterpreted as linear here because gradients are converted to FColor
 			// without the SRgb conversion before being passed to the renderer for some reason.
-			Color = Color.ToFColor( true ).ReinterpretAsLinear();
+			const FLinearColor Color = ColorStop.Get<1>().ToFColor( true ).ReinterpretAsLinear();
 
 			float TimeFraction = (Time - StartTime) / SectionDuration;
-
 			GradientStops.Add( FSlateGradientStop( FVector2D( TimeFraction * Painter.SectionGeometry.Size.X, 0 ), Color ) );
 		}
 
@@ -102,144 +84,119 @@ int32 FColorPropertySection::OnPaintSection( FSequencerSectionPainter& Painter )
 }
 
 
-void FColorPropertySection::ConsolidateColorCurves( TArray< TKeyValuePair<float, FLinearColor> >& OutColorKeys, const UMovieSceneColorSection* Section ) const
+void FColorPropertySection::ConsolidateColorCurves( TArray< TTuple<float, FLinearColor> >& OutColorKeys, const UMovieSceneColorSection* InColorSection, const FTimeToPixel& TimeConverter ) const
 {
-	// Get the default color from the current property value, or use black if the current property
-	// value can't be found.
-	FLinearColor DefaultColor;
-	TOptional<FLinearColor> CurrentPropertyValue = GetPropertyValueAsLinearColor();
-	if (CurrentPropertyValue.IsSet())
-	{
-		DefaultColor = CurrentPropertyValue.GetValue();
-	}
-	else
-	{
-		DefaultColor = FLinearColor::Black;
-	}
+	FLinearColor DefaultColor = GetPropertyValueAsLinearColor();
 
-	// @todo Sequencer Optimize - This could all get cached, instead of recalculating everything every OnPaint
-
-	const FRichCurve* Curves[4] = {
-		&Section->GetRedCurve(),
-		&Section->GetGreenCurve(),
-		&Section->GetBlueCurve(),
-		&Section->GetAlphaCurve()
-	};
-
-	// @todo Sequencer Optimize - This is a O(n^2) loop!
-	// Our times are floats, which means we can't use a map and
-	// do a quick lookup to see if the keys already exist
-	// because the keys are ordered, we could take advantage of that, however
-	TArray<float> TimesWithKeys;
-	for ( int32 i = 0; i < 4; ++i )
+	UMovieSceneSection* Section = WeakSection.Get();
+	if (Section)
 	{
-		const FRichCurve* Curve = Curves[i];
-		for ( auto It( Curve->GetKeyIterator() ); It; ++It )
+		// @todo Sequencer Optimize - This could all get cached, instead of recalculating everything every OnPaint
+
+		TArrayView<FMovieSceneFloatChannel*> FloatChannels = Section->GetChannelProxy().GetChannels<FMovieSceneFloatChannel>();
+
+		// Gather all the channels with keys
+		TArray<TArrayView<const FFrameNumber>, TInlineAllocator<4>> ChannelTimes;
+		for (int32 Index = 0; Index < 4; ++Index)
 		{
-			float KeyTime = It->Time;
+			if (FloatChannels[Index]->GetTimes().Num())
+			{
+				ChannelTimes.Add(FloatChannels[Index]->GetTimes());
+			}
+		}
 
-			bool bShouldAddKey = true;
+		// Keep adding color stops for similar times until we have nothing left
+		while ( ChannelTimes.Num() )
+		{
+			// Find the earliest time from the remaining channels
+			FFrameNumber Time = TNumericLimits<int32>::Max();
+			for (const TArrayView<const FFrameNumber>& Channel : ChannelTimes)
+			{
+				Time = FMath::Min(Time, Channel[0]);
+			}
 
-			int32 InsertKeyIndex = INDEX_NONE;
-			for ( int32 k = 0; k < TimesWithKeys.Num(); ++k )
+			// Slice the channels until we no longer match the next time
+			for (TArrayView<const FFrameNumber>& Channel : ChannelTimes)
 			{
-				if ( FMath::IsNearlyEqual( TimesWithKeys[k], KeyTime ) )
+				int32 SliceIndex = 0;
+				while (SliceIndex < Channel.Num() && Time == Channel[SliceIndex])
 				{
-					bShouldAddKey = false;
-					break;
+					++SliceIndex;
 				}
-				else if ( TimesWithKeys[k] > KeyTime )
+
+				if (SliceIndex > 0)
 				{
-					InsertKeyIndex = k;
-					break;
+					int32 NewNum = Channel.Num() - SliceIndex;
+					Channel = NewNum > 0 ? Channel.Slice(SliceIndex, NewNum) : TArrayView<const FFrameNumber>();
 				}
 			}
 
-			if ( InsertKeyIndex == INDEX_NONE && bShouldAddKey )
+			// Remove empty channels with no keys left
+			for (int32 Index = ChannelTimes.Num()-1; Index >= 0; --Index)
 			{
-				InsertKeyIndex = TimesWithKeys.Num();
+				if (ChannelTimes[Index].Num() == 0)
+				{
+					ChannelTimes.RemoveAt(Index, 1, false);
+				}
 			}
 
-			if ( bShouldAddKey )
-			{
-				TimesWithKeys.Insert( KeyTime, InsertKeyIndex );
-			}
+			FLinearColor ColorAtTime = DefaultColor;
+			FloatChannels[0]->Evaluate(Time, ColorAtTime.R);
+			FloatChannels[1]->Evaluate(Time, ColorAtTime.G);
+			FloatChannels[2]->Evaluate(Time, ColorAtTime.B);
+			FloatChannels[3]->Evaluate(Time, ColorAtTime.A);
+
+			OutColorKeys.Add(MakeTuple(float(Time / TimeConverter.GetFrameResolution()), ColorAtTime));
 		}
 	}
 
 	// Enforce at least one key for the default value
-	if (TimesWithKeys.Num() == 0)
-	{
-		TimesWithKeys.Add(0);
-	}
-
-	// @todo Sequencer Optimize - This another O(n^2) loop, since Eval is O(n)!
-	for ( int32 i = 0; i < TimesWithKeys.Num(); ++i )
+	if (OutColorKeys.Num() == 0)
 	{
-		float Time = TimesWithKeys[i];
-
-		FLinearColor Color(
-			Section->GetRedCurve().Eval(Time, DefaultColor.R),
-			Section->GetGreenCurve().Eval(Time, DefaultColor.G),
-			Section->GetBlueCurve().Eval(Time, DefaultColor.B),
-			Section->GetAlphaCurve().Eval(Time, DefaultColor.A)
-			);
-		OutColorKeys.Add( TKeyValuePair<float, FLinearColor>( TimesWithKeys[i], Color ) );
+		OutColorKeys.Add(MakeTuple(0.f, DefaultColor));
 	}
 }
 
-
-TOptional<FLinearColor> FColorPropertySection::GetPropertyValueAsLinearColor() const
+FLinearColor FColorPropertySection::GetPropertyValueAsLinearColor() const
 {
-	UStructProperty* ColorStructProperty = Cast<UStructProperty>(GetProperty());
-	if (ColorStructProperty != nullptr)
+	UMovieSceneSection* Section = WeakSection.Get();
+	TSharedPtr<ISequencer> Sequencer = WeakSequencer.Pin();
+
+	FLinearColor LinearColor = FLinearColor::Black;
+
+	if (Section && Sequencer.IsValid())
 	{
-		if (ColorStructProperty->Struct->GetFName() == SlateColorName)
+		// Find the first object bound to this object binding ID, and apply each channel's external value to the color if possible
+		for (TWeakObjectPtr<> WeakObject : Sequencer->FindObjectsInCurrentSequence(ObjectBindingID))
 		{
-			TOptional<FSlateColor> SlateColor = GetPropertyValue<FSlateColor>();
-			if (SlateColor.IsSet())
+			if (UObject* Object = WeakObject.Get())
 			{
-				return TOptional<FLinearColor>(SlateColor.GetValue().GetSpecifiedColor());
-			}
-		}
+				// Access the editor data for the float channels which define how to extract the property value from the object
+				TArrayView<const TMovieSceneExternalValue<float>> ExternalValues = Section->GetChannelProxy().GetAllSpecializedEditorData<FMovieSceneFloatChannel>();
 
-		if (ColorStructProperty->Struct->GetFName() == NAME_LinearColor)
-		{
-			return GetPropertyValue<FLinearColor>();
-		}
+				FTrackInstancePropertyBindings* BindingsPtr = PropertyBindings.IsSet() ? &PropertyBindings.GetValue() : nullptr;
 
-		if (ColorStructProperty->Struct->GetFName() == NAME_Color)
-		{
-			TOptional<FColor> Color = GetPropertyValue<FColor>();
-			if (Color.IsSet())
-			{
-				return TOptional<FLinearColor>(Color.GetValue());
+				if (ExternalValues[0].OnGetExternalValue)
+				{
+					LinearColor.R = ExternalValues[0].OnGetExternalValue(*Object, BindingsPtr).Get(0.f);
+				}
+				if (ExternalValues[1].OnGetExternalValue)
+				{
+					LinearColor.G = ExternalValues[1].OnGetExternalValue(*Object, BindingsPtr).Get(0.f);
+				}
+				if (ExternalValues[2].OnGetExternalValue)
+				{
+					LinearColor.B = ExternalValues[2].OnGetExternalValue(*Object, BindingsPtr).Get(0.f);
+				}
+				if (ExternalValues[3].OnGetExternalValue)
+				{
+					LinearColor.A = ExternalValues[3].OnGetExternalValue(*Object, BindingsPtr).Get(0.f);
+				}
+
+				break;
 			}
 		}
 	}
-	return TOptional<FLinearColor>();
-}
-
-TOptional<float> FColorPropertySection::GetColorRedValue() const
-{
-	TOptional<FLinearColor> CurrentValue = GetPropertyValueAsLinearColor();
-	return CurrentValue.IsSet() ? CurrentValue.GetValue().R : TOptional<float>();
-}
 
-TOptional<float> FColorPropertySection::GetColorGreenValue() const
-{
-	TOptional<FLinearColor> CurrentValue = GetPropertyValueAsLinearColor();
-	return CurrentValue.IsSet() ? CurrentValue.GetValue().G : TOptional<float>();
-}
-
-TOptional<float> FColorPropertySection::GetColorBlueValue() const
-{
-	TOptional<FLinearColor> CurrentValue = GetPropertyValueAsLinearColor();
-	return CurrentValue.IsSet() ? CurrentValue.GetValue().B : TOptional<float>();
-}
-
-TOptional<float> FColorPropertySection::GetColorAlphaValue() const
-{
-	TOptional<FLinearColor> CurrentValue = GetPropertyValueAsLinearColor();
-	return CurrentValue.IsSet() ? CurrentValue.GetValue().A : TOptional<float>();
+	return LinearColor;
 }
diff --git a/Engine/Source/Editor/MovieSceneTools/Private/Sections/ColorPropertySection.h b/Engine/Source/Editor/MovieSceneTools/Private/Sections/ColorPropertySection.h
--- a/Engine/Source/Editor/MovieSceneTools/Private/Sections/ColorPropertySection.h
+++ b/Engine/Source/Editor/MovieSceneTools/Private/Sections/ColorPropertySection.h
@@ -2,10 +2,12 @@
 
 #pragma once
 
-#include "CoreMinimal.h"
-#include "Misc/Guid.h"
-#include "PropertySection.h"
+#include "ISequencerSection.h"
+#include "MovieSceneCommonHelpers.h"
 
+struct  FTimeToPixel;
+
+class ISequencer;
 class FSequencerSectionPainter;
 class UMovieSceneColorSection;
 
@@ -13,42 +15,37 @@ class UMovieSceneColorSection;
 * A color section implementation
 */
 class FColorPropertySection
-	: public FPropertySection
+	: public FSequencerSection
 {
 public:
 
 	/**
 	* Creates a new color property section.
 	*
-	* @param InSequencer The sequencer which is controlling this property section.
-	* @param InObjectBinding The object binding for the object which owns the property that this section is animating.
-	* @param InPropertyName The name of the property which is animated by this section.
-	* @param InPropertyPath A string representing the path to the property which is animated by this section.
 	* @param InSectionObject The section object which is being displayed and edited.
-	* @param InDisplayName A display name for the section being displayed and edited.
+	* @param InObjectBindingID The ID of the object this section is bound to
+	* @param InSequencer The sequencer this section is for
 	*/
-	FColorPropertySection(ISequencer* InSequencer, FGuid InObjectBinding, FName InPropertyName, const FString& InPropertyPath, UMovieSceneSection& InSectionObject, const FText& InDisplayName)
-		: FPropertySection(InSequencer, InObjectBinding, InPropertyName, InPropertyPath, InSectionObject, InDisplayName)
-	{
-	}
-
-public:
+	FColorPropertySection(UMovieSceneSection& InSectionObject, const FGuid& InObjectBindingID, TWeakPtr<ISequencer> InSequencer);
 
-	// FPropertySection interface
+private:
 
-	virtual void GenerateSectionLayout(class ISectionLayoutBuilder& LayoutBuilder) const override;
 	virtual int32 OnPaintSection(FSequencerSectionPainter& Painter) const override;
 
-protected:
-
 	/** Consolidate color curves for all track sections. */
-	void ConsolidateColorCurves(TArray< TKeyValuePair<float, FLinearColor> >& OutColorKeys, const UMovieSceneColorSection* Section) const;
-	
-	/** Gets current value of the property being edited as a linear color. */
-	TOptional<FLinearColor> GetPropertyValueAsLinearColor() const;
-
-	TOptional<float> GetColorRedValue() const;
-	TOptional<float> GetColorGreenValue() const;
-	TOptional<float> GetColorBlueValue() const;
-	TOptional<float> GetColorAlphaValue() const;
+	void ConsolidateColorCurves(TArray< TTuple<float, FLinearColor> >& OutColorKeys, const UMovieSceneColorSection* Section, const FTimeToPixel& TimeConverter) const;
+
+	/** Get the current value of the object's property as a linear color */
+	FLinearColor GetPropertyValueAsLinearColor() const;
+
+private:
+
+	/** The bound objects ID */
+	FGuid ObjectBindingID;
+
+	/** Weak pointer to the sequencer this section is for */
+	TWeakPtr<ISequencer> WeakSequencer;
+
+	/** Property bindings used for retrieving object properties */
+	mutable TOptional<FTrackInstancePropertyBindings> PropertyBindings;
 };
diff --git a/Engine/Source/Editor/MovieSceneTools/Private/Sections/LevelVisibilitySection.h b/Engine/Source/Editor/MovieSceneTools/Private/Sections/LevelVisibilitySection.h
--- a/Engine/Source/Editor/MovieSceneTools/Private/Sections/LevelVisibilitySection.h
+++ b/Engine/Source/Editor/MovieSceneTools/Private/Sections/LevelVisibilitySection.h
@@ -29,7 +29,6 @@ public:
 	virtual UMovieSceneSection* GetSectionObject() override;
 	virtual TSharedRef<SWidget> GenerateSectionWidget() override;
 	virtual int32 OnPaintSection( FSequencerSectionPainter& InPainter ) const override;
-	virtual void GenerateSectionLayout( ISectionLayoutBuilder& LayoutBuilder ) const override { }
 
 private:
 
diff --git a/Engine/Source/Editor/MovieSceneTools/Private/Sections/ParameterSection.cpp b/Engine/Source/Editor/MovieSceneTools/Private/Sections/ParameterSection.cpp
--- a/Engine/Source/Editor/MovieSceneTools/Private/Sections/ParameterSection.cpp
+++ b/Engine/Source/Editor/MovieSceneTools/Private/Sections/ParameterSection.cpp
@@ -4,115 +4,13 @@
 #include "ISectionLayoutBuilder.h"
 #include "ScopedTransaction.h"
 #include "Sections/MovieSceneParameterSection.h"
-#include "FloatCurveKeyArea.h"
-
 
 #define LOCTEXT_NAMESPACE "ParameterSection"
 
-DECLARE_DELEGATE(FLayoutParameterDelegate);
-
-
-struct FIndexAndLayoutParameterDelegate
-{
-	int32 Index;
-	FLayoutParameterDelegate LayoutParameterDelegate;
-
-	bool operator<(FIndexAndLayoutParameterDelegate const& Other) const
-	{
-		return Index < Other.Index;
-	}
-};
-
-
-void LayoutScalarParameter( ISectionLayoutBuilder* LayoutBuilder, FScalarParameterNameAndCurve* ScalarParameterNameAndCurve, UMovieSceneParameterSection* ParameterSection )
-{
-	LayoutBuilder->AddKeyArea(
-		ScalarParameterNameAndCurve->ParameterName,
-		FText::FromName( ScalarParameterNameAndCurve->ParameterName ),
-		MakeShareable( new FFloatCurveKeyArea( &ScalarParameterNameAndCurve->ParameterCurve, ParameterSection ) ) );
-}
-
-
-void LayoutVectorParameter( ISectionLayoutBuilder* LayoutBuilder, FVectorParameterNameAndCurves* VectorParameterNameAndCurves, UMovieSceneParameterSection* ParameterSection )
-{
-	LayoutBuilder->PushCategory( VectorParameterNameAndCurves->ParameterName, FText::FromName( VectorParameterNameAndCurves->ParameterName ) );
-	{
-		LayoutBuilder->AddKeyArea( "X", NSLOCTEXT( "ParameterSection", "XArea", "X" ),
-			MakeShareable( new FFloatCurveKeyArea( &VectorParameterNameAndCurves->XCurve, ParameterSection ) ) );
-		LayoutBuilder->AddKeyArea( "Y", NSLOCTEXT( "ParameterSection", "YArea", "Y" ), 
-			MakeShareable( new FFloatCurveKeyArea( &VectorParameterNameAndCurves->YCurve, ParameterSection ) ) );
-		LayoutBuilder->AddKeyArea( "Z", NSLOCTEXT( "ParameterSection", "ZArea", "Z" ), 
-			MakeShareable( new FFloatCurveKeyArea( &VectorParameterNameAndCurves->ZCurve, ParameterSection ) ) );
-	}
-	LayoutBuilder->PopCategory();
-}
-
-
-void LayoutColorParameter( ISectionLayoutBuilder* LayoutBuilder, FColorParameterNameAndCurves* ColorParameterNameAndCurves, UMovieSceneParameterSection* ParameterSection )
-{
-	LayoutBuilder->PushCategory( ColorParameterNameAndCurves->ParameterName, FText::FromName( ColorParameterNameAndCurves->ParameterName ) );
-	{
-		LayoutBuilder->AddKeyArea( "R", NSLOCTEXT( "ParameterSection", "RedArea", "Red" ),
-			MakeShareable( new FFloatCurveKeyArea( &ColorParameterNameAndCurves->RedCurve, ParameterSection ) ) );
-		LayoutBuilder->AddKeyArea( "G", NSLOCTEXT( "ParameterSection", "GreenArea", "Green" ),
-			MakeShareable( new FFloatCurveKeyArea( &ColorParameterNameAndCurves->GreenCurve, ParameterSection ) ) );
-		LayoutBuilder->AddKeyArea( "B", NSLOCTEXT( "ParameterSection", "BlueArea", "Blue" ),
-			MakeShareable( new FFloatCurveKeyArea( &ColorParameterNameAndCurves->BlueCurve, ParameterSection ) ) );
-		LayoutBuilder->AddKeyArea( "A", NSLOCTEXT( "ParameterSection", "OpacityArea", "Opacity" ),
-			MakeShareable( new FFloatCurveKeyArea( &ColorParameterNameAndCurves->AlphaCurve, ParameterSection ) ) );
-	}
-	LayoutBuilder->PopCategory();
-}
-
-void FParameterSection::GenerateSectionLayout( class ISectionLayoutBuilder& LayoutBuilder ) const
-{
-	// TODO - This seems a bit convoluted here to have to use delegates to get the layout to happen in the right order based on the indices.
-	UMovieSceneParameterSection* ParameterSection = Cast<UMovieSceneParameterSection>( &SectionObject );
-	TArray<FIndexAndLayoutParameterDelegate> IndexAndLayoutDelegates;
-
-	// Collect scalar parameters.
-	TArray<FScalarParameterNameAndCurve>* ScalarParameterNamesAndCurves = ParameterSection->GetScalarParameterNamesAndCurves();
-	for ( int32 i = 0; i < ScalarParameterNamesAndCurves->Num(); i++)
-	{
-		FIndexAndLayoutParameterDelegate IndexAndLayoutDelegate;
-		IndexAndLayoutDelegate.Index = (*ScalarParameterNamesAndCurves)[i].Index;
-		IndexAndLayoutDelegate.LayoutParameterDelegate.BindStatic(&LayoutScalarParameter, &LayoutBuilder, &(*ScalarParameterNamesAndCurves)[i], ParameterSection);
-		IndexAndLayoutDelegates.Add(IndexAndLayoutDelegate);
-	}
-
-	// Collect vector parameters
-	TArray<FVectorParameterNameAndCurves>* VectorParameterNamesAndCurves = ParameterSection->GetVectorParameterNamesAndCurves();
-	for ( int32 i = 0; i < VectorParameterNamesAndCurves->Num(); i++ )
-	{
-		FIndexAndLayoutParameterDelegate IndexAndLayoutDelegate;
-		IndexAndLayoutDelegate.Index = ( *VectorParameterNamesAndCurves )[i].Index;
-		IndexAndLayoutDelegate.LayoutParameterDelegate.BindStatic( &LayoutVectorParameter, &LayoutBuilder, &( *VectorParameterNamesAndCurves )[i], ParameterSection );
-		IndexAndLayoutDelegates.Add( IndexAndLayoutDelegate );
-	}
-
-	// Collect Color parameters
-	TArray<FColorParameterNameAndCurves>* ColorParameterNamesAndCurves = ParameterSection->GetColorParameterNamesAndCurves();
-	for ( int32 i = 0; i < ColorParameterNamesAndCurves->Num(); i++ )
-	{
-		FIndexAndLayoutParameterDelegate IndexAndLayoutDelegate;
-		IndexAndLayoutDelegate.Index = (*ColorParameterNamesAndCurves)[i].Index;
-		IndexAndLayoutDelegate.LayoutParameterDelegate.BindStatic( &LayoutColorParameter, &LayoutBuilder, &(*ColorParameterNamesAndCurves)[i], ParameterSection );
-		IndexAndLayoutDelegates.Add( IndexAndLayoutDelegate );
-	}
-
-	// Sort and perform layout
-	IndexAndLayoutDelegates.Sort();
-	for ( FIndexAndLayoutParameterDelegate& IndexAndLayoutDelegate : IndexAndLayoutDelegates )
-	{
-		IndexAndLayoutDelegate.LayoutParameterDelegate.Execute();
-	}
-}
-
-
 bool FParameterSection::RequestDeleteCategory( const TArray<FName>& CategoryNamePath )
 {
 	const FScopedTransaction Transaction( LOCTEXT( "DeleteVectorOrColorParameter", "Delete vector or color parameter" ) );
-	UMovieSceneParameterSection* ParameterSection = Cast<UMovieSceneParameterSection>( &SectionObject );
+	UMovieSceneParameterSection* ParameterSection = Cast<UMovieSceneParameterSection>( WeakSection.Get() );
 	if( ParameterSection->Modify() )
 	{
 		bool bVectorParameterDeleted = ParameterSection->RemoveVectorParameter( CategoryNamePath[0] );
@@ -129,7 +27,7 @@ bool FParameterSection::RequestDeleteKeyArea( const TArray<FName>& KeyAreaNamePa
 	if ( KeyAreaNamePath.Num() == 1)
 	{
 		const FScopedTransaction Transaction( LOCTEXT( "DeleteScalarParameter", "Delete scalar parameter" ) );
-		UMovieSceneParameterSection* ParameterSection = Cast<UMovieSceneParameterSection>( &SectionObject );
+		UMovieSceneParameterSection* ParameterSection = Cast<UMovieSceneParameterSection>( WeakSection.Get() );
 		if (ParameterSection->TryModify())
 		{
 			return ParameterSection->RemoveScalarParameter( KeyAreaNamePath[0] );
diff --git a/Engine/Source/Editor/MovieSceneTools/Private/Sections/ParameterSection.h b/Engine/Source/Editor/MovieSceneTools/Private/Sections/ParameterSection.h
--- a/Engine/Source/Editor/MovieSceneTools/Private/Sections/ParameterSection.h
+++ b/Engine/Source/Editor/MovieSceneTools/Private/Sections/ParameterSection.h
@@ -2,26 +2,24 @@
 
 #pragma once
 
-#include "CoreMinimal.h"
-#include "PropertySection.h"
+#include "ISequencerSection.h"
 
 /**
  * A movie scene section for material parameters.
  */
 class FParameterSection
-	: public FPropertySection
+	: public FSequencerSection
 {
 public:
 
-	FParameterSection(UMovieSceneSection& InSectionObject, const FText& SectionName)
-		: FPropertySection(InSectionObject, SectionName)
+	FParameterSection(UMovieSceneSection& InSectionObject)
+		: FSequencerSection(InSectionObject)
 	{ }
 
 public:
 
 	//~ ISequencerSection interface
 
-	virtual void GenerateSectionLayout(class ISectionLayoutBuilder& LayoutBuilder) const override;
 	virtual bool RequestDeleteCategory(const TArray<FName>& CategoryNamePath) override;
 	virtual bool RequestDeleteKeyArea(const TArray<FName>& KeyAreaNamePath) override;
 };
diff --git a/Engine/Source/Editor/MovieSceneTools/Private/Sections/ThumbnailSection.cpp b/Engine/Source/Editor/MovieSceneTools/Private/Sections/ThumbnailSection.cpp
--- a/Engine/Source/Editor/MovieSceneTools/Private/Sections/ThumbnailSection.cpp
+++ b/Engine/Source/Editor/MovieSceneTools/Private/Sections/ThumbnailSection.cpp
@@ -16,7 +16,8 @@
 #include "PropertyEditorModule.h"
 #include "IDetailsView.h"
 #include "IVREditorModule.h"
-
+#include "MovieScene.h"
+#include "MovieSceneTimeHelpers.h"
 
 #define LOCTEXT_NAMESPACE "FThumbnailSection"
 
@@ -108,7 +109,7 @@ void FThumbnailSection::BuildSectionContextMenu(FMenuBuilder& MenuBuilder, const
 
 				TSharedPtr<ISequencer> Sequencer = SequencerPtr.Pin();
 
-				FText CurrentTime = FText::FromString(Sequencer->GetZeroPadNumericTypeInterface()->ToString(Sequencer->GetLocalTime()));
+				FText CurrentTime = FText::FromString(Sequencer->GetNumericTypeInterface()->ToString(Sequencer->GetLocalTime().Time.GetFrame().Value));
 
 				InMenuBuilder.BeginSection(NAME_None, LOCTEXT("ThisSectionText", "This Section"));
 				{
@@ -124,7 +125,7 @@ void FThumbnailSection::BuildSectionContextMenu(FMenuBuilder& MenuBuilder, const
 						FSlateIcon(),
 						FUIAction(
 							FExecuteAction::CreateLambda([=]{
-								SetSingleTime(Sequencer->GetLocalTime());
+								SetSingleTime(Sequencer->GetLocalTime().AsSeconds());
 								GetMutableDefault<UMovieSceneUserThumbnailSettings>()->bDrawSingleThumbnails = true;
 								GetMutableDefault<UMovieSceneUserThumbnailSettings>()->SaveConfig();
 							})
@@ -208,11 +209,11 @@ int32 FThumbnailSection::OnPaintSection( FSequencerSectionPainter& InPainter ) c
 
 	// @todo Sequencer: Need a way to visualize the key here
 
-	const TRange<float> VisibleRange = GetVisibleRange();
-	const TRange<float> GenerationRange = GetTotalRange();
+	const TRange<double> VisibleRange = GetVisibleRange();
+	const TRange<double> GenerationRange = GetTotalRange();
+
+	const float TimePerPx = GenerationRange.Size<double>() / InPainter.SectionGeometry.GetLocalSize().X;
 
-	const float TimePerPx = GenerationRange.Size<float>() / InPainter.SectionGeometry.GetLocalSize().X;
-	
 	FSlateRect ThumbnailClipRect = SectionGeometry.GetLayoutBoundingRect().InsetBy(FMargin(SectionThumbnailPadding, 0.f)).IntersectionWith(InPainter.SectionClippingRect);
 
 	for (const TSharedPtr<FTrackEditorThumbnail>& Thumbnail : ThumbnailCache.GetThumbnails())
@@ -221,13 +222,13 @@ int32 FThumbnailSection::OnPaintSection( FSequencerSectionPainter& InPainter ) c
 		FIntPoint ThumbnailSize = Thumbnail->GetSize();
 		
 		// Calculate the paint geometry for this thumbnail
-		TOptional<float> SingleReferenceFrame = ThumbnailCache.GetSingleReferenceFrame();
+		TOptional<double> SingleReferenceFrame = ThumbnailCache.GetSingleReferenceFrame();
 
 		// Single thumbnails are always drawn at the start of the section, clamped to the visible range
 		// Thumbnail sequences draw relative to their actual position in the sequence/section
-		const int32 Offset = SingleReferenceFrame.IsSet() ?
-			FMath::Max((VisibleRange.GetLowerBoundValue() - GenerationRange.GetLowerBoundValue()) / TimePerPx, 0.f) + SectionThumbnailPadding :
-			(Thumbnail->GetTimeRange().GetLowerBoundValue() - GenerationRange.GetLowerBoundValue()) / TimePerPx;
+		const int32 Offset = SingleReferenceFrame.IsSet()
+			? FMath::Max(float(VisibleRange.GetLowerBoundValue() - GenerationRange.GetLowerBoundValue()) / TimePerPx, 0.f) + SectionThumbnailPadding
+			: (Thumbnail->GetTimeRange().GetLowerBoundValue() - GenerationRange.GetLowerBoundValue()) / TimePerPx;
 
 		FPaintGeometry PaintGeometry = SectionGeometry.ToPaintGeometry(
 			ThumbnailSize,
@@ -269,33 +270,42 @@ int32 FThumbnailSection::OnPaintSection( FSequencerSectionPainter& InPainter ) c
 	return LayerId + 2;
 }
 
-TRange<float> FThumbnailSection::GetVisibleRange() const
+TRange<double> FThumbnailSection::GetVisibleRange() const
 {
-	TRange<float> GlobalVisibleRange = SequencerPtr.Pin()->GetViewRange();
+	TRange<double> GlobalVisibleRange = SequencerPtr.Pin()->GetViewRange();
 	if (TimeSpace == ETimeSpace::Global)
 	{
 		return GlobalVisibleRange;
 	}
-	
-	return TRange<float>(
-		GlobalVisibleRange.GetLowerBoundValue() - Section->GetStartTime(),
-		GlobalVisibleRange.GetUpperBoundValue() - Section->GetStartTime()
+
+	const FFrameRate FrameResolution = Section->GetTypedOuter<UMovieScene>()->GetFrameResolution();
+
+	const double     StartOffset = Section->HasStartFrame() ? Section->GetInclusiveStartFrame() / FrameResolution : 0.0;
+	const double     EndOffset   = Section->HasEndFrame()   ? Section->GetExclusiveEndFrame()   / FrameResolution : 0.0;
+
+	return TRange<double>(
+		GlobalVisibleRange.GetLowerBoundValue() - StartOffset,
+		GlobalVisibleRange.GetUpperBoundValue() - EndOffset
 	);
 }
 
-TRange<float> FThumbnailSection::GetTotalRange() const
+TRange<double> FThumbnailSection::GetTotalRange() const
 {
-	if (Section->IsInfinite())
-	{
-		return GetVisibleRange();
-	}
-	else if (TimeSpace == ETimeSpace::Global)
+	TRange<FFrameNumber> SectionRange    = Section->GetRange();
+	FFrameRate           FrameResolution = Section->GetTypedOuter<UMovieScene>()->GetFrameResolution();
+
+	if (TimeSpace == ETimeSpace::Global)
 	{
-		return Section->GetRange();
+		return SectionRange / FrameResolution;
 	}
 	else
 	{
-		return TRange<float>(0.f, Section->GetRange().Size<float>());
+		const bool bHasDiscreteSize = SectionRange.GetLowerBound().IsClosed() && SectionRange.GetUpperBound().IsClosed();
+		TRangeBound<double> UpperBound = bHasDiscreteSize
+			? TRangeBound<double>::Exclusive(FFrameNumber(MovieScene::DiscreteSize(SectionRange)) / FrameResolution)
+			: TRangeBound<double>::Open();
+
+		return TRange<double>(0, UpperBound);
 	}
 }
 
@@ -333,7 +343,7 @@ void FViewportThumbnailSection::PreDraw(FTrackEditorThumbnail& Thumbnail, FLevel
 
 		SavedPlaybackStatus = Sequencer->GetPlaybackStatus();
 		Sequencer->SetPlaybackStatus(EMovieScenePlayerStatus::Jumping);
-		Sequencer->SetLocalTimeDirectly(Thumbnail.GetEvalPosition());
+		Sequencer->SetLocalTimeDirectly(Thumbnail.GetEvalPosition() * Sequencer->GetLocalTime().Rate );
 		Sequencer->ForceEvaluate();
 
 		ViewportClient.SetActorLock(Camera);
diff --git a/Engine/Source/Editor/MovieSceneTools/Private/Sections/TransformPropertySection.cpp b/Engine/Source/Editor/MovieSceneTools/Private/Sections/TransformPropertySection.cpp
--- a/Engine/Source/Editor/MovieSceneTools/Private/Sections/TransformPropertySection.cpp
+++ b/Engine/Source/Editor/MovieSceneTools/Private/Sections/TransformPropertySection.cpp
@@ -1,131 +1,18 @@
 // Copyright 1998-2018 Epic Games, Inc. All Rights Reserved.
 
 #include "Sections/TransformPropertySection.h"
-#include "FloatCurveKeyArea.h"
 #include "ISectionLayoutBuilder.h"
 #include "Sections/MovieScene3DTransformSection.h"
 #include "SequencerSectionPainter.h"
 #include "Framework/MultiBox/MultiBoxBuilder.h"
 #include "ScopedTransaction.h"
+#include "ISequencer.h"
 
 #define LOCTEXT_NAMESPACE "FTransformSection"
 
-void FTransformSection::AssignProperty(FName PropertyName, const FString& PropertyPath)
-{
-	PropertyBindings = FTrackInstancePropertyBindings(PropertyName, PropertyPath);
-}
-
-int32 FTransformSection::OnPaintSection(FSequencerSectionPainter& InPainter) const
-{
-	return InPainter.PaintSectionBackground();
-}
-
-void FTransformSection::GenerateSectionLayout(ISectionLayoutBuilder& LayoutBuilder) const
-{
-	static const FLinearColor BlueKeyAreaColor(0.0f, 0.0f, 0.7f, 0.5f);
-	static const FLinearColor GreenKeyAreaColor(0.0f, 0.7f, 0.0f, 0.5f);
-	static const FLinearColor RedKeyAreaColor(0.7f, 0.0f, 0.0f, 0.5f);
-
-	UMovieScene3DTransformSection* TransformSection = CastChecked<UMovieScene3DTransformSection>(Section.Get());
-	EMovieSceneTransformChannel Channels = TransformSection->GetMask().GetChannels();
-
-	// Helper function to tidy up TAttribute::Create syntax
-	typedef TOptional<float> (FTransformSection::*ValueGetter)(EAxis::Type) const;
-
-	auto MakeExternalValue = [this](ValueGetter Getter, EAxis::Type Axis) -> TAttribute<TOptional<float>>
-	{
-		return TAttribute<TOptional<float>>::Create(TAttribute<TOptional<float>>::FGetter::CreateRaw(this, Getter, Axis));
-	};
-
-	// This generates the tree structure for the transform section
-	if (EnumHasAnyFlags(Channels, EMovieSceneTransformChannel::Translation))
-	{
-		LayoutBuilder.PushCategory("Location", LOCTEXT("LocationArea", "Location"));
-
-		if (EnumHasAllFlags(Channels, EMovieSceneTransformChannel::TranslationX))
-		{
-			TSharedRef<FFloatCurveKeyArea> TranslationXKeyArea = MakeShared<FFloatCurveKeyArea>(&TransformSection->GetTranslationCurve(EAxis::X),
-				MakeExternalValue(&FTransformSection::GetTranslationValue, EAxis::X), TransformSection, RedKeyAreaColor);
-			LayoutBuilder.AddKeyArea("Location.X", LOCTEXT("LocXArea", "X"), TranslationXKeyArea );
-		}
-		if (EnumHasAllFlags(Channels, EMovieSceneTransformChannel::TranslationY))
-		{
-			TSharedRef<FFloatCurveKeyArea> TranslationYKeyArea = MakeShared<FFloatCurveKeyArea>(&TransformSection->GetTranslationCurve(EAxis::Y),
-				MakeExternalValue(&FTransformSection::GetTranslationValue, EAxis::Y), TransformSection, GreenKeyAreaColor);
-			LayoutBuilder.AddKeyArea("Location.Y", LOCTEXT("LocYArea", "Y"), TranslationYKeyArea );
-		}
-		if (EnumHasAllFlags(Channels, EMovieSceneTransformChannel::TranslationZ))
-		{
-			TSharedRef<FFloatCurveKeyArea> TranslationZKeyArea = MakeShared<FFloatCurveKeyArea>(&TransformSection->GetTranslationCurve(EAxis::Z),
-				MakeExternalValue(&FTransformSection::GetTranslationValue, EAxis::Z), TransformSection, BlueKeyAreaColor);
-			LayoutBuilder.AddKeyArea("Location.Z", LOCTEXT("LocZArea", "Z"), TranslationZKeyArea );
-		}
-
-		LayoutBuilder.PopCategory();
-	}
-
-	if (EnumHasAnyFlags(Channels, EMovieSceneTransformChannel::Rotation))
-	{
-		LayoutBuilder.PushCategory( "Rotation", LOCTEXT("RotationArea", "Rotation") );
-
-		if (EnumHasAllFlags(Channels, EMovieSceneTransformChannel::RotationX))
-		{
-			TSharedRef<FFloatCurveKeyArea> RotationXKeyArea = MakeShared<FFloatCurveKeyArea>(&TransformSection->GetRotationCurve(EAxis::X),
-				MakeExternalValue(&FTransformSection::GetRotationValue, EAxis::X), TransformSection, RedKeyAreaColor);
-			LayoutBuilder.AddKeyArea("Rotation.X", LOCTEXT("RotXArea", "X"), RotationXKeyArea );
-		}
-		if (EnumHasAllFlags(Channels, EMovieSceneTransformChannel::RotationY))
-		{
-			TSharedRef<FFloatCurveKeyArea> RotationYKeyArea = MakeShared<FFloatCurveKeyArea>(&TransformSection->GetRotationCurve(EAxis::Y),
-				MakeExternalValue(&FTransformSection::GetRotationValue, EAxis::Y), TransformSection, GreenKeyAreaColor);
-			LayoutBuilder.AddKeyArea("Rotation.Y", LOCTEXT("RotYArea", "Y"), RotationYKeyArea );
-		}
-		if (EnumHasAllFlags(Channels, EMovieSceneTransformChannel::RotationZ))
-		{
-			TSharedRef<FFloatCurveKeyArea> RotationZKeyArea = MakeShared<FFloatCurveKeyArea>(&TransformSection->GetRotationCurve(EAxis::Z),
-				MakeExternalValue(&FTransformSection::GetRotationValue, EAxis::Z), TransformSection, BlueKeyAreaColor);
-			LayoutBuilder.AddKeyArea("Rotation.Z", LOCTEXT("RotZArea", "Z"), RotationZKeyArea );
-		}
-
-		LayoutBuilder.PopCategory();
-	}
-
-	if (EnumHasAnyFlags(Channels, EMovieSceneTransformChannel::Scale))
-	{
-		LayoutBuilder.PushCategory( "Scale", LOCTEXT("ScaleArea", "Scale") );
-
-		if (EnumHasAllFlags(Channels, EMovieSceneTransformChannel::ScaleX))
-		{
-			TSharedRef<FFloatCurveKeyArea> ScaleXKeyArea = MakeShared<FFloatCurveKeyArea>(&TransformSection->GetScaleCurve(EAxis::X),
-				MakeExternalValue(&FTransformSection::GetScaleValue, EAxis::X), TransformSection, RedKeyAreaColor);
-			LayoutBuilder.AddKeyArea("Scale.X", LOCTEXT("ScaleXArea", "X"), ScaleXKeyArea );
-		}
-		if (EnumHasAllFlags(Channels, EMovieSceneTransformChannel::ScaleY))
-		{
-			TSharedRef<FFloatCurveKeyArea> ScaleYKeyArea = MakeShared<FFloatCurveKeyArea>(&TransformSection->GetScaleCurve(EAxis::Y),
-				MakeExternalValue(&FTransformSection::GetScaleValue, EAxis::Y), TransformSection, GreenKeyAreaColor);
-			LayoutBuilder.AddKeyArea("Scale.Y", LOCTEXT("ScaleYArea", "Y"), ScaleYKeyArea );
-		}
-		if (EnumHasAllFlags(Channels, EMovieSceneTransformChannel::ScaleZ))
-		{
-			TSharedRef<FFloatCurveKeyArea> ScaleZKeyArea = MakeShared<FFloatCurveKeyArea>(&TransformSection->GetScaleCurve(EAxis::Z),
-				MakeExternalValue(&FTransformSection::GetScaleValue, EAxis::Z), TransformSection, BlueKeyAreaColor);
-			LayoutBuilder.AddKeyArea("Scale.Z", LOCTEXT("ScaleZArea", "Z"), ScaleZKeyArea );
-		}
-
-		LayoutBuilder.PopCategory();
-	}
-
-	if (EnumHasAnyFlags(Channels, EMovieSceneTransformChannel::Weight))
-	{
-		TSharedRef<FFloatCurveKeyArea> WeightKeyArea = MakeShared<FFloatCurveKeyArea>(&TransformSection->GetManualWeightCurve(), TransformSection);
-		LayoutBuilder.AddKeyArea("Weight", LOCTEXT("WeightArea", "Weight"), WeightKeyArea );
-	}
-}
-
 void FTransformSection::BuildSectionContextMenu(FMenuBuilder& MenuBuilder, const FGuid& InObjectBinding)
 {
-	UMovieScene3DTransformSection* TransformSection = CastChecked<UMovieScene3DTransformSection>(Section.Get());
+	UMovieScene3DTransformSection* TransformSection = CastChecked<UMovieScene3DTransformSection>(WeakSection.Get());
 	TSharedPtr<ISequencer> SequencerPtr = WeakSequencer.Pin();
 
 	auto MakeUIAction = [=](EMovieSceneTransformChannel ChannelsToToggle)
@@ -235,79 +122,4 @@ void FTransformSection::BuildSectionContextMenu(FMenuBuilder& MenuBuilder, const
 	MenuBuilder.EndSection();
 }
 
-TOptional<FTransform> FTransformSection::GetCurrentValue() const
-{
-	TSharedPtr<ISequencer> Sequencer = WeakSequencer.Pin();
-	if (ensure(PropertyBindings.IsSet()) && Sequencer.IsValid())
-	{
-		for (TWeakObjectPtr<> WeakObject : Sequencer->FindBoundObjects(ObjectBinding, Sequencer->GetFocusedTemplateID()))
-		{
-			if (UObject* Object = WeakObject.Get())
-			{
-				return PropertyBindings->GetCurrentValue<FTransform>(*Object);
-			}
-		}
-	}
-	return TOptional<FTransform>();
-}
-
-TOptional<FRotator> FTransformSection::GetCurrentRotator() const
-{
-	TSharedPtr<ISequencer> Sequencer = WeakSequencer.Pin();
-	if (ensure(PropertyBindings.IsSet()) && Sequencer.IsValid())
-	{
-		for (TWeakObjectPtr<> WeakObject : Sequencer->FindBoundObjects(ObjectBinding, Sequencer->GetFocusedTemplateID()))
-		{
-			if (UObject* Object = WeakObject.Get())
-			{
-				return PropertyBindings->GetCurrentValue<FTransform>(*Object).GetRotation().Rotator();
-			}
-		}
-	}
-	return TOptional<FRotator>();
-}
-
-TOptional<float> FTransformSection::GetTranslationValue(EAxis::Type Axis) const
-{
-	FTransform Transform = GetCurrentValue().Get(FTransform::Identity);
-	switch (Axis)
-	{
-	case EAxis::X:		return Transform.GetTranslation().X;
-	case EAxis::Y:		return Transform.GetTranslation().Y;
-	case EAxis::Z:		return Transform.GetTranslation().Z;
-	}
-	return TOptional<float>();
-}
-
-TOptional<float> FTransformSection::GetRotationValue(EAxis::Type Axis) const
-{
-	TOptional<FRotator> Rotator = GetCurrentRotator();
-
-	if (!Rotator.IsSet())
-	{
-		FTransform Transform = GetCurrentValue().Get(FTransform::Identity);
-		Rotator = Transform.GetRotation().Rotator();
-	}
-
-	switch (Axis)
-	{
-	case EAxis::X:		return Rotator.GetValue().Roll;
-	case EAxis::Y:		return Rotator.GetValue().Pitch;
-	case EAxis::Z:		return Rotator.GetValue().Yaw;
-	}
-	return TOptional<float>();
-}
-
-TOptional<float> FTransformSection::GetScaleValue(EAxis::Type Axis) const
-{
-	FTransform Transform = GetCurrentValue().Get(FTransform::Identity);
-	switch (Axis)
-	{
-	case EAxis::X:		return Transform.GetScale3D().X;
-	case EAxis::Y:		return Transform.GetScale3D().Y;
-	case EAxis::Z:		return Transform.GetScale3D().Z;
-	}
-	return TOptional<float>();
-}
-
-#undef LOCTEXT_NAMESPACE
+#undef LOCTEXT_NAMESPACE
\ No newline at end of file
diff --git a/Engine/Source/Editor/MovieSceneTools/Private/Sections/TransformPropertySection.h b/Engine/Source/Editor/MovieSceneTools/Private/Sections/TransformPropertySection.h
--- a/Engine/Source/Editor/MovieSceneTools/Private/Sections/TransformPropertySection.h
+++ b/Engine/Source/Editor/MovieSceneTools/Private/Sections/TransformPropertySection.h
@@ -3,13 +3,16 @@
 #pragma once
 
 #include "CoreMinimal.h"
-#include "Misc/Guid.h"
-#include "PropertySection.h"
+#include "ISequencerSection.h"
+
+struct FGuid;
+class ISequencer;
+class UMovieSceneSection;
 
 /**
  * An implementation of vector property sections.
  */
-class FTransformSection : public ISequencerSection
+class FTransformSection : public FSequencerSection
 {
 public:
 
@@ -18,45 +21,18 @@ public:
 	*
 	* @param InSection The section object which is being displayed and edited.
 	* @param InSequencer The sequencer which is controlling this property section.
-	* @param InObjectBinding The object binding for the object which owns the property that this section is animating.
 	*/
-	FTransformSection(TWeakObjectPtr<UMovieSceneSection> InSection, TWeakPtr<ISequencer> InSequencer, FGuid InObjectBinding)
-		: Section(InSection), WeakSequencer(InSequencer), ObjectBinding(InObjectBinding)
+	FTransformSection(UMovieSceneSection& InSection, TWeakPtr<ISequencer> InSequencer)
+		: FSequencerSection(InSection), WeakSequencer(InSequencer)
 	{
 	}
 
-	/** Assign a property to this section */
-	void AssignProperty(FName PropertyName, const FString& PropertyPath);
-
 public:
 
-	virtual UMovieSceneSection* GetSectionObject() override { return Section.Get(); }
-	virtual int32 OnPaintSection( FSequencerSectionPainter& InPainter ) const override;
-	virtual void GenerateSectionLayout(ISectionLayoutBuilder& LayoutBuilder) const override;
 	virtual void BuildSectionContextMenu(FMenuBuilder& MenuBuilder, const FGuid& InObjectBinding) override;
 
 protected:
 
-	virtual TOptional<FTransform> GetCurrentValue() const;
-	virtual TOptional<FRotator> GetCurrentRotator() const;
-
-protected:
-
-	TOptional<float> GetTranslationValue(EAxis::Type Axis) const;
-	TOptional<float> GetRotationValue(EAxis::Type Axis) const;
-	TOptional<float> GetScaleValue(EAxis::Type Axis) const;
-
-protected:
-
-	/** The section that we are editing */
-	TWeakObjectPtr<UMovieSceneSection> Section;
-
 	/** The sequencer which is controlling this section. */
 	TWeakPtr<ISequencer> WeakSequencer;
-
-	/** An object binding for the object which owns the property being animated by this section. */
-	FGuid ObjectBinding;
-
-	/** An object which is used to retrieve the value of a property based on it's name and path. */
-	mutable TOptional<FTrackInstancePropertyBindings> PropertyBindings;
 };
diff --git a/Engine/Source/Editor/MovieSceneTools/Private/TrackEditorThumbnail/TrackEditorThumbnail.cpp b/Engine/Source/Editor/MovieSceneTools/Private/TrackEditorThumbnail/TrackEditorThumbnail.cpp
--- a/Engine/Source/Editor/MovieSceneTools/Private/TrackEditorThumbnail/TrackEditorThumbnail.cpp
+++ b/Engine/Source/Editor/MovieSceneTools/Private/TrackEditorThumbnail/TrackEditorThumbnail.cpp
@@ -52,7 +52,7 @@ public:
 /* FTrackEditorThumbnail structors
  *****************************************************************************/
 
-FTrackEditorThumbnail::FTrackEditorThumbnail(const FOnThumbnailDraw& InOnDraw, const FIntPoint& InSize, TRange<float> InTimeRange, float InPosition)
+FTrackEditorThumbnail::FTrackEditorThumbnail(const FOnThumbnailDraw& InOnDraw, const FIntPoint& InSize, TRange<double> InTimeRange, double InPosition)
 	: OnDraw(InOnDraw)
 	, Size(InSize)
 	, Texture(nullptr)
@@ -299,12 +299,12 @@ FTrackEditorThumbnailCache::~FTrackEditorThumbnailCache()
 }
 
 
-void FTrackEditorThumbnailCache::SetSingleReferenceFrame(TOptional<float> InReferenceFrame)
+void FTrackEditorThumbnailCache::SetSingleReferenceFrame(TOptional<double> InReferenceFrame)
 {
 	CurrentCache.SingleReferenceFrame = InReferenceFrame;
 }
 
-void FTrackEditorThumbnailCache::Update(const TRange<float>& NewRange, const TRange<float>& VisibleRange, const FIntPoint& AllottedSize, const FIntPoint& InDesiredSize, EThumbnailQuality Quality, double InCurrentTime)
+void FTrackEditorThumbnailCache::Update(const TRange<double>& NewRange, const TRange<double>& VisibleRange, const FIntPoint& AllottedSize, const FIntPoint& InDesiredSize, EThumbnailQuality Quality, double InCurrentTime)
 {
 	PreviousCache.TimeRange = CurrentCache.TimeRange;
 	PreviousCache.VisibleRange = CurrentCache.VisibleRange;
@@ -495,14 +495,14 @@ void FTrackEditorThumbnailCache::UpdateSingleThumbnail()
 {
 	Thumbnails.Reset();
 
-	const float TimePerPx = CurrentCache.TimeRange.Size<float>() / CurrentCache.AllottedSize.X;
-	const float HalfRange = CurrentCache.DesiredSize.X*TimePerPx*.5f;
-	const float EvalPosition = CurrentCache.SingleReferenceFrame.GetValue();
+	const double TimePerPx    = CurrentCache.TimeRange.Size<double>() / CurrentCache.AllottedSize.X;
+	const double HalfRange    = CurrentCache.DesiredSize.X*TimePerPx*.5;
+	const double EvalPosition = CurrentCache.SingleReferenceFrame.GetValue();
 
 	TSharedPtr<FTrackEditorThumbnail> NewThumbnail = MakeShareable(new FTrackEditorThumbnail(
 		FOnThumbnailDraw::CreateRaw(this, &FTrackEditorThumbnailCache::DrawViewportThumbnail),
 		CurrentCache.DesiredSize,
-		TRange<float>(EvalPosition - HalfRange, EvalPosition + HalfRange),
+		TRange<double>(EvalPosition - HalfRange, EvalPosition + HalfRange),
 		EvalPosition
 	));
 
@@ -558,8 +558,8 @@ void FTrackEditorThumbnailCache::UpdateFilledThumbnails()
 		}
 	}
 
-	TRange<float> MaxRange(CurrentCache.VisibleRange.GetLowerBoundValue() - CurrentCache.VisibleRange.Size<float>(), CurrentCache.VisibleRange.GetUpperBoundValue() + CurrentCache.VisibleRange.Size<float>());
-	TRange<float> Boundary = TRange<float>::Intersection(CurrentCache.TimeRange, MaxRange);
+	TRange<double> MaxRange(CurrentCache.VisibleRange.GetLowerBoundValue() - CurrentCache.VisibleRange.Size<double>(), CurrentCache.VisibleRange.GetUpperBoundValue() + CurrentCache.VisibleRange.Size<double>());
+	TRange<double> Boundary = TRange<double>::Intersection(CurrentCache.TimeRange, MaxRange);
 
 	if (!Boundary.IsEmpty())
 	{
@@ -577,26 +577,26 @@ void FTrackEditorThumbnailCache::UpdateFilledThumbnails()
 }
 
 
-void FTrackEditorThumbnailCache::GenerateFront(const TRange<float>& Boundary)
+void FTrackEditorThumbnailCache::GenerateFront(const TRange<double>& Boundary)
 {
 	if (!Thumbnails.Num())
 	{
 		return;
 	}
 
-	const float TimePerPx = CurrentCache.TimeRange.Size<float>() / CurrentCache.AllottedSize.X;
-	float EndTime = Thumbnails[0]->GetTimeRange().GetLowerBoundValue();
+	const double TimePerPx = CurrentCache.TimeRange.Size<double>() / CurrentCache.AllottedSize.X;
+	double EndTime = Thumbnails[0]->GetTimeRange().GetLowerBoundValue();
 
 	while (EndTime > Boundary.GetLowerBoundValue())
 	{
 		FIntPoint TextureSize = CurrentCache.DesiredSize;
 
-		TRange<float> TimeRange(EndTime - TextureSize.X * TimePerPx, EndTime);
+		TRange<double> TimeRange(EndTime - TextureSize.X * TimePerPx, EndTime);
 
 		// Evaluate the thumbnail along the length of its duration, based on its position in the sequence
-		const float FrameLength = TimeRange.Size<float>();
-		float TotalLerp = (TimeRange.GetLowerBoundValue() - CurrentCache.TimeRange.GetLowerBoundValue()) / (CurrentCache.TimeRange.Size<float>() - FrameLength);
-		float EvalPosition = CurrentCache.TimeRange.GetLowerBoundValue() + FMath::Clamp(TotalLerp, 0.f, .99f)*CurrentCache.TimeRange.Size<float>();
+		const double FrameLength = TimeRange.Size<double>();
+		double TotalLerp = (TimeRange.GetLowerBoundValue() - CurrentCache.TimeRange.GetLowerBoundValue()) / (CurrentCache.TimeRange.Size<double>() - FrameLength);
+		double EvalPosition = CurrentCache.TimeRange.GetLowerBoundValue() + FMath::Clamp(TotalLerp, 0.0, .99)*CurrentCache.TimeRange.Size<double>();
 
 		TSharedPtr<FTrackEditorThumbnail> NewThumbnail = MakeShareable(new FTrackEditorThumbnail(
 			FOnThumbnailDraw::CreateRaw(this, &FTrackEditorThumbnailCache::DrawViewportThumbnail),
@@ -613,10 +613,10 @@ void FTrackEditorThumbnailCache::GenerateFront(const TRange<float>& Boundary)
 }
 
 
-void FTrackEditorThumbnailCache::GenerateBack(const TRange<float>& Boundary)
+void FTrackEditorThumbnailCache::GenerateBack(const TRange<double>& Boundary)
 {
-	const float TimePerPx = CurrentCache.TimeRange.Size<float>() / CurrentCache.AllottedSize.X;
-	float StartTime = Thumbnails.Num() ? Thumbnails.Last()->GetTimeRange().GetUpperBoundValue() : Boundary.GetLowerBoundValue();
+	const double TimePerPx = CurrentCache.TimeRange.Size<double>() / CurrentCache.AllottedSize.X;
+	double StartTime = Thumbnails.Num() ? Thumbnails.Last()->GetTimeRange().GetUpperBoundValue() : Boundary.GetLowerBoundValue();
 
 	while (StartTime < Boundary.GetUpperBoundValue())
 	{
@@ -624,18 +624,18 @@ void FTrackEditorThumbnailCache::GenerateBack(const TRange<float>& Boundary)
 
 		{
 			// Move the thumbnail to the center of the space if we're the only thumbnail, and we don't fit on
-			float Overflow = TextureSize.X*TimePerPx - CurrentCache.TimeRange.Size<float>();
+			double Overflow = TextureSize.X*TimePerPx - CurrentCache.TimeRange.Size<double>();
 			if (Thumbnails.Num() == 0 && Overflow > 0)
 			{
 				StartTime -= Overflow*.5f;
 			}
 		}
 
-		TRange<float> TimeRange(StartTime, StartTime + TextureSize.X * TimePerPx);
+		TRange<double> TimeRange(StartTime, StartTime + TextureSize.X * TimePerPx);
 
-		const float FrameLength = TimeRange.Size<float>();
-		float TotalLerp = (TimeRange.GetLowerBoundValue() - CurrentCache.TimeRange.GetLowerBoundValue()) / (CurrentCache.TimeRange.Size<float>() - FrameLength);
-		float EvalPosition = CurrentCache.TimeRange.GetLowerBoundValue() + FMath::Clamp(TotalLerp, 0.f, .99f)*CurrentCache.TimeRange.Size<float>();
+		const double FrameLength = TimeRange.Size<double>();
+		double TotalLerp = (TimeRange.GetLowerBoundValue() - CurrentCache.TimeRange.GetLowerBoundValue()) / (CurrentCache.TimeRange.Size<double>() - FrameLength);
+		double EvalPosition = CurrentCache.TimeRange.GetLowerBoundValue() + FMath::Clamp(TotalLerp, 0.0, .99)*CurrentCache.TimeRange.Size<double>();
 
 		TSharedPtr<FTrackEditorThumbnail> NewThumbnail = MakeShareable(new FTrackEditorThumbnail(
 			FOnThumbnailDraw::CreateRaw(this, &FTrackEditorThumbnailCache::DrawViewportThumbnail),
diff --git a/Engine/Source/Editor/MovieSceneTools/Private/TrackEditors/AttachTrackEditor.cpp b/Engine/Source/Editor/MovieSceneTools/Private/TrackEditors/AttachTrackEditor.cpp
--- a/Engine/Source/Editor/MovieSceneTools/Private/TrackEditors/AttachTrackEditor.cpp
+++ b/Engine/Source/Editor/MovieSceneTools/Private/TrackEditors/AttachTrackEditor.cpp
@@ -62,8 +62,6 @@ public:
 		return FText::GetEmpty(); 
 	}
 
-	virtual void GenerateSectionLayout( class ISectionLayoutBuilder& LayoutBuilder ) const override {}
-
 	virtual int32 OnPaintSection( FSequencerSectionPainter& InPainter ) const override 
 	{
 		return InPainter.PaintSectionBackground();
@@ -209,7 +207,7 @@ void F3DAttachTrackEditor::ActorSocketPicked(const FName SocketName, USceneCompo
 	}
 }
 
-FKeyPropertyResult F3DAttachTrackEditor::AddKeyInternal( float KeyTime, const TArray<TWeakObjectPtr<UObject>> Objects, const FName SocketName, const FName ComponentName, FActorPickerID ActorPickerID)
+FKeyPropertyResult F3DAttachTrackEditor::AddKeyInternal( FFrameNumber KeyTime, const TArray<TWeakObjectPtr<UObject>> Objects, const FName SocketName, const FName ComponentName, FActorPickerID ActorPickerID)
 {
 	FKeyPropertyResult KeyPropertyResult;
 
@@ -248,12 +246,11 @@ FKeyPropertyResult F3DAttachTrackEditor::AddKeyInternal( float KeyTime, const TA
 			if (ensure(Track))
 			{
 				// Clamp to next attach section's start time or the end of the current sequencer view range
-				float AttachEndTime = GetSequencer()->GetViewRange().GetUpperBoundValue();
+				FFrameNumber AttachEndTime = (GetSequencer()->GetViewRange().GetUpperBoundValue() * Track->GetTypedOuter<UMovieScene>()->GetFrameResolution()).FrameNumber;
 
-				for (int32 AttachSectionIndex = 0; AttachSectionIndex < Track->GetAllSections().Num(); ++AttachSectionIndex)
+				for (UMovieSceneSection* Section : Track->GetAllSections())
 				{
-					float StartTime = Track->GetAllSections()[AttachSectionIndex]->GetStartTime();
-					float EndTime = Track->GetAllSections()[AttachSectionIndex]->GetEndTime();
+					FFrameNumber StartTime = Section->HasStartFrame() ? Section->GetInclusiveStartFrame() : 0;
 					if (KeyTime < StartTime)
 					{
 						if (AttachEndTime > StartTime)
@@ -263,7 +260,7 @@ FKeyPropertyResult F3DAttachTrackEditor::AddKeyInternal( float KeyTime, const TA
 					}
 				}
 
-				Cast<UMovieScene3DAttachTrack>(Track)->AddConstraint( KeyTime, AttachEndTime, SocketName, ComponentName, ConstraintBindingID);
+				Cast<UMovieScene3DAttachTrack>(Track)->AddConstraint( KeyTime, (AttachEndTime - KeyTime).Value, SocketName, ComponentName, ConstraintBindingID);
 				KeyPropertyResult.bTrackModified = true;
 			}
 		}
diff --git a/Engine/Source/Editor/MovieSceneTools/Private/TrackEditors/AttachTrackEditor.h b/Engine/Source/Editor/MovieSceneTools/Private/TrackEditors/AttachTrackEditor.h
--- a/Engine/Source/Editor/MovieSceneTools/Private/TrackEditors/AttachTrackEditor.h
+++ b/Engine/Source/Editor/MovieSceneTools/Private/TrackEditors/AttachTrackEditor.h
@@ -56,5 +56,5 @@ public:
 private:
 
 	/** Delegate for AnimatablePropertyChanged in AddKey */
-	FKeyPropertyResult AddKeyInternal(float KeyTime, const TArray<TWeakObjectPtr<UObject>> Objects, const FName SocketName, const FName ComponentName, FActorPickerID ActorPickerID);
+	FKeyPropertyResult AddKeyInternal(FFrameNumber KeyTime, const TArray<TWeakObjectPtr<UObject>> Objects, const FName SocketName, const FName ComponentName, FActorPickerID ActorPickerID);
 };
diff --git a/Engine/Source/Editor/MovieSceneTools/Private/TrackEditors/AudioTrackEditor.cpp b/Engine/Source/Editor/MovieSceneTools/Private/TrackEditors/AudioTrackEditor.cpp
--- a/Engine/Source/Editor/MovieSceneTools/Private/TrackEditors/AudioTrackEditor.cpp
+++ b/Engine/Source/Editor/MovieSceneTools/Private/TrackEditors/AudioTrackEditor.cpp
@@ -31,9 +31,11 @@
 #include "AssetRegistryModule.h"
 #include "MatineeImportTools.h"
 #include "Matinee/InterpTrackSound.h"
-#include "FloatCurveKeyArea.h"
 #include "ISectionLayoutBuilder.h"
 
+#include "DragAndDrop/AssetDragDropOp.h"
+#include "QualifiedFrameTime.h"
+
 
 #define LOCTEXT_NAMESPACE "FAudioTrackEditor"
 
@@ -297,6 +299,8 @@ void FAudioThumbnail::GenerateWaveformPreview(TArray<uint8>& OutData, TRange<flo
 		FAudioDevice* AudioDevice = GEngine->GetMainAudioDevice();
 		if (AudioDevice)
 		{
+			AudioDevice->StopAllSounds(true);
+
 			EDecompressionType DecompressionType =  SoundWave->DecompressionType;
 			SoundWave->DecompressionType = DTYPE_Native;
 
@@ -315,22 +319,26 @@ void FAudioThumbnail::GenerateWaveformPreview(TArray<uint8>& OutData, TRange<flo
 	const int32 LookupDataSize = SoundWave->RawPCMDataSize;
 	const int32 LookupSize = LookupDataSize * sizeof(uint8) / sizeof(int16);
 
-	if (!LookupData)
+	if (!LookupData || !AudioSection->HasStartFrame() || !AudioSection->HasEndFrame())
 	{
 		return;
 	}
 
+	FFrameRate FrameRate            = AudioSection->GetTypedOuter<UMovieScene>()->GetFrameResolution();
+	float      PitchMultiplierValue = AudioSection->GetPitchMultiplierChannel().GetDefault().Get(1.f);
+	double     SectionStartTime     = AudioSection->GetInclusiveStartFrame() / FrameRate;
+
 	// @todo Sequencer This fixes looping drawing by pretending we are only dealing with a SoundWave
-	TRange<float> AudioTrueRange = TRange<float>(
-		AudioSection->GetStartTime() - AudioSection->GetStartOffset(),
-		AudioSection->GetStartTime() - AudioSection->GetStartOffset() + DeriveUnloopedDuration(AudioSection) * AudioSection->GetPitchMultiplierCurve().GetDefaultValue());
-	float TrueRangeSize = AudioTrueRange.Size<float>();
+	TRange<double> AudioTrueRange = TRange<double>(
+		SectionStartTime - AudioSection->GetStartOffset(),
+		SectionStartTime - AudioSection->GetStartOffset() + DeriveUnloopedDuration(AudioSection) * PitchMultiplierValue);
 
+	float TrueRangeSize = AudioTrueRange.Size<float>();
 	float DrawRangeSize = DrawRange.Size<float>();
 
 	const int32 MaxAmplitude = NumChannels == 1 ? GetSize().Y : GetSize().Y / 2;
 
-	int32 DrawOffsetPx = FMath::Max(FMath::RoundToInt((DrawRange.GetLowerBoundValue() - AudioSection->GetRange().GetLowerBoundValue()) / DisplayScale), 0);
+	int32 DrawOffsetPx = FMath::Max(FMath::RoundToInt((DrawRange.GetLowerBoundValue() - SectionStartTime) / DisplayScale), 0);
 
 
 	// In order to prevent flickering waveforms when moving the display position/range around, we have to lock our sample position and spline segments to the view range
@@ -646,17 +654,6 @@ float FAudioSection::GetSectionHeight() const
 	return Section.GetTypedOuter<UMovieSceneAudioTrack>()->GetRowHeight();
 }
 
-void FAudioSection::GenerateSectionLayout( class ISectionLayoutBuilder& LayoutBuilder ) const
-{
-	UMovieSceneAudioSection* AudioSection = Cast<UMovieSceneAudioSection>( &Section );
-
-	SoundVolumeArea = MakeShareable( new FFloatCurveKeyArea( &AudioSection->GetSoundVolumeCurve(), AudioSection ) ) ;
-	PitchMultiplierArea = MakeShareable( new FFloatCurveKeyArea( &AudioSection->GetPitchMultiplierCurve(), AudioSection ) );
-
-	LayoutBuilder.AddKeyArea( "Volume", NSLOCTEXT( "FAudioSection", "SoundVolumeArea", "Volume" ), SoundVolumeArea.ToSharedRef() );
-	LayoutBuilder.AddKeyArea( "Pitch Multiplier", NSLOCTEXT( "FAudioSection", "PitchMultiplierArea", "Pitch Multiplier" ), PitchMultiplierArea.ToSharedRef() );
-}
-
 int32 FAudioSection::OnPaintSection( FSequencerSectionPainter& Painter ) const
 {
 	int32 LayerId = Painter.PaintSectionBackground();
@@ -697,11 +694,12 @@ void FAudioSection::Tick( const FGeometry& AllottedGeometry, const FGeometry& Pa
 		const float LeftMostVisiblePixel = FMath::Max(ParentRect.Left, 0.f);
 		const float RightMostVisiblePixel = FMath::Min(ParentRect.Right, AllottedGeometry.GetLocalSize().X);
 
-		FTimeToPixel TimeToPixelConverter( AllottedGeometry, TRange<float>( AudioSection->GetStartTime(), AudioSection->GetEndTime() ) );
+		FFrameRate   FrameResolution = AudioSection->GetTypedOuter<UMovieScene>()->GetFrameResolution();
+		FTimeToPixel TimeToPixel( AllottedGeometry, AudioSection->GetRange() / FrameResolution, FrameResolution );
 
 		TRange<float> DrawRange = TRange<float>(
-			TimeToPixelConverter.PixelToTime(LeftMostVisiblePixel),
-			TimeToPixelConverter.PixelToTime(RightMostVisiblePixel)
+			TimeToPixel.PixelToSeconds(LeftMostVisiblePixel),
+			TimeToPixel.PixelToSeconds(RightMostVisiblePixel)
 			);
 
 		// generate texture x offset and x size
@@ -740,18 +738,18 @@ void FAudioSection::BeginSlipSection()
 {
 	UMovieSceneAudioSection* AudioSection = Cast<UMovieSceneAudioSection>(&Section);
 	InitialStartOffsetDuringResize = AudioSection->GetStartOffset();
-	InitialStartTimeDuringResize = AudioSection->GetStartTime();
+	InitialStartTimeDuringResize   = FFrameNumber(AudioSection->HasStartFrame() ? AudioSection->GetInclusiveStartFrame() : 0) / AudioSection->GetTypedOuter<UMovieScene>()->GetFrameResolution();
 }
 
-void FAudioSection::SlipSection(float SlipTime)
+void FAudioSection::SlipSection(double SlipTime)
 {
 	UMovieSceneAudioSection* AudioSection = Cast<UMovieSceneAudioSection>(&Section);
 
-	float StartOffset = (SlipTime - InitialStartTimeDuringResize);
+	double StartOffset = (SlipTime - InitialStartTimeDuringResize);
 	StartOffset += InitialStartOffsetDuringResize;
 
 	// Ensure start offset is not less than 0
-	StartOffset = FMath::Max(StartOffset, 0.f);
+	StartOffset = FMath::Max(StartOffset, 0.0);
 
 	AudioSection->SetStartOffset(StartOffset);
 
@@ -882,6 +880,66 @@ void FAudioTrackEditor::Resize(float NewSize, UMovieSceneTrack* InTrack)
 	}
 }
 
+bool FAudioTrackEditor::OnAllowDrop(const FDragDropEvent& DragDropEvent, UMovieSceneTrack* Track, int32 RowIndex, const FGuid& TargetObjectGuid)
+{
+	if (!Track->IsA(UMovieSceneAudioTrack::StaticClass()))
+	{
+		return false;
+	}
+
+	TSharedPtr<FDragDropOperation> Operation = DragDropEvent.GetOperation();
+
+	if (!Operation.IsValid() || !Operation->IsOfType<FAssetDragDropOp>() )
+	{
+		return false;
+	}
+	
+	TSharedPtr<FAssetDragDropOp> DragDropOp = StaticCastSharedPtr<FAssetDragDropOp>( Operation );
+
+	for (const FAssetData& AssetData : DragDropOp->GetAssets())
+	{
+		if (Cast<USoundBase>(AssetData.GetAsset()))
+		{
+			return true;
+		}
+	}
+
+	return false;
+}
+
+
+FReply FAudioTrackEditor::OnDrop(const FDragDropEvent& DragDropEvent, UMovieSceneTrack* Track, int32 RowIndex, const FGuid& TargetObjectGuid)
+{
+	if (!Track->IsA(UMovieSceneAudioTrack::StaticClass()))
+	{
+		return FReply::Unhandled();
+	}
+
+	TSharedPtr<FDragDropOperation> Operation = DragDropEvent.GetOperation();
+
+	if (!Operation.IsValid() || !Operation->IsOfType<FAssetDragDropOp>() )
+	{
+		return FReply::Unhandled();
+	}
+	
+	TSharedPtr<FAssetDragDropOp> DragDropOp = StaticCastSharedPtr<FAssetDragDropOp>( Operation );
+
+	bool bAnyDropped = false;
+	for (const FAssetData& AssetData : DragDropOp->GetAssets())
+	{
+		USoundBase* Sound = Cast<USoundBase>(AssetData.GetAsset());
+
+		if (Sound)
+		{
+			AnimatablePropertyChanged( FOnKeyProperty::CreateRaw(this, &FAudioTrackEditor::AddNewMasterSound, Sound, RowIndex));
+
+			bAnyDropped = true;
+		}
+	}
+
+	return bAnyDropped ? FReply::Handled() : FReply::Unhandled();
+}
+
 TSharedRef<ISequencerSection> FAudioTrackEditor::MakeSectionInterface( UMovieSceneSection& SectionObject, UMovieSceneTrack& Track, FGuid ObjectBinding )
 {
 	check( SupportsType( SectionObject.GetOuter()->GetClass() ) );
@@ -908,6 +966,8 @@ bool FAudioTrackEditor::HandleAssetAdded(UObject* Asset, const FGuid& TargetObje
 	{
 		auto Sound = Cast<USoundBase>(Asset);
 		
+		const FScopedTransaction Transaction(LOCTEXT("AddAudio_Transaction", "Add Audio"));
+
 		if (TargetObjectGuid.IsValid())
 		{
 			TArray<TWeakObjectPtr<>> OutObjects;
@@ -920,7 +980,8 @@ bool FAudioTrackEditor::HandleAssetAdded(UObject* Asset, const FGuid& TargetObje
 		}
 		else
 		{
-			AnimatablePropertyChanged( FOnKeyProperty::CreateRaw(this, &FAudioTrackEditor::AddNewMasterSound, Sound));
+			int32 RowIndex = INDEX_NONE;
+			AnimatablePropertyChanged( FOnKeyProperty::CreateRaw(this, &FAudioTrackEditor::AddNewMasterSound, Sound, RowIndex));
 		}
 
 		return true;
@@ -929,15 +990,20 @@ bool FAudioTrackEditor::HandleAssetAdded(UObject* Asset, const FGuid& TargetObje
 }
 
 
-FKeyPropertyResult FAudioTrackEditor::AddNewMasterSound( float KeyTime, USoundBase* Sound )
+FKeyPropertyResult FAudioTrackEditor::AddNewMasterSound( FFrameNumber KeyTime, USoundBase* Sound, int32 RowIndex )
 {
 	FKeyPropertyResult KeyPropertyResult;
 
+	UMovieScene* FocusedMovieScene = GetFocusedMovieScene();
+	FocusedMovieScene->Modify();
+
 	FFindOrCreateMasterTrackResult<UMovieSceneAudioTrack> TrackResult = FindOrCreateMasterTrack<UMovieSceneAudioTrack>();
 	UMovieSceneTrack* Track = TrackResult.Track;
+	Track->Modify();
 
 	auto AudioTrack = Cast<UMovieSceneAudioTrack>(Track);
-	AudioTrack->AddNewSound( Sound, KeyTime );
+	UMovieSceneSection* NewSection = AudioTrack->AddNewSoundOnRow( Sound, KeyTime, RowIndex );
+
 	if (TrackResult.bWasCreated)
 	{
 		AudioTrack->SetDisplayName(LOCTEXT("AudioTrackName", "Audio"));
@@ -945,11 +1011,15 @@ FKeyPropertyResult FAudioTrackEditor::AddNewMasterSound( float KeyTime, USoundBa
 
 	KeyPropertyResult.bTrackModified = true;
 
+	GetSequencer()->EmptySelection();
+	GetSequencer()->SelectSection(NewSection);
+	GetSequencer()->ThrobSectionSelection();
+
 	return KeyPropertyResult;
 }
 
 
-FKeyPropertyResult FAudioTrackEditor::AddNewAttachedSound( float KeyTime, USoundBase* Sound, TArray<TWeakObjectPtr<UObject>> ObjectsToAttachTo )
+FKeyPropertyResult FAudioTrackEditor::AddNewAttachedSound( FFrameNumber KeyTime, USoundBase* Sound, TArray<TWeakObjectPtr<UObject>> ObjectsToAttachTo )
 {
 	FKeyPropertyResult KeyPropertyResult;
 
@@ -970,9 +1040,14 @@ FKeyPropertyResult FAudioTrackEditor::AddNewAttachedSound( float KeyTime, USound
 			if (ensure(Track))
 			{
 				auto AudioTrack = Cast<UMovieSceneAudioTrack>(Track);
-				AudioTrack->AddNewSound(Sound, KeyTime);
-				AudioTrack->SetDisplayName(LOCTEXT("AudioTrackName", "Audio"));
+				AudioTrack->Modify();
+				UMovieSceneSection* NewSection = AudioTrack->AddNewSound(Sound, KeyTime);
+				AudioTrack->SetDisplayName(LOCTEXT("AudioTrackName", "Audio"));				
 				KeyPropertyResult.bTrackModified = true;
+
+				GetSequencer()->EmptySelection();
+				GetSequencer()->SelectSection(NewSection);
+				GetSequencer()->ThrobSectionSelection();
 			}
 		}
 	}
@@ -1001,6 +1076,11 @@ void FAudioTrackEditor::HandleAddAudioTrackMenuEntryExecute()
 
 	NewTrack->SetDisplayName(LOCTEXT("AudioTrackName", "Audio"));
 
+	if (GetSequencer().IsValid())
+	{
+		GetSequencer()->OnAddTrack(NewTrack);
+	}
+
 	GetSequencer()->NotifyMovieSceneDataChanged( EMovieSceneDataChangeType::MovieSceneStructureItemAdded );
 }
 
@@ -1056,8 +1136,12 @@ void FAudioTrackEditor::OnAudioAssetSelected(const FAssetData& AssetData, UMovie
 			auto AudioTrack = Cast<UMovieSceneAudioTrack>(Track);
 			AudioTrack->Modify();
 
-			float KeyTime = GetSequencer()->GetLocalTime();
-			AudioTrack->AddNewSound( NewSound, KeyTime );
+			FFrameTime KeyTime = GetSequencer()->GetLocalTime().Time;
+			UMovieSceneSection* NewSection = AudioTrack->AddNewSound( NewSound, KeyTime.FrameNumber );				
+
+			GetSequencer()->EmptySelection();
+			GetSequencer()->SelectSection(NewSection);
+			GetSequencer()->ThrobSectionSelection();
 
 			GetSequencer()->NotifyMovieSceneDataChanged( EMovieSceneDataChangeType::MovieSceneStructureItemAdded );
 		}
diff --git a/Engine/Source/Editor/MovieSceneTools/Private/TrackEditors/CameraAnimTrackEditor.cpp b/Engine/Source/Editor/MovieSceneTools/Private/TrackEditors/CameraAnimTrackEditor.cpp
--- a/Engine/Source/Editor/MovieSceneTools/Private/TrackEditors/CameraAnimTrackEditor.cpp
+++ b/Engine/Source/Editor/MovieSceneTools/Private/TrackEditors/CameraAnimTrackEditor.cpp
@@ -25,22 +25,17 @@
 #define LOCTEXT_NAMESPACE "FCameraAnimTrackEditor"
 
 
-class FCameraAnimSection : public ISequencerSection
+class FCameraAnimSection : public FSequencerSection
 {
 public:
 	FCameraAnimSection(UMovieSceneSection& InSection)
-		: Section( InSection )
+		: FSequencerSection( InSection )
 	{ }
 
 	/** ISequencerSection interface */
-	virtual UMovieSceneSection* GetSectionObject() override
-	{ 
-		return &Section;
-	}
-	
 	virtual FText GetSectionTitle() const override 
 	{ 
-		UMovieSceneCameraAnimSection const* const AnimSection = Cast<UMovieSceneCameraAnimSection>(&Section);
+		UMovieSceneCameraAnimSection const* const AnimSection = Cast<UMovieSceneCameraAnimSection>(WeakSection.Get());
 		UCameraAnim const* const Anim = AnimSection ? AnimSection->AnimData.CameraAnim : nullptr;
 		if (Anim)
 		{
@@ -48,21 +43,6 @@ public:
 		}
 		return LOCTEXT("NoCameraAnimSection", "No Camera Anim");
 	}
-
-	virtual void GenerateSectionLayout( class ISectionLayoutBuilder& LayoutBuilder ) const override
-	{
-// 		UMovieSceneCameraAnimSection* PathSection = Cast<UMovieSceneCameraAnimSection>(&Section);
-// 		LayoutBuilder.AddKeyArea("Timing", LOCTEXT("TimingArea", "Timing"), MakeShareable( new FFloatCurveKeyArea ( &PathSection->GetTimingCurve( ), PathSection ) ) );
-	}
-
-	virtual int32 OnPaintSection( FSequencerSectionPainter& InPainter ) const override 
-	{
-		return InPainter.PaintSectionBackground();
-	}
-
-private:
-	/** The section we are visualizing */
-	UMovieSceneSection& Section;
 };
 
 
@@ -130,6 +110,8 @@ bool FCameraAnimTrackEditor::HandleAssetAdded(UObject* Asset, const FGuid& Targe
 				OutObjects.Add(Object);
 			}
 
+			const FScopedTransaction Transaction(LOCTEXT("AddCameraAnim_Transaction", "Add Camera Anim"));
+
 			AnimatablePropertyChanged(FOnKeyProperty::CreateRaw(this, &FCameraAnimTrackEditor::AddKeyInternal, OutObjects, CameraAnim));
 
 			return true;
@@ -231,7 +213,7 @@ void FCameraAnimTrackEditor::OnCameraAnimAssetSelected(const FAssetData& AssetDa
 
 
 
-FKeyPropertyResult FCameraAnimTrackEditor::AddKeyInternal(float KeyTime, const TArray<TWeakObjectPtr<UObject>> Objects, UCameraAnim* CameraAnim)
+FKeyPropertyResult FCameraAnimTrackEditor::AddKeyInternal(FFrameNumber KeyTime, const TArray<TWeakObjectPtr<UObject>> Objects, UCameraAnim* CameraAnim)
 {
 	FKeyPropertyResult KeyPropertyResult;
 
@@ -250,8 +232,12 @@ FKeyPropertyResult FCameraAnimTrackEditor::AddKeyInternal(float KeyTime, const T
 
 			if (ensure(Track))
 			{
-				Cast<UMovieSceneCameraAnimTrack>(Track)->AddNewCameraAnim(KeyTime, CameraAnim);
+				UMovieSceneSection* NewSection = Cast<UMovieSceneCameraAnimTrack>(Track)->AddNewCameraAnim(KeyTime, CameraAnim);
 				KeyPropertyResult.bTrackModified = true;
+
+				GetSequencer()->EmptySelection();
+				GetSequencer()->SelectSection(NewSection);
+				GetSequencer()->ThrobSectionSelection();
 			}
 		}
 	}
diff --git a/Engine/Source/Editor/MovieSceneTools/Private/TrackEditors/CameraAnimTrackEditor.h b/Engine/Source/Editor/MovieSceneTools/Private/TrackEditors/CameraAnimTrackEditor.h
--- a/Engine/Source/Editor/MovieSceneTools/Private/TrackEditors/CameraAnimTrackEditor.h
+++ b/Engine/Source/Editor/MovieSceneTools/Private/TrackEditors/CameraAnimTrackEditor.h
@@ -54,7 +54,7 @@ public:
 private:
 
 	/** Delegate for AnimatablePropertyChanged in AddKey */
-	FKeyPropertyResult AddKeyInternal(float KeyTime, const TArray<TWeakObjectPtr<UObject>> Objects, UCameraAnim* CameraAnim);
+	FKeyPropertyResult AddKeyInternal(FFrameNumber KeyTime, const TArray<TWeakObjectPtr<UObject>> Objects, UCameraAnim* CameraAnim);
 
 	/** Animation sub menu */
 	TSharedRef<SWidget> BuildCameraAnimSubMenu(FGuid ObjectBinding);
diff --git a/Engine/Source/Editor/MovieSceneTools/Private/TrackEditors/CameraCutTrackEditor.cpp b/Engine/Source/Editor/MovieSceneTools/Private/TrackEditors/CameraCutTrackEditor.cpp
--- a/Engine/Source/Editor/MovieSceneTools/Private/TrackEditors/CameraCutTrackEditor.cpp
+++ b/Engine/Source/Editor/MovieSceneTools/Private/TrackEditors/CameraCutTrackEditor.cpp
@@ -13,6 +13,7 @@
 #include "GameFramework/WorldSettings.h"
 #include "LevelEditorViewport.h"
 #include "Sections/CameraCutSection.h"
+#include "Sections/MovieSceneCameraCutSection.h"
 #include "SequencerUtilities.h"
 #include "Editor.h"
 #include "ActorEditorUtils.h"
@@ -184,11 +185,11 @@ void FCameraCutTrackEditor::Tick(float DeltaTime)
 	{
 		SequencerPin->EnterSilentMode();
 
-		float SavedTime = SequencerPin->GetLocalTime();
+		FQualifiedFrameTime SavedTime = SequencerPin->GetLocalTime();
 
 		if (DeltaTime > 0.f && ThumbnailPool->DrawThumbnails())
 		{
-			SequencerPin->SetLocalTimeDirectly(SavedTime);
+			SequencerPin->SetLocalTimeDirectly(SavedTime.Time);
 		}
 
 		SequencerPin->ExitSilentMode();
@@ -202,7 +203,7 @@ const FSlateBrush* FCameraCutTrackEditor::GetIconBrush() const
 }
 
 
-bool FCameraCutTrackEditor::OnAllowDrop(const FDragDropEvent& DragDropEvent, UMovieSceneTrack* Track)
+bool FCameraCutTrackEditor::OnAllowDrop(const FDragDropEvent& DragDropEvent, UMovieSceneTrack* Track, int32 RowIndex, const FGuid& TargetObjectGuid)
 {
 	if (!Track->IsA(UMovieSceneCameraCutTrack::StaticClass()))
 	{
@@ -236,7 +237,7 @@ bool FCameraCutTrackEditor::OnAllowDrop(const FDragDropEvent& DragDropEvent, UMo
 }
 
 
-FReply FCameraCutTrackEditor::OnDrop(const FDragDropEvent& DragDropEvent, UMovieSceneTrack* Track)
+FReply FCameraCutTrackEditor::OnDrop(const FDragDropEvent& DragDropEvent, UMovieSceneTrack* Track, int32 RowIndex, const FGuid& TargetObjectGuid)
 {
 	if (!Track->IsA(UMovieSceneCameraCutTrack::StaticClass()))
 	{
@@ -276,16 +277,20 @@ FReply FCameraCutTrackEditor::OnDrop(const FDragDropEvent& DragDropEvent, UMovie
 /* FCameraCutTrackEditor implementation
  *****************************************************************************/
 
-FKeyPropertyResult FCameraCutTrackEditor::AddKeyInternal( float KeyTime, const FGuid ObjectGuid )
+FKeyPropertyResult FCameraCutTrackEditor::AddKeyInternal( FFrameNumber KeyTime, const FGuid ObjectGuid )
 {
 	FKeyPropertyResult KeyPropertyResult;
 
 	UMovieSceneCameraCutTrack* CameraCutTrack = FindOrCreateCameraCutTrack();
 	const TArray<UMovieSceneSection*>& AllSections = CameraCutTrack->GetAllSections();
 
-	CameraCutTrack->AddNewCameraCut(FMovieSceneObjectBindingID(ObjectGuid, MovieSceneSequenceID::Root), KeyTime);
+	UMovieSceneCameraCutSection* NewSection = CameraCutTrack->AddNewCameraCut(FMovieSceneObjectBindingID(ObjectGuid, MovieSceneSequenceID::Root), KeyTime);
 	KeyPropertyResult.bTrackModified = true;
 
+	GetSequencer()->EmptySelection();
+	GetSequencer()->SelectSection(NewSection);
+	GetSequencer()->ThrobSectionSelection();
+
 	return KeyPropertyResult;
 }
 
@@ -320,7 +325,12 @@ bool FCameraCutTrackEditor::HandleAddCameraCutTrackMenuEntryCanExecute() const
 
 void FCameraCutTrackEditor::HandleAddCameraCutTrackMenuEntryExecute()
 {
-	FindOrCreateCameraCutTrack();
+	UMovieSceneCameraCutTrack* CameraCutTrack = FindOrCreateCameraCutTrack();
+
+	if (GetSequencer().IsValid())
+	{
+		GetSequencer()->OnAddTrack(CameraCutTrack);
+	}
 	GetSequencer()->NotifyMovieSceneDataChanged( EMovieSceneDataChangeType::MovieSceneStructureItemAdded );
 }
 
@@ -424,13 +434,17 @@ TSharedRef<SWidget> FCameraCutTrackEditor::HandleAddCameraCutComboButtonGetMenuC
 
 void FCameraCutTrackEditor::CreateNewSectionFromBinding(FMovieSceneObjectBindingID InBindingID)
 {
-	auto CreateNewSection = [this, InBindingID](float KeyTime)
+	auto CreateNewSection = [this, InBindingID](FFrameNumber KeyTime)
 	{
 		FKeyPropertyResult KeyPropertyResult;
 
-		FindOrCreateCameraCutTrack()->AddNewCameraCut(InBindingID, KeyTime);
+		UMovieSceneCameraCutSection* NewSection = FindOrCreateCameraCutTrack()->AddNewCameraCut(InBindingID, KeyTime);
 		KeyPropertyResult.bTrackModified = true;
 
+		GetSequencer()->EmptySelection();
+		GetSequencer()->SelectSection(NewSection);
+		GetSequencer()->ThrobSectionSelection();
+		
 		return KeyPropertyResult;
 	};
 
diff --git a/Engine/Source/Editor/MovieSceneTools/Private/TrackEditors/CameraCutTrackEditor.h b/Engine/Source/Editor/MovieSceneTools/Private/TrackEditors/CameraCutTrackEditor.h
--- a/Engine/Source/Editor/MovieSceneTools/Private/TrackEditors/CameraCutTrackEditor.h
+++ b/Engine/Source/Editor/MovieSceneTools/Private/TrackEditors/CameraCutTrackEditor.h
@@ -61,13 +61,13 @@ public:
 	virtual bool SupportsSequence(UMovieSceneSequence* InSequence) const override;
 	virtual void Tick(float DeltaTime) override;
 	virtual const FSlateBrush* GetIconBrush() const override;
-	virtual bool OnAllowDrop(const FDragDropEvent& DragDropEvent, UMovieSceneTrack* Track) override;
-	virtual FReply OnDrop(const FDragDropEvent& DragDropEvent, UMovieSceneTrack* Track) override;
+	virtual bool OnAllowDrop(const FDragDropEvent& DragDropEvent, UMovieSceneTrack* Track, int32 RowIndex, const FGuid& TargetObjectGuid) override;
+	virtual FReply OnDrop(const FDragDropEvent& DragDropEvent, UMovieSceneTrack* Track, int32 RowIndex, const FGuid& TargetObjectGuid) override;
 
 protected:
 
 	/** Delegate for AnimatablePropertyChanged in AddKey */
-	FKeyPropertyResult AddKeyInternal(float AutoKeyTime, const FGuid ObjectGuid);
+	FKeyPropertyResult AddKeyInternal(FFrameNumber AutoKeyTime, const FGuid ObjectGuid);
 
 	/** Find or create a camera cut track in the currently focused movie scene. */
 	UMovieSceneCameraCutTrack* FindOrCreateCameraCutTrack();
diff --git a/Engine/Source/Editor/MovieSceneTools/Private/TrackEditors/CameraShakeTrackEditor.cpp b/Engine/Source/Editor/MovieSceneTools/Private/TrackEditors/CameraShakeTrackEditor.cpp
--- a/Engine/Source/Editor/MovieSceneTools/Private/TrackEditors/CameraShakeTrackEditor.cpp
+++ b/Engine/Source/Editor/MovieSceneTools/Private/TrackEditors/CameraShakeTrackEditor.cpp
@@ -29,22 +29,16 @@ static const FName ParentClassTagName(TEXT("ParentClass"));
 static const FString CameraShakeClassPath(TEXT("Class'/Script/Engine.CameraShake'"));
 
 
-class FCameraShakeSection : public ISequencerSection
+class FCameraShakeSection : public FSequencerSection
 {
 public:
 	FCameraShakeSection(UMovieSceneSection& InSection)
-		: Section(InSection)
+		: FSequencerSection(InSection)
 	{ }
 
-	/** ISequencerSection interface */
-	virtual UMovieSceneSection* GetSectionObject() override
-	{
-		return &Section;
-	}
-
 	virtual FText GetSectionTitle() const override
 	{
-		UMovieSceneCameraShakeSection const* const ShakeSection = Cast<UMovieSceneCameraShakeSection>(&Section);
+		UMovieSceneCameraShakeSection const* const ShakeSection = Cast<UMovieSceneCameraShakeSection>(WeakSection.Get());
 		UClass const* const Shake = ShakeSection ? ShakeSection->ShakeData.ShakeClass : nullptr;
 		if (Shake)
 		{
@@ -52,22 +46,6 @@ public:
 		}
 		return LOCTEXT("NoCameraShakeSection", "No Camera Shake");
 	}
-
-	virtual void GenerateSectionLayout(class ISectionLayoutBuilder& LayoutBuilder) const override
-	{
-// 		UMovieSceneCameraShakeSection* PathSection = Cast<UMovieSceneCameraShakeSection>(&Section);
-// 		LayoutBuilder.AddKeyArea("Timing", LOCTEXT("TimingArea", "Timing"), MakeShareable( new FFloatCurveKeyArea ( &PathSection->GetTimingCurve( ), PathSection ) ) );
-	}
-
-	virtual int32 OnPaintSection(FSequencerSectionPainter& InPainter) const override
-	{
-		return InPainter.PaintSectionBackground();
-	}
-
-private:
-
-	/** The section we are visualizing */
-	UMovieSceneSection& Section;
 };
 
 
@@ -135,6 +113,9 @@ bool FCameraShakeTrackEditor::HandleAssetAdded(UObject* Asset, const FGuid& Targ
 			{
 				OutObjects.Add(Object);
 			}
+
+			const FScopedTransaction Transaction(LOCTEXT("AddCameraShake_Transaction", "Add Camera Shake"));
+			
 			AnimatablePropertyChanged(FOnKeyProperty::CreateRaw(this, &FCameraShakeTrackEditor::AddKeyInternal, OutObjects, ShakeClass));
 
 			return true;
@@ -242,7 +223,7 @@ void FCameraShakeTrackEditor::OnCameraShakeAssetSelected(const FAssetData& Asset
 }
 
 
-FKeyPropertyResult FCameraShakeTrackEditor::AddKeyInternal(float KeyTime, const TArray<TWeakObjectPtr<UObject>> Objects, TSubclassOf<UCameraShake> ShakeClass)
+FKeyPropertyResult FCameraShakeTrackEditor::AddKeyInternal(FFrameNumber KeyTime, const TArray<TWeakObjectPtr<UObject>> Objects, TSubclassOf<UCameraShake> ShakeClass)
 {
 	FKeyPropertyResult KeyPropertyResult;
 
@@ -261,8 +242,12 @@ FKeyPropertyResult FCameraShakeTrackEditor::AddKeyInternal(float KeyTime, const
 
 			if (ensure(Track))
 			{
-				Cast<UMovieSceneCameraShakeTrack>(Track)->AddNewCameraShake(KeyTime, ShakeClass);
+				UMovieSceneSection* NewSection = Cast<UMovieSceneCameraShakeTrack>(Track)->AddNewCameraShake(KeyTime, ShakeClass);
 				KeyPropertyResult.bTrackModified = true;
+				
+				GetSequencer()->EmptySelection();
+				GetSequencer()->SelectSection(NewSection);
+				GetSequencer()->ThrobSectionSelection();
 			}
 		}
 	}
diff --git a/Engine/Source/Editor/MovieSceneTools/Private/TrackEditors/CameraShakeTrackEditor.h b/Engine/Source/Editor/MovieSceneTools/Private/TrackEditors/CameraShakeTrackEditor.h
--- a/Engine/Source/Editor/MovieSceneTools/Private/TrackEditors/CameraShakeTrackEditor.h
+++ b/Engine/Source/Editor/MovieSceneTools/Private/TrackEditors/CameraShakeTrackEditor.h
@@ -54,7 +54,7 @@ public:
 private:
 
 	/** Delegate for AnimatablePropertyChanged in AddKey */
-	FKeyPropertyResult AddKeyInternal(float KeyTime, const TArray<TWeakObjectPtr<UObject>> Objects, TSubclassOf<UCameraShake> ShakeClass);
+	FKeyPropertyResult AddKeyInternal(FFrameNumber KeyTime, const TArray<TWeakObjectPtr<UObject>> Objects, TSubclassOf<UCameraShake> ShakeClass);
 
 	/** Animation sub menu */
 	TSharedRef<SWidget> BuildCameraShakeSubMenu(FGuid ObjectBinding);
