Copying //UE4/Dev-Sequencer to //UE4/Dev-Main (Source: //UE4/Dev-Sequencer @ 3946692)

diff --git a/Engine/Source/Editor/MovieSceneTools/Private/TrackEditors/CinematicShotTrackEditor.cpp b/Engine/Source/Editor/MovieSceneTools/Private/TrackEditors/CinematicShotTrackEditor.cpp
--- a/Engine/Source/Editor/MovieSceneTools/Private/TrackEditors/CinematicShotTrackEditor.cpp
+++ b/Engine/Source/Editor/MovieSceneTools/Private/TrackEditors/CinematicShotTrackEditor.cpp
@@ -32,6 +32,9 @@
 #include "MovieSceneToolsProjectSettings.h"
 #include "Editor.h"
 #include "DragAndDrop/AssetDragDropOp.h"
+#include "MovieSceneTimeHelpers.h"
+#include "Framework/Notifications/NotificationManager.h"
+#include "Widgets/Notifications/SNotificationList.h"
 
 #define LOCTEXT_NAMESPACE "FCinematicShotTrackEditor"
 
@@ -151,10 +154,17 @@ bool FCinematicShotTrackEditor::HandleAssetAdded(UObject* Asset, const FGuid& Ta
 
 	if (CanAddSubSequence(*Sequence))
 	{
-		AnimatablePropertyChanged(FOnKeyProperty::CreateRaw(this, &FCinematicShotTrackEditor::HandleSequenceAdded, Sequence));
+		const FScopedTransaction Transaction(LOCTEXT("AddShot_Transaction", "Add Shot"));
+		
+		int32 RowIndex = INDEX_NONE;
+		AnimatablePropertyChanged(FOnKeyProperty::CreateRaw(this, &FCinematicShotTrackEditor::HandleSequenceAdded, Sequence, RowIndex));
 
 		return true;
 	}
+		
+	FNotificationInfo Info(FText::Format( LOCTEXT("InvalidSequence", "Invalid level sequence {0}. There could be a circular dependency."), Sequence->GetDisplayName()));	
+	Info.bUseLargeFont = false;
+	FSlateNotificationManager::Get().AddNotification(Info);
 
 	return false;
 }
@@ -186,7 +196,7 @@ void FCinematicShotTrackEditor::Tick(float DeltaTime)
 	{
 		SequencerPin->EnterSilentMode();
 
-		float SavedTime = SequencerPin->GetGlobalTime();
+		FFrameTime SavedTime = SequencerPin->GetGlobalTime().Time;
 
 		if (DeltaTime > 0.f && ThumbnailPool->DrawThumbnails())
 		{
@@ -221,7 +231,7 @@ const FSlateBrush* FCinematicShotTrackEditor::GetIconBrush() const
 	return FEditorStyle::GetBrush("Sequencer.Tracks.CinematicShot");
 }
 
-bool FCinematicShotTrackEditor::OnAllowDrop(const FDragDropEvent& DragDropEvent, UMovieSceneTrack* Track)
+bool FCinematicShotTrackEditor::OnAllowDrop(const FDragDropEvent& DragDropEvent, UMovieSceneTrack* Track, int32 RowIndex, const FGuid& TargetObjectGuid)
 {
 	if (!Track->IsA(UMovieSceneCinematicShotTrack::StaticClass()))
 	{
@@ -249,7 +259,7 @@ bool FCinematicShotTrackEditor::OnAllowDrop(const FDragDropEvent& DragDropEvent,
 }
 
 
-FReply FCinematicShotTrackEditor::OnDrop(const FDragDropEvent& DragDropEvent, UMovieSceneTrack* Track)
+FReply FCinematicShotTrackEditor::OnDrop(const FDragDropEvent& DragDropEvent, UMovieSceneTrack* Track, int32 RowIndex, const FGuid& TargetObjectGuid)
 {
 	if (!Track->IsA(UMovieSceneCinematicShotTrack::StaticClass()))
 	{
@@ -272,7 +282,7 @@ FReply FCinematicShotTrackEditor::OnDrop(const FDragDropEvent& DragDropEvent, UM
 
 		if (Sequence)
 		{
-			AnimatablePropertyChanged( FOnKeyProperty::CreateRaw( this, &FCinematicShotTrackEditor::AddKeyInternal, Sequence) );
+			AnimatablePropertyChanged( FOnKeyProperty::CreateRaw( this, &FCinematicShotTrackEditor::AddKeyInternal, Sequence, RowIndex) );
 			
 			bAnyDropped = true;
 		}
@@ -281,9 +291,19 @@ FReply FCinematicShotTrackEditor::OnDrop(const FDragDropEvent& DragDropEvent, UM
 	return bAnyDropped ? FReply::Handled() : FReply::Unhandled();
 }
 
-UMovieSceneSubSection* FCinematicShotTrackEditor::CreateShotInternal(FString& NewShotName, float NewShotStartTime, UMovieSceneCinematicShotSection* ShotToDuplicate)
+UMovieSceneSubSection* FCinematicShotTrackEditor::CreateShotInternal(FString& NewShotName, FFrameNumber NewShotStartTime, UMovieSceneCinematicShotSection* ShotToDuplicate)
 {
-	FString NewShotPath = MovieSceneToolHelpers::GenerateNewShotPath(GetSequencer()->GetFocusedMovieSceneSequence()->GetMovieScene(), NewShotName);
+	FString NewShotPath;
+	
+	if (ShotToDuplicate != nullptr)
+	{
+		// If duplicating a shot, use that shot's path
+		NewShotPath = FPaths::GetPath(ShotToDuplicate->GetSequence()->GetPathName());
+	}
+	else
+	{
+		NewShotPath = MovieSceneToolHelpers::GenerateNewShotPath(GetSequencer()->GetFocusedMovieSceneSequence()->GetMovieScene(), NewShotName);
+	}
 
 	// Create a new level sequence asset with the appropriate name
 	IAssetTools& AssetTools = FModuleManager::GetModuleChecked<FAssetToolsModule>("AssetTools").Get();
@@ -317,16 +337,12 @@ UMovieSceneSubSection* FCinematicShotTrackEditor::CreateShotInternal(FString& Ne
 
 	UMovieSceneSequence* NewSequence = Cast<UMovieSceneSequence>(NewAsset);
 
-	float Duration = NewSequence->GetMovieScene()->GetPlaybackRange().Size<float>();
-	if (ShotToDuplicate != nullptr)
-	{
-		Duration = ShotToDuplicate->GetEndTime() - ShotToDuplicate->GetStartTime();
-	}
+	int32 Duration = MovieScene::DiscreteSize(ShotToDuplicate ? ShotToDuplicate->GetRange() : NewSequence->GetMovieScene()->GetPlaybackRange());
 
 	UMovieSceneCinematicShotTrack* CinematicShotTrack = FindOrCreateCinematicShotTrack();
 
 	// Create a cinematic shot section. 
-	UMovieSceneSubSection* NewSection = CinematicShotTrack->AddSequence(NewSequence, NewShotStartTime, Duration);	
+	UMovieSceneSubSection* NewSection = CinematicShotTrack->AddSequence(NewSequence, NewShotStartTime, Duration);
 	return NewSection;
 }
 
@@ -334,18 +350,21 @@ void FCinematicShotTrackEditor::InsertShot()
 {
 	const FScopedTransaction Transaction(LOCTEXT("InsertShot_Transaction", "Insert Shot"));
 
-	float NewShotStartTime = GetSequencer()->GetLocalTime();
+	FFrameTime NewShotStartTime = GetSequencer()->GetLocalTime().Time;
 
 	UMovieSceneCinematicShotTrack* CinematicShotTrack = FindOrCreateCinematicShotTrack();
-	FString NewShotName = MovieSceneToolHelpers::GenerateNewShotName(CinematicShotTrack->GetAllSections(), NewShotStartTime);
+	FString NewShotName = MovieSceneToolHelpers::GenerateNewShotName(CinematicShotTrack->GetAllSections(), NewShotStartTime.FrameNumber);
 
-	UMovieSceneSubSection* NewShot = CreateShotInternal(NewShotName, NewShotStartTime);
+	UMovieSceneSubSection* NewShot = CreateShotInternal(NewShotName, NewShotStartTime.FrameNumber);
 	if (NewShot)
 	{
 		NewShot->SetRowIndex(MovieSceneToolHelpers::FindAvailableRowIndex(CinematicShotTrack, NewShot));
 	}
 
 	GetSequencer()->NotifyMovieSceneDataChanged( EMovieSceneDataChangeType::MovieSceneStructureItemAdded );
+	GetSequencer()->EmptySelection();
+	GetSequencer()->SelectSection(NewShot);
+	GetSequencer()->ThrobSectionSelection();
 }
 
 
@@ -355,22 +374,24 @@ void FCinematicShotTrackEditor::InsertFiller()
 
 	const FScopedTransaction Transaction(LOCTEXT("InsertFiller_Transaction", "Insert Filler"));
 
-	float NewShotStartTime = GetSequencer()->GetLocalTime();
+	FQualifiedFrameTime CurrentTime = GetSequencer()->GetLocalTime();
 
 	UMovieSceneCinematicShotTrack* CinematicShotTrack = FindOrCreateCinematicShotTrack();
 
-	float Duration = ProjectSettings->DefaultDuration;
+	int32 Duration = (ProjectSettings->DefaultDuration * CurrentTime.Rate).FrameNumber.Value;
 
 	UMovieSceneSequence* NullSequence = nullptr;
 
-	UMovieSceneSubSection* NewSection = CinematicShotTrack->AddSequence(NullSequence, NewShotStartTime, Duration);	
+	UMovieSceneSubSection* NewSection = CinematicShotTrack->AddSequence(NullSequence, CurrentTime.Time.FrameNumber, Duration);
 
 	UMovieSceneCinematicShotSection* NewCinematicShotSection = Cast<UMovieSceneCinematicShotSection>(NewSection);
 
 	NewCinematicShotSection->SetShotDisplayName(FText(LOCTEXT("Filler", "Filler")).ToString());
-	NewCinematicShotSection->SetRowIndex(MovieSceneToolHelpers::FindAvailableRowIndex(CinematicShotTrack, NewSection));
 
 	GetSequencer()->NotifyMovieSceneDataChanged( EMovieSceneDataChangeType::MovieSceneStructureItemAdded );
+	GetSequencer()->EmptySelection();
+	GetSequencer()->SelectSection(NewSection);
+	GetSequencer()->ThrobSectionSelection();
 }
 
 
@@ -379,20 +400,24 @@ void FCinematicShotTrackEditor::DuplicateShot(UMovieSceneCinematicShotSection* S
 	const FScopedTransaction Transaction(LOCTEXT("DuplicateShot_Transaction", "Duplicate Shot"));
 
 	UMovieSceneCinematicShotTrack* CinematicShotTrack = FindOrCreateCinematicShotTrack();
-	FString NewShotName = MovieSceneToolHelpers::GenerateNewShotName(CinematicShotTrack->GetAllSections(), Section->GetStartTime());
+
+	FFrameNumber StartTime = Section->HasStartFrame() ? Section->GetInclusiveStartFrame() : 0;
+	FString NewShotName = MovieSceneToolHelpers::GenerateNewShotName(CinematicShotTrack->GetAllSections(), StartTime);
 
 	// Duplicate the shot and put it on the next available row
-	UMovieSceneSubSection* NewShot = CreateShotInternal(NewShotName, Section->GetStartTime(), Section);
+	UMovieSceneSubSection* NewShot = CreateShotInternal(NewShotName, StartTime, Section);
 	if (NewShot)
 	{
-		NewShot->SetStartTime(Section->GetStartTime());
-		NewShot->SetEndTime(Section->GetEndTime());
+		NewShot->SetRange(Section->GetRange());
 		NewShot->SetRowIndex(MovieSceneToolHelpers::FindAvailableRowIndex(CinematicShotTrack, NewShot));
-		NewShot->Parameters.StartOffset = Section->Parameters.StartOffset;
+		NewShot->Parameters.SetStartFrameOffset(Section->Parameters.GetStartFrameOffset());
 		NewShot->Parameters.TimeScale = Section->Parameters.TimeScale;
-		NewShot->SetPreRollTime(Section->GetPreRollTime());
+		NewShot->SetPreRollFrames(Section->GetPreRollFrames());
 
 		GetSequencer()->NotifyMovieSceneDataChanged( EMovieSceneDataChangeType::MovieSceneStructureItemAdded );
+		GetSequencer()->EmptySelection();
+		GetSequencer()->SelectSection(NewShot);
+		GetSequencer()->ThrobSectionSelection();
 	}
 }
 
@@ -430,12 +455,12 @@ void FCinematicShotTrackEditor::NewTake(UMovieSceneCinematicShotSection* Section
 
 		FString NewShotName = MovieSceneToolHelpers::ComposeShotName(ShotPrefix, ShotNumber, NewTakeNumber);
 
-		float NewShotStartTime = Section->GetStartTime();
-		float NewShotEndTime = Section->GetEndTime();
-		float NewShotStartOffset = Section->Parameters.StartOffset;
-		float NewShotTimeScale = Section->Parameters.TimeScale;
-		float NewShotPrerollTime = Section->GetPreRollTime();
-		int32 NewRowIndex = Section->GetRowIndex();
+		TRange<FFrameNumber> NewShotRange         = Section->GetRange();
+		int32                NewShotStartOffset   = Section->Parameters.GetStartFrameOffset();
+		float                NewShotTimeScale     = Section->Parameters.TimeScale;
+		int32                NewShotPrerollFrames = Section->GetPreRollFrames();
+		int32                NewRowIndex          = Section->GetRowIndex();
+		FFrameNumber         NewShotStartTime     = NewShotRange.GetLowerBound().IsClosed() ? MovieScene::DiscreteInclusiveLower(NewShotRange) : 0;
 
 		UMovieSceneSubSection* NewShot = CreateShotInternal(NewShotName, NewShotStartTime, Section);
 
@@ -444,14 +469,16 @@ void FCinematicShotTrackEditor::NewTake(UMovieSceneCinematicShotSection* Section
 			UMovieSceneCinematicShotTrack* CinematicShotTrack = FindOrCreateCinematicShotTrack();
 			CinematicShotTrack->RemoveSection(*Section);
 
-			NewShot->SetStartTime(NewShotStartTime);
-			NewShot->SetEndTime(NewShotEndTime);
-			NewShot->Parameters.StartOffset = NewShotStartOffset;
+			NewShot->SetRange(NewShotRange);
+			NewShot->Parameters.SetStartFrameOffset(NewShotStartOffset);
 			NewShot->Parameters.TimeScale = NewShotTimeScale;
-			NewShot->SetPreRollTime(NewShotPrerollTime);
+			NewShot->SetPreRollFrames(NewShotPrerollFrames);
 			NewShot->SetRowIndex(NewRowIndex);
 
 			GetSequencer()->NotifyMovieSceneDataChanged( EMovieSceneDataChangeType::MovieSceneStructureItemsChanged );
+			GetSequencer()->EmptySelection();
+			GetSequencer()->SelectSection(NewShot);
+			GetSequencer()->ThrobSectionSelection();
 		}
 	}
 }
@@ -468,25 +495,27 @@ void FCinematicShotTrackEditor::SwitchTake(UMovieSceneCinematicShotSection* Sect
 		UMovieSceneSequence* MovieSceneSequence = CastChecked<UMovieSceneSequence>(TakeObject);
 
 		UMovieSceneCinematicShotTrack* CinematicShotTrack = FindOrCreateCinematicShotTrack();
-		
-		float NewShotStartTime = Section->GetStartTime();
-		float NewShotEndTime = Section->GetEndTime();
-		float NewShotStartOffset = Section->Parameters.StartOffset;
-		float NewShotTimeScale = Section->Parameters.TimeScale;
-		float NewShotPrerollTime = Section->GetPreRollTime();
 
-		const float Duration = NewShotEndTime - NewShotStartTime;
-		UMovieSceneSubSection* NewShot = CinematicShotTrack->AddSequence(MovieSceneSequence, NewShotStartTime, Duration);	
+		TRange<FFrameNumber> NewShotRange         = Section->GetRange();
+		int32                NewShotStartOffset   = Section->Parameters.GetStartFrameOffset();
+		float                NewShotTimeScale     = Section->Parameters.TimeScale;
+		int32                NewShotPrerollFrames = Section->GetPreRollFrames();
+		int32                NewRowIndex          = Section->GetRowIndex();
+		FFrameNumber         NewShotStartTime     = NewShotRange.GetLowerBound().IsClosed() ? MovieScene::DiscreteInclusiveLower(NewShotRange) : 0;
+		int32                NewShotRowIndex      = Section->GetRowIndex();
+
+		const int32 Duration = (NewShotRange.GetLowerBound().IsClosed() && NewShotRange.GetUpperBound().IsClosed() ) ? MovieScene::DiscreteSize(NewShotRange) : 1;
+		UMovieSceneSubSection* NewShot = CinematicShotTrack->AddSequence(MovieSceneSequence, NewShotStartTime, Duration);
 
 		if (NewShot != nullptr)
 		{
 			CinematicShotTrack->RemoveSection(*Section);
 
-			NewShot->SetStartTime(NewShotStartTime);
-			NewShot->SetEndTime(NewShotEndTime);
-			NewShot->Parameters.StartOffset = NewShotStartOffset;
+			NewShot->SetRange(NewShotRange);
+			NewShot->Parameters.SetStartFrameOffset(NewShotStartOffset);
 			NewShot->Parameters.TimeScale = NewShotTimeScale;
-			NewShot->SetPreRollTime(NewShotPrerollTime);
+			NewShot->SetPreRollFrames(NewShotPrerollFrames);
+			NewShot->SetRowIndex(NewShotRowIndex);
 		}
 
 		GetSequencer()->NotifyMovieSceneDataChanged( EMovieSceneDataChangeType::MovieSceneStructureItemsChanged );
@@ -508,7 +537,12 @@ bool FCinematicShotTrackEditor::HandleAddCinematicShotTrackMenuEntryCanExecute()
 
 void FCinematicShotTrackEditor::HandleAddCinematicShotTrackMenuEntryExecute()
 {
-	FindOrCreateCinematicShotTrack();
+	UMovieSceneCinematicShotTrack* ShotTrack = FindOrCreateCinematicShotTrack();
+	if (GetSequencer().IsValid())
+	{
+		GetSequencer()->OnAddTrack(ShotTrack);
+	}
+	GetSequencer()->NotifyMovieSceneDataChanged(EMovieSceneDataChangeType::MovieSceneStructureItemAdded);
 }
 
 
@@ -563,24 +597,41 @@ void FCinematicShotTrackEditor::HandleAddCinematicShotComboButtonMenuEntryExecut
 	{
 		UMovieSceneSequence* MovieSceneSequence = CastChecked<UMovieSceneSequence>(AssetData.GetAsset());
 
-		AnimatablePropertyChanged( FOnKeyProperty::CreateRaw( this, &FCinematicShotTrackEditor::AddKeyInternal, MovieSceneSequence) );
+		int32 RowIndex = INDEX_NONE;
+		AnimatablePropertyChanged( FOnKeyProperty::CreateRaw( this, &FCinematicShotTrackEditor::AddKeyInternal, MovieSceneSequence, RowIndex) );
 	}
 }
 
 
-FKeyPropertyResult FCinematicShotTrackEditor::AddKeyInternal(float KeyTime, UMovieSceneSequence* InMovieSceneSequence)
+FKeyPropertyResult FCinematicShotTrackEditor::AddKeyInternal(FFrameNumber KeyTime, UMovieSceneSequence* InMovieSceneSequence, int32 RowIndex)
 {	
 	FKeyPropertyResult KeyPropertyResult;
 
 	if (CanAddSubSequence(*InMovieSceneSequence))
 	{
 		UMovieSceneCinematicShotTrack* CinematicShotTrack = FindOrCreateCinematicShotTrack();
-		float EndTime = InMovieSceneSequence->GetMovieScene()->GetPlaybackRange().Size<float>();
-		UMovieSceneSubSection* NewSection = CinematicShotTrack->AddSequence(InMovieSceneSequence, KeyTime, EndTime);
-		NewSection->SetRowIndex(MovieSceneToolHelpers::FindAvailableRowIndex(CinematicShotTrack, NewSection));
+		
+		const FQualifiedFrameTime InnerDuration = FQualifiedFrameTime(
+			MovieScene::DiscreteSize(InMovieSceneSequence->GetMovieScene()->GetPlaybackRange()),
+			InMovieSceneSequence->GetMovieScene()->GetFrameResolution());
+
+		const FFrameRate OuterFrameRate = CinematicShotTrack->GetTypedOuter<UMovieScene>()->GetFrameResolution();
+		const int32      OuterDuration  = InnerDuration.ConvertTo(OuterFrameRate).FrameNumber.Value;
+
+		UMovieSceneSubSection* NewSection = CinematicShotTrack->AddSequenceOnRow(InMovieSceneSequence, KeyTime, OuterDuration, RowIndex);
 		KeyPropertyResult.bTrackModified = true;
+
+		GetSequencer()->EmptySelection();
+		GetSequencer()->SelectSection(NewSection);
+		GetSequencer()->ThrobSectionSelection();
+
+		return KeyPropertyResult;
 	}
 
+	FNotificationInfo Info(FText::Format( LOCTEXT("InvalidSequence", "Invalid level sequence {0}. There could be a circular dependency."), InMovieSceneSequence->GetDisplayName()));	
+	Info.bUseLargeFont = false;
+	FSlateNotificationManager::Get().AddNotification(Info);
+
 	return KeyPropertyResult;
 }
 
@@ -680,14 +731,19 @@ bool FCinematicShotTrackEditor::CanAddSubSequence(const UMovieSceneSequence& Seq
 
 	// make sure we are not contained in the other sequence (circular dependency)
 	// @todo sequencer: this check is not sufficient (does not prevent circular dependencies of 2+ levels)
-	UMovieSceneCinematicShotTrack* SequenceCinematicShotTrack = SequenceMovieScene->FindMasterTrack<UMovieSceneCinematicShotTrack>();
+	UMovieSceneSubTrack* SequenceSubTrack = SequenceMovieScene->FindMasterTrack<UMovieSceneSubTrack>();
+	if (SequenceSubTrack && SequenceSubTrack->ContainsSequence(*FocusedSequence, true))
+	{
+		return false;
+	}
 
-	if (SequenceCinematicShotTrack == nullptr)
+	UMovieSceneCinematicShotTrack* SequenceCinematicTrack = SequenceMovieScene->FindMasterTrack<UMovieSceneCinematicShotTrack>();
+	if (SequenceCinematicTrack && SequenceCinematicTrack->ContainsSequence(*FocusedSequence, true))
 	{
-		return true;
+		return false;
 	}
 
-	return !SequenceCinematicShotTrack->ContainsSequence(*FocusedSequence, true);
+	return true;
 }
 
 
@@ -698,16 +754,26 @@ void FCinematicShotTrackEditor::OnUpdateCameraCut(UObject* CameraObject, bool bJ
 }
 
 
-FKeyPropertyResult FCinematicShotTrackEditor::HandleSequenceAdded(float KeyTime, UMovieSceneSequence* Sequence)
+FKeyPropertyResult FCinematicShotTrackEditor::HandleSequenceAdded(FFrameNumber KeyTime, UMovieSceneSequence* Sequence, int32 RowIndex)
 {
 	FKeyPropertyResult KeyPropertyResult;
 
 	auto CinematicShotTrack = FindOrCreateCinematicShotTrack();
-	float Duration = Sequence->GetMovieScene()->GetPlaybackRange().Size<float>();
-	UMovieSceneSubSection* NewSection = CinematicShotTrack->AddSequence(Sequence, KeyTime, Duration);
-	NewSection->SetRowIndex(MovieSceneToolHelpers::FindAvailableRowIndex(CinematicShotTrack, NewSection));
+
+	const FQualifiedFrameTime InnerDuration = FQualifiedFrameTime(
+		MovieScene::DiscreteSize(Sequence->GetMovieScene()->GetPlaybackRange()),
+		Sequence->GetMovieScene()->GetFrameResolution());
+
+	const FFrameRate OuterFrameRate = CinematicShotTrack->GetTypedOuter<UMovieScene>()->GetFrameResolution();
+	const int32      OuterDuration  = InnerDuration.ConvertTo(OuterFrameRate).FrameNumber.Value;
+
+	UMovieSceneSubSection* NewSection = CinematicShotTrack->AddSequenceOnRow(Sequence, KeyTime, OuterDuration, RowIndex);
 	KeyPropertyResult.bTrackModified = true;
 
+	GetSequencer()->EmptySelection();
+	GetSequencer()->SelectSection(NewSection);
+	GetSequencer()->ThrobSectionSelection();
+
 	return KeyPropertyResult;
 }
 
@@ -740,9 +806,8 @@ void FCinematicShotTrackEditor::ImportEDL()
 
 	const FMovieSceneCaptureSettings& Settings = MovieSceneCapture->GetSettings();
 	FString SaveDirectory = FPaths::ConvertRelativePathToFull(Settings.OutputDirectory.Path);
-	float FrameRate = Settings.FrameRate;
 
-	if (MovieSceneToolHelpers::ShowImportEDLDialog(MovieScene, FrameRate, SaveDirectory))
+	if (MovieSceneToolHelpers::ShowImportEDLDialog(MovieScene, Settings.FrameRate, SaveDirectory))
 	{
 		GetSequencer()->NotifyMovieSceneDataChanged( EMovieSceneDataChangeType::MovieSceneStructureItemsChanged );
 	}
@@ -779,9 +844,8 @@ void FCinematicShotTrackEditor::ExportEDL()
 	FString SaveDirectory = FPaths::ConvertRelativePathToFull(Settings.OutputDirectory.Path);
 	int32 HandleFrames = Settings.HandleFrames;
 
-	float FrameRate = 1.0f / MovieScene->GetFixedFrameInterval();
-
-	MovieSceneToolHelpers::ShowExportEDLDialog(MovieScene, FrameRate, SaveDirectory, HandleFrames);
+	// @todo: sequencer-timecode: resolution vs playrate
+	MovieSceneToolHelpers::ShowExportEDLDialog(MovieScene, MovieScene->GetFrameResolution(), SaveDirectory, HandleFrames);
 }
 
 
diff --git a/Engine/Source/Editor/MovieSceneTools/Private/TrackEditors/CinematicShotTrackEditor.h b/Engine/Source/Editor/MovieSceneTools/Private/TrackEditors/CinematicShotTrackEditor.h
--- a/Engine/Source/Editor/MovieSceneTools/Private/TrackEditors/CinematicShotTrackEditor.h
+++ b/Engine/Source/Editor/MovieSceneTools/Private/TrackEditors/CinematicShotTrackEditor.h
@@ -62,8 +62,8 @@ public:
 	virtual void Tick(float DeltaTime) override;
 	virtual void BuildTrackContextMenu( FMenuBuilder& MenuBuilder, UMovieSceneTrack* Track ) override;
 	virtual const FSlateBrush* GetIconBrush() const override;
-	virtual bool OnAllowDrop(const FDragDropEvent& DragDropEvent, UMovieSceneTrack* Track) override;
-	virtual FReply OnDrop(const FDragDropEvent& DragDropEvent, UMovieSceneTrack* Track) override;
+	virtual bool OnAllowDrop(const FDragDropEvent& DragDropEvent, UMovieSceneTrack* Track, int32 RowIndex, const FGuid& TargetObjectGuid) override;
+	virtual FReply OnDrop(const FDragDropEvent& DragDropEvent, UMovieSceneTrack* Track, int32 RowIndex, const FGuid& TargetObjectGuid) override;
 
 	/*
 	 * Insert shot. 
@@ -121,7 +121,7 @@ private:
 	 * @param ShotToDuplicate The shot to duplicate.
 	 * @return The new shot.
 	 */
-	UMovieSceneSubSection* CreateShotInternal(FString& NewShotName, float NewShotStartTime, UMovieSceneCinematicShotSection* ShotToDuplicate = nullptr);
+	UMovieSceneSubSection* CreateShotInternal(FString& NewShotName, FFrameNumber NewShotStartTime, UMovieSceneCinematicShotSection* ShotToDuplicate = nullptr);
 
 
 private:
@@ -142,7 +142,7 @@ private:
 	UMovieSceneCinematicShotTrack* FindOrCreateCinematicShotTrack();
 
 	/** Delegate for AnimatablePropertyChanged in AddKey */
-	FKeyPropertyResult AddKeyInternal(float KeyTime, UMovieSceneSequence* InMovieSceneSequence);
+	FKeyPropertyResult AddKeyInternal(FFrameNumber KeyTime, UMovieSceneSequence* InMovieSceneSequence, int32 RowIndex);
 
 	/** Delegate for shots button lock state */
 	ECheckBoxState AreShotsLocked() const; 
@@ -168,7 +168,7 @@ private:
 	void OnUpdateCameraCut(UObject* CameraObject, bool bJumpCut);
 
 	/** Callback for AnimatablePropertyChanged in HandleAssetAdded. */
-	FKeyPropertyResult HandleSequenceAdded(float KeyTime, UMovieSceneSequence* Sequence);
+	FKeyPropertyResult HandleSequenceAdded(FFrameNumber KeyTime, UMovieSceneSequence* Sequence, int32 RowIndex);
 
 	/** Callback for ImportEDL. */
 	void ImportEDL();
diff --git a/Engine/Source/Editor/MovieSceneTools/Private/TrackEditors/EventTrackEditor.cpp b/Engine/Source/Editor/MovieSceneTools/Private/TrackEditors/EventTrackEditor.cpp
--- a/Engine/Source/Editor/MovieSceneTools/Private/TrackEditors/EventTrackEditor.cpp
+++ b/Engine/Source/Editor/MovieSceneTools/Private/TrackEditors/EventTrackEditor.cpp
@@ -6,7 +6,7 @@
 #include "EditorStyleSet.h"
 #include "UObject/Package.h"
 #include "Tracks/MovieSceneEventTrack.h"
-#include "Sections/EventTrackSection.h"
+#include "ISequencerSection.h"
 #include "MovieSceneObjectBindingIDPicker.h"
 #include "IDetailCustomization.h"
 #include "PropertyEditorModule.h"
@@ -116,12 +116,6 @@ void FEventTrackEditor::BuildTrackContextMenu(FMenuBuilder& MenuBuilder, UMovieS
 }
 
 
-TSharedRef<ISequencerSection> FEventTrackEditor::MakeSectionInterface(UMovieSceneSection& SectionObject, UMovieSceneTrack& Track, FGuid ObjectBinding)
-{
-	return MakeShareable(new FEventTrackSection(SectionObject, GetSequencer()));
-}
-
-
 bool FEventTrackEditor::SupportsType(TSubclassOf<UMovieSceneTrack> Type) const
 {
 	return (Type == UMovieSceneEventTrack::StaticClass());
@@ -165,6 +159,11 @@ void FEventTrackEditor::HandleAddEventTrackMenuEntryExecute()
 	NewTrack->AddSection(*NewSection);
 	NewTrack->SetDisplayName(LOCTEXT("TrackName", "Events"));
 
+	if (GetSequencer().IsValid())
+	{
+		GetSequencer()->OnAddTrack(NewTrack);
+	}
+
 	GetSequencer()->NotifyMovieSceneDataChanged( EMovieSceneDataChangeType::MovieSceneStructureItemAdded );
 }
 
diff --git a/Engine/Source/Editor/MovieSceneTools/Private/TrackEditors/EventTrackEditor.h b/Engine/Source/Editor/MovieSceneTools/Private/TrackEditors/EventTrackEditor.h
--- a/Engine/Source/Editor/MovieSceneTools/Private/TrackEditors/EventTrackEditor.h
+++ b/Engine/Source/Editor/MovieSceneTools/Private/TrackEditors/EventTrackEditor.h
@@ -44,7 +44,6 @@ public:
 
 	virtual void BuildAddTrackMenu(FMenuBuilder& MenuBuilder) override;
 	virtual void BuildObjectBindingTrackMenu(FMenuBuilder& MenuBuilder, const FGuid& ObjectBinding, const UClass* ObjectClass) override;
-	virtual TSharedRef<ISequencerSection> MakeSectionInterface(UMovieSceneSection& SectionObject, UMovieSceneTrack& Track, FGuid ObjectBinding) override;
 	virtual bool SupportsType( TSubclassOf<UMovieSceneTrack> Type ) const override;
 	virtual bool SupportsSequence(UMovieSceneSequence* InSequence) const override;
 	virtual const FSlateBrush* GetIconBrush() const override;
diff --git a/Engine/Source/Editor/MovieSceneTools/Private/TrackEditors/FadeTrackEditor.cpp b/Engine/Source/Editor/MovieSceneTools/Private/TrackEditors/FadeTrackEditor.cpp
--- a/Engine/Source/Editor/MovieSceneTools/Private/TrackEditors/FadeTrackEditor.cpp
+++ b/Engine/Source/Editor/MovieSceneTools/Private/TrackEditors/FadeTrackEditor.cpp
@@ -6,7 +6,7 @@
 #include "EditorStyleSet.h"
 #include "Tracks/MovieSceneFadeTrack.h"
 #include "Sections/MovieSceneFadeSection.h"
-#include "Sections/FloatPropertySection.h"
+#include "ISequencerSection.h"
 #include "CommonMovieSceneTools.h"
 
 #define LOCTEXT_NAMESPACE "FFadeTrackEditor"
@@ -15,15 +15,12 @@
  * Class for fade sections handles drawing of fade gradient
  */
 class FFadeSection
-	: public FFloatPropertySection
+	: public FSequencerSection
 {
 public:
 
 	/** Constructor. */
-	FFadeSection(UMovieSceneSection& InSectionObject, const FText& SectionName) : FFloatPropertySection(InSectionObject, SectionName) {}
-
-	/** Virtual destructor. */
-	virtual ~FFadeSection() {}
+	FFadeSection(UMovieSceneSection& InSectionObject) : FSequencerSection(InSectionObject) {}
 
 public:
 
@@ -33,64 +30,42 @@ public:
 
 		const ESlateDrawEffect DrawEffects = Painter.bParentEnabled ? ESlateDrawEffect::None : ESlateDrawEffect::DisabledEffect;
 
-		const UMovieSceneFadeSection* FadeSection = Cast<const UMovieSceneFadeSection>( &SectionObject );
+		FVector2D GradientSize = FVector2D( Painter.SectionGeometry.Size.X - 2.f, Painter.SectionGeometry.Size.Y - 3.0f );
+		FPaintGeometry PaintGeometry = Painter.SectionGeometry.ToPaintGeometry( FVector2D( 1.f, 3.f ), GradientSize );
+
+		const UMovieSceneFadeSection* FadeSection = Cast<const UMovieSceneFadeSection>( WeakSection.Get() );
+
+		FTimeToPixel TimeConverter    = Painter.GetTimeConverter();
+		FFrameRate   FrameResolution  = TimeConverter.GetFrameResolution();
 
-		const FTimeToPixel& TimeConverter = Painter.GetTimeConverter();
+		const double HalfFrameTime    = FFrameTime(0, 0.5f) / FrameResolution;
+		const double StartTimeSeconds = TimeConverter.PixelToSeconds(1.f) - HalfFrameTime;
+		const double EndTimeSeconds   = TimeConverter.PixelToSeconds(Painter.SectionGeometry.GetLocalSize().X-2.f) - HalfFrameTime;
+		const double TimeThreshold    = FMath::Max(0.0001, TimeConverter.PixelToSeconds(5) - TimeConverter.PixelToSeconds(0));
+		const double DurationSeconds  = EndTimeSeconds - StartTimeSeconds;
 
-		const float StartTime = TimeConverter.PixelToTime(0.f);
-		const float EndTime = TimeConverter.PixelToTime(Painter.SectionGeometry.GetLocalSize().X);
-		const float SectionDuration = EndTime - StartTime;
+		TArray<TTuple<double, double>> CurvePoints;
+		FadeSection->GetChannel().PopulateCurvePoints(StartTimeSeconds, EndTimeSeconds, TimeThreshold, 0.1f, FrameResolution, CurvePoints);
+
+		TArray<FSlateGradientStop> GradientStops;
+		for (TTuple<double, double> Vector : CurvePoints)
+		{
+			GradientStops.Add( FSlateGradientStop(
+				FVector2D( (Vector.Get<0>() - StartTimeSeconds) / DurationSeconds * Painter.SectionGeometry.Size.X, 0 ),
+				FadeSection->FadeColor.CopyWithNewOpacity(Vector.Get<1>()) )
+			);
+		}
 
-		if ( !FMath::IsNearlyZero( SectionDuration ) )
+		if ( GradientStops.Num() > 0 )
 		{
-			FVector2D GradientSize = FVector2D( Painter.SectionGeometry.Size.X - 2.f, Painter.SectionGeometry.Size.Y - 3.0f );
-
-			FPaintGeometry PaintGeometry = Painter.SectionGeometry.ToPaintGeometry( FVector2D( 1.f, 3.f ), GradientSize );
-
-			TArray<FSlateGradientStop> GradientStops;
-
-			TArray< TKeyValuePair<float, FLinearColor> > FadeKeys;
-
-			TArray<float> TimesWithKeys;
-			for ( auto It( FadeSection->GetFloatCurve().GetKeyIterator() ); It; ++It )
-			{
-				float Time = It->Time;
-				TimesWithKeys.Add(Time);
-			}
-
-			// Enforce at least one key for the default value
-			if (TimesWithKeys.Num() == 0)
-			{
-				TimesWithKeys.Add(0);
-			}
-
-			for (auto Time : TimesWithKeys)
-			{
-				float Value = FadeSection->Eval(Time, 0.f);
-			
-				FLinearColor Color = FadeSection->FadeColor;
-				Color.A = Value*255.f;
-
-				float TimeFraction = (Time - StartTime) / SectionDuration;
-
-				if (TimeFraction > 0)
-				{
-					GradientStops.Add( FSlateGradientStop( FVector2D( TimeFraction * Painter.SectionGeometry.Size.X, 0 ),
-						Color ) );
-				}
-			}
-
-			if ( GradientStops.Num() > 0 )
-			{
-				FSlateDrawElement::MakeGradient(
-					Painter.DrawElements,
-					Painter.LayerId + 1,
-					PaintGeometry,
-					GradientStops,
-					Orient_Vertical,
-					DrawEffects
-					);
-			}
+			FSlateDrawElement::MakeGradient(
+				Painter.DrawElements,
+				Painter.LayerId + 1,
+				PaintGeometry,
+				GradientStops,
+				Orient_Vertical,
+				DrawEffects
+				);
 		}
 
 		return LayerId + 1;
@@ -119,7 +94,7 @@ FFadeTrackEditor::FFadeTrackEditor(TSharedRef<ISequencer> InSequencer)
 
 TSharedRef<ISequencerSection> FFadeTrackEditor::MakeSectionInterface(UMovieSceneSection& SectionObject, UMovieSceneTrack& Track, FGuid ObjectBinding)
 {
-	return MakeShareable(new FFadeSection(SectionObject, Track.GetDisplayName()));
+	return MakeShareable(new FFadeSection(SectionObject));
 }
 
 void FFadeTrackEditor::BuildAddTrackMenu(FMenuBuilder& MenuBuilder)
@@ -129,7 +104,8 @@ void FFadeTrackEditor::BuildAddTrackMenu(FMenuBuilder& MenuBuilder)
 		LOCTEXT("AddFadeTrackTooltip", "Adds a new track that controls the fade of the sequence."),
 		FSlateIcon(FEditorStyle::GetStyleSetName(), "Sequencer.Tracks.Fade"),
 		FUIAction(
-			FExecuteAction::CreateRaw(this, &FFadeTrackEditor::HandleAddFadeTrackMenuEntryExecute)
+			FExecuteAction::CreateRaw(this, &FFadeTrackEditor::HandleAddFadeTrackMenuEntryExecute),
+			FCanExecuteAction::CreateRaw(this, &FFadeTrackEditor::HandleAddFadeTrackMenuEntryCanExecute)
 		)
 	);
 }
@@ -180,9 +156,18 @@ void FFadeTrackEditor::HandleAddFadeTrackMenuEntryExecute()
 	check(NewSection);
 
 	FadeTrack->AddSection(*NewSection);
-
+	if (GetSequencer().IsValid())
+	{
+		GetSequencer()->OnAddTrack(FadeTrack);
+	}
 	GetSequencer()->NotifyMovieSceneDataChanged( EMovieSceneDataChangeType::MovieSceneStructureItemAdded );
 }
 
+bool FFadeTrackEditor::HandleAddFadeTrackMenuEntryCanExecute() const
+{
+	UMovieScene* FocusedMovieScene = GetFocusedMovieScene();
+	
+	return ((FocusedMovieScene != nullptr) && (FocusedMovieScene->FindMasterTrack<UMovieSceneFadeTrack>() == nullptr));
+}
 
 #undef LOCTEXT_NAMESPACE
diff --git a/Engine/Source/Editor/MovieSceneTools/Private/TrackEditors/FadeTrackEditor.h b/Engine/Source/Editor/MovieSceneTools/Private/TrackEditors/FadeTrackEditor.h
--- a/Engine/Source/Editor/MovieSceneTools/Private/TrackEditors/FadeTrackEditor.h
+++ b/Engine/Source/Editor/MovieSceneTools/Private/TrackEditors/FadeTrackEditor.h
@@ -50,4 +50,5 @@ private:
 
 	/** Callback for executing the "Add Fade Track" menu entry. */
 	void HandleAddFadeTrackMenuEntryExecute();
+	bool HandleAddFadeTrackMenuEntryCanExecute() const;
 };
diff --git a/Engine/Source/Editor/MovieSceneTools/Private/TrackEditors/LevelVisibilityTrackEditor.cpp b/Engine/Source/Editor/MovieSceneTools/Private/TrackEditors/LevelVisibilityTrackEditor.cpp
--- a/Engine/Source/Editor/MovieSceneTools/Private/TrackEditors/LevelVisibilityTrackEditor.cpp
+++ b/Engine/Source/Editor/MovieSceneTools/Private/TrackEditors/LevelVisibilityTrackEditor.cpp
@@ -63,7 +63,7 @@ TSharedPtr<SWidget> FLevelVisibilityTrackEditor::BuildOutlinerEditWidget( const
 }
 
 
-void FLevelVisibilityTrackEditor::AddNewSection( UMovieScene* MovieScene, UMovieSceneTrack* LevelVisibilityTrack, ELevelVisibility Visibility )
+UMovieSceneLevelVisibilitySection* FLevelVisibilityTrackEditor::AddNewSection( UMovieScene* MovieScene, UMovieSceneTrack* LevelVisibilityTrack, ELevelVisibility Visibility )
 {
 	const FScopedTransaction Transaction( LOCTEXT( "AddLevelVisibilitySection_Transaction", "Add Level Visibility Trigger" ) );
 
@@ -71,8 +71,9 @@ void FLevelVisibilityTrackEditor::AddNewSection( UMovieScene* MovieScene, UMovie
 
 	UMovieSceneLevelVisibilitySection* LevelVisibilitySection = CastChecked<UMovieSceneLevelVisibilitySection>( LevelVisibilityTrack->CreateNewSection() );
 	LevelVisibilitySection->SetVisibility( Visibility );
-	LevelVisibilitySection->SetStartTime( MovieScene->GetPlaybackRange().GetLowerBoundValue() );
-	LevelVisibilitySection->SetEndTime( MovieScene->GetPlaybackRange().GetUpperBoundValue() );
+
+	TRange<FFrameNumber> SectionRange = MovieScene->GetPlaybackRange();
+	LevelVisibilitySection->SetRange(SectionRange);
 
 	int32 RowIndex = -1;
 	for ( const UMovieSceneSection* Section : LevelVisibilityTrack->GetAllSections() )
@@ -82,6 +83,8 @@ void FLevelVisibilityTrackEditor::AddNewSection( UMovieScene* MovieScene, UMovie
 	LevelVisibilitySection->SetRowIndex(RowIndex + 1);
 
 	LevelVisibilityTrack->AddSection( *LevelVisibilitySection );
+
+	return LevelVisibilitySection;
 }
 
 
@@ -100,9 +103,15 @@ void FLevelVisibilityTrackEditor::OnAddTrack()
 	UMovieSceneLevelVisibilityTrack* NewTrack = FocusedMovieScene->AddMasterTrack<UMovieSceneLevelVisibilityTrack>();
 	checkf( NewTrack != nullptr, TEXT("Failed to create new level visibility track.") );
 
-	AddNewSection( FocusedMovieScene, NewTrack, ELevelVisibility::Visible );
-
+	UMovieSceneLevelVisibilitySection* NewSection = AddNewSection( FocusedMovieScene, NewTrack, ELevelVisibility::Visible );
+	if (GetSequencer().IsValid())
+	{
+		GetSequencer()->OnAddTrack(NewTrack);
+	}
 	GetSequencer()->NotifyMovieSceneDataChanged( EMovieSceneDataChangeType::MovieSceneStructureItemAdded );
+	GetSequencer()->EmptySelection();
+	GetSequencer()->SelectSection(NewSection);
+	GetSequencer()->ThrobSectionSelection();
 }
 
 
@@ -137,8 +146,12 @@ void FLevelVisibilityTrackEditor::OnAddNewSection( UMovieSceneTrack* LevelVisibi
 		return;
 	}
 
-	AddNewSection( FocusedMovieScene, LevelVisibilityTrack, Visibility );
+	UMovieSceneLevelVisibilitySection* NewSection = AddNewSection( FocusedMovieScene, LevelVisibilityTrack, Visibility );
+
 	GetSequencer()->NotifyMovieSceneDataChanged( EMovieSceneDataChangeType::MovieSceneStructureItemAdded );
+	GetSequencer()->EmptySelection();
+	GetSequencer()->SelectSection(NewSection);
+	GetSequencer()->ThrobSectionSelection();
 }
 
 #undef LOCTEXT_NAMESPACE
diff --git a/Engine/Source/Editor/MovieSceneTools/Private/TrackEditors/LevelVisibilityTrackEditor.h b/Engine/Source/Editor/MovieSceneTools/Private/TrackEditors/LevelVisibilityTrackEditor.h
--- a/Engine/Source/Editor/MovieSceneTools/Private/TrackEditors/LevelVisibilityTrackEditor.h
+++ b/Engine/Source/Editor/MovieSceneTools/Private/TrackEditors/LevelVisibilityTrackEditor.h
@@ -13,6 +13,7 @@
 #include "MovieSceneTrackEditor.h"
 
 class FMenuBuilder;
+class UMovieSceneLevelVisibilitySection;
 enum class ELevelVisibility : uint8;
 
 /**
@@ -54,7 +55,7 @@ public:
 
 private:
 	/** Adds a new section which spans the length of the owning movie scene with the specified visibility. */
-	void AddNewSection(UMovieScene* MovieScene, UMovieSceneTrack* LevelVisibilityTrack, ELevelVisibility Visibility );
+	UMovieSceneLevelVisibilitySection* AddNewSection(UMovieScene* MovieScene, UMovieSceneTrack* LevelVisibilityTrack, ELevelVisibility Visibility );
 
 	/** Handles when the add track menu item is activated. */
 	void OnAddTrack();
diff --git a/Engine/Source/Editor/MovieSceneTools/Private/TrackEditors/MaterialParameterCollectionTrackEditor.cpp b/Engine/Source/Editor/MovieSceneTools/Private/TrackEditors/MaterialParameterCollectionTrackEditor.cpp
--- a/Engine/Source/Editor/MovieSceneTools/Private/TrackEditors/MaterialParameterCollectionTrackEditor.cpp
+++ b/Engine/Source/Editor/MovieSceneTools/Private/TrackEditors/MaterialParameterCollectionTrackEditor.cpp
@@ -32,7 +32,7 @@ TSharedRef<ISequencerSection> FMaterialParameterCollectionTrackEditor::MakeSecti
 	UMovieSceneParameterSection* ParameterSection = Cast<UMovieSceneParameterSection>(&SectionObject);
 	checkf(ParameterSection != nullptr, TEXT("Unsupported section type."));
 
-	return MakeShareable(new FParameterSection(*ParameterSection, FText::FromName(ParameterSection->GetFName())));
+	return MakeShareable(new FParameterSection(*ParameterSection));
 }
 
 TSharedRef<SWidget> CreateAssetPicker(FOnAssetSelected OnAssetSelected)
@@ -137,6 +137,10 @@ void FMaterialParameterCollectionTrackEditor::AddTrackToSequence(const FAssetDat
 	Track->MPC = MPC;
 	Track->SetDisplayName(FText::FromString(MPC->GetName()));
 
+	if (GetSequencer().IsValid())
+	{
+		GetSequencer()->OnAddTrack(Track);
+	}
 	GetSequencer()->NotifyMovieSceneDataChanged(EMovieSceneDataChangeType::MovieSceneStructureItemAdded);
 }
 
@@ -216,7 +220,7 @@ void FMaterialParameterCollectionTrackEditor::AddScalarParameter(UMovieSceneMate
 		return;
 	}
 
-	float KeyTime = GetTimeForKey();
+	FFrameNumber KeyTime = GetTimeForKey();
 
 	const FScopedTransaction Transaction(LOCTEXT("AddScalarParameter", "Add scalar parameter"));
 	Track->Modify();
@@ -232,7 +236,7 @@ void FMaterialParameterCollectionTrackEditor::AddVectorParameter(UMovieSceneMate
 		return;
 	}
 
-	float KeyTime = GetTimeForKey();
+	FFrameNumber KeyTime = GetTimeForKey();
 
 	const FScopedTransaction Transaction(LOCTEXT("AddVectorParameter", "Add vector parameter"));
 	Track->Modify();
diff --git a/Engine/Source/Editor/MovieSceneTools/Private/TrackEditors/MaterialTrackEditor.cpp b/Engine/Source/Editor/MovieSceneTools/Private/TrackEditors/MaterialTrackEditor.cpp
--- a/Engine/Source/Editor/MovieSceneTools/Private/TrackEditors/MaterialTrackEditor.cpp
+++ b/Engine/Source/Editor/MovieSceneTools/Private/TrackEditors/MaterialTrackEditor.cpp
@@ -27,7 +27,7 @@ TSharedRef<ISequencerSection> FMaterialTrackEditor::MakeSectionInterface( UMovie
 	UMovieSceneParameterSection* ParameterSection = Cast<UMovieSceneParameterSection>(&SectionObject);
 	checkf( ParameterSection != nullptr, TEXT("Unsupported section type.") );
 
-	return MakeShareable(new FParameterSection( *ParameterSection, FText::FromName(ParameterSection->GetFName())));
+	return MakeShareable(new FParameterSection( *ParameterSection ));
 }
 
 
@@ -147,7 +147,7 @@ UMaterial* FMaterialTrackEditor::GetMaterialForTrack( FGuid ObjectBinding, UMovi
 
 void FMaterialTrackEditor::AddScalarParameter( FGuid ObjectBinding, UMovieSceneMaterialTrack* MaterialTrack, FName ParameterName )
 {
-	float KeyTime = GetTimeForKey();
+	FFrameNumber KeyTime = GetTimeForKey();
 
 	UMaterialInterface* Material = GetMaterialInterfaceForTrack(ObjectBinding, MaterialTrack);
 	if (Material != nullptr)
@@ -164,7 +164,7 @@ void FMaterialTrackEditor::AddScalarParameter( FGuid ObjectBinding, UMovieSceneM
 
 void FMaterialTrackEditor::AddColorParameter( FGuid ObjectBinding, UMovieSceneMaterialTrack* MaterialTrack, FName ParameterName )
 {
-	float KeyTime = GetTimeForKey();
+	FFrameNumber KeyTime = GetTimeForKey();
 
 	UMaterialInterface* Material = GetMaterialInterfaceForTrack( ObjectBinding, MaterialTrack );
 	if ( Material != nullptr )
diff --git a/Engine/Source/Editor/MovieSceneTools/Private/TrackEditors/ParticleParameterTrackEditor.cpp b/Engine/Source/Editor/MovieSceneTools/Private/TrackEditors/ParticleParameterTrackEditor.cpp
--- a/Engine/Source/Editor/MovieSceneTools/Private/TrackEditors/ParticleParameterTrackEditor.cpp
+++ b/Engine/Source/Editor/MovieSceneTools/Private/TrackEditors/ParticleParameterTrackEditor.cpp
@@ -31,7 +31,7 @@ TSharedRef<ISequencerSection> FParticleParameterTrackEditor::MakeSectionInterfac
 	UMovieSceneParameterSection* ParameterSection = Cast<UMovieSceneParameterSection>(&SectionObject);
 	checkf( ParameterSection != nullptr, TEXT("Unsupported section type.") );
 
-	return MakeShareable(new FParameterSection( *ParameterSection, FText::FromName(ParameterSection->GetFName())));
+	return MakeShareable(new FParameterSection( *ParameterSection ));
 }
 
 
@@ -142,7 +142,7 @@ void FParticleParameterTrackEditor::AddScalarParameter( FGuid ObjectBinding, UMo
 	{
 		const FScopedTransaction Transaction(LOCTEXT("AddScalarParameter", "Add scalar parameter"));
 
-		float KeyTime = GetTimeForKey();
+		FFrameNumber KeyTime = GetTimeForKey();
 		float Value;
 		ParticleSystemComponent->GetFloatParameter(ParameterName, Value);
 
@@ -161,7 +161,7 @@ void FParticleParameterTrackEditor::AddVectorParameter( FGuid ObjectBinding, UMo
 	{
 		const FScopedTransaction Transaction(LOCTEXT("AddVectorParameter", "Add vector parameter"));
 
-		float KeyTime = GetTimeForKey();
+		FFrameNumber KeyTime = GetTimeForKey();
 		FVector Value;
 		ParticleSystemComponent->GetVectorParameter(ParameterName, Value);
 
@@ -180,7 +180,7 @@ void FParticleParameterTrackEditor::AddColorParameter( FGuid ObjectBinding, UMov
 	{
 		const FScopedTransaction Transaction(LOCTEXT("AddColorParameter", "Add color parameter"));
 
-		float KeyTime = GetTimeForKey();
+		FFrameNumber KeyTime = GetTimeForKey();
 		FLinearColor Value;
 		ParticleSystemComponent->GetColorParameter(ParameterName, Value);
 
diff --git a/Engine/Source/Editor/MovieSceneTools/Private/TrackEditors/ParticleTrackEditor.cpp b/Engine/Source/Editor/MovieSceneTools/Private/TrackEditors/ParticleTrackEditor.cpp
--- a/Engine/Source/Editor/MovieSceneTools/Private/TrackEditors/ParticleTrackEditor.cpp
+++ b/Engine/Source/Editor/MovieSceneTools/Private/TrackEditors/ParticleTrackEditor.cpp
@@ -16,7 +16,6 @@
 #include "CommonMovieSceneTools.h"
 #include "Particles/ParticleLODLevel.h"
 #include "Particles/ParticleModuleRequired.h"
-#include "EnumKeyArea.h"
 #include "Matinee/InterpTrackToggle.h"
 #include "MatineeImportTools.h"
 
@@ -35,11 +34,6 @@ FParticleSection::FParticleSection( UMovieSceneSection& InSection, TSharedRef<IS
 	: Section( InSection )
 	, OwningSequencerPtr( InOwningSequencer )
 {
-	ParticleKeyEnum = FindObject<UEnum>( ANY_PACKAGE, TEXT( "EParticleKey" ) );
-	checkf( ParticleKeyEnum != nullptr, TEXT( "FParticleSection could not find the EParticleKey UEnum by name." ) )
-
-	LeftKeyBrush = FEditorStyle::GetBrush( "Sequencer.KeyLeft" );
-	RightKeyBrush = FEditorStyle::GetBrush( "Sequencer.KeyRight" );
 }
 
 
@@ -58,12 +52,6 @@ float FParticleSection::GetSectionHeight() const
 	return (float)AnimatableParticleEditorConstants::ParticleTrackHeight;
 }
 
-void FParticleSection::GenerateSectionLayout( class ISectionLayoutBuilder& LayoutBuilder ) const
-{
-	UMovieSceneParticleSection* ParticleSection = Cast<UMovieSceneParticleSection>( &Section );
-	LayoutBuilder.SetSectionAsKeyArea( MakeShareable( new FEnumKeyArea( ParticleSection->GetParticleCurve(), ParticleSection, ParticleKeyEnum ) ) );
-}
-
 int32 FParticleSection::OnPaintSection( FSequencerSectionPainter& InPainter ) const
 {
 	TSharedPtr<ISequencer> OwningSequencer = OwningSequencerPtr.Pin();
@@ -81,7 +69,7 @@ int32 FParticleSection::OnPaintSection( FSequencerSectionPainter& InPainter ) co
 
 	// @todo Sequencer - These values should be cached and then refreshed only when the particle system changes.
 	bool bIsLooping = false;
-	float LastEmitterEndTime = 0;
+	double LastEmitterEndTime = 0;
 	UMovieSceneParticleSection* ParticleSection = Cast<UMovieSceneParticleSection>( &Section );
 	if ( ParticleSection != nullptr )
 	{
@@ -116,7 +104,7 @@ int32 FParticleSection::OnPaintSection( FSequencerSectionPainter& InPainter ) co
 						{
 							UParticleModuleRequired* RequiredModule = Emitter->GetLODLevel( 0 )->RequiredModule;
 							bIsLooping |= RequiredModule->EmitterLoops == 0;
-							LastEmitterEndTime = FMath::Max( LastEmitterEndTime, RequiredModule->EmitterDelay + RequiredModule->EmitterDuration );
+							LastEmitterEndTime = FMath::Max( LastEmitterEndTime, double(RequiredModule->EmitterDelay) + RequiredModule->EmitterDuration );
 						}
 					}
 				}
@@ -125,53 +113,64 @@ int32 FParticleSection::OnPaintSection( FSequencerSectionPainter& InPainter ) co
 	}
 
 	// @todo Sequencer - This should only draw the visible ranges.
-	TArray<TRange<float>> DrawRanges;
-	TOptional<float> CurrentRangeStart;
-	for ( auto KeyIterator = ParticleSection->GetParticleCurve().GetKeyIterator(); KeyIterator; ++KeyIterator )
+	TArray<TRange<float>>   DrawRanges;
+	TOptional<double> CurrentRangeStart;
+
+	if (ParticleSection != nullptr)
 	{
-		FIntegralKey Key = *KeyIterator;
-		if ( (EParticleKey::Type)Key.Value == EParticleKey::Activate )
+		TMovieSceneChannel<const uint8> ChannelInterface = ParticleSection->ParticleKeys.GetInterface();
+
+		TArrayView<const FFrameNumber> Times  = ChannelInterface.GetTimes();
+		TArrayView<const uint8>        Values = ChannelInterface.GetValues();
+
+		for (int32 Index = 0; Index < Times.Num(); ++Index)
 		{
-			if ( CurrentRangeStart.IsSet() == false )
-			{
-				CurrentRangeStart = Key.Time;
-			}
-			else
+			const double       Time  = Times[Index] / TimeToPixelConverter.GetFrameResolution();
+			const EParticleKey Value = (EParticleKey)Values[Index];
+
+			if ( Value == EParticleKey::Activate )
 			{
-				if ( bIsLooping == false )
+				if ( CurrentRangeStart.IsSet() == false )
 				{
-					if ( Key.Time > CurrentRangeStart.GetValue() + LastEmitterEndTime )
-					{
-						DrawRanges.Add( TRange<float>( CurrentRangeStart.GetValue(), CurrentRangeStart.GetValue() + LastEmitterEndTime ) );
-					}
-					else
+					CurrentRangeStart = Time;
+				}
+				else
+				{
+					if ( bIsLooping == false )
 					{
-						DrawRanges.Add( TRange<float>( CurrentRangeStart.GetValue(), Key.Time ) );
+						if ( Time > CurrentRangeStart.GetValue() + LastEmitterEndTime )
+						{
+							DrawRanges.Add( TRange<float>( CurrentRangeStart.GetValue(), CurrentRangeStart.GetValue() + LastEmitterEndTime ) );
+						}
+						else
+						{
+							DrawRanges.Add( TRange<float>( CurrentRangeStart.GetValue(), Time ) );
+						}
+						CurrentRangeStart = Time;
 					}
-					CurrentRangeStart = Key.Time;
 				}
 			}
-		}
-		if ( (EParticleKey::Type)Key.Value == EParticleKey::Deactivate )
-		{
-			if ( CurrentRangeStart.IsSet() )
+			if ( Value == EParticleKey::Deactivate )
 			{
-				if (bIsLooping)
-				{
-					DrawRanges.Add( TRange<float>( CurrentRangeStart.GetValue(), Key.Time ) );
-				}
-				else
+				if ( CurrentRangeStart.IsSet() )
 				{
-					if ( Key.Time > CurrentRangeStart.GetValue() + LastEmitterEndTime )
+					if (bIsLooping)
 					{
-						DrawRanges.Add( TRange<float>( CurrentRangeStart.GetValue(), CurrentRangeStart.GetValue() + LastEmitterEndTime ) );
+						DrawRanges.Add( TRange<float>( CurrentRangeStart.GetValue(), Time ) );
 					}
 					else
 					{
-						DrawRanges.Add( TRange<float>( CurrentRangeStart.GetValue(), Key.Time ) );
+						if ( Time > CurrentRangeStart.GetValue() + LastEmitterEndTime )
+						{
+							DrawRanges.Add( TRange<float>( CurrentRangeStart.GetValue(), CurrentRangeStart.GetValue() + LastEmitterEndTime ) );
+						}
+						else
+						{
+							DrawRanges.Add( TRange<float>( CurrentRangeStart.GetValue(), Time ) );
+						}
 					}
+					CurrentRangeStart.Reset();
 				}
-				CurrentRangeStart.Reset();
 			}
 		}
 	}
@@ -189,8 +188,8 @@ int32 FParticleSection::OnPaintSection( FSequencerSectionPainter& InPainter ) co
 
 	for ( const TRange<float>& DrawRange : DrawRanges )
 	{
-		float XOffset = TimeToPixelConverter.TimeToPixel(DrawRange.GetLowerBoundValue());
-		float XSize = TimeToPixelConverter.TimeToPixel(DrawRange.GetUpperBoundValue()) - XOffset;
+		float XOffset = TimeToPixelConverter.SecondsToPixel(DrawRange.GetLowerBoundValue());
+		float XSize   = TimeToPixelConverter.SecondsToPixel(DrawRange.GetUpperBoundValue()) - XOffset;
 		FSlateDrawElement::MakeBox(
 			InPainter.DrawElements,
 			InPainter.LayerId,
@@ -212,43 +211,6 @@ int32 FParticleSection::OnPaintSection( FSequencerSectionPainter& InPainter ) co
 }
 
 
-const FSlateBrush* FParticleSection::GetKeyBrush( FKeyHandle KeyHandle ) const
-{
-	UMovieSceneParticleSection* ParticleSection = Cast<UMovieSceneParticleSection>( &Section );
-	if ( ParticleSection != nullptr )
-	{
-		FIntegralKey ParticleKey = ParticleSection->GetParticleCurve().GetKey(KeyHandle);
-		if ( (EParticleKey::Type)ParticleKey.Value == EParticleKey::Activate )
-		{
-			return LeftKeyBrush;
-		}
-		else if ( (EParticleKey::Type)ParticleKey.Value == EParticleKey::Deactivate )
-		{
-			return RightKeyBrush;
-		}
-	}
-	return nullptr;
-}
-
-FVector2D FParticleSection::GetKeyBrushOrigin( FKeyHandle KeyHandle ) const
-{
-	UMovieSceneParticleSection* ParticleSection = Cast<UMovieSceneParticleSection>( &Section );
-	if ( ParticleSection != nullptr )
-	{
-		FIntegralKey ParticleKey = ParticleSection->GetParticleCurve().GetKey(KeyHandle);
-		if ( (EParticleKey::Type)ParticleKey.Value == EParticleKey::Activate )
-		{
-			return FVector2D(-1.0f, 1.0f);
-		}
-		else if ( (EParticleKey::Type)ParticleKey.Value == EParticleKey::Deactivate )
-		{
-			return FVector2D(1.0f, 1.0f);
-		}
-	}
-	return FVector2D(0.0f, 0.0f);
-}
-
-
 FParticleTrackEditor::FParticleTrackEditor( TSharedRef<ISequencer> InSequencer )
 	: FMovieSceneTrackEditor( InSequencer ) 
 { }
@@ -303,7 +265,7 @@ void FParticleTrackEditor::AddParticleKey( const FGuid ObjectGuid )
 }
 
 
-FKeyPropertyResult FParticleTrackEditor::AddKeyInternal( float KeyTime, UObject* Object )
+FKeyPropertyResult FParticleTrackEditor::AddKeyInternal( FFrameNumber KeyTime, UObject* Object )
 {
 	FKeyPropertyResult KeyPropertyResult;
 
diff --git a/Engine/Source/Editor/MovieSceneTools/Private/TrackEditors/ParticleTrackEditor.h b/Engine/Source/Editor/MovieSceneTools/Private/TrackEditors/ParticleTrackEditor.h
--- a/Engine/Source/Editor/MovieSceneTools/Private/TrackEditors/ParticleTrackEditor.h
+++ b/Engine/Source/Editor/MovieSceneTools/Private/TrackEditors/ParticleTrackEditor.h
@@ -58,7 +58,7 @@ public:
 private:
 
 	/** Delegate for AnimatablePropertyChanged in AddKey. */
-	FKeyPropertyResult AddKeyInternal( float KeyTime, UObject* Object);
+	FKeyPropertyResult AddKeyInternal( FFrameNumber KeyTime, UObject* Object);
 };
 
 
@@ -79,10 +79,7 @@ public:
 	// ISequencerSection interface
 	virtual UMovieSceneSection* GetSectionObject() override;
 	virtual float GetSectionHeight() const override;
-	virtual void GenerateSectionLayout( class ISectionLayoutBuilder& LayoutBuilder ) const override;
 	virtual int32 OnPaintSection( FSequencerSectionPainter& InPainter ) const override;
-	virtual const FSlateBrush* GetKeyBrush( FKeyHandle KeyHandle ) const override;
-	virtual FVector2D GetKeyBrushOrigin( FKeyHandle KeyHandle ) const override;
 	virtual bool SectionIsResizable() const override { return false; }
 
 private:
@@ -92,10 +89,4 @@ private:
 
 	/** The sequencer that owns this section */
 	TWeakPtr<ISequencer> OwningSequencerPtr;
-
-	/** The UEnum for the EParticleKey enum */
-	const UEnum* ParticleKeyEnum;
-
-	const FSlateBrush* LeftKeyBrush;
-	const FSlateBrush* RightKeyBrush;
 };
diff --git a/Engine/Source/Editor/MovieSceneTools/Private/TrackEditors/PathTrackEditor.cpp b/Engine/Source/Editor/MovieSceneTools/Private/TrackEditors/PathTrackEditor.cpp
--- a/Engine/Source/Editor/MovieSceneTools/Private/TrackEditors/PathTrackEditor.cpp
+++ b/Engine/Source/Editor/MovieSceneTools/Private/TrackEditors/PathTrackEditor.cpp
@@ -10,7 +10,6 @@
 #include "ISectionLayoutBuilder.h"
 #include "ActorEditorUtils.h"
 #include "Components/SplineComponent.h"
-#include "FloatCurveKeyArea.h"
 
 
 #define LOCTEXT_NAMESPACE "FPathTrackEditor"
@@ -55,13 +54,6 @@ public:
 		return FText::GetEmpty(); 
 	}
 
-	virtual void GenerateSectionLayout( class ISectionLayoutBuilder& LayoutBuilder ) const override
-	{
-		UMovieScene3DPathSection* PathSection = Cast<UMovieScene3DPathSection>( &Section );
-
-		LayoutBuilder.AddKeyArea("Timing", LOCTEXT("TimingArea", "Timing"), MakeShareable( new FFloatCurveKeyArea ( &PathSection->GetTimingCurve( ), PathSection ) ) );
-	}
-
 	virtual int32 OnPaintSection( FSequencerSectionPainter& InPainter ) const override 
 	{
 		return InPainter.PaintSectionBackground();
@@ -195,7 +187,7 @@ void F3DPathTrackEditor::ActorSocketPicked(const FName SocketName, USceneCompone
 	}
 }
 
-FKeyPropertyResult F3DPathTrackEditor::AddKeyInternal( float KeyTime, const TArray<TWeakObjectPtr<UObject>> Objects, FActorPickerID ActorPickerID)
+FKeyPropertyResult F3DPathTrackEditor::AddKeyInternal( FFrameNumber KeyTime, const TArray<TWeakObjectPtr<UObject>> Objects, FActorPickerID ActorPickerID)
 {
 	FKeyPropertyResult KeyPropertyResult;
 
@@ -234,12 +226,11 @@ FKeyPropertyResult F3DPathTrackEditor::AddKeyInternal( float KeyTime, const TArr
 			if (ensure(Track))
 			{
 				// Clamp to next path section's start time or the end of the current sequencer view range
-				float PathEndTime = GetSequencer()->GetViewRange().GetUpperBoundValue();
+				FFrameNumber PathEndTime = ( GetSequencer()->GetViewRange().GetUpperBoundValue() * Track->GetTypedOuter<UMovieScene>()->GetFrameResolution() ).FrameNumber;
 	
-				for (int32 PathSectionIndex = 0; PathSectionIndex < Track->GetAllSections().Num(); ++PathSectionIndex)
+				for (UMovieSceneSection* Section : Track->GetAllSections())
 				{
-					float StartTime = Track->GetAllSections()[PathSectionIndex]->GetStartTime();
-					float EndTime = Track->GetAllSections()[PathSectionIndex]->GetEndTime();
+					FFrameNumber StartTime = Section->HasStartFrame() ? Section->GetInclusiveStartFrame() : TNumericLimits<int32>::Lowest();
 					if (KeyTime < StartTime)
 					{
 						if (PathEndTime > StartTime)
@@ -249,7 +240,7 @@ FKeyPropertyResult F3DPathTrackEditor::AddKeyInternal( float KeyTime, const TArr
 					}
 				}
 
-				Cast<UMovieScene3DPathTrack>(Track)->AddConstraint( KeyTime, PathEndTime, NAME_None, NAME_None, ConstraintBindingID );
+				Cast<UMovieScene3DPathTrack>(Track)->AddConstraint( KeyTime, (PathEndTime - KeyTime).Value, NAME_None, NAME_None, ConstraintBindingID );
 				KeyPropertyResult.bTrackModified = true;
 			}
 		}
diff --git a/Engine/Source/Editor/MovieSceneTools/Private/TrackEditors/PathTrackEditor.h b/Engine/Source/Editor/MovieSceneTools/Private/TrackEditors/PathTrackEditor.h
--- a/Engine/Source/Editor/MovieSceneTools/Private/TrackEditors/PathTrackEditor.h
+++ b/Engine/Source/Editor/MovieSceneTools/Private/TrackEditors/PathTrackEditor.h
@@ -56,7 +56,7 @@ public:
 private:
 
 	/** Delegate for AnimatablePropertyChanged in AddKey */
-	FKeyPropertyResult AddKeyInternal(float KeyTime, const TArray<TWeakObjectPtr<UObject>> Objects, FActorPickerID ActorPickerID);
+	FKeyPropertyResult AddKeyInternal(FFrameNumber KeyTime, const TArray<TWeakObjectPtr<UObject>> Objects, FActorPickerID ActorPickerID);
 };
 
 
diff --git a/Engine/Source/Editor/MovieSceneTools/Private/TrackEditors/PropertyTrackEditors/ActorReferencePropertyTrackEditor.cpp b/Engine/Source/Editor/MovieSceneTools/Private/TrackEditors/PropertyTrackEditors/ActorReferencePropertyTrackEditor.cpp
--- a/Engine/Source/Editor/MovieSceneTools/Private/TrackEditors/PropertyTrackEditors/ActorReferencePropertyTrackEditor.cpp
+++ b/Engine/Source/Editor/MovieSceneTools/Private/TrackEditors/PropertyTrackEditors/ActorReferencePropertyTrackEditor.cpp
@@ -2,33 +2,32 @@
 
 #include "TrackEditors/PropertyTrackEditors/ActorReferencePropertyTrackEditor.h"
 #include "GameFramework/Actor.h"
-#include "Sections/ActorReferencePropertySection.h"
 
 
-TSharedRef<ISequencerTrackEditor> FActorReferencePropertyTrackEditor::CreateTrackEditor( TSharedRef<ISequencer> OwningSequencer )
+namespace MovieScene
 {
-	return MakeShareable(new FActorReferencePropertyTrackEditor(OwningSequencer));
+	bool ValueExistsAtTime(const FMovieSceneActorReferenceData* Channel, FFrameTime Time, FMovieSceneObjectBindingID Value, float Tolerance = KINDA_SMALL_NUMBER)
+	{
+		return Channel->Evaluate(Time) == Value;
+	}
 }
 
-
-TSharedRef<ISequencerSection> FActorReferencePropertyTrackEditor::MakeSectionInterface(UMovieSceneSection& SectionObject, UMovieSceneTrack& Track, FGuid ObjectBinding)
+TSharedRef<ISequencerTrackEditor> FActorReferencePropertyTrackEditor::CreateTrackEditor( TSharedRef<ISequencer> OwningSequencer )
 {
-	UMovieScenePropertyTrack* PropertyTrack = Cast<UMovieScenePropertyTrack>(&Track);
-	checkf(PropertyTrack != nullptr, TEXT("Incompatible track in FActorReferencePropertyTrackEditor"));
-	return MakeShareable(new FActorReferencePropertySection(GetSequencer().Get(), ObjectBinding, PropertyTrack->GetPropertyName(), PropertyTrack->GetPropertyPath(), SectionObject, Track.GetDisplayName()));
+	return MakeShareable(new FActorReferencePropertyTrackEditor(OwningSequencer));
 }
 
-
-void FActorReferencePropertyTrackEditor::GenerateKeysFromPropertyChanged( const FPropertyChangedParams& PropertyChangedParams, TArray<FGuid>& NewGeneratedKeys, TArray<FGuid>& DefaultGeneratedKeys )
+void FActorReferencePropertyTrackEditor::GenerateKeysFromPropertyChanged( const FPropertyChangedParams& PropertyChangedParams, FGeneratedTrackKeys& OutGeneratedKeys )
 {
-
 	AActor* NewReferencedActor = PropertyChangedParams.GetPropertyValue<AActor*>();
 	if ( NewReferencedActor != nullptr )
 	{
 		FGuid ActorGuid = GetSequencer()->GetHandleToObject( NewReferencedActor );
 		if ( ActorGuid.IsValid() )
 		{
-			NewGeneratedKeys.Add( ActorGuid );
+			FMovieSceneActorReferenceKey NewKey;
+			NewKey.Object = FMovieSceneObjectBindingID(ActorGuid, MovieSceneSequenceID::Root, EMovieSceneObjectBindingSpace::Local);
+			OutGeneratedKeys.Add(FMovieSceneChannelValueSetter::Create<FMovieSceneActorReferenceData>(0, NewKey, true));
 		}
 	}
 }
diff --git a/Engine/Source/Editor/MovieSceneTools/Private/TrackEditors/PropertyTrackEditors/ActorReferencePropertyTrackEditor.h b/Engine/Source/Editor/MovieSceneTools/Private/TrackEditors/PropertyTrackEditors/ActorReferencePropertyTrackEditor.h
--- a/Engine/Source/Editor/MovieSceneTools/Private/TrackEditors/PropertyTrackEditors/ActorReferencePropertyTrackEditor.h
+++ b/Engine/Source/Editor/MovieSceneTools/Private/TrackEditors/PropertyTrackEditors/ActorReferencePropertyTrackEditor.h
@@ -16,7 +16,7 @@
  * A property track editor for actor references.
  */
 class FActorReferencePropertyTrackEditor
-	: public FPropertyTrackEditor<UMovieSceneActorReferenceTrack, UMovieSceneActorReferenceSection, FGuid>
+	: public FPropertyTrackEditor<UMovieSceneActorReferenceTrack>
 {
 public:
 
@@ -46,13 +46,9 @@ public:
 	 */
 	static TSharedRef<ISequencerTrackEditor> CreateTrackEditor(TSharedRef<ISequencer> OwningSequencer);
 
-	//~ ISequencerTrackEditor interface
-
-	virtual TSharedRef<ISequencerSection> MakeSectionInterface(UMovieSceneSection& SectionObject, UMovieSceneTrack& Track, FGuid ObjectBinding) override;
-
 protected:
 
 	//~ FPropertyTrackEditor interface
 
-	virtual void GenerateKeysFromPropertyChanged(const FPropertyChangedParams& PropertyChangedParams, TArray<FGuid>& NewGeneratedKeys, TArray<FGuid>& DefaultGeneratedKeys) override;
+	virtual void GenerateKeysFromPropertyChanged(const FPropertyChangedParams& PropertyChangedParams, FGeneratedTrackKeys& OutGeneratedKeys) override;
 };
diff --git a/Engine/Source/Editor/MovieSceneTools/Private/TrackEditors/PropertyTrackEditors/BoolPropertyTrackEditor.cpp b/Engine/Source/Editor/MovieSceneTools/Private/TrackEditors/PropertyTrackEditors/BoolPropertyTrackEditor.cpp
--- a/Engine/Source/Editor/MovieSceneTools/Private/TrackEditors/PropertyTrackEditors/BoolPropertyTrackEditor.cpp
+++ b/Engine/Source/Editor/MovieSceneTools/Private/TrackEditors/PropertyTrackEditors/BoolPropertyTrackEditor.cpp
@@ -12,13 +12,12 @@ TSharedRef<ISequencerTrackEditor> FBoolPropertyTrackEditor::CreateTrackEditor(TS
 
 TSharedRef<ISequencerSection> FBoolPropertyTrackEditor::MakeSectionInterface(UMovieSceneSection& SectionObject, UMovieSceneTrack& Track, FGuid ObjectBinding)
 {
-	UMovieScenePropertyTrack* PropertyTrack = Cast<UMovieScenePropertyTrack>(&Track);
-	checkf(PropertyTrack != nullptr, TEXT("Incompatible track in FBoolPropertyTrackEditor"));
-	return MakeShareable(new FBoolPropertySection(GetSequencer().Get(), ObjectBinding, PropertyTrack->GetPropertyName(), PropertyTrack->GetPropertyPath(), SectionObject, Track.GetDisplayName()));
+	return MakeShared<FBoolPropertySection>(SectionObject);
 }
 
 
-void FBoolPropertyTrackEditor::GenerateKeysFromPropertyChanged( const FPropertyChangedParams& PropertyChangedParams, TArray<bool>& NewGeneratedKeys, TArray<bool>& DefaultGeneratedKeys )
+void FBoolPropertyTrackEditor::GenerateKeysFromPropertyChanged( const FPropertyChangedParams& PropertyChangedParams, FGeneratedTrackKeys& OutGeneratedKeys )
 {
-	NewGeneratedKeys.Add( PropertyChangedParams.GetPropertyValue<bool>() );
+	const bool KeyedValue = PropertyChangedParams.GetPropertyValue<bool>();
+	OutGeneratedKeys.Add(FMovieSceneChannelValueSetter::Create<FMovieSceneBoolChannel>(0, KeyedValue, true));
 }
diff --git a/Engine/Source/Editor/MovieSceneTools/Private/TrackEditors/PropertyTrackEditors/BytePropertyTrackEditor.cpp b/Engine/Source/Editor/MovieSceneTools/Private/TrackEditors/PropertyTrackEditors/BytePropertyTrackEditor.cpp
--- a/Engine/Source/Editor/MovieSceneTools/Private/TrackEditors/PropertyTrackEditors/BytePropertyTrackEditor.cpp
+++ b/Engine/Source/Editor/MovieSceneTools/Private/TrackEditors/PropertyTrackEditors/BytePropertyTrackEditor.cpp
@@ -1,7 +1,6 @@
 // Copyright 1998-2018 Epic Games, Inc. All Rights Reserved.
 
 #include "TrackEditors/PropertyTrackEditors/BytePropertyTrackEditor.h"
-#include "Sections/BytePropertySection.h"
 #include "UObject/EnumProperty.h"
 
 
@@ -11,15 +10,6 @@ TSharedRef<ISequencerTrackEditor> FBytePropertyTrackEditor::CreateTrackEditor( T
 }
 
 
-TSharedRef<ISequencerSection> FBytePropertyTrackEditor::MakeSectionInterface(UMovieSceneSection& SectionObject, UMovieSceneTrack& Track, FGuid ObjectBinding)
-{
-	UMovieSceneByteTrack* ByteTrack = Cast<UMovieSceneByteTrack>(&Track);
-	checkf(ByteTrack != nullptr, TEXT("Incompatible track in FBoolPropertyTrackEditor"));
-	return MakeShareable(new FBytePropertySection(GetSequencer().Get(), ObjectBinding, ByteTrack->GetPropertyName(), ByteTrack->GetPropertyPath(),
-		SectionObject, Track.GetDisplayName(), ByteTrack->GetEnum()));
-}
-
-
 UEnum* GetEnumForByteTrack(TSharedPtr<ISequencer> Sequencer, const FGuid& OwnerObjectHandle, FName PropertyName, UMovieSceneByteTrack* ByteTrack)
 {
 	TSet<UEnum*> PropertyEnums;
@@ -83,7 +73,8 @@ UMovieSceneTrack* FBytePropertyTrackEditor::AddTrack(UMovieScene* FocusedMovieSc
 }
 
 
-void FBytePropertyTrackEditor::GenerateKeysFromPropertyChanged( const FPropertyChangedParams& PropertyChangedParams, TArray<uint8>& NewGeneratedKeys, TArray<uint8>& DefaultGeneratedKeys )
+void FBytePropertyTrackEditor::GenerateKeysFromPropertyChanged( const FPropertyChangedParams& PropertyChangedParams, FGeneratedTrackKeys& OutGeneratedKeys )
 {
-	NewGeneratedKeys.Add( PropertyChangedParams.GetPropertyValue<uint8>() );
+	uint8 KeyedValue = PropertyChangedParams.GetPropertyValue<uint8>();
+	OutGeneratedKeys.Add(FMovieSceneChannelValueSetter::Create<FMovieSceneByteChannel>(0, KeyedValue, true));
 }
diff --git a/Engine/Source/Editor/MovieSceneTools/Private/TrackEditors/PropertyTrackEditors/BytePropertyTrackEditor.h b/Engine/Source/Editor/MovieSceneTools/Private/TrackEditors/PropertyTrackEditors/BytePropertyTrackEditor.h
--- a/Engine/Source/Editor/MovieSceneTools/Private/TrackEditors/PropertyTrackEditors/BytePropertyTrackEditor.h
+++ b/Engine/Source/Editor/MovieSceneTools/Private/TrackEditors/PropertyTrackEditors/BytePropertyTrackEditor.h
@@ -16,7 +16,7 @@
 * A property track editor for byte and enumerations.
 */
 class FBytePropertyTrackEditor
-	: public FPropertyTrackEditor<UMovieSceneByteTrack, UMovieSceneByteSection, uint8>
+	: public FPropertyTrackEditor<UMovieSceneByteTrack>
 {
 public:
 
@@ -53,11 +53,10 @@ public:
 	//~ ISequencerTrackEditor interface
 
 	virtual UMovieSceneTrack* AddTrack(UMovieScene* FocusedMovieScene, const FGuid& ObjectHandle, TSubclassOf<class UMovieSceneTrack> TrackClass, FName UniqueTypeName) override;
-	virtual TSharedRef<ISequencerSection> MakeSectionInterface(UMovieSceneSection& SectionObject, UMovieSceneTrack& Track, FGuid ObjectBinding) override;
 
 protected:
 
 	//~ FPropertyTrackEditor interface
 
-	virtual void GenerateKeysFromPropertyChanged(const FPropertyChangedParams& PropertyChangedParams, TArray<uint8>& NewGeneratedKeys, TArray<uint8>& DefaultGeneratedKeys) override;
+	virtual void GenerateKeysFromPropertyChanged(const FPropertyChangedParams& PropertyChangedParams, FGeneratedTrackKeys& OutGeneratedKeys) override;
 };
diff --git a/Engine/Source/Editor/MovieSceneTools/Private/TrackEditors/PropertyTrackEditors/ColorPropertyTrackEditor.cpp b/Engine/Source/Editor/MovieSceneTools/Private/TrackEditors/PropertyTrackEditors/ColorPropertyTrackEditor.cpp
--- a/Engine/Source/Editor/MovieSceneTools/Private/TrackEditors/PropertyTrackEditors/ColorPropertyTrackEditor.cpp
+++ b/Engine/Source/Editor/MovieSceneTools/Private/TrackEditors/PropertyTrackEditors/ColorPropertyTrackEditor.cpp
@@ -25,11 +25,11 @@ TSharedRef<ISequencerSection> FColorPropertyTrackEditor::MakeSectionInterface(UM
 {
 	UMovieScenePropertyTrack* PropertyTrack = Cast<UMovieScenePropertyTrack>(&Track);
 	checkf(PropertyTrack != nullptr, TEXT("Incompatible track in FColorPropertyTrackEditor"));
-	return MakeShareable(new FColorPropertySection(GetSequencer().Get(), ObjectBinding, PropertyTrack->GetPropertyName(), PropertyTrack->GetPropertyPath(), SectionObject, Track.GetDisplayName()));
+	return MakeShared<FColorPropertySection>(SectionObject, ObjectBinding, GetSequencer());
 }
 
 
-void FColorPropertyTrackEditor::GenerateKeysFromPropertyChanged( const FPropertyChangedParams& PropertyChangedParams, TArray<FColorKey>& NewGeneratedKeys, TArray<FColorKey>& DefaultGeneratedKeys )
+void FColorPropertyTrackEditor::GenerateKeysFromPropertyChanged( const FPropertyChangedParams& PropertyChangedParams, FGeneratedTrackKeys& OutGeneratedKeys )
 {
 	UProperty* Property = PropertyChangedParams.PropertyPath.GetLeafMostProperty().Property.Get();
 	if (!Property)
@@ -42,8 +42,6 @@ void FColorPropertyTrackEditor::GenerateKeysFromPropertyChanged( const FProperty
 	FName PropertyName = Property->GetFName();
 
 	bool bIsFColor = StructName == NAME_Color;
-	bool bIsFLinearColor = StructName == NAME_LinearColor;
-	bool bIsSlateColor = StructName == FName( "SlateColor" );
 
 	FLinearColor ColorValue;
 
@@ -61,19 +59,18 @@ void FColorPropertyTrackEditor::GenerateKeysFromPropertyChanged( const FProperty
 		ColorValue.A = 1;
 	}
 
-	FName ChannelName = PropertyChangedParams.StructPropertyNameToKey;
+	FPropertyPath StructPath = PropertyChangedParams.StructPathToKey;
+	FName ChannelName = StructPath.GetNumProperties() != 0 ? StructPath.GetLeafMostProperty().Property->GetFName() : NAME_None;
 
-	TArray<FColorKey>& RedKeys = ChannelName == NAME_None || ChannelName == RedName || ChannelName == SpecifiedColorName ? NewGeneratedKeys : DefaultGeneratedKeys;
-	RedKeys.Add( FColorKey( EKeyColorChannel::Red, ColorValue.R, bIsSlateColor ) );
+	const bool bKeyRed   =  ChannelName == NAME_None || ChannelName == RedName   || ChannelName == SpecifiedColorName;
+	const bool bKeyGreen =  ChannelName == NAME_None || ChannelName == GreenName || ChannelName == SpecifiedColorName;
+	const bool bKeyBlue  =  ChannelName == NAME_None || ChannelName == BlueName  || ChannelName == SpecifiedColorName;
+	const bool bKeyAlpha =  ChannelName == NAME_None || ChannelName == AlphaName || ChannelName == SpecifiedColorName;
 
-	TArray<FColorKey>& GreenKeys = ChannelName == NAME_None || ChannelName == GreenName || ChannelName == SpecifiedColorName ? NewGeneratedKeys : DefaultGeneratedKeys;
-	GreenKeys.Add( FColorKey( EKeyColorChannel::Green, ColorValue.G, bIsSlateColor ) );
-
-	TArray<FColorKey>& BlueKeys =  ChannelName == NAME_None || ChannelName == BlueName || ChannelName == SpecifiedColorName ? NewGeneratedKeys : DefaultGeneratedKeys;
-	BlueKeys.Add( FColorKey( EKeyColorChannel::Blue, ColorValue.B, bIsSlateColor ) );
-
-	TArray<FColorKey>& AlphaKeys = ChannelName == NAME_None || ChannelName == AlphaName || ChannelName == SpecifiedColorName ? NewGeneratedKeys : DefaultGeneratedKeys;
-	AlphaKeys.Add( FColorKey( EKeyColorChannel::Alpha, ColorValue.A, bIsSlateColor ) );
+	OutGeneratedKeys.Add(FMovieSceneChannelValueSetter::Create<FMovieSceneFloatChannel>(0, ColorValue.R, bKeyRed));
+	OutGeneratedKeys.Add(FMovieSceneChannelValueSetter::Create<FMovieSceneFloatChannel>(1, ColorValue.G, bKeyGreen));
+	OutGeneratedKeys.Add(FMovieSceneChannelValueSetter::Create<FMovieSceneFloatChannel>(2, ColorValue.B, bKeyBlue));
+	OutGeneratedKeys.Add(FMovieSceneChannelValueSetter::Create<FMovieSceneFloatChannel>(3, ColorValue.A, bKeyAlpha));
 }
 
 void CopyInterpColorTrack(TSharedRef<ISequencer> Sequencer, UInterpTrackColorProp* ColorPropTrack, UMovieSceneColorTrack* ColorTrack)
diff --git a/Engine/Source/Editor/MovieSceneTools/Private/TrackEditors/PropertyTrackEditors/ColorPropertyTrackEditor.h b/Engine/Source/Editor/MovieSceneTools/Private/TrackEditors/PropertyTrackEditors/ColorPropertyTrackEditor.h
--- a/Engine/Source/Editor/MovieSceneTools/Private/TrackEditors/PropertyTrackEditors/ColorPropertyTrackEditor.h
+++ b/Engine/Source/Editor/MovieSceneTools/Private/TrackEditors/PropertyTrackEditors/ColorPropertyTrackEditor.h
@@ -15,7 +15,7 @@
 * A property track editor for colors.
 */
 class FColorPropertyTrackEditor
-	: public FPropertyTrackEditor<UMovieSceneColorTrack, UMovieSceneColorSection, FColorKey>
+	: public FPropertyTrackEditor<UMovieSceneColorTrack>
 {
 public:
 
@@ -25,7 +25,7 @@ public:
 	 * @param InSequencer The sequencer instance to be used by this tool.
 	 */
 	FColorPropertyTrackEditor(TSharedRef<ISequencer> InSequencer)
-		: FPropertyTrackEditor<UMovieSceneColorTrack, UMovieSceneColorSection, FColorKey>(InSequencer, GetAnimatedPropertyTypes())
+		: FPropertyTrackEditor(InSequencer, GetAnimatedPropertyTypes())
 	{
 	}
 
@@ -57,7 +57,7 @@ protected:
 
 	//~ FPropertyTrackEditor interface
 
-	virtual void GenerateKeysFromPropertyChanged(const FPropertyChangedParams& PropertyChangedParams, TArray<FColorKey>& NewGeneratedKeys, TArray<FColorKey>& DefaultGeneratedKeys) override;
+	virtual void GenerateKeysFromPropertyChanged(const FPropertyChangedParams& PropertyChangedParams, FGeneratedTrackKeys& OutGeneratedKeys) override;
 
 protected:
 
diff --git a/Engine/Source/Editor/MovieSceneTools/Private/TrackEditors/PropertyTrackEditors/FloatPropertyTrackEditor.cpp b/Engine/Source/Editor/MovieSceneTools/Private/TrackEditors/PropertyTrackEditors/FloatPropertyTrackEditor.cpp
--- a/Engine/Source/Editor/MovieSceneTools/Private/TrackEditors/PropertyTrackEditors/FloatPropertyTrackEditor.cpp
+++ b/Engine/Source/Editor/MovieSceneTools/Private/TrackEditors/PropertyTrackEditors/FloatPropertyTrackEditor.cpp
@@ -3,7 +3,6 @@
 #include "TrackEditors/PropertyTrackEditors/FloatPropertyTrackEditor.h"
 #include "Editor/UnrealEdEngine.h"
 #include "UnrealEdGlobals.h"
-#include "Sections/FloatPropertySection.h"
 #include "MatineeImportTools.h"
 #include "Matinee/InterpTrackFloatBase.h"
 
@@ -14,23 +13,10 @@ TSharedRef<ISequencerTrackEditor> FFloatPropertyTrackEditor::CreateTrackEditor(
 }
 
 
-TSharedRef<ISequencerSection> FFloatPropertyTrackEditor::MakeSectionInterface(UMovieSceneSection& SectionObject, UMovieSceneTrack& Track, FGuid ObjectBinding)
+void FFloatPropertyTrackEditor::GenerateKeysFromPropertyChanged( const FPropertyChangedParams& PropertyChangedParams, FGeneratedTrackKeys& OutGeneratedKeys )
 {
-	UMovieScenePropertyTrack* PropertyTrack = Cast<UMovieScenePropertyTrack>(&Track);
-	if (PropertyTrack != nullptr)
-	{
-		return MakeShareable(new FFloatPropertySection(GetSequencer().Get(), ObjectBinding, PropertyTrack->GetPropertyName(), PropertyTrack->GetPropertyPath(), SectionObject, Track.GetDisplayName()));
-	}
-	else
-	{
-		return MakeShareable(new FFloatPropertySection(SectionObject, Track.GetDisplayName()));
-	}
-}
-
-
-void FFloatPropertyTrackEditor::GenerateKeysFromPropertyChanged( const FPropertyChangedParams& PropertyChangedParams, TArray<float>& NewGeneratedKeys, TArray<float>& DefaultGeneratedKeys )
-{
-	NewGeneratedKeys.Add( PropertyChangedParams.GetPropertyValue<float>() );
+	const float KeyedValue =  PropertyChangedParams.GetPropertyValue<float>();
+	OutGeneratedKeys.Add(FMovieSceneChannelValueSetter::Create<FMovieSceneFloatChannel>(0, KeyedValue, true));
 }
 
 
diff --git a/Engine/Source/Editor/MovieSceneTools/Private/TrackEditors/PropertyTrackEditors/FloatPropertyTrackEditor.h b/Engine/Source/Editor/MovieSceneTools/Private/TrackEditors/PropertyTrackEditors/FloatPropertyTrackEditor.h
--- a/Engine/Source/Editor/MovieSceneTools/Private/TrackEditors/PropertyTrackEditors/FloatPropertyTrackEditor.h
+++ b/Engine/Source/Editor/MovieSceneTools/Private/TrackEditors/PropertyTrackEditors/FloatPropertyTrackEditor.h
@@ -15,7 +15,7 @@
  * A property track editor for floats.
  */
 class FFloatPropertyTrackEditor
-	: public FPropertyTrackEditor<UMovieSceneFloatTrack, UMovieSceneFloatSection, float>
+	: public FPropertyTrackEditor<UMovieSceneFloatTrack>
 {
 public:
 
@@ -50,11 +50,7 @@ public:
 
 	virtual void BuildTrackContextMenu(FMenuBuilder& MenuBuilder, UMovieSceneTrack* Track) override;
 
-	//~ ISequencerTrackEditor interface
-
-	virtual TSharedRef<ISequencerSection> MakeSectionInterface(UMovieSceneSection& SectionObject, UMovieSceneTrack& Track, FGuid ObjectBinding) override;
-
 protected:
 
-	virtual void GenerateKeysFromPropertyChanged(const FPropertyChangedParams& PropertyChangedParams, TArray<float>& NewGeneratedKeys, TArray<float>& DefaultGeneratedKeys) override;
+	virtual void GenerateKeysFromPropertyChanged(const FPropertyChangedParams& PropertyChangedParams, FGeneratedTrackKeys& OutGeneratedKeys) override;
 };
diff --git a/Engine/Source/Editor/MovieSceneTools/Private/TrackEditors/PropertyTrackEditors/IntegerPropertyTrackEditor.cpp b/Engine/Source/Editor/MovieSceneTools/Private/TrackEditors/PropertyTrackEditors/IntegerPropertyTrackEditor.cpp
--- a/Engine/Source/Editor/MovieSceneTools/Private/TrackEditors/PropertyTrackEditors/IntegerPropertyTrackEditor.cpp
+++ b/Engine/Source/Editor/MovieSceneTools/Private/TrackEditors/PropertyTrackEditors/IntegerPropertyTrackEditor.cpp
@@ -1,7 +1,6 @@
 // Copyright 1998-2018 Epic Games, Inc. All Rights Reserved.
 
 #include "TrackEditors/PropertyTrackEditors/IntegerPropertyTrackEditor.h"
-#include "Sections/IntegerPropertySection.h"
 
 
 TSharedRef<ISequencerTrackEditor> FIntegerPropertyTrackEditor::CreateTrackEditor( TSharedRef<ISequencer> OwningSequencer )
@@ -10,21 +9,8 @@ TSharedRef<ISequencerTrackEditor> FIntegerPropertyTrackEditor::CreateTrackEditor
 }
 
 
-TSharedRef<ISequencerSection> FIntegerPropertyTrackEditor::MakeSectionInterface(UMovieSceneSection& SectionObject, UMovieSceneTrack& Track, FGuid ObjectBinding)
+void FIntegerPropertyTrackEditor::GenerateKeysFromPropertyChanged( const FPropertyChangedParams& PropertyChangedParams, FGeneratedTrackKeys& OutGeneratedKeys )
 {
-	UMovieScenePropertyTrack* PropertyTrack = Cast<UMovieScenePropertyTrack>(&Track);
-	if (PropertyTrack != nullptr)
-	{
-		return MakeShareable(new FIntegerPropertySection(GetSequencer().Get(), ObjectBinding, PropertyTrack->GetPropertyName(), PropertyTrack->GetPropertyPath(), SectionObject, Track.GetDisplayName()));
-	}
-	else
-	{
-		return MakeShareable(new FIntegerPropertySection(SectionObject, Track.GetDisplayName()));
-	}
-}
-
-
-void FIntegerPropertyTrackEditor::GenerateKeysFromPropertyChanged( const FPropertyChangedParams& PropertyChangedParams, TArray<int32>& NewGeneratedKeys, TArray<int32>& DefaultGeneratedKeys )
-{
-	NewGeneratedKeys.Add( PropertyChangedParams.GetPropertyValue<int32>() );
+	const int32 KeyedValue = PropertyChangedParams.GetPropertyValue<int32>();
+	OutGeneratedKeys.Add(FMovieSceneChannelValueSetter::Create<FMovieSceneIntegerChannel>(0, KeyedValue, true));
 }
diff --git a/Engine/Source/Editor/MovieSceneTools/Private/TrackEditors/PropertyTrackEditors/IntegerPropertyTrackEditor.h b/Engine/Source/Editor/MovieSceneTools/Private/TrackEditors/PropertyTrackEditors/IntegerPropertyTrackEditor.h
--- a/Engine/Source/Editor/MovieSceneTools/Private/TrackEditors/PropertyTrackEditors/IntegerPropertyTrackEditor.h
+++ b/Engine/Source/Editor/MovieSceneTools/Private/TrackEditors/PropertyTrackEditors/IntegerPropertyTrackEditor.h
@@ -15,7 +15,7 @@
  * A property track editor for integers.
  */
 class FIntegerPropertyTrackEditor
-	: public FPropertyTrackEditor<UMovieSceneIntegerTrack, UMovieSceneIntegerSection, int32>
+	: public FPropertyTrackEditor<UMovieSceneIntegerTrack>
 {
 public:
 
@@ -44,13 +44,7 @@ public:
 	 */
 	static TSharedRef<ISequencerTrackEditor> CreateTrackEditor(TSharedRef<ISequencer> OwningSequencer);
 
-public:
-
-	//~ ISequencerTrackEditor interface
-
-	virtual TSharedRef<ISequencerSection> MakeSectionInterface(UMovieSceneSection& SectionObject, UMovieSceneTrack& Track, FGuid ObjectBinding) override;
-
 protected:
 
-	virtual void GenerateKeysFromPropertyChanged(const FPropertyChangedParams& PropertyChangedParams, TArray<int32>& NewGeneratedKeys, TArray<int32>& DefaultGeneratedKeys) override;
+	virtual void GenerateKeysFromPropertyChanged(const FPropertyChangedParams& PropertyChangedParams, FGeneratedTrackKeys& OutGeneratedKeys) override;
 };
diff --git a/Engine/Source/Editor/MovieSceneTools/Private/TrackEditors/PropertyTrackEditors/StringPropertyTrackEditor.cpp b/Engine/Source/Editor/MovieSceneTools/Private/TrackEditors/PropertyTrackEditors/StringPropertyTrackEditor.cpp
--- a/Engine/Source/Editor/MovieSceneTools/Private/TrackEditors/PropertyTrackEditors/StringPropertyTrackEditor.cpp
+++ b/Engine/Source/Editor/MovieSceneTools/Private/TrackEditors/PropertyTrackEditors/StringPropertyTrackEditor.cpp
@@ -1,7 +1,6 @@
 // Copyright 1998-2018 Epic Games, Inc. All Rights Reserved.
 
 #include "TrackEditors/PropertyTrackEditors/StringPropertyTrackEditor.h"
-#include "Sections/StringPropertySection.h"
 
 
 TSharedRef<ISequencerTrackEditor> FStringPropertyTrackEditor::CreateTrackEditor( TSharedRef<ISequencer> OwningSequencer )
@@ -10,15 +9,7 @@ TSharedRef<ISequencerTrackEditor> FStringPropertyTrackEditor::CreateTrackEditor(
 }
 
 
-TSharedRef<ISequencerSection> FStringPropertyTrackEditor::MakeSectionInterface(UMovieSceneSection& SectionObject, UMovieSceneTrack& Track, FGuid ObjectBinding)
-{
-	UMovieScenePropertyTrack* PropertyTrack = Cast<UMovieScenePropertyTrack>(&Track);
-	checkf(PropertyTrack != nullptr, TEXT("Incompatible track in FStringPropertyTrackEditor"));
-	return MakeShareable(new FStringPropertySection(GetSequencer().Get(), ObjectBinding, PropertyTrack->GetPropertyName(), PropertyTrack->GetPropertyPath(), SectionObject, Track.GetDisplayName()));
-}
-
-
-void FStringPropertyTrackEditor::GenerateKeysFromPropertyChanged( const FPropertyChangedParams& PropertyChangedParams, TArray<FString>& NewGeneratedKeys, TArray<FString>& DefaultGeneratedKeys )
+void FStringPropertyTrackEditor::GenerateKeysFromPropertyChanged( const FPropertyChangedParams& PropertyChangedParams, FGeneratedTrackKeys& OutGeneratedKeys )
 {
 	void* CurrentObject = PropertyChangedParams.ObjectsThatChanged[0];
 	void* PropertyValue = nullptr;
@@ -34,6 +25,6 @@ void FStringPropertyTrackEditor::GenerateKeysFromPropertyChanged( const FPropert
 	if ( StrProperty )
 	{
 		FString StrPropertyValue = StrProperty->GetPropertyValue(CurrentObject);
-		NewGeneratedKeys.Add(StrPropertyValue);
+		OutGeneratedKeys.Add(FMovieSceneChannelValueSetter::Create<FMovieSceneStringChannel>(0, MoveTemp(StrPropertyValue), true));
 	}
 }
diff --git a/Engine/Source/Editor/MovieSceneTools/Private/TrackEditors/PropertyTrackEditors/StringPropertyTrackEditor.h b/Engine/Source/Editor/MovieSceneTools/Private/TrackEditors/PropertyTrackEditors/StringPropertyTrackEditor.h
--- a/Engine/Source/Editor/MovieSceneTools/Private/TrackEditors/PropertyTrackEditors/StringPropertyTrackEditor.h
+++ b/Engine/Source/Editor/MovieSceneTools/Private/TrackEditors/PropertyTrackEditors/StringPropertyTrackEditor.h
@@ -15,7 +15,7 @@
  * A property track editor for strings.
  */
 class FStringPropertyTrackEditor
-	: public FPropertyTrackEditor<UMovieSceneStringTrack, UMovieSceneStringSection, FString>
+	: public FPropertyTrackEditor<UMovieSceneStringTrack>
 {
 public:
 
@@ -44,13 +44,9 @@ public:
 	 */
 	static TSharedRef<ISequencerTrackEditor> CreateTrackEditor(TSharedRef<ISequencer> OwningSequencer);
 
-	//~ ISequencerTrackEditor interface
-
-	virtual TSharedRef<ISequencerSection> MakeSectionInterface(UMovieSceneSection& SectionObject, UMovieSceneTrack& Track, FGuid ObjectBinding) override;
-
 protected:
 
 	//~ FPropertyTrackEditor interface
 
-	virtual void GenerateKeysFromPropertyChanged(const FPropertyChangedParams& PropertyChangedParams, TArray<FString>& NewGeneratedKeys, TArray<FString>& DefaultGeneratedKeys) override;
+	virtual void GenerateKeysFromPropertyChanged(const FPropertyChangedParams& PropertyChangedParams, FGeneratedTrackKeys& OutGeneratedKeys) override;
 };
diff --git a/Engine/Source/Editor/MovieSceneTools/Private/TrackEditors/PropertyTrackEditors/TransformPropertyTrackEditor.cpp b/Engine/Source/Editor/MovieSceneTools/Private/TrackEditors/PropertyTrackEditors/TransformPropertyTrackEditor.cpp
--- a/Engine/Source/Editor/MovieSceneTools/Private/TrackEditors/PropertyTrackEditors/TransformPropertyTrackEditor.cpp
+++ b/Engine/Source/Editor/MovieSceneTools/Private/TrackEditors/PropertyTrackEditors/TransformPropertyTrackEditor.cpp
@@ -15,12 +15,7 @@ TSharedRef<ISequencerTrackEditor> FTransformPropertyTrackEditor::CreateTrackEdit
 
 TSharedRef<ISequencerSection> FTransformPropertyTrackEditor::MakeSectionInterface(UMovieSceneSection& SectionObject, UMovieSceneTrack& Track, FGuid ObjectBinding)
 {
-	UMovieScenePropertyTrack* PropertyTrack = Cast<UMovieScenePropertyTrack>(&Track);
-	checkf(PropertyTrack != nullptr, TEXT("Incompatible track in FTransformPropertyTrackEditor"));
-
-	TSharedRef<FTransformSection> NewSection = MakeShared<FTransformSection>(&SectionObject, GetSequencer(), ObjectBinding);
-	NewSection->AssignProperty(PropertyTrack->GetPropertyName(), PropertyTrack->GetPropertyPath());
-	return NewSection;
+	return MakeShared<FTransformSection>(SectionObject, GetSequencer());
 }
 
 
@@ -48,23 +43,23 @@ TSharedPtr<SWidget> FTransformPropertyTrackEditor::BuildOutlinerEditWidget(const
 }
 
 
-void FTransformPropertyTrackEditor::GenerateKeysFromPropertyChanged( const FPropertyChangedParams& PropertyChangedParams, TArray<FTransformKey>& NewGeneratedKeys, TArray<FTransformKey>& DefaultGeneratedKeys )
+void FTransformPropertyTrackEditor::GenerateKeysFromPropertyChanged( const FPropertyChangedParams& PropertyChangedParams, FGeneratedTrackKeys& OutGeneratedKeys )
 {
 	FTransform Transform = PropertyChangedParams.GetPropertyValue<FTransform>();
 
 	const FVector& Translation = Transform.GetTranslation();
-	NewGeneratedKeys.Add(FTransformKey(EKey3DTransformChannel::Translation, EAxis::X, Translation.X, false));
-	NewGeneratedKeys.Add(FTransformKey(EKey3DTransformChannel::Translation, EAxis::Y, Translation.Y, false));
-	NewGeneratedKeys.Add(FTransformKey(EKey3DTransformChannel::Translation, EAxis::Z, Translation.Z, false));
+	OutGeneratedKeys.Add(FMovieSceneChannelValueSetter::Create<FMovieSceneFloatChannel>(0, Translation.X, true));
+	OutGeneratedKeys.Add(FMovieSceneChannelValueSetter::Create<FMovieSceneFloatChannel>(1, Translation.Y, true));
+	OutGeneratedKeys.Add(FMovieSceneChannelValueSetter::Create<FMovieSceneFloatChannel>(2, Translation.Z, true));
 
 	const FRotator& Rotator = Transform.GetRotation().Rotator();
-	NewGeneratedKeys.Add(FTransformKey(EKey3DTransformChannel::Rotation, EAxis::X, Rotator.Roll, false));
-	NewGeneratedKeys.Add(FTransformKey(EKey3DTransformChannel::Rotation, EAxis::Y, Rotator.Pitch, false));
-	NewGeneratedKeys.Add(FTransformKey(EKey3DTransformChannel::Rotation, EAxis::Z, Rotator.Yaw, false));
+	OutGeneratedKeys.Add(FMovieSceneChannelValueSetter::Create<FMovieSceneFloatChannel>(3, Rotator.Roll, true));
+	OutGeneratedKeys.Add(FMovieSceneChannelValueSetter::Create<FMovieSceneFloatChannel>(4, Rotator.Pitch, true));
+	OutGeneratedKeys.Add(FMovieSceneChannelValueSetter::Create<FMovieSceneFloatChannel>(5, Rotator.Yaw, true));
 
 	const FVector& Scale = Transform.GetScale3D();
-	NewGeneratedKeys.Add(FTransformKey(EKey3DTransformChannel::Scale, EAxis::X, Scale.X, false));
-	NewGeneratedKeys.Add(FTransformKey(EKey3DTransformChannel::Scale, EAxis::Y, Scale.Y, false));
-	NewGeneratedKeys.Add(FTransformKey(EKey3DTransformChannel::Scale, EAxis::Z, Scale.Z, false));
+	OutGeneratedKeys.Add(FMovieSceneChannelValueSetter::Create<FMovieSceneFloatChannel>(6, Scale.X, true));
+	OutGeneratedKeys.Add(FMovieSceneChannelValueSetter::Create<FMovieSceneFloatChannel>(7, Scale.Y, true));
+	OutGeneratedKeys.Add(FMovieSceneChannelValueSetter::Create<FMovieSceneFloatChannel>(8, Scale.Z, true));
 }
 
diff --git a/Engine/Source/Editor/MovieSceneTools/Private/TrackEditors/PropertyTrackEditors/TransformPropertyTrackEditor.h b/Engine/Source/Editor/MovieSceneTools/Private/TrackEditors/PropertyTrackEditors/TransformPropertyTrackEditor.h
--- a/Engine/Source/Editor/MovieSceneTools/Private/TrackEditors/PropertyTrackEditors/TransformPropertyTrackEditor.h
+++ b/Engine/Source/Editor/MovieSceneTools/Private/TrackEditors/PropertyTrackEditors/TransformPropertyTrackEditor.h
@@ -16,7 +16,7 @@
  * A property track editor for transforms.
  */
 class FTransformPropertyTrackEditor
-	: public FPropertyTrackEditor<UMovieSceneTransformTrack, UMovieScene3DTransformSection, FTransformKey>
+	: public FPropertyTrackEditor<UMovieSceneTransformTrack>
 {
 public:
 
@@ -26,7 +26,7 @@ public:
 	 * @param InSequencer The sequencer instance to be used by this tool.
 	 */
 	FTransformPropertyTrackEditor(TSharedRef<ISequencer> InSequencer)
-		: FPropertyTrackEditor<UMovieSceneTransformTrack, UMovieScene3DTransformSection, FTransformKey>(InSequencer, GetAnimatedPropertyTypes())
+		: FPropertyTrackEditor(InSequencer, GetAnimatedPropertyTypes())
 	{ }
 
 	/**
@@ -55,5 +55,5 @@ protected:
 
 	//~ FPropertyTrackEditor interface
 
-	virtual void GenerateKeysFromPropertyChanged(const FPropertyChangedParams& PropertyChangedParams, TArray<FTransformKey>& NewGeneratedKeys, TArray<FTransformKey>& DefaultGeneratedKeys) override;
+	virtual void GenerateKeysFromPropertyChanged(const FPropertyChangedParams& PropertyChangedParams, FGeneratedTrackKeys& OutGeneratedKeys) override;
 };
diff --git a/Engine/Source/Editor/MovieSceneTools/Private/TrackEditors/PropertyTrackEditors/VectorPropertyTrackEditor.cpp b/Engine/Source/Editor/MovieSceneTools/Private/TrackEditors/PropertyTrackEditors/VectorPropertyTrackEditor.cpp
--- a/Engine/Source/Editor/MovieSceneTools/Private/TrackEditors/PropertyTrackEditors/VectorPropertyTrackEditor.cpp
+++ b/Engine/Source/Editor/MovieSceneTools/Private/TrackEditors/PropertyTrackEditors/VectorPropertyTrackEditor.cpp
@@ -1,7 +1,6 @@
 // Copyright 1998-2018 Epic Games, Inc. All Rights Reserved.
 
 #include "TrackEditors/PropertyTrackEditors/VectorPropertyTrackEditor.h"
-#include "Sections/VectorPropertySection.h"
 #include "Tracks/MovieSceneVectorTrack.h"
 #include "MatineeImportTools.h"
 #include "Matinee/InterpTrackVectorProp.h"
@@ -20,15 +19,7 @@ TSharedRef<ISequencerTrackEditor> FVectorPropertyTrackEditor::CreateTrackEditor(
 	return MakeShareable( new FVectorPropertyTrackEditor( InSequencer ) );
 }
 
-
-TSharedRef<ISequencerSection> FVectorPropertyTrackEditor::MakeSectionInterface(UMovieSceneSection& SectionObject, UMovieSceneTrack& Track, FGuid ObjectBinding)
-{
-	UMovieScenePropertyTrack* PropertyTrack = Cast<UMovieScenePropertyTrack>(&Track);
-	checkf(PropertyTrack != nullptr, TEXT("Incompatible track in FVectorTrackEditor"));
-	return MakeShareable(new FVectorPropertySection(GetSequencer().Get(), ObjectBinding, PropertyTrack->GetPropertyName(), PropertyTrack->GetPropertyPath(), SectionObject, Track.GetDisplayName()));
-}
-
-void FVectorPropertyTrackEditor::GenerateKeysFromPropertyChanged( const FPropertyChangedParams& PropertyChangedParams, TArray<FVectorKey>& NewGeneratedKeys, TArray<FVectorKey>& DefaultGeneratedKeys )
+void FVectorPropertyTrackEditor::GenerateKeysFromPropertyChanged( const FPropertyChangedParams& PropertyChangedParams, FGeneratedTrackKeys& OutGeneratedKeys )
 {
 	const UStructProperty* StructProp = Cast<const UStructProperty>( PropertyChangedParams.PropertyPath.GetLeafMostProperty().Property.Get() );
 
@@ -62,30 +53,29 @@ void FVectorPropertyTrackEditor::GenerateKeysFromPropertyChanged( const FPropert
 		Channels = 4;
 	}
 
-	FName ChannelName = PropertyChangedParams.StructPropertyNameToKey;
+	FPropertyPath StructPath = PropertyChangedParams.StructPathToKey;
+	FName ChannelName = StructPath.GetNumProperties() != 0 ? StructPath.GetLeafMostProperty().Property->GetFName() : NAME_None;
 
-	TArray<FVectorKey>& XKeys = ChannelName == NAME_None || ChannelName == XName ? NewGeneratedKeys : DefaultGeneratedKeys;
-	XKeys.Add( FVectorKey( EKeyVectorChannel::X, VectorValues.X ) );
+	const bool bKeyX = ChannelName == NAME_None || ChannelName == XName;
+	const bool bKeyY = ChannelName == NAME_None || ChannelName == YName;
 
-	TArray<FVectorKey>& YKeys = ChannelName == NAME_None || ChannelName == YName ? NewGeneratedKeys : DefaultGeneratedKeys;
-	YKeys.Add( FVectorKey( EKeyVectorChannel::Y, VectorValues.Y ) );
+	OutGeneratedKeys.Add(FMovieSceneChannelValueSetter::Create<FMovieSceneFloatChannel>(0, VectorValues.X, bKeyX));
+	OutGeneratedKeys.Add(FMovieSceneChannelValueSetter::Create<FMovieSceneFloatChannel>(1, VectorValues.Y, bKeyY));
 
 	if ( Channels >= 3 )
 	{
-		TArray<FVectorKey>& ZKeys = ChannelName == NAME_None || ChannelName == ZName ? NewGeneratedKeys : DefaultGeneratedKeys;
-		ZKeys.Add( FVectorKey( EKeyVectorChannel::Z, VectorValues.Z ) );
+		OutGeneratedKeys.Add(FMovieSceneChannelValueSetter::Create<FMovieSceneFloatChannel>(2, VectorValues.Z, ChannelName == NAME_None || ChannelName == ZName));
 	}
 
 	if ( Channels >= 4 )
 	{
-		TArray<FVectorKey>& WKeys = ChannelName == NAME_None || ChannelName == WName ? NewGeneratedKeys : DefaultGeneratedKeys;
-		WKeys.Add( FVectorKey( EKeyVectorChannel::W, VectorValues.W ) );
+		OutGeneratedKeys.Add(FMovieSceneChannelValueSetter::Create<FMovieSceneFloatChannel>(3, VectorValues.W, ChannelName == NAME_None || ChannelName == WName));
 	}
 }
 
 void FVectorPropertyTrackEditor::InitializeNewTrack( UMovieSceneVectorTrack* NewTrack, FPropertyChangedParams PropertyChangedParams )
 {
-	FPropertyTrackEditor<UMovieSceneVectorTrack, UMovieSceneVectorSection, FVectorKey>::InitializeNewTrack( NewTrack, PropertyChangedParams );
+	FPropertyTrackEditor::InitializeNewTrack( NewTrack, PropertyChangedParams );
 	const UStructProperty* StructProp = Cast<const UStructProperty>( PropertyChangedParams.PropertyPath.GetLeafMostProperty().Property.Get() );
 	FName StructName = StructProp->Struct->GetFName();
 
diff --git a/Engine/Source/Editor/MovieSceneTools/Private/TrackEditors/PropertyTrackEditors/VectorPropertyTrackEditor.h b/Engine/Source/Editor/MovieSceneTools/Private/TrackEditors/PropertyTrackEditors/VectorPropertyTrackEditor.h
--- a/Engine/Source/Editor/MovieSceneTools/Private/TrackEditors/PropertyTrackEditors/VectorPropertyTrackEditor.h
+++ b/Engine/Source/Editor/MovieSceneTools/Private/TrackEditors/PropertyTrackEditors/VectorPropertyTrackEditor.h
@@ -16,7 +16,7 @@
  * A property track editor for vectors.
  */
 class FVectorPropertyTrackEditor
-	: public FPropertyTrackEditor<UMovieSceneVectorTrack, UMovieSceneVectorSection, FVectorKey>
+	: public FPropertyTrackEditor<UMovieSceneVectorTrack>
 {
 public:
 
@@ -26,7 +26,7 @@ public:
 	 * @param InSequencer The sequencer instance to be used by this tool.
 	 */
 	FVectorPropertyTrackEditor(TSharedRef<ISequencer> InSequencer)
-		: FPropertyTrackEditor<UMovieSceneVectorTrack, UMovieSceneVectorSection, FVectorKey>(InSequencer, GetAnimatedPropertyTypes())
+		: FPropertyTrackEditor(InSequencer, GetAnimatedPropertyTypes())
 	{
 	}
 
@@ -50,10 +50,6 @@ public:
 	 */
 	static TSharedRef<ISequencerTrackEditor> CreateTrackEditor(TSharedRef<ISequencer> OwningSequencer);
 
-	//~ ISequencerTrackEditor interface
-
-	virtual TSharedRef<ISequencerSection> MakeSectionInterface(UMovieSceneSection& SectionObject, UMovieSceneTrack& Track, FGuid ObjectBinding) override;
-
 protected:
 
 	//~ FMovieSceneTrackEditor interface
@@ -62,7 +58,7 @@ protected:
 
 	//~ FPropertyTrackEditor interface
 
-	virtual void GenerateKeysFromPropertyChanged(const FPropertyChangedParams& PropertyChangedParams, TArray<FVectorKey>& NewGeneratedKeys, TArray<FVectorKey>& DefaultGeneratedKeys) override;
+	virtual void GenerateKeysFromPropertyChanged(const FPropertyChangedParams& PropertyChangedParams, FGeneratedTrackKeys& OutGeneratedKeys) override;
 	virtual void InitializeNewTrack(UMovieSceneVectorTrack* NewTrack, FPropertyChangedParams PropertyChangedParams) override;
 
 private:
diff --git a/Engine/Source/Editor/MovieSceneTools/Private/TrackEditors/PropertyTrackEditors/VisibilityPropertyTrackEditor.cpp b/Engine/Source/Editor/MovieSceneTools/Private/TrackEditors/PropertyTrackEditors/VisibilityPropertyTrackEditor.cpp
--- a/Engine/Source/Editor/MovieSceneTools/Private/TrackEditors/PropertyTrackEditors/VisibilityPropertyTrackEditor.cpp
+++ b/Engine/Source/Editor/MovieSceneTools/Private/TrackEditors/PropertyTrackEditors/VisibilityPropertyTrackEditor.cpp
@@ -1,7 +1,7 @@
 // Copyright 1998-2018 Epic Games, Inc. All Rights Reserved.
 
 #include "TrackEditors/PropertyTrackEditors/VisibilityPropertyTrackEditor.h"
-#include "Sections/VisibilityPropertySection.h"
+#include "Sections/BoolPropertySection.h"
 
 
 TSharedRef<ISequencerTrackEditor> FVisibilityPropertyTrackEditor::CreateTrackEditor( TSharedRef<ISequencer> OwningSequencer )
@@ -12,13 +12,13 @@ TSharedRef<ISequencerTrackEditor> FVisibilityPropertyTrackEditor::CreateTrackEdi
 
 TSharedRef<ISequencerSection> FVisibilityPropertyTrackEditor::MakeSectionInterface(UMovieSceneSection& SectionObject, UMovieSceneTrack& Track, FGuid ObjectBinding)
 {
-	UMovieScenePropertyTrack* PropertyTrack = Cast<UMovieScenePropertyTrack>(&Track);
-	checkf(PropertyTrack != nullptr, TEXT("Incompatible track in FVisibilityPropertyTrackEditor"));
-	return MakeShareable(new FVisibilityPropertySection(GetSequencer().Get(), ObjectBinding, PropertyTrack->GetPropertyName(), PropertyTrack->GetPropertyPath(), SectionObject, Track.GetDisplayName()));
+	return MakeShared<FBoolPropertySection>(SectionObject);
 }
 
 
-void FVisibilityPropertyTrackEditor::GenerateKeysFromPropertyChanged( const FPropertyChangedParams& PropertyChangedParams, TArray<bool>& NewGeneratedKeys, TArray<bool>& DefaultGeneratedKeys )
+void FVisibilityPropertyTrackEditor::GenerateKeysFromPropertyChanged( const FPropertyChangedParams& PropertyChangedParams, FGeneratedTrackKeys& OutGeneratedKeys )
 {
-	NewGeneratedKeys.Add(!PropertyChangedParams.GetPropertyValue<bool>());
+	// Invert the property key since the underlying property is actually 'hidden'
+	bool KeyedValue = !PropertyChangedParams.GetPropertyValue<bool>();
+	OutGeneratedKeys.Add(FMovieSceneChannelValueSetter::Create<FMovieSceneBoolChannel>(0, KeyedValue, true));
 }
diff --git a/Engine/Source/Editor/MovieSceneTools/Private/TrackEditors/PropertyTrackEditors/VisibilityPropertyTrackEditor.h b/Engine/Source/Editor/MovieSceneTools/Private/TrackEditors/PropertyTrackEditors/VisibilityPropertyTrackEditor.h
--- a/Engine/Source/Editor/MovieSceneTools/Private/TrackEditors/PropertyTrackEditors/VisibilityPropertyTrackEditor.h
+++ b/Engine/Source/Editor/MovieSceneTools/Private/TrackEditors/PropertyTrackEditors/VisibilityPropertyTrackEditor.h
@@ -17,15 +17,13 @@
  * to use a UMovieSceneVisibilityTrack through metadata.
  */
 class FVisibilityPropertyTrackEditor
-	: public FPropertyTrackEditor<UMovieSceneVisibilityTrack, UMovieSceneBoolSection, bool>
+	: public FPropertyTrackEditor<UMovieSceneVisibilityTrack>
 {
 public:
 
 	/** Constructor. */
 	FVisibilityPropertyTrackEditor(TSharedRef<ISequencer> InSequencer)
-		// Don't supply any property type names to watch since the FBoolPropertyTrackEditor is already watching for bool property changes.
-		// @todo: Is this right? comment says we should't be passing the property type name, yet we currently are.
-		: FPropertyTrackEditor<UMovieSceneVisibilityTrack, UMovieSceneBoolSection, bool>(InSequencer, GetAnimatedPropertyTypes())
+		: FPropertyTrackEditor(InSequencer, GetAnimatedPropertyTypes())
 	{ }
 
 	/**
@@ -44,14 +42,11 @@ public:
 	 */
 	static TSharedRef<ISequencerTrackEditor> CreateTrackEditor(TSharedRef<ISequencer> OwningSequencer);
 
-	//~ ISequencerTrackEditor interface
-
-	virtual TSharedRef<ISequencerSection> MakeSectionInterface(UMovieSceneSection& SectionObject, UMovieSceneTrack& Track, FGuid ObjectBinding) override;
-
 protected:
 
 	//~ FPropertyTrackEditor interface
 
-	virtual void GenerateKeysFromPropertyChanged(const FPropertyChangedParams& PropertyChangedParams, TArray<bool>& NewGeneratedKeys, TArray<bool>& DefaultGeneratedKeys) override;
+	virtual TSharedRef<ISequencerSection> MakeSectionInterface(UMovieSceneSection& SectionObject, UMovieSceneTrack& Track, FGuid ObjectBinding) override;
+	virtual void GenerateKeysFromPropertyChanged(const FPropertyChangedParams& PropertyChangedParams, FGeneratedTrackKeys& OutGeneratedKeys) override;
 	virtual bool ForCustomizedUseOnly() { return true; }
 };
