Copying //UE4/Dev-Sequencer to //UE4/Dev-Main (Source: //UE4/Dev-Sequencer @ 3946692)

diff --git a/Engine/Source/Editor/MovieSceneTools/Private/TrackEditors/SkeletalAnimationTrackEditor.cpp b/Engine/Source/Editor/MovieSceneTools/Private/TrackEditors/SkeletalAnimationTrackEditor.cpp
--- a/Engine/Source/Editor/MovieSceneTools/Private/TrackEditors/SkeletalAnimationTrackEditor.cpp
+++ b/Engine/Source/Editor/MovieSceneTools/Private/TrackEditors/SkeletalAnimationTrackEditor.cpp
@@ -25,11 +25,12 @@
 #include "MatineeImportTools.h"
 #include "Matinee/InterpTrackAnimControl.h"
 #include "SequencerUtilities.h"
-#include "FloatCurveKeyArea.h"
 #include "ISectionLayoutBuilder.h"
 #include "Animation/AnimMontage.h"
 #include "Animation/AnimSequence.h"
 #include "EditorStyleSet.h"
+#include "DragAndDrop/AssetDragDropOp.h"
+#include "MovieSceneTimeHelpers.h"
 
 
 namespace SkeletalAnimationEditorConstants
@@ -44,7 +45,7 @@ namespace SkeletalAnimationEditorConstants
 FSkeletalAnimationSection::FSkeletalAnimationSection( UMovieSceneSection& InSection )
 	: Section(*CastChecked<UMovieSceneSkeletalAnimationSection>(&InSection))
 	, InitialStartOffsetDuringResize(0.f)
-	, InitialStartTimeDuringResize(0.f)
+	, InitialStartTimeDuringResize(0)
 { }
 
 
@@ -69,13 +70,6 @@ float FSkeletalAnimationSection::GetSectionHeight() const
 	return (float)SkeletalAnimationEditorConstants::AnimationTrackHeight;
 }
 
-void FSkeletalAnimationSection::GenerateSectionLayout( class ISectionLayoutBuilder& LayoutBuilder ) const
-{
-	WeightArea = MakeShareable( new FFloatCurveKeyArea( &Section.Params.Weight, &Section ) );
-
-	LayoutBuilder.AddKeyArea( "Weight", NSLOCTEXT( "FSkeletalAnimationSection", "WeightArea", "Weight" ), WeightArea.ToSharedRef() );
-}
-
 
 int32 FSkeletalAnimationSection::OnPaintSection( FSequencerSectionPainter& Painter ) const
 {
@@ -87,18 +81,25 @@ int32 FSkeletalAnimationSection::OnPaintSection( FSequencerSectionPainter& Paint
 
 	static const FSlateBrush* GenericDivider = FEditorStyle::GetBrush("Sequencer.GenericDivider");
 
+	if (!Section.HasStartFrame() || !Section.HasEndFrame())
+	{
+		return LayerId;
+	}
+
 	// Add lines where the animation starts and ends/loops
 	float AnimPlayRate = FMath::IsNearlyZero(Section.Params.PlayRate) ? 1.0f : Section.Params.PlayRate;
 	float SeqLength = (Section.Params.GetSequenceLength() - (Section.Params.StartOffset + Section.Params.EndOffset)) / AnimPlayRate;
 
+	FFrameRate FrameResolution = TimeToPixelConverter.GetFrameResolution();
 	if (!FMath::IsNearlyZero(SeqLength, KINDA_SMALL_NUMBER) && SeqLength > 0)
 	{
-		float MaxOffset = Section.GetRange().Size<float>();
+		float MaxOffset  = Section.GetRange().Size<FFrameTime>() / FrameResolution;
 		float OffsetTime = SeqLength;
+		float StartTime  = Section.GetInclusiveStartFrame() / FrameResolution;
 
 		while (OffsetTime < MaxOffset)
 		{
-			float OffsetPixel = TimeToPixelConverter.TimeToPixel(Section.GetStartTime() + OffsetTime) - TimeToPixelConverter.TimeToPixel(Section.GetStartTime());
+			float OffsetPixel = TimeToPixelConverter.SecondsToPixel(StartTime + OffsetTime) - TimeToPixelConverter.SecondsToPixel(StartTime);
 
 			FSlateDrawElement::MakeBox(
 				Painter.DrawElements,
@@ -121,21 +122,23 @@ int32 FSkeletalAnimationSection::OnPaintSection( FSequencerSectionPainter& Paint
 void FSkeletalAnimationSection::BeginResizeSection()
 {
 	InitialStartOffsetDuringResize = Section.Params.StartOffset;
-	InitialStartTimeDuringResize = Section.GetStartTime();
+	InitialStartTimeDuringResize   = Section.HasStartFrame() ? Section.GetInclusiveStartFrame() : 0;
 }
 
-void FSkeletalAnimationSection::ResizeSection(ESequencerSectionResizeMode ResizeMode, float ResizeTime)
+void FSkeletalAnimationSection::ResizeSection(ESequencerSectionResizeMode ResizeMode, FFrameNumber ResizeTime)
 {
 	// Adjust the start offset when resizing from the beginning
 	if (ResizeMode == SSRM_LeadingEdge)
 	{
-		float StartOffset = (ResizeTime - InitialStartTimeDuringResize) * Section.Params.PlayRate;
+		FFrameRate FrameRate   = Section.GetTypedOuter<UMovieScene>()->GetFrameResolution();
+		float      StartOffset = (ResizeTime - InitialStartTimeDuringResize) / FrameRate * Section.Params.PlayRate;
+
 		StartOffset += InitialStartOffsetDuringResize;
 
 		// Ensure start offset is not less than 0 and adjust ResizeTime
 		if (StartOffset < 0)
 		{
-			ResizeTime = ResizeTime - (StartOffset / Section.Params.PlayRate);
+			ResizeTime = ResizeTime - ((StartOffset * Section.Params.PlayRate) * FrameRate).RoundToFrame();
 
 			StartOffset = 0.f;
 		}
@@ -151,9 +154,9 @@ void FSkeletalAnimationSection::BeginSlipSection()
 	BeginResizeSection();
 }
 
-void FSkeletalAnimationSection::SlipSection(float SlipTime)
+void FSkeletalAnimationSection::SlipSection(double SlipTime)
 {
-	float StartOffset = (SlipTime - InitialStartTimeDuringResize) * Section.Params.PlayRate;
+	float StartOffset = (SlipTime - InitialStartTimeDuringResize / Section.GetTypedOuter<UMovieScene>()->GetFrameResolution()) * Section.Params.PlayRate;
 	StartOffset += InitialStartOffsetDuringResize;
 
 	// Ensure start offset is not less than 0
@@ -241,7 +244,10 @@ bool FSkeletalAnimationTrackEditor::HandleAssetAdded(UObject* Asset, const FGuid
 				
 				UMovieSceneTrack* Track = nullptr;
 
-				AnimatablePropertyChanged(FOnKeyProperty::CreateRaw(this, &FSkeletalAnimationTrackEditor::AddKeyInternal, Object, AnimSequence, Track));
+				const FScopedTransaction Transaction(LOCTEXT("AddAnimation_Transaction", "Add Animation"));
+
+				int32 RowIndex = INDEX_NONE;
+				AnimatablePropertyChanged(FOnKeyProperty::CreateRaw(this, &FSkeletalAnimationTrackEditor::AddKeyInternal, Object, AnimSequence, Track, RowIndex));
 
 				return true;
 			}
@@ -346,12 +352,13 @@ void FSkeletalAnimationTrackEditor::OnAnimationAssetSelected(const FAssetData& A
 		UAnimSequenceBase* AnimSequence = CastChecked<UAnimSequenceBase>(AssetData.GetAsset());
 
 		UObject* Object = SequencerPtr->FindSpawnedObjectOrTemplate(ObjectBinding);
-		AnimatablePropertyChanged( FOnKeyProperty::CreateRaw( this, &FSkeletalAnimationTrackEditor::AddKeyInternal, Object, AnimSequence, Track) );
+		int32 RowIndex = INDEX_NONE;
+		AnimatablePropertyChanged( FOnKeyProperty::CreateRaw( this, &FSkeletalAnimationTrackEditor::AddKeyInternal, Object, AnimSequence, Track, RowIndex) );
 	}
 }
 
 
-FKeyPropertyResult FSkeletalAnimationTrackEditor::AddKeyInternal( float KeyTime, UObject* Object, class UAnimSequenceBase* AnimSequence, UMovieSceneTrack* Track )
+FKeyPropertyResult FSkeletalAnimationTrackEditor::AddKeyInternal( FFrameNumber KeyTime, UObject* Object, class UAnimSequenceBase* AnimSequence, UMovieSceneTrack* Track, int32 RowIndex )
 {
 	FKeyPropertyResult KeyPropertyResult;
 
@@ -370,8 +377,12 @@ FKeyPropertyResult FSkeletalAnimationTrackEditor::AddKeyInternal( float KeyTime,
 		{
 			Track->Modify();
 
-			Cast<UMovieSceneSkeletalAnimationTrack>(Track)->AddNewAnimation( KeyTime, AnimSequence );
+			UMovieSceneSection* NewSection = Cast<UMovieSceneSkeletalAnimationTrack>(Track)->AddNewAnimationOnRow( KeyTime, AnimSequence, RowIndex );
 			KeyPropertyResult.bTrackModified = true;
+
+			GetSequencer()->EmptySelection();
+			GetSequencer()->SelectSection(NewSection);
+			GetSequencer()->ThrobSectionSelection();
 		}
 	}
 
@@ -413,7 +424,7 @@ USkeleton* FSkeletalAnimationTrackEditor::AcquireSkeletonFromObjectGuid(const FG
 
 void CopyInterpAnimControlTrack(TSharedRef<ISequencer> Sequencer, UInterpTrackAnimControl* MatineeAnimControlTrack, UMovieSceneSkeletalAnimationTrack* SkeletalAnimationTrack)
 {
-	float EndPlaybackRange = Sequencer.Get().GetFocusedMovieSceneSequence()->GetMovieScene()->GetPlaybackRange().GetUpperBoundValue();
+	FFrameNumber EndPlaybackRange = MovieScene::DiscreteExclusiveUpper(Sequencer.Get().GetFocusedMovieSceneSequence()->GetMovieScene()->GetPlaybackRange());
 
 	if (FMatineeImportTools::CopyInterpAnimControlTrack(MatineeAnimControlTrack, SkeletalAnimationTrack, EndPlaybackRange))
 	{
@@ -466,4 +477,82 @@ TSharedPtr<SWidget> FSkeletalAnimationTrackEditor::BuildOutlinerEditWidget(const
 	}
 }
 
+bool FSkeletalAnimationTrackEditor::OnAllowDrop(const FDragDropEvent& DragDropEvent, UMovieSceneTrack* Track, int32 RowIndex, const FGuid& TargetObjectGuid)
+{
+	if (!Track->IsA(UMovieSceneSkeletalAnimationTrack::StaticClass()))
+	{
+		return false;
+	}
+
+	TSharedPtr<FDragDropOperation> Operation = DragDropEvent.GetOperation();
+
+	if (!Operation.IsValid() || !Operation->IsOfType<FAssetDragDropOp>() )
+	{
+		return false;
+	}
+	
+	if (!TargetObjectGuid.IsValid())
+	{
+		return false;
+	}
+
+	USkeleton* Skeleton = AcquireSkeletonFromObjectGuid(TargetObjectGuid);
+
+	TSharedPtr<FAssetDragDropOp> DragDropOp = StaticCastSharedPtr<FAssetDragDropOp>( Operation );
+
+	for (const FAssetData& AssetData : DragDropOp->GetAssets())
+	{
+		UAnimSequenceBase* AnimSequence = Cast<UAnimSequenceBase>(AssetData.GetAsset());
+
+		if (AnimSequence && Skeleton && Skeleton == AnimSequence->GetSkeleton())
+		{
+			return true;
+		}
+	}
+
+	return false;
+}
+
+
+FReply FSkeletalAnimationTrackEditor::OnDrop(const FDragDropEvent& DragDropEvent, UMovieSceneTrack* Track, int32 RowIndex, const FGuid& TargetObjectGuid)
+{
+	if (!Track->IsA(UMovieSceneSkeletalAnimationTrack::StaticClass()))
+	{
+		return FReply::Unhandled();
+	}
+
+	TSharedPtr<FDragDropOperation> Operation = DragDropEvent.GetOperation();
+
+	if (!Operation.IsValid() || !Operation->IsOfType<FAssetDragDropOp>() )
+	{
+		return FReply::Unhandled();
+	}
+	
+	if (!TargetObjectGuid.IsValid())
+	{
+		return FReply::Unhandled();
+	}
+
+	USkeleton* Skeleton = AcquireSkeletonFromObjectGuid(TargetObjectGuid);
+
+	TSharedPtr<FAssetDragDropOp> DragDropOp = StaticCastSharedPtr<FAssetDragDropOp>( Operation );
+
+	bool bAnyDropped = false;
+	for (const FAssetData& AssetData : DragDropOp->GetAssets())
+	{
+		UAnimSequenceBase* AnimSequence = Cast<UAnimSequenceBase>(AssetData.GetAsset());
+
+		if (AnimSequence && Skeleton && Skeleton == AnimSequence->GetSkeleton())
+		{
+			UObject* Object = GetSequencer()->FindSpawnedObjectOrTemplate(TargetObjectGuid);
+				
+			AnimatablePropertyChanged( FOnKeyProperty::CreateRaw(this, &FSkeletalAnimationTrackEditor::AddKeyInternal, Object, AnimSequence, Track, RowIndex));
+
+			bAnyDropped = true;
+		}
+	}
+
+	return bAnyDropped ? FReply::Handled() : FReply::Unhandled();
+}
+
 #undef LOCTEXT_NAMESPACE
diff --git a/Engine/Source/Editor/MovieSceneTools/Private/TrackEditors/SkeletalAnimationTrackEditor.h b/Engine/Source/Editor/MovieSceneTools/Private/TrackEditors/SkeletalAnimationTrackEditor.h
--- a/Engine/Source/Editor/MovieSceneTools/Private/TrackEditors/SkeletalAnimationTrackEditor.h
+++ b/Engine/Source/Editor/MovieSceneTools/Private/TrackEditors/SkeletalAnimationTrackEditor.h
@@ -13,7 +13,6 @@
 #include "MovieSceneTrackEditor.h"
 
 struct FAssetData;
-class FFloatCurveKeyArea;
 class FMenuBuilder;
 class FSequencerSectionPainter;
 class UMovieSceneSkeletalAnimationSection;
@@ -54,6 +53,8 @@ public:
 	virtual bool SupportsType( TSubclassOf<UMovieSceneTrack> Type ) const override;
 	virtual void BuildTrackContextMenu( FMenuBuilder& MenuBuilder, UMovieSceneTrack* Track ) override;
 	virtual TSharedPtr<SWidget> BuildOutlinerEditWidget(const FGuid& ObjectBinding, UMovieSceneTrack* Track, const FBuildEditWidgetParams& Params) override;
+	virtual bool OnAllowDrop(const FDragDropEvent& DragDropEvent, UMovieSceneTrack* Track, int32 RowIndex, const FGuid& TargetObjectGuid) override;
+	virtual FReply OnDrop(const FDragDropEvent& DragDropEvent, UMovieSceneTrack* Track, int32 RowIndex, const FGuid& TargetObjectGuid) override;
 
 private:
 
@@ -68,7 +69,7 @@ private:
 	void OnAnimationAssetSelected(const FAssetData& AssetData, FGuid ObjectBinding, UMovieSceneTrack* Track);
 
 	/** Delegate for AnimatablePropertyChanged in AddKey */
-	FKeyPropertyResult AddKeyInternal(float KeyTime, UObject* Object, class UAnimSequenceBase* AnimSequence, UMovieSceneTrack* Track);
+	FKeyPropertyResult AddKeyInternal(FFrameNumber KeyTime, UObject* Object, class UAnimSequenceBase* AnimSequence, UMovieSceneTrack* Track, int32 RowIndex);
 
 	/** Gets a skeleton from an object guid in the movie scene */
 	class USkeleton* AcquireSkeletonFromObjectGuid(const FGuid& Guid);
@@ -95,24 +96,20 @@ public:
 	virtual UMovieSceneSection* GetSectionObject() override;
 	virtual FText GetSectionTitle() const override;
 	virtual float GetSectionHeight() const override;
-	virtual void GenerateSectionLayout( class ISectionLayoutBuilder& LayoutBuilder ) const override;
 	virtual int32 OnPaintSection( FSequencerSectionPainter& Painter ) const override;
 	virtual void BeginResizeSection() override;
-	virtual void ResizeSection(ESequencerSectionResizeMode ResizeMode, float ResizeTime) override;
+	virtual void ResizeSection(ESequencerSectionResizeMode ResizeMode, FFrameNumber ResizeTime) override;
 	virtual void BeginSlipSection() override;
-	virtual void SlipSection(float SlipTime) override;
+	virtual void SlipSection(double SlipTime) override;
 
 private:
 
 	/** The section we are visualizing */
 	UMovieSceneSkeletalAnimationSection& Section;
 
-	/** Weight key areas. */
-	mutable TSharedPtr<FFloatCurveKeyArea> WeightArea;
-
 	/** Cached start offset value valid only during resize */
 	float InitialStartOffsetDuringResize;
 	
 	/** Cached start time valid only during resize */
-	float InitialStartTimeDuringResize;
+	FFrameNumber InitialStartTimeDuringResize;
 };
diff --git a/Engine/Source/Editor/MovieSceneTools/Private/TrackEditors/SlomoTrackEditor.cpp b/Engine/Source/Editor/MovieSceneTools/Private/TrackEditors/SlomoTrackEditor.cpp
--- a/Engine/Source/Editor/MovieSceneTools/Private/TrackEditors/SlomoTrackEditor.cpp
+++ b/Engine/Source/Editor/MovieSceneTools/Private/TrackEditors/SlomoTrackEditor.cpp
@@ -35,7 +35,8 @@ void FSlomoTrackEditor::BuildAddTrackMenu(FMenuBuilder& MenuBuilder)
 		LOCTEXT("AddPlayRateTrackTooltip", "Adds a new track that controls the playback rate of the sequence."),
 		FSlateIcon(FEditorStyle::GetStyleSetName(), "Sequencer.Tracks.Slomo"),
 		FUIAction(
-			FExecuteAction::CreateRaw(this, &FSlomoTrackEditor::HandleAddSlomoTrackMenuEntryExecute)
+			FExecuteAction::CreateRaw(this, &FSlomoTrackEditor::HandleAddSlomoTrackMenuEntryExecute),
+			FCanExecuteAction::CreateRaw(this, &FSlomoTrackEditor::HandleAddSlomoTrackMenuEntryCanExecute)
 		)
 	);
 }
@@ -87,8 +88,19 @@ void FSlomoTrackEditor::HandleAddSlomoTrackMenuEntryExecute()
 
 	SlomoTrack->AddSection(*NewSection);
 
+	if (GetSequencer().IsValid())
+	{
+		GetSequencer()->OnAddTrack(SlomoTrack);
+	}
+
 	GetSequencer()->NotifyMovieSceneDataChanged( EMovieSceneDataChangeType::MovieSceneStructureItemAdded );
 }
 
+bool FSlomoTrackEditor::HandleAddSlomoTrackMenuEntryCanExecute() const
+{
+	UMovieScene* FocusedMovieScene = GetFocusedMovieScene();
+
+	return ((FocusedMovieScene != nullptr) && (FocusedMovieScene->FindMasterTrack<UMovieSceneSlomoTrack>() == nullptr));
+}
 
 #undef LOCTEXT_NAMESPACE
diff --git a/Engine/Source/Editor/MovieSceneTools/Private/TrackEditors/SlomoTrackEditor.h b/Engine/Source/Editor/MovieSceneTools/Private/TrackEditors/SlomoTrackEditor.h
--- a/Engine/Source/Editor/MovieSceneTools/Private/TrackEditors/SlomoTrackEditor.h
+++ b/Engine/Source/Editor/MovieSceneTools/Private/TrackEditors/SlomoTrackEditor.h
@@ -47,4 +47,5 @@ private:
 
 	/** Callback for executing the "Add Slomo Track" menu entry. */
 	void HandleAddSlomoTrackMenuEntryExecute();
+	bool HandleAddSlomoTrackMenuEntryCanExecute() const;
 };
diff --git a/Engine/Source/Editor/MovieSceneTools/Private/TrackEditors/SpawnTrackEditor.cpp b/Engine/Source/Editor/MovieSceneTools/Private/TrackEditors/SpawnTrackEditor.cpp
--- a/Engine/Source/Editor/MovieSceneTools/Private/TrackEditors/SpawnTrackEditor.cpp
+++ b/Engine/Source/Editor/MovieSceneTools/Private/TrackEditors/SpawnTrackEditor.cpp
@@ -66,12 +66,6 @@ bool FSpawnTrackEditor::SupportsSequence(UMovieSceneSequence* InSequence) const
 }
 
 
-TSharedRef<ISequencerSection> FSpawnTrackEditor::MakeSectionInterface(UMovieSceneSection& SectionObject, UMovieSceneTrack& Track, FGuid ObjectBinding)
-{
-	return MakeShareable(new FBoolPropertySection(SectionObject, Track.GetDisplayName()));
-}
-
-
 void FSpawnTrackEditor::HandleAddSpawnTrackMenuEntryExecute(FGuid ObjectBinding)
 {
 	FScopedTransaction AddSpawnTrackTransaction(LOCTEXT("AddSpawnTrack_Transaction", "Add Spawn Track"));
diff --git a/Engine/Source/Editor/MovieSceneTools/Private/TrackEditors/SubTrackEditor.cpp b/Engine/Source/Editor/MovieSceneTools/Private/TrackEditors/SubTrackEditor.cpp
--- a/Engine/Source/Editor/MovieSceneTools/Private/TrackEditors/SubTrackEditor.cpp
+++ b/Engine/Source/Editor/MovieSceneTools/Private/TrackEditors/SubTrackEditor.cpp
@@ -24,6 +24,8 @@
 #include "SequenceRecorderSettings.h"
 #include "DragAndDrop/AssetDragDropOp.h"
 #include "MovieSceneToolHelpers.h"
+#include "QualifiedFrameTime.h"
+#include "MovieSceneTimeHelpers.h"
 
 namespace SubTrackEditorConstants
 {
@@ -42,12 +44,13 @@ class FSubSection
 {
 public:
 
-	FSubSection(TSharedPtr<ISequencer> InSequencer, UMovieSceneSection& InSection, const FText& InDisplayName)
+	FSubSection(TSharedPtr<ISequencer> InSequencer, UMovieSceneSection& InSection, const FText& InDisplayName, TSharedPtr<FSubTrackEditor> InSubTrackEditor)
 		: DisplayName(InDisplayName)
 		, SectionObject(*CastChecked<UMovieSceneSubSection>(&InSection))
 		, Sequencer(InSequencer)
-		, InitialStartOffsetDuringResize(0.f)
-		, InitialStartTimeDuringResize(0.f)
+		, SubTrackEditor(InSubTrackEditor)
+		, InitialStartOffsetDuringResize(0)
+		, InitialStartTimeDuringResize(0)
 	{
 	}
 
@@ -55,11 +58,6 @@ public:
 
 	// ISequencerSection interface
 
-	virtual void GenerateSectionLayout(class ISectionLayoutBuilder& LayoutBuilder) const override
-	{
-		// do nothing
-	}
-
 	virtual float GetSectionHeight() const override
 	{
 		return SubTrackEditorConstants::TrackHeight;
@@ -114,129 +112,96 @@ public:
 	{
 		int32 LayerId = InPainter.PaintSectionBackground();
 
-		const float SectionSize = SectionObject.GetTimeSize();
+		ESlateDrawEffect DrawEffects = InPainter.bParentEnabled ? ESlateDrawEffect::None : ESlateDrawEffect::DisabledEffect;
 
-		if (SectionSize <= 0.0f)
+		TRange<FFrameNumber> SectionRange = SectionObject.GetRange();
+		if (SectionRange.GetLowerBound().IsOpen() || SectionRange.GetUpperBound().IsOpen())
 		{
-			return LayerId;
+			return InPainter.LayerId;
 		}
 
-		const float DrawScale = InPainter.SectionGeometry.Size.X / SectionSize;
-		const ESlateDrawEffect DrawEffects = InPainter.bParentEnabled
-			? ESlateDrawEffect::None
-			: ESlateDrawEffect::DisabledEffect;
+		const FFrameNumber SectionStartFrame = SectionObject.GetInclusiveStartFrame();
+		const FFrameNumber SectionEndFrame   = SectionObject.GetExclusiveEndFrame();
+		const int32        SectionSize       = MovieScene::DiscreteSize(SectionRange);
 
-		UMovieScene* MovieScene = nullptr;
-		FFloatRange PlaybackRange;
-		if(SectionObject.GetSequence() != nullptr)
+		if (SectionSize <= 0)
 		{
-			MovieScene = SectionObject.GetSequence()->GetMovieScene();
-			PlaybackRange = MovieScene->GetPlaybackRange();
-		}
-		else
-		{
-			UMovieSceneTrack* MovieSceneTrack = CastChecked<UMovieSceneTrack>(SectionObject.GetOuter());
-			MovieScene = CastChecked<UMovieScene>(MovieSceneTrack->GetOuter());
-			PlaybackRange = MovieScene->GetPlaybackRange();
-		}
-
-		// add box for the working size
-		const float StartOffset = 1.0f/SectionObject.Parameters.TimeScale * SectionObject.Parameters.StartOffset;
-		const float WorkingStart = -1.0f/SectionObject.Parameters.TimeScale * PlaybackRange.GetLowerBoundValue() - StartOffset;
-		const float WorkingSize = 1.0f/SectionObject.Parameters.TimeScale * (MovieScene != nullptr ? MovieScene->GetEditorData().WorkingRange.Size<float>() : 1.0f);
-		
-		if(UMovieSceneSubSection::GetRecordingSection() == &SectionObject)
-		{
-			FColor SubSectionColor = FColor(180, 75, 75, 190);
-	
-			ISequenceRecorder& SequenceRecorder = FModuleManager::LoadModuleChecked<ISequenceRecorder>("SequenceRecorder");
-			if(SequenceRecorder.IsRecording())
-			{
-				SubSectionColor = FColor(200, 10, 10, 190);
-			}
-
-			FSlateDrawElement::MakeBox(
-				InPainter.DrawElements,
-				++LayerId,
-				InPainter.SectionGeometry.ToPaintGeometry(
-					FVector2D(WorkingStart * DrawScale, 0.f),
-					FVector2D(WorkingSize * DrawScale, InPainter.SectionGeometry.Size.Y)
-				),
-				FEditorStyle::GetBrush("Sequencer.Section.BackgroundTint"),
-				DrawEffects,
-				SubSectionColor
-			);
+			return InPainter.LayerId;
 		}
 
-		// add dark tint for left out-of-bounds & working range
-		if (StartOffset < 0.0f)
-		{
-			FSlateDrawElement::MakeBox(
-				InPainter.DrawElements,
-				++LayerId,
-				InPainter.SectionGeometry.ToPaintGeometry(
-					FVector2D(0.0f, 0.f),
-					FVector2D(-StartOffset * DrawScale, InPainter.SectionGeometry.Size.Y)
-				),
-				FEditorStyle::GetBrush("WhiteBrush"),
-				ESlateDrawEffect::None,
-				FLinearColor::Black.CopyWithNewOpacity(0.2f)
-			);
-		}
+		const float        PixelsPerFrame    = InPainter.SectionGeometry.Size.X / float(SectionSize);
 
-		// add green line for playback start
-		if (StartOffset < 0)
+		UMovieSceneSequence* InnerSequence = SectionObject.GetSequence();
+		if (InnerSequence)
 		{
-			FSlateDrawElement::MakeBox(
-				InPainter.DrawElements,
-				++LayerId,
-				InPainter.SectionGeometry.ToPaintGeometry(
-					FVector2D(-StartOffset * DrawScale, 0.f),
-					FVector2D(1.0f, InPainter.SectionGeometry.Size.Y)
-				),
-				FEditorStyle::GetBrush("WhiteBrush"),
-				ESlateDrawEffect::None,
-				FColor(32, 128, 32)	// 120, 75, 50 (HSV)
-			);
-		}
+			UMovieScene*         MovieScene    = InnerSequence->GetMovieScene();
+			TRange<FFrameNumber> PlaybackRange = MovieScene->GetPlaybackRange();
 
-		// add dark tint for right out-of-bounds & working range
-		const float PlaybackEnd = 1.0f/SectionObject.Parameters.TimeScale * PlaybackRange.Size<float>() - StartOffset;
+			FMovieSceneSequenceTransform InnerToOuterTransform = SectionObject.OuterToInnerTransform().Inverse();
 
-		if (PlaybackEnd < SectionSize)
-		{
-			FSlateDrawElement::MakeBox(
-				InPainter.DrawElements,
-				++LayerId,
-				InPainter.SectionGeometry.ToPaintGeometry(
-					FVector2D(PlaybackEnd * DrawScale, 0.f),
-					FVector2D((SectionSize - PlaybackEnd) * DrawScale, InPainter.SectionGeometry.Size.Y)
-				),
-				FEditorStyle::GetBrush("WhiteBrush"),
-				ESlateDrawEffect::None,
-				FLinearColor::Black.CopyWithNewOpacity(0.2f)
-			);
-		}
+			const FFrameNumber PlaybackStart = (MovieScene::DiscreteInclusiveLower(PlaybackRange) * InnerToOuterTransform).FloorToFrame();
+			if (SectionRange.Contains(PlaybackStart))
+			{
+				const int32 StartOffset = (PlaybackStart-SectionStartFrame).Value;
+				// add dark tint for left out-of-bounds
+				FSlateDrawElement::MakeBox(
+					InPainter.DrawElements,
+					InPainter.LayerId++,
+					InPainter.SectionGeometry.ToPaintGeometry(
+						FVector2D(0.0f, 0.f),
+						FVector2D(StartOffset * PixelsPerFrame, InPainter.SectionGeometry.Size.Y)
+					),
+					FEditorStyle::GetBrush("WhiteBrush"),
+					DrawEffects,
+					FLinearColor::Black.CopyWithNewOpacity(0.5f)
+				);
+
+				// add green line for playback start
+				FSlateDrawElement::MakeBox(
+					InPainter.DrawElements,
+					InPainter.LayerId++,
+					InPainter.SectionGeometry.ToPaintGeometry(
+						FVector2D(StartOffset * PixelsPerFrame, 0.f),
+						FVector2D(1.0f, InPainter.SectionGeometry.Size.Y)
+					),
+					FEditorStyle::GetBrush("WhiteBrush"),
+					DrawEffects,
+					FColor(32, 128, 32)	// 120, 75, 50 (HSV)
+				);
+			}
 
-		// add red line for playback end
-		if (PlaybackEnd <= SectionSize)
-		{
-			FSlateDrawElement::MakeBox(
-				InPainter.DrawElements,
-				++LayerId,
-				InPainter.SectionGeometry.ToPaintGeometry(
-					FVector2D(PlaybackEnd * DrawScale, 0.f),
-					FVector2D(1.0f, InPainter.SectionGeometry.Size.Y)
-				),
-				FEditorStyle::GetBrush("WhiteBrush"),
-				ESlateDrawEffect::None,
-				FColor(128, 32, 32)	// 0, 75, 50 (HSV)
-			);
-		}
+			const FFrameNumber PlaybackEnd = (MovieScene::DiscreteExclusiveUpper(PlaybackRange) * InnerToOuterTransform).FloorToFrame();
+			if (SectionRange.Contains(PlaybackEnd))
+			{
+				// add dark tint for right out-of-bounds
+				const int32 EndOffset = (PlaybackEnd-SectionStartFrame).Value;
+				FSlateDrawElement::MakeBox(
+					InPainter.DrawElements,
+					InPainter.LayerId++,
+					InPainter.SectionGeometry.ToPaintGeometry(
+						FVector2D(EndOffset * PixelsPerFrame, 0.f),
+						FVector2D((SectionSize - EndOffset) * PixelsPerFrame, InPainter.SectionGeometry.Size.Y)
+					),
+					FEditorStyle::GetBrush("WhiteBrush"),
+					DrawEffects,
+					FLinearColor::Black.CopyWithNewOpacity(0.5f)
+				);
+
+
+				// add red line for playback end
+				FSlateDrawElement::MakeBox(
+					InPainter.DrawElements,
+					InPainter.LayerId++,
+					InPainter.SectionGeometry.ToPaintGeometry(
+						FVector2D(EndOffset * PixelsPerFrame, 0.f),
+						FVector2D(1.0f, InPainter.SectionGeometry.Size.Y)
+					),
+					FEditorStyle::GetBrush("WhiteBrush"),
+					DrawEffects,
+					FColor(128, 32, 32)	// 0, 75, 50 (HSV)
+				);
+			}
 
-		// add number of tracks within the section's sequence
-		if(SectionObject.GetSequence() != nullptr && MovieScene != nullptr)
-		{
 			int32 NumTracks = MovieScene->GetPossessableCount() + MovieScene->GetSpawnableCount() + MovieScene->GetMasterTracks().Num();
 
 			FSlateDrawElement::MakeText(
@@ -249,20 +214,42 @@ public:
 				FColor(200, 200, 200)
 			);
 		}
-		else
+		else if (UMovieSceneSubSection::GetRecordingSection() == &SectionObject)
 		{
-			// if we are primed for recording, display where we will create the recording
-			FString Path = SectionObject.GetTargetPathToRecordTo() / SectionObject.GetTargetSequenceName();
+			FColor SubSectionColor = FColor(180, 75, 75, 190);
+	
+			ISequenceRecorder& SequenceRecorder = FModuleManager::LoadModuleChecked<ISequenceRecorder>("SequenceRecorder");
+			if(SequenceRecorder.IsRecording())
+			{
+				SubSectionColor = FColor(200, 10, 10, 190);
+			}
 
-			FSlateDrawElement::MakeText(
+			FSlateDrawElement::MakeBox(
 				InPainter.DrawElements,
 				++LayerId,
-				InPainter.SectionGeometry.ToOffsetPaintGeometry(FVector2D(11.0f, 32.0f)),
-				FText::Format(LOCTEXT("RecordingDestination", "Target: \"{0}\""), FText::FromString(Path)),
-				FEditorStyle::GetFontStyle("NormalFont"),
+				InPainter.SectionGeometry.ToPaintGeometry(
+					FVector2D(0.f, 0.f),
+					InPainter.SectionGeometry.Size
+				),
+				FEditorStyle::GetBrush("Sequencer.Section.BackgroundTint"),
 				DrawEffects,
-				FColor(200, 200, 200)
+				SubSectionColor
 			);
+
+			// display where we will create the recording
+			FString Path = SectionObject.GetTargetPathToRecordTo() / SectionObject.GetTargetSequenceName();
+			if (Path.Len() > 0)
+			{
+				FSlateDrawElement::MakeText(
+					InPainter.DrawElements,
+					++LayerId,
+					InPainter.SectionGeometry.ToOffsetPaintGeometry(FVector2D(11.0f, 32.0f)),
+					FText::Format(LOCTEXT("RecordingDestination", "Target: \"{0}\""), FText::FromString(Path)),
+					FEditorStyle::GetFontStyle("NormalFont"),
+					DrawEffects,
+					FColor(200, 200, 200)
+				);
+			}
 		}
 
 		return LayerId;
@@ -281,25 +268,59 @@ public:
 		return FReply::Handled();
 	}
 
+	virtual void BuildSectionContextMenu(FMenuBuilder& MenuBuilder, const FGuid& ObjectBinding)
+	{
+		ISequencerSection::BuildSectionContextMenu(MenuBuilder, ObjectBinding);
+
+		MenuBuilder.AddSubMenu(
+			LOCTEXT("TakesMenu", "Takes"),
+			LOCTEXT("TakesMenuTooltip", "Sub section takes"),
+			FNewMenuDelegate::CreateLambda([=](FMenuBuilder& InMenuBuilder){ AddTakesMenu(InMenuBuilder); }));
+	}
+
 	virtual void BeginSlipSection() override
 	{
-		InitialStartOffsetDuringResize = SectionObject.Parameters.StartOffset;
-		InitialStartTimeDuringResize = SectionObject.GetStartTime();
+		InitialStartOffsetDuringResize = SectionObject.Parameters.GetStartFrameOffset();
+		InitialStartTimeDuringResize   = ( SectionObject.HasStartFrame() ? SectionObject.GetInclusiveStartFrame() : 0 ) / SectionObject.GetTypedOuter<UMovieScene>()->GetFrameResolution();
 	}
 
-	virtual void SlipSection(float SlipTime) override
+	virtual void SlipSection(double SlipTime) override
 	{
-		float StartOffset = (SlipTime - InitialStartTimeDuringResize) / SectionObject.Parameters.TimeScale;
-		StartOffset += InitialStartOffsetDuringResize;
+		UMovieSceneSequence* InnerSequence = SectionObject.GetSequence();
 
-		// Ensure start offset is not less than 0
-		StartOffset = FMath::Max(StartOffset, 0.f);
+		// Adjust the start offset when resizing from the beginning
+		if (InnerSequence)
+		{
+			const FFrameRate    InnerFrameRate   = InnerSequence->GetMovieScene()->GetFrameResolution();
+			const double        ResizeDifference = SlipTime - InitialStartTimeDuringResize;
+			const int32         NewStartOffset   = ( (ResizeDifference * SectionObject.Parameters.TimeScale) * InnerFrameRate ).FrameNumber.Value;
 
-		SectionObject.Parameters.StartOffset = StartOffset;
+			// Ensure start offset is not less than 0
+			SectionObject.Parameters.SetStartFrameOffset(FMath::Max(NewStartOffset, 0));
+		}
 
 		ISequencerSection::SlipSection(SlipTime);
 	}
 
+private:
+
+	void AddTakesMenu(FMenuBuilder& MenuBuilder)
+	{
+		TArray<uint32> TakeNumbers;
+		uint32 CurrentTakeNumber = INDEX_NONE;
+		MovieSceneToolHelpers::GatherTakes(&SectionObject, TakeNumbers, CurrentTakeNumber);
+
+		for (auto TakeNumber : TakeNumbers)
+		{
+			MenuBuilder.AddMenuEntry(
+				FText::Format(LOCTEXT("TakeNumber", "Take {0}"), FText::AsNumber(TakeNumber)),
+				FText::Format(LOCTEXT("TakeNumberTooltip", "Switch to take {0}"), FText::AsNumber(TakeNumber)),
+				TakeNumber == CurrentTakeNumber ? FSlateIcon(FEditorStyle::GetStyleSetName(), "Sequencer.Star") : FSlateIcon(FEditorStyle::GetStyleSetName(), "Sequencer.Empty"),
+				FUIAction(FExecuteAction::CreateSP(SubTrackEditor.Pin().ToSharedRef(), &FSubTrackEditor::SwitchTake, &SectionObject, TakeNumber))
+			);
+		}
+	}
+
 private:
 
 	/** Display name of the section */
@@ -311,11 +332,14 @@ private:
 	/** Sequencer interface */
 	TWeakPtr<ISequencer> Sequencer;
 
+	/** The sub track editor that contains this section */
+	TWeakPtr<FSubTrackEditor> SubTrackEditor;
+
 	/** Cached start offset value valid only during resize */
-	float InitialStartOffsetDuringResize;
-	
+	int32 InitialStartOffsetDuringResize;
+
 	/** Cached start time valid only during resize */
-	float InitialStartTimeDuringResize;
+	double InitialStartTimeDuringResize;
 };
 
 
@@ -365,7 +389,7 @@ TSharedRef<ISequencerTrackEditor> FSubTrackEditor::CreateTrackEditor(TSharedRef<
 
 TSharedRef<ISequencerSection> FSubTrackEditor::MakeSectionInterface(UMovieSceneSection& SectionObject, UMovieSceneTrack& Track, FGuid ObjectBinding)
 {
-	return MakeShareable(new FSubSection(GetSequencer(), SectionObject, Track.GetDisplayName()));
+	return MakeShareable(new FSubSection(GetSequencer(), SectionObject, Track.GetDisplayName(), SharedThis(this)));
 }
 
 
@@ -393,11 +417,18 @@ bool FSubTrackEditor::HandleAssetAdded(UObject* Asset, const FGuid& TargetObject
 
 	if (CanAddSubSequence(*Sequence))
 	{
-		AnimatablePropertyChanged(FOnKeyProperty::CreateRaw(this, &FSubTrackEditor::HandleSequenceAdded, Sequence));
+		const FScopedTransaction Transaction(LOCTEXT("AddSubScene_Transaction", "Add Subscene"));
+
+		int32 RowIndex = INDEX_NONE;
+		AnimatablePropertyChanged(FOnKeyProperty::CreateRaw(this, &FSubTrackEditor::HandleSequenceAdded, Sequence, RowIndex));
 
 		return true;
 	}
 
+	FNotificationInfo Info(FText::Format( LOCTEXT("InvalidSequence", "Invalid level sequence {0}. There could be a circular dependency."), Sequence->GetDisplayName()));	
+	Info.bUseLargeFont = false;
+	FSlateNotificationManager::Get().AddNotification(Info);
+
 	return false;
 }
 
@@ -418,7 +449,7 @@ const FSlateBrush* FSubTrackEditor::GetIconBrush() const
 }
 
 
-bool FSubTrackEditor::OnAllowDrop(const FDragDropEvent& DragDropEvent, UMovieSceneTrack* Track)
+bool FSubTrackEditor::OnAllowDrop(const FDragDropEvent& DragDropEvent, UMovieSceneTrack* Track, int32 RowIndex, const FGuid& TargetObjectGuid)
 {
 	if (!Track->IsA(UMovieSceneSubTrack::StaticClass()) || Track->IsA(UMovieSceneCinematicShotTrack::StaticClass()))
 	{
@@ -446,7 +477,7 @@ bool FSubTrackEditor::OnAllowDrop(const FDragDropEvent& DragDropEvent, UMovieSce
 }
 
 
-FReply FSubTrackEditor::OnDrop(const FDragDropEvent& DragDropEvent, UMovieSceneTrack* Track)
+FReply FSubTrackEditor::OnDrop(const FDragDropEvent& DragDropEvent, UMovieSceneTrack* Track, int32 RowIndex, const FGuid& TargetObjectGuid)
 {
 	if (!Track->IsA(UMovieSceneSubTrack::StaticClass()) || Track->IsA(UMovieSceneCinematicShotTrack::StaticClass()))
 	{
@@ -469,7 +500,7 @@ FReply FSubTrackEditor::OnDrop(const FDragDropEvent& DragDropEvent, UMovieSceneT
 
 		if (Sequence)
 		{
-			AnimatablePropertyChanged(FOnKeyProperty::CreateRaw(this, &FSubTrackEditor::HandleSequenceAdded, Sequence));
+			AnimatablePropertyChanged(FOnKeyProperty::CreateRaw(this, &FSubTrackEditor::HandleSequenceAdded, Sequence, RowIndex));
 
 			bAnyDropped = true;
 		}
@@ -502,13 +533,18 @@ bool FSubTrackEditor::CanAddSubSequence(const UMovieSceneSequence& Sequence) con
 	// make sure we are not contained in the other sequence (circular dependency)
 	// @todo sequencer: this check is not sufficient (does not prevent circular dependencies of 2+ levels)
 	UMovieSceneSubTrack* SequenceSubTrack = SequenceMovieScene->FindMasterTrack<UMovieSceneSubTrack>();
+	if (SequenceSubTrack && SequenceSubTrack->ContainsSequence(*FocusedSequence, true))
+	{
+		return false;
+	}
 
-	if (SequenceSubTrack == nullptr)
+	UMovieSceneCinematicShotTrack* SequenceCinematicTrack = SequenceMovieScene->FindMasterTrack<UMovieSceneCinematicShotTrack>();
+	if (SequenceCinematicTrack && SequenceCinematicTrack->ContainsSequence(*FocusedSequence, true))
 	{
-		return true;
+		return false;
 	}
 
-	return !SequenceSubTrack->ContainsSequence(*FocusedSequence, true);
+	return true;
 }
 
 
@@ -530,6 +566,10 @@ void FSubTrackEditor::HandleAddSubTrackMenuEntryExecute()
 	auto NewTrack = FocusedMovieScene->AddMasterTrack<UMovieSceneSubTrack>();
 	ensure(NewTrack);
 
+	if (GetSequencer().IsValid())
+	{
+		GetSequencer()->OnAddTrack(NewTrack);
+	}
 	GetSequencer()->NotifyMovieSceneDataChanged( EMovieSceneDataChangeType::MovieSceneStructureItemAdded );
 }
 
@@ -620,22 +660,33 @@ void FSubTrackEditor::HandleAddSubSequenceComboButtonMenuEntryExecute(const FAss
 	{
 		UMovieSceneSequence* MovieSceneSequence = CastChecked<UMovieSceneSequence>(AssetData.GetAsset());
 
-		AnimatablePropertyChanged( FOnKeyProperty::CreateRaw( this, &FSubTrackEditor::AddKeyInternal, MovieSceneSequence, InTrack) );
+		int32 RowIndex = INDEX_NONE;
+		AnimatablePropertyChanged( FOnKeyProperty::CreateRaw( this, &FSubTrackEditor::AddKeyInternal, MovieSceneSequence, InTrack, RowIndex) );
 	}
 }
 
-FKeyPropertyResult FSubTrackEditor::AddKeyInternal(float KeyTime, UMovieSceneSequence* InMovieSceneSequence, UMovieSceneTrack* InTrack)
+FKeyPropertyResult FSubTrackEditor::AddKeyInternal(FFrameNumber KeyTime, UMovieSceneSequence* InMovieSceneSequence, UMovieSceneTrack* InTrack, int32 RowIndex)
 {	
 	FKeyPropertyResult KeyPropertyResult;
 
 	if (CanAddSubSequence(*InMovieSceneSequence))
 	{
 		UMovieSceneSubTrack* SubTrack = Cast<UMovieSceneSubTrack>(InTrack);
-		float Duration = InMovieSceneSequence->GetMovieScene()->GetPlaybackRange().Size<float>();
-		UMovieSceneSubSection* NewSection = SubTrack->AddSequence(InMovieSceneSequence, KeyTime, Duration);
-		NewSection->SetRowIndex(MovieSceneToolHelpers::FindAvailableRowIndex(SubTrack, NewSection));
+
+		const FQualifiedFrameTime InnerDuration = FQualifiedFrameTime(
+			MovieScene::DiscreteSize(InMovieSceneSequence->GetMovieScene()->GetPlaybackRange()),
+			InMovieSceneSequence->GetMovieScene()->GetFrameResolution());
+
+		const FFrameRate OuterFrameRate = SubTrack->GetTypedOuter<UMovieScene>()->GetFrameResolution();
+		const int32      OuterDuration  = InnerDuration.ConvertTo(OuterFrameRate).FrameNumber.Value;
+
+		UMovieSceneSubSection* NewSection = SubTrack->AddSequenceOnRow(InMovieSceneSequence, KeyTime, OuterDuration, RowIndex);
 		KeyPropertyResult.bTrackModified = true;
 
+		GetSequencer()->EmptySelection();
+		GetSequencer()->SelectSection(NewSection);
+		GetSequencer()->ThrobSectionSelection();
+
 		return KeyPropertyResult;
 	}
 
@@ -646,16 +697,26 @@ FKeyPropertyResult FSubTrackEditor::AddKeyInternal(float KeyTime, UMovieSceneSeq
 	return KeyPropertyResult;
 }
 
-FKeyPropertyResult FSubTrackEditor::HandleSequenceAdded(float KeyTime, UMovieSceneSequence* Sequence)
+FKeyPropertyResult FSubTrackEditor::HandleSequenceAdded(FFrameNumber KeyTime, UMovieSceneSequence* Sequence, int32 RowIndex)
 {
 	FKeyPropertyResult KeyPropertyResult;
 
 	auto SubTrack = FindOrCreateMasterTrack<UMovieSceneSubTrack>().Track;
-	float Duration = Sequence->GetMovieScene()->GetPlaybackRange().Size<float>();
-	UMovieSceneSubSection* NewSection = SubTrack->AddSequence(Sequence, KeyTime, Duration);
-	NewSection->SetRowIndex(MovieSceneToolHelpers::FindAvailableRowIndex(SubTrack, NewSection));
+
+	const FQualifiedFrameTime InnerDuration = FQualifiedFrameTime(
+		MovieScene::DiscreteSize(Sequence->GetMovieScene()->GetPlaybackRange()),
+		Sequence->GetMovieScene()->GetFrameResolution());
+
+	const FFrameRate OuterFrameRate = SubTrack->GetTypedOuter<UMovieScene>()->GetFrameResolution();
+	const int32      OuterDuration  = InnerDuration.ConvertTo(OuterFrameRate).FrameNumber.Value;
+
+	UMovieSceneSubSection* NewSection = SubTrack->AddSequenceOnRow(Sequence, KeyTime, OuterDuration, RowIndex);
 	KeyPropertyResult.bTrackModified = true;
 
+	GetSequencer()->EmptySelection();
+	GetSequencer()->SelectSection(NewSection);
+	GetSequencer()->ThrobSectionSelection();
+
 	return KeyPropertyResult;
 }
 
@@ -671,7 +732,7 @@ void FSubTrackEditor::HandleRecordNewSequence(AActor* InActorToRecord, UMovieSce
 	AnimatablePropertyChanged( FOnKeyProperty::CreateRaw( this, &FSubTrackEditor::HandleRecordNewSequenceInternal, InActorToRecord, InTrack) );
 }
 
-FKeyPropertyResult FSubTrackEditor::HandleRecordNewSequenceInternal(float KeyTime, AActor* InActorToRecord, UMovieSceneTrack* InTrack)
+FKeyPropertyResult FSubTrackEditor::HandleRecordNewSequenceInternal(FFrameNumber KeyTime, AActor* InActorToRecord, UMovieSceneTrack* InTrack)
 {
 	FKeyPropertyResult KeyPropertyResult;
 
@@ -679,12 +740,52 @@ FKeyPropertyResult FSubTrackEditor::HandleRecordNewSequenceInternal(float KeyTim
 	UMovieSceneSubSection* Section = SubTrack->AddSequenceToRecord();
 
 	// @todo: we could default to the same directory as a parent sequence, or the last sequence recorded. Lots of options!
-	Section->SetTargetSequenceName(GetDefault<USequenceRecorderSettings>()->SequenceName);
-	Section->SetTargetPathToRecordTo(GetDefault<USequenceRecorderSettings>()->SequenceRecordingBasePath.Path);
+	ISequenceRecorder& SequenceRecorder = FModuleManager::LoadModuleChecked<ISequenceRecorder>("SequenceRecorder");
+
+	Section->SetTargetSequenceName(SequenceRecorder.GetSequenceRecordingName());
+	Section->SetTargetPathToRecordTo(SequenceRecorder.GetSequenceRecordingBasePath());
 	Section->SetActorToRecord(InActorToRecord);
 	KeyPropertyResult.bTrackModified = true;
 
 	return KeyPropertyResult;
 }
 
+void FSubTrackEditor::SwitchTake(UMovieSceneSubSection* Section, uint32 TakeNumber)
+{
+	const FScopedTransaction Transaction(LOCTEXT("SwitchTake_Transaction", "Switch Take"));
+
+	UObject* TakeObject = MovieSceneToolHelpers::GetTake(Section, TakeNumber);
+
+	if (TakeObject && TakeObject->IsA(UMovieSceneSequence::StaticClass()))
+	{
+		UMovieSceneSequence* MovieSceneSequence = CastChecked<UMovieSceneSequence>(TakeObject);
+
+		UMovieSceneSubTrack* SubTrack = FindOrCreateMasterTrack<UMovieSceneSubTrack>().Track;
+		
+		TRange<FFrameNumber> NewShotRange         = Section->GetRange();
+		int32                NewShotStartOffset   = Section->Parameters.GetStartFrameOffset();
+		float                NewShotTimeScale     = Section->Parameters.TimeScale;
+		int32                NewShotPrerollFrames = Section->GetPreRollFrames();
+		int32                NewRowIndex          = Section->GetRowIndex();
+		FFrameNumber         NewShotStartTime     = NewShotRange.GetLowerBound().IsClosed() ? MovieScene::DiscreteInclusiveLower(NewShotRange) : 0;
+		int32                NewShotRowIndex      = Section->GetRowIndex();
+
+		const int32 Duration = (NewShotRange.GetLowerBound().IsClosed() && NewShotRange.GetUpperBound().IsClosed() ) ? MovieScene::DiscreteSize(NewShotRange) : 1;
+		UMovieSceneSubSection* NewShot = SubTrack->AddSequence(MovieSceneSequence, NewShotStartTime, Duration);
+
+		if (NewShot != nullptr)
+		{
+			SubTrack->RemoveSection(*Section);
+
+			NewShot->SetRange(NewShotRange);
+			NewShot->Parameters.SetStartFrameOffset(NewShotStartOffset);
+			NewShot->Parameters.TimeScale = NewShotTimeScale;
+			NewShot->SetPreRollFrames(NewShotPrerollFrames);
+			NewShot->SetRowIndex(NewShotRowIndex);
+		}
+
+		GetSequencer()->NotifyMovieSceneDataChanged( EMovieSceneDataChangeType::MovieSceneStructureItemsChanged );
+	}
+}
+
 #undef LOCTEXT_NAMESPACE
diff --git a/Engine/Source/Editor/MovieSceneTools/Private/TrackEditors/TransformTrackEditor.cpp b/Engine/Source/Editor/MovieSceneTools/Private/TrackEditors/TransformTrackEditor.cpp
--- a/Engine/Source/Editor/MovieSceneTools/Private/TrackEditors/TransformTrackEditor.cpp
+++ b/Engine/Source/Editor/MovieSceneTools/Private/TrackEditors/TransformTrackEditor.cpp
@@ -19,12 +19,12 @@
 #include "MatineeImportTools.h"
 #include "Matinee/InterpTrackMove.h"
 #include "Matinee/InterpTrackMoveAxis.h"
-#include "FloatCurveKeyArea.h"
 #include "IContentBrowserSingleton.h"
 #include "ContentBrowserModule.h"
 #include "Editor.h"
 #include "TransformPropertySection.h"
 #include "SequencerUtilities.h"
+#include "MovieSceneToolHelpers.h"
 
 #define LOCTEXT_NAMESPACE "MovieScene_TransformTrack"
 
@@ -46,64 +46,6 @@ void GetActorAndSceneComponentFromObject( UObject* Object, AActor*& OutActor, US
 	}
 }
 
-/**
- * Class that draws a transform section in the sequencer
- */
-class F3DTransformSection : public FTransformSection
-{
-public:
-
-	F3DTransformSection(TWeakObjectPtr<UMovieSceneSection> InSection, TWeakPtr<ISequencer> InSequencer, FGuid InObjectBinding)
-		: FTransformSection(InSection, InSequencer, InObjectBinding)
-	{
-	}
-
-protected:
-
-	virtual TOptional<FTransform> GetCurrentValue() const
-	{
-		TSharedPtr<ISequencer> Sequencer = WeakSequencer.Pin();
-		if (Sequencer.IsValid())
-		{
-			for (TWeakObjectPtr<> WeakObject : Sequencer->FindBoundObjects(ObjectBinding, Sequencer->GetFocusedTemplateID()))
-			{
-				if (UObject* Object = WeakObject.Get())
-				{
-					AActor* Actor = nullptr;
-					USceneComponent* Component = nullptr;
-					GetActorAndSceneComponentFromObject(Object, Actor, Component);
-					if (Component)
-					{
-						return Component->GetRelativeTransform();
-					}
-				}
-			}
-		}
-		return TOptional<FTransform>();
-	}
-	virtual TOptional<FRotator> GetCurrentRotator() const
-	{
-		TSharedPtr<ISequencer> Sequencer = WeakSequencer.Pin();
-		if (Sequencer.IsValid())
-		{
-			for (TWeakObjectPtr<> WeakObject : Sequencer->FindBoundObjects(ObjectBinding, Sequencer->GetFocusedTemplateID()))
-			{
-				if (UObject* Object = WeakObject.Get())
-				{
-					AActor* Actor = nullptr;
-					USceneComponent* Component = nullptr;
-					GetActorAndSceneComponentFromObject(Object, Actor, Component);
-					if (Component)
-					{
-						return Component->RelativeRotation;
-					}
-				}
-			}
-		}
-		return TOptional<FRotator>();
-	}
-};
-
 
 class F3DTransformTrackCommands
 	: public TCommands<F3DTransformTrackCommands>
@@ -153,7 +95,7 @@ void F3DTransformTrackCommands::RegisterCommands()
 FName F3DTransformTrackEditor::TransformPropertyName("Transform");
 
 F3DTransformTrackEditor::F3DTransformTrackEditor( TSharedRef<ISequencer> InSequencer )
-	: FKeyframeTrackEditor<UMovieScene3DTransformTrack, UMovieScene3DTransformSection, FTransformKey>( InSequencer ) 
+	: FKeyframeTrackEditor<UMovieScene3DTransformTrack>( InSequencer ) 
 {
 	// Listen for actor/component movement
 	GEditor->OnBeginObjectMovement().AddRaw( this, &F3DTransformTrackEditor::OnPreTransformChanged );
@@ -299,7 +241,7 @@ void F3DTransformTrackEditor::BuildTrackContextMenu( FMenuBuilder& MenuBuilder,
 TSharedRef<ISequencerSection> F3DTransformTrackEditor::MakeSectionInterface(UMovieSceneSection& SectionObject, UMovieSceneTrack& Track, FGuid ObjectBinding)
 {
 	check(SupportsType(SectionObject.GetOuter()->GetClass()));
-	return MakeShareable(new F3DTransformSection(&SectionObject, GetSequencer(), ObjectBinding));
+	return MakeShared<FTransformSection>(SectionObject, GetSequencer());
 }
 
 bool F3DTransformTrackEditor::HasTransformTrack(UObject& InObject) const
@@ -324,7 +266,7 @@ void F3DTransformTrackEditor::OnPreTransformChanged( UObject& InObject )
 		return;
 	}
 
-	float AutoKeyTime = GetTimeForKey();
+	FFrameNumber AutoKeyTime = GetTimeForKey();
 	AActor* Actor = Cast<AActor>(&InObject);
 	// If Sequencer is allowed to autokey and we are clicking on an Actor that can be autokeyed
 	if(Actor && !Actor->IsEditorOnly())
@@ -340,8 +282,6 @@ void F3DTransformTrackEditor::OnPreTransformChanged( UObject& InObject )
 			FTransformData Transform( SceneComponentThatChanged );
 			
 			ObjectToExistingTransform.Add(&InObject, Transform);
-
-			const bool bUnwindRotation = GetSequencer()->IsRecordingLive();
 			
 			bool bObjectHasTransformTrack = HasTransformTrack(InObject);
 			bool bComponentHasTransformTrack = HasTransformTrack(*SceneComponentThatChanged);
@@ -350,29 +290,14 @@ void F3DTransformTrackEditor::OnPreTransformChanged( UObject& InObject )
 			// If keying only happens at the end of interaction, the transform after interaction would end up incorrectly as the default state.
 			if (!bObjectHasTransformTrack && !bComponentHasTransformTrack)
 			{
-				FTransformData DummyTransform;
-				AddTransformKeys(bComponentHasTransformTrack ? SceneComponentThatChanged : &InObject, DummyTransform, Transform, EKey3DTransformChannel::All, bUnwindRotation, ESequencerKeyMode::AutoKey);
+				TOptional<FTransformData> LastTransform;
+				AddTransformKeys(bComponentHasTransformTrack ? SceneComponentThatChanged : &InObject, LastTransform, Transform, EMovieSceneTransformChannel::All, ESequencerKeyMode::AutoKey);
 			}
 		}
 	}
 }
 
 
-/**
- * Temp struct used because delegates only accept 4 or less payloads
- * FTransformKey is immutable and would require heavy re-architecting to fit here
- */
-struct FTransformDataPair
-{
-	FTransformDataPair(FTransformData InTransformData, FTransformData InLastTransformData)
-		: TransformData(InTransformData)
-		, LastTransformData(InLastTransformData) {}
-
-	FTransformData TransformData;
-	FTransformData LastTransformData;
-};
-
-
 void F3DTransformTrackEditor::OnTransformChanged( UObject& InObject )
 {
 	if (!GetSequencer()->IsAllowedToChange())
@@ -388,7 +313,11 @@ void F3DTransformTrackEditor::OnTransformChanged( UObject& InObject )
 	if( SceneComponentThatChanged != nullptr && (Actor && !Actor->IsEditorOnly()))
 	{
 		// Find an existing transform if possible.  If one exists we will compare against the new one to decide what components of the transform need keys
-		FTransformData ExistingTransform = ObjectToExistingTransform.FindRef( &InObject );
+		TOptional<FTransformData> ExistingTransform;
+		if (const FTransformData* Found = ObjectToExistingTransform.Find( &InObject ))
+		{
+			ExistingTransform = *Found;
+		}
 
 		// Remove it from the list of cached transforms. 
 		// @todo sequencer livecapture: This can be made much for efficient by not removing cached state during live capture situation
@@ -397,16 +326,14 @@ void F3DTransformTrackEditor::OnTransformChanged( UObject& InObject )
 		// Build new transform data
 		FTransformData NewTransformData( SceneComponentThatChanged );
 
-		const bool bUnwindRotation = GetSequencer()->IsRecordingLive();
-
 		bool bComponentHasTransformTrack = HasTransformTrack(*SceneComponentThatChanged);
 
-		AddTransformKeys(bComponentHasTransformTrack ? SceneComponentThatChanged : &InObject, ExistingTransform, NewTransformData, EKey3DTransformChannel::All, bUnwindRotation, ESequencerKeyMode::AutoKey);
+		AddTransformKeys(bComponentHasTransformTrack ? SceneComponentThatChanged : &InObject, ExistingTransform, NewTransformData, EMovieSceneTransformChannel::All, ESequencerKeyMode::AutoKey);
 	}
 }
 
 
-void F3DTransformTrackEditor::OnAddTransformKeysForSelectedObjects( EKey3DTransformChannel::Type Channel )
+void F3DTransformTrackEditor::OnAddTransformKeysForSelectedObjects( EMovieSceneTransformChannel Channel )
 {
 	// WASD hotkeys to fly the viewport can conflict with hotkeys for setting keyframes (ie. s). 
 	// If the viewport is moving, disregard setting keyframes.
@@ -448,19 +375,19 @@ void F3DTransformTrackEditor::BindCommands(TSharedRef<FUICommandList> SequencerC
 
 	SequencerCommandBindings->MapAction(
 		Commands.AddTransformKey,
-		FExecuteAction::CreateSP( this, &F3DTransformTrackEditor::OnAddTransformKeysForSelectedObjects, EKey3DTransformChannel::All ) );
+		FExecuteAction::CreateSP( this, &F3DTransformTrackEditor::OnAddTransformKeysForSelectedObjects, EMovieSceneTransformChannel::All ) );
 
 	SequencerCommandBindings->MapAction(
 		Commands.AddTranslationKey,
-		FExecuteAction::CreateSP( this, &F3DTransformTrackEditor::OnAddTransformKeysForSelectedObjects, EKey3DTransformChannel::Translation ) );
+		FExecuteAction::CreateSP( this, &F3DTransformTrackEditor::OnAddTransformKeysForSelectedObjects, EMovieSceneTransformChannel::Translation ) );
 
 	SequencerCommandBindings->MapAction(
 		Commands.AddRotationKey,
-		FExecuteAction::CreateSP( this, &F3DTransformTrackEditor::OnAddTransformKeysForSelectedObjects, EKey3DTransformChannel::Rotation ) );
+		FExecuteAction::CreateSP( this, &F3DTransformTrackEditor::OnAddTransformKeysForSelectedObjects, EMovieSceneTransformChannel::Rotation ) );
 
 	SequencerCommandBindings->MapAction(
 		Commands.AddScaleKey,
-		FExecuteAction::CreateSP( this, &F3DTransformTrackEditor::OnAddTransformKeysForSelectedObjects, EKey3DTransformChannel::Scale ) );
+		FExecuteAction::CreateSP( this, &F3DTransformTrackEditor::OnAddTransformKeysForSelectedObjects, EMovieSceneTransformChannel::Scale ) );
 
 	Commands.BindingCount++;
 }
@@ -470,9 +397,10 @@ void F3DTransformTrackEditor::BuildObjectBindingEditButtons(TSharedPtr<SHorizont
 {
 	bool bHasCameraComponent = false;
 
-	for (auto Object : GetSequencer()->FindObjectsInCurrentSequence(ObjectGuid))
+	UObject* Object = GetSequencer()->FindSpawnedObjectOrTemplate(ObjectGuid);
+	if (Object && Object->IsA<AActor>())
 	{
-		AActor* Actor = Cast<AActor>(Object.Get());
+		AActor* Actor = Cast<AActor>(Object);
 		if (Actor != nullptr)
 		{
 			UCameraComponent* CameraComponent = MovieSceneHelpers::CameraComponentFromActor(Actor);
@@ -519,7 +447,7 @@ void F3DTransformTrackEditor::BuildObjectBindingTrackMenu(FMenuBuilder& MenuBuil
 			NSLOCTEXT("Sequencer", "AddPTransformTooltip", "Adds a transform track."),
 			FSlateIcon(),
 			FUIAction( 
-				FExecuteAction::CreateSP( this, &F3DTransformTrackEditor::AddTransformKeysForHandle, ObjectBinding, EKey3DTransformChannel::All, ESequencerKeyMode::ManualKey )
+				FExecuteAction::CreateSP( this, &F3DTransformTrackEditor::AddTransformKeysForHandle, ObjectBinding, EMovieSceneTransformChannel::All, ESequencerKeyMode::ManualKey )
 			)
 		);
 	}
@@ -706,20 +634,6 @@ FText F3DTransformTrackEditor::GetLockCameraToolTip(FGuid ObjectGuid) const
 	return FText();
 }
 
-void GetKeysForVector( bool LastVectorIsValid, const FVector& LastVector, const FVector& CurrentVector, EKey3DTransformChannel::Type VectorChannel, EKey3DTransformChannel::Type ChannelsToKey, bool bUnwindRotation, TArray<FTransformKey>& OutNewKeys, TArray<FTransformKey>& OutDefaultKeys)
-{
-	bool bKeyChannel = ChannelsToKey == EKey3DTransformChannel::All || ChannelsToKey == VectorChannel;
-
-	TArray<FTransformKey>& XKeys = bKeyChannel && ( LastVectorIsValid == false || FMath::IsNearlyEqual( LastVector.X, CurrentVector.X ) == false ) ? OutNewKeys : OutDefaultKeys;
-	XKeys.Add( FTransformKey( VectorChannel, EAxis::X, CurrentVector.X, bUnwindRotation ) );
-	
-	TArray<FTransformKey>& YKeys = bKeyChannel && ( LastVectorIsValid == false || FMath::IsNearlyEqual( LastVector.Y, CurrentVector.Y ) == false ) ? OutNewKeys : OutDefaultKeys;
-	YKeys.Add( FTransformKey( VectorChannel, EAxis::Y, CurrentVector.Y, bUnwindRotation ) );
-
-	TArray<FTransformKey>& ZKeys = bKeyChannel && ( LastVectorIsValid == false || FMath::IsNearlyEqual( LastVector.Z, CurrentVector.Z ) == false ) ? OutNewKeys : OutDefaultKeys;
-	ZKeys.Add( FTransformKey( VectorChannel, EAxis::Z, CurrentVector.Z, bUnwindRotation ) );
-}
-
 float UnwindChannel(const float& OldValue, float NewValue)
 {
 	while( NewValue - OldValue > 180.0f )
@@ -733,33 +647,66 @@ float UnwindChannel(const float& OldValue, float NewValue)
 	return NewValue;
 }
 
-void F3DTransformTrackEditor::GetTransformKeys( const FTransformData& LastTransform, const FTransformData& CurrentTransform, EKey3DTransformChannel::Type ChannelsToKey, bool bUnwindRotation, TArray<FTransformKey>& OutNewKeys, TArray<FTransformKey>& OutDefaultKeys )
+void F3DTransformTrackEditor::GetTransformKeys( const TOptional<FTransformData>& LastTransform, const FTransformData& CurrentTransform, EMovieSceneTransformChannel ChannelsToKey, FGeneratedTrackKeys& OutGeneratedKeys )
 {
-	bool bLastVectorIsValid = LastTransform.IsValid();
+	bool bLastVectorIsValid = LastTransform.IsSet();
 
 	// If key all is enabled, for a key on all the channels
 	if (GetSequencer()->GetKeyAllEnabled())
 	{
 		bLastVectorIsValid = false;
-		ChannelsToKey = EKey3DTransformChannel::All;
+		ChannelsToKey = EMovieSceneTransformChannel::All;
 	}
 
-	// Unwind rotation
-	FRotator CurrentRotator = CurrentTransform.Rotation;
-	FRotator LastRotator = LastTransform.Rotation;
-	if (LastTransform.IsValid())
+	// Set translation keys/defaults
 	{
-		CurrentRotator.Yaw = UnwindChannel(LastRotator.Yaw, CurrentRotator.Yaw);
-		CurrentRotator.Pitch = UnwindChannel(LastRotator.Pitch, CurrentRotator.Pitch);
-		CurrentRotator.Roll = UnwindChannel(LastRotator.Roll, CurrentRotator.Roll);
+		FVector LastVector = bLastVectorIsValid ? LastTransform->Translation : FVector(), CurrentVector = CurrentTransform.Translation;
+		const bool bKeyX = EnumHasAnyFlags(ChannelsToKey, EMovieSceneTransformChannel::TranslationX) && ( !bLastVectorIsValid || !FMath::IsNearlyEqual(LastVector.X, CurrentVector.X) );
+		const bool bKeyY = EnumHasAnyFlags(ChannelsToKey, EMovieSceneTransformChannel::TranslationY) && ( !bLastVectorIsValid || !FMath::IsNearlyEqual(LastVector.Y, CurrentVector.Y) );
+		const bool bKeyZ = EnumHasAnyFlags(ChannelsToKey, EMovieSceneTransformChannel::TranslationZ) && ( !bLastVectorIsValid || !FMath::IsNearlyEqual(LastVector.Z, CurrentVector.Z) );
+
+		OutGeneratedKeys.Add(FMovieSceneChannelValueSetter::Create<FMovieSceneFloatChannel>(0, CurrentVector.X, bKeyX));
+		OutGeneratedKeys.Add(FMovieSceneChannelValueSetter::Create<FMovieSceneFloatChannel>(1, CurrentVector.Y, bKeyY));
+		OutGeneratedKeys.Add(FMovieSceneChannelValueSetter::Create<FMovieSceneFloatChannel>(2, CurrentVector.Z, bKeyZ));
 	}
 
-	GetKeysForVector( bLastVectorIsValid, LastTransform.Translation, CurrentTransform.Translation, EKey3DTransformChannel::Translation, ChannelsToKey, bUnwindRotation, OutNewKeys, OutDefaultKeys );
-	GetKeysForVector( bLastVectorIsValid, LastRotator.Euler(), CurrentRotator.Euler(), EKey3DTransformChannel::Rotation, ChannelsToKey, bUnwindRotation, OutNewKeys, OutDefaultKeys );
-	GetKeysForVector( bLastVectorIsValid, LastTransform.Scale, CurrentTransform.Scale, EKey3DTransformChannel::Scale, ChannelsToKey, bUnwindRotation, OutNewKeys, OutDefaultKeys );
+	// Set rotation keys/defaults
+	{
+		FVector LastVector = bLastVectorIsValid ? LastTransform->Rotation.Euler() : FVector(), CurrentVector = CurrentTransform.Rotation.Euler();
+
+		if (bLastVectorIsValid)
+		{
+			FRotator CurrentRotator = CurrentTransform.Rotation;
+			FRotator LastRotator = LastTransform->Rotation;
+
+			CurrentRotator.Yaw = UnwindChannel(LastRotator.Yaw, CurrentRotator.Yaw);
+			CurrentRotator.Pitch = UnwindChannel(LastRotator.Pitch, CurrentRotator.Pitch);
+			CurrentRotator.Roll = UnwindChannel(LastRotator.Roll, CurrentRotator.Roll);
+		}
+
+		const bool bKeyX = EnumHasAnyFlags(ChannelsToKey, EMovieSceneTransformChannel::RotationX) && ( !bLastVectorIsValid || !FMath::IsNearlyEqual(LastVector.X, CurrentVector.X) );
+		const bool bKeyY = EnumHasAnyFlags(ChannelsToKey, EMovieSceneTransformChannel::RotationY) && ( !bLastVectorIsValid || !FMath::IsNearlyEqual(LastVector.Y, CurrentVector.Y) );
+		const bool bKeyZ = EnumHasAnyFlags(ChannelsToKey, EMovieSceneTransformChannel::RotationZ) && ( !bLastVectorIsValid || !FMath::IsNearlyEqual(LastVector.Z, CurrentVector.Z) );
+
+		OutGeneratedKeys.Add(FMovieSceneChannelValueSetter::Create<FMovieSceneFloatChannel>(3, CurrentVector.X, bKeyX));
+		OutGeneratedKeys.Add(FMovieSceneChannelValueSetter::Create<FMovieSceneFloatChannel>(4, CurrentVector.Y, bKeyY));
+		OutGeneratedKeys.Add(FMovieSceneChannelValueSetter::Create<FMovieSceneFloatChannel>(5, CurrentVector.Z, bKeyZ));
+	}
+
+	// Set scale keys/defaults
+	{
+		FVector LastVector = bLastVectorIsValid ? LastTransform->Scale : FVector(), CurrentVector = CurrentTransform.Scale;
+		const bool bKeyX = EnumHasAnyFlags(ChannelsToKey, EMovieSceneTransformChannel::ScaleX) && ( !bLastVectorIsValid || !FMath::IsNearlyEqual(LastVector.X, CurrentVector.X) );
+		const bool bKeyY = EnumHasAnyFlags(ChannelsToKey, EMovieSceneTransformChannel::ScaleY) && ( !bLastVectorIsValid || !FMath::IsNearlyEqual(LastVector.Y, CurrentVector.Y) );
+		const bool bKeyZ = EnumHasAnyFlags(ChannelsToKey, EMovieSceneTransformChannel::ScaleZ) && ( !bLastVectorIsValid || !FMath::IsNearlyEqual(LastVector.Z, CurrentVector.Z) );
+
+		OutGeneratedKeys.Add(FMovieSceneChannelValueSetter::Create<FMovieSceneFloatChannel>(6, CurrentVector.X, bKeyX));
+		OutGeneratedKeys.Add(FMovieSceneChannelValueSetter::Create<FMovieSceneFloatChannel>(7, CurrentVector.Y, bKeyY));
+		OutGeneratedKeys.Add(FMovieSceneChannelValueSetter::Create<FMovieSceneFloatChannel>(8, CurrentVector.Z, bKeyZ));
+	}
 }
 
-void F3DTransformTrackEditor::AddTransformKeysForHandle( FGuid ObjectHandle, EKey3DTransformChannel::Type ChannelToKey, ESequencerKeyMode KeyMode )
+void F3DTransformTrackEditor::AddTransformKeysForHandle( FGuid ObjectHandle, EMovieSceneTransformChannel ChannelToKey, ESequencerKeyMode KeyMode )
 {
 	for ( TWeakObjectPtr<UObject> Object : GetSequencer()->FindObjectsInCurrentSequence(ObjectHandle) )
 	{
@@ -768,7 +715,7 @@ void F3DTransformTrackEditor::AddTransformKeysForHandle( FGuid ObjectHandle, EKe
 }
 
 
-void F3DTransformTrackEditor::AddTransformKeysForObject( UObject* Object, EKey3DTransformChannel::Type ChannelToKey, ESequencerKeyMode KeyMode )
+void F3DTransformTrackEditor::AddTransformKeysForObject( UObject* Object, EMovieSceneTransformChannel ChannelToKey, ESequencerKeyMode KeyMode )
 {
 	AActor* Actor = nullptr;
 	USceneComponent* SceneComponent = nullptr;
@@ -778,48 +725,60 @@ void F3DTransformTrackEditor::AddTransformKeysForObject( UObject* Object, EKey3D
 		FTransformData CurrentTransform( SceneComponent );
 		if(Object->GetClass()->IsChildOf(AActor::StaticClass()))
 		{
-			AddTransformKeys( Actor, FTransformData(), CurrentTransform, ChannelToKey, false, KeyMode );
+			AddTransformKeys( Actor, TOptional<FTransformData>(), CurrentTransform, ChannelToKey, KeyMode );
 		}
 		else if(Object->GetClass()->IsChildOf(USceneComponent::StaticClass()))
 		{
-			AddTransformKeys( SceneComponent, FTransformData(), CurrentTransform, ChannelToKey, false, KeyMode );
+			AddTransformKeys( SceneComponent, TOptional<FTransformData>(), CurrentTransform, ChannelToKey, KeyMode );
 		}
 	}
 }
 
 
-void F3DTransformTrackEditor::AddTransformKeys( UObject* ObjectToKey, const FTransformData& LastTransform, const FTransformData& CurrentTransform, EKey3DTransformChannel::Type ChannelsToKey, bool bUnwindRotation, ESequencerKeyMode KeyMode )
+void F3DTransformTrackEditor::AddTransformKeys( UObject* ObjectToKey, const TOptional<FTransformData>& LastTransform, const FTransformData& CurrentTransform, EMovieSceneTransformChannel ChannelsToKey, ESequencerKeyMode KeyMode )
 {
 	if (!GetSequencer()->IsAllowedToChange())
 	{
 		return;
 	}
 
-	TArray<FTransformKey> NewKeys;
-	TArray<FTransformKey> DefaultKeys;
-	GetTransformKeys(LastTransform, CurrentTransform, ChannelsToKey, bUnwindRotation, NewKeys, DefaultKeys);
+	TSharedRef<FGeneratedTrackKeys> GeneratedKeys = MakeShared<FGeneratedTrackKeys>();
+	GetTransformKeys(LastTransform, CurrentTransform, ChannelsToKey, *GeneratedKeys);
 
-	AnimatablePropertyChanged( FOnKeyProperty::CreateRaw(this, &F3DTransformTrackEditor::OnAddTransformKeys, ObjectToKey, &NewKeys, &DefaultKeys, CurrentTransform, KeyMode ) );
-}
+	auto InitializeNewTrack = [](UMovieScene3DTransformTrack* NewTrack)
+	{
+		NewTrack->SetPropertyNameAndPath(TransformPropertyName, TransformPropertyName.ToString());
+	};
+
+	auto OnKeyProperty = [=](FFrameNumber Time) -> FKeyPropertyResult
+	{
+		return this->AddKeysToObjects({ ObjectToKey }, Time, *GeneratedKeys, KeyMode, UMovieScene3DTransformTrack::StaticClass(), TransformPropertyName, InitializeNewTrack);
+	};
 
+	AnimatablePropertyChanged( FOnKeyProperty::CreateLambda(OnKeyProperty) );
+}
 
-FKeyPropertyResult F3DTransformTrackEditor::OnAddTransformKeys( float Time, UObject* ObjectToKey, TArray<FTransformKey>* NewKeys, TArray<FTransformKey>* DefaultKeys, FTransformData CurrentTransform, ESequencerKeyMode KeyMode )
+void AddUnwoundKey(FMovieSceneFloatChannel& Channel, FFrameNumber Time, float Value)
 {
-	TArray<UObject*> ObjectsToKey;
-	ObjectsToKey.Add(ObjectToKey);
-	
-	return AddKeysToObjects(
-		ObjectsToKey,
-		Time,
-		*NewKeys,
-		*DefaultKeys,
-		KeyMode,
-		UMovieScene3DTransformTrack::StaticClass(),
-		TransformPropertyName,
-		[this](UMovieScene3DTransformTrack* NewTrack) {
-			NewTrack->SetPropertyNameAndPath(TransformPropertyName, TransformPropertyName.ToString());
+	int32 Index = Channel.AddLinearKey(Time, Value);
+
+	TArrayView<FMovieSceneFloatValue> Values = Channel.GetInterface().GetValues();
+	if (Index >= 1)
+	{
+		const float PreviousValue = Values[Index - 1].Value;
+		float NewValue = Value;
+
+		while (NewValue - PreviousValue > 180.0f)
+		{
+			NewValue -= 360.f;
+		}
+		while (NewValue - PreviousValue < -180.0f)
+		{
+			NewValue += 360.f;
 		}
-	);
+
+		Values[Index].Value = NewValue;
+	}
 }
 
 
@@ -871,16 +830,25 @@ void F3DTransformTrackEditor::ImportAnimSequenceTransforms(const FAssetData& Ass
 		TransformTrack->Modify();
 
 		UMovieScene3DTransformSection* Section = Cast<UMovieScene3DTransformSection>(TransformTrack->CreateNewSection());
-		Section->GetScaleCurve(EAxis::X).SetDefaultValue(1);
-		Section->GetScaleCurve(EAxis::Y).SetDefaultValue(1);
-		Section->GetScaleCurve(EAxis::Z).SetDefaultValue(1);
+		FFrameRate FrameResolution = Section->GetTypedOuter<UMovieScene>()->GetFrameResolution();
+
+		TArrayView<FMovieSceneFloatChannel*> FloatChannels = Section->GetChannelProxy().GetChannels<FMovieSceneFloatChannel>();
+
+		// Set default translation and rotation
+		for (int32 Index = 0; Index < 6; ++Index)
+		{
+			FloatChannels[Index]->SetDefault(0.f);
+		}
+		// Set default scale
+		for (int32 Index = 6; Index < 9; ++Index)
+		{
+			FloatChannels[Index]->SetDefault(1.f);
+		}
+
 		TransformTrack->AddSection(*Section);
 
 		if (Section->TryModify())
 		{
-			float SectionMin = Section->GetStartTime();
-			float SectionMax = Section->GetEndTime();
-
 			struct FTempTransformKey
 			{
 				FTransform Transform;
@@ -943,32 +911,35 @@ void F3DTransformTrackEditor::ImportAnimSequenceTransforms(const FAssetData& Ass
 				FMath::WindRelativeAnglesDegrees(Rotator.Roll, NextRotator.Roll);
 			}
 
-			const bool bUnwindRotation = false;
+			TRange<FFrameNumber> Range = Section->GetRange();
 			for(const FTempTransformKey& TempKey : TempKeys)
 			{
-				SectionMin = FMath::Min( SectionMin, TempKey.Time );
-				SectionMax = FMath::Max( SectionMax, TempKey.Time );
+				FFrameNumber KeyTime = (TempKey.Time * FrameResolution).RoundToFrame();
 
-				Section->SetStartTime( SectionMin );
-				Section->SetEndTime( SectionMax );
+				Range = TRange<FFrameNumber>::Hull(Range, TRange<FFrameNumber>(KeyTime));
 
 				const FVector Translation = TempKey.Transform.GetTranslation();
 				const FVector Rotation = TempKey.WoundRotation.Euler();
 				const FVector Scale = TempKey.Transform.GetScale3D();
 
-				Section->AddKey(TempKey.Time, FTransformKey(EKey3DTransformChannel::Translation, EAxis::X, Translation.X, bUnwindRotation), EMovieSceneKeyInterpolation::Linear);
-				Section->AddKey(TempKey.Time, FTransformKey(EKey3DTransformChannel::Translation, EAxis::Y, Translation.Y, bUnwindRotation), EMovieSceneKeyInterpolation::Linear);
-				Section->AddKey(TempKey.Time, FTransformKey(EKey3DTransformChannel::Translation, EAxis::Z, Translation.Z, bUnwindRotation), EMovieSceneKeyInterpolation::Linear);
+				TArrayView<FMovieSceneFloatChannel*> Channels = Section->GetChannelProxy().GetChannels<FMovieSceneFloatChannel>();
 
-				Section->AddKey(TempKey.Time, FTransformKey(EKey3DTransformChannel::Rotation, EAxis::X, Rotation.X, bUnwindRotation), EMovieSceneKeyInterpolation::Linear);
-				Section->AddKey(TempKey.Time, FTransformKey(EKey3DTransformChannel::Rotation, EAxis::Y, Rotation.Y, bUnwindRotation), EMovieSceneKeyInterpolation::Linear);
-				Section->AddKey(TempKey.Time, FTransformKey(EKey3DTransformChannel::Rotation, EAxis::Z, Rotation.Z, bUnwindRotation), EMovieSceneKeyInterpolation::Linear);
+				Channels[0]->AddLinearKey(KeyTime, Translation.X);
+				Channels[1]->AddLinearKey(KeyTime, Translation.Y);
+				Channels[2]->AddLinearKey(KeyTime, Translation.Z);
 
-				Section->AddKey(TempKey.Time, FTransformKey(EKey3DTransformChannel::Scale, EAxis::X, Scale.X, bUnwindRotation), EMovieSceneKeyInterpolation::Linear);
-				Section->AddKey(TempKey.Time, FTransformKey(EKey3DTransformChannel::Scale, EAxis::Y, Scale.Y, bUnwindRotation), EMovieSceneKeyInterpolation::Linear);
-				Section->AddKey(TempKey.Time, FTransformKey(EKey3DTransformChannel::Scale, EAxis::Z, Scale.Z, bUnwindRotation), EMovieSceneKeyInterpolation::Linear);
+				AddUnwoundKey(*Channels[3], KeyTime, Rotation.X);
+				AddUnwoundKey(*Channels[4], KeyTime, Rotation.Y);
+				AddUnwoundKey(*Channels[5], KeyTime, Rotation.Z);
+
+				Channels[6]->AddLinearKey(KeyTime, Scale.X);
+				Channels[7]->AddLinearKey(KeyTime, Scale.Y);
+				Channels[8]->AddLinearKey(KeyTime, Scale.Z);
 			}
 
+			Section->SetRange(Range);
+			Section->SetRowIndex(MovieSceneToolHelpers::FindAvailableRowIndex(TransformTrack, Section));
+
 			Sequencer->NotifyMovieSceneDataChanged( EMovieSceneDataChangeType::MovieSceneStructureItemAdded );
 		}
 	}
diff --git a/Engine/Source/Editor/MovieSceneTools/Private/TrackEditors/TransformTrackEditor.h b/Engine/Source/Editor/MovieSceneTools/Private/TrackEditors/TransformTrackEditor.h
--- a/Engine/Source/Editor/MovieSceneTools/Private/TrackEditors/TransformTrackEditor.h
+++ b/Engine/Source/Editor/MovieSceneTools/Private/TrackEditors/TransformTrackEditor.h
@@ -23,7 +23,7 @@ class SHorizontalBox;
  * Tools for animatable transforms
  */
 class F3DTransformTrackEditor
-	: public FKeyframeTrackEditor<UMovieScene3DTransformTrack, UMovieScene3DTransformSection, FTransformKey>
+	: public FKeyframeTrackEditor<UMovieScene3DTransformTrack>
 {
 public:
 
@@ -96,14 +96,14 @@ private:
 
 	/** Generates transform keys based on the last transform, the current transform, and other options. 
 		One transform key is generated for each individual key to be added to the section. */
-	void GetTransformKeys( const FTransformData& LastTransform, const FTransformData& CurrentTransform, EKey3DTransformChannel::Type ChannelsToKey, bool bUnwindRotation, TArray<FTransformKey>& OutNewKeys, TArray<FTransformKey>& OutDefaultKeys );
+	void GetTransformKeys( const TOptional<FTransformData>& LastTransform, const FTransformData& CurrentTransform, EMovieSceneTransformChannel ChannelsToKey, FGeneratedTrackKeys& OutGeneratedKeys );
 
 	/**
 	* Adds transform tracks and keys to the selected objects in the level.
 	*
 	* @param Channel The transform channel to add keys for.
 	*/
-	void OnAddTransformKeysForSelectedObjects( EKey3DTransformChannel::Type Channel );
+	void OnAddTransformKeysForSelectedObjects( EMovieSceneTransformChannel Channel );
 
 	/** 
 	 * Adds transform keys to an object represented by a handle.
@@ -112,7 +112,7 @@ private:
 	 * @param ChannelToKey The channels to add keys to.
 	 * @param KeyParams Parameters which control how the keys are added. 
 	 */
-	void AddTransformKeysForHandle( FGuid ObjectHandle, EKey3DTransformChannel::Type ChannelToKey, ESequencerKeyMode KeyMode );
+	void AddTransformKeysForHandle( FGuid ObjectHandle, EMovieSceneTransformChannel ChannelToKey, ESequencerKeyMode KeyMode );
 
 	/**
 	* Adds transform keys to a specific object.
@@ -121,7 +121,7 @@ private:
 	* @param ChannelToKey The channels to add keys to.
 	* @param KeyParams Parameters which control how the keys are added.
 	*/
-	void AddTransformKeysForObject( UObject* Object, EKey3DTransformChannel::Type ChannelToKey, ESequencerKeyMode KeyMode );
+	void AddTransformKeysForObject( UObject* Object, EMovieSceneTransformChannel ChannelToKey, ESequencerKeyMode KeyMode );
 
 	/**
 	* Adds keys to a specific actor.
@@ -129,20 +129,9 @@ private:
 	* @param LastTransform The last known transform of the actor if any.
 	* @param CurrentTransform The current transform of the actor.
 	* @param ChannelToKey The channels to add keys to.
-	* @param bUnwindRotation Whether or not rotation key values should be unwound.
 	* @param KeyParams Parameters which control how the keys are added.
 	*/
-	void AddTransformKeys( UObject* ObjectToKey, const FTransformData& LastTransform, const FTransformData& CurrentTransform, EKey3DTransformChannel::Type ChannelsToKey, bool bUnwindRotation, ESequencerKeyMode KeyMode );
-
-	/**
-	* Delegate target of AnimatablePropertyChanged which actually adds the keys.
-
-	* @param Time The time to add keys.
-	* @param ObjectToKey The object to add keys to.
-	* @param Keys The keys to add.
-	* @param KeyParams Parameters which control how the keys are added.
-	*/
-	FKeyPropertyResult OnAddTransformKeys( float Time, UObject* ObjectToKey, TArray<FTransformKey>* NewKeys, TArray<FTransformKey>* DefaultKeys, FTransformData CurrentTransform, ESequencerKeyMode KeyMode );
+	void AddTransformKeys( UObject* ObjectToKey, const TOptional<FTransformData>& LastTransform, const FTransformData& CurrentTransform, EMovieSceneTransformChannel ChannelsToKey, ESequencerKeyMode KeyMode );
 
 private:
 
diff --git a/Engine/Source/Editor/MovieSceneTools/Public/ClipboardTypes.h b/Engine/Source/Editor/MovieSceneTools/Public/ClipboardTypes.h
--- a/Engine/Source/Editor/MovieSceneTools/Public/ClipboardTypes.h
+++ b/Engine/Source/Editor/MovieSceneTools/Public/ClipboardTypes.h
@@ -4,11 +4,12 @@
 
 #include "CoreMinimal.h"
 #include "Misc/Guid.h"
-#include "Curves/RichCurve.h"
 
 #if WITH_EDITOR
 
+#include "Channels/MovieSceneFloatChannel.h"
 #include "Sections/MovieSceneEventSection.h"
+#include "Sections/MovieSceneActorReferenceSection.h"
 #include "MovieSceneClipboard.h"
 
 namespace MovieSceneClipboard
@@ -23,9 +24,14 @@ namespace MovieSceneClipboard
 		static FName Name("Int");
 		return Name;
 	}
-	template<> inline FName GetKeyTypeName<FRichCurveKey>()
+	template<> inline FName GetKeyTypeName<int64>()
 	{
-		static FName Name("RichKey");
+		static FName Name("Int64");
+		return Name;
+	}
+	template<> inline FName GetKeyTypeName<FMovieSceneFloatValue>()
+	{
+		static FName Name("Float");
 		return Name;
 	}
 	template<> inline FName GetKeyTypeName<FName>()
@@ -43,9 +49,9 @@ namespace MovieSceneClipboard
 		static FName Name("String");
 		return Name;
 	}
-	template<> inline FName GetKeyTypeName<FGuid>()
+	template<> inline FName GetKeyTypeName<FMovieSceneActorReferenceKey>()
 	{
-		static FName Name("Guid");
+		static FName Name("MovieSceneActorReferenceKey");
 		return Name;
 	}
 	template<> inline FName GetKeyTypeName<FEventPayload>()
diff --git a/Engine/Source/Editor/MovieSceneTools/Public/CommonMovieSceneTools.h b/Engine/Source/Editor/MovieSceneTools/Public/CommonMovieSceneTools.h
--- a/Engine/Source/Editor/MovieSceneTools/Public/CommonMovieSceneTools.h
+++ b/Engine/Source/Editor/MovieSceneTools/Public/CommonMovieSceneTools.h
@@ -2,8 +2,9 @@
 
 #pragma once
 
-#include "CoreMinimal.h"
+#include "CoreTypes.h"
 #include "Layout/Geometry.h"
+#include "FrameRate.h"
 
 /**
  * Utility for converting time units to slate pixel units and vice versa
@@ -11,12 +12,14 @@
 struct FTimeToPixel
 {
 public:
-	FTimeToPixel( const FGeometry& AllottedGeometry, TRange<float> InLocalViewRange )
-		: MaxPixelsPerInput(1000.0f)
-		, LocalViewRange( InLocalViewRange )
+	FTimeToPixel( const FGeometry& AllottedGeometry, const TRange<double>& InLocalViewRange, const FFrameRate& InFrameResolution )
+		: ViewRangeStartSeconds( InLocalViewRange.GetLowerBoundValue() )
+		, FrameResolution( InFrameResolution )
 	{
-		const float ViewRange = LocalViewRange.Size<float>();
-		PixelsPerInput = ViewRange > 0 ? AllottedGeometry.GetLocalSize().X / ViewRange : MaxPixelsPerInput;
+		double VisibleWidth = InLocalViewRange.Size<double>();
+
+		const float MaxPixelsPerSecond = 1000.f;
+		PixelsPerSecond = VisibleWidth > 0 ? AllottedGeometry.GetLocalSize().X / VisibleWidth : MaxPixelsPerSecond;
 	}
 
 	/**
@@ -25,9 +28,9 @@ public:
 	 * @param Time	The time to convert
 	 * @return The pixel equivalent of the time
 	 */
-	float TimeToPixel( float Time ) const
+	float SecondsToPixel( double Time ) const
 	{
-		return (Time - LocalViewRange.GetLowerBoundValue()) * PixelsPerInput;
+		return (Time - ViewRangeStartSeconds) * PixelsPerSecond;
 	}
 
 	/**
@@ -36,23 +39,58 @@ public:
 	 * @param PixelX The x value of a pixel coordinate relative to the geometry that was passed into the constructor.
 	 * @return The time where the pixel is located
 	 */
-	float PixelToTime( float PixelX ) const
+	double PixelToSeconds( float PixelX ) const
 	{
-		return (PixelX/PixelsPerInput) + LocalViewRange.GetLowerBoundValue();
+		return (PixelX/PixelsPerSecond) + ViewRangeStartSeconds;
 	}
 
 	/**
-	 * @return The pixels per input
+	 * Converts a frame time to a pixel point relative to the geometry of a widget (passed into the constructor)
+	 *
+	 * @param Time The time to convert
+	 * @return The pixel equivalent of the frame time
 	 */
-	float GetPixelsPerInput() const 
+	float FrameToPixel( const FFrameTime& Time ) const
 	{
-		return PixelsPerInput;
+		return (Time / FrameResolution - ViewRangeStartSeconds) * PixelsPerSecond;
 	}
+
+	/**
+	 * Converts a pixel value to frame time 
+	 *
+	 * @param PixelX The x value of a pixel coordinate relative to the geometry that was passed into the constructor.
+	 * @return The frame time where the pixel is located
+	 */
+	FFrameTime PixelToFrame( float PixelX ) const
+	{
+		return ( PixelX / PixelsPerSecond + ViewRangeStartSeconds ) * FrameResolution;
+	}
+
+	/**
+	 * Converts a pixel delta value to delta frame time 
+	 *
+	 * @param PixelDelta The delta value in pixel space
+	 * @return The equivalent delta frame time
+	 */
+	FFrameTime PixelDeltaToFrame( float PixelDelta ) const
+	{
+		return ( PixelDelta / PixelsPerSecond ) * FrameResolution;
+	}
+
+	/**
+	 * Retrieve the frame resolution of the current sequence
+	 */
+	FFrameRate GetFrameResolution() const
+	{
+		return FrameResolution;
+	}
+
 private:
-	const float MaxPixelsPerInput;
 
-	/** time range of the sequencer */
-	TRange<float> LocalViewRange;
+	/** time range of the sequencer in seconds */
+	double ViewRangeStartSeconds;
+	/** The frame resolution of the current timeline */
+	FFrameRate FrameResolution;
 	/** The number of pixels in the view range */
-	float PixelsPerInput;
+	float PixelsPerSecond;
 };
diff --git a/Engine/Source/Editor/MovieSceneTools/Public/FrameNumberDetailsCustomization.h b/Engine/Source/Editor/MovieSceneTools/Public/FrameNumberDetailsCustomization.h
new file mode 100644
--- /dev/null
+++ b/Engine/Source/Editor/MovieSceneTools/Public/FrameNumberDetailsCustomization.h
@@ -0,0 +1,45 @@
+// Copyright 1998-2018 Epic Games, Inc. All Rights Reserved.
+
+#pragma once
+
+#include "CoreMinimal.h"
+#include "IPropertyTypeCustomization.h"
+#include "Widgets/Input/NumericTypeInterface.h"
+
+class IDetailLayoutBuilder;
+
+/**
+ *  Customize the FFrameNumber to support conversion from seconds/frames/timecode formats.
+ */
+class MOVIESCENETOOLS_API FFrameNumberDetailsCustomization : public IPropertyTypeCustomization
+{
+public:
+	static TSharedRef<IPropertyTypeCustomization> MakeInstance(TSharedPtr<INumericTypeInterface<double>> InNumericTypeInterface)
+	{
+		return MakeShared<FFrameNumberDetailsCustomization>(InNumericTypeInterface);
+	}
+
+	FFrameNumberDetailsCustomization(TSharedPtr<INumericTypeInterface<double>> InNumericTypeInterface)
+	{
+		NumericTypeInterface = InNumericTypeInterface;
+	}
+
+	/** IPropertyTypeCustomization interface */
+	virtual void CustomizeHeader(TSharedRef<IPropertyHandle> PropertyHandle, FDetailWidgetRow& HeaderRow, IPropertyTypeCustomizationUtils& CustomizationUtils) override {};
+	virtual void CustomizeChildren(TSharedRef<IPropertyHandle> PropertyHandle, IDetailChildrenBuilder& ChildBuilder, IPropertyTypeCustomizationUtils& CustomizationUtils) override;
+
+private:
+	FText OnGetTimeText() const;
+	void OnTimeTextCommitted(const FText& InText, ETextCommit::Type CommitInfo);
+
+	/** The Numeric Type interface used to convert between display formats and internal frame resolution. */
+	TSharedPtr<INumericTypeInterface<double>> NumericTypeInterface;
+
+	/** Store the property handle to the FrameNumber field so we can get/set the value on the object via text box callbacks. */
+	TSharedPtr<IPropertyHandle> FrameNumberProperty;
+
+	/** If they've used the UIMin metadata on the FFrameNumber property, we store that for use via text box callbacks. */
+	int32 UIClampMin;
+	/** If they've used the UIMax metadata on the FFrameNumber property, we store that for use via text box callbacks. */
+	int32 UIClampMax;
+};
diff --git a/Engine/Source/Editor/MovieSceneTools/Public/KeyframeTrackEditor.h b/Engine/Source/Editor/MovieSceneTools/Public/KeyframeTrackEditor.h
--- a/Engine/Source/Editor/MovieSceneTools/Public/KeyframeTrackEditor.h
+++ b/Engine/Source/Editor/MovieSceneTools/Public/KeyframeTrackEditor.h
@@ -14,16 +14,137 @@
 #include "MovieSceneTrackEditor.h"
 #include "Framework/MultiBox/MultiBoxBuilder.h"
 #include "MovieSceneCommonHelpers.h"
+#include "Channels/MovieSceneChannelProxy.h"
 
-template<typename KeyDataType> class IKeyframeSection;
+
+struct FMovieSceneChannelValueSetter
+{
+	FMovieSceneChannelValueSetter(const FMovieSceneChannelValueSetter&) = delete;
+	FMovieSceneChannelValueSetter& operator=(const FMovieSceneChannelValueSetter&) = delete;
+
+	FMovieSceneChannelValueSetter(FMovieSceneChannelValueSetter&&) = default;
+	FMovieSceneChannelValueSetter& operator=(FMovieSceneChannelValueSetter&&) = default;
+
+	/** Templated construction function that can add a key (and potentially also set a default) for the specified channel and value */
+	template<typename ChannelType, typename ValueType>
+	static FMovieSceneChannelValueSetter Create(int32 ChannelIndex, const ValueType& InNewValue, bool bAddKey)
+	{
+		FMovieSceneChannelValueSetter NewValue;
+		if (bAddKey)
+		{
+			NewValue.Impl = TAddKeyImpl<ChannelType, ValueType>(ChannelIndex, InNewValue);
+		}
+		else
+		{
+			NewValue.Impl = TSetDefaultImpl<ChannelType, ValueType>(ChannelIndex, InNewValue);
+		}
+		return MoveTemp(NewValue);
+	}
+
+	struct IImpl
+	{
+		virtual ~IImpl(){}
+
+		/* Returns whether a key was created */
+		virtual bool Apply(FMovieSceneChannelProxy& Proxy, FFrameNumber InTime, EMovieSceneKeyInterpolation InterpolationMode, bool bKeyEvenIfUnchanged, bool bKeyEvenIfEmpty) const { return false; }
+
+		virtual void ApplyDefault(FMovieSceneChannelProxy& Proxy) const { }
+	};
+
+	IImpl* operator->()
+	{
+		return &Impl.GetValue();
+	}
+
+	const IImpl* operator->() const
+	{
+		return &Impl.GetValue();
+	}
+
+private:
+
+	FMovieSceneChannelValueSetter()
+	{}
+
+	template<typename ChannelType, typename ValueType>
+	struct TSetDefaultImpl : IImpl
+	{
+		int32 ChannelIndex;
+		ValueType ValueToSet;
+
+		TSetDefaultImpl(int32 InChannelIndex, const ValueType& InValue)
+			: ChannelIndex(InChannelIndex), ValueToSet(InValue)
+		{}
+
+		virtual void ApplyDefault(FMovieSceneChannelProxy& Proxy) const override
+		{
+			ChannelType* Channel = Proxy.GetChannel<ChannelType>(ChannelIndex);
+			if (Channel)
+			{
+				Channel->SetDefault(ValueToSet);
+			}
+		}
+	};
+
+	template<typename ChannelType, typename ValueType>
+	struct TAddKeyImpl : IImpl
+	{
+		int32 ChannelIndex;
+		ValueType ValueToSet;
+
+		TAddKeyImpl(int32 InChannelIndex, const ValueType& InValue)
+			: ChannelIndex(InChannelIndex), ValueToSet(InValue)
+		{}
+
+		virtual bool Apply(FMovieSceneChannelProxy& Proxy, FFrameNumber InTime, EMovieSceneKeyInterpolation InterpolationMode, bool bKeyEvenIfUnchanged, bool bKeyEvenIfEmpty) const override
+		{
+			bool bKeyCreated = false;
+			using namespace MovieScene;
+
+			ChannelType* Channel = Proxy.GetChannel<ChannelType>(ChannelIndex);
+			if (Channel)
+			{
+				bool bShouldKeyChannel = bKeyEvenIfUnchanged;
+				if (!bShouldKeyChannel)
+				{
+					bShouldKeyChannel = !ValueExistsAtTime(Channel, InTime, ValueToSet);
+				}
+
+				if (bShouldKeyChannel)
+				{
+					if (Channel->GetInterface().GetTimes().Num() != 0 || bKeyEvenIfEmpty)
+					{
+						AddKeyToChannel(Channel, InTime, ValueToSet, InterpolationMode);
+						bKeyCreated = true;
+					}
+				}
+			}
+
+			return bKeyCreated;
+		}
+
+		virtual void ApplyDefault(FMovieSceneChannelProxy& Proxy) const override
+		{
+			using namespace MovieScene;
+
+			ChannelType* Channel = Proxy.GetChannel<ChannelType>(ChannelIndex);
+			if (Channel)
+			{
+				Channel->SetDefault(ValueToSet);
+			}
+		}
+	};
+	TInlineValue<IImpl> Impl;
+};
+
+typedef TArray<FMovieSceneChannelValueSetter, TInlineAllocator<1>> FGeneratedTrackKeys;
 
 /**
- * A base class for track editors that edit tracks which contain sections implementing IKeyframeSection.
+ * A base class for track editors that edit tracks which contain sections implementing GetKeyDataInterface.
   */
-template<typename TrackType, typename SectionType, typename KeyDataType>
+template<typename TrackType>
 class FKeyframeTrackEditor : public FMovieSceneTrackEditor
 {
-
 public:
 	/**
 	* Constructor
@@ -62,9 +183,8 @@ protected:
 	 *
 	 * @param ObjectsToKey An array of objects to add keyframes to.
 	 * @param KeyTime The time to add keys.
-	 * @param NewKeys The new keys to add.
-	 * @param DefaultKeys Extra keys with default values which shouldn't be added directly, but are needed to set correct
-	 *        default values when adding single channel keys for multi-channel tracks like vectors.
+	 * @param KeyedChannels Aggregate of optionally enabled channel values for which to create new keys
+	 * @param ChannelDefaults Aggregate of optionally enabled channel default values that should be applied to the section
 	 * @param KeyParams The parameters to control keyframing behavior.
 	 * @param TrackClass The class of track which should be created if specified in the parameters.
 	 * @param PropertyName The name of the property to add keys for.
@@ -73,8 +193,7 @@ protected:
 	 * @return Whether or not a handle guid or track was created. Note this does not return true if keys were added or modified.
 	 */
 	FKeyPropertyResult AddKeysToObjects(
-		TArray<UObject*> ObjectsToKey, float KeyTime,
-		const TArray<KeyDataType>& NewKeys, const TArray<KeyDataType>& DefaultKeys,
+		TArrayView<UObject* const> ObjectsToKey, FFrameNumber KeyTime, const FGeneratedTrackKeys& GeneratedKeys,
 		ESequencerKeyMode KeyMode, TSubclassOf<UMovieSceneTrack> TrackClass, FName PropertyName,
 		TFunction<void(TrackType*)> OnInitializeNewTrack)
 	{
@@ -97,7 +216,7 @@ protected:
 
 			if ( ObjectHandle.IsValid() )
 			{
-				KeyPropertyResult |= AddKeysToHandle( ObjectHandle, KeyTime, NewKeys, DefaultKeys, KeyMode, TrackClass, PropertyName, OnInitializeNewTrack );
+				KeyPropertyResult |= AddKeysToHandle( ObjectHandle, KeyTime, GeneratedKeys, KeyMode, TrackClass, PropertyName, OnInitializeNewTrack );
 			}
 		}
 		return KeyPropertyResult;
@@ -106,14 +225,18 @@ protected:
 
 private:
 
-	virtual void ClearDefaults( UMovieSceneTrack* Track )
+	void ClearDefaults( UMovieSceneTrack* Track )
 	{
 		const FScopedTransaction Transaction(NSLOCTEXT("KeyframeTrackEditor", "ClearTrackDefaultsTransaction", "Clear track defaults"));
 		for (UMovieSceneSection* Section : Track->GetAllSections())
 		{
 			Section->Modify();
-			IKeyframeSection<KeyDataType>* KeyframeSection = CastChecked<SectionType>(Section);
-			KeyframeSection->ClearDefaults();
+
+			// Clear all defaults on the section
+			for (const FMovieSceneChannelEntry& Entry : Section->GetChannelProxy().GetAllEntries())
+			{
+				Entry.GetBatchChannelInterface().ClearDefaults_Batch(Entry.GetChannels());
+			}
 		}
 		GetSequencer()->NotifyMovieSceneDataChanged(EMovieSceneDataChangeType::TrackValueChanged);
 	}
@@ -134,9 +257,7 @@ private:
 	 *
 	 * @param ObjectsToKey An array of objects to add keyframes to.
 	 * @param KeyTime The time to add keys.
-	 * @param NewKeys The new keys to add.
-	 * @param DefaultKeys Extra keys with default values which shouldn't be added directly, but are needed to set correct
-	 *        default values when adding single channel keys for multi-channel tracks like vectors.
+	 * @param GeneratedKeys Array of keys to set
 	 * @param KeyParams The parameters to control keyframing behavior.
 	 * @param TrackClass The class of track which should be created if specified in the parameters.
 	 * @param PropertyName The name of the property to add keys for.
@@ -145,13 +266,11 @@ private:
 	 * @return Whether or not a track was created. Note this does not return true if keys were added or modified.
 	*/
 	FKeyPropertyResult AddKeysToHandle(
-		FGuid ObjectHandle, float KeyTime,
-		const TArray<KeyDataType>& NewKeys, const TArray<KeyDataType>& DefaultKeys,
+		FGuid ObjectHandle, FFrameNumber KeyTime, const FGeneratedTrackKeys& GeneratedKeys,
 		ESequencerKeyMode KeyMode, TSubclassOf<UMovieSceneTrack> TrackClass, FName PropertyName,
 		TFunction<void(TrackType*)> OnInitializeNewTrack)
 	{
 		bool bTrackCreated = false;
-		bool bSectionCreated = false;
 
 		EAutoChangeMode AutoChangeMode = GetSequencer()->GetAutoChangeMode();
 		EAllowEditsMode AllowEditsMode = GetSequencer()->GetAllowEditsMode();
@@ -175,10 +294,31 @@ private:
 			}
 		}
 
+		bool bSectionCreated = false;
+
 		FKeyPropertyResult KeyPropertyResult;
+
 		if ( Track )
 		{
-			KeyPropertyResult |= AddKeysToTrack( Track, KeyTime, NewKeys, DefaultKeys, KeyMode, bTrackCreated );
+			UMovieSceneSection* SectionToKey = Track->FindOrExtendSection(KeyTime);
+
+			// If there's no overlapping section to key, create one only if a track was newly created. Otherwise, skip keying altogether
+			// so that the user is forced to create a section to key on.
+			if (bTrackCreated && !SectionToKey)
+			{
+				Track->Modify();
+
+				SectionToKey = Track->FindOrAddSection(KeyTime, bSectionCreated);
+				if (bSectionCreated && GetSequencer()->GetInfiniteKeyAreas())
+				{
+					SectionToKey->SetRange(TRange<FFrameNumber>::All());
+				}
+			}
+
+			if (SectionToKey && CanAutoKeySection(SectionToKey, KeyTime))
+			{
+				KeyPropertyResult |= AddKeysToSection( SectionToKey, KeyTime, GeneratedKeys, KeyMode );
+			}
 		}
 
 		KeyPropertyResult.bTrackCreated |= bTrackCreated || bSectionCreated;
@@ -186,158 +326,56 @@ private:
 		return KeyPropertyResult;
 	}
 
-	FKeyPropertyResult AddKeysToTrack(
-		TrackType* Track, float KeyTime,
-		const TArray<KeyDataType>& NewKeys, const TArray<KeyDataType>& DefaultKeys,
-		ESequencerKeyMode KeyMode, bool bNewTrack)
+	/* Returns whether a section was added */
+	FKeyPropertyResult AddKeysToSection(UMovieSceneSection* Section, FFrameNumber KeyTime, const FGeneratedTrackKeys& Keys, ESequencerKeyMode KeyMode)
 	{
-		bool bSectionCreated = false;
-		bool bKeyCreated = false;
-		bool bInfiniteKeyAreas = GetSequencer()->GetInfiniteKeyAreas();
+		FKeyPropertyResult KeyPropertyResult;
+		if (!Section->TryModify())
+		{
+			return KeyPropertyResult;
+		}
 
 		EAutoChangeMode AutoChangeMode = GetSequencer()->GetAutoChangeMode();
 
-		if ( KeyMode != ESequencerKeyMode::AutoKey || AutoChangeMode == EAutoChangeMode::AutoKey || AutoChangeMode == EAutoChangeMode::All )
+		FMovieSceneChannelProxy& Proxy = Section->GetChannelProxy();
+			
+		const bool bSetDefaults = GetSequencer()->GetAutoSetTrackDefaults();
+		
+		if ( KeyMode != ESequencerKeyMode::AutoKey || AutoChangeMode == EAutoChangeMode::AutoKey || AutoChangeMode == EAutoChangeMode::All || bSetDefaults)
 		{
 			EMovieSceneKeyInterpolation InterpolationMode = GetSequencer()->GetKeyInterpolation();
 
-			bool bKeyEvenIfUnchanged =
+			const bool bKeyEvenIfUnchanged =
 				KeyMode == ESequencerKeyMode::ManualKeyForced ||
 				GetSequencer()->GetKeyAllEnabled();
 
-			bool bKeyEvenIfEmpty =
+			const bool bKeyEvenIfEmpty =
 				(KeyMode == ESequencerKeyMode::AutoKey && AutoChangeMode == EAutoChangeMode::All) ||
 				KeyMode == ESequencerKeyMode::ManualKeyForced;
 
-			for (const KeyDataType& NewKey : NewKeys)
-			{
-				if ( NewKeyIsNewData( Track, KeyTime, NewKey ) || bKeyEvenIfUnchanged )
-				{
-					if ( HasKeys( Track, NewKey ) || bKeyEvenIfEmpty )
-					{
-						bSectionCreated |= AddKey( Track, KeyTime, NewKey, InterpolationMode, bInfiniteKeyAreas );
-						bKeyCreated = true;
-					}
-				}
-			}
-		}
-
-		if (GetSequencer()->GetAutoSetTrackDefaults())
-		{
-			for (const KeyDataType& NewKey : NewKeys)
+			for (const FMovieSceneChannelValueSetter& GeneratedKey : Keys)
 			{
-				bSectionCreated |= SetDefault(Track, KeyTime, NewKey, bInfiniteKeyAreas);
+				KeyPropertyResult.bKeyCreated |= GeneratedKey->Apply(Proxy, KeyTime, InterpolationMode, bKeyEvenIfUnchanged, bKeyEvenIfEmpty);
 			}
 
-			for (const KeyDataType& DefaultKey : DefaultKeys)
+			if (bSetDefaults)
 			{
-				bSectionCreated |= SetDefault(Track, KeyTime, DefaultKey, bInfiniteKeyAreas);
+				for (const FMovieSceneChannelValueSetter& GeneratedKey : Keys)
+				{
+					GeneratedKey->ApplyDefault(Proxy);
+				}
 			}
 		}
 
-		// If a new track was created but no keys or defaults were set, make sure a new section is created too to allow the user to edit it.
-		if ( bNewTrack && bSectionCreated == false)
-		{
-			Track->Modify();
-			UMovieSceneSection* NewSection = Track->FindOrAddSection(KeyTime, bSectionCreated);
-			NewSection->SetIsInfinite(bInfiniteKeyAreas);
-		}
-
-		FKeyPropertyResult KeyPropertyResult;
-		KeyPropertyResult.bTrackCreated |= bSectionCreated;
-		KeyPropertyResult.bKeyCreated = bKeyCreated;
-
 		return KeyPropertyResult;
 	}
 
-	bool NewKeyIsNewData( TrackType* Track, float Time, const KeyDataType& KeyData ) const
-	{
-		IKeyframeSection<KeyDataType>* NearestKeyframeSection = CastChecked<SectionType>( MovieSceneHelpers::FindNearestSectionAtTime( Track->GetAllSections(), Time ), ECastCheckedType::NullAllowed );
-		return NearestKeyframeSection == nullptr || NearestKeyframeSection->NewKeyIsNewData( Time, KeyData );
-	}
-
-	bool HasKeys( TrackType* Track, const KeyDataType& KeyData ) const
-	{
-		for ( UMovieSceneSection* Section : Track->GetAllSections() )
-		{
-			IKeyframeSection<KeyDataType>* KeyframeSection = CastChecked<SectionType>( Section );
-			if ( KeyframeSection->HasKeys( KeyData ) )
-			{
-				return true;
-			}
-		}
-		return false;
-	}
-
-	using FMovieSceneTrackEditor::AddKey;
-
-	/* Return whether a section was added */
-	bool AddKey( TrackType* Track, float Time, const KeyDataType& KeyData, EMovieSceneKeyInterpolation KeyInterpolation, bool bInfiniteKeyAreas )
-	{
-		bool bSectionAdded = false;
-		Track->Modify();
-		UMovieSceneSection* NewSection = Track->FindOrAddSection( Time, bSectionAdded );
-		if (!bSectionAdded && !CanAutoKeySection(NewSection, Time))
-		{
-			return false;
-		}
-
-		IKeyframeSection<KeyDataType>* KeyframeSection = CastChecked<SectionType>( NewSection );
-		KeyframeSection->AddKey( Time, KeyData, KeyInterpolation );
-
-		if (bSectionAdded)
-		{
-			NewSection->SetIsInfinite(bInfiniteKeyAreas);
-		}
-		return bSectionAdded;
-	}
-
 	/** Check whether we can autokey the specified section at the specified time */
-	static bool CanAutoKeySection(UMovieSceneSection* Section, float Time)
+	static bool CanAutoKeySection(UMovieSceneSection* Section, FFrameNumber Time)
 	{
 		FOptionalMovieSceneBlendType BlendType = Section->GetBlendType();
 		// Sections are only eligible for autokey if they are not blendable (or absolute), and overlap the current time
-		return ( !BlendType.IsValid() || BlendType.Get() == EMovieSceneBlendType::Absolute ) &&
-			   ( Section->IsInfinite() || Section->GetRange().Contains(Time) );
-	}
-
-	/* Return whether a section was added */
-	bool SetDefault( TrackType* Track, float Time, const KeyDataType& KeyData, bool bInfiniteKeyAreas )
-	{
-		bool bSectionAdded = false;
-		const TArray<UMovieSceneSection*>& Sections = Track->GetAllSections();
-		if ( Sections.Num() > 0)
-		{
-			for ( UMovieSceneSection* Section : Sections )
-			{
-				if ( !CanAutoKeySection(Section, Time) )
-				{
-					continue;
-				}
-
-				IKeyframeSection<KeyDataType>* KeyframeSection = CastChecked<SectionType>( Section );
-				if (!KeyframeSection->HasKeys(KeyData))
-				{
-					KeyframeSection->SetDefault( KeyData );
-				}
-			}
-		}
-		else
-		{
-			UMovieSceneSection* NewSection = Track->FindOrAddSection( Time, bSectionAdded );
-			IKeyframeSection<KeyDataType>* KeyframeSection = CastChecked<SectionType>( NewSection );
-			if (!KeyframeSection->HasKeys(KeyData))
-			{
-				KeyframeSection->SetDefault( KeyData );
-			}
-
-			if (bSectionAdded)
-			{
-				NewSection->SetIsInfinite(bInfiniteKeyAreas);
-			}
-		}
-		return bSectionAdded;
+		return ( !BlendType.IsValid() || BlendType.Get() == EMovieSceneBlendType::Absolute ) && Section->GetRange().Contains(Time);
 	}
 };
 
-
diff --git a/Engine/Source/Editor/MovieSceneTools/Public/MatineeImportTools.h b/Engine/Source/Editor/MovieSceneTools/Public/MatineeImportTools.h
--- a/Engine/Source/Editor/MovieSceneTools/Public/MatineeImportTools.h
+++ b/Engine/Source/Editor/MovieSceneTools/Public/MatineeImportTools.h
@@ -4,6 +4,7 @@
 #include "CoreMinimal.h"
 #include "Matinee/InterpTrackToggle.h"
 #include "Sections/MovieSceneParticleSection.h"
+#include "Channels/MovieSceneFloatChannel.h"
 
 class AMatineeActor;
 class IMovieScenePlayer;
@@ -43,10 +44,10 @@ public:
 	static ERichCurveTangentMode MatineeInterpolationToRichCurveTangent( EInterpCurveMode CurveMode );
 
 	/** Tries to convert a matinee toggle to a particle key. */
-	static bool TryConvertMatineeToggleToOutParticleKey( ETrackToggleAction ToggleAction, EParticleKey::Type& OutParticleKey );
+	static bool TryConvertMatineeToggleToOutParticleKey( ETrackToggleAction ToggleAction, EParticleKey& OutParticleKey );
 
 	/** Adds a key to a rich curve based on matinee curve key data. */
-	static void SetOrAddKey( FRichCurve& Curve, float Time, float Value, float ArriveTangent, float LeaveTangent, EInterpCurveMode MatineeInterpMode );
+	static void SetOrAddKey( TMovieSceneChannel<FMovieSceneFloatValue>& Curve, FFrameNumber Time, float Value, float ArriveTangent, float LeaveTangent, EInterpCurveMode MatineeInterpMode );
 
 	/** Copies keys from a matinee bool track to a sequencer bool track. */
 	static bool CopyInterpBoolTrack( UInterpTrackBoolProp* MatineeBoolTrack, UMovieSceneBoolTrack* BoolTrack );
@@ -70,7 +71,7 @@ public:
 	static bool CopyInterpParticleTrack( UInterpTrackToggle* MatineeToggleTrack, UMovieSceneParticleTrack* ParticleTrack );
 
 	/** Copies keys from a matinee anim control track to a sequencer skeletal animation track. */
-	static bool CopyInterpAnimControlTrack( UInterpTrackAnimControl* MatineeAnimControlTrack, UMovieSceneSkeletalAnimationTrack* SkeletalAnimationTrack, float EndPlaybackRange );
+	static bool CopyInterpAnimControlTrack( UInterpTrackAnimControl* MatineeAnimControlTrack, UMovieSceneSkeletalAnimationTrack* SkeletalAnimationTrack, FFrameNumber EndPlaybackRange );
 
 	/** Copies keys from a matinee sound track to a sequencer audio track. */
 	static bool CopyInterpSoundTrack( UInterpTrackSound* MatineeSoundTrack, UMovieSceneAudioTrack* AudioTrack );
diff --git a/Engine/Source/Editor/MovieSceneTools/Public/MovieSceneObjectBindingIDCustomization.h b/Engine/Source/Editor/MovieSceneTools/Public/MovieSceneObjectBindingIDCustomization.h
--- a/Engine/Source/Editor/MovieSceneTools/Public/MovieSceneObjectBindingIDCustomization.h
+++ b/Engine/Source/Editor/MovieSceneTools/Public/MovieSceneObjectBindingIDCustomization.h
@@ -14,7 +14,7 @@ class ISequencer;
 struct FMovieSceneObjectBindingID;
 
 
-class FMovieSceneObjectBindingIDCustomization
+class MOVIESCENETOOLS_API FMovieSceneObjectBindingIDCustomization
 	: public IPropertyTypeCustomization
 	, FMovieSceneObjectBindingIDPicker
 {
@@ -27,7 +27,7 @@ public:
 		: FMovieSceneObjectBindingIDPicker(InLocalSequenceID, InSequencer)
 	{}
 
-	MOVIESCENETOOLS_API static void BindTo(TSharedRef<ISequencer> InSequencer);
+	static void BindTo(TSharedRef<ISequencer> InSequencer);
 
 	virtual void CustomizeHeader(TSharedRef<IPropertyHandle> PropertyHandle, FDetailWidgetRow& HeaderRow, IPropertyTypeCustomizationUtils& CustomizationUtils) override;
 	virtual void CustomizeChildren(TSharedRef<IPropertyHandle> StructPropertyHandle, class IDetailChildrenBuilder& ChildBuilder, IPropertyTypeCustomizationUtils& StructCustomizationUtils) override {}
