Copying //UE4/Dev-Sequencer to //UE4/Dev-Main (Source: //UE4/Dev-Sequencer @ 4044339)

diff --git a/Engine/Source/Editor/CurveEditor/Public/CurveDataAbstraction.h b/Engine/Source/Editor/CurveEditor/Public/CurveDataAbstraction.h
--- a/Engine/Source/Editor/CurveEditor/Public/CurveDataAbstraction.h
+++ b/Engine/Source/Editor/CurveEditor/Public/CurveDataAbstraction.h
@@ -38,23 +38,32 @@ struct FKeyAttributes
 		bHasLeaveTangent  = 0;
 		bHasInterpMode    = 0;
 		bHasTangentMode   = 0;
+		bHasTangentWeightMode = 0;
+		bHasArriveTangentWeight = 0;
+		bHasLeaveTangentWeight = 0;
 	}
 
 	/**
 	 * Check whether this key has the specified attributes
 	 */
-	bool HasArriveTangent() const { return bHasArriveTangent; }
-	bool HasLeaveTangent() const  { return bHasLeaveTangent;  }
-	bool HasInterpMode() const    { return bHasInterpMode;    }
-	bool HasTangentMode() const   { return bHasTangentMode;   }
-
+	bool HasArriveTangent() const		{ return bHasArriveTangent;		}
+	bool HasLeaveTangent() const		{ return bHasLeaveTangent;		}
+	bool HasInterpMode() const			{ return bHasInterpMode;		}
+	bool HasTangentMode() const			{ return bHasTangentMode;		}
+	bool HasTangentWeightMode() const	{ return bHasTangentWeightMode; }
+	bool HasArriveTangentWeight() const { return bHasArriveTangentWeight; }
+	bool HasLeaveTangentWeight() const  { return bHasLeaveTangentWeight; }
+	
 	/**
 	 * Retrieve specific attributes for this key. Must check such attributes exist.
 	 */
-	float GetArriveTangent() const               { check(bHasArriveTangent); return ArriveTangent; }
-	float GetLeaveTangent() const                { check(bHasLeaveTangent);  return LeaveTangent;  }
-	ERichCurveInterpMode GetInterpMode() const   { check(bHasInterpMode);    return InterpMode;    }
-	ERichCurveTangentMode GetTangentMode() const { check(bHasTangentMode);   return TangentMode;   }
+	float GetArriveTangent() const							{ check(bHasArriveTangent); return ArriveTangent; }	
+	float GetLeaveTangent() const							{ check(bHasLeaveTangent);  return LeaveTangent;  }
+	ERichCurveInterpMode GetInterpMode() const				{ check(bHasInterpMode);    return InterpMode;    }
+	ERichCurveTangentMode GetTangentMode() const			{ check(bHasTangentMode);   return TangentMode;   }
+	ERichCurveTangentWeightMode GetTangentWeightMode()const { check(bHasTangentWeightMode);   return TangentWeightMode; }
+	float GetArriveTangentWeight() const { check(bHasArriveTangentWeight); return ArriveTangentWeight; }
+	float GetLeaveTangentWeight() const { check(bHasLeaveTangentWeight);  return LeaveTangentWeight; }
 
 	/**
 	 * Set the attributes for this key
@@ -63,6 +72,9 @@ struct FKeyAttributes
 	FKeyAttributes& SetLeaveTangent(float InLeaveTangent)               { bHasLeaveTangent = 1;     LeaveTangent = InLeaveTangent;   return *this; }
 	FKeyAttributes& SetInterpMode(ERichCurveInterpMode InInterpMode)    { bHasInterpMode = 1;       InterpMode = InInterpMode;       return *this; }
 	FKeyAttributes& SetTangentMode(ERichCurveTangentMode InTangentMode) { bHasTangentMode = 1;      TangentMode = InTangentMode;     return *this; }
+	FKeyAttributes& SetTangentWeightMode(ERichCurveTangentWeightMode InTangentWeightMode) { bHasTangentWeightMode = 1;      TangentWeightMode = InTangentWeightMode;     return *this; }
+	FKeyAttributes& SetArriveTangentWeight(float InArriveTangentWeight) { bHasArriveTangentWeight = 1;    ArriveTangentWeight = InArriveTangentWeight; return *this; }
+	FKeyAttributes& SetLeaveTangentWeight(float InLeaveTangentWeight) { bHasLeaveTangentWeight = 1;     LeaveTangentWeight = InLeaveTangentWeight;   return *this; }
 
 	/**
 	 * Reset specific attributes of this key, implying such attributes are not supported
@@ -71,6 +83,9 @@ struct FKeyAttributes
 	void UnsetLeaveTangent()  { bHasLeaveTangent = 0;  }
 	void UnsetInterpMode()    { bHasInterpMode = 0;    }
 	void UnsetTangentMode()   { bHasTangentMode = 0;   }
+	void UnsetTangentWeightMode() { bHasTangentWeightMode = 0; }
+	void UnsetArriveTangentWeight() { bHasArriveTangentWeight = 0; }
+	void UnsetLeaveTangentWeight() { bHasLeaveTangentWeight = 0; }
 
 	/**
 	 * Generate a new set of attributes that contains only those attributes common to both A and B
@@ -98,6 +113,20 @@ struct FKeyAttributes
 			NewAttributes.SetTangentMode(A.TangentMode);
 		}
 
+		if (A.bHasTangentWeightMode && B.bHasTangentWeightMode && A.TangentWeightMode == B.TangentWeightMode)
+		{
+			NewAttributes.SetTangentWeightMode(A.TangentWeightMode);
+		}
+
+		if (A.bHasArriveTangentWeight && B.bHasArriveTangentWeight && A.ArriveTangentWeight == B.ArriveTangentWeight)
+		{
+			NewAttributes.SetArriveTangentWeight(A.ArriveTangentWeight);
+		}
+
+		if (A.bHasLeaveTangentWeight && B.bHasLeaveTangentWeight && A.LeaveTangentWeight == B.LeaveTangentWeight)
+		{
+			NewAttributes.SetLeaveTangentWeight(A.LeaveTangentWeight);
+		}
 		return NewAttributes;
 	}
 
@@ -111,6 +140,12 @@ private:
 	uint8 bHasInterpMode : 1;
 	/** True if this key supports tangent modes */
 	uint8 bHasTangentMode : 1;
+	/** True if this key supports tangent modes */
+	uint8 bHasTangentWeightMode : 1;
+	/** True if this key supports entry tangents weights*/
+	uint8 bHasArriveTangentWeight : 1;
+	/** True if this key supports exit tangents weights*/
+	uint8 bHasLeaveTangentWeight : 1;
 
 	/** This key's entry tangent, if bHasArriveTangent */
 	float ArriveTangent;
@@ -120,6 +155,13 @@ private:
 	ERichCurveInterpMode InterpMode;
 	/** This key's tangent mode, if bHasTangentMode */
 	ERichCurveTangentMode TangentMode;
+	/** This key's tangent weight mode, if bHasTangentWeightMode */
+	ERichCurveTangentWeightMode TangentWeightMode;
+	/** This key's entry tangent weight, if bHasArriveTangentWeight */
+	float ArriveTangentWeight;
+	/** This key's exit tangent weight, if bHasLeaveTangentWeight */
+	float LeaveTangentWeight;
+
 };
 
 /**
@@ -150,7 +192,6 @@ struct FCurveAttributes
 	 */
 	FCurveAttributes& SetPreExtrapolation(ERichCurveExtrapolation InPreExtrapolation)   { bHasPreExtrapolation = 1;  PreExtrapolation = InPreExtrapolation;   return *this; }
 	FCurveAttributes& SetPostExtrapolation(ERichCurveExtrapolation InPostExtrapolation) { bHasPostExtrapolation = 1; PostExtrapolation = InPostExtrapolation; return *this; }
-
 	/**
 	 * Reset the extended properties of this curve, implying it does not support such properties
 	 */
diff --git a/Engine/Source/Editor/CurveEditor/Public/CurveEditor.h b/Engine/Source/Editor/CurveEditor/Public/CurveEditor.h
--- a/Engine/Source/Editor/CurveEditor/Public/CurveEditor.h
+++ b/Engine/Source/Editor/CurveEditor/Public/CurveEditor.h
@@ -34,6 +34,8 @@ public:
 	 */
 	FCurveEditorSelection Selection;
 
+	
+
 public:
 
 	/** Attribute used to retrieve the current input snap rate (also used for display) */
@@ -270,6 +272,18 @@ public:
 	* Get A Vector for the given slope, usually a tangent, and length. Used to draw the tangent.
 	*/
 	static FVector2D GetVectorFromSlopeAndLength(float Slope, float Length);
+
+public:
+	/**
+	* Given the position of a tangent in screen space get it's position in normal time/value space.
+	*/
+	FVector2D GetTangentPositionInScreenSpace(const FVector2D &StartPos, float Tangent, float Weight) const;
+
+	/**
+	* Given point and tangent position in screen space, get the tangent and it's weight value in normal time/value space.
+	*/
+	void GetTangentAndWeightFromScreenPosition(const FVector2D &StartPos, const  FVector2D &TangentPos, float &Tangent, float &Weight) const;
+
 protected:
 
 	/**
diff --git a/Engine/Source/Editor/CurveEditor/Public/CurveEditorCommands.h b/Engine/Source/Editor/CurveEditor/Public/CurveEditorCommands.h
--- a/Engine/Source/Editor/CurveEditor/Public/CurveEditorCommands.h
+++ b/Engine/Source/Editor/CurveEditor/Public/CurveEditorCommands.h
@@ -35,6 +35,8 @@ public:
 	TSharedPtr<FUICommandInfo> InterpolationCubicAuto;
 	TSharedPtr<FUICommandInfo> InterpolationCubicUser;
 	TSharedPtr<FUICommandInfo> InterpolationCubicBreak;
+	TSharedPtr<FUICommandInfo> InterpolationToggleWeighted;
+
 
 	TSharedPtr<FUICommandInfo> FlattenTangents;
 	TSharedPtr<FUICommandInfo> StraightenTangents;
diff --git a/Engine/Source/Editor/CurveEditor/Public/CurveModel.h b/Engine/Source/Editor/CurveEditor/Public/CurveModel.h
--- a/Engine/Source/Editor/CurveEditor/Public/CurveModel.h
+++ b/Engine/Source/Editor/CurveEditor/Public/CurveModel.h
@@ -20,6 +20,7 @@ struct FCurveEditorScreenSpace;
 class FName;
 class SWidget;
 class FCurveEditor;
+class UObject;
 
 enum class ECurvePointType : uint8;
 
@@ -162,12 +163,13 @@ public:
 	}
 
 	/**
-	 * Create an edit interface for this curve
+	 * Create key proxy objects for the specified key handles. One object should be assigned to OutObjects per index within InKeyHandles
+	 *
+	 * @param InKeyHandles           Array of key handles to create edit objects for
+	 * @param OutObjects             (Out) Array to receive objects that should be used to edit each of the input key handles.
 	 */
-	virtual TSharedPtr<SWidget> CreateEditUI(TSharedPtr<FCurveEditor> InCurveEditor, FCurveModelID ThisCurveID)
-	{
-		return nullptr;
-	}
+	virtual void CreateKeyProxies(TArrayView<const FKeyHandle> InKeyHandles, TArrayView<UObject*> OutObjects)
+	{}
 
 public:
 
diff --git a/Engine/Source/Editor/CurveEditor/Public/SCurveEditorPanel.h b/Engine/Source/Editor/CurveEditor/Public/SCurveEditorPanel.h
--- a/Engine/Source/Editor/CurveEditor/Public/SCurveEditorPanel.h
+++ b/Engine/Source/Editor/CurveEditor/Public/SCurveEditorPanel.h
@@ -11,7 +11,9 @@
 #include "CurveEditor.h"
 #include "CurveEditorTypes.h"
 
-class SScrollBox;
+struct FCurveEditorEditObjectContainer;
+
+class IDetailsView;
 
 /**
  * Curve editor widget that reflects the state of an FCurveEditor
@@ -27,6 +29,9 @@ class CURVEEDITOR_API SCurveEditorPanel : public SCompoundWidget
 
 	SLATE_END_ARGS()
 
+	SCurveEditorPanel();
+	~SCurveEditorPanel();
+
 	/**
 	 * Construct a new curve editor panel widget
 	 */
@@ -48,6 +53,14 @@ class CURVEEDITOR_API SCurveEditorPanel : public SCompoundWidget
 		return CommandList;
 	}
 
+	/**
+	 * Access the details view used for editing selected keys
+	 */
+	TSharedPtr<IDetailsView> GetKeyDetailsView() const
+	{
+		return KeyDetailsView;
+	}
+
 private:
 
 	virtual int32 OnPaint(const FPaintArgs& Args, const FGeometry& AllottedGeometry, const FSlateRect& MyCullingRect, FSlateWindowElementList& OutDrawElements, int32 LayerId, const FWidgetStyle& InWidgetStyle, bool bParentEnabled) const override;
@@ -128,6 +141,12 @@ private:
 		return CompareCommonInterpolationMode(InterpMode) && CachedCommonKeyAttributes.HasTangentMode() && CachedCommonKeyAttributes.GetTangentMode() == TangentMode;
 	}
 
+	/** Compare all the currently selected keys' tangent modes against the specified tangent mode */
+	bool CompareCommonTangentWeightMode(ERichCurveInterpMode InterpMode, ERichCurveTangentWeightMode TangentWeightMode) const
+	{
+		return CompareCommonInterpolationMode(InterpMode) && CachedCommonKeyAttributes.HasTangentWeightMode() && CachedCommonKeyAttributes.GetTangentWeightMode() == TangentWeightMode;
+	}
+
 	/** Compare all the visible curves' pre-extrapolation modes against the specified extrapolation mode */
 	bool CompareCommonPreExtrapolationMode(ERichCurveExtrapolation PreExtrapolationMode) const
 	{
@@ -140,6 +159,15 @@ private:
 		return CachedCommonCurveAttributes.HasPostExtrapolation() && CachedCommonCurveAttributes.GetPostExtrapolation() == PostExtrapolationMode;
 	}
 
+	/**
+	 * Toggle weighted tangents on the current selection
+	 */
+	void ToggleWeightedTangents();
+
+	/**
+	 * Check whether we can toggle weighted tangents on the current selection
+	 */
+	bool CanToggleWeightedTangents() const;
 
 	/**
 	*  Brings up Dialog to specify Tolerance before doing the key reduction.
@@ -193,11 +221,14 @@ private:
 	/** Cached key attributes that are common to all selected keys */
 	FKeyAttributes CachedCommonKeyAttributes;
 
+	/** True if the current selection supports weighted tangents, false otherwise */
+	bool bSelectionSupportsWeightedTangents;
+
 	/** Attribute used for retrieving the desired grid line color */
 	TAttribute<FLinearColor> GridLineTintAttribute;
 
-	/** Edit panel scroll box */
-	TSharedPtr<SScrollBox> EditPanel;
+	/** Edit panel */
+	TSharedPtr<IDetailsView> KeyDetailsView;
 
 	/** Map of edit UI widgets for each curve in the current selection set */
 	TMap<FCurveModelID, TSharedPtr<SWidget>> CurveToEditUI;
@@ -231,4 +262,7 @@ private:
 
 	/** Generic Popup Entry */
 	TWeakPtr<IMenu> EntryPopupMenu;
+
+	/** Container of objects that are being used to edit keys on the curve editor */
+	TUniquePtr<FCurveEditorEditObjectContainer> EditObjects;
 };
\ No newline at end of file
diff --git a/Engine/Source/Editor/EditorStyle/Private/SlateEditorStyle.cpp b/Engine/Source/Editor/EditorStyle/Private/SlateEditorStyle.cpp
--- a/Engine/Source/Editor/EditorStyle/Private/SlateEditorStyle.cpp
+++ b/Engine/Source/Editor/EditorStyle/Private/SlateEditorStyle.cpp
@@ -1783,10 +1783,10 @@ void FSlateEditorStyle::FStyle::SetupGeneralStyles()
 		Set( "Sequencer.GenericDivider", new IMAGE_BRUSH( "Sequencer/GenericDivider", FVector2D(2.f, 2.f), FLinearColor::White, ESlateBrushTileType::Vertical ) );
 
 		Set( "Sequencer.Timeline.ScrubFill", new BOX_BRUSH( "Sequencer/ScrubFill", FMargin( 2.f/4.f, 0.f ) ) );
-		Set( "Sequencer.Timeline.ScrubHandleDown", new BOX_BRUSH( "Sequencer/ScrubHandleDown", FMargin( 6.f/13.f, 5/12.f, 6/13.f, 8/12.f ) ) );
-		Set( "Sequencer.Timeline.ScrubHandleUp", new BOX_BRUSH( "Sequencer/ScrubHandleUp", FMargin( 6.f/13.f, 8/12.f, 6/13.f, 5/12.f ) ) );
-		Set( "Sequencer.Timeline.ClampedScrubHandleDown", new BOX_BRUSH( "Sequencer/ScrubHandleDown_Clamped", FMargin( 6.f/13.f, 3.f/12.f, 6.f/13.f, 7.f/12.f ) ) );
-		Set( "Sequencer.Timeline.ClampedScrubHandleUp", new BOX_BRUSH( "Sequencer/ScrubHandleUp_Clamped", FMargin( 6.f/13.f, 8/12.f, 6/13.f, 5/12.f ) ) );
+		Set( "Sequencer.Timeline.FrameBlockScrubHandleDown", new BOX_BRUSH( "Sequencer/ScrubHandleDown", FMargin( 6.f/13.f, 5/12.f, 6/13.f, 8/12.f ) ) );
+		Set( "Sequencer.Timeline.FrameBlockScrubHandleUp", new BOX_BRUSH( "Sequencer/ScrubHandleUp", FMargin( 6.f/13.f, 8/12.f, 6/13.f, 5/12.f ) ) );
+		Set( "Sequencer.Timeline.VanillaScrubHandleDown", new BOX_BRUSH( "Sequencer/ScrubHandleDown_Clamped", FMargin( 6.f/13.f, 3.f/12.f, 6.f/13.f, 7.f/12.f ) ) );
+		Set( "Sequencer.Timeline.VanillaScrubHandleUp", new BOX_BRUSH( "Sequencer/ScrubHandleUp_Clamped", FMargin( 6.f/13.f, 8/12.f, 6/13.f, 5/12.f ) ) );
 		Set( "Sequencer.Timeline.ScrubHandleWhole", new BOX_BRUSH( "Sequencer/ScrubHandleWhole", FMargin( 6.f/13.f, 10/24.f, 6/13.f, 10/24.f  ) ) );
 		Set( "Sequencer.Timeline.RangeHandleLeft", new BOX_BRUSH( "Sequencer/GenericGripLeft", FMargin(5.f/16.f) ) );
 		Set( "Sequencer.Timeline.RangeHandleRight", new BOX_BRUSH( "Sequencer/GenericGripRight", FMargin(5.f/16.f) ) );
@@ -1890,10 +1890,12 @@ void FSlateEditorStyle::FStyle::SetupGeneralStyles()
 		Set( "Sequencer.AllowSequencerEditsOnly.Small", new IMAGE_BRUSH( "Sequencer/Main_Icons/Icon_Sequencer_Allow_Sequencer_Edits_Only_24x", Icon24x24 ) );
 		Set( "Sequencer.AllowLevelEditsOnly", new IMAGE_BRUSH( "Sequencer/Main_Icons/Icon_Sequencer_Allow_Level_Edits_Only_24x", Icon48x48 ) );
 		Set( "Sequencer.AllowLevelEditsOnly.Small", new IMAGE_BRUSH( "Sequencer/Main_Icons/Icon_Sequencer_Allow_Level_Edits_Only_24x", Icon24x24 ) );		
-		Set( "Sequencer.KeyAllEnabled", new IMAGE_BRUSH( "Sequencer/Main_Icons/Icon_Sequencer_Key_All_24x", Icon48x48 ) );
-		Set( "Sequencer.KeyAllEnabled.Small", new IMAGE_BRUSH( "Sequencer/Main_Icons/Icon_Sequencer_Key_All_24x", Icon24x24 ) );
-		Set( "Sequencer.KeyAllDisabled", new IMAGE_BRUSH( "Sequencer/Main_Icons/Icon_Sequencer_Key_Part_24x", Icon48x48 ) );
-		Set( "Sequencer.KeyAllDisabled.Small", new IMAGE_BRUSH( "Sequencer/Main_Icons/Icon_Sequencer_Key_Part_24x", Icon24x24 ) );
+		Set( "Sequencer.SetKeyAll", new IMAGE_BRUSH( "Sequencer/Main_Icons/Icon_Sequencer_Key_All_24x", Icon48x48 ) );
+		Set( "Sequencer.SetKeyAll.Small", new IMAGE_BRUSH( "Sequencer/Main_Icons/Icon_Sequencer_Key_All_24x", Icon24x24 ) );
+		Set( "Sequencer.SetKeyGroup", new IMAGE_BRUSH("Sequencer/Main_Icons/Icon_Sequencer_Key_Group_24x", Icon48x48));
+		Set( "Sequencer.SetKeyGroup.Small", new IMAGE_BRUSH("Sequencer/Main_Icons/Icon_Sequencer_Key_Group_24x", Icon24x24));
+		Set( "Sequencer.SetKeyChanged", new IMAGE_BRUSH( "Sequencer/Main_Icons/Icon_Sequencer_Key_Part_24x", Icon48x48 ) );
+		Set( "Sequencer.SetKeyChanged.Small", new IMAGE_BRUSH( "Sequencer/Main_Icons/Icon_Sequencer_Key_Part_24x", Icon24x24 ) );
 		Set( "Sequencer.ToggleIsSnapEnabled", new IMAGE_BRUSH( "Sequencer/Main_Icons/Icon_Sequencer_Snap_24x", Icon48x48 ) );
 		Set( "Sequencer.ToggleIsSnapEnabled.Small", new IMAGE_BRUSH( "Sequencer/Main_Icons/Icon_Sequencer_Snap_24x", Icon24x24 ) );
 		Set( "Sequencer.ToggleShowCurveEditor", new IMAGE_BRUSH("Icons/SequencerIcons/icon_Sequencer_CurveEditor_24x", Icon48x48) );
@@ -1970,6 +1972,17 @@ void FSlateEditorStyle::FStyle::SetupGeneralStyles()
 			.SetForegroundColor(FSlateColor::UseSubduedForeground())
 			.SetArrowsImage(FSlateNoResource())
 		);
+
+		Set("Sequencer.PlayTimeSpinBox", FSpinBoxStyle(GetWidgetStyle<FSpinBoxStyle>("SpinBox"))
+			.SetTextPadding(FMargin(0))
+			.SetBackgroundBrush(FSlateNoResource())
+			.SetHoveredBackgroundBrush(FSlateNoResource())
+			.SetInactiveFillBrush(FSlateNoResource())
+			.SetActiveFillBrush(FSlateNoResource())
+			.SetForegroundColor(SelectionColor_Pressed)
+			.SetArrowsImage(FSlateNoResource())
+		);
+
 		Set( "Sequencer.HyperlinkTextBox", FEditableTextBoxStyle()
 			.SetFont( DEFAULT_FONT( "Regular", 9 ) )
 			.SetBackgroundImageNormal( FSlateNoResource() )
@@ -2012,6 +2025,8 @@ void FSlateEditorStyle::FStyle::SetupGeneralStyles()
 		Set( "SequenceRecorder.Common.StopAll", new IMAGE_BRUSH( "SequenceRecorder/icon_StopAll_40x", Icon40x40 ) );
 		Set( "SequenceRecorder.Common.AddRecording.Small", new IMAGE_BRUSH( "SequenceRecorder/icon_AddRecording_40x", Icon20x20 ) );
 		Set( "SequenceRecorder.Common.AddRecording", new IMAGE_BRUSH( "SequenceRecorder/icon_AddRecording_40x", Icon40x40 ) );
+		Set( "SequenceRecorder.Common.AddCurrentPlayerRecording.Small", new IMAGE_BRUSH( "SequenceRecorder/icon_AddCurrentPlayerRecording_40x", Icon20x20 ) );
+		Set( "SequenceRecorder.Common.AddCurrentPlayerRecording", new IMAGE_BRUSH( "SequenceRecorder/icon_AddCurrentPlayerRecording_40x", Icon40x40 ) );
 		Set( "SequenceRecorder.Common.RemoveRecording.Small", new IMAGE_BRUSH( "SequenceRecorder/icon_RemoveRecording_40x", Icon20x20 ) );
 		Set( "SequenceRecorder.Common.RemoveRecording", new IMAGE_BRUSH( "SequenceRecorder/icon_RemoveRecording_40x", Icon40x40 ) );
 		Set( "SequenceRecorder.Common.RemoveAllRecordings.Small", new IMAGE_BRUSH( "SequenceRecorder/icon_RemoveRecording_40x", Icon20x20 ) );
@@ -6007,6 +6022,7 @@ void FSlateEditorStyle::FStyle::SetupClassIconsAndThumbnails()
 		Set("ClassIcon.MovableMobilityIcon", new IMAGE_BRUSH("Icons/ActorIcons/Light_Movable_16x", Icon16x16));
 		Set("ClassIcon.StationaryMobilityIcon", new IMAGE_BRUSH("Icons/ActorIcons/Light_Adjustable_16x", Icon16x16));
 		Set("ClassIcon.ComponentMobilityHeaderIcon", new IMAGE_BRUSH("Icons/ActorIcons/ComponentMobilityHeader_7x16", Icon7x16));
+		Set("ClassIcon.AnimInstance", new IMAGE_BRUSH("Icons/AssetIcons/AnimInstance_16x", Icon16x16));
 
 		//@TODO: PAPER2D: Defined here until it is possible to define these in a plugin
 		{
@@ -6908,6 +6924,9 @@ void FSlateEditorStyle::FStyle::SetupToolkitStyles()
 		Set( "CurveEditor.CurveAutoClamped", new IMAGE_BRUSH( "Icons/icon_CurveEditor_AutoClamped_40x", Icon40x40 ) );
 		Set( "CurveEditor.CurveUser", new IMAGE_BRUSH( "Icons/icon_CurveEditor_User_40x", Icon40x40 ) );
 		Set( "CurveEditor.CurveBreak", new IMAGE_BRUSH( "Icons/icon_CurveEditor_Break_40x", Icon40x40 ) );
+		Set( "CurveEditor.CurveWeight", new IMAGE_BRUSH("Icons/icon_CurveEditor_Break_40x", Icon40x40));
+		Set( "CurveEditor.CurveNonWeight", new IMAGE_BRUSH("Icons/icon_CurveEditor_Break_40x", Icon40x40));
+
 		Set( "CurveEditor.Linear", new IMAGE_BRUSH( "Icons/icon_CurveEditor_Linear_40x", Icon40x40 ) );
 		Set( "CurveEditor.Constant", new IMAGE_BRUSH( "Icons/icon_CurveEditor_Constant_40x", Icon40x40 ) );
 		Set( "CurveEditor.FlattenTangents", new IMAGE_BRUSH( "Icons/icon_CurveEditor_Flatten_40x", Icon40x40 ) );
@@ -6925,6 +6944,9 @@ void FSlateEditorStyle::FStyle::SetupToolkitStyles()
 		Set( "CurveEditor.CurveAutoClamped.Small", new IMAGE_BRUSH( "Icons/icon_CurveEditor_AutoClamped_40x", Icon20x20 ) );
 		Set( "CurveEditor.CurveUser.Small", new IMAGE_BRUSH( "Icons/icon_CurveEditor_User_40x", Icon20x20 ) );
 		Set( "CurveEditor.CurveBreak.Small", new IMAGE_BRUSH( "Icons/icon_CurveEditor_Break_40x", Icon20x20 ) );
+		Set("CurveEditor.CurveWeight.Small", new IMAGE_BRUSH("Icons/icon_CurveEditor_Break_40x", Icon20x20));
+		Set("CurveEditor.CurveNonWeight.Small", new IMAGE_BRUSH("Icons/icon_CurveEditor_Break_40x", Icon20x20));
+
 		Set( "CurveEditor.Linear.Small", new IMAGE_BRUSH( "Icons/icon_CurveEditor_Linear_40x", Icon20x20 ) );
 		Set( "CurveEditor.Constant.Small", new IMAGE_BRUSH( "Icons/icon_CurveEditor_Constant_40x", Icon20x20 ) );
 		Set( "CurveEditor.FlattenTangents.Small", new IMAGE_BRUSH( "Icons/icon_CurveEditor_Flatten_40x", Icon20x20 ) );
@@ -6954,6 +6976,8 @@ void FSlateEditorStyle::FStyle::SetupToolkitStyles()
 		Set("GenericCurveEditor.InterpolationCubicUser.Small", new IMAGE_BRUSH("Icons/icon_CurveEditor_User_40x", Icon20x20));
 		Set("GenericCurveEditor.InterpolationCubicBreak", new IMAGE_BRUSH("Icons/icon_CurveEditor_Break_40x", Icon40x40));
 		Set("GenericCurveEditor.InterpolationCubicBreak.Small", new IMAGE_BRUSH("Icons/icon_CurveEditor_Break_40x", Icon20x20));
+		Set("GenericCurveEditor.InterpolationToggleWeighted", new IMAGE_BRUSH("Icons/icon_CurveEditor_WeightedTangents_40x", Icon40x40));
+		Set("GenericCurveEditor.InterpolationToggleWeighted.Small", new IMAGE_BRUSH("Icons/icon_CurveEditor_WeightedTangents_40x", Icon20x20));
 		Set("GenericCurveEditor.InterpolationLinear", new IMAGE_BRUSH("Icons/icon_CurveEditor_Linear_40x", Icon40x40));
 		Set("GenericCurveEditor.InterpolationLinear.Small", new IMAGE_BRUSH("Icons/icon_CurveEditor_Linear_40x", Icon20x20));
 		Set("GenericCurveEditor.InterpolationConstant", new IMAGE_BRUSH("Icons/icon_CurveEditor_Constant_40x", Icon40x40));
diff --git a/Engine/Source/Editor/MovieSceneCaptureDialog/Private/MovieSceneCaptureDialogModule.cpp b/Engine/Source/Editor/MovieSceneCaptureDialog/Private/MovieSceneCaptureDialogModule.cpp
--- a/Engine/Source/Editor/MovieSceneCaptureDialog/Private/MovieSceneCaptureDialogModule.cpp
+++ b/Engine/Source/Editor/MovieSceneCaptureDialog/Private/MovieSceneCaptureDialogModule.cpp
@@ -604,7 +604,7 @@ private:
 			}
 		}
 
-		if (CaptureObject->Settings.GameModeOverride != nullptr)
+		if (CaptureObject->Settings.GameModeOverride != nullptr && CapturingFromWorld != nullptr)
 		{
 			CapturingFromWorld->GetWorldSettings()->DefaultGameMode = CachedGameMode;
 		}
diff --git a/Engine/Source/Editor/MovieSceneTools/MovieSceneTools.Build.cs b/Engine/Source/Editor/MovieSceneTools/MovieSceneTools.Build.cs
--- a/Engine/Source/Editor/MovieSceneTools/MovieSceneTools.Build.cs
+++ b/Engine/Source/Editor/MovieSceneTools/MovieSceneTools.Build.cs
@@ -24,7 +24,8 @@ public class MovieSceneTools : ModuleRules
 				"CoreUObject",
 				"Engine",
                 "InputCore",
-				"UnrealEd",
+                "MovieSceneCapture",
+                "UnrealEd",
 				"Sequencer"
 			}
 		);
@@ -35,10 +36,12 @@ public class MovieSceneTools : ModuleRules
                 "CinematicCamera",
                 "CurveEditor",
                 "DesktopPlatform",
-				"LevelSequence",
+                "Json",
+                "JsonUtilities",
+                "LevelSequence",
+                "MessageLog",
 				"MovieScene",
-                "MovieSceneCapture",
-				"MovieSceneTracks",
+ 				"MovieSceneTracks",
 				"BlueprintGraph",
 				"KismetCompiler",
                 "GraphEditor",
@@ -54,6 +57,8 @@ public class MovieSceneTools : ModuleRules
 				"ShaderCore",
 				"SequenceRecorder",
                 "AnimationCore",
+				"TimeManagement",
+                "XmlParser",
 			}
 		);
 
diff --git a/Engine/Source/Editor/MovieSceneTools/Private/Channels/FloatChannelCurveModel.cpp b/Engine/Source/Editor/MovieSceneTools/Private/Channels/FloatChannelCurveModel.cpp
--- a/Engine/Source/Editor/MovieSceneTools/Private/Channels/FloatChannelCurveModel.cpp
+++ b/Engine/Source/Editor/MovieSceneTools/Private/Channels/FloatChannelCurveModel.cpp
@@ -13,15 +13,20 @@
 #include "CurveEditorSnapMetrics.h"
 #include "EditorStyleSet.h"
 #include "BuiltInChannelEditors.h"
-#include "SKeyEditInterface.h"
 #include "SequencerChannelTraits.h"
 #include "ISequencer.h"
+#include "Channels/FloatChannelKeyProxy.h"
 
 FFloatChannelCurveModel::FFloatChannelCurveModel(TMovieSceneChannelHandle<FMovieSceneFloatChannel> InChannel, UMovieSceneSection* OwningSection, TWeakPtr<ISequencer> InWeakSequencer)
 {
 	ChannelHandle = InChannel;
 	WeakSection = OwningSection;
 	WeakSequencer = InWeakSequencer;
+	if (UMovieSceneSection* Section = WeakSection.Get())
+	{
+		FMovieSceneFloatChannel* Channel = ChannelHandle.Get();
+		Channel->SetTickResolution(Section->GetTypedOuter<UMovieScene>()->GetTickResolution());
+	}
 }
 
 const void* FFloatChannelCurveModel::GetCurve() const
@@ -37,12 +42,6 @@ void FFloatChannelCurveModel::Modify()
 	}
 }
 
-double FFloatChannelCurveModel::GetInputDisplayOffset() const
-{
-	UMovieSceneSection* Section = WeakSection.Get();
-	return Section ? ( FFrameTime(0, 0.5f) / Section->GetTypedOuter<UMovieScene>()->GetFrameResolution() ) : 0.0;
-}
-
 void FFloatChannelCurveModel::AddKeys(TArrayView<const FKeyPosition> InKeyPositions, TArrayView<const FKeyAttributes> InKeyAttributes, TArrayView<TOptional<FKeyHandle>>* OutKeyHandles)
 {
 	check(InKeyPositions.Num() == InKeyAttributes.Num() && (!OutKeyHandles || OutKeyHandles->Num() == InKeyPositions.Num()));
@@ -54,14 +53,15 @@ void FFloatChannelCurveModel::AddKeys(TArrayView<const FKeyPosition> InKeyPositi
 		Section->Modify();
 
 		TMovieSceneChannel<FMovieSceneFloatValue> ChannelInterface = Channel->GetInterface();
-		FFrameRate FrameResolution = Section->GetTypedOuter<UMovieScene>()->GetFrameResolution();
+		FFrameRate TickResolution = Section->GetTypedOuter<UMovieScene>()->GetTickResolution();
 
 		for (int32 Index = 0; Index < InKeyPositions.Num(); ++Index)
 		{
 			FKeyPosition   Position   = InKeyPositions[Index];
 			FKeyAttributes Attributes = InKeyAttributes[Index];
 
-			FFrameNumber Time = (Position.InputValue * FrameResolution).RoundToFrame();
+			FFrameNumber Time = (Position.InputValue * TickResolution).RoundToFrame();
+			Section->ExpandToFrame(Time);
 
 			FMovieSceneFloatValue Value(Position.OutputValue);
 
@@ -88,10 +88,10 @@ bool FFloatChannelCurveModel::Evaluate(double Time, double& OutValue) const
 
 	if (Channel && Section)
 	{
-		FFrameRate FrameResolution = Section->GetTypedOuter<UMovieScene>()->GetFrameResolution();
+		FFrameRate TickResolution = Section->GetTypedOuter<UMovieScene>()->GetTickResolution();
 
 		float ThisValue = 0.f;
-		if (Channel->Evaluate(Time * FrameResolution, ThisValue))
+		if (Channel->Evaluate(Time * TickResolution, ThisValue))
 		{
 			OutValue = ThisValue;
 			return true;
@@ -131,7 +131,7 @@ void FFloatChannelCurveModel::DrawCurve(const FCurveEditor& CurveEditor, TArray<
 	{
 		FCurveEditorScreenSpace ScreenSpace = CurveEditor.GetScreenSpace();
 
-		FFrameRate   FrameResolution  = Section->GetTypedOuter<UMovieScene>()->GetFrameResolution();
+		FFrameRate   TickResolution   = Section->GetTypedOuter<UMovieScene>()->GetTickResolution();
 
 		const double DisplayOffset    = GetInputDisplayOffset();
 		const double StartTimeSeconds = ScreenSpace.GetInputMin() - DisplayOffset;
@@ -139,7 +139,7 @@ void FFloatChannelCurveModel::DrawCurve(const FCurveEditor& CurveEditor, TArray<
 		const double TimeThreshold    = FMath::Max(0.0001, 1.0 / ScreenSpace.PixelsPerInput());
 		const double ValueThreshold   = FMath::Max(0.0001, 1.0 / ScreenSpace.PixelsPerOutput());
 
-		Channel->PopulateCurvePoints(StartTimeSeconds, EndTimeSeconds, TimeThreshold, ValueThreshold, FrameResolution, InterpolatingPoints);
+		Channel->PopulateCurvePoints(StartTimeSeconds, EndTimeSeconds, TimeThreshold, ValueThreshold, TickResolution, InterpolatingPoints);
 	}
 }
 
@@ -150,14 +150,14 @@ void FFloatChannelCurveModel::GetKeys(const FCurveEditor& CurveEditor, double Mi
 
 	if (Channel && Section)
 	{
-		FFrameRate FrameResolution = Section->GetTypedOuter<UMovieScene>()->GetFrameResolution();
+		FFrameRate TickResolution = Section->GetTypedOuter<UMovieScene>()->GetTickResolution();
 
 		TMovieSceneChannel<FMovieSceneFloatValue> ChannelInterface = Channel->GetInterface();
 		TArrayView<const FFrameNumber>            Times  = ChannelInterface.GetTimes();
 		TArrayView<const FMovieSceneFloatValue>   Values = ChannelInterface.GetValues();
 
-		const FFrameNumber StartFrame = MinTime <= MIN_int32 ? MIN_int32 : (MinTime * FrameResolution).CeilToFrame();
-		const FFrameNumber EndFrame   = MaxTime >= MAX_int32 ? MAX_int32 : (MaxTime * FrameResolution).FloorToFrame();
+		const FFrameNumber StartFrame = MinTime <= MIN_int32 ? MIN_int32 : (MinTime * TickResolution).CeilToFrame();
+		const FFrameNumber EndFrame   = MaxTime >= MAX_int32 ? MAX_int32 : (MaxTime * TickResolution).FloorToFrame();
 
 		const int32 StartingIndex = Algo::LowerBound(Times, StartFrame);
 		const int32 EndingIndex   = Algo::UpperBound(Times, EndFrame);
@@ -195,7 +195,7 @@ void FFloatChannelCurveModel::GetKeyPositions(TArrayView<const FKeyHandle> InKey
 
 	if (Channel && Section)
 	{
-		FFrameRate FrameResolution = Section->GetTypedOuter<UMovieScene>()->GetFrameResolution();
+		FFrameRate TickResolution = Section->GetTypedOuter<UMovieScene>()->GetTickResolution();
 
 		TMovieSceneChannel<FMovieSceneFloatValue> ChannelInterface = Channel->GetInterface();
 		TArrayView<const FFrameNumber>            Times            = ChannelInterface.GetTimes();
@@ -206,7 +206,7 @@ void FFloatChannelCurveModel::GetKeyPositions(TArrayView<const FKeyHandle> InKey
 			int32 KeyIndex = ChannelInterface.GetIndex(InKeys[Index]);
 			if (KeyIndex != INDEX_NONE)
 			{
-				OutKeyPositions[Index].InputValue  = Times[KeyIndex] / FrameResolution;
+				OutKeyPositions[Index].InputValue  = Times[KeyIndex] / TickResolution;
 				OutKeyPositions[Index].OutputValue = Values[KeyIndex].Value;
 			}
 		}
@@ -222,7 +222,7 @@ void FFloatChannelCurveModel::SetKeyPositions(TArrayView<const FKeyHandle> InKey
 	{
 		Section->MarkAsChanged();
 
-		FFrameRate FrameResolution = Section->GetTypedOuter<UMovieScene>()->GetFrameResolution();
+		FFrameRate TickResolution = Section->GetTypedOuter<UMovieScene>()->GetTickResolution();
 
 		TMovieSceneChannel<FMovieSceneFloatValue> ChannelInterface = Channel->GetInterface();
 		for (int32 Index = 0; Index < InKeys.Num(); ++Index)
@@ -230,10 +230,12 @@ void FFloatChannelCurveModel::SetKeyPositions(TArrayView<const FKeyHandle> InKey
 			int32 KeyIndex = ChannelInterface.GetIndex(InKeys[Index]);
 			if (KeyIndex != INDEX_NONE)
 			{
-				FFrameTime NewTime = InKeyPositions[Index].InputValue * FrameResolution;
+				FFrameNumber NewTime = (InKeyPositions[Index].InputValue * TickResolution).FloorToFrame();
 
-				KeyIndex = ChannelInterface.MoveKey(KeyIndex, NewTime.FloorToFrame());
+				KeyIndex = ChannelInterface.MoveKey(KeyIndex, NewTime);
 				ChannelInterface.GetValues()[KeyIndex].Value = InKeyPositions[Index].OutputValue;
+
+				Section->ExpandToFrame(NewTime);
 			}
 		}
 
@@ -251,7 +253,7 @@ void FFloatChannelCurveModel::GetKeyAttributes(TArrayView<const FKeyHandle> InKe
 		TArrayView<const FFrameNumber>    Times  = ChannelInterface.GetTimes();
 		TArrayView<FMovieSceneFloatValue> Values = ChannelInterface.GetValues();
 
-		float TimeInterval = Section->GetTypedOuter<UMovieScene>()->GetFrameResolution().AsInterval();
+		float TimeInterval = Section->GetTypedOuter<UMovieScene>()->GetTickResolution().AsInterval();
 
 		for (int32 Index = 0; Index < InKeys.Num(); ++Index)
 		{
@@ -275,6 +277,15 @@ void FFloatChannelCurveModel::GetKeyAttributes(TArrayView<const FKeyHandle> InKe
 					{
 						Attributes.SetLeaveTangent(KeyValue.Tangent.LeaveTangent / TimeInterval);
 					}
+					if (KeyValue.InterpMode == RCIM_Cubic)
+					{
+						Attributes.SetTangentWeightMode(KeyValue.Tangent.TangentWeightMode);
+						if (KeyValue.Tangent.TangentWeightMode != RCTWM_WeightedNone)
+						{
+							Attributes.SetArriveTangentWeight(KeyValue.Tangent.ArriveTangentWeight);
+							Attributes.SetLeaveTangentWeight(KeyValue.Tangent.LeaveTangentWeight);
+						}
+					}
 				}
 			}
 		}
@@ -293,7 +304,8 @@ void FFloatChannelCurveModel::SetKeyAttributes(TArrayView<const FKeyHandle> InKe
 		TMovieSceneChannel<FMovieSceneFloatValue> ChannelInterface = Channel->GetInterface();
 		TArrayView<FMovieSceneFloatValue> Values = ChannelInterface.GetValues();
 
-		float TimeInterval = Section->GetTypedOuter<UMovieScene>()->GetFrameResolution().AsInterval();
+		FFrameRate TickResolution = Section->GetTypedOuter<UMovieScene>()->GetTickResolution();
+		float TimeInterval = TickResolution.AsInterval();
 
 		for (int32 Index = 0; Index < InKeys.Num(); ++Index)
 		{
@@ -302,15 +314,57 @@ void FFloatChannelCurveModel::SetKeyAttributes(TArrayView<const FKeyHandle> InKe
 			{
 				const FKeyAttributes&  Attributes = InAttributes[Index];
 				FMovieSceneFloatValue& KeyValue   = Values[KeyIndex];
-
 				if (Attributes.HasInterpMode())    { KeyValue.InterpMode  = Attributes.GetInterpMode();  bAutoSetTangents = true; }
-				if (Attributes.HasTangentMode())   { KeyValue.TangentMode = Attributes.GetTangentMode(); bAutoSetTangents = true; }
+				if (Attributes.HasTangentMode())
+				{
+					KeyValue.TangentMode = Attributes.GetTangentMode();
+					if (KeyValue.TangentMode == RCTM_Auto)
+					{
+						KeyValue.Tangent.TangentWeightMode = RCTWM_WeightedNone;
+					}
+					bAutoSetTangents = true;
+				}
+				if (Attributes.HasTangentWeightMode()) 
+				{ 
+					if (KeyValue.Tangent.TangentWeightMode == RCTWM_WeightedNone) //set tangent weights to default use
+					{
+						TArrayView<const FFrameNumber> Times = Channel->GetTimes();
+						
+						//calculate a tangent weight based upon tangent and time difference
+						//calculate arrive tangent weight
+						if (KeyIndex > 0 )
+						{
+							const float X = TickResolution.AsSeconds(Times[KeyIndex].Value - Times[KeyIndex - 1].Value);
+							const float ArriveTangentNormal = KeyValue.Tangent.ArriveTangent / (TimeInterval);
+							const float Y = ArriveTangentNormal * X;
+							KeyValue.Tangent.ArriveTangentWeight = FMath::Sqrt(X*X + Y*Y) * 0.3f;
+						}
+						//calculate leave weight
+						if(KeyIndex < ( Times.Num() - 1))
+						{
+							const float X = TickResolution.AsSeconds(Times[KeyIndex].Value - Times[KeyIndex + 1].Value);
+							const float LeaveTangentNormal = KeyValue.Tangent.LeaveTangent / (TimeInterval);
+							const float Y = LeaveTangentNormal * X;
+							KeyValue.Tangent.LeaveTangentWeight = FMath::Sqrt(X*X + Y*Y) * 0.3f;
+						}
+					}
+					KeyValue.Tangent.TangentWeightMode = Attributes.GetTangentWeightMode();
+
+					if( KeyValue.Tangent.TangentWeightMode != RCTWM_WeightedNone )
+					{
+						if (KeyValue.TangentMode != RCTM_User && KeyValue.TangentMode != RCTM_Break)
+						{
+							KeyValue.TangentMode = RCTM_User;
+						}
+					}
+				}
 
 				if (Attributes.HasArriveTangent())
 				{
 					if (KeyValue.TangentMode == RCTM_Auto)
 					{
 						KeyValue.TangentMode = RCTM_User;
+						KeyValue.Tangent.TangentWeightMode = RCTWM_WeightedNone;
 					}
 
 					KeyValue.Tangent.ArriveTangent = Attributes.GetArriveTangent() * TimeInterval;
@@ -325,6 +379,7 @@ void FFloatChannelCurveModel::SetKeyAttributes(TArrayView<const FKeyHandle> InKe
 					if (KeyValue.TangentMode == RCTM_Auto)
 					{
 						KeyValue.TangentMode = RCTM_User;
+						KeyValue.Tangent.TangentWeightMode = RCTWM_WeightedNone;
 					}
 
 					KeyValue.Tangent.LeaveTangent = Attributes.GetLeaveTangent() * TimeInterval;
@@ -333,6 +388,37 @@ void FFloatChannelCurveModel::SetKeyAttributes(TArrayView<const FKeyHandle> InKe
 						KeyValue.Tangent.ArriveTangent = KeyValue.Tangent.LeaveTangent;
 					}
 				}
+
+				if (Attributes.HasArriveTangentWeight())
+				{
+					if (KeyValue.TangentMode == RCTM_Auto)
+					{
+						KeyValue.TangentMode = RCTM_User;
+						KeyValue.Tangent.TangentWeightMode = RCTWM_WeightedNone;
+					}
+
+					KeyValue.Tangent.ArriveTangentWeight = Attributes.GetArriveTangentWeight(); 
+					if (KeyValue.InterpMode == RCIM_Cubic && KeyValue.TangentMode != RCTM_Break)
+					{
+						KeyValue.Tangent.LeaveTangentWeight = KeyValue.Tangent.ArriveTangentWeight;
+					}
+				}
+
+				if (Attributes.HasLeaveTangentWeight())
+				{
+				
+					if (KeyValue.TangentMode == RCTM_Auto)
+					{
+						KeyValue.TangentMode = RCTM_User;
+						KeyValue.Tangent.TangentWeightMode = RCTWM_WeightedNone;
+					}
+
+					KeyValue.Tangent.LeaveTangentWeight = Attributes.GetLeaveTangentWeight();
+					if (KeyValue.InterpMode == RCIM_Cubic && KeyValue.TangentMode != RCTM_Break)
+					{
+						KeyValue.Tangent.ArriveTangentWeight = KeyValue.Tangent.LeaveTangentWeight;
+					}
+				}
 			}
 		}
 
@@ -373,81 +459,13 @@ void FFloatChannelCurveModel::SetCurveAttributes(const FCurveAttributes& InCurve
 	}
 }
 
-class SCurveEditorKeyEditInterface : public SCompoundWidget
+void FFloatChannelCurveModel::CreateKeyProxies(TArrayView<const FKeyHandle> InKeyHandles, TArrayView<UObject*> OutObjects)
 {
-public:
-	SLATE_BEGIN_ARGS(SCurveEditorKeyEditInterface){}
-	SLATE_END_ARGS()
-
-	void Construct(
-		const FArguments& InArgs,
-		TSharedRef<ISequencer> InSequencer,
-		FCurveModelID InCurveID,
-		TWeakPtr<FCurveEditor> InWeakCurveEditor,
-		TWeakObjectPtr<UMovieSceneSection> InWeakSection,
-		TMovieSceneChannelHandle<FMovieSceneFloatChannel> InChannelHandle
-		)
+	for (int32 Index = 0; Index < InKeyHandles.Num(); ++Index)
 	{
-		CachedSelectionSerialNumber = 0;
-
-		CurveID         = InCurveID;
-		WeakCurveEditor = InWeakCurveEditor;
-		WeakSection     = InWeakSection;
-		ChannelHandle   = InChannelHandle;
-
-		ChildSlot
-		[
-			SAssignNew(Interface, SKeyEditInterface, InSequencer)
-			.EditData(this, &SCurveEditorKeyEditInterface::UpdateAndRetrieveEditData)
-		];
-	}
+		UFloatChannelKeyProxy* NewProxy = NewObject<UFloatChannelKeyProxy>(GetTransientPackage(), NAME_None);
 
-	virtual void Tick( const FGeometry& AllottedGeometry, const double InCurrentTime, const float InDeltaTime )
-	{
-		TSharedPtr<FCurveEditor> CurveEditor = WeakCurveEditor.Pin();
-		if (CurveEditor.IsValid() && CurveEditor->Selection.GetSerialNumber() != CachedSelectionSerialNumber)
-		{
-			Interface->Initialize();
-			CachedSelectionSerialNumber = CurveEditor->Selection.GetSerialNumber();
-		}
+		NewProxy->Initialize(InKeyHandles[Index], ChannelHandle, WeakSection);
+		OutObjects[Index] = NewProxy;
 	}
-
-private:
-
-	FKeyEditData UpdateAndRetrieveEditData() const
-	{
-		FKeyEditData EditData;
-
-		TSharedPtr<FCurveEditor> CurveEditor = WeakCurveEditor.Pin();
-		if (CurveEditor.IsValid())
-		{
-			const FKeyHandleSet* SelectedKeys = CurveEditor->Selection.FindForCurve(CurveID.GetValue());
-			if (SelectedKeys && SelectedKeys->Num() == 1)
-			{
-				using namespace Sequencer;
-				EditData.KeyStruct = GetKeyStruct(ChannelHandle, SelectedKeys->AsArray()[0]);
-				EditData.OwningSection = WeakSection;
-			}
-		}
-
-		return EditData;
-	}
-
-	TOptional<FCurveModelID> CurveID;
-	TWeakPtr<FCurveEditor> WeakCurveEditor;
-	TWeakObjectPtr<UMovieSceneSection> WeakSection;
-	TMovieSceneChannelHandle<FMovieSceneFloatChannel> ChannelHandle;
-	uint32 CachedSelectionSerialNumber;
-	TSharedPtr<SKeyEditInterface> Interface;
-};
-
-TSharedPtr<SWidget> FFloatChannelCurveModel::CreateEditUI(TSharedPtr<FCurveEditor> InCurveEditor, FCurveModelID ThisCurveID)
-{
-	TSharedPtr<ISequencer> Sequencer = WeakSequencer.Pin();
-	if (Sequencer.IsValid())
-	{
-		return SNew(SCurveEditorKeyEditInterface, Sequencer.ToSharedRef(), ThisCurveID, TWeakPtr<FCurveEditor>(InCurveEditor), WeakSection, ChannelHandle);
-	}
-
-	return nullptr;
 }
\ No newline at end of file
diff --git a/Engine/Source/Editor/MovieSceneTools/Private/Channels/FloatChannelCurveModel.h b/Engine/Source/Editor/MovieSceneTools/Private/Channels/FloatChannelCurveModel.h
--- a/Engine/Source/Editor/MovieSceneTools/Private/Channels/FloatChannelCurveModel.h
+++ b/Engine/Source/Editor/MovieSceneTools/Private/Channels/FloatChannelCurveModel.h
@@ -39,9 +39,7 @@ public:
 	virtual void AddKeys(TArrayView<const FKeyPosition> InKeyPositions, TArrayView<const FKeyAttributes> InAttributes, TArrayView<TOptional<FKeyHandle>>* OutKeyHandles) override;
 	virtual void RemoveKeys(TArrayView<const FKeyHandle> InKeys) override;
 
-	virtual double GetInputDisplayOffset() const override;
-
-	virtual TSharedPtr<SWidget> CreateEditUI(TSharedPtr<FCurveEditor> InCurveEditor, FCurveModelID ThisCurveID) override;
+	virtual void CreateKeyProxies(TArrayView<const FKeyHandle> InKeyHandles, TArrayView<UObject*> OutObjects) override;
 
 private:
 
diff --git a/Engine/Source/Editor/MovieSceneTools/Private/CurveKeyEditors/SequencerKeyEditor.h b/Engine/Source/Editor/MovieSceneTools/Private/CurveKeyEditors/SequencerKeyEditor.h
--- a/Engine/Source/Editor/MovieSceneTools/Private/CurveKeyEditors/SequencerKeyEditor.h
+++ b/Engine/Source/Editor/MovieSceneTools/Private/CurveKeyEditors/SequencerKeyEditor.h
@@ -9,6 +9,7 @@
 #include "Channels/MovieSceneChannelTraits.h"
 #include "SequencerChannelTraits.h"
 #include "Channels/MovieSceneChannelHandle.h"
+#include "MovieSceneTimeHelpers.h"
 
 
 template<typename ChannelType, typename ValueType>
@@ -71,12 +72,13 @@ struct TSequencerKeyEditor
 
 		ChannelType* Channel = ChannelHandle.Get();
 		ISequencer* Sequencer = WeakSequencer.Pin().Get();
+		UMovieSceneSection* OwningSection = WeakSection.Get();
 
 		ValueType Result{};
 
-		if (Channel && Sequencer)
+		if (Channel && Sequencer && OwningSection)
 		{
-			const FFrameTime CurrentTime = Sequencer->GetLocalTime().Time;
+			const FFrameTime CurrentTime = MovieScene::ClampToDiscreteRange(Sequencer->GetLocalTime().Time, OwningSection->GetRange());
 			EvaluateChannel(Channel, CurrentTime, Result);
 		}
 
diff --git a/Engine/Source/Editor/MovieSceneTools/Private/FCPXML/FCPXMLExport.cpp b/Engine/Source/Editor/MovieSceneTools/Private/FCPXML/FCPXMLExport.cpp
new file mode 100644
--- /dev/null
+++ b/Engine/Source/Editor/MovieSceneTools/Private/FCPXML/FCPXMLExport.cpp
@@ -0,0 +1,753 @@
+// Copyright 1998-2018 Epic Games, Inc. All Rights Reserved.
+
+#include "FCPXML/FCPXMLExport.h"
+#include "MovieScene.h"
+#include "LevelSequence.h"
+#include "Tracks/MovieSceneAudioTrack.h"
+#include "Sections/MovieSceneCinematicShotSection.h"
+#include "Tracks/MovieSceneCinematicShotTrack.h"
+#include "MovieSceneTimeHelpers.h"
+#include "AssetRegistryModule.h"
+#include "UObject/MetaData.h"
+#include "Logging/TokenizedMessage.h"
+
+#define LOCTEXT_NAMESPACE "FCPXMLExporter"
+
+DEFINE_LOG_CATEGORY_STATIC(LogFCPXMLExporter, Log, All);
+
+FFCPXMLExportVisitor::FFCPXMLExportVisitor(FString InSaveFilename, TSharedRef<FMovieSceneExportData> InExportData, TSharedRef<FMovieSceneTranslatorContext> InExportContext) :
+	  FFCPXMLNodeVisitor()
+	, ExportData(InExportData)
+	, ExportContext(InExportContext)
+	, SaveFilePath(InSaveFilename)
+{
+	SaveFilePath = FPaths::GetPath(InSaveFilename);
+	if (FPaths::IsRelative(SaveFilePath))
+	{
+		SaveFilePath = FPaths::ConvertRelativePathToFull(SaveFilePath);
+	}
+
+	SequenceId = 0;
+}
+
+FFCPXMLExportVisitor::~FFCPXMLExportVisitor() {}
+
+bool FFCPXMLExportVisitor::VisitNode(TSharedRef<FFCPXMLBasicNode> InBasicNode)
+{
+	return InBasicNode->VisitChildren(*this);
+}
+
+bool FFCPXMLExportVisitor::VisitNode(TSharedRef<FFCPXMLXmemlNode> InXmemlNode)
+{
+	float EditTime = 0.f;
+
+	// Construct the FCP 7 XML structure from Sequencer class objects
+	if (InXmemlNode->GetChildCount() == 0)
+	{
+		if (!ConstructProjectNode(InXmemlNode)) 
+		{
+			return false;
+		}
+	}
+
+	/** @todo - MERGE METADATA
+
+	   Merging the newly exported XML structure with pre-existing
+	   XML metadata will be implemented here. The traversal would proceed
+	   through the new XML structure, referring back to the metadata 
+	   XML structure to incorporate any missing or desired attributes or
+	   elements. Alternatively, the traversal might be invoked directly on
+	   the metadata referring back to and modifying the new XML structure. 
+
+	*/
+
+	return InXmemlNode->VisitChildren(*this);
+}
+
+bool FFCPXMLExportVisitor::VisitNode(TSharedRef<FFCPXMLSequenceNode> InSequenceNode)
+{
+	return InSequenceNode->VisitChildren(*this);
+}
+
+bool FFCPXMLExportVisitor::VisitNode(TSharedRef<FFCPXMLVideoNode> InVideoNode)
+{
+	return InVideoNode->VisitChildren(*this);
+}
+
+bool FFCPXMLExportVisitor::VisitNode(TSharedRef<FFCPXMLAudioNode> InAudioNode)
+{
+	return InAudioNode->VisitChildren(*this);
+}
+
+bool FFCPXMLExportVisitor::VisitNode(TSharedRef<FFCPXMLTrackNode> InTrackNode)
+{
+	return InTrackNode->VisitChildren(*this);
+}
+
+bool FFCPXMLExportVisitor::VisitNode(TSharedRef<FFCPXMLClipNode> InClipNode)
+{
+	return InClipNode->VisitChildren(*this);
+}
+
+bool FFCPXMLExportVisitor::VisitNode(TSharedRef<FFCPXMLClipItemNode> InClipItemNode)
+{
+	return InClipItemNode->VisitChildren(*this);
+}
+
+bool FFCPXMLExportVisitor::VisitNode(TSharedRef<FFCPXMLFileNode> InFileNode)
+{
+	return InFileNode->VisitChildren(*this);
+}
+
+
+/** Creates project node. */
+bool FFCPXMLExportVisitor::ConstructProjectNode(TSharedRef<FFCPXMLNode> InParentNode)
+{
+	TSharedRef<FFCPXMLNode> ProjectNode = InParentNode->CreateChildNode(TEXT("project"));
+
+	TSharedRef<FFCPXMLNode> NameNode = ProjectNode->CreateChildNode(TEXT("name"));
+	NameNode->SetContent(ExportData->MovieSceneData->Name + TEXT("Project"));
+
+	TSharedRef<FFCPXMLNode> ChildrenNode = ProjectNode->CreateChildNode(TEXT("children"));
+
+	if (!ConstructMasterClipNodes(ChildrenNode))
+	{
+		return false;
+	}
+
+	if (!ConstructSequenceNode(ChildrenNode))
+	{
+		return false;
+	}
+
+	return true;
+}
+
+/** Creates master clip node. */
+bool FFCPXMLExportVisitor::ConstructMasterClipNodes(TSharedRef<FFCPXMLNode> InParentNode)
+{
+	if (!ExportData->IsExportDataValid() || !ExportData->MovieSceneData.IsValid() || !ExportData->MovieSceneData->CinematicMasterTrack.IsValid())
+	{
+		return false;
+	}
+
+	int32 MasterClipId = 0;
+
+	TSharedPtr<FMovieSceneExportCinematicMasterTrackData> CinematicMasterTrackData = ExportData->MovieSceneData->CinematicMasterTrack;
+	if (!CinematicMasterTrackData.IsValid())
+	{
+		return false;
+	}
+
+	for (TSharedPtr<FMovieSceneExportCinematicSectionData> CinematicSection : CinematicMasterTrackData->CinematicSections)
+	{
+		if (!CinematicSection.IsValid())
+		{
+			continue;
+		}
+
+		// skip sections outside of playback range
+		if (CinematicSection->bWithinPlaybackRange == false)
+		{
+			CinematicSection->bEnabled = false;
+			continue;
+		}
+
+		// skip sections if media file does not exist
+		FString FilePathName = GetFilePathName(CinematicSection->ShotFilename);
+		if (!FPaths::FileExists(FilePathName))
+		{
+			// add warning message and skip this section
+			ExportContext->AddMessage(
+				EMessageSeverity::Warning,
+				FText::Format(LOCTEXT("SkippingSection", "Warning: Skipping section {0}, media file does not exist: {1}."), FText::FromString(CinematicSection->ShotDisplayName), FText::FromString(FilePathName)));
+			CinematicSection->bEnabled = false;
+			continue;
+		}
+
+
+		if (!ConstructMasterClipNode(InParentNode, CinematicSection, ++MasterClipId))
+		{
+			return false;
+		}
+	}
+	return true;
+}
+
+bool FFCPXMLExportVisitor::ConstructMasterClipNode(TSharedRef<FFCPXMLNode> InParentNode, const TSharedPtr<FMovieSceneExportCinematicSectionData> InCinematicSectionData, int32 InMasterClipId)
+{
+	if (!InCinematicSectionData.IsValid())
+	{
+		return false;
+	}
+
+	int32 Duration{ 0 };
+	int32 StartFrame{ 0 };
+	int32 EndFrame{ 0 };
+	int32 InFrame{ 0 };
+	int32 OutFrame{ 0 };
+	FString SectionName(TEXT(""));
+	GetSectionFrames(InCinematicSectionData, Duration, StartFrame, EndFrame, InFrame, OutFrame);
+	GetSectionName(InCinematicSectionData, SectionName);
+
+	/** Construct a master clip id name based on the cinematic section and id */
+	FString MasterClipName(TEXT(""));
+	if (!CreateMasterClipIdName(InCinematicSectionData, InMasterClipId, MasterClipName))
+	{
+		return false;
+	}
+
+	TSharedRef<FFCPXMLNode> ClipNode = InParentNode->CreateChildNode(TEXT("clip"));
+	ClipNode->AddAttribute(TEXT("id"), MasterClipName);
+
+	// @todo add to file's masterclip and refidmap HERE
+
+	ClipNode->CreateChildNode(TEXT("masterclipid"))->SetContent(MasterClipName);
+	ClipNode->CreateChildNode(TEXT("ismasterclip"))->SetContent(true);
+	ClipNode->CreateChildNode(TEXT("duration"))->SetContent(Duration);
+
+	if (!ConstructRateNode(ClipNode))
+	{
+		return false;
+	}
+
+	ClipNode->CreateChildNode(TEXT("in"))->SetContent(InFrame);
+	ClipNode->CreateChildNode(TEXT("out"))->SetContent(OutFrame);
+	ClipNode->CreateChildNode(TEXT("name"))->SetContent(SectionName);
+	TSharedRef<FFCPXMLNode> MediaNode = ClipNode->CreateChildNode(TEXT("media"));
+	TSharedRef<FFCPXMLNode> VideoNode = MediaNode->CreateChildNode(TEXT("video"));
+	TSharedRef<FFCPXMLNode> TrackNode = VideoNode->CreateChildNode(TEXT("track"));
+
+	if (!ConstructVideoClipItemNode(TrackNode, InCinematicSectionData, true))
+	{
+		return false;
+	}
+
+	if (!ConstructSectionLoggingInfoNode(ClipNode, InCinematicSectionData, SectionName))
+	{
+		return false;
+	}
+
+	if (!ConstructColorInfoNode(ClipNode, InCinematicSectionData))
+	{
+		return false;
+	}
+
+	return true;
+}
+
+/** Creates logginginfo node. */
+bool FFCPXMLExportVisitor::ConstructSectionLoggingInfoNode(TSharedRef<FFCPXMLNode> InParentNode, const TSharedPtr<FMovieSceneExportCinematicSectionData> InCinematicSectionData, const FString& InSectionName)
+{
+	if (InCinematicSectionData->CinematicShotSection == nullptr)
+	{
+		return false;
+	}
+
+	TSharedRef<FFCPXMLNode> LoggingInfoNode = InParentNode->CreateChildNode(TEXT("logginginfo"));
+	ConstructLoggingInfoElements(LoggingInfoNode, InCinematicSectionData->CinematicShotSection);
+
+	TSharedPtr<FFCPXMLNode> LogNoteNode = LoggingInfoNode->GetChildNode(TEXT("lognote"), ENodeInherit::NoInherit, ENodeReference::NoReferences);
+	if (!LogNoteNode.IsValid())
+	{
+		LogNoteNode = LoggingInfoNode->CreateChildNode(TEXT("lognote"));
+	}
+	LogNoteNode->SetContent(FFCPXMLClipNode::GetMetadataSectionName(InSectionName));
+
+	return true;
+}
+
+/** Creates logginginfo elements. */
+void FFCPXMLExportVisitor::ConstructLoggingInfoElements(TSharedRef<FFCPXMLNode> InLoggingInfoNode, const UObject* InObject)
+{
+	TSharedRef<FFCPXMLNode> DescriptionNode = InLoggingInfoNode->CreateChildNode(TEXT("description"));
+	TSharedRef<FFCPXMLNode> SceneNode = InLoggingInfoNode->CreateChildNode(TEXT("scene"));
+	TSharedRef<FFCPXMLNode> ShotTakeNode = InLoggingInfoNode->CreateChildNode(TEXT("shottake"));
+	TSharedRef<FFCPXMLNode> GoodNode = InLoggingInfoNode->CreateChildNode(TEXT("good"));
+	TSharedRef<FFCPXMLNode> OriginalVideoNode = InLoggingInfoNode->CreateChildNode(TEXT("originalvideofilename"));
+	TSharedRef<FFCPXMLNode> OriginalAudioNode = InLoggingInfoNode->CreateChildNode(TEXT("originalaudiofilename"));
+
+	if (InObject != nullptr)
+	{
+		SetLoggingInfoElementValue(DescriptionNode, InObject, TEXT("description"));
+		SetLoggingInfoElementValue(SceneNode, InObject, TEXT("scene"));
+		SetLoggingInfoElementValue(ShotTakeNode, InObject, TEXT("shottake"));
+		SetLoggingInfoElementValue(GoodNode, InObject, TEXT("good"));
+		SetLoggingInfoElementValue(OriginalVideoNode, InObject, TEXT("originalvideofilename"));
+		SetLoggingInfoElementValue(OriginalAudioNode, InObject, TEXT("originalaudiofilename"));
+	}
+}
+
+/** Set logginginfo element value. */
+void FFCPXMLExportVisitor::SetLoggingInfoElementValue(TSharedPtr<FFCPXMLNode> InNode, const UObject* InObject, const FString& InElement)
+{
+	if (InNode.IsValid() && InObject != nullptr)
+	{
+		UPackage* Package = InObject->GetOutermost();
+		check(Package);
+
+		UMetaData* MetaData = Package->GetMetaData();
+		check(MetaData);
+
+		if (MetaData->HasValue(InObject, *InElement))
+		{
+			FString Value = MetaData->GetValue(InObject, *InElement);
+			InNode->SetContent(Value);
+		}
+	}	
+}
+
+/** Creates colorinfo node. */
+bool FFCPXMLExportVisitor::ConstructColorInfoNode(TSharedRef<FFCPXMLNode> InParentNode, const TSharedPtr<FMovieSceneExportCinematicSectionData> InCinematicSectionData)
+{
+	TSharedPtr<FFCPXMLNode> ColorInfoNode = InParentNode->CreateChildNode(TEXT("colorinfo"));
+	ColorInfoNode->CreateChildNode(TEXT("lut"));
+	ColorInfoNode->CreateChildNode(TEXT("lut1"));
+	ColorInfoNode->CreateChildNode(TEXT("asc_sop"));
+	ColorInfoNode->CreateChildNode(TEXT("asc_sat"));
+	ColorInfoNode->CreateChildNode(TEXT("lut2"));
+
+	return true;
+}
+
+bool FFCPXMLExportVisitor::ConstructSequenceNode(TSharedRef<FFCPXMLNode> InParentNode)
+{
+	if (!ExportData->IsExportDataValid() || !ExportData->MovieSceneData.IsValid()) 
+	{ 
+		return false; 
+	}
+
+	TSharedRef<FFCPXMLNode> SequenceNode = InParentNode->CreateChildNode("sequence");
+
+	// attributes
+	SequenceNode->AddAttribute(TEXT("id"), FString::Printf(TEXT("sequence%d"), ++SequenceId));
+
+	// required elements
+	TSharedRef<FFCPXMLNode> DurationNode = SequenceNode->CreateChildNode(TEXT("duration"));
+	DurationNode->SetContent(ExportData->MovieSceneData->Duration);
+
+	if (!ConstructRateNode(SequenceNode)) 
+	{ 
+		return false; 
+	}
+
+	TSharedRef<FFCPXMLNode> NameNode = SequenceNode->CreateChildNode(TEXT("name"));
+	NameNode->SetContent(ExportData->MovieSceneData->Name);
+
+	TSharedRef<FFCPXMLNode> MediaNode = SequenceNode->CreateChildNode(TEXT("media"));
+
+	if (!ConstructVideoNode(MediaNode)) 
+	{ 
+		return false; 
+	}
+
+	if (!ConstructAudioNode(MediaNode))
+	{
+		return false;
+	}
+
+	if (!ConstructTimecodeNode(SequenceNode))
+	{
+		return false;
+	}
+
+	return true;
+}
+
+bool FFCPXMLExportVisitor::ConstructVideoNode(TSharedRef<FFCPXMLNode> InParentNode)
+{
+	if (!ExportData->IsExportDataValid() || !ExportData->MovieSceneData.IsValid() || !ExportData->MovieSceneData->CinematicMasterTrack.IsValid())
+	{ 
+		return false; 
+	}
+
+	TSharedPtr<FMovieSceneExportCinematicMasterTrackData> CinematicMasterTrackData = ExportData->MovieSceneData->CinematicMasterTrack;
+	if (!CinematicMasterTrackData.IsValid())
+	{
+		return false;
+	}
+
+	TSharedRef<FFCPXMLNode> VideoNode = InParentNode->CreateChildNode(TEXT("video"));
+
+	TSharedRef<FFCPXMLNode> FormatNode = VideoNode->CreateChildNode(TEXT("format"));
+	int32 Width{ 0 };
+	int32 Height{ 0 };
+	GetDefaultImageResolution(Width, Height);
+	if (!ConstructSampleCharacteristicsNode(FormatNode, Width, Height))
+	{
+		return false;
+	}
+ 
+	for (TSharedPtr<FMovieSceneExportCinematicTrackData> CinematicTrack : CinematicMasterTrackData->CinematicTracks)
+	{
+		if (!ConstructVideoTrackNode(VideoNode, CinematicTrack)) 
+		{ 
+			return false;
+		}
+	}
+
+	return true;
+}
+
+bool FFCPXMLExportVisitor::ConstructAudioNode(TSharedRef<FFCPXMLNode> InParentNode)
+{
+	if (!ExportData->IsExportDataValid())
+	{ 
+		return false; 
+	}
+
+	TSharedRef<FFCPXMLNode> AudioNode = InParentNode->CreateChildNode(TEXT("audio"));
+
+	//TSharedRef<FFCPXMLNode> FormatNode = InVideoNode->CreateChildNode(TEXT("format"));
+
+	for (TSharedPtr<FMovieSceneExportAudioTrackData> AudioTrack : ExportData->MovieSceneData->AudioTracks)
+	{
+		if (!ConstructAudioTrackNode(AudioNode, AudioTrack))
+		{
+			return false;
+		}
+	}
+
+	return true;
+}
+
+bool FFCPXMLExportVisitor::ConstructVideoTrackNode(TSharedRef<FFCPXMLNode> InParentNode, const TSharedPtr<FMovieSceneExportCinematicTrackData> InCinematicTrackData)
+{
+	if (!ExportData->IsExportDataValid() || !InCinematicTrackData.IsValid())
+	{
+		return false;
+	}
+
+	TSharedRef<FFCPXMLNode> TrackNode = InParentNode->CreateChildNode(TEXT("track"));
+
+	for (TSharedPtr<FMovieSceneExportCinematicSectionData> CinematicSection : InCinematicTrackData->CinematicSections)
+	{
+		// skip disabled sections
+		if (!CinematicSection.IsValid() || !CinematicSection->bEnabled)
+		{
+			continue;
+		}
+		
+		if (!ConstructVideoClipItemNode(TrackNode, CinematicSection, false))
+		{
+			return false;
+		}
+	}
+
+	TSharedRef<FFCPXMLNode> EnabledNode = TrackNode->CreateChildNode(TEXT("enabled"));
+	EnabledNode->SetContent(true);
+
+	TSharedRef<FFCPXMLNode> LockedNode = TrackNode->CreateChildNode(TEXT("locked"));
+	LockedNode->SetContent(false);
+
+	return true;
+}
+
+bool FFCPXMLExportVisitor::ConstructAudioTrackNode(TSharedRef<FFCPXMLNode> InParentNode, const TSharedPtr<FMovieSceneExportAudioTrackData> InAudioTrackData)
+{
+	if (!ExportData->IsExportDataValid() || !InAudioTrackData.IsValid())
+	{
+		return false;
+	}
+
+	TSharedRef<FFCPXMLNode> TrackNode = InParentNode->CreateChildNode(TEXT("track"));
+
+	for (TSharedPtr<FMovieSceneExportAudioSectionData> AudioSection : InAudioTrackData->AudioSections)
+	{
+		if (!ConstructAudioClipItemNode(TrackNode, AudioSection))
+		{
+			return false;
+		}
+	}
+
+	TSharedRef<FFCPXMLNode> EnabledNode = TrackNode->CreateChildNode(TEXT("enabled"));
+	EnabledNode->SetContent(true);
+
+	TSharedRef<FFCPXMLNode> LockedNode = TrackNode->CreateChildNode(TEXT("locked"));
+	LockedNode->SetContent(false);
+
+	return true;
+}
+
+bool FFCPXMLExportVisitor::ConstructVideoClipItemNode(TSharedRef<FFCPXMLNode> InParentNode, const TSharedPtr<FMovieSceneExportCinematicSectionData> InCinematicSectionData, bool bInMasterClip)
+{
+	if (!ExportData->IsExportDataValid() || !InCinematicSectionData.IsValid())
+	{
+		return false;
+	}
+
+	TSharedRef<FFCPXMLNode> ClipItemNode = InParentNode->CreateChildNode(TEXT("clipitem"));
+
+	int32 Duration{ 0 };
+	int32 In{ 0 };
+	int32 Out{ 0 };
+	int32 Start{ 0 };
+	int32 End{ 0 };
+	FString SectionName(TEXT(""));
+
+	GetSectionFrames(InCinematicSectionData, Duration, Start, End, In, Out);
+	GetSectionName(InCinematicSectionData, SectionName);
+
+	FString ClipItemIdName = TEXT("");
+	if (!GetMasterClipItemIdName(InCinematicSectionData, ClipItemIdName))
+	{
+		return false;
+	}
+
+	FString MasterClipIdName = TEXT("");
+	if (!GetMasterClipIdName(InCinematicSectionData, MasterClipIdName))
+	{
+		return false;
+	}
+
+	// attributes
+	ClipItemNode->AddAttribute(TEXT("id"), ClipItemIdName);
+
+	// elements
+	ClipItemNode->CreateChildNode(TEXT("masterclipid"))->SetContent(MasterClipIdName);
+	ClipItemNode->CreateChildNode(TEXT("ismasterclip"))->SetContent(bInMasterClip);
+	ClipItemNode->CreateChildNode(TEXT("name"))->SetContent(SectionName);
+	ClipItemNode->CreateChildNode(TEXT("enabled"))->SetContent(true);
+	ClipItemNode->CreateChildNode(TEXT("duration"))->SetContent(Duration);
+
+	if (!ConstructRateNode(ClipItemNode))
+	{
+		return false;
+	}
+
+	if (!bInMasterClip)
+	{
+		ClipItemNode->CreateChildNode(TEXT("start"))->SetContent(Start);
+		ClipItemNode->CreateChildNode(TEXT("end"))->SetContent(End);
+	}
+
+	ClipItemNode->CreateChildNode(TEXT("in"))->SetContent(In);
+	ClipItemNode->CreateChildNode(TEXT("out"))->SetContent(Out);
+
+	if (bInMasterClip)
+	{
+		ClipItemNode->CreateChildNode(TEXT("anamorphic"))->SetContent(false);
+		ClipItemNode->CreateChildNode(TEXT("pixelaspectratio"))->SetContent(TEXT("square"));
+		ClipItemNode->CreateChildNode(TEXT("fielddominance"))->SetContent(TEXT("lower"));
+	}
+
+	if (!ConstructVideoFileNode(ClipItemNode, InCinematicSectionData, Duration, bInMasterClip))
+	{
+		return false;
+	}
+
+	return true;
+}
+
+bool FFCPXMLExportVisitor::ConstructAudioClipItemNode(TSharedRef<FFCPXMLNode> InParentNode, const TSharedPtr<FMovieSceneExportAudioSectionData> InAudioSectionData)
+{
+	if (!ExportData->IsExportDataValid() || !InAudioSectionData.IsValid())
+	{
+		return false;
+	}
+
+	// @todo - audio
+
+	return true;
+}
+
+bool FFCPXMLExportVisitor::ConstructVideoFileNode(TSharedRef<FFCPXMLNode> InParentNode, const TSharedPtr<FMovieSceneExportCinematicSectionData> InCinematicSectionData, uint32 Duration, bool bInMasterClip)
+{
+	if (!InCinematicSectionData.IsValid())
+	{
+		return false;
+	}
+
+	FString FileIdName = TEXT("");
+	if (!GetMasterClipFileIdName(InCinematicSectionData, FileIdName))
+	{
+		return false;
+	}
+
+	// attributes
+	TSharedRef<FFCPXMLNode> FileNode = InParentNode->CreateChildNode(TEXT("file"));
+	FileNode->AddAttribute(TEXT("id"), FileIdName);
+
+	if (bInMasterClip)
+	{
+		FString FilePathName = SaveFilePath + TEXT("/") + InCinematicSectionData->ShotFilename;
+		FString FilePathUrl = FString(TEXT("file://localhost/")) + FilePathName.Replace(TEXT(" "), TEXT("%20")).Replace(TEXT(":"), TEXT("%3a"));
+
+		// required elements
+		TSharedRef<FFCPXMLNode> NameNode = FileNode->CreateChildNode(TEXT("name"));
+		NameNode->SetContent(InCinematicSectionData->ShotFilename);
+
+		TSharedRef<FFCPXMLNode> PathUrlNode = FileNode->CreateChildNode(TEXT("pathurl"));
+		PathUrlNode->SetContent(FilePathUrl);
+
+		if (!ConstructRateNode(FileNode))
+		{
+			return false;
+		}
+
+		TSharedRef<FFCPXMLNode> DurationNode = FileNode->CreateChildNode(TEXT("duration"));
+		DurationNode->SetContent(static_cast<int32>(Duration));
+
+		if (!ConstructTimecodeNode(FileNode))
+		{
+			return false;
+		}
+
+		TSharedRef<FFCPXMLNode> MediaNode = FileNode->CreateChildNode(TEXT("media"));
+		TSharedRef<FFCPXMLNode> VideoNode = MediaNode->CreateChildNode(TEXT("video"));
+
+		int32 Width{ 0 };
+		int32 Height{ 0 };
+		GetDefaultImageResolution(Width, Height);
+		if (!ConstructSampleCharacteristicsNode(VideoNode, Width, Height))
+		{
+			return false;
+		}
+	}
+
+	return true;
+}
+
+bool FFCPXMLExportVisitor::ConstructAudioFileNode(TSharedRef<FFCPXMLNode> InParentNode, const TSharedPtr<FMovieSceneExportAudioSectionData> InAudioSectionData)
+{
+	// @todo - audio
+
+	return true;
+}
+
+bool FFCPXMLExportVisitor::ConstructSampleCharacteristicsNode(TSharedRef<FFCPXMLNode> InParentNode, int InWidth, int InHeight)
+{
+	TSharedRef<FFCPXMLNode> SampleCharacteristicsNode = InParentNode->CreateChildNode(TEXT("samplecharacteristics"));
+
+	if (!ConstructRateNode(SampleCharacteristicsNode))
+	{
+		return false;
+	}
+
+	SampleCharacteristicsNode->CreateChildNode(TEXT("width"))->SetContent(InWidth);
+	SampleCharacteristicsNode->CreateChildNode(TEXT("height"))->SetContent(InHeight);
+	SampleCharacteristicsNode->CreateChildNode(TEXT("anamorphic"))->SetContent(false);
+	SampleCharacteristicsNode->CreateChildNode(TEXT("pixelaspectratio"))->SetContent(TEXT("square"));
+	SampleCharacteristicsNode->CreateChildNode(TEXT("fielddominance"))->SetContent(TEXT("lower"));
+
+	return true;
+}
+
+bool FFCPXMLExportVisitor::ConstructRateNode(TSharedRef<FFCPXMLNode> InParentNode)
+{
+	TSharedRef<FFCPXMLNode> RateNode = InParentNode->CreateChildNode(TEXT("rate"));
+
+	TSharedRef<FFCPXMLNode> TimebaseNode = RateNode->CreateChildNode(TEXT("timebase"));
+	TimebaseNode->SetContent(static_cast<int32>(ExportData->GetNearestWholeFrameRate()));
+
+	TSharedRef<FFCPXMLNode> NTSCNode = RateNode->CreateChildNode(TEXT("ntsc"));
+	NTSCNode->SetContent(ExportData->GetFrameRateIsNTSC());
+
+	return true;
+}
+
+bool FFCPXMLExportVisitor::ConstructTimecodeNode(TSharedRef<FFCPXMLNode> InParentNode)
+{
+	TSharedRef<FFCPXMLNode> TimecodeNode = InParentNode->CreateChildNode(TEXT("timecode"));
+
+	if (!ConstructRateNode(TimecodeNode))
+	{
+		return false;
+	}
+
+	TimecodeNode->CreateChildNode(TEXT("string"))->SetContent(TEXT("00:00:00:00"));
+	TimecodeNode->CreateChildNode(TEXT("frame"))->SetContent(0);
+	TimecodeNode->CreateChildNode(TEXT("displayformat"))->SetContent(TEXT("NDF"));
+
+	TSharedRef<FFCPXMLNode> ReelNode = TimecodeNode->CreateChildNode(TEXT("reel"));
+	ReelNode->CreateChildNode(TEXT("name"))->SetContent(TEXT(""));
+
+	return true;
+}
+
+/** Get duration, in and out frames for a given shot section */
+bool FFCPXMLExportVisitor::GetSectionFrames(const TSharedPtr<FMovieSceneExportCinematicSectionData> InCinematicSectionData, int32& OutDuration, int32& OutStartFrame, int32&OutEndFrame, int32& OutInFrame, int32& OutOutFrame)
+{
+	if (!InCinematicSectionData.IsValid() || !ExportData->MovieSceneData.IsValid())
+	{
+		return false;
+	}
+
+	int32 Handles = ExportData->GetHandleFrames().FloorToFrame().Value;
+	OutStartFrame = InCinematicSectionData->StartFrame.Value;
+	OutEndFrame = InCinematicSectionData->EndFrame.Value;
+	OutDuration = OutEndFrame - OutStartFrame;
+	OutInFrame = Handles;
+	OutOutFrame = Handles + OutDuration;
+
+	return true;
+}
+
+/** Get name of a given shot section */
+bool FFCPXMLExportVisitor::GetSectionName(const TSharedPtr<FMovieSceneExportCinematicSectionData> InCinematicSectionData, FString& OutSectionName)
+{
+	if (!InCinematicSectionData.IsValid())
+	{
+		return false;
+	}
+	OutSectionName = InCinematicSectionData->ShotDisplayName;
+	return true;
+}
+
+bool FFCPXMLExportVisitor::CreateMasterClipIdName(const TSharedPtr<FMovieSceneExportCinematicSectionData> InCinematicSectionData, uint32 InId, FString& OutName)
+{
+	if (!InCinematicSectionData.IsValid())
+	{
+		return false;
+	}
+	FString Name(TEXT(""));
+	GetSectionName(InCinematicSectionData, Name);
+	MasterClipIdMap.Add(Name, InId);
+	OutName = FString::Printf(TEXT("masterclip%d_"), InId) + Name;
+	return true;
+}
+
+bool FFCPXMLExportVisitor::GetIdName(const TSharedPtr<FMovieSceneExportCinematicSectionData> InCinematicSectionData, FString InPrefix, FString &OutName)
+{
+	if (!InCinematicSectionData.IsValid())
+	{
+		return false;
+	}
+	FString Name(TEXT(""));
+	GetSectionName(InCinematicSectionData, Name);
+	uint32* Id = MasterClipIdMap.Find(Name);
+	if (Id == nullptr)
+	{
+		return false;
+	}
+	OutName = InPrefix + FString::Printf(TEXT("%d_"), *Id) + Name;
+	return true;
+}
+
+bool FFCPXMLExportVisitor::GetMasterClipIdName(const TSharedPtr<FMovieSceneExportCinematicSectionData> InCinematicSectionData, FString& OutName)
+{
+	return GetIdName(InCinematicSectionData, FString(TEXT("masterclip")), OutName);
+}
+
+bool FFCPXMLExportVisitor::GetMasterClipItemIdName(const TSharedPtr<FMovieSceneExportCinematicSectionData> InCinematicSectionData, FString& OutName)
+{
+	return GetIdName(InCinematicSectionData, FString(TEXT("clipitem")), OutName);
+}
+
+bool FFCPXMLExportVisitor::GetMasterClipFileIdName(const TSharedPtr<FMovieSceneExportCinematicSectionData> InCinematicSectionData, FString& OutName)
+{
+	return GetIdName(InCinematicSectionData, FString(TEXT("file")), OutName);
+}
+
+void FFCPXMLExportVisitor::GetDefaultImageResolution(int32& OutWidth, int32& OutHeight)
+{
+	OutWidth = 1280;
+	OutHeight = 720;
+}
+
+FString FFCPXMLExportVisitor::GetFilePathName(const FString& InSectionName) const
+{
+	return (SaveFilePath + TEXT("/") + InSectionName);
+}
+
+#undef LOCTEXT_NAMESPACE
diff --git a/Engine/Source/Editor/MovieSceneTools/Private/FCPXML/FCPXMLExport.h b/Engine/Source/Editor/MovieSceneTools/Private/FCPXML/FCPXMLExport.h
new file mode 100644
--- /dev/null
+++ b/Engine/Source/Editor/MovieSceneTools/Private/FCPXML/FCPXMLExport.h
@@ -0,0 +1,120 @@
+// Copyright 1998-2018 Epic Games, Inc. All Rights Reserved.
+
+#pragma once
+
+#include "FCPXML/FCPXMLNode.h"
+#include "MovieScene.h"
+#include "MovieSceneTranslator.h"
+#include "Tracks/MovieSceneAudioTrack.h"
+#include "Sections/MovieSceneCinematicShotSection.h"
+#include "Tracks/MovieSceneCinematicShotTrack.h"
+
+/** The FFCPXMLExportVisitor class exports from Sequencer data into the FCP 7 XML structure.
+
+    This class will eventually be used to merge the data with an existing XML structure 
+	representing previously imported material. This is intended to preserve metadata 
+	roundtrip between Sequencer and the FCP XML format.
+	
+	Currently the VisitNode functions are mostly empty. This is where the
+	merge data functionality will be implemented. */
+
+class FFCPXMLExportVisitor : public FFCPXMLNodeVisitor
+{
+public:
+	/** Constructor */
+	FFCPXMLExportVisitor(FString InSaveFilename, TSharedRef<FMovieSceneExportData> InExportData, TSharedRef<FMovieSceneTranslatorContext> InExportContext);
+	/** Destructor */
+	virtual ~FFCPXMLExportVisitor();
+
+public:
+
+	/** Called when visiting a FFCPXMLBasicNode during visitor traversal. */
+	virtual bool VisitNode(TSharedRef<FFCPXMLBasicNode> InBasicNode) override final;
+	/** Called when visiting a FFCPXMLXmemlNode during visitor traversal. */
+	virtual bool VisitNode(TSharedRef<FFCPXMLXmemlNode> InXmemlNode) override final;
+	/** Called when visiting a FFCPXMLSequenceNode during visitor traversal. */
+	virtual bool VisitNode(TSharedRef<FFCPXMLSequenceNode> InSequenceNode) override final;
+	/** Called when visiting a FFCPXMLVideoNode during visitor traversal. */
+	virtual bool VisitNode(TSharedRef<FFCPXMLVideoNode> InVideoNode) override final;
+	/** Called when visiting a FFCPXMLAudioNode during visitor traversal. */
+	virtual bool VisitNode(TSharedRef<FFCPXMLAudioNode> InAudioNode) override final;
+	/** Called when visiting a FFCPXMLTrackNode during visitor traversal. */
+	virtual bool VisitNode(TSharedRef<FFCPXMLTrackNode> InTrackNode) override final;
+	/** Called when visiting a FFCPXMLClipNode during visitor traversal. */
+	virtual bool VisitNode(TSharedRef<FFCPXMLClipNode> InClipNode) override final;
+	/** Called when visiting a FFCPXMLClipItemNode during visitor traversal. */
+	virtual bool VisitNode(TSharedRef<FFCPXMLClipItemNode> InClipItemNode) override final;
+	/** Called when visiting a FFCPXMLFileNode during visitor traversal. */
+	virtual bool VisitNode(TSharedRef<FFCPXMLFileNode> InFileNode) override final;
+
+public:
+	/** Creates project node. */
+	bool ConstructProjectNode(TSharedRef<FFCPXMLNode> InParentNode);
+	/** Creates master clip nodes. */
+	bool ConstructMasterClipNodes(TSharedRef<FFCPXMLNode> InParentNode);
+	/** Creates master clip node. */
+	bool ConstructMasterClipNode(TSharedRef<FFCPXMLNode> InParentNode, const TSharedPtr<FMovieSceneExportCinematicSectionData> InCinematicSectionData, int32 InMasterClipId);
+	/** Creates colorinfo node. */
+	bool ConstructColorInfoNode(TSharedRef<FFCPXMLNode> InParentNode, const TSharedPtr<FMovieSceneExportCinematicSectionData> InCinematicSectionData);
+	/** Creates logginginfo node. */
+	bool ConstructSectionLoggingInfoNode(TSharedRef<FFCPXMLNode> InParentNode, const TSharedPtr<FMovieSceneExportCinematicSectionData> InCinematicSectionData, const FString& InSectionName);
+	/** Creates sequence node. */
+	bool ConstructSequenceNode(TSharedRef<FFCPXMLNode> InParentNode);
+	/** Creates video node. */
+	bool ConstructVideoNode(TSharedRef<FFCPXMLNode> InParentNode);
+	/** Creates audio node. */
+	bool ConstructAudioNode(TSharedRef<FFCPXMLNode> InParentNode);
+	/** Creates cinematic track node. */
+	bool ConstructVideoTrackNode(TSharedRef<FFCPXMLNode> InParentNode, const TSharedPtr<FMovieSceneExportCinematicTrackData> InCinematicTrackData);
+	/** Creates audio track node. */
+	bool ConstructAudioTrackNode(TSharedRef<FFCPXMLNode> InParentNode, const TSharedPtr<FMovieSceneExportAudioTrackData> InAudioTrackData);
+	/** Creates video clip item node. */
+	bool ConstructVideoClipItemNode(TSharedRef<FFCPXMLNode> InParentNode, const TSharedPtr<FMovieSceneExportCinematicSectionData> InCinematicSectionData, bool bInMasterClip);
+	/** Creates audio clip item node. */
+	bool ConstructAudioClipItemNode(TSharedRef<FFCPXMLNode> InParentNode, const TSharedPtr<FMovieSceneExportAudioSectionData> InAudioSectionData);
+	/** Creates video file node. */
+	bool ConstructVideoFileNode(TSharedRef<FFCPXMLNode> InParentNode, const TSharedPtr<FMovieSceneExportCinematicSectionData> InCinematicSectionData, uint32 Duration, bool bInMasterClip);
+	/** Creates audio file node. */
+	bool ConstructAudioFileNode(TSharedRef<FFCPXMLNode> InParentNode, const TSharedPtr<FMovieSceneExportAudioSectionData> InAudioSectionData);
+	/** Creates sample characteristics node */
+	bool ConstructSampleCharacteristicsNode(TSharedRef<FFCPXMLNode> InParentNode, int InWidth, int InHeight);
+	/** Creates rate node. */
+	bool ConstructRateNode(TSharedRef<FFCPXMLNode> InParentNode);
+	/** Creates timecode node. */
+	bool ConstructTimecodeNode(TSharedRef<FFCPXMLNode> InParentNode);
+	/** Creates logginginfo elements. */
+	void ConstructLoggingInfoElements(TSharedRef<FFCPXMLNode> InLoggingInfoNode, const UObject* InObject) ;
+	/** Set logginginfo element value. */
+	void SetLoggingInfoElementValue(TSharedPtr<FFCPXMLNode> InNode, const UObject* InObject, const FString& InElement) ;
+	/** Get duration, start, end, in, out frames for a given shot section */
+	bool GetSectionFrames(const TSharedPtr<FMovieSceneExportCinematicSectionData> InCinematicSectionData, int32& OutDuration, int32& OutStartFrame, int32& OutEndFrame, int32& OutInFrame, int32& OutOutFrame);
+	/** Get start and end frames for a given shot section */
+	//bool GetSectionStartEnd(const TSharedPtr<FMovieSceneExportCinematicSectionData> InCinematicSectionData, int32& OutStartFrame, int32& OutEndFrame);
+	/** Get name of a given shot section */
+	bool GetSectionName(const TSharedPtr<FMovieSceneExportCinematicSectionData> InCinematicSectionData, FString& OutSectionName);
+
+private:
+
+	/** Construct a master clip id name based on input name and id */
+	bool CreateMasterClipIdName(const TSharedPtr<FMovieSceneExportCinematicSectionData> InCinematicSectionData, uint32 InId, FString& OutName);
+	/** Get id name for given section and prefix */
+	bool GetIdName(const TSharedPtr<FMovieSceneExportCinematicSectionData> InCinematicSectionData, FString InPrefix, FString &OutName);
+	/** Get master clip id name for given section */
+	bool GetMasterClipIdName(const TSharedPtr<FMovieSceneExportCinematicSectionData> InCinematicSectionData, FString& OutName);
+	/** Get master clip's clipitem id name for given section */
+	bool GetMasterClipItemIdName(const TSharedPtr<FMovieSceneExportCinematicSectionData> InCinematicSectionData, FString& OutName);
+	/** Get master clip's file id name based on input name and id */
+	bool GetMasterClipFileIdName(const TSharedPtr<FMovieSceneExportCinematicSectionData> InCinematicSectionData, FString& OutName);
+	/** Get the default image resolution */
+	void GetDefaultImageResolution(int32& OutWidth, int32& OutHeight);
+	/** Return filename with path */
+	FString GetFilePathName(const FString& InSectionName) const;
+
+private:
+
+	TSharedRef<FMovieSceneExportData> ExportData;
+	TSharedRef<FMovieSceneTranslatorContext> ExportContext;
+	FString SaveFilePath;
+	TMap<FString, uint32> MasterClipIdMap;
+	uint32 SequenceId;
+};
diff --git a/Engine/Source/Editor/MovieSceneTools/Private/FCPXML/FCPXMLImport.cpp b/Engine/Source/Editor/MovieSceneTools/Private/FCPXML/FCPXMLImport.cpp
new file mode 100644
--- /dev/null
+++ b/Engine/Source/Editor/MovieSceneTools/Private/FCPXML/FCPXMLImport.cpp
@@ -0,0 +1,349 @@
+// Copyright 1998-2018 Epic Games, Inc. All Rights Reserved.
+
+#include "FCPXML/FCPXMLImport.h"
+#include "MovieScene.h"
+#include "MovieSceneTranslator.h"
+#include "LevelSequence.h"
+#include "Tracks/MovieSceneAudioTrack.h"
+#include "Sections/MovieSceneCinematicShotSection.h"
+#include "Tracks/MovieSceneCinematicShotTrack.h"
+#include "AssetRegistryModule.h"
+#include "UObject/MetaData.h"
+
+#define LOCTEXT_NAMESPACE "FCPXMLImporter"
+// DEFINE_LOG_CATEGORY_STATIC(FCPXMLImporter, Log, All);
+
+#define INHERIT true
+#define NO_INHERIT false
+
+FFCPXMLImportVisitor::FFCPXMLImportVisitor(TSharedRef<FMovieSceneImportData> InImportData, TSharedRef<FMovieSceneTranslatorContext> InImportContext) :
+	FFCPXMLNodeVisitor()
+	, ImportData(InImportData)
+	, ImportContext(InImportContext)
+	, bInSequenceNode(false)
+	, bInVideoNode(false)
+	, bInAudioNode(false)
+	, bInVideoTrackNode(false)
+	, bInAudioTrackNode(false)
+	, CurrVideoTrackRowIndex(0)
+	, CurrAudioTrackIndex(0)
+{
+}
+	
+FFCPXMLImportVisitor::~FFCPXMLImportVisitor() {}
+
+bool FFCPXMLImportVisitor::VisitNode(TSharedRef<FFCPXMLBasicNode> InBasicNode)
+{
+	return InBasicNode->VisitChildren(*this);
+}
+
+bool FFCPXMLImportVisitor::VisitNode(TSharedRef<FFCPXMLXmemlNode> InXmemlNode)
+{
+	return InXmemlNode->VisitChildren(*this);
+}
+
+bool FFCPXMLImportVisitor::VisitNode(TSharedRef<FFCPXMLSequenceNode> InSequenceNode)
+{ 
+	bool bPrev = bInSequenceNode;
+	bInSequenceNode = true;
+
+	// Sequences can be referenced so flag to visit reference node children
+	bool bSuccess = InSequenceNode->VisitChildren(*this, true);
+	
+	bInSequenceNode = false;
+
+	return bSuccess;
+}
+
+bool FFCPXMLImportVisitor::VisitNode(TSharedRef<FFCPXMLVideoNode> InVideoNode)
+{
+	/*
+	if (!ImportData->GetCinematicMasterTrackData(true).IsValid())
+	{
+		return false;
+	}
+	*/
+	bool bPrev = bInVideoNode;
+	bInVideoNode = true;
+	int32 PrevVideoTrackRowIndex = CurrVideoTrackRowIndex;
+	CurrVideoTrackRowIndex = 0;
+
+	bool bSuccess = InVideoNode->VisitChildren(*this);
+	
+	bInVideoNode = bPrev;
+	CurrVideoTrackRowIndex = PrevVideoTrackRowIndex;
+
+	return bSuccess;
+}
+
+bool FFCPXMLImportVisitor::VisitNode(TSharedRef<FFCPXMLAudioNode> InAudioNode)
+{
+	// @todo
+	CurrAudioTrackIndex = 0;
+	return InAudioNode->VisitChildren(*this);
+}
+
+bool FFCPXMLImportVisitor::VisitNode(TSharedRef<FFCPXMLTrackNode> InTrackNode)
+{
+	bool bPrevInTrackNode = false;
+	int32 PrevTrackIndex = 0;
+
+	if (bInSequenceNode && bInVideoNode)
+	{
+		bPrevInTrackNode = bInVideoTrackNode;
+		bInVideoTrackNode = true;
+	}
+	else if (bInSequenceNode && bInAudioNode)
+	{
+		bPrevInTrackNode = bInAudioTrackNode;
+		bInAudioTrackNode = true;
+	}
+	
+	bool bSuccess = InTrackNode->VisitChildren(*this);
+
+	if (bInSequenceNode && bInVideoNode)
+	{
+		bInVideoTrackNode = bPrevInTrackNode;
+		CurrVideoTrackRowIndex++;
+	}
+	else if (bInSequenceNode && bInAudioNode)
+	{
+		bInAudioTrackNode = bPrevInTrackNode;
+		CurrAudioTrackIndex++;
+	}
+
+	return bSuccess;
+}
+
+bool FFCPXMLImportVisitor::VisitNode(TSharedRef<FFCPXMLClipNode> InClipNode)
+{
+	// masterclip
+	bool bIsMasterClip = false;
+	FString MasterClipName(TEXT(""));
+	FString MetadataSectionName(TEXT(""));
+
+	if (InClipNode->GetChildValue<bool>("ismasterclip", bIsMasterClip, ENodeInherit::NoInherit) && bIsMasterClip)
+	{
+		bool bHasMasterClipName = InClipNode->GetChildValue<FString>("masterclipid", MasterClipName, ENodeInherit::NoInherit);
+
+		TSharedPtr<FFCPXMLNode> LoggingInfoNode = InClipNode->GetChildNode("logginginfo", ENodeInherit::NoInherit, ENodeReference::NoReferences);
+		if (LoggingInfoNode.IsValid())
+		{
+			bool bHasLoggingShotTrack = LoggingInfoNode->GetChildValue<FString>("lognote", MetadataSectionName, ENodeInherit::NoInherit, ENodeReference::NoReferences);
+
+			if (!MetadataSectionName.IsEmpty())
+			{
+				if (bHasMasterClipName)
+				{
+					FString CinematicSectionName;
+					if (FFCPXMLClipNode::GetSequencerSectionName(MetadataSectionName, CinematicSectionName))
+					{
+						AddMasterClipSectionName(MasterClipName, CinematicSectionName);
+						AddSectionMetaData(CinematicSectionName, LoggingInfoNode);
+					}
+				}
+			}
+		}
+	}
+
+	// Clips can be referenced so flag to visit reference node children
+	return InClipNode->VisitChildren(*this, true);
+}
+
+bool FFCPXMLImportVisitor::VisitNode(TSharedRef<FFCPXMLClipItemNode> InClipItemNode)
+{
+	if (bInSequenceNode && bInVideoTrackNode)
+	{
+		// Sequencer section name 
+		FString CinematicSectionName(TEXT(""));
+		bool bHasCinematicSectionName = false;
+
+		// masterclip
+		FString MasterClipIdName(TEXT(""));
+		if (InClipItemNode->GetChildValue<FString>("masterclipid", MasterClipIdName, ENodeInherit::NoInherit))
+		{
+			bHasCinematicSectionName = GetMasterClipSectionName(MasterClipIdName, CinematicSectionName);
+		}
+
+		// Frame rate 
+		int32 FrameRateTimebase = 0;
+		if (!InClipItemNode->GetChildSubValue<int32>("rate", "timebase", FrameRateTimebase, ENodeInherit::CheckInherit))
+		{
+			return false;
+		}
+
+		bool bFrameRateIsNTSC = false;
+		if (!InClipItemNode->GetChildSubValue<bool>("rate", "ntsc", bFrameRateIsNTSC, ENodeInherit::CheckInherit))
+		{
+			bFrameRateIsNTSC = false;
+		}
+
+		// Shot name
+		FString ShotName;
+		if (!InClipItemNode->GetChildValue("name", ShotName, ENodeInherit::NoInherit) || ShotName.IsEmpty())
+		{
+			return false;
+		}
+
+		if (!bHasCinematicSectionName)
+		{
+			CinematicSectionName = ShotName;
+		}
+
+		// Start frame 
+		int32 StartVal = 0;
+		if (!InClipItemNode->GetChildValue("start", StartVal, ENodeInherit::NoInherit))
+		{
+			return false;
+		}
+		FFrameNumber StartFrame = StartVal;
+
+		// End frame 
+		int32 EndVal = 0;
+		if (!InClipItemNode->GetChildValue("end", EndVal, ENodeInherit::NoInherit))
+		{
+			return false;
+		}
+
+		// In frame 
+		int32 InVal = 0;
+		bool bInDefault = false;
+		if (!InClipItemNode->GetChildValue("in", InVal, ENodeInherit::NoInherit))
+		{
+			bInDefault = true;
+		}
+
+		// Out frame 
+		int32 OutVal = 0;
+		bool bOutDefault = false;
+		if (!InClipItemNode->GetChildValue("out", OutVal, ENodeInherit::NoInherit))
+		{
+			bOutDefault = true;
+		}
+
+		uint32 Numerator = static_cast<uint32>(FrameRateTimebase);
+		FFrameRate FrameRate(Numerator, 1);
+		FFrameNumber StartOffset = bInDefault ? 0 : InVal;
+		FFrameNumber Start = StartVal;
+		FFrameNumber End = EndVal;
+
+		// @todo handle xml subclips?
+
+		TSharedPtr<FMovieSceneImportCinematicSectionData> SectionData = ImportData->FindCinematicSection(CinematicSectionName);
+		if (!SectionData.IsValid())
+		{
+			SectionData = ImportData->CreateCinematicSection(CinematicSectionName, CurrVideoTrackRowIndex, FrameRate, Start, End, StartOffset);
+			if (!SectionData.IsValid())
+			{
+				return false;
+			}
+		}
+		else
+		{
+			if (!ImportData->SetCinematicSection(SectionData, CurrVideoTrackRowIndex, FrameRate, Start, End, StartOffset))
+			{
+				return false;
+			}
+		}
+
+		ImportSectionMetaData(CinematicSectionName, SectionData->CinematicSection);
+	}
+
+	// Clip items can be referenced so flag to visit reference node children
+	return InClipItemNode->VisitChildren(*this, true);
+}
+
+bool FFCPXMLImportVisitor::VisitNode(TSharedRef<FFCPXMLFileNode> InFileNode)
+{
+	// Files can be referenced so flag to visit reference node children
+	return InFileNode->VisitChildren(*this, true);
+}
+
+bool FFCPXMLImportVisitor::AddMasterClipSectionName(const FString& InMasterClipIdName, const FString& InSequencerSectionName)
+{
+	if (MasterClipSectionNameMap.Contains(InMasterClipIdName))
+	{
+		return false;
+	}
+
+	MasterClipSectionNameMap.Add(InMasterClipIdName, InSequencerSectionName);
+	return true;
+}
+
+bool FFCPXMLImportVisitor::GetMasterClipSectionName(const FString& InMasterClipIdName, FString& OutSequencerSectionName) const
+{
+	const FString* Name = MasterClipSectionNameMap.Find(InMasterClipIdName);
+	if (Name == nullptr)
+	{
+		return false;
+	}
+
+	OutSequencerSectionName = *Name;
+	return true;
+}
+
+
+/** Add metadata entry */
+bool FFCPXMLImportVisitor::AddSectionMetaData(const FString& InSequencerSectionName, TSharedPtr<FFCPXMLNode> InLoggingInfoNode)
+{
+	if (SectionMetaDataMap.Contains(InSequencerSectionName))
+	{
+		return false;
+	}
+
+	SectionMetaDataMap.Add(InSequencerSectionName, InLoggingInfoNode);
+	return true;
+}
+
+void SetMetaDataValue(TSharedPtr<FFCPXMLNode> InNode, const FString& InElement, UMetaData* InMetaData, const UMovieSceneCinematicShotSection* InSection)
+{
+	if (!InNode.IsValid() || InMetaData == nullptr || InSection == nullptr)
+	{
+		return;
+	}
+
+	FString Value;
+	InNode->GetChildValue(InElement, Value, ENodeInherit::NoInherit, ENodeReference::NoReferences);
+	if (!Value.IsEmpty())
+	{
+		InMetaData->SetValue(InSection, *InElement, *Value);
+	}
+}
+
+/** Store clip metadata */
+bool FFCPXMLImportVisitor::ImportSectionMetaData(const FString& InSequencerSectionName, const UMovieSceneCinematicShotSection* InSection)
+{
+	if (InSequencerSectionName.IsEmpty() || InSection == nullptr)
+	{
+		return false;
+	}
+
+	TSharedPtr<FFCPXMLNode>* LoggingInfoNode = SectionMetaDataMap.Find(InSequencerSectionName);
+	if (LoggingInfoNode != nullptr && (*LoggingInfoNode).IsValid())
+	{
+		UPackage* Package = InSection->GetOutermost();
+		check(Package);
+
+		UMetaData* MetaData = Package->GetMetaData();
+		check(MetaData);
+
+		SetMetaDataValue(*LoggingInfoNode, TEXT("description"), MetaData, InSection);
+		SetMetaDataValue(*LoggingInfoNode, TEXT("scene"), MetaData, InSection);
+		SetMetaDataValue(*LoggingInfoNode, TEXT("shottake"), MetaData, InSection);
+		SetMetaDataValue(*LoggingInfoNode, TEXT("good"), MetaData, InSection);
+		SetMetaDataValue(*LoggingInfoNode, TEXT("originalvideofilename"), MetaData, InSection);
+		SetMetaDataValue(*LoggingInfoNode, TEXT("originalaudiofilename"), MetaData, InSection);
+/*
+		FString Value;
+		(*LoggingInfoNode)->GetChildValue(TEXT("description"), Value, ENodeInherit::NoInherit, ENodeReference::NoReferences);
+		if (!Value.IsEmpty())
+		{
+			MetaData->SetValue(InSection, TEXT("description"), *Value);
+		}
+		*/
+	}
+
+	return true;
+}
+
+#undef LOCTEXT_NAMESPACE
diff --git a/Engine/Source/Editor/MovieSceneTools/Private/FrameNumberDetailsCustomization.cpp b/Engine/Source/Editor/MovieSceneTools/Private/FrameNumberDetailsCustomization.cpp
--- a/Engine/Source/Editor/MovieSceneTools/Private/FrameNumberDetailsCustomization.cpp
+++ b/Engine/Source/Editor/MovieSceneTools/Private/FrameNumberDetailsCustomization.cpp
@@ -76,10 +76,10 @@ void FFrameNumberDetailsCustomization::OnTimeTextCommitted(const FText& InText,
 {
 	int32 ExistingValue = 0.0;
 	FrameNumberProperty->GetValue(ExistingValue);
-	TOptional<double> FrameResolution = NumericTypeInterface->FromString(InText.ToString(), ExistingValue);
-	if (FrameResolution.IsSet())
+	TOptional<double> TickResolution = NumericTypeInterface->FromString(InText.ToString(), ExistingValue);
+	if (TickResolution.IsSet())
 	{
-		double ClampedValue = FMath::Clamp(FrameResolution.GetValue(), (double)UIClampMin, (double)UIClampMax);
+		double ClampedValue = FMath::Clamp(TickResolution.GetValue(), (double)UIClampMin, (double)UIClampMax);
 		FrameNumberProperty->SetValue((int32)ClampedValue);
 	}
 }
diff --git a/Engine/Source/Editor/MovieSceneTools/Private/MatineeImportTools.cpp b/Engine/Source/Editor/MovieSceneTools/Private/MatineeImportTools.cpp
--- a/Engine/Source/Editor/MovieSceneTools/Private/MatineeImportTools.cpp
+++ b/Engine/Source/Editor/MovieSceneTools/Private/MatineeImportTools.cpp
@@ -136,7 +136,7 @@ bool FMatineeImportTools::CopyInterpBoolTrack( UInterpTrackBoolProp* MatineeBool
 
 	BoolTrack->Modify();
 
-	FFrameRate   FrameRate    = BoolTrack->GetTypedOuter<UMovieScene>()->GetFrameResolution();
+	FFrameRate   FrameRate    = BoolTrack->GetTypedOuter<UMovieScene>()->GetTickResolution();
 	FFrameNumber FirstKeyTime = (MatineeBoolTrack->GetKeyframeTime( 0 ) * FrameRate).RoundToFrame();
 
 	UMovieSceneBoolSection* Section = Cast<UMovieSceneBoolSection>( MovieSceneHelpers::FindSectionAtTime( BoolTrack->GetAllSections(), FirstKeyTime ) );
@@ -177,7 +177,7 @@ bool FMatineeImportTools::CopyInterpFloatTrack( UInterpTrackFloatBase* MatineeFl
 
 	FloatTrack->Modify();
 
-	FFrameRate   FrameRate    = FloatTrack->GetTypedOuter<UMovieScene>()->GetFrameResolution();
+	FFrameRate   FrameRate    = FloatTrack->GetTypedOuter<UMovieScene>()->GetTickResolution();
 	FFrameNumber FirstKeyTime = (MatineeFloatTrack->GetKeyframeTime( 0 ) * FrameRate).RoundToFrame();;
 
 	UMovieSceneFloatSection* Section = Cast<UMovieSceneFloatSection>( MovieSceneHelpers::FindSectionAtTime( FloatTrack->GetAllSections(), FirstKeyTime ) );
@@ -222,7 +222,7 @@ bool FMatineeImportTools::CopyInterpVectorTrack( UInterpTrackVectorProp* Matinee
 
 	VectorTrack->Modify();
 
-	FFrameRate   FrameRate    = VectorTrack->GetTypedOuter<UMovieScene>()->GetFrameResolution();
+	FFrameRate   FrameRate    = VectorTrack->GetTypedOuter<UMovieScene>()->GetTickResolution();
 	FFrameNumber FirstKeyTime = (MatineeVectorTrack->GetKeyframeTime( 0 ) * FrameRate).RoundToFrame();
 
 	UMovieSceneVectorSection* Section = Cast<UMovieSceneVectorSection>( MovieSceneHelpers::FindSectionAtTime( VectorTrack->GetAllSections(), FirstKeyTime ) );
@@ -275,7 +275,7 @@ bool FMatineeImportTools::CopyInterpColorTrack( UInterpTrackColorProp* ColorProp
 
 	ColorTrack->Modify();
 
-	FFrameRate   FrameRate    = ColorTrack->GetTypedOuter<UMovieScene>()->GetFrameResolution();
+	FFrameRate   FrameRate    = ColorTrack->GetTypedOuter<UMovieScene>()->GetTickResolution();
 	FFrameNumber FirstKeyTime = (ColorPropTrack->GetKeyframeTime( 0 ) * FrameRate).RoundToFrame();
 
 	UMovieSceneColorSection* Section = Cast<UMovieSceneColorSection>( MovieSceneHelpers::FindSectionAtTime( ColorTrack->GetAllSections(), FirstKeyTime ) );
@@ -333,7 +333,7 @@ bool FMatineeImportTools::CopyInterpLinearColorTrack( UInterpTrackLinearColorPro
 
 	ColorTrack->Modify();
 
-	FFrameRate   FrameRate    = ColorTrack->GetTypedOuter<UMovieScene>()->GetFrameResolution();
+	FFrameRate   FrameRate    = ColorTrack->GetTypedOuter<UMovieScene>()->GetTickResolution();
 	FFrameNumber FirstKeyTime = (LinearColorPropTrack->GetKeyframeTime( 0 ) * FrameRate).RoundToFrame();
 
 	UMovieSceneColorSection* Section = Cast<UMovieSceneColorSection>( MovieSceneHelpers::FindSectionAtTime( ColorTrack->GetAllSections(), FirstKeyTime ) );
@@ -392,7 +392,7 @@ bool FMatineeImportTools::CopyInterpMoveTrack( UInterpTrackMove* MoveTrack, UMov
 
 	TransformTrack->Modify();
 
-	FFrameRate   FrameRate    = TransformTrack->GetTypedOuter<UMovieScene>()->GetFrameResolution();
+	FFrameRate   FrameRate    = TransformTrack->GetTypedOuter<UMovieScene>()->GetTickResolution();
 	FFrameNumber FirstKeyTime = (MoveTrack->GetKeyframeTime( 0 ) * FrameRate).RoundToFrame();
 
 	UMovieScene3DTransformSection* Section = Cast<UMovieScene3DTransformSection>( MovieSceneHelpers::FindSectionAtTime( TransformTrack->GetAllSections(), FirstKeyTime ) );
@@ -516,7 +516,7 @@ bool FMatineeImportTools::CopyInterpParticleTrack( UInterpTrackToggle* MatineeTo
 
 	ParticleTrack->Modify();
 
-	FFrameRate   FrameRate    = ParticleTrack->GetTypedOuter<UMovieScene>()->GetFrameResolution();
+	FFrameRate   FrameRate    = ParticleTrack->GetTypedOuter<UMovieScene>()->GetTickResolution();
 	FFrameNumber FirstKeyTime = (MatineeToggleTrack->GetKeyframeTime( 0 ) * FrameRate).RoundToFrame();
 
 	UMovieSceneParticleSection* Section = Cast<UMovieSceneParticleSection>( MovieSceneHelpers::FindSectionAtTime( ParticleTrack->GetAllSections(), FirstKeyTime ) );
@@ -564,7 +564,7 @@ bool FMatineeImportTools::CopyInterpAnimControlTrack( UInterpTrackAnimControl* M
 	const FScopedTransaction Transaction( NSLOCTEXT( "Sequencer", "PasteMatineeAnimTrack", "Paste Matinee Anim Track" ) );
 	bool bSectionCreated = false;
 
-	FFrameRate FrameRate = SkeletalAnimationTrack->GetTypedOuter<UMovieScene>()->GetFrameResolution();
+	FFrameRate FrameRate = SkeletalAnimationTrack->GetTypedOuter<UMovieScene>()->GetTickResolution();
 
 	SkeletalAnimationTrack->Modify();
 	SkeletalAnimationTrack->RemoveAllAnimationData();
@@ -617,7 +617,7 @@ bool FMatineeImportTools::CopyInterpSoundTrack( UInterpTrackSound* MatineeSoundT
 	const FScopedTransaction Transaction( NSLOCTEXT( "Sequencer", "PasteMatineeSoundTrack", "Paste Matinee Sound Track" ) );
 	bool bSectionCreated = false;
 
-	FFrameRate FrameRate = AudioTrack->GetTypedOuter<UMovieScene>()->GetFrameResolution();
+	FFrameRate FrameRate = AudioTrack->GetTypedOuter<UMovieScene>()->GetTickResolution();
 
 	AudioTrack->Modify();
 
@@ -652,7 +652,7 @@ bool FMatineeImportTools::CopyInterpFadeTrack( UInterpTrackFade* MatineeFadeTrac
 
 	FadeTrack->Modify();
 
-	FFrameRate   FrameRate    = FadeTrack->GetTypedOuter<UMovieScene>()->GetFrameResolution();
+	FFrameRate   FrameRate    = FadeTrack->GetTypedOuter<UMovieScene>()->GetTickResolution();
 	FFrameNumber FirstKeyTime = (MatineeFadeTrack->GetKeyframeTime( 0 ) * FrameRate).RoundToFrame();
 
 	UMovieSceneFadeSection* Section = Cast<UMovieSceneFadeSection>( MovieSceneHelpers::FindSectionAtTime( FadeTrack->GetAllSections(), FirstKeyTime ) );
@@ -695,7 +695,7 @@ bool FMatineeImportTools::CopyInterpDirectorTrack( UInterpTrackDirector* Directo
 	const FScopedTransaction Transaction( NSLOCTEXT( "Sequencer", "PasteMatineeDirectorTrack", "Paste Matinee Director Track" ) );
 	bool bCutsAdded = false;
 
-	FFrameRate FrameRate = CameraCutTrack->GetTypedOuter<UMovieScene>()->GetFrameResolution();
+	FFrameRate FrameRate = CameraCutTrack->GetTypedOuter<UMovieScene>()->GetTickResolution();
 
 	CameraCutTrack->Modify();
 	
@@ -735,7 +735,7 @@ bool FMatineeImportTools::CopyInterpEventTrack( UInterpTrackEvent* MatineeEventT
 
 	if (MatineeEventTrack->EventTrack.Num())
 	{
-		FFrameRate   FrameRate    = EventTrack->GetTypedOuter<UMovieScene>()->GetFrameResolution();
+		FFrameRate   FrameRate    = EventTrack->GetTypedOuter<UMovieScene>()->GetTickResolution();
 		FFrameNumber FirstKeyTime = (MatineeEventTrack->EventTrack[0].Time * FrameRate).RoundToFrame();
 
 		UMovieSceneEventSection* Section = Cast<UMovieSceneEventSection>( MovieSceneHelpers::FindSectionAtTime( EventTrack->GetAllSections(), FirstKeyTime ) );
@@ -781,7 +781,7 @@ bool FMatineeImportTools::CopyInterpVisibilityTrack( UInterpTrackVisibility* Mat
 
 	if (MatineeVisibilityTrack->VisibilityTrack.Num())
 	{
-		FFrameRate   FrameRate    = VisibilityTrack->GetTypedOuter<UMovieScene>()->GetFrameResolution();
+		FFrameRate   FrameRate    = VisibilityTrack->GetTypedOuter<UMovieScene>()->GetTickResolution();
 		FFrameNumber FirstKeyTime = (MatineeVisibilityTrack->GetKeyframeTime( 0 ) * FrameRate).RoundToFrame();
 
 		UMovieSceneBoolSection* Section = Cast<UMovieSceneBoolSection>( MovieSceneHelpers::FindSectionAtTime( VisibilityTrack->GetAllSections(), FirstKeyTime ) );
diff --git a/Engine/Source/Editor/MovieSceneTools/Private/MovieSceneSectionDetailsCustomization.cpp b/Engine/Source/Editor/MovieSceneTools/Private/MovieSceneSectionDetailsCustomization.cpp
--- a/Engine/Source/Editor/MovieSceneTools/Private/MovieSceneSectionDetailsCustomization.cpp
+++ b/Engine/Source/Editor/MovieSceneTools/Private/MovieSceneSectionDetailsCustomization.cpp
@@ -180,10 +180,10 @@ FText FMovieSceneSectionDetailsCustomization::OnGetRangeStartText() const
 void FMovieSceneSectionDetailsCustomization::OnRangeStartTextCommitted(const FText& InText, ETextCommit::Type CommitInfo)
 {
 	// Find the new value for the start range.
-	TOptional<double> FrameResolution = NumericTypeInterface->FromString(InText.ToString(), 0.0);
+	TOptional<double> NewStart = NumericTypeInterface->FromString(InText.ToString(), 0.0);
 
 	// Early out if we couldn't parse it, no need to reset them all to zero.
-	if (!FrameResolution.IsSet())
+	if (!NewStart.IsSet())
 	{
 		return;
 	}
@@ -200,7 +200,7 @@ void FMovieSceneSectionDetailsCustomization::OnRangeStartTextCommitted(const FTe
 		FMovieSceneFrameRange* MovieSceneFrameRange = (FMovieSceneFrameRange*)RawData[i];
 		if (MovieSceneFrameRange && !MovieSceneFrameRange->Value.GetLowerBound().IsOpen())
 		{
-			MovieSceneFrameRange->Value.SetLowerBoundValue(FFrameTime::FromDecimal(FrameResolution.GetValue()).RoundToFrame());
+			MovieSceneFrameRange->Value.SetLowerBoundValue(FFrameTime::FromDecimal(NewStart.GetValue()).RoundToFrame());
 		}
 	}
 
@@ -368,10 +368,10 @@ FText FMovieSceneSectionDetailsCustomization::OnGetRangeEndText() const
 void FMovieSceneSectionDetailsCustomization::OnRangeEndTextCommitted(const FText& InText, ETextCommit::Type CommitInfo)
 {
 	// Find the new value for the start range.
-	TOptional<double> FrameResolution = NumericTypeInterface->FromString(InText.ToString(), 0.0);
+	TOptional<double> NewEnd = NumericTypeInterface->FromString(InText.ToString(), 0.0);
 
 	// Early out if we couldn't parse it, no need to reset them all to zero.
-	if (!FrameResolution.IsSet())
+	if (!NewEnd.IsSet())
 	{
 		return;
 	}
@@ -388,7 +388,7 @@ void FMovieSceneSectionDetailsCustomization::OnRangeEndTextCommitted(const FText
 		FMovieSceneFrameRange* MovieSceneFrameRange = (FMovieSceneFrameRange*)RawData[i];
 		if (MovieSceneFrameRange && !MovieSceneFrameRange->Value.GetUpperBound().IsOpen())
 		{
-			MovieSceneFrameRange->Value.SetUpperBoundValue(FFrameTime::FromDecimal(FrameResolution.GetValue()).RoundToFrame());
+			MovieSceneFrameRange->Value.SetUpperBoundValue(FFrameTime::FromDecimal(NewEnd.GetValue()).RoundToFrame());
 		}
 	}
 
