Copying //UE4/Dev-Sequencer to //UE4/Dev-Main (Source: //UE4/Dev-Sequencer @ 4044339)

diff --git a/Engine/Source/Editor/MovieSceneTools/Private/MovieSceneToolHelpers.cpp b/Engine/Source/Editor/MovieSceneTools/Private/MovieSceneToolHelpers.cpp
--- a/Engine/Source/Editor/MovieSceneTools/Private/MovieSceneToolHelpers.cpp
+++ b/Engine/Source/Editor/MovieSceneTools/Private/MovieSceneToolHelpers.cpp
@@ -26,7 +26,9 @@
 #include "DesktopPlatformModule.h"
 #include "Framework/Notifications/NotificationManager.h"
 #include "Widgets/Notifications/SNotificationList.h"
-#include "MovieSceneCaptureHelpers.h"
+#include "MovieSceneTranslatorEDL.h"
+#include "MessageLogModule.h"
+#include "IMessageLogListing.h"
 #include "FbxImporter.h"
 #include "MatineeImportTools.h"
 #include "MovieSceneToolsProjectSettings.h"
@@ -51,7 +53,7 @@
 /* MovieSceneToolHelpers
  *****************************************************************************/
 
-void MovieSceneToolHelpers::TrimSection(const TSet<TWeakObjectPtr<UMovieSceneSection>>& Sections, FFrameNumber Time, bool bTrimLeft)
+void MovieSceneToolHelpers::TrimSection(const TSet<TWeakObjectPtr<UMovieSceneSection>>& Sections, FQualifiedFrameTime Time, bool bTrimLeft)
 {
 	for (auto Section : Sections)
 	{
@@ -63,7 +65,7 @@ void MovieSceneToolHelpers::TrimSection(const TSet<TWeakObjectPtr<UMovieSceneSec
 }
 
 
-void MovieSceneToolHelpers::SplitSection(const TSet<TWeakObjectPtr<UMovieSceneSection>>& Sections, FFrameNumber Time)
+void MovieSceneToolHelpers::SplitSection(const TSet<TWeakObjectPtr<UMovieSceneSection>>& Sections, FQualifiedFrameTime Time)
 {
 	for (auto Section : Sections)
 	{
@@ -487,7 +489,7 @@ int32 MovieSceneToolHelpers::FindAvailableRowIndex(UMovieSceneTrack* InTrack, UM
 		}
 	}
 
-	return InTrack->GetMaxRowIndex() + 1;
+	return InTrack->GetMaxRowIndex();
 }
 
 class SEnumCombobox : public SComboBox<TSharedPtr<int32>>
@@ -627,7 +629,7 @@ bool MovieSceneToolHelpers::ShowImportEDLDialog(UMovieScene* InMovieScene, FFram
 
 	const FScopedTransaction Transaction( NSLOCTEXT( "MovieSceneTools", "ImportEDLTransaction", "Import EDL" ) );
 
-	return MovieSceneCaptureHelpers::ImportEDL(InMovieScene, InFrameRate, OpenFilenames[0]);
+	return MovieSceneTranslatorEDL::ImportEDL(InMovieScene, InFrameRate, OpenFilenames[0]);
 }
 
 bool MovieSceneToolHelpers::ShowExportEDLDialog(const UMovieScene* InMovieScene, FFrameRate InFrameRate, FString InSaveDirectory, int32 InHandleFrames)
@@ -664,7 +666,7 @@ bool MovieSceneToolHelpers::ShowExportEDLDialog(const UMovieScene* InMovieScene,
 		return false;
 	}
 
-	if (MovieSceneCaptureHelpers::ExportEDL(InMovieScene, InFrameRate, SaveFilenames[0], InHandleFrames))
+	if (MovieSceneTranslatorEDL::ExportEDL(InMovieScene, InFrameRate, SaveFilenames[0], InHandleFrames))
 	{
 		const FString AbsoluteFilename = FPaths::ConvertRelativePathToFull(SaveFilenames[0]);
 		const FString SaveDirectory = FPaths::GetPath(AbsoluteFilename);
@@ -681,6 +683,131 @@ bool MovieSceneToolHelpers::ShowExportEDLDialog(const UMovieScene* InMovieScene,
 	return false;
 }
 
+bool MovieSceneToolHelpers::MovieSceneTranslatorImport(FMovieSceneImporter* InImporter, UMovieScene* InMovieScene, FFrameRate InFrameRate, FString InOpenDirectory)
+{
+	TArray<FString> OpenFilenames;
+	IDesktopPlatform* DesktopPlatform = FDesktopPlatformModule::Get();
+	bool bOpen = false;
+	if (DesktopPlatform)
+	{
+		FString FileTypeDescription = InImporter->GetFileTypeDescription().ToString();
+		FString DialogTitle = InImporter->GetDialogTitle().ToString();
+
+		bOpen = DesktopPlatform->OpenFileDialog(
+			FSlateApplication::Get().FindBestParentWindowHandleForDialogs(nullptr),
+			DialogTitle,
+			InOpenDirectory,
+			TEXT(""),
+			FileTypeDescription,
+			EFileDialogFlags::None,
+			OpenFilenames
+		);
+	}
+
+	if (!bOpen || !OpenFilenames.Num())
+	{
+		return false;
+	}
+
+	FScopedTransaction Transaction(InImporter->GetTransactionDescription());
+
+	TSharedRef<FMovieSceneTranslatorContext> ImportContext(new FMovieSceneTranslatorContext);
+	ImportContext->Init();
+
+	bool bSuccess = InImporter->Import(InMovieScene, InFrameRate, OpenFilenames[0], ImportContext);
+
+	// Display any messages in context
+	MovieSceneTranslatorDisplayMessages(InImporter, ImportContext);
+
+	// Roll back transaction when import fails.
+	if (!bSuccess)
+	{
+		Transaction.Cancel();
+	}
+
+	return bSuccess;
+}
+
+bool MovieSceneToolHelpers::MovieSceneTranslatorExport(FMovieSceneExporter* InExporter, const UMovieScene* InMovieScene, FFrameRate InFrameRate, FString InSaveDirectory, int32 InHandleFrames)
+{
+	if (InExporter == nullptr || InMovieScene == nullptr)
+	{
+		return false;
+	}
+
+	TArray<FString> SaveFilenames;
+	FString SequenceName = InMovieScene->GetOuter()->GetName();
+
+	// Pop open a dialog to request the location of the edl
+	IDesktopPlatform* DesktopPlatform = FDesktopPlatformModule::Get();
+	bool bSave = false;
+	if (DesktopPlatform)
+	{
+		FString FileTypeDescription = InExporter->GetFileTypeDescription().ToString();
+		FString DialogTitle = InExporter->GetDialogTitle().ToString();
+		FString FileExtension = InExporter->GetDefaultFileExtension().ToString();
+
+		bSave = DesktopPlatform->SaveFileDialog(
+			FSlateApplication::Get().FindBestParentWindowHandleForDialogs(nullptr),
+			DialogTitle,
+			InSaveDirectory,
+			SequenceName + TEXT(".") + FileExtension,
+			FileTypeDescription,
+			EFileDialogFlags::None,
+			SaveFilenames
+		);
+	}
+
+	if (!bSave || !SaveFilenames.Num())
+	{
+		return false;
+	}
+
+	TSharedRef<FMovieSceneTranslatorContext> ExportContext(new FMovieSceneTranslatorContext);
+	ExportContext->Init();
+
+	bool bSuccess = InExporter->Export(InMovieScene, InFrameRate, SaveFilenames[0], InHandleFrames, ExportContext);
+	
+	// Display any messages in context
+	MovieSceneTranslatorDisplayMessages(InExporter, ExportContext);
+
+	if (bSuccess)
+	{
+		const FString AbsoluteFilename = FPaths::ConvertRelativePathToFull(SaveFilenames[0]);
+		const FString SaveDirectory = FPaths::GetPath(AbsoluteFilename);
+
+		FNotificationInfo NotificationInfo(InExporter->GetNotificationExportFinished());
+		NotificationInfo.ExpireDuration = 5.f;
+		NotificationInfo.Hyperlink = FSimpleDelegate::CreateStatic([](FString InDirectory) { FPlatformProcess::ExploreFolder(*InDirectory); }, SaveDirectory);
+		NotificationInfo.HyperlinkText = InExporter->GetNotificationHyperlinkText();
+		FSlateNotificationManager::Get().AddNotification(NotificationInfo);
+	}
+
+	return bSuccess;
+}
+
+void MovieSceneToolHelpers::MovieSceneTranslatorDisplayMessages(FMovieSceneTranslator *InTranslator, TSharedRef<FMovieSceneTranslatorContext> InContext)
+{
+	if (InTranslator == nullptr || InContext->GetMessages().Num() == 0)
+	{
+		return;
+	}
+
+	// Clear any old messages after an import or export
+	const FName LogTitle = InTranslator->GetMessageLogWindowTitle();
+	FMessageLogModule& MessageLogModule = FModuleManager::LoadModuleChecked<FMessageLogModule>("MessageLog");
+	TSharedRef<IMessageLogListing> LogListing = MessageLogModule.GetLogListing(LogTitle);
+	LogListing->SetLabel(InTranslator->GetMessageLogLabel());
+	LogListing->ClearMessages();
+
+	for (TSharedRef<FTokenizedMessage> Message : InContext->GetMessages())
+	{
+		LogListing->AddMessage(Message);
+	}
+
+	MessageLogModule.OpenMessageLog(LogTitle);
+}
+
 bool ImportFBXProperty(FString NodeName, FString AnimatedPropertyName, FGuid ObjectBinding, UnFbx::FFbxCurvesAPI& CurveAPI, UMovieScene* InMovieScene, ISequencer& InSequencer)
 {
 	const UMovieSceneToolsProjectSettings* ProjectSettings = GetDefault<UMovieSceneToolsProjectSettings>();
@@ -732,7 +859,7 @@ bool ImportFBXProperty(FString NodeName, FString AnimatedPropertyName, FGuid Obj
 
 			if (FloatTrack)
 			{
-				FFrameRate FrameRate = FloatTrack->GetTypedOuter<UMovieScene>()->GetFrameResolution();
+				FFrameRate FrameRate = FloatTrack->GetTypedOuter<UMovieScene>()->GetTickResolution();
 
 				bool bSectionAdded = false;
 				UMovieSceneFloatSection* FloatSection = Cast<UMovieSceneFloatSection>(FloatTrack->FindOrAddSection(0, bSectionAdded));
@@ -741,6 +868,8 @@ bool ImportFBXProperty(FString NodeName, FString AnimatedPropertyName, FGuid Obj
 					continue;
 				}
 
+				FloatSection->Modify();
+
 				if (bSectionAdded)
 				{
 					FloatSection->SetRange(TRange<FFrameNumber>::All());
@@ -756,18 +885,19 @@ bool ImportFBXProperty(FString NodeName, FString AnimatedPropertyName, FGuid Obj
 				TMovieSceneChannel<FMovieSceneFloatValue> ChannelInterface = Channel->GetInterface();
 
 				ChannelInterface.Reset();
+				double DecimalRate = FrameRate.AsDecimal();
 				for (int32 KeyIndex = 0; KeyIndex < CurveHandle.Points.Num(); ++KeyIndex)
 				{
 					float ArriveTangent = CurveHandle.Points[KeyIndex].ArriveTangent;
 					if (KeyIndex > 0)
 					{
-						ArriveTangent = ArriveTangent / (CurveHandle.Points[KeyIndex].InVal - CurveHandle.Points[KeyIndex-1].InVal);
+						ArriveTangent = ArriveTangent / ((CurveHandle.Points[KeyIndex].InVal - CurveHandle.Points[KeyIndex - 1].InVal) * DecimalRate);
 					}
 					
 					float LeaveTangent = CurveHandle.Points[KeyIndex].LeaveTangent;
 					if (KeyIndex < CurveHandle.Points.Num() - 1)
 					{
-						LeaveTangent = LeaveTangent / (CurveHandle.Points[KeyIndex+1].InVal - CurveHandle.Points[KeyIndex].InVal);
+						LeaveTangent = LeaveTangent / ((CurveHandle.Points[KeyIndex+1].InVal - CurveHandle.Points[KeyIndex].InVal) * DecimalRate);
 					}
 
 					FFrameNumber KeyTime = (CurveHandle.Points[KeyIndex].InVal * FrameRate).RoundToFrame();
@@ -777,9 +907,12 @@ bool ImportFBXProperty(FString NodeName, FString AnimatedPropertyName, FGuid Obj
 				if (ImportFBXSettings->bReduceKeys)
 				{
 					FKeyDataOptimizationParams Params;
-					Params.bAutoSetInterpolation = true;
-					MovieScene::Optimize(Channel, Params);
+					Params.Tolerance = ImportFBXSettings->ReduceKeysTolerance;
+					using namespace MovieScene;
+					Optimize(Channel, Params);
 				}
+				Channel->AutoSetTangents();
+
 
 				TArrayView<const FFrameNumber> AllTimes = ChannelInterface.GetTimes();
 				if (AllTimes.Num() > 0)
@@ -797,19 +930,19 @@ void ImportTransformChannel(const FInterpCurveFloat& Source, FMovieSceneFloatCha
 {
 	TMovieSceneChannel<FMovieSceneFloatValue> ChannelInterface = Dest->GetInterface();
 	ChannelInterface.Reset();
-
+	double DecimalRate = DestFrameRate.AsDecimal();
 	for (int32 KeyIndex = 0; KeyIndex < Source.Points.Num(); ++KeyIndex)
 	{
 		float ArriveTangent = Source.Points[KeyIndex].ArriveTangent;
 		if (KeyIndex > 0)
 		{
-			ArriveTangent = ArriveTangent / (Source.Points[KeyIndex].InVal - Source.Points[KeyIndex-1].InVal);
+			ArriveTangent = ArriveTangent / ((Source.Points[KeyIndex].InVal - Source.Points[KeyIndex-1].InVal) * DecimalRate);
 		}
 		
 		float LeaveTangent = Source.Points[KeyIndex].LeaveTangent;
 		if (KeyIndex < Source.Points.Num() - 1)
 		{
-			LeaveTangent = LeaveTangent / (Source.Points[KeyIndex+1].InVal - Source.Points[KeyIndex].InVal);
+			LeaveTangent = LeaveTangent / ((Source.Points[KeyIndex+1].InVal - Source.Points[KeyIndex].InVal) * DecimalRate);
 		}
 
 		if (bNegateTangents)
@@ -821,10 +954,15 @@ void ImportTransformChannel(const FInterpCurveFloat& Source, FMovieSceneFloatCha
 		FFrameNumber KeyTime = (Source.Points[KeyIndex].InVal * DestFrameRate).RoundToFrame();
 		FMatineeImportTools::SetOrAddKey(ChannelInterface, KeyTime, Source.Points[KeyIndex].OutVal, ArriveTangent, LeaveTangent, Source.Points[KeyIndex].InterpMode);
 	}
-
-	FKeyDataOptimizationParams Params;
-	Params.bAutoSetInterpolation = true;
-	MovieScene::Optimize(Dest, Params);
+	const UMovieSceneUserImportFBXSettings* ImportFBXSettings = GetDefault<UMovieSceneUserImportFBXSettings>();
+	if (ImportFBXSettings->bReduceKeys)
+	{
+		FKeyDataOptimizationParams Params;
+		Params.Tolerance = ImportFBXSettings->ReduceKeysTolerance;
+		using namespace MovieScene;
+		Optimize(Dest, Params);
+	}
+	Dest->AutoSetTangents();
 }
 
 bool ImportFBXTransform(FString NodeName, FGuid ObjectBinding, UnFbx::FFbxCurvesAPI& CurveAPI, UMovieScene* InMovieScene)
@@ -852,7 +990,10 @@ bool ImportFBXTransform(FString NodeName, FGuid ObjectBinding, UnFbx::FFbxCurves
 		return false;
 	}
 
-	FFrameRate FrameRate = TransformSection->GetTypedOuter<UMovieScene>()->GetFrameResolution();
+	TransformSection->Modify();
+
+	FFrameRate FrameRate = TransformSection->GetTypedOuter<UMovieScene>()->GetTickResolution();
+
 
 	if (bSectionAdded)
 	{
@@ -887,17 +1028,11 @@ bool ImportFBXTransform(FString NodeName, FGuid ObjectBinding, UnFbx::FFbxCurves
 	ImportTransformChannel(Scale[1],         Channels[7], FrameRate, false);
 	ImportTransformChannel(Scale[2],         Channels[8], FrameRate, false);
 
-
-	FKeyDataOptimizationParams Params;
-	Params.bAutoSetInterpolation = true;
-
 	TOptional<TRangeBound<FFrameNumber>> MinLower;
 	TOptional<TRangeBound<FFrameNumber>> MaxUpper;
 
 	for (int32 ChannelIndex = 0; ChannelIndex < 9; ++ChannelIndex)
 	{
-		MovieScene::Optimize(Channels[ChannelIndex], Params);
-
 		TArrayView<const FFrameNumber> AllTimes = Channels[ChannelIndex]->GetTimes();
 		if (AllTimes.Num() > 0)
 		{
@@ -1009,11 +1144,49 @@ FbxNode* RetrieveObjectFromName(const TCHAR* ObjectName, FbxNode* Root)
 	return nullptr;
 }
 
+void CopyCameraProperties(FbxCamera* CameraNode, ACineCameraActor* CameraActor)
+{
+	float FieldOfView;
+	float FocalLength;
+
+	if (CameraNode->GetApertureMode() == FbxCamera::eFocalLength)
+	{
+		FocalLength = CameraNode->FocalLength.Get();
+		FieldOfView = CameraNode->ComputeFieldOfView(FocalLength);
+	}
+	else
+	{
+		FieldOfView = CameraNode->FieldOfView.Get();
+		FocalLength = CameraNode->ComputeFocalLength(FieldOfView);
+	}
+
+	float ApertureWidth = CameraNode->GetApertureWidth();
+	float ApertureHeight = CameraNode->GetApertureHeight();
+
+	UCineCameraComponent* CineCameraComponent = CameraActor->GetCineCameraComponent();
 
-void ImportFBXCamera(UnFbx::FFbxImporter* FbxImporter, UMovieScene* InMovieScene, ISequencer& InSequencer, TMap<FGuid, FString>& InObjectBindingMap)
+	CineCameraComponent->SetProjectionMode(CameraNode->ProjectionType.Get() == FbxCamera::ePerspective ? ECameraProjectionMode::Perspective : ECameraProjectionMode::Orthographic);
+	CineCameraComponent->SetAspectRatio(CameraNode->AspectWidth.Get() / CameraNode->AspectHeight.Get());
+	CineCameraComponent->SetOrthoNearClipPlane(CameraNode->NearPlane.Get());
+	CineCameraComponent->SetOrthoFarClipPlane(CameraNode->FarPlane.Get());
+	CineCameraComponent->SetOrthoWidth(CameraNode->OrthoZoom.Get());
+	CineCameraComponent->SetFieldOfView(FieldOfView);
+	CineCameraComponent->FilmbackSettings.SensorWidth = FUnitConversion::Convert(ApertureWidth, EUnit::Inches, EUnit::Millimeters);
+	CineCameraComponent->FilmbackSettings.SensorHeight = FUnitConversion::Convert(ApertureHeight, EUnit::Inches, EUnit::Millimeters);
+	if (FocalLength < CineCameraComponent->LensSettings.MinFocalLength)
+	{
+		CineCameraComponent->LensSettings.MinFocalLength = FocalLength;
+	}
+	if (FocalLength > CineCameraComponent->LensSettings.MaxFocalLength)
+	{
+		CineCameraComponent->LensSettings.MaxFocalLength = FocalLength;
+	}
+	CineCameraComponent->CurrentFocalLength = FocalLength;
+}
+
+void ImportFBXCamera(UnFbx::FFbxImporter* FbxImporter, UMovieScene* InMovieScene, ISequencer& InSequencer, TMap<FGuid, FString>& InObjectBindingMap, bool bCreateCameras)
 {
-	const UMovieSceneUserImportFBXSettings* ImportFBXSettings = GetDefault<UMovieSceneUserImportFBXSettings>();
-	if (ImportFBXSettings->bCreateCameras)
+	if (bCreateCameras)
 	{
 		TArray<FbxCamera*> AllCameras;
 		GetCameras(FbxImporter->Scene->GetRootNode(), AllCameras);
@@ -1059,12 +1232,18 @@ void ImportFBXCamera(UnFbx::FFbxImporter* FbxImporter, UMovieScene* InMovieScene
 			ACineCameraActor* NewCamera = World->SpawnActor<ACineCameraActor>(SpawnParams);
 			NewCamera->SetActorLabel(*CameraName);
 
+			// Copy camera properties before adding default tracks so that initial camera properties match and can be restored after sequencer finishes
+			CopyCameraProperties(UnmatchedCamera, NewCamera);
+
 			TArray<TWeakObjectPtr<AActor> > NewCameras;
 			NewCameras.Add(NewCamera);
 			TArray<FGuid> NewCameraGuids = InSequencer.AddActors(NewCameras);
 
-			InObjectBindingMap.Add(NewCameraGuids[0]);
-			InObjectBindingMap[NewCameraGuids[0]] = CameraName;
+			if (NewCameraGuids.Num())
+			{
+				InObjectBindingMap.Add(NewCameraGuids[0]);
+				InObjectBindingMap[NewCameraGuids[0]] = CameraName;
+			}
 		}
 	}
 
@@ -1114,9 +1293,6 @@ void ImportFBXCamera(UnFbx::FFbxImporter* FbxImporter, UMovieScene* InMovieScene
 			FocalLength = CameraNode->ComputeFocalLength(FieldOfView);
 		}
 
-		float ApertureWidth = CameraNode->GetApertureWidth();
-		float ApertureHeight = CameraNode->GetApertureHeight();
-
 		for (TWeakObjectPtr<>& WeakObject : BoundObjects)
 		{
 			UObject* FoundObject = WeakObject.Get();
@@ -1124,23 +1300,7 @@ void ImportFBXCamera(UnFbx::FFbxImporter* FbxImporter, UMovieScene* InMovieScene
 			{
 				ACineCameraActor* CineCameraActor = Cast<ACineCameraActor>(FoundObject);
 				UCineCameraComponent* CineCameraComponent = CineCameraActor->GetCineCameraComponent();
-
-				CineCameraComponent->SetProjectionMode(CameraNode->ProjectionType.Get() == FbxCamera::ePerspective ? ECameraProjectionMode::Perspective : ECameraProjectionMode::Orthographic);
-				CineCameraComponent->SetAspectRatio(CameraNode->AspectWidth.Get() / CameraNode->AspectHeight.Get());
-				CineCameraComponent->SetOrthoNearClipPlane(CameraNode->NearPlane.Get());
-				CineCameraComponent->SetOrthoFarClipPlane(CameraNode->FarPlane.Get());
-				CineCameraComponent->SetOrthoWidth(CameraNode->OrthoZoom.Get());
-				CineCameraComponent->SetFieldOfView(FieldOfView);
-				CineCameraComponent->FilmbackSettings.SensorWidth = FUnitConversion::Convert(ApertureWidth, EUnit::Inches, EUnit::Millimeters);
-				CineCameraComponent->FilmbackSettings.SensorHeight = FUnitConversion::Convert(ApertureHeight, EUnit::Inches, EUnit::Millimeters);
-				if (FocalLength < CineCameraComponent->LensSettings.MinFocalLength)
-				{
-					CineCameraComponent->LensSettings.MinFocalLength = FocalLength;
-				}
-				if (FocalLength > CineCameraComponent->LensSettings.MaxFocalLength)
-				{
-					CineCameraComponent->LensSettings.MaxFocalLength = FocalLength;
-				}
+				CopyCameraProperties(CameraNode, CineCameraActor);
 
 				// Set the default value of the current focal length section
 				FGuid PropertyOwnerGuid = InSequencer.GetHandleToObject(CineCameraComponent);
@@ -1159,6 +1319,8 @@ void ImportFBXCamera(UnFbx::FFbxImporter* FbxImporter, UMovieScene* InMovieScene
 						continue;
 					}
 
+					FloatSection->Modify();
+
 					if (bSectionAdded)
 					{
 						FloatSection->SetRange(TRange<FFrameNumber>::All());
@@ -1226,23 +1388,26 @@ void ImportCameraCut(UnFbx::FFbxImporter* FbxImporter, UMovieScene* InMovieScene
 	GetCameras(FbxImporter->Scene->GetRootNode(), AllCameras);
 
 	UMovieSceneCameraCutTrack* CameraCutTrack = GetCameraCutTrack(InMovieScene);
-	FFrameRate FrameRate = CameraCutTrack->GetTypedOuter<UMovieScene>()->GetFrameResolution();
+	FFrameRate FrameRate = CameraCutTrack->GetTypedOuter<UMovieScene>()->GetTickResolution();
 
 	FbxAnimCurve* AnimCurve = CameraSwitcher->CameraIndex.GetCurve(AnimLayer);
-	for (int i = 0; i < AnimCurve->KeyGetCount(); ++i)
+	if (AnimCurve)
 	{
-		FbxAnimCurveKey key = AnimCurve->KeyGet(i);
-		int value = (int)key.GetValue() - 1;
-		if (value >= 0 && value < AllCameras.Num())
+		for (int i = 0; i < AnimCurve->KeyGetCount(); ++i)
 		{
-			FGuid CameraGuid = FindCameraGuid(AllCameras[value], InObjectBindingMap);
-			if (CameraGuid != FGuid())
+			FbxAnimCurveKey key = AnimCurve->KeyGet(i);
+			int value = (int)key.GetValue() - 1;
+			if (value >= 0 && value < AllCameras.Num())
 			{
-				CameraCutTrack->AddNewCameraCut(FMovieSceneObjectBindingID(CameraGuid, MovieSceneSequenceID::Root), (((float)key.GetTime().GetSecondDouble()) * FrameRate).FloorToFrame());
+				FGuid CameraGuid = FindCameraGuid(AllCameras[value], InObjectBindingMap);
+				if (CameraGuid != FGuid())
+				{
+					CameraCutTrack->AddNewCameraCut(FMovieSceneObjectBindingID(CameraGuid, MovieSceneSequenceID::Root), (((float)key.GetTime().GetSecondDouble()) * FrameRate).FloorToFrame());
+				}
 			}
 		}
+		InSequencer.NotifyMovieSceneDataChanged(EMovieSceneDataChangeType::MovieSceneStructureItemAdded);
 	}
-	InSequencer.NotifyMovieSceneDataChanged(EMovieSceneDataChangeType::MovieSceneStructureItemAdded);
 }
 
 class SMovieSceneImportFBXSettings : public SCompoundWidget, public FGCObject
@@ -1342,8 +1507,9 @@ private:
 
 		const FScopedTransaction Transaction( NSLOCTEXT( "MovieSceneTools", "ImportFBXTransaction", "Import FBX" ) );
 
-		// Import static cameras first
-		ImportFBXCamera(FbxImporter, MovieScene, *Sequencer, ObjectBindingMap);
+		// Import static cameras first if not mapping to any particular object binding nodes
+		bool bCreateCameras = ImportFBXSettings->bCreateCameras && ObjectBindingMap.Num() == 0;
+		ImportFBXCamera(FbxImporter, MovieScene, *Sequencer, ObjectBindingMap, bCreateCameras);
 
 		UnFbx::FFbxCurvesAPI CurveAPI;
 		FbxImporter->PopulateAnimatedCurveData(CurveAPI);
@@ -1538,19 +1704,37 @@ UObject* MovieSceneToolHelpers::ExportToCameraAnim(UMovieScene* InMovieScene, FG
 			}
 			else
 			{
-				FFrameRate FrameResolution = InMovieScene->GetFrameResolution();
+				FFrameRate TickResolution = InMovieScene->GetTickResolution();
 				UMovieScene3DTransformSection* TransformSection = Cast<UMovieScene3DTransformSection>(Sections[0]);
 				TArrayView<FMovieSceneFloatChannel*> FloatChannels = TransformSection->GetChannelProxy().GetChannels<FMovieSceneFloatChannel>();
 
-				CopyKeyDataToMoveAxis(FloatChannels[0]->GetInterface(), MoveAxies[AXIS_TranslationX], FrameResolution);
-				CopyKeyDataToMoveAxis(FloatChannels[1]->GetInterface(), MoveAxies[AXIS_TranslationY], FrameResolution);
-				CopyKeyDataToMoveAxis(FloatChannels[2]->GetInterface(), MoveAxies[AXIS_TranslationZ], FrameResolution);
-				CopyKeyDataToMoveAxis(FloatChannels[3]->GetInterface(), MoveAxies[AXIS_RotationX],    FrameResolution);
-				CopyKeyDataToMoveAxis(FloatChannels[4]->GetInterface(), MoveAxies[AXIS_RotationY],    FrameResolution);
-				CopyKeyDataToMoveAxis(FloatChannels[5]->GetInterface(), MoveAxies[AXIS_RotationZ],    FrameResolution);
+				CopyKeyDataToMoveAxis(FloatChannels[0]->GetInterface(), MoveAxies[AXIS_TranslationX], TickResolution);
+				CopyKeyDataToMoveAxis(FloatChannels[1]->GetInterface(), MoveAxies[AXIS_TranslationY], TickResolution);
+				CopyKeyDataToMoveAxis(FloatChannels[2]->GetInterface(), MoveAxies[AXIS_TranslationZ], TickResolution);
+				CopyKeyDataToMoveAxis(FloatChannels[3]->GetInterface(), MoveAxies[AXIS_RotationX],    TickResolution);
+				CopyKeyDataToMoveAxis(FloatChannels[4]->GetInterface(), MoveAxies[AXIS_RotationY],    TickResolution);
+				CopyKeyDataToMoveAxis(FloatChannels[5]->GetInterface(), MoveAxies[AXIS_RotationZ],    TickResolution);
 			}
 		}
 	}
 
 	return NewAsset;
 }
+
+
+bool MovieSceneToolHelpers::HasHiddenMobility(const UClass* ObjectClass)
+{
+	if (ObjectClass)
+	{
+		static const FName NAME_HideCategories(TEXT("HideCategories"));
+		if (ObjectClass->HasMetaData(NAME_HideCategories))
+		{
+			if (ObjectClass->GetMetaData(NAME_HideCategories).Contains(TEXT("Mobility")))
+			{
+				return true;
+			}
+		}
+	}
+
+	return false;
+}
diff --git a/Engine/Source/Editor/MovieSceneTools/Private/MovieSceneTranslator.cpp b/Engine/Source/Editor/MovieSceneTools/Private/MovieSceneTranslator.cpp
new file mode 100644
--- /dev/null
+++ b/Engine/Source/Editor/MovieSceneTools/Private/MovieSceneTranslator.cpp
@@ -0,0 +1,676 @@
+// Copyright 1998-2018 Epic Games, Inc. All Rights Reserved.
+
+#include "MovieSceneTranslator.h"
+#include "MovieScene.h"
+#include "Misc/FileHelper.h"
+#include "Misc/Paths.h"
+#include "Modules/ModuleManager.h"
+#include "AssetData.h"
+#include "LevelSequence.h"
+#include "Tracks/MovieSceneAudioTrack.h"
+#include "Sections/MovieSceneAudioSection.h"
+#include "Tracks/MovieSceneCinematicShotTrack.h"
+#include "Sections/MovieSceneCinematicShotSection.h"
+#include "MovieSceneTimeHelpers.h"
+#include "AssetRegistryModule.h"
+
+#define LOCTEXT_NAMESPACE "MovieSceneTranslator"
+
+FMovieSceneExportData::FMovieSceneExportData(const UMovieScene* InMovieScene, FFrameRate InFrameRate, int32 InHandleFrames, FString InSaveFilename, TSharedPtr<FMovieSceneTranslatorContext> InContext)
+{
+	if (InMovieScene == nullptr)
+	{
+		bExportDataIsValid = false;
+		return;
+	}
+
+	ExportContext = InContext;
+	if (ExportContext.IsValid())
+	{
+		ExportContext->Init();
+	}
+
+	FrameRate = InFrameRate;
+	HandleFrames = ConvertFrameTime(FFrameTime(InHandleFrames), InMovieScene->GetTickResolution(), FrameRate);
+	SaveFilename = InSaveFilename;
+
+	bExportDataIsValid = ConstructData(InMovieScene);
+}
+
+FMovieSceneExportData::FMovieSceneExportData()
+{
+	bExportDataIsValid = false;
+}
+
+FMovieSceneExportData::~FMovieSceneExportData()
+{
+}
+
+bool FMovieSceneExportData::IsExportDataValid() const
+{
+	return bExportDataIsValid;
+}
+
+bool FMovieSceneExportData::ConstructData(const UMovieScene* InMovieScene)
+{
+	if (InMovieScene == nullptr) 
+	{ 
+		return false; 
+	}
+
+	SaveFilenamePath = FPaths::GetPath(SaveFilename);
+	if (FPaths::IsRelative(SaveFilenamePath))
+	{
+		SaveFilenamePath = FPaths::ConvertRelativePathToFull(SaveFilenamePath);
+	}
+
+	return ConstructMovieSceneData(InMovieScene);
+}
+
+bool FMovieSceneExportData::ConstructMovieSceneData(const UMovieScene* InMovieScene)
+{
+	if (InMovieScene == nullptr)
+	{
+		return false;
+	}
+
+	MovieSceneData = MakeShared<FMovieSceneExportMovieSceneData>();
+
+	FFrameRate Resolution = InMovieScene->GetTickResolution();
+
+	TRange<FFrameNumber> PlaybackRange = InMovieScene->GetPlaybackRange();
+
+	if (PlaybackRange.HasLowerBound())
+	{
+		MovieSceneData->PlaybackRangeStartFrame = ConvertFrameTime(PlaybackRange.GetLowerBoundValue(), Resolution, FrameRate).CeilToFrame();
+	}
+	else
+	{
+		UE_LOG(LogMovieScene, Error, TEXT("Invalid condition: Movie scene playback range has infinite lower bound."));
+		return false;
+	}
+
+	if (PlaybackRange.HasUpperBound())
+	{
+		MovieSceneData->PlaybackRangeEndFrame = ConvertFrameTime(PlaybackRange.GetUpperBoundValue(), Resolution, FrameRate).CeilToFrame();
+	}
+	else
+	{
+		UE_LOG(LogMovieScene, Error, TEXT("Invalid condition: Movie scene playback range has infinite upper bound."));
+		return false;
+	}
+
+	MovieSceneData->Name = InMovieScene->GetOuter()->GetName();
+	MovieSceneData->Path = InMovieScene->GetOuter()->GetPathName();
+	MovieSceneData->Resolution = Resolution;
+	MovieSceneData->Duration = ConvertFrameTime(MovieScene::DiscreteSize(PlaybackRange), Resolution, FrameRate).FrameNumber.Value;
+
+	bool bFoundCinematicMasterTrack = false;
+
+	const TArray<UMovieSceneTrack*> MasterTracks = InMovieScene->GetMasterTracks();
+	for (UMovieSceneTrack* MasterTrack : MasterTracks)
+	{
+		if (!bFoundCinematicMasterTrack && MasterTrack->IsA(UMovieSceneCinematicShotTrack::StaticClass()))
+		{
+			const UMovieSceneCinematicShotTrack* CinematicTrack = Cast<UMovieSceneCinematicShotTrack>(MasterTrack);
+			if (!ConstructCinematicMasterTrackData(InMovieScene, CinematicTrack))
+			{ 
+				return false; 
+			}
+			bFoundCinematicMasterTrack = true;
+		}
+		else if (MasterTrack->IsA(UMovieSceneAudioTrack::StaticClass()))
+		{
+			const UMovieSceneAudioTrack* AudioTrack = Cast<UMovieSceneAudioTrack>(MasterTrack);
+			if (!ConstructAudioTrackData(InMovieScene, MovieSceneData, AudioTrack))
+			{ 
+				return false; 
+			}
+		}
+	}
+
+	return true;
+}
+
+bool FMovieSceneExportData::ConstructCinematicMasterTrackData(const UMovieScene* InMovieScene, const UMovieSceneCinematicShotTrack* InCinematicMasterTrack)
+{
+	if (InMovieScene == nullptr || !MovieSceneData.IsValid())
+	{
+		return false;
+	}
+
+	TSharedPtr<FMovieSceneExportCinematicMasterTrackData> MasterTrackData = MakeShared<FMovieSceneExportCinematicMasterTrackData>();
+	MovieSceneData->CinematicMasterTrack = MasterTrackData;
+
+	// Construct sections & create track row index array
+	TArray<int32> CinematicTrackRowIndices;
+
+	for (UMovieSceneSection* Section : InCinematicMasterTrack->GetAllSections())
+	{
+		const UMovieSceneCinematicShotSection* CinematicSection = Cast<UMovieSceneCinematicShotSection>(Section);
+
+		if (CinematicSection != nullptr && CinematicSection->GetSequence() != nullptr)
+		{
+			if (!ConstructCinematicSectionData(InMovieScene, MasterTrackData, CinematicSection))
+			{
+				return false;
+			}
+
+			int32 RowIndex = CinematicSection->GetRowIndex();
+			if (RowIndex >= 0)
+			{
+				CinematicTrackRowIndices.AddUnique(RowIndex);
+			}
+		}
+	}
+
+	// Construct tracks and point to sections
+	CinematicTrackRowIndices.Sort();
+
+	for (int32 CinematicTrackRowIndex : CinematicTrackRowIndices)
+	{	
+		if (!ConstructCinematicTrackData(InMovieScene, MasterTrackData, CinematicTrackRowIndex))
+		{
+			return false;
+		}		
+	}
+
+	return true;
+}
+
+bool FMovieSceneExportData::ConstructCinematicTrackData(const UMovieScene* InMovieScene, TSharedPtr<FMovieSceneExportCinematicMasterTrackData> InCinematicMasterTrackData, int32 RowIndex)
+{
+	if (InMovieScene == nullptr || !InCinematicMasterTrackData.IsValid() || !MovieSceneData.IsValid() || !MovieSceneData->CinematicMasterTrack.IsValid())
+	{
+		return false;
+	}
+
+	TSharedPtr<FMovieSceneExportCinematicTrackData> TrackData = MakeShared<FMovieSceneExportCinematicTrackData>(RowIndex);
+	MovieSceneData->CinematicMasterTrack->CinematicTracks.Add(TrackData);
+
+	for (TSharedPtr<FMovieSceneExportCinematicSectionData> Section : InCinematicMasterTrackData->CinematicSections)
+	{
+		if (Section.IsValid() && Section->RowIndex == RowIndex)
+		{
+			TrackData->CinematicSections.Add(Section);
+		}
+	}
+
+	return true;
+}
+
+bool FMovieSceneExportData::ConstructAudioTrackData(const UMovieScene* InMovieScene, TSharedPtr<FMovieSceneExportMovieSceneData> InMovieSceneData, const UMovieSceneAudioTrack* InAudioTrack)
+{
+	if (InMovieScene == nullptr || InAudioTrack == nullptr || !InMovieSceneData.IsValid())
+	{
+		return false;
+	}
+
+	TSharedPtr<FMovieSceneExportAudioTrackData> TrackData = MakeShared<FMovieSceneExportAudioTrackData>();
+	InMovieSceneData->AudioTracks.Add(TrackData);
+
+	for (UMovieSceneSection* ShotSection : InAudioTrack->GetAllSections())
+	{
+		const UMovieSceneAudioSection* AudioSection = Cast<UMovieSceneAudioSection>(ShotSection);
+
+		if (!ConstructAudioSectionData(InMovieScene, TrackData, AudioSection))
+		{
+			return false;
+		}
+	}
+
+	return true;
+}
+
+bool FMovieSceneExportData::ConstructCinematicSectionData(const UMovieScene* InMovieScene, TSharedPtr<FMovieSceneExportCinematicMasterTrackData> InMasterTrackData, const UMovieSceneCinematicShotSection* InCinematicSection)
+{
+	if (InMovieScene == nullptr || !InMasterTrackData.IsValid() || InCinematicSection == nullptr)
+	{
+		return false;
+	}
+
+	TSharedPtr<FMovieSceneExportCinematicSectionData> SectionData = MakeShared<FMovieSceneExportCinematicSectionData>();
+	InMasterTrackData->CinematicSections.Add(SectionData);
+
+	SectionData->ShotDisplayName = InCinematicSection->GetShotDisplayName();
+	SectionData->ShotFilename = SectionData->ShotDisplayName + TEXT(".avi");
+	SectionData->CinematicShotSection = InCinematicSection;
+
+	ConstructSectionData(InMovieScene, SectionData, InCinematicSection, EMovieSceneTranslatorSectionType::Cinematic, SectionData->ShotDisplayName);
+
+	return true;
+}
+
+bool FMovieSceneExportData::ConstructAudioSectionData(const UMovieScene* InMovieScene, TSharedPtr<FMovieSceneExportAudioTrackData> InTrackData, const UMovieSceneAudioSection* InAudioSection)
+{
+	if (InMovieScene == nullptr || !InTrackData.IsValid() || InAudioSection == nullptr)
+	{
+		return false;
+	}
+
+	TSharedPtr<FMovieSceneExportAudioSectionData> SectionData = MakeShared<FMovieSceneExportAudioSectionData>();
+	InTrackData->AudioSections.Add(SectionData);
+
+	ConstructSectionData(InMovieScene, SectionData, InAudioSection, EMovieSceneTranslatorSectionType::Audio, TEXT(""));
+
+	return true;
+}
+
+bool FMovieSceneExportData::ConstructSectionData(const UMovieScene* InMovieScene, TSharedPtr<FMovieSceneExportSectionData> InSectionData, const UMovieSceneSection* InSection, EMovieSceneTranslatorSectionType InSectionType, const FString& InSectionDisplayName)
+{
+	if (InMovieScene == nullptr || !MovieSceneData.IsValid() || InSection == nullptr || !InSectionData.IsValid())
+	{
+		return false;
+	}
+
+	InSectionData->RowIndex = InSection->GetRowIndex();
+
+	if (InSection->HasStartFrame())
+	{
+		FFrameTime InclusiveStartFrame = InSection->GetInclusiveStartFrame();
+		FFrameTime ConvertedStartFrame = ConvertFrameTime(InclusiveStartFrame, MovieSceneData->Resolution, FrameRate);
+		InSectionData->StartFrame = ConvertedStartFrame.CeilToFrame();
+
+		if (ExportContext.IsValid() && InSectionType == EMovieSceneTranslatorSectionType::Cinematic && ConvertedStartFrame.GetSubFrame() > 0.0f)
+		{
+			ExportContext->AddMessage(EMessageSeverity::Warning,
+				FText::Format(LOCTEXT("SectionStartNotDivisByDisplayRate", "Section '{0}' starts on tick {1} which is not evenly divisible by the display rate {2}. Enable snapping and adjust the start or edit the section properties to ensure it lands evenly on a whole frame."),
+					FText::FromString(InSectionDisplayName),
+					FText::AsNumber(InclusiveStartFrame.CeilToFrame().Value),
+					FText::FromString(FString::SanitizeFloat(FrameRate.AsDecimal()))));
+		}
+	}
+	else
+	{
+		InSectionData->StartFrame = FFrameNumber(0);
+		if (ExportContext.IsValid() && InSectionType == EMovieSceneTranslatorSectionType::Cinematic)
+		{
+			ExportContext->AddMessage(EMessageSeverity::Warning, FText::Format(LOCTEXT("SectionHasNoStartFrame", "Section '{0}' has no start frame. Start frame will default to 0."), FText::FromString(InSectionDisplayName)));
+		}
+	}
+
+	if (InSection->HasEndFrame())
+	{
+		FFrameTime ExclusiveEndFrame = InSection->GetExclusiveEndFrame();
+		FFrameTime ConvertedEndFrame = ConvertFrameTime(ExclusiveEndFrame, MovieSceneData->Resolution, FrameRate);
+		InSectionData->EndFrame = ConvertedEndFrame.CeilToFrame();
+
+		if (ExportContext.IsValid() && InSectionType == EMovieSceneTranslatorSectionType::Cinematic && ConvertedEndFrame.GetSubFrame() > 0.0f)
+		{
+			ExportContext->AddMessage(EMessageSeverity::Warning,
+				FText::Format(LOCTEXT("SectionEndNotDivisByDisplayRate", "Section '{0}' ends on tick {1} which is not evenly divisible by the display rate {2}. Enable snapping and adjust the end or edit the section properties to ensure it lands evenly on a whole frame."),
+					FText::FromString(InSectionDisplayName),
+					FText::FromString(FString::FromInt(ExclusiveEndFrame.CeilToFrame().Value)), 
+					FText::FromString(FString::SanitizeFloat(FrameRate.AsDecimal()))));
+		}
+	}
+	else
+	{
+		InSectionData->EndFrame = MovieSceneData->PlaybackRangeEndFrame;
+		if (ExportContext.IsValid() && InSectionType == EMovieSceneTranslatorSectionType::Cinematic)
+		{
+			ExportContext->AddMessage(EMessageSeverity::Warning, FText::Format(LOCTEXT("SectionHasNoEndFrame", "Section '{0}' has no end frame. End frame will default to playback range end."), FText::FromString(InSectionDisplayName)));
+		}
+	}
+
+	// @todo handle intersection with playback range
+	TRange<FFrameNumber> PlaybackRange = InMovieScene->GetPlaybackRange();
+	TRange<FFrameNumber> EditRange = InSection->GetRange();
+	TRange<FFrameNumber> Intersection = TRange<FFrameNumber>::Intersection(PlaybackRange, EditRange);
+	InSectionData->bWithinPlaybackRange = EditRange.Overlaps(PlaybackRange);
+	InSectionData->bEnabled = true;
+
+	return true;
+}
+
+FString FMovieSceneExportData::GetFilename() const
+{
+	return SaveFilename;
+}
+
+FString FMovieSceneExportData::GetFilenamePath() const
+{
+	return SaveFilenamePath;
+}
+
+FFrameRate FMovieSceneExportData::GetFrameRate() const
+{
+	return FrameRate;
+}
+
+uint32 FMovieSceneExportData::GetNearestWholeFrameRate() const
+{
+	if (GetFrameRateIsNTSC())
+	{
+		double Rate = FrameRate.AsDecimal();
+		return static_cast<int32>(FMath::FloorToDouble(Rate + 0.5));
+	}
+	return static_cast<uint32>(FrameRate.AsDecimal());
+}
+
+bool FMovieSceneExportData::GetFrameRateIsNTSC() const
+{
+	float FractionalPart = FMath::Frac(FrameRate.AsDecimal());
+	return (!FMath::IsNearlyZero(FractionalPart));
+}
+
+FFrameTime FMovieSceneExportData::GetHandleFrames() const
+{
+	return HandleFrames;
+}
+
+FMovieSceneImportData::FMovieSceneImportData(UMovieScene* InMovieScene, TSharedPtr<FMovieSceneTranslatorContext> InContext)
+{
+	if (InMovieScene == nullptr)
+	{
+		return;
+	}
+
+	ImportContext = InContext;
+
+	MovieSceneData = ConstructMovieSceneData(InMovieScene);
+}
+
+FMovieSceneImportData::FMovieSceneImportData() : MovieSceneData(nullptr)
+{
+}
+
+FMovieSceneImportData::~FMovieSceneImportData()
+{
+}
+
+bool FMovieSceneImportData::IsImportDataValid() const
+{
+	return MovieSceneData.IsValid();
+}
+
+TSharedPtr<FMovieSceneImportCinematicSectionData> FMovieSceneImportData::FindCinematicSection(const FString& InName) 
+{
+	TSharedPtr<FMovieSceneImportCinematicMasterTrackData> MasterTrackData = GetCinematicMasterTrackData(false);
+	if (!MasterTrackData.IsValid())
+	{
+		return nullptr;
+	}
+
+	for (TSharedPtr<FMovieSceneImportCinematicSectionData> CinematicSectionData : MasterTrackData->CinematicSectionsData)
+	{
+		UMovieSceneCinematicShotSection* CinematicShotSection = CinematicSectionData->CinematicSection;
+		if (CinematicShotSection != nullptr)
+		{
+			UMovieSceneSequence* ShotSequence = CinematicShotSection->GetSequence();
+
+			if (ShotSequence != nullptr && ShotSequence->GetName() == InName)
+			{
+				return CinematicSectionData;
+			}
+		}
+	}
+	return nullptr;
+}
+
+/** Create cinematic section */
+TSharedPtr<FMovieSceneImportCinematicSectionData> FMovieSceneImportData::CreateCinematicSection(FString InName, int32 InRow, FFrameRate InFrameRate, FFrameNumber InStartFrame, FFrameNumber InEndFrame, FFrameNumber InStartOffsetFrame)
+{
+	if (!MovieSceneData.IsValid() || MovieSceneData->MovieScene == nullptr)
+	{ 
+		return nullptr;
+	}
+
+	TSharedPtr<FMovieSceneImportCinematicMasterTrackData> MasterTrackData = GetCinematicMasterTrackData(true);
+	if (!MasterTrackData.IsValid())
+	{
+		return nullptr;
+	}
+
+	UMovieSceneCinematicShotTrack* MasterTrack = MasterTrackData->CinematicMasterTrack;
+	if (MasterTrack == nullptr)
+	{
+		return nullptr;
+	}
+
+	UMovieSceneSequence* SequenceToAdd = nullptr;
+
+	FAssetRegistryModule& AssetRegistryModule = FModuleManager::LoadModuleChecked<FAssetRegistryModule>(TEXT("AssetRegistry"));
+
+	// Collect a full list of assets with the specified class
+	TArray<FAssetData> AssetDataArray;
+	AssetRegistryModule.Get().GetAssetsByClass(ULevelSequence::StaticClass()->GetFName(), AssetDataArray);
+
+	for (FAssetData AssetData : AssetDataArray)
+	{
+		if (AssetData.AssetName == *InName)
+		{
+			SequenceToAdd = Cast<ULevelSequence>(AssetData.GetAsset());
+			break;
+		}
+	}
+
+	if (SequenceToAdd == nullptr)
+	{
+		return nullptr;
+	}
+
+	// both FCP XML and Sequencer have inclusive start frame, exclusive end frame
+	FFrameRate Resolution = MovieSceneData->MovieScene->GetTickResolution();
+	FFrameNumber StartFrame = ConvertFrameTime(InStartFrame, InFrameRate, Resolution).RoundToFrame();
+	FFrameNumber StartOffsetFrame = ConvertFrameTime(InStartOffsetFrame, InFrameRate, Resolution).RoundToFrame();
+	FFrameNumber EndFrame = ConvertFrameTime(InEndFrame, InFrameRate, Resolution).RoundToFrame();
+	int32 Duration = (EndFrame - StartFrame).Value;
+	
+	MasterTrack->Modify();
+	UMovieSceneCinematicShotSection* Section = Cast<UMovieSceneCinematicShotSection>(MasterTrack->AddSequence(SequenceToAdd, StartFrame, Duration));
+	if (Section == nullptr)
+	{
+		return nullptr;
+	}
+	Section->Modify();
+	Section->SetRowIndex(InRow);
+	Section->Parameters.SetStartFrameOffset(StartOffsetFrame.Value);
+	Section->SetRange(TRange<FFrameNumber>(StartFrame, EndFrame));
+
+	TSharedPtr<FMovieSceneImportCinematicSectionData> SectionData = ConstructCinematicSectionData(Section);
+	MasterTrackData->CinematicSectionsData.Add(SectionData);
+
+	return SectionData;
+}
+
+bool FMovieSceneImportData::SetCinematicSection(TSharedPtr<FMovieSceneImportCinematicSectionData> InSection, int32 InRow, FFrameRate InFrameRate, FFrameNumber InStartFrame, FFrameNumber InEndFrame, FFrameNumber InStartOffsetFrame)
+{
+	if (!InSection.IsValid() || InSection->CinematicSection == nullptr)
+	{
+		return false;
+	}
+
+	FFrameRate Resolution = MovieSceneData->MovieScene->GetTickResolution();
+	FFrameNumber StartFrame = ConvertFrameTime(InStartFrame, InFrameRate, Resolution).GetFrame();
+	FFrameNumber StartOffsetFrame = ConvertFrameTime(InStartOffsetFrame, InFrameRate, Resolution).GetFrame();
+	FFrameNumber EndFrame = ConvertFrameTime(InEndFrame, InFrameRate, Resolution).GetFrame();
+
+	InSection->CinematicSection->Modify();
+	InSection->CinematicSection->Parameters.SetStartFrameOffset(StartOffsetFrame.Value);
+	InSection->CinematicSection->SetRange(TRange<FFrameNumber>(StartFrame, EndFrame));
+	if (InRow != InSection->CinematicSection->GetRowIndex())
+	{
+		InSection->CinematicSection->SetRowIndex(InRow);
+	}
+
+	return true;
+}
+
+TSharedPtr<FMovieSceneImportAudioSectionData> FMovieSceneImportData::FindAudioSection(FString InName) const
+{
+	// @todo
+
+	return nullptr;
+}
+
+
+TSharedPtr<FMovieSceneImportCinematicMasterTrackData> FMovieSceneImportData::GetCinematicMasterTrackData(bool CreateTrackIfNull) 
+{
+	if (!MovieSceneData.IsValid())
+	{
+		return nullptr;
+	}
+	if (!MovieSceneData->CinematicMasterTrackData.IsValid() && CreateTrackIfNull)
+	{
+		UMovieSceneCinematicShotTrack* CinematicMasterTrack = MovieSceneData->MovieScene->AddMasterTrack<UMovieSceneCinematicShotTrack>();
+		MovieSceneData->CinematicMasterTrackData = ConstructCinematicMasterTrackData(CinematicMasterTrack);
+	}
+	return MovieSceneData->CinematicMasterTrackData;
+}
+
+TSharedPtr<FMovieSceneImportMovieSceneData> FMovieSceneImportData::ConstructMovieSceneData(UMovieScene* InMovieScene)
+{
+	if (InMovieScene == nullptr)
+	{
+		return nullptr;
+	}
+
+	MovieSceneData = MakeShared<FMovieSceneImportMovieSceneData>();
+	MovieSceneData->MovieScene = InMovieScene;
+
+	// Get cinematic master track
+	UMovieSceneCinematicShotTrack* CinematicMasterTrack = InMovieScene->FindMasterTrack<UMovieSceneCinematicShotTrack>();
+	if (CinematicMasterTrack != nullptr)
+	{
+		MovieSceneData->CinematicMasterTrackData = ConstructCinematicMasterTrackData(CinematicMasterTrack);
+		if (!MovieSceneData->CinematicMasterTrackData.IsValid())
+		{
+			return nullptr;
+		}
+	}
+
+	// Get audio tracks
+	const TArray<UMovieSceneTrack*> MasterTracks = InMovieScene->GetMasterTracks();
+	for (UMovieSceneTrack* MasterTrack : MasterTracks)
+	{
+		if (MasterTrack->IsA(UMovieSceneAudioTrack::StaticClass()))
+		{
+
+			UMovieSceneAudioTrack* AudioTrack = Cast<UMovieSceneAudioTrack>(MasterTrack);
+			TSharedPtr<FMovieSceneImportAudioTrackData> AudioTrackData = ConstructAudioTrackData(AudioTrack);
+			if (!AudioTrackData.IsValid())
+			{
+				return nullptr;
+			}
+			MovieSceneData->AudioTracksData.Add(AudioTrackData);
+		}
+	}
+
+	return MovieSceneData;
+}
+
+
+TSharedPtr<FMovieSceneImportCinematicMasterTrackData> FMovieSceneImportData::ConstructCinematicMasterTrackData(UMovieSceneCinematicShotTrack* InCinematicMasterTrack)
+{
+	if (!MovieSceneData.IsValid() || InCinematicMasterTrack == nullptr)
+	{
+		return nullptr;
+	}
+
+	TSharedPtr<FMovieSceneImportCinematicMasterTrackData> TrackData = MakeShared<FMovieSceneImportCinematicMasterTrackData>();
+	TrackData->CinematicMasterTrack = InCinematicMasterTrack;
+
+	for (UMovieSceneSection* ShotSection : InCinematicMasterTrack->GetAllSections())
+	{
+		UMovieSceneCinematicShotSection* CinematicSection = Cast<UMovieSceneCinematicShotSection>(ShotSection);
+
+		if (CinematicSection != nullptr && CinematicSection->GetSequence() != nullptr)
+		{
+			TSharedPtr<FMovieSceneImportCinematicSectionData> CinematicSectionData = ConstructCinematicSectionData(CinematicSection);
+			if (!CinematicSectionData.IsValid())
+			{
+				return nullptr;
+			}
+			TrackData->CinematicSectionsData.Add(CinematicSectionData);
+		}
+	}
+
+	return TrackData;
+}
+
+TSharedPtr<FMovieSceneImportAudioTrackData> FMovieSceneImportData::ConstructAudioTrackData(UMovieSceneAudioTrack* InAudioTrack)
+{
+	if (!MovieSceneData.IsValid() || InAudioTrack == nullptr)
+	{
+		return nullptr;
+	}
+
+	TSharedPtr<FMovieSceneImportAudioTrackData> TrackData = MakeShared<FMovieSceneImportAudioTrackData>();
+	TrackData->AudioTrack = InAudioTrack;
+
+	for (UMovieSceneSection* ShotSection : InAudioTrack->GetAllSections())
+	{
+		UMovieSceneAudioSection* AudioSection = Cast<UMovieSceneAudioSection>(ShotSection);
+
+		TSharedPtr<FMovieSceneImportAudioSectionData> AudioSectionData = ConstructAudioSectionData(AudioSection);
+		if (!AudioSectionData.IsValid())
+		{
+			return nullptr;
+		}
+		TrackData->AudioSectionsData.Add(AudioSectionData);
+	}
+
+	return TrackData;
+}
+
+TSharedPtr<FMovieSceneImportCinematicSectionData> FMovieSceneImportData::ConstructCinematicSectionData(UMovieSceneCinematicShotSection* InCinematicSection)
+{
+	if (!MovieSceneData.IsValid() || MovieSceneData->MovieScene == nullptr || InCinematicSection == nullptr)
+	{
+		return nullptr;
+	}	
+
+	TSharedPtr<FMovieSceneImportCinematicSectionData> SectionData = MakeShared<FMovieSceneImportCinematicSectionData>();
+	SectionData->CinematicSection = InCinematicSection;
+	
+	return SectionData;
+}
+
+TSharedPtr<FMovieSceneImportAudioSectionData> FMovieSceneImportData::ConstructAudioSectionData(UMovieSceneAudioSection* InAudioSection)
+{
+	if (InAudioSection == nullptr)
+	{
+		return nullptr;
+	}
+
+	TSharedPtr<FMovieSceneImportAudioSectionData> SectionData = MakeShared<FMovieSceneImportAudioSectionData>();
+	SectionData->AudioSection = InAudioSection;
+
+	return SectionData;
+}
+
+
+void FMovieSceneTranslatorContext::Init()
+{
+	ClearMessages();
+}
+
+void FMovieSceneTranslatorContext::AddMessage(EMessageSeverity::Type InMessageSeverity, FText InMessage)
+{
+	Messages.Add(FTokenizedMessage::Create(InMessageSeverity, InMessage));
+}
+
+void FMovieSceneTranslatorContext::ClearMessages()
+{
+	Messages.Empty();
+}
+
+bool FMovieSceneTranslatorContext::ContainsMessageType(EMessageSeverity::Type InMessageSeverity) const
+{
+	for (const TSharedRef<FTokenizedMessage>& Message : Messages)
+	{
+		if (Message->GetSeverity() == InMessageSeverity)
+		{
+			return true;
+		}
+	}
+	return false;
+}
+
+const TArray<TSharedRef<FTokenizedMessage>>& FMovieSceneTranslatorContext::GetMessages() const
+{
+	return Messages;
+}
+
+#undef LOCTEXT_NAMESPACE
diff --git a/Engine/Source/Editor/MovieSceneTools/Private/Sections/BoolPropertySection.cpp b/Engine/Source/Editor/MovieSceneTools/Private/Sections/BoolPropertySection.cpp
--- a/Engine/Source/Editor/MovieSceneTools/Private/Sections/BoolPropertySection.cpp
+++ b/Engine/Source/Editor/MovieSceneTools/Private/Sections/BoolPropertySection.cpp
@@ -15,13 +15,13 @@ int32 FBoolPropertySection::OnPaintSection( FSequencerSectionPainter& Painter )
 	// custom drawing for bool curves
 	UMovieSceneBoolSection* BoolSection = Cast<UMovieSceneBoolSection>( WeakSection.Get() );
 
-	TArray<float> SectionSwitchTimes;
+	TArray<FFrameTime> SectionSwitchTimes;
 
 	const FTimeToPixel& TimeConverter = Painter.GetTimeConverter();
 
 	// Add the start time
-	const float StartTime = TimeConverter.PixelToSeconds(0.f);
-	const float EndTime   = TimeConverter.PixelToSeconds(Painter.SectionGeometry.GetLocalSize().X);
+	const FFrameNumber StartTime = TimeConverter.PixelToFrame(0.f).FloorToFrame();
+	const FFrameNumber EndTime   = TimeConverter.PixelToFrame(Painter.SectionGeometry.GetLocalSize().X).CeilToFrame();
 	
 	SectionSwitchTimes.Add(StartTime);
 
@@ -35,10 +35,9 @@ int32 FBoolPropertySection::OnPaintSection( FSequencerSectionPainter& Painter )
 
 	for ( FFrameNumber Time : BoolChannel->GetTimes() )
 	{
-		double TimeInSeconds = Time / TimeConverter.GetFrameResolution();
-		if ( TimeInSeconds > StartTime && TimeInSeconds < EndTime )
+		if ( Time > StartTime && Time < EndTime )
 		{
-			SectionSwitchTimes.Add( TimeInSeconds );
+			SectionSwitchTimes.Add( Time );
 		}
 	}
 
@@ -55,15 +54,15 @@ int32 FBoolPropertySection::OnPaintSection( FSequencerSectionPainter& Painter )
 
 	for ( int32 i = 0; i < SectionSwitchTimes.Num() - 1; ++i )
 	{
-		float ThisTime = SectionSwitchTimes[i];
+		FFrameTime ThisTime = SectionSwitchTimes[i];
 
 		bool ValueAtTime = false;
-		BoolChannel->Evaluate(ThisTime * TimeConverter.GetFrameResolution(), ValueAtTime);
+		BoolChannel->Evaluate(ThisTime, ValueAtTime);
 
 		const FColor Color = ValueAtTime ? FColor(0, 255, 0, 125) : FColor(255, 0, 0, 125);
 		
-		FVector2D StartPos(TimeConverter.SecondsToPixel(ThisTime), VerticalOffset);
-		FVector2D Size(TimeConverter.SecondsToPixel(SectionSwitchTimes[i+1]) - StartPos.X, Height);
+		FVector2D StartPos(TimeConverter.FrameToPixel(ThisTime), VerticalOffset);
+		FVector2D Size(TimeConverter.FrameToPixel(SectionSwitchTimes[i+1]) - StartPos.X, Height);
 
 		FSlateDrawElement::MakeBox(
 			Painter.DrawElements,
diff --git a/Engine/Source/Editor/MovieSceneTools/Private/Sections/CameraCutSection.cpp b/Engine/Source/Editor/MovieSceneTools/Private/Sections/CameraCutSection.cpp
--- a/Engine/Source/Editor/MovieSceneTools/Private/Sections/CameraCutSection.cpp
+++ b/Engine/Source/Editor/MovieSceneTools/Private/Sections/CameraCutSection.cpp
@@ -41,7 +41,7 @@ void FCameraCutSection::SetSingleTime(double GlobalTime)
 	UMovieSceneCameraCutSection* CameraCutSection = Cast<UMovieSceneCameraCutSection>(Section);
 	if (CameraCutSection && CameraCutSection->HasStartFrame())
 	{
-		double ReferenceOffsetSeconds = CameraCutSection->GetInclusiveStartFrame() / CameraCutSection->GetTypedOuter<UMovieScene>()->GetFrameResolution();
+		double ReferenceOffsetSeconds = CameraCutSection->GetInclusiveStartFrame() / CameraCutSection->GetTypedOuter<UMovieScene>()->GetTickResolution();
 		CameraCutSection->SetThumbnailReferenceOffset(GlobalTime - ReferenceOffsetSeconds);
 	}
 }
@@ -53,7 +53,7 @@ void FCameraCutSection::Tick(const FGeometry& AllottedGeometry, const FGeometry&
 	{
 		if (GetDefault<UMovieSceneUserThumbnailSettings>()->bDrawSingleThumbnails && CameraCutSection->HasStartFrame())
 		{
-			double ReferenceOffsetSeconds = CameraCutSection->GetInclusiveStartFrame() / CameraCutSection->GetTypedOuter<UMovieScene>()->GetFrameResolution() + CameraCutSection->GetThumbnailReferenceOffset();
+			double ReferenceOffsetSeconds = CameraCutSection->GetInclusiveStartFrame() / CameraCutSection->GetTypedOuter<UMovieScene>()->GetTickResolution() + CameraCutSection->GetThumbnailReferenceOffset();
 			ThumbnailCache.SetSingleReferenceFrame(ReferenceOffsetSeconds);
 		}
 		else
diff --git a/Engine/Source/Editor/MovieSceneTools/Private/Sections/CinematicShotSection.cpp b/Engine/Source/Editor/MovieSceneTools/Private/Sections/CinematicShotSection.cpp
--- a/Engine/Source/Editor/MovieSceneTools/Private/Sections/CinematicShotSection.cpp
+++ b/Engine/Source/Editor/MovieSceneTools/Private/Sections/CinematicShotSection.cpp
@@ -33,7 +33,7 @@ FCinematicShotSection::FCinematicSectionCache::FCinematicSectionCache(UMovieScen
 		UMovieSceneSequence* InnerSequence = Section->GetSequence();
 		if (InnerSequence)
 		{
-			InnerFrameRate = InnerSequence->GetMovieScene()->GetFrameResolution();
+			InnerFrameRate = InnerSequence->GetMovieScene()->GetTickResolution();
 		}
 
 		InnerFrameOffset = Section->Parameters.GetStartFrameOffset();
@@ -72,7 +72,7 @@ FMargin FCinematicShotSection::GetContentPadding() const
 
 void FCinematicShotSection::SetSingleTime(double GlobalTime)
 {
-	double ReferenceOffsetSeconds = SectionObject.HasStartFrame() ? SectionObject.GetInclusiveStartFrame() / SectionObject.GetTypedOuter<UMovieScene>()->GetFrameResolution() : 0;
+	double ReferenceOffsetSeconds = SectionObject.HasStartFrame() ? SectionObject.GetInclusiveStartFrame() / SectionObject.GetTypedOuter<UMovieScene>()->GetTickResolution() : 0;
 	SectionObject.SetThumbnailReferenceOffset(GlobalTime - ReferenceOffsetSeconds);
 }
 
@@ -89,8 +89,8 @@ void FCinematicShotSection::ResizeSection(ESequencerSectionResizeMode ResizeMode
 	// Adjust the start offset when resizing from the beginning
 	if (ResizeMode == SSRM_LeadingEdge && InnerSequence)
 	{
-		const FFrameRate    OuterFrameRate   = SectionObject.GetTypedOuter<UMovieScene>()->GetFrameResolution();
-		const FFrameRate    InnerFrameRate   = InnerSequence->GetMovieScene()->GetFrameResolution();
+		const FFrameRate    OuterFrameRate   = SectionObject.GetTypedOuter<UMovieScene>()->GetTickResolution();
+		const FFrameRate    InnerFrameRate   = InnerSequence->GetMovieScene()->GetTickResolution();
 		const FFrameNumber  ResizeDifference = ResizeTime - InitialStartTimeDuringResize;
 		const FFrameTime    InnerFrameTime   = ConvertFrameTime(ResizeDifference, OuterFrameRate, InnerFrameRate);
 		const int32         NewStartOffset   = FFrameTime::FromDecimal(InnerFrameTime.AsDecimal() * SectionObject.Parameters.TimeScale).FrameNumber.Value;
@@ -114,8 +114,8 @@ void FCinematicShotSection::SlipSection(double SlipTime)
 	// Adjust the start offset when resizing from the beginning
 	if (InnerSequence)
 	{
-		const FFrameRate    OuterFrameRate   = SectionObject.GetTypedOuter<UMovieScene>()->GetFrameResolution();
-		const FFrameRate    InnerFrameRate   = InnerSequence->GetMovieScene()->GetFrameResolution();
+		const FFrameRate    OuterFrameRate   = SectionObject.GetTypedOuter<UMovieScene>()->GetTickResolution();
+		const FFrameRate    InnerFrameRate   = InnerSequence->GetMovieScene()->GetTickResolution();
 		const double        ResizeDifference = SlipTime - InitialStartTimeDuringResize / OuterFrameRate;
 		const int32         NewStartOffset   = ( (ResizeDifference * SectionObject.Parameters.TimeScale) * InnerFrameRate ).FrameNumber.Value;
 
@@ -139,7 +139,7 @@ void FCinematicShotSection::Tick(const FGeometry& AllottedGeometry, const FGeome
 	// Update single reference frame settings
 	if (GetDefault<UMovieSceneUserThumbnailSettings>()->bDrawSingleThumbnails && SectionObject.HasStartFrame())
 	{
-		double ReferenceTime = SectionObject.GetInclusiveStartFrame() / SectionObject.GetTypedOuter<UMovieScene>()->GetFrameResolution() + SectionObject.GetThumbnailReferenceOffset();
+		double ReferenceTime = SectionObject.GetInclusiveStartFrame() / SectionObject.GetTypedOuter<UMovieScene>()->GetTickResolution() + SectionObject.GetThumbnailReferenceOffset();
 		ThumbnailCache.SetSingleReferenceFrame(ReferenceTime);
 	}
 	else
diff --git a/Engine/Source/Editor/MovieSceneTools/Private/Sections/ColorPropertySection.cpp b/Engine/Source/Editor/MovieSceneTools/Private/Sections/ColorPropertySection.cpp
--- a/Engine/Source/Editor/MovieSceneTools/Private/Sections/ColorPropertySection.cpp
+++ b/Engine/Source/Editor/MovieSceneTools/Private/Sections/ColorPropertySection.cpp
@@ -146,7 +146,7 @@ void FColorPropertySection::ConsolidateColorCurves( TArray< TTuple<float, FLinea
 			FloatChannels[2]->Evaluate(Time, ColorAtTime.B);
 			FloatChannels[3]->Evaluate(Time, ColorAtTime.A);
 
-			OutColorKeys.Add(MakeTuple(float(Time / TimeConverter.GetFrameResolution()), ColorAtTime));
+			OutColorKeys.Add(MakeTuple(float(Time / TimeConverter.GetTickResolution()), ColorAtTime));
 		}
 	}
 
diff --git a/Engine/Source/Editor/MovieSceneTools/Private/Sections/EventSection.cpp b/Engine/Source/Editor/MovieSceneTools/Private/Sections/EventSection.cpp
new file mode 100644
--- /dev/null
+++ b/Engine/Source/Editor/MovieSceneTools/Private/Sections/EventSection.cpp
@@ -0,0 +1,106 @@
+// Copyright 1998-2018 Epic Games, Inc. All Rights Reserved.
+
+#include "Sections/EventSection.h"
+#include "Rendering/DrawElements.h"
+#include "SequencerSectionPainter.h"
+#include "ISectionLayoutBuilder.h"
+#include "EditorStyleSet.h"
+#include "CommonMovieSceneTools.h"
+#include "Sections/MovieSceneEventSection.h"
+#include "MovieSceneTrack.h"
+#include "SequencerTimeSliderController.h"
+#include "Fonts/FontMeasure.h"
+#include "Framework/Application/SlateApplication.h"
+#include "EditorStyleSet.h"
+
+int32 FEventSection::OnPaintSection( FSequencerSectionPainter& Painter ) const
+{
+	int32 LayerId = Painter.PaintSectionBackground();
+	UMovieSceneEventSection* EventSection = Cast<UMovieSceneEventSection>( WeakSection.Get() );
+	if (!EventSection)
+	{
+		return LayerId;
+	}
+
+	// Draw event names for keys on sections that are on selected tracks
+	TSharedPtr<ISequencer> SequencerPtr = Sequencer.Pin();
+	TArray<UMovieSceneTrack*> SelectedTracks;
+	SequencerPtr->GetSelectedTracks(SelectedTracks);
+
+	bool bSelected = false;
+	for (auto SelectedTrack : SelectedTracks)
+	{
+		if (SelectedTrack->GetAllSections().Contains(EventSection))
+		{
+			bSelected = true;
+			break;
+		}
+	}
+
+	if (!bSelected)
+	{
+		return LayerId;
+	}
+
+	const FTimeToPixel& TimeToPixelConverter = Painter.GetTimeConverter();
+
+	const ESlateDrawEffect DrawEffects = Painter.bParentEnabled
+		? ESlateDrawEffect::None
+		: ESlateDrawEffect::DisabledEffect;
+
+	for (int32 KeyIndex = 0; KeyIndex < EventSection->GetEventData().GetKeyTimes().Num(); ++KeyIndex)
+	{
+		FFrameNumber EventTime = EventSection->GetEventData().GetKeyTimes()[KeyIndex];
+		FEventPayload EventData = EventSection->GetEventData().GetKeyValues()[KeyIndex];
+
+		if (!EventSection->GetRange().Contains(EventTime))
+		{
+			continue;
+		}
+
+		FString EventString = EventData.EventName.ToString();
+		if (EventString.IsEmpty())
+		{
+			continue;
+		}
+
+		const float Time = TimeToPixelConverter.FrameToPixel(EventTime); 
+
+		const FSlateFontInfo SmallLayoutFont = FCoreStyle::GetDefaultFontStyle("Bold", 10);
+		const TSharedRef< FSlateFontMeasure > FontMeasureService = FSlateApplication::Get().GetRenderer()->GetFontMeasureService();
+		FVector2D TextSize = FontMeasureService->Measure(EventString, SmallLayoutFont);
+
+		// Flip the text position if getting near the end of the view range
+		static const float TextOffsetPx = 10.f;
+		bool  bDrawLeft = (Painter.SectionGeometry.Size.X - Time) < (TextSize.X + 22.f) - TextOffsetPx;
+		float TextPosition = bDrawLeft ? Time - TextSize.X - TextOffsetPx : Time + TextOffsetPx;
+		//handle mirrored labels
+		const float MajorTickHeight = 4.0f; 
+		FVector2D TextOffset(TextPosition, Painter.SectionGeometry.Size.Y - (MajorTickHeight + TextSize.Y));
+
+		const FLinearColor DrawColor = FEditorStyle::GetSlateColor("SelectionColor").GetColor(FWidgetStyle());
+		const FVector2D BoxPadding = FVector2D(4.0f, 2.0f);
+
+		// draw event string
+		FSlateDrawElement::MakeBox(
+			Painter.DrawElements,
+			LayerId + 5,
+			Painter.SectionGeometry.ToPaintGeometry(TextOffset - BoxPadding, TextSize + 2.0f * BoxPadding),
+			FEditorStyle::GetBrush("WhiteBrush"),
+			ESlateDrawEffect::None,
+			FLinearColor::Black.CopyWithNewOpacity(0.5f)
+		);
+
+		FSlateDrawElement::MakeText(
+			Painter.DrawElements,
+			LayerId + 6,
+			Painter.SectionGeometry.ToPaintGeometry(TextOffset, TextSize),
+			EventString,
+			SmallLayoutFont,
+			DrawEffects,
+			DrawColor
+		);
+	}
+
+	return LayerId + 1;
+}
diff --git a/Engine/Source/Editor/MovieSceneTools/Private/Sections/EventSection.h b/Engine/Source/Editor/MovieSceneTools/Private/Sections/EventSection.h
new file mode 100644
--- /dev/null
+++ b/Engine/Source/Editor/MovieSceneTools/Private/Sections/EventSection.h
@@ -0,0 +1,36 @@
+// Copyright 1998-2018 Epic Games, Inc. All Rights Reserved.
+
+#pragma once
+
+#include "ISequencerSection.h"
+#include "ISequencer.h"
+
+class FSequencerSectionPainter;
+
+/**
+* An implementation of event sections.
+*/
+class FEventSection
+	: public FSequencerSection
+{
+public:
+
+	/**
+	* Creates a new event section for editing event sections.
+	*
+	* @param InSectionObject The section object which is being displayed and edited.
+	* @param InDisplayName A display name for the section being displayed and edited.
+	*/
+	FEventSection(UMovieSceneSection& InSectionObject, TWeakPtr<ISequencer> InSequencer)
+		: FSequencerSection(InSectionObject)
+		, Sequencer(InSequencer)
+	{}
+
+public:
+
+	virtual int32 OnPaintSection(FSequencerSectionPainter& Painter) const override;
+
+private:
+
+	TWeakPtr<ISequencer> Sequencer;
+};
diff --git a/Engine/Source/Editor/MovieSceneTools/Private/Sections/ThumbnailSection.cpp b/Engine/Source/Editor/MovieSceneTools/Private/Sections/ThumbnailSection.cpp
--- a/Engine/Source/Editor/MovieSceneTools/Private/Sections/ThumbnailSection.cpp
+++ b/Engine/Source/Editor/MovieSceneTools/Private/Sections/ThumbnailSection.cpp
@@ -278,10 +278,10 @@ TRange<double> FThumbnailSection::GetVisibleRange() const
 		return GlobalVisibleRange;
 	}
 
-	const FFrameRate FrameResolution = Section->GetTypedOuter<UMovieScene>()->GetFrameResolution();
+	const FFrameRate TickResolution = Section->GetTypedOuter<UMovieScene>()->GetTickResolution();
 
-	const double     StartOffset = Section->HasStartFrame() ? Section->GetInclusiveStartFrame() / FrameResolution : 0.0;
-	const double     EndOffset   = Section->HasEndFrame()   ? Section->GetExclusiveEndFrame()   / FrameResolution : 0.0;
+	const double     StartOffset = Section->HasStartFrame() ? Section->GetInclusiveStartFrame() / TickResolution : 0.0;
+	const double     EndOffset   = Section->HasEndFrame()   ? Section->GetExclusiveEndFrame()   / TickResolution : 0.0;
 
 	return TRange<double>(
 		GlobalVisibleRange.GetLowerBoundValue() - StartOffset,
@@ -291,18 +291,18 @@ TRange<double> FThumbnailSection::GetVisibleRange() const
 
 TRange<double> FThumbnailSection::GetTotalRange() const
 {
-	TRange<FFrameNumber> SectionRange    = Section->GetRange();
-	FFrameRate           FrameResolution = Section->GetTypedOuter<UMovieScene>()->GetFrameResolution();
+	TRange<FFrameNumber> SectionRange   = Section->GetRange();
+	FFrameRate           TickResolution = Section->GetTypedOuter<UMovieScene>()->GetTickResolution();
 
 	if (TimeSpace == ETimeSpace::Global)
 	{
-		return SectionRange / FrameResolution;
+		return SectionRange / TickResolution;
 	}
 	else
 	{
 		const bool bHasDiscreteSize = SectionRange.GetLowerBound().IsClosed() && SectionRange.GetUpperBound().IsClosed();
 		TRangeBound<double> UpperBound = bHasDiscreteSize
-			? TRangeBound<double>::Exclusive(FFrameNumber(MovieScene::DiscreteSize(SectionRange)) / FrameResolution)
+			? TRangeBound<double>::Exclusive(FFrameNumber(MovieScene::DiscreteSize(SectionRange)) / TickResolution)
 			: TRangeBound<double>::Open();
 
 		return TRange<double>(0, UpperBound);
diff --git a/Engine/Source/Editor/MovieSceneTools/Private/Sections/TransformPropertySection.cpp b/Engine/Source/Editor/MovieSceneTools/Private/Sections/TransformPropertySection.cpp
--- a/Engine/Source/Editor/MovieSceneTools/Private/Sections/TransformPropertySection.cpp
+++ b/Engine/Source/Editor/MovieSceneTools/Private/Sections/TransformPropertySection.cpp
@@ -122,4 +122,55 @@ void FTransformSection::BuildSectionContextMenu(FMenuBuilder& MenuBuilder, const
 	MenuBuilder.EndSection();
 }
 
+bool FTransformSection::RequestDeleteCategory(const TArray<FName>& CategoryNamePaths)
+{
+	UMovieScene3DTransformSection* TransformSection = CastChecked<UMovieScene3DTransformSection>(WeakSection.Get());
+	TSharedPtr<ISequencer> SequencerPtr = WeakSequencer.Pin();
+		
+	const FScopedTransaction Transaction( LOCTEXT( "DeleteTransformCategory", "Delete transform category" ) );
+
+	if (TransformSection->TryModify())
+	{
+		FName CategoryName = CategoryNamePaths[CategoryNamePaths.Num()-1];
+		
+		EMovieSceneTransformChannel Channel = TransformSection->GetMask().GetChannels();
+		EMovieSceneTransformChannel ChannelToRemove = TransformSection->GetMaskByName(CategoryName).GetChannels();
+
+		Channel = Channel ^ ChannelToRemove;
+
+		TransformSection->SetMask(Channel);
+			
+		SequencerPtr->NotifyMovieSceneDataChanged(EMovieSceneDataChangeType::MovieSceneStructureItemsChanged);
+		return true;
+	}
+
+	return false;
+}
+
+bool FTransformSection::RequestDeleteKeyArea(const TArray<FName>& KeyAreaNamePaths)
+{
+	UMovieScene3DTransformSection* TransformSection = CastChecked<UMovieScene3DTransformSection>(WeakSection.Get());
+	TSharedPtr<ISequencer> SequencerPtr = WeakSequencer.Pin();
+
+	const FScopedTransaction Transaction( LOCTEXT( "DeleteTransformChannel", "Delete transform channel" ) );
+
+	if (TransformSection->TryModify())
+	{
+		// Only delete the last key area path which is the channel. ie. TranslationX as opposed to Translation
+		FName KeyAreaName = KeyAreaNamePaths[KeyAreaNamePaths.Num()-1];
+
+		EMovieSceneTransformChannel Channel = TransformSection->GetMask().GetChannels();
+		EMovieSceneTransformChannel ChannelToRemove = TransformSection->GetMaskByName(KeyAreaName).GetChannels();
+
+		Channel = Channel ^ ChannelToRemove;
+
+		TransformSection->SetMask(Channel);
+					
+		SequencerPtr->NotifyMovieSceneDataChanged(EMovieSceneDataChangeType::MovieSceneStructureItemsChanged);
+		return true;
+	}
+
+	return true;
+}
+
 #undef LOCTEXT_NAMESPACE
\ No newline at end of file
diff --git a/Engine/Source/Editor/MovieSceneTools/Private/Sections/TransformPropertySection.h b/Engine/Source/Editor/MovieSceneTools/Private/Sections/TransformPropertySection.h
--- a/Engine/Source/Editor/MovieSceneTools/Private/Sections/TransformPropertySection.h
+++ b/Engine/Source/Editor/MovieSceneTools/Private/Sections/TransformPropertySection.h
@@ -31,6 +31,11 @@ public:
 
 	virtual void BuildSectionContextMenu(FMenuBuilder& MenuBuilder, const FGuid& InObjectBinding) override;
 
+	//~ ISequencerSection interface
+
+	virtual bool RequestDeleteCategory(const TArray<FName>& CategoryNamePath) override;
+	virtual bool RequestDeleteKeyArea(const TArray<FName>& KeyAreaNamePath) override;
+
 protected:
 
 	/** The sequencer which is controlling this section. */
diff --git a/Engine/Source/Editor/MovieSceneTools/Private/TrackEditors/AttachTrackEditor.cpp b/Engine/Source/Editor/MovieSceneTools/Private/TrackEditors/AttachTrackEditor.cpp
--- a/Engine/Source/Editor/MovieSceneTools/Private/TrackEditors/AttachTrackEditor.cpp
+++ b/Engine/Source/Editor/MovieSceneTools/Private/TrackEditors/AttachTrackEditor.cpp
@@ -118,8 +118,13 @@ TSharedRef<ISequencerSection> F3DAttachTrackEditor::MakeSectionInterface( UMovie
 
 void F3DAttachTrackEditor::BuildObjectBindingTrackMenu(FMenuBuilder& MenuBuilder, const FGuid& ObjectBinding, const UClass* ObjectClass)
 {
-	if (ObjectClass->IsChildOf(AActor::StaticClass()))
+	if (ObjectClass != nullptr && ObjectClass->IsChildOf(AActor::StaticClass()))
 	{
+		if (MovieSceneToolHelpers::HasHiddenMobility(ObjectClass))
+		{
+			return;
+		}
+
 		UMovieSceneSection* DummySection = nullptr;
 
 		MenuBuilder.AddSubMenu(
@@ -184,7 +189,7 @@ void F3DAttachTrackEditor::ActorSocketPicked(const FName SocketName, USceneCompo
 		else if (ActorPickerID.ActorPicked.IsValid())
 		{
 			FGuid ParentActorId = FindOrCreateHandleToObject(ActorPickerID.ActorPicked.Get()).Handle;
-			ConstraintBindingID = FMovieSceneObjectBindingID(ParentActorId, MovieSceneSequenceID::Root);
+			ConstraintBindingID = FMovieSceneObjectBindingID(ParentActorId, GetSequencer()->GetFocusedTemplateID());
 		}
 
 		if (ConstraintBindingID.IsValid())
@@ -222,7 +227,7 @@ FKeyPropertyResult F3DAttachTrackEditor::AddKeyInternal( FFrameNumber KeyTime, c
 		FFindOrCreateHandleResult HandleResult = FindOrCreateHandleToObject(ActorPickerID.ActorPicked.Get());
 		FGuid ParentActorId = HandleResult.Handle;
 		KeyPropertyResult.bHandleCreated |= HandleResult.bWasCreated;
-		ConstraintBindingID = FMovieSceneObjectBindingID(ParentActorId, MovieSceneSequenceID::Root);
+		ConstraintBindingID = FMovieSceneObjectBindingID(ParentActorId, GetSequencer()->GetFocusedTemplateID());
 	}
 
 	if (!ConstraintBindingID.IsValid())
@@ -246,7 +251,7 @@ FKeyPropertyResult F3DAttachTrackEditor::AddKeyInternal( FFrameNumber KeyTime, c
 			if (ensure(Track))
 			{
 				// Clamp to next attach section's start time or the end of the current sequencer view range
-				FFrameNumber AttachEndTime = (GetSequencer()->GetViewRange().GetUpperBoundValue() * Track->GetTypedOuter<UMovieScene>()->GetFrameResolution()).FrameNumber;
+				FFrameNumber AttachEndTime = (GetSequencer()->GetViewRange().GetUpperBoundValue() * Track->GetTypedOuter<UMovieScene>()->GetTickResolution()).FrameNumber;
 
 				for (UMovieSceneSection* Section : Track->GetAllSections())
 				{
@@ -260,7 +265,8 @@ FKeyPropertyResult F3DAttachTrackEditor::AddKeyInternal( FFrameNumber KeyTime, c
 					}
 				}
 
-				Cast<UMovieScene3DAttachTrack>(Track)->AddConstraint( KeyTime, (AttachEndTime - KeyTime).Value, SocketName, ComponentName, ConstraintBindingID);
+				int32 Duration = FMath::Max(0, (AttachEndTime - KeyTime).Value);
+				Cast<UMovieScene3DAttachTrack>(Track)->AddConstraint( KeyTime, Duration, SocketName, ComponentName, ConstraintBindingID);
 				KeyPropertyResult.bTrackModified = true;
 			}
 		}
diff --git a/Engine/Source/Editor/MovieSceneTools/Private/TrackEditors/AudioTrackEditor.cpp b/Engine/Source/Editor/MovieSceneTools/Private/TrackEditors/AudioTrackEditor.cpp
--- a/Engine/Source/Editor/MovieSceneTools/Private/TrackEditors/AudioTrackEditor.cpp
+++ b/Engine/Source/Editor/MovieSceneTools/Private/TrackEditors/AudioTrackEditor.cpp
@@ -324,7 +324,7 @@ void FAudioThumbnail::GenerateWaveformPreview(TArray<uint8>& OutData, TRange<flo
 		return;
 	}
 
-	FFrameRate FrameRate            = AudioSection->GetTypedOuter<UMovieScene>()->GetFrameResolution();
+	FFrameRate FrameRate            = AudioSection->GetTypedOuter<UMovieScene>()->GetTickResolution();
 	float      PitchMultiplierValue = AudioSection->GetPitchMultiplierChannel().GetDefault().Get(1.f);
 	double     SectionStartTime     = AudioSection->GetInclusiveStartFrame() / FrameRate;
 
@@ -694,8 +694,8 @@ void FAudioSection::Tick( const FGeometry& AllottedGeometry, const FGeometry& Pa
 		const float LeftMostVisiblePixel = FMath::Max(ParentRect.Left, 0.f);
 		const float RightMostVisiblePixel = FMath::Min(ParentRect.Right, AllottedGeometry.GetLocalSize().X);
 
-		FFrameRate   FrameResolution = AudioSection->GetTypedOuter<UMovieScene>()->GetFrameResolution();
-		FTimeToPixel TimeToPixel( AllottedGeometry, AudioSection->GetRange() / FrameResolution, FrameResolution );
+		FFrameRate   TickResolution = AudioSection->GetTypedOuter<UMovieScene>()->GetTickResolution();
+		FTimeToPixel TimeToPixel( AllottedGeometry, AudioSection->GetRange() / TickResolution, TickResolution );
 
 		TRange<float> DrawRange = TRange<float>(
 			TimeToPixel.PixelToSeconds(LeftMostVisiblePixel),
@@ -738,7 +738,7 @@ void FAudioSection::BeginSlipSection()
 {
 	UMovieSceneAudioSection* AudioSection = Cast<UMovieSceneAudioSection>(&Section);
 	InitialStartOffsetDuringResize = AudioSection->GetStartOffset();
-	InitialStartTimeDuringResize   = FFrameNumber(AudioSection->HasStartFrame() ? AudioSection->GetInclusiveStartFrame() : 0) / AudioSection->GetTypedOuter<UMovieScene>()->GetFrameResolution();
+	InitialStartTimeDuringResize   = FFrameNumber(AudioSection->HasStartFrame() ? AudioSection->GetInclusiveStartFrame() : 0) / AudioSection->GetTypedOuter<UMovieScene>()->GetTickResolution();
 }
 
 void FAudioSection::SlipSection(double SlipTime)
@@ -1097,6 +1097,7 @@ TSharedRef<SWidget> FAudioTrackEditor::BuildAudioSubMenu(UMovieSceneTrack* Track
 	FAssetPickerConfig AssetPickerConfig;
 	{
 		AssetPickerConfig.OnAssetSelected = FOnAssetSelected::CreateRaw( this, &FAudioTrackEditor::OnAudioAssetSelected, Track);
+		AssetPickerConfig.OnAssetEnterPressed = FOnAssetEnterPressed::CreateRaw( this, &FAudioTrackEditor::OnAudioAssetEnterPressed, Track);
 		AssetPickerConfig.bAllowNullSelection = false;
 		AssetPickerConfig.InitialAssetViewType = EAssetViewType::List;
 		for (auto ClassName : DerivedClassNames)
@@ -1147,4 +1148,14 @@ void FAudioTrackEditor::OnAudioAssetSelected(const FAssetData& AssetData, UMovie
 		}
 	}
 }
+
+void FAudioTrackEditor::OnAudioAssetEnterPressed(const TArray<FAssetData>& AssetData, UMovieSceneTrack* Track)
+{
+	if (AssetData.Num() > 0)
+	{
+		OnAudioAssetSelected(AssetData[0].GetAsset(), Track);
+	}
+}
+
+
 #undef LOCTEXT_NAMESPACE
diff --git a/Engine/Source/Editor/MovieSceneTools/Private/TrackEditors/CameraAnimTrackEditor.cpp b/Engine/Source/Editor/MovieSceneTools/Private/TrackEditors/CameraAnimTrackEditor.cpp
--- a/Engine/Source/Editor/MovieSceneTools/Private/TrackEditors/CameraAnimTrackEditor.cpp
+++ b/Engine/Source/Editor/MovieSceneTools/Private/TrackEditors/CameraAnimTrackEditor.cpp
@@ -160,6 +160,7 @@ void FCameraAnimTrackEditor::AddCameraAnimSubMenu(FMenuBuilder& MenuBuilder, FGu
 	FAssetPickerConfig AssetPickerConfig;
 	{
 		AssetPickerConfig.OnAssetSelected = FOnAssetSelected::CreateRaw(this, &FCameraAnimTrackEditor::OnCameraAnimAssetSelected, ObjectBinding);
+		AssetPickerConfig.OnAssetEnterPressed = FOnAssetEnterPressed::CreateRaw(this, &FCameraAnimTrackEditor::OnCameraAnimAssetEnterPressed, ObjectBinding);
 		AssetPickerConfig.bAllowNullSelection = false;
 		AssetPickerConfig.InitialAssetViewType = EAssetViewType::List;
 		AssetPickerConfig.Filter.ClassNames.Add(UCameraAnim::StaticClass()->GetFName());
@@ -211,7 +212,13 @@ void FCameraAnimTrackEditor::OnCameraAnimAssetSelected(const FAssetData& AssetDa
 	}
 }
 
-
+void FCameraAnimTrackEditor::OnCameraAnimAssetEnterPressed(const TArray<FAssetData>& AssetData, FGuid ObjectBinding)
+{
+	if (AssetData.Num() > 0)
+	{
+		OnCameraAnimAssetSelected(AssetData[0].GetAsset(), ObjectBinding);
+	}
+}
 
 FKeyPropertyResult FCameraAnimTrackEditor::AddKeyInternal(FFrameNumber KeyTime, const TArray<TWeakObjectPtr<UObject>> Objects, UCameraAnim* CameraAnim)
 {
diff --git a/Engine/Source/Editor/MovieSceneTools/Private/TrackEditors/CameraAnimTrackEditor.h b/Engine/Source/Editor/MovieSceneTools/Private/TrackEditors/CameraAnimTrackEditor.h
--- a/Engine/Source/Editor/MovieSceneTools/Private/TrackEditors/CameraAnimTrackEditor.h
+++ b/Engine/Source/Editor/MovieSceneTools/Private/TrackEditors/CameraAnimTrackEditor.h
@@ -63,6 +63,9 @@ private:
 	/** Animation asset selected */
 	void OnCameraAnimAssetSelected(const FAssetData& AssetData, FGuid ObjectBinding);
 
+	/** Animation asset enter pressed */
+	void OnCameraAnimAssetEnterPressed(const TArray<FAssetData>& AssetData, FGuid ObjectBinding);
+
 	class UCameraComponent* AcquireCameraComponentFromObjectGuid(const FGuid& Guid);
 };
 
diff --git a/Engine/Source/Editor/MovieSceneTools/Private/TrackEditors/CameraShakeTrackEditor.cpp b/Engine/Source/Editor/MovieSceneTools/Private/TrackEditors/CameraShakeTrackEditor.cpp
--- a/Engine/Source/Editor/MovieSceneTools/Private/TrackEditors/CameraShakeTrackEditor.cpp
+++ b/Engine/Source/Editor/MovieSceneTools/Private/TrackEditors/CameraShakeTrackEditor.cpp
@@ -170,6 +170,7 @@ void FCameraShakeTrackEditor::AddCameraShakeSubMenu(FMenuBuilder& MenuBuilder, F
 	FAssetPickerConfig AssetPickerConfig;
 	{
 		AssetPickerConfig.OnAssetSelected = FOnAssetSelected::CreateRaw(this, &FCameraShakeTrackEditor::OnCameraShakeAssetSelected, ObjectBinding);
+		AssetPickerConfig.OnAssetEnterPressed = FOnAssetEnterPressed::CreateRaw(this, &FCameraShakeTrackEditor::OnCameraShakeAssetEnterPressed, ObjectBinding);
 		AssetPickerConfig.bAllowNullSelection = false;
 		AssetPickerConfig.InitialAssetViewType = EAssetViewType::List;
 		AssetPickerConfig.Filter.ClassNames.Add(UBlueprint::StaticClass()->GetFName());
@@ -254,6 +255,15 @@ FKeyPropertyResult FCameraShakeTrackEditor::AddKeyInternal(FFrameNumber KeyTime,
 	return KeyPropertyResult;
 }
 
+void FCameraShakeTrackEditor::OnCameraShakeAssetEnterPressed(const TArray<FAssetData>& AssetData, FGuid ObjectBinding)
+{
+	if (AssetData.Num() > 0)
+	{
+		OnCameraShakeAssetSelected(AssetData[0].GetAsset(), ObjectBinding);
+	}
+}
+
+
 UCameraComponent* FCameraShakeTrackEditor::AcquireCameraComponentFromObjectGuid(const FGuid& Guid)
 {
 	USkeleton* Skeleton = nullptr;
diff --git a/Engine/Source/Editor/MovieSceneTools/Private/TrackEditors/CameraShakeTrackEditor.h b/Engine/Source/Editor/MovieSceneTools/Private/TrackEditors/CameraShakeTrackEditor.h
--- a/Engine/Source/Editor/MovieSceneTools/Private/TrackEditors/CameraShakeTrackEditor.h
+++ b/Engine/Source/Editor/MovieSceneTools/Private/TrackEditors/CameraShakeTrackEditor.h
@@ -63,6 +63,9 @@ private:
 	/** Animation asset selected */
 	void OnCameraShakeAssetSelected(const FAssetData& AssetData, FGuid ObjectBinding);
 
+	/** Animation asset enter pressed */
+	void OnCameraShakeAssetEnterPressed(const TArray<FAssetData>& AssetData, FGuid ObjectBinding);
+
 	class UCameraComponent* AcquireCameraComponentFromObjectGuid(const FGuid& Guid);
 };
 
diff --git a/Engine/Source/Editor/MovieSceneTools/Private/TrackEditors/CinematicShotTrackEditor.cpp b/Engine/Source/Editor/MovieSceneTools/Private/TrackEditors/CinematicShotTrackEditor.cpp
--- a/Engine/Source/Editor/MovieSceneTools/Private/TrackEditors/CinematicShotTrackEditor.cpp
+++ b/Engine/Source/Editor/MovieSceneTools/Private/TrackEditors/CinematicShotTrackEditor.cpp
@@ -17,6 +17,7 @@
 #include "EditorStyleSet.h"
 #include "LevelEditorViewport.h"
 #include "MovieSceneToolHelpers.h"
+#include "FCPXML/FCPXMLMovieSceneTranslator.h"
 #include "Sections/CinematicShotSection.h"
 #include "SequencerUtilities.h"
 #include "IAssetTools.h"
@@ -38,7 +39,6 @@
 
 #define LOCTEXT_NAMESPACE "FCinematicShotTrackEditor"
 
-
 /* FCinematicShotTrackEditor structors
  *****************************************************************************/
 
@@ -210,6 +210,8 @@ void FCinematicShotTrackEditor::Tick(float DeltaTime)
 
 void FCinematicShotTrackEditor::BuildTrackContextMenu( FMenuBuilder& MenuBuilder, UMovieSceneTrack* Track )
 {
+	MenuBuilder.BeginSection("Import/Export", NSLOCTEXT("Sequencer", "ImportExportMenuSectionName", "Import/Export"));
+
 	MenuBuilder.AddMenuEntry(
 		NSLOCTEXT( "Sequencer", "ImportEDL", "Import EDL..." ),
 		NSLOCTEXT( "Sequencer", "ImportEDLTooltip", "Import Edit Decision List (EDL) for non-linear editors." ),
@@ -223,6 +225,22 @@ void FCinematicShotTrackEditor::BuildTrackContextMenu( FMenuBuilder& MenuBuilder
 		FSlateIcon(),
 		FUIAction(
 		FExecuteAction::CreateRaw(this, &FCinematicShotTrackEditor::ExportEDL )));
+
+	MenuBuilder.AddMenuEntry(
+		NSLOCTEXT("Sequencer", "ImportFCPXML", "Import Final Cut Pro 7 XML..."),
+		NSLOCTEXT("Sequencer", "ImportFCPXMLTooltip", "Import Final Cut Pro 7 XML file for non-linear editors."),
+		FSlateIcon(),
+		FUIAction(
+			FExecuteAction::CreateRaw(this, &FCinematicShotTrackEditor::ImportFCPXML )));
+
+	MenuBuilder.AddMenuEntry(
+		NSLOCTEXT("Sequencer", "ExportFCPXML", "Export Final Cut Pro 7 XML..."),
+		NSLOCTEXT("Sequencer", "ExportFCPXMLTooltip", "Export Final Cut Pro 7 XML file for non-linear editors."),
+		FSlateIcon(),
+		FUIAction(
+			FExecuteAction::CreateRaw(this, &FCinematicShotTrackEditor::ExportFCPXML )));
+
+	MenuBuilder.EndSection();
 }
 
 
@@ -567,6 +585,7 @@ TSharedRef<SWidget> FCinematicShotTrackEditor::HandleAddCinematicShotComboButton
 	FAssetPickerConfig AssetPickerConfig;
 	{
 		AssetPickerConfig.OnAssetSelected = FOnAssetSelected::CreateRaw( this, &FCinematicShotTrackEditor::HandleAddCinematicShotComboButtonMenuEntryExecute);
+		AssetPickerConfig.OnAssetEnterPressed = FOnAssetEnterPressed::CreateRaw( this, &FCinematicShotTrackEditor::HandleAddCinematicShotComboButtonMenuEntryEnterPressed);
 		AssetPickerConfig.bAllowNullSelection = false;
 		AssetPickerConfig.InitialAssetViewType = EAssetViewType::Tile;
 		AssetPickerConfig.Filter.ClassNames.Add(TEXT("LevelSequence"));
@@ -602,6 +621,13 @@ void FCinematicShotTrackEditor::HandleAddCinematicShotComboButtonMenuEntryExecut
 	}
 }
 
+void FCinematicShotTrackEditor::HandleAddCinematicShotComboButtonMenuEntryEnterPressed(const TArray<FAssetData>& AssetData)
+{
+	if (AssetData.Num() > 0)
+	{
+		HandleAddCinematicShotComboButtonMenuEntryExecute(AssetData[0].GetAsset());
+	}
+}
 
 FKeyPropertyResult FCinematicShotTrackEditor::AddKeyInternal(FFrameNumber KeyTime, UMovieSceneSequence* InMovieSceneSequence, int32 RowIndex)
 {	
@@ -613,9 +639,9 @@ FKeyPropertyResult FCinematicShotTrackEditor::AddKeyInternal(FFrameNumber KeyTim
 		
 		const FQualifiedFrameTime InnerDuration = FQualifiedFrameTime(
 			MovieScene::DiscreteSize(InMovieSceneSequence->GetMovieScene()->GetPlaybackRange()),
-			InMovieSceneSequence->GetMovieScene()->GetFrameResolution());
+			InMovieSceneSequence->GetMovieScene()->GetTickResolution());
 
-		const FFrameRate OuterFrameRate = CinematicShotTrack->GetTypedOuter<UMovieScene>()->GetFrameResolution();
+		const FFrameRate OuterFrameRate = CinematicShotTrack->GetTypedOuter<UMovieScene>()->GetTickResolution();
 		const int32      OuterDuration  = InnerDuration.ConvertTo(OuterFrameRate).FrameNumber.Value;
 
 		UMovieSceneSubSection* NewSection = CinematicShotTrack->AddSequenceOnRow(InMovieSceneSequence, KeyTime, OuterDuration, RowIndex);
@@ -762,9 +788,9 @@ FKeyPropertyResult FCinematicShotTrackEditor::HandleSequenceAdded(FFrameNumber K
 
 	const FQualifiedFrameTime InnerDuration = FQualifiedFrameTime(
 		MovieScene::DiscreteSize(Sequence->GetMovieScene()->GetPlaybackRange()),
-		Sequence->GetMovieScene()->GetFrameResolution());
+		Sequence->GetMovieScene()->GetTickResolution());
 
-	const FFrameRate OuterFrameRate = CinematicShotTrack->GetTypedOuter<UMovieScene>()->GetFrameResolution();
+	const FFrameRate OuterFrameRate = CinematicShotTrack->GetTypedOuter<UMovieScene>()->GetTickResolution();
 	const int32      OuterDuration  = InnerDuration.ConvertTo(OuterFrameRate).FrameNumber.Value;
 
 	UMovieSceneSubSection* NewSection = CinematicShotTrack->AddSequenceOnRow(Sequence, KeyTime, OuterDuration, RowIndex);
@@ -807,7 +833,7 @@ void FCinematicShotTrackEditor::ImportEDL()
 	const FMovieSceneCaptureSettings& Settings = MovieSceneCapture->GetSettings();
 	FString SaveDirectory = FPaths::ConvertRelativePathToFull(Settings.OutputDirectory.Path);
 
-	if (MovieSceneToolHelpers::ShowImportEDLDialog(MovieScene, Settings.FrameRate, SaveDirectory))
+	if (MovieSceneToolHelpers::ShowImportEDLDialog(MovieScene, MovieScene->GetDisplayRate(), SaveDirectory))
 	{
 		GetSequencer()->NotifyMovieSceneDataChanged( EMovieSceneDataChangeType::MovieSceneStructureItemsChanged );
 	}
@@ -844,9 +870,88 @@ void FCinematicShotTrackEditor::ExportEDL()
 	FString SaveDirectory = FPaths::ConvertRelativePathToFull(Settings.OutputDirectory.Path);
 	int32 HandleFrames = Settings.HandleFrames;
 
-	// @todo: sequencer-timecode: resolution vs playrate
-	MovieSceneToolHelpers::ShowExportEDLDialog(MovieScene, MovieScene->GetFrameResolution(), SaveDirectory, HandleFrames);
+	MovieSceneToolHelpers::ShowExportEDLDialog(MovieScene, MovieScene->GetDisplayRate(), SaveDirectory, HandleFrames);
 }
 
 
+void FCinematicShotTrackEditor::ImportFCPXML()
+{
+	UMovieSceneSequence* FocusedSequence = GetSequencer()->GetFocusedMovieSceneSequence();
+	if (!FocusedSequence)
+	{
+		return;
+	}
+
+	UMovieScene* MovieScene = FocusedSequence->GetMovieScene();
+	if (!MovieScene)
+	{
+		return;
+	}
+
+	UAutomatedLevelSequenceCapture* MovieSceneCapture = Cast<UAutomatedLevelSequenceCapture>(IMovieSceneCaptureModule::Get().GetFirstActiveMovieSceneCapture());
+	if (!MovieSceneCapture)
+	{
+		MovieSceneCapture = NewObject<UAutomatedLevelSequenceCapture>(GetTransientPackage(), UAutomatedLevelSequenceCapture::StaticClass(), NAME_None, RF_Transient);
+		MovieSceneCapture->LoadFromConfig();
+	}
+
+	if (!MovieSceneCapture)
+	{
+		return;
+	}
+
+	const FMovieSceneCaptureSettings& Settings = MovieSceneCapture->GetSettings();
+	FString SaveDirectory = FPaths::ConvertRelativePathToFull(Settings.OutputDirectory.Path);
+
+	FFCPXMLImporter *Importer = new FFCPXMLImporter;
+
+	if (MovieSceneToolHelpers::MovieSceneTranslatorImport(Importer, MovieScene, MovieScene->GetDisplayRate(), SaveDirectory))
+	{
+		GetSequencer()->NotifyMovieSceneDataChanged(EMovieSceneDataChangeType::MovieSceneStructureItemsChanged);
+	}
+
+	delete Importer;
+}
+
+
+void FCinematicShotTrackEditor::ExportFCPXML()
+{
+	UMovieSceneSequence* FocusedSequence = GetSequencer()->GetFocusedMovieSceneSequence();
+	if (!FocusedSequence)
+	{
+		return;
+	}
+
+	const UMovieScene* MovieScene = FocusedSequence->GetMovieScene();
+	if (!MovieScene)
+	{
+		return;
+	}
+
+	UAutomatedLevelSequenceCapture* MovieSceneCapture = Cast<UAutomatedLevelSequenceCapture>(IMovieSceneCaptureModule::Get().GetFirstActiveMovieSceneCapture());
+	if (!MovieSceneCapture)
+	{
+		MovieSceneCapture = NewObject<UAutomatedLevelSequenceCapture>(GetTransientPackage(), UAutomatedLevelSequenceCapture::StaticClass(), NAME_None, RF_Transient);
+		MovieSceneCapture->LoadFromConfig();
+	}
+
+	if (!MovieSceneCapture)
+	{
+		return;
+	}
+
+	const FMovieSceneCaptureSettings& Settings = MovieSceneCapture->GetSettings();
+	FString SaveDirectory = FPaths::ConvertRelativePathToFull(Settings.OutputDirectory.Path);
+	int32 HandleFrames = Settings.HandleFrames;
+
+	FFCPXMLExporter *Exporter = new FFCPXMLExporter;
+
+	MovieSceneToolHelpers::MovieSceneTranslatorExport(Exporter, MovieScene, MovieScene->GetDisplayRate(), SaveDirectory, HandleFrames);
+
+	delete Exporter;
+}
+
+
+
+
 #undef LOCTEXT_NAMESPACE
diff --git a/Engine/Source/Editor/MovieSceneTools/Private/TrackEditors/CinematicShotTrackEditor.h b/Engine/Source/Editor/MovieSceneTools/Private/TrackEditors/CinematicShotTrackEditor.h
--- a/Engine/Source/Editor/MovieSceneTools/Private/TrackEditors/CinematicShotTrackEditor.h
+++ b/Engine/Source/Editor/MovieSceneTools/Private/TrackEditors/CinematicShotTrackEditor.h
@@ -138,6 +138,9 @@ private:
 	/** Callback for executing a menu entry in the "Add Shot" combo button. */
 	void HandleAddCinematicShotComboButtonMenuEntryExecute(const FAssetData& AssetData);
 
+	/** Callback for executing a menu entry in the "Add Shot" combo button when enter pressed. */
+	void HandleAddCinematicShotComboButtonMenuEntryEnterPressed(const TArray<FAssetData>& AssetData);
+
 	/** Find or create a cinematic shot track in the currently focused movie scene. */
 	UMovieSceneCinematicShotTrack* FindOrCreateCinematicShotTrack();
 
@@ -176,6 +179,12 @@ private:
 	/** Callback for ExportEDL. */
 	void ExportEDL();
 
+	/** Callback for ImportFCPXML. */
+	void ImportFCPXML();
+
+	/** Callback for ExportFCPXML. */
+	void ExportFCPXML();
+
 private:
 
 	/** The Thumbnail pool which draws all the viewport thumbnails for the shot track. */
diff --git a/Engine/Source/Editor/MovieSceneTools/Private/TrackEditors/EventTrackEditor.cpp b/Engine/Source/Editor/MovieSceneTools/Private/TrackEditors/EventTrackEditor.cpp
--- a/Engine/Source/Editor/MovieSceneTools/Private/TrackEditors/EventTrackEditor.cpp
+++ b/Engine/Source/Editor/MovieSceneTools/Private/TrackEditors/EventTrackEditor.cpp
@@ -15,6 +15,8 @@
 #include "MovieSceneObjectBindingIDCustomization.h"
 #include "DetailCategoryBuilder.h"
 #include "IDetailPropertyRow.h"
+#include "Widgets/Layout/SBox.h"
+#include "Sections/EventSection.h"
 
 #define LOCTEXT_NAMESPACE "FEventTrackEditor"
 
@@ -28,6 +30,12 @@ TSharedRef<ISequencerTrackEditor> FEventTrackEditor::CreateTrackEditor(TSharedRe
 }
 
 
+TSharedRef<ISequencerSection> FEventTrackEditor::MakeSectionInterface(UMovieSceneSection& SectionObject, UMovieSceneTrack& Track, FGuid ObjectBinding)
+{
+	return MakeShared<FEventSection>(SectionObject, GetSequencer());
+}
+
+
 /* FEventTrackEditor structors
  *****************************************************************************/
 
@@ -96,7 +104,8 @@ void FEventTrackEditor::BuildTrackContextMenu(FMenuBuilder& MenuBuilder, UMovieS
 		FDetailsViewArgs DetailsViewArgs(false,false,false,FDetailsViewArgs::HideNameArea,true);
 		DetailsViewArgs.DefaultsOnlyVisibility = EEditDefaultsOnlyNodeVisibility::Automatic;
 		DetailsViewArgs.bShowOptions = false;
-		
+		DetailsViewArgs.ColumnWidth = 0.55f;
+
 		TSharedRef<IDetailsView> DetailsView = PropertyEditor.CreateDetailView(DetailsViewArgs);
 		
 		// Register the custom type layout for the class
@@ -109,7 +118,15 @@ void FEventTrackEditor::BuildTrackContextMenu(FMenuBuilder& MenuBuilder, UMovieS
 		DetailsView->SetObject(EventTrack, true);
 
 		// Add it to the menu
-		SubMenuBuilder.AddWidget(DetailsView, FText(), true, false);
+		TSharedRef< SWidget > DetailsViewWidget =
+			SNew(SBox)
+			.MaxDesiredHeight(400.0f)
+			.WidthOverride(450.0f)
+		[
+			DetailsView
+		];
+
+		SubMenuBuilder.AddWidget(DetailsViewWidget, FText(), true, false);
 	};
 
 	MenuBuilder.AddSubMenu(LOCTEXT("Properties_MenuText", "Properties"), FText(), FNewMenuDelegate::CreateLambda(PopulateSubMenu));
diff --git a/Engine/Source/Editor/MovieSceneTools/Private/TrackEditors/EventTrackEditor.h b/Engine/Source/Editor/MovieSceneTools/Private/TrackEditors/EventTrackEditor.h
--- a/Engine/Source/Editor/MovieSceneTools/Private/TrackEditors/EventTrackEditor.h
+++ b/Engine/Source/Editor/MovieSceneTools/Private/TrackEditors/EventTrackEditor.h
@@ -49,6 +49,9 @@ public:
 	virtual const FSlateBrush* GetIconBrush() const override;
 	virtual void BuildTrackContextMenu( FMenuBuilder& MenuBuilder, UMovieSceneTrack* Track ) override;
 
+	//~ FPropertyTrackEditor interface
+	virtual TSharedRef<ISequencerSection> MakeSectionInterface(UMovieSceneSection& SectionObject, UMovieSceneTrack& Track, FGuid ObjectBinding) override;
+
 private:
 
 	/** Callback for executing the "Add Event Track" menu entry. */
diff --git a/Engine/Source/Editor/MovieSceneTools/Private/TrackEditors/FadeTrackEditor.cpp b/Engine/Source/Editor/MovieSceneTools/Private/TrackEditors/FadeTrackEditor.cpp
--- a/Engine/Source/Editor/MovieSceneTools/Private/TrackEditors/FadeTrackEditor.cpp
+++ b/Engine/Source/Editor/MovieSceneTools/Private/TrackEditors/FadeTrackEditor.cpp
@@ -36,16 +36,15 @@ public:
 		const UMovieSceneFadeSection* FadeSection = Cast<const UMovieSceneFadeSection>( WeakSection.Get() );
 
 		FTimeToPixel TimeConverter    = Painter.GetTimeConverter();
-		FFrameRate   FrameResolution  = TimeConverter.GetFrameResolution();
+		FFrameRate   TickResolution   = TimeConverter.GetTickResolution();
 
-		const double HalfFrameTime    = FFrameTime(0, 0.5f) / FrameResolution;
-		const double StartTimeSeconds = TimeConverter.PixelToSeconds(1.f) - HalfFrameTime;
-		const double EndTimeSeconds   = TimeConverter.PixelToSeconds(Painter.SectionGeometry.GetLocalSize().X-2.f) - HalfFrameTime;
+		const double StartTimeSeconds = TimeConverter.PixelToSeconds(1.f);
+		const double EndTimeSeconds   = TimeConverter.PixelToSeconds(Painter.SectionGeometry.GetLocalSize().X-2.f);
 		const double TimeThreshold    = FMath::Max(0.0001, TimeConverter.PixelToSeconds(5) - TimeConverter.PixelToSeconds(0));
 		const double DurationSeconds  = EndTimeSeconds - StartTimeSeconds;
 
 		TArray<TTuple<double, double>> CurvePoints;
-		FadeSection->GetChannel().PopulateCurvePoints(StartTimeSeconds, EndTimeSeconds, TimeThreshold, 0.1f, FrameResolution, CurvePoints);
+		FadeSection->GetChannel().PopulateCurvePoints(StartTimeSeconds, EndTimeSeconds, TimeThreshold, 0.1f, TickResolution, CurvePoints);
 
 		TArray<FSlateGradientStop> GradientStops;
 		for (TTuple<double, double> Vector : CurvePoints)
diff --git a/Engine/Source/Editor/MovieSceneTools/Private/TrackEditors/MaterialParameterCollectionTrackEditor.cpp b/Engine/Source/Editor/MovieSceneTools/Private/TrackEditors/MaterialParameterCollectionTrackEditor.cpp
--- a/Engine/Source/Editor/MovieSceneTools/Private/TrackEditors/MaterialParameterCollectionTrackEditor.cpp
+++ b/Engine/Source/Editor/MovieSceneTools/Private/TrackEditors/MaterialParameterCollectionTrackEditor.cpp
@@ -35,11 +35,12 @@ TSharedRef<ISequencerSection> FMaterialParameterCollectionTrackEditor::MakeSecti
 	return MakeShareable(new FParameterSection(*ParameterSection));
 }
 
-TSharedRef<SWidget> CreateAssetPicker(FOnAssetSelected OnAssetSelected)
+TSharedRef<SWidget> CreateAssetPicker(FOnAssetSelected OnAssetSelected, FOnAssetEnterPressed OnAssetEnterPressed)
 {
 	FAssetPickerConfig AssetPickerConfig;
 	{
 		AssetPickerConfig.OnAssetSelected = OnAssetSelected;
+		AssetPickerConfig.OnAssetEnterPressed = OnAssetEnterPressed;
 		AssetPickerConfig.bAllowNullSelection = false;
 		AssetPickerConfig.InitialAssetViewType = EAssetViewType::List;
 		AssetPickerConfig.Filter.bRecursiveClasses = true;
@@ -72,9 +73,17 @@ void FMaterialParameterCollectionTrackEditor::BuildTrackContextMenu(FMenuBuilder
 		}
 	};
 
-	auto SubMenuCallback = [this, AssignAsset](FMenuBuilder& SubMenuBuilder)
+	auto AssignAssetEnterPressed = [AssignAsset](const TArray<FAssetData>& InAssetData)
 	{
-		SubMenuBuilder.AddWidget(CreateAssetPicker(FOnAssetSelected::CreateLambda(AssignAsset)), FText::GetEmpty(), true);
+		if (InAssetData.Num() > 0)
+		{
+			AssignAsset(InAssetData[0].GetAsset());
+		}
+	};
+
+	auto SubMenuCallback = [this, AssignAsset, AssignAssetEnterPressed](FMenuBuilder& SubMenuBuilder)
+	{
+		SubMenuBuilder.AddWidget(CreateAssetPicker(FOnAssetSelected::CreateLambda(AssignAsset), FOnAssetEnterPressed::CreateLambda(AssignAssetEnterPressed)), FText::GetEmpty(), true);
 	};
 
 	MenuBuilder.AddSubMenu(
@@ -89,7 +98,7 @@ void FMaterialParameterCollectionTrackEditor::BuildAddTrackMenu(FMenuBuilder& Me
 {
 	auto SubMenuCallback = [this](FMenuBuilder& SubMenuBuilder)
 	{
-		SubMenuBuilder.AddWidget(CreateAssetPicker(FOnAssetSelected::CreateRaw(this, &FMaterialParameterCollectionTrackEditor::AddTrackToSequence)), FText::GetEmpty(), true);
+		SubMenuBuilder.AddWidget(CreateAssetPicker(FOnAssetSelected::CreateRaw(this, &FMaterialParameterCollectionTrackEditor::AddTrackToSequence), FOnAssetEnterPressed::CreateRaw(this, &FMaterialParameterCollectionTrackEditor::AddTrackToSequenceEnterPressed)), FText::GetEmpty(), true);
 	};
 
 	MenuBuilder.AddSubMenu(
@@ -144,6 +153,14 @@ void FMaterialParameterCollectionTrackEditor::AddTrackToSequence(const FAssetDat
 	GetSequencer()->NotifyMovieSceneDataChanged(EMovieSceneDataChangeType::MovieSceneStructureItemAdded);
 }
 
+void FMaterialParameterCollectionTrackEditor::AddTrackToSequenceEnterPressed(const TArray<FAssetData>& InAssetData)
+{
+	if (InAssetData.Num() > 0)
+	{
+		AddTrackToSequence(InAssetData[0].GetAsset());
+	}
+}
+
 bool FMaterialParameterCollectionTrackEditor::SupportsType(TSubclassOf<UMovieSceneTrack> Type) const
 {
 	return Type == UMovieSceneMaterialParameterCollectionTrack::StaticClass();
diff --git a/Engine/Source/Editor/MovieSceneTools/Private/TrackEditors/MaterialParameterCollectionTrackEditor.h b/Engine/Source/Editor/MovieSceneTools/Private/TrackEditors/MaterialParameterCollectionTrackEditor.h
--- a/Engine/Source/Editor/MovieSceneTools/Private/TrackEditors/MaterialParameterCollectionTrackEditor.h
+++ b/Engine/Source/Editor/MovieSceneTools/Private/TrackEditors/MaterialParameterCollectionTrackEditor.h
@@ -45,4 +45,5 @@ private:
 	void AddVectorParameter(UMovieSceneMaterialParameterCollectionTrack* Track, FCollectionVectorParameter Parameter);
 
 	void AddTrackToSequence(const FAssetData& InAssetData);
+	void AddTrackToSequenceEnterPressed(const TArray<FAssetData>& InAssetData);
 };
diff --git a/Engine/Source/Editor/MovieSceneTools/Private/TrackEditors/ParticleTrackEditor.cpp b/Engine/Source/Editor/MovieSceneTools/Private/TrackEditors/ParticleTrackEditor.cpp
--- a/Engine/Source/Editor/MovieSceneTools/Private/TrackEditors/ParticleTrackEditor.cpp
+++ b/Engine/Source/Editor/MovieSceneTools/Private/TrackEditors/ParticleTrackEditor.cpp
@@ -125,7 +125,7 @@ int32 FParticleSection::OnPaintSection( FSequencerSectionPainter& InPainter ) co
 
 		for (int32 Index = 0; Index < Times.Num(); ++Index)
 		{
-			const double       Time  = Times[Index] / TimeToPixelConverter.GetFrameResolution();
+			const double       Time  = Times[Index] / TimeToPixelConverter.GetTickResolution();
 			const EParticleKey Value = (EParticleKey)Values[Index];
 
 			if ( Value == EParticleKey::Activate )
diff --git a/Engine/Source/Editor/MovieSceneTools/Private/TrackEditors/PathTrackEditor.cpp b/Engine/Source/Editor/MovieSceneTools/Private/TrackEditors/PathTrackEditor.cpp
--- a/Engine/Source/Editor/MovieSceneTools/Private/TrackEditors/PathTrackEditor.cpp
+++ b/Engine/Source/Editor/MovieSceneTools/Private/TrackEditors/PathTrackEditor.cpp
@@ -10,6 +10,7 @@
 #include "ISectionLayoutBuilder.h"
 #include "ActorEditorUtils.h"
 #include "Components/SplineComponent.h"
+#include "MovieSceneToolHelpers.h"
 
 
 #define LOCTEXT_NAMESPACE "FPathTrackEditor"
@@ -105,8 +106,13 @@ TSharedRef<ISequencerSection> F3DPathTrackEditor::MakeSectionInterface( UMovieSc
 
 void F3DPathTrackEditor::BuildObjectBindingTrackMenu(FMenuBuilder& MenuBuilder, const FGuid& ObjectBinding, const UClass* ObjectClass)
 {
-	if (ObjectClass->IsChildOf(AActor::StaticClass()))
+	if (ObjectClass && ObjectClass->IsChildOf(AActor::StaticClass()))
 	{
+		if (MovieSceneToolHelpers::HasHiddenMobility(ObjectClass))
+		{
+			return;
+		}
+
 		UMovieSceneSection* DummySection = nullptr;
 
 		MenuBuilder.AddSubMenu(
@@ -168,7 +174,7 @@ void F3DPathTrackEditor::ActorSocketPicked(const FName SocketName, USceneCompone
 		else if (ActorPickerID.ActorPicked.IsValid())
 		{
 			FGuid ParentActorId = FindOrCreateHandleToObject(ActorPickerID.ActorPicked.Get()).Handle;
-			ConstraintBindingID = FMovieSceneObjectBindingID(ParentActorId, MovieSceneSequenceID::Root);
+			ConstraintBindingID = FMovieSceneObjectBindingID(ParentActorId, GetSequencer()->GetFocusedTemplateID());
 		}
 
 		if (ConstraintBindingID.IsValid())
@@ -202,7 +208,7 @@ FKeyPropertyResult F3DPathTrackEditor::AddKeyInternal( FFrameNumber KeyTime, con
 		FFindOrCreateHandleResult HandleResult = FindOrCreateHandleToObject(ActorPickerID.ActorPicked.Get());
 		FGuid ParentActorId = HandleResult.Handle;
 		KeyPropertyResult.bHandleCreated |= HandleResult.bWasCreated;
-		ConstraintBindingID = FMovieSceneObjectBindingID(ParentActorId, MovieSceneSequenceID::Root);
+		ConstraintBindingID = FMovieSceneObjectBindingID(ParentActorId, GetSequencer()->GetFocusedTemplateID());
 	}
 
 	if (!ConstraintBindingID.IsValid())
@@ -226,7 +232,7 @@ FKeyPropertyResult F3DPathTrackEditor::AddKeyInternal( FFrameNumber KeyTime, con
 			if (ensure(Track))
 			{
 				// Clamp to next path section's start time or the end of the current sequencer view range
-				FFrameNumber PathEndTime = ( GetSequencer()->GetViewRange().GetUpperBoundValue() * Track->GetTypedOuter<UMovieScene>()->GetFrameResolution() ).FrameNumber;
+				FFrameNumber PathEndTime = ( GetSequencer()->GetViewRange().GetUpperBoundValue() * Track->GetTypedOuter<UMovieScene>()->GetTickResolution() ).FrameNumber;
 	
 				for (UMovieSceneSection* Section : Track->GetAllSections())
 				{
@@ -240,7 +246,8 @@ FKeyPropertyResult F3DPathTrackEditor::AddKeyInternal( FFrameNumber KeyTime, con
 					}
 				}
 
-				Cast<UMovieScene3DPathTrack>(Track)->AddConstraint( KeyTime, (PathEndTime - KeyTime).Value, NAME_None, NAME_None, ConstraintBindingID );
+				int32 Duration = FMath::Max(0, (PathEndTime - KeyTime).Value);
+				Cast<UMovieScene3DPathTrack>(Track)->AddConstraint( KeyTime, Duration, NAME_None, NAME_None, ConstraintBindingID );
 				KeyPropertyResult.bTrackModified = true;
 			}
 		}
