Copying //UE4/Dev-Sequencer to //UE4/Dev-Main (Source: //UE4/Dev-Sequencer @ 4044339)

diff --git a/Engine/Source/Editor/MovieSceneTools/Private/TrackEditors/SkeletalAnimationTrackEditor.cpp b/Engine/Source/Editor/MovieSceneTools/Private/TrackEditors/SkeletalAnimationTrackEditor.cpp
--- a/Engine/Source/Editor/MovieSceneTools/Private/TrackEditors/SkeletalAnimationTrackEditor.cpp
+++ b/Engine/Source/Editor/MovieSceneTools/Private/TrackEditors/SkeletalAnimationTrackEditor.cpp
@@ -28,10 +28,18 @@
 #include "ISectionLayoutBuilder.h"
 #include "Animation/AnimMontage.h"
 #include "Animation/AnimSequence.h"
+#include "Animation/PoseAsset.h"
 #include "EditorStyleSet.h"
 #include "DragAndDrop/AssetDragDropOp.h"
 #include "MovieSceneTimeHelpers.h"
-
+#include "Fonts/FontMeasure.h"
+#include "SequencerTimeSliderController.h"
+#include "AnimationEditorUtils.h"
+#include "Factories/PoseAssetFactory.h"
+#include "Misc/MessageDialog.h"
+#include "Framework/Notifications/NotificationManager.h"
+#include "Widgets/Notifications/SNotificationList.h"
+#include "Toolkits/AssetEditorManager.h"
 
 namespace SkeletalAnimationEditorConstants
 {
@@ -41,9 +49,69 @@ namespace SkeletalAnimationEditorConstants
 
 #define LOCTEXT_NAMESPACE "FSkeletalAnimationTrackEditor"
 
+USkeletalMeshComponent* AcquireSkeletalMeshFromObjectGuid(const FGuid& Guid, TSharedPtr<ISequencer> SequencerPtr)
+{
+	UObject* BoundObject = SequencerPtr.IsValid() ? SequencerPtr->FindSpawnedObjectOrTemplate(Guid) : nullptr;
+
+	if (AActor* Actor = Cast<AActor>(BoundObject))
+	{
+		TInlineComponentArray<USkeletalMeshComponent*> SkeletalMeshComponents;
+		Actor->GetComponents(SkeletalMeshComponents);
+
+		for (int32 j = 0; j <SkeletalMeshComponents.Num(); ++j)
+		{
+			USkeletalMeshComponent* SkeletalMeshComp = SkeletalMeshComponents[j];
+			if (SkeletalMeshComp)
+			{
+				return SkeletalMeshComp;
+			}
+		}
+	}
+	else if(USkeletalMeshComponent* SkeletalMeshComponent = Cast<USkeletalMeshComponent>(BoundObject))
+	{
+		if (SkeletalMeshComponent->SkeletalMesh)
+		{
+			return SkeletalMeshComponent;
+		}
+	}
+
+	return nullptr;
+}
+
+USkeleton* AcquireSkeletonFromObjectGuid(const FGuid& Guid, TSharedPtr<ISequencer> SequencerPtr)
+{
+	UObject* BoundObject = SequencerPtr.IsValid() ? SequencerPtr->FindSpawnedObjectOrTemplate(Guid) : nullptr;
+
+	if (AActor* Actor = Cast<AActor>(BoundObject))
+	{
+		TInlineComponentArray<USkeletalMeshComponent*> SkeletalMeshComponents;
+		Actor->GetComponents(SkeletalMeshComponents);
+
+		for (int32 j = 0; j <SkeletalMeshComponents.Num(); ++j)
+		{
+			USkeletalMeshComponent* SkeletalMeshComp = SkeletalMeshComponents[j];
+			if (SkeletalMeshComp->SkeletalMesh && SkeletalMeshComp->SkeletalMesh->Skeleton)
+			{
+				// @todo Multiple actors, multiple components
+				return SkeletalMeshComp->SkeletalMesh->Skeleton;
+			}
+		}
+	}
+	else if(USkeletalMeshComponent* SkeletalMeshComponent = Cast<USkeletalMeshComponent>(BoundObject))
+	{
+		if (SkeletalMeshComponent->SkeletalMesh)
+		{
+			return SkeletalMeshComponent->SkeletalMesh->Skeleton;
+		}
+	}
+
+	return nullptr;
+}
+
 
-FSkeletalAnimationSection::FSkeletalAnimationSection( UMovieSceneSection& InSection )
+FSkeletalAnimationSection::FSkeletalAnimationSection( UMovieSceneSection& InSection, TWeakPtr<ISequencer> InSequencer)
 	: Section(*CastChecked<UMovieSceneSkeletalAnimationSection>(&InSection))
+	, Sequencer(InSequencer)
 	, InitialStartOffsetDuringResize(0.f)
 	, InitialStartTimeDuringResize(0)
 { }
@@ -90,12 +158,12 @@ int32 FSkeletalAnimationSection::OnPaintSection( FSequencerSectionPainter& Paint
 	float AnimPlayRate = FMath::IsNearlyZero(Section.Params.PlayRate) ? 1.0f : Section.Params.PlayRate;
 	float SeqLength = (Section.Params.GetSequenceLength() - (Section.Params.StartOffset + Section.Params.EndOffset)) / AnimPlayRate;
 
-	FFrameRate FrameResolution = TimeToPixelConverter.GetFrameResolution();
+	FFrameRate TickResolution = TimeToPixelConverter.GetTickResolution();
 	if (!FMath::IsNearlyZero(SeqLength, KINDA_SMALL_NUMBER) && SeqLength > 0)
 	{
-		float MaxOffset  = Section.GetRange().Size<FFrameTime>() / FrameResolution;
+		float MaxOffset  = Section.GetRange().Size<FFrameTime>() / TickResolution;
 		float OffsetTime = SeqLength;
-		float StartTime  = Section.GetInclusiveStartFrame() / FrameResolution;
+		float StartTime  = Section.GetInclusiveStartFrame() / TickResolution;
 
 		while (OffsetTime < MaxOffset)
 		{
@@ -116,6 +184,56 @@ int32 FSkeletalAnimationSection::OnPaintSection( FSequencerSectionPainter& Paint
 		}
 	}
 
+	TSharedPtr<ISequencer> SequencerPtr = Sequencer.Pin();
+	if (Painter.bIsSelected && SequencerPtr.IsValid())
+	{
+		FFrameTime CurrentTime = SequencerPtr->GetLocalTime().Time;
+		if (Section.GetRange().Contains(CurrentTime.FrameNumber) && Section.Params.Animation != nullptr)
+		{
+			const float Time = TimeToPixelConverter.FrameToPixel(CurrentTime); 
+
+			// Draw the current time next to the scrub handle
+			const float AnimTime = Section.MapTimeToAnimation(CurrentTime, TickResolution);
+			int32 FrameTime = Section.Params.Animation->GetFrameAtTime(AnimTime);
+			FString FrameString = FString::FromInt(FrameTime);
+
+			const FSlateFontInfo SmallLayoutFont = FCoreStyle::GetDefaultFontStyle("Bold", 10);
+			const TSharedRef< FSlateFontMeasure > FontMeasureService = FSlateApplication::Get().GetRenderer()->GetFontMeasureService();
+			FVector2D TextSize = FontMeasureService->Measure(FrameString, SmallLayoutFont);
+
+			// Flip the text position if getting near the end of the view range
+			static const float TextOffsetPx = 10.f;
+			bool  bDrawLeft = (Painter.SectionGeometry.Size.X - Time) < (TextSize.X + 22.f) - TextOffsetPx;
+			float TextPosition = bDrawLeft ? Time - TextSize.X - TextOffsetPx : Time + TextOffsetPx;
+			//handle mirrored labels
+			const float MajorTickHeight = 9.0f; 
+			FVector2D TextOffset(TextPosition, Painter.SectionGeometry.Size.Y - (MajorTickHeight + TextSize.Y));
+
+			const FLinearColor DrawColor = FEditorStyle::GetSlateColor("SelectionColor").GetColor(FWidgetStyle());
+			const FVector2D BoxPadding = FVector2D(4.0f, 2.0f);
+			// draw time string
+	
+			FSlateDrawElement::MakeBox(
+				Painter.DrawElements,
+				LayerId + 5,
+				Painter.SectionGeometry.ToPaintGeometry(TextOffset - BoxPadding, TextSize + 2.0f * BoxPadding),
+				FEditorStyle::GetBrush("WhiteBrush"),
+				ESlateDrawEffect::None,
+				FLinearColor::Black.CopyWithNewOpacity(0.5f)
+			);
+
+			FSlateDrawElement::MakeText(
+				Painter.DrawElements,
+				LayerId + 6,
+				Painter.SectionGeometry.ToPaintGeometry(TextOffset, TextSize),
+				FrameString,
+				SmallLayoutFont,
+				DrawEffects,
+				DrawColor
+			);
+		}
+	}
+	
 	return LayerId;
 }
 
@@ -130,7 +248,7 @@ void FSkeletalAnimationSection::ResizeSection(ESequencerSectionResizeMode Resize
 	// Adjust the start offset when resizing from the beginning
 	if (ResizeMode == SSRM_LeadingEdge)
 	{
-		FFrameRate FrameRate   = Section.GetTypedOuter<UMovieScene>()->GetFrameResolution();
+		FFrameRate FrameRate   = Section.GetTypedOuter<UMovieScene>()->GetTickResolution();
 		float      StartOffset = (ResizeTime - InitialStartTimeDuringResize) / FrameRate * Section.Params.PlayRate;
 
 		StartOffset += InitialStartOffsetDuringResize;
@@ -156,7 +274,7 @@ void FSkeletalAnimationSection::BeginSlipSection()
 
 void FSkeletalAnimationSection::SlipSection(double SlipTime)
 {
-	float StartOffset = (SlipTime - InitialStartTimeDuringResize / Section.GetTypedOuter<UMovieScene>()->GetFrameResolution()) * Section.Params.PlayRate;
+	float StartOffset = (SlipTime - InitialStartTimeDuringResize / Section.GetTypedOuter<UMovieScene>()->GetTickResolution()) * Section.Params.PlayRate;
 	StartOffset += InitialStartOffsetDuringResize;
 
 	// Ensure start offset is not less than 0
@@ -170,6 +288,88 @@ void FSkeletalAnimationSection::SlipSection(double SlipTime)
 	ISequencerSection::SlipSection(SlipTime);
 }
 
+bool FSkeletalAnimationSection::CreatePoseAsset(const TArray<UObject*> NewAssets, FGuid InObjectBinding)
+{
+	USkeletalMeshComponent* SkeletalMeshComponent = AcquireSkeletalMeshFromObjectGuid(InObjectBinding, Sequencer.Pin());
+
+	bool bResult = false;
+	if (NewAssets.Num() > 0)
+	{
+		for (auto NewAsset : NewAssets)
+		{
+			UPoseAsset* NewPoseAsset = Cast<UPoseAsset>(NewAsset);
+			if (NewPoseAsset)
+			{
+				NewPoseAsset->AddOrUpdatePoseWithUniqueName(SkeletalMeshComponent);
+				bResult = true;
+			}
+		}
+
+		// if it contains error, warn them
+		if (bResult)
+		{				
+			FText NotificationText;
+			if (NewAssets.Num() == 1)
+			{
+				NotificationText = FText::Format(LOCTEXT("CreatePoseAsset", "{0} Pose assets created."), NewAssets.Num());
+			}
+			else
+			{
+				NotificationText = FText::Format(LOCTEXT("CreatePoseAsset", "Pose asset created: '{0}'."), FText::FromString(NewAssets[0]->GetName()));
+			}
+						
+			FNotificationInfo Info(NotificationText);	
+			Info.ExpireDuration = 8.0f;	
+			Info.bUseLargeFont = false;
+			Info.Hyperlink = FSimpleDelegate::CreateLambda([NewAssets]()
+			{
+				FAssetEditorManager::Get().OpenEditorForAssets(NewAssets);
+			});
+			Info.HyperlinkText = FText::Format(LOCTEXT("OpenNewPoseAssetHyperlink", "Open {0}"), FText::FromString(NewAssets[0]->GetName()));
+				
+			TSharedPtr<SNotificationItem> Notification = FSlateNotificationManager::Get().AddNotification(Info);
+			if ( Notification.IsValid() )
+			{
+				Notification->SetCompletionState( SNotificationItem::CS_Success );
+			}
+		}
+		else
+		{
+			FMessageDialog::Open(EAppMsgType::Ok, LOCTEXT("FailedToCreateAsset", "Failed to create asset"));
+		}
+	}
+	return bResult;
+}
+
+
+void FSkeletalAnimationSection::HandleCreatePoseAsset(FGuid InObjectBinding)
+{
+	USkeleton* Skeleton = AcquireSkeletonFromObjectGuid(InObjectBinding, Sequencer.Pin());
+	if (Skeleton)
+	{
+		TArray<TWeakObjectPtr<UObject>> Skeletons;
+		Skeletons.Add(Skeleton);
+		AnimationEditorUtils::ExecuteNewAnimAsset<UPoseAssetFactory, UPoseAsset>(Skeletons, FString("_PoseAsset"), FAnimAssetCreated::CreateSP(this, &FSkeletalAnimationSection::CreatePoseAsset, InObjectBinding), false);
+	}
+}
+
+
+void FSkeletalAnimationSection::BuildSectionContextMenu(FMenuBuilder& MenuBuilder, const FGuid& InObjectBinding)
+{
+	MenuBuilder.BeginSection(NAME_None, LOCTEXT("SkeletonMenuText", "Skeleton"));
+
+	MenuBuilder.AddMenuEntry(
+		LOCTEXT("CreatePoseAsset", "Create Pose Asset"), 
+		LOCTEXT("CreatePoseAsset_ToolTip", "Create Animation from current pose"),
+		FSlateIcon(), 
+		FUIAction(FExecuteAction::CreateRaw(this, &FSkeletalAnimationSection::HandleCreatePoseAsset, InObjectBinding)),
+		NAME_None, 
+		EUserInterfaceActionType::Button);
+
+	MenuBuilder.EndSection();
+}
+
+
 FSkeletalAnimationTrackEditor::FSkeletalAnimationTrackEditor( TSharedRef<ISequencer> InSequencer )
 	: FMovieSceneTrackEditor( InSequencer ) 
 { }
@@ -191,13 +391,13 @@ TSharedRef<ISequencerSection> FSkeletalAnimationTrackEditor::MakeSectionInterfac
 {
 	check( SupportsType( SectionObject.GetOuter()->GetClass() ) );
 	
-	return MakeShareable( new FSkeletalAnimationSection(SectionObject) );
+	return MakeShareable( new FSkeletalAnimationSection(SectionObject, GetSequencer()) );
 }
 
 
 void FSkeletalAnimationTrackEditor::AddKey(const FGuid& ObjectGuid)
 {
-	USkeleton* Skeleton = AcquireSkeletonFromObjectGuid(ObjectGuid);
+	USkeleton* Skeleton = AcquireSkeletonFromObjectGuid(ObjectGuid, GetSequencer());
 
 	if (Skeleton)
 	{
@@ -236,7 +436,7 @@ bool FSkeletalAnimationTrackEditor::HandleAssetAdded(UObject* Asset, const FGuid
 		
 		if (TargetObjectGuid.IsValid())
 		{
-			USkeleton* Skeleton = AcquireSkeletonFromObjectGuid(TargetObjectGuid);
+			USkeleton* Skeleton = AcquireSkeletonFromObjectGuid(TargetObjectGuid, GetSequencer());
 
 			if (Skeleton && Skeleton == AnimSequence->GetSkeleton())
 			{
@@ -263,7 +463,7 @@ void FSkeletalAnimationTrackEditor::BuildObjectBindingTrackMenu(FMenuBuilder& Me
 	{
 		const TSharedPtr<ISequencer> ParentSequencer = GetSequencer();
 
-		USkeleton* Skeleton = AcquireSkeletonFromObjectGuid(ObjectBinding);
+		USkeleton* Skeleton = AcquireSkeletonFromObjectGuid(ObjectBinding, GetSequencer());
 
 		if (Skeleton)
 		{
@@ -319,6 +519,7 @@ void FSkeletalAnimationTrackEditor::AddAnimationSubMenu(FMenuBuilder& MenuBuilde
 	FAssetPickerConfig AssetPickerConfig;
 	{
 		AssetPickerConfig.OnAssetSelected = FOnAssetSelected::CreateRaw( this, &FSkeletalAnimationTrackEditor::OnAnimationAssetSelected, ObjectBinding, Track);
+		AssetPickerConfig.OnAssetEnterPressed = FOnAssetEnterPressed::CreateRaw( this, &FSkeletalAnimationTrackEditor::OnAnimationAssetEnterPressed, ObjectBinding, Track);
 		AssetPickerConfig.bAllowNullSelection = false;
 		AssetPickerConfig.InitialAssetViewType = EAssetViewType::List;
 		AssetPickerConfig.OnShouldFilterAsset = FOnShouldFilterAsset::CreateRaw(this, &FSkeletalAnimationTrackEditor::ShouldFilterAsset);
@@ -357,6 +558,14 @@ void FSkeletalAnimationTrackEditor::OnAnimationAssetSelected(const FAssetData& A
 	}
 }
 
+void FSkeletalAnimationTrackEditor::OnAnimationAssetEnterPressed(const TArray<FAssetData>& AssetData, FGuid ObjectBinding, UMovieSceneTrack* Track)
+{
+	if (AssetData.Num() > 0)
+	{
+		OnAnimationAssetSelected(AssetData[0].GetAsset(), ObjectBinding, Track);
+	}
+}
+
 
 FKeyPropertyResult FSkeletalAnimationTrackEditor::AddKeyInternal( FFrameNumber KeyTime, UObject* Object, class UAnimSequenceBase* AnimSequence, UMovieSceneTrack* Track, int32 RowIndex )
 {
@@ -389,39 +598,6 @@ FKeyPropertyResult FSkeletalAnimationTrackEditor::AddKeyInternal( FFrameNumber K
 	return KeyPropertyResult;
 }
 
-
-USkeleton* FSkeletalAnimationTrackEditor::AcquireSkeletonFromObjectGuid(const FGuid& Guid)
-{
-	TSharedPtr<ISequencer> SequencerPtr = GetSequencer();
-	UObject* BoundObject = SequencerPtr.IsValid() ? SequencerPtr->FindSpawnedObjectOrTemplate(Guid) : nullptr;
-
-	if (AActor* Actor = Cast<AActor>(BoundObject))
-	{
-		TInlineComponentArray<USkeletalMeshComponent*> SkeletalMeshComponents;
-		Actor->GetComponents(SkeletalMeshComponents);
-
-		for (int32 j = 0; j <SkeletalMeshComponents.Num(); ++j)
-		{
-			USkeletalMeshComponent* SkeletalMeshComp = SkeletalMeshComponents[j];
-			if (SkeletalMeshComp->SkeletalMesh && SkeletalMeshComp->SkeletalMesh->Skeleton)
-			{
-				// @todo Multiple actors, multiple components
-				return SkeletalMeshComp->SkeletalMesh->Skeleton;
-			}
-		}
-	}
-	else if(USkeletalMeshComponent* SkeletalMeshComponent = Cast<USkeletalMeshComponent>(BoundObject))
-	{
-		if (SkeletalMeshComponent->SkeletalMesh)
-		{
-			return SkeletalMeshComponent->SkeletalMesh->Skeleton;
-		}
-	}
-
-	return nullptr;
-}
-
-
 void CopyInterpAnimControlTrack(TSharedRef<ISequencer> Sequencer, UInterpTrackAnimControl* MatineeAnimControlTrack, UMovieSceneSkeletalAnimationTrack* SkeletalAnimationTrack)
 {
 	FFrameNumber EndPlaybackRange = MovieScene::DiscreteExclusiveUpper(Sequencer.Get().GetFocusedMovieSceneSequence()->GetMovieScene()->GetPlaybackRange());
@@ -455,7 +631,7 @@ void FSkeletalAnimationTrackEditor::BuildTrackContextMenu( FMenuBuilder& MenuBui
 
 TSharedPtr<SWidget> FSkeletalAnimationTrackEditor::BuildOutlinerEditWidget(const FGuid& ObjectBinding, UMovieSceneTrack* Track, const FBuildEditWidgetParams& Params)
 {
-	USkeleton* Skeleton = AcquireSkeletonFromObjectGuid(ObjectBinding);
+	USkeleton* Skeleton = AcquireSkeletonFromObjectGuid(ObjectBinding, GetSequencer());
 
 	if (Skeleton)
 	{
@@ -496,7 +672,7 @@ bool FSkeletalAnimationTrackEditor::OnAllowDrop(const FDragDropEvent& DragDropEv
 		return false;
 	}
 
-	USkeleton* Skeleton = AcquireSkeletonFromObjectGuid(TargetObjectGuid);
+	USkeleton* Skeleton = AcquireSkeletonFromObjectGuid(TargetObjectGuid, GetSequencer());
 
 	TSharedPtr<FAssetDragDropOp> DragDropOp = StaticCastSharedPtr<FAssetDragDropOp>( Operation );
 
@@ -533,7 +709,7 @@ FReply FSkeletalAnimationTrackEditor::OnDrop(const FDragDropEvent& DragDropEvent
 		return FReply::Unhandled();
 	}
 
-	USkeleton* Skeleton = AcquireSkeletonFromObjectGuid(TargetObjectGuid);
+	USkeleton* Skeleton = AcquireSkeletonFromObjectGuid(TargetObjectGuid, GetSequencer());
 
 	TSharedPtr<FAssetDragDropOp> DragDropOp = StaticCastSharedPtr<FAssetDragDropOp>( Operation );
 
diff --git a/Engine/Source/Editor/MovieSceneTools/Private/TrackEditors/SkeletalAnimationTrackEditor.h b/Engine/Source/Editor/MovieSceneTools/Private/TrackEditors/SkeletalAnimationTrackEditor.h
--- a/Engine/Source/Editor/MovieSceneTools/Private/TrackEditors/SkeletalAnimationTrackEditor.h
+++ b/Engine/Source/Editor/MovieSceneTools/Private/TrackEditors/SkeletalAnimationTrackEditor.h
@@ -68,11 +68,11 @@ private:
 	/** Animation asset selected */
 	void OnAnimationAssetSelected(const FAssetData& AssetData, FGuid ObjectBinding, UMovieSceneTrack* Track);
 
+	/** Animation asset enter pressed */
+	void OnAnimationAssetEnterPressed(const TArray<FAssetData>& AssetData, FGuid ObjectBinding, UMovieSceneTrack* Track);
+
 	/** Delegate for AnimatablePropertyChanged in AddKey */
 	FKeyPropertyResult AddKeyInternal(FFrameNumber KeyTime, UObject* Object, class UAnimSequenceBase* AnimSequence, UMovieSceneTrack* Track, int32 RowIndex);
-
-	/** Gets a skeleton from an object guid in the movie scene */
-	class USkeleton* AcquireSkeletonFromObjectGuid(const FGuid& Guid);
 };
 
 
@@ -84,7 +84,7 @@ class FSkeletalAnimationSection
 public:
 
 	/** Constructor. */
-	FSkeletalAnimationSection( UMovieSceneSection& InSection );
+	FSkeletalAnimationSection( UMovieSceneSection& InSection, TWeakPtr<ISequencer> InSequencer);
 
 	/** Virtual destructor. */
 	virtual ~FSkeletalAnimationSection() { }
@@ -101,12 +101,21 @@ public:
 	virtual void ResizeSection(ESequencerSectionResizeMode ResizeMode, FFrameNumber ResizeTime) override;
 	virtual void BeginSlipSection() override;
 	virtual void SlipSection(double SlipTime) override;
+	virtual void BuildSectionContextMenu(FMenuBuilder& MenuBuilder, const FGuid& InObjectBinding) override;
+
+private:
+
+	bool CreatePoseAsset(const TArray<UObject*> NewAssets, FGuid InObjectBinding);
+	void HandleCreatePoseAsset(FGuid InObjectBinding);
 
 private:
 
 	/** The section we are visualizing */
 	UMovieSceneSkeletalAnimationSection& Section;
 
+	/** Used to draw animation frame, need selection state and local time*/
+	TWeakPtr<ISequencer> Sequencer;
+
 	/** Cached start offset value valid only during resize */
 	float InitialStartOffsetDuringResize;
 	
diff --git a/Engine/Source/Editor/MovieSceneTools/Private/TrackEditors/SubTrackEditor.cpp b/Engine/Source/Editor/MovieSceneTools/Private/TrackEditors/SubTrackEditor.cpp
--- a/Engine/Source/Editor/MovieSceneTools/Private/TrackEditors/SubTrackEditor.cpp
+++ b/Engine/Source/Editor/MovieSceneTools/Private/TrackEditors/SubTrackEditor.cpp
@@ -281,7 +281,7 @@ public:
 	virtual void BeginSlipSection() override
 	{
 		InitialStartOffsetDuringResize = SectionObject.Parameters.GetStartFrameOffset();
-		InitialStartTimeDuringResize   = ( SectionObject.HasStartFrame() ? SectionObject.GetInclusiveStartFrame() : 0 ) / SectionObject.GetTypedOuter<UMovieScene>()->GetFrameResolution();
+		InitialStartTimeDuringResize   = ( SectionObject.HasStartFrame() ? SectionObject.GetInclusiveStartFrame() : 0 ) / SectionObject.GetTypedOuter<UMovieScene>()->GetTickResolution();
 	}
 
 	virtual void SlipSection(double SlipTime) override
@@ -291,7 +291,7 @@ public:
 		// Adjust the start offset when resizing from the beginning
 		if (InnerSequence)
 		{
-			const FFrameRate    InnerFrameRate   = InnerSequence->GetMovieScene()->GetFrameResolution();
+			const FFrameRate    InnerFrameRate   = InnerSequence->GetMovieScene()->GetTickResolution();
 			const double        ResizeDifference = SlipTime - InitialStartTimeDuringResize;
 			const int32         NewStartOffset   = ( (ResizeDifference * SectionObject.Parameters.TimeScale) * InnerFrameRate ).FrameNumber.Value;
 
@@ -629,6 +629,7 @@ TSharedRef<SWidget> FSubTrackEditor::HandleAddSubSequenceComboButtonGetMenuConte
 		FAssetPickerConfig AssetPickerConfig;
 		{
 			AssetPickerConfig.OnAssetSelected = FOnAssetSelected::CreateRaw( this, &FSubTrackEditor::HandleAddSubSequenceComboButtonMenuEntryExecute, InTrack);
+			AssetPickerConfig.OnAssetEnterPressed = FOnAssetEnterPressed::CreateRaw( this, &FSubTrackEditor::HandleAddSubSequenceComboButtonMenuEntryEnterPressed, InTrack);
 			AssetPickerConfig.bAllowNullSelection = false;
 			AssetPickerConfig.InitialAssetViewType = EAssetViewType::Tile;
 			AssetPickerConfig.Filter.ClassNames.Add(TEXT("LevelSequence"));
@@ -665,6 +666,14 @@ void FSubTrackEditor::HandleAddSubSequenceComboButtonMenuEntryExecute(const FAss
 	}
 }
 
+void FSubTrackEditor::HandleAddSubSequenceComboButtonMenuEntryEnterPressed(const TArray<FAssetData>& AssetData, UMovieSceneTrack* InTrack)
+{
+	if (AssetData.Num() > 0)
+	{
+		HandleAddSubSequenceComboButtonMenuEntryExecute(AssetData[0].GetAsset(), InTrack);
+	}
+}
+
 FKeyPropertyResult FSubTrackEditor::AddKeyInternal(FFrameNumber KeyTime, UMovieSceneSequence* InMovieSceneSequence, UMovieSceneTrack* InTrack, int32 RowIndex)
 {	
 	FKeyPropertyResult KeyPropertyResult;
@@ -675,9 +684,9 @@ FKeyPropertyResult FSubTrackEditor::AddKeyInternal(FFrameNumber KeyTime, UMovieS
 
 		const FQualifiedFrameTime InnerDuration = FQualifiedFrameTime(
 			MovieScene::DiscreteSize(InMovieSceneSequence->GetMovieScene()->GetPlaybackRange()),
-			InMovieSceneSequence->GetMovieScene()->GetFrameResolution());
+			InMovieSceneSequence->GetMovieScene()->GetTickResolution());
 
-		const FFrameRate OuterFrameRate = SubTrack->GetTypedOuter<UMovieScene>()->GetFrameResolution();
+		const FFrameRate OuterFrameRate = SubTrack->GetTypedOuter<UMovieScene>()->GetTickResolution();
 		const int32      OuterDuration  = InnerDuration.ConvertTo(OuterFrameRate).FrameNumber.Value;
 
 		UMovieSceneSubSection* NewSection = SubTrack->AddSequenceOnRow(InMovieSceneSequence, KeyTime, OuterDuration, RowIndex);
@@ -705,9 +714,9 @@ FKeyPropertyResult FSubTrackEditor::HandleSequenceAdded(FFrameNumber KeyTime, UM
 
 	const FQualifiedFrameTime InnerDuration = FQualifiedFrameTime(
 		MovieScene::DiscreteSize(Sequence->GetMovieScene()->GetPlaybackRange()),
-		Sequence->GetMovieScene()->GetFrameResolution());
+		Sequence->GetMovieScene()->GetTickResolution());
 
-	const FFrameRate OuterFrameRate = SubTrack->GetTypedOuter<UMovieScene>()->GetFrameResolution();
+	const FFrameRate OuterFrameRate = SubTrack->GetTypedOuter<UMovieScene>()->GetTickResolution();
 	const int32      OuterDuration  = InnerDuration.ConvertTo(OuterFrameRate).FrameNumber.Value;
 
 	UMovieSceneSubSection* NewSection = SubTrack->AddSequenceOnRow(Sequence, KeyTime, OuterDuration, RowIndex);
diff --git a/Engine/Source/Editor/MovieSceneTools/Private/TrackEditors/TransformTrackEditor.cpp b/Engine/Source/Editor/MovieSceneTools/Private/TrackEditors/TransformTrackEditor.cpp
--- a/Engine/Source/Editor/MovieSceneTools/Private/TrackEditors/TransformTrackEditor.cpp
+++ b/Engine/Source/Editor/MovieSceneTools/Private/TrackEditors/TransformTrackEditor.cpp
@@ -212,6 +212,7 @@ void F3DTransformTrackEditor::BuildTrackContextMenu( FMenuBuilder& MenuBuilder,
 		AssetPickerConfig.SelectionMode = ESelectionMode::Single;
 		AssetPickerConfig.Filter.ClassNames.Add(UAnimSequence::StaticClass()->GetFName());
 		AssetPickerConfig.OnAssetSelected = FOnAssetSelected::CreateStatic(&F3DTransformTrackEditor::ImportAnimSequenceTransforms, InSequencer, InTransformTrack);
+		AssetPickerConfig.OnAssetEnterPressed = FOnAssetEnterPressed::CreateStatic(&F3DTransformTrackEditor::ImportAnimSequenceTransformsEnterPressed, InSequencer, InTransformTrack);
 
 		FContentBrowserModule& ContentBrowserModule = FModuleManager::LoadModuleChecked<FContentBrowserModule>(TEXT("ContentBrowser"));
 
@@ -440,8 +441,13 @@ void F3DTransformTrackEditor::BuildObjectBindingEditButtons(TSharedPtr<SHorizont
 
 void F3DTransformTrackEditor::BuildObjectBindingTrackMenu(FMenuBuilder& MenuBuilder, const FGuid& ObjectBinding, const UClass* ObjectClass)
 {
-	if (ObjectClass->IsChildOf(AActor::StaticClass()) || ObjectClass->IsChildOf(USceneComponent::StaticClass()))
+	if (ObjectClass != nullptr && (ObjectClass->IsChildOf(AActor::StaticClass()) || ObjectClass->IsChildOf(USceneComponent::StaticClass())))
 	{
+		if (MovieSceneToolHelpers::HasHiddenMobility(ObjectClass))
+		{
+			return;
+		}
+
 		MenuBuilder.AddMenuEntry(
 			NSLOCTEXT("Sequencer", "AddTransform", "Transform"),
 			NSLOCTEXT("Sequencer", "AddPTransformTooltip", "Adds a transform track."),
@@ -652,7 +658,7 @@ void F3DTransformTrackEditor::GetTransformKeys( const TOptional<FTransformData>&
 	bool bLastVectorIsValid = LastTransform.IsSet();
 
 	// If key all is enabled, for a key on all the channels
-	if (GetSequencer()->GetKeyAllEnabled())
+	if (GetSequencer()->GetKeyGroupMode() == EKeyGroupMode::KeyAll)
 	{
 		bLastVectorIsValid = false;
 		ChannelsToKey = EMovieSceneTransformChannel::All;
@@ -661,9 +667,14 @@ void F3DTransformTrackEditor::GetTransformKeys( const TOptional<FTransformData>&
 	// Set translation keys/defaults
 	{
 		FVector LastVector = bLastVectorIsValid ? LastTransform->Translation : FVector(), CurrentVector = CurrentTransform.Translation;
-		const bool bKeyX = EnumHasAnyFlags(ChannelsToKey, EMovieSceneTransformChannel::TranslationX) && ( !bLastVectorIsValid || !FMath::IsNearlyEqual(LastVector.X, CurrentVector.X) );
-		const bool bKeyY = EnumHasAnyFlags(ChannelsToKey, EMovieSceneTransformChannel::TranslationY) && ( !bLastVectorIsValid || !FMath::IsNearlyEqual(LastVector.Y, CurrentVector.Y) );
-		const bool bKeyZ = EnumHasAnyFlags(ChannelsToKey, EMovieSceneTransformChannel::TranslationZ) && ( !bLastVectorIsValid || !FMath::IsNearlyEqual(LastVector.Z, CurrentVector.Z) );
+		bool bKeyX = EnumHasAnyFlags(ChannelsToKey, EMovieSceneTransformChannel::TranslationX) && ( !bLastVectorIsValid || !FMath::IsNearlyEqual(LastVector.X, CurrentVector.X) );
+		bool bKeyY = EnumHasAnyFlags(ChannelsToKey, EMovieSceneTransformChannel::TranslationY) && ( !bLastVectorIsValid || !FMath::IsNearlyEqual(LastVector.Y, CurrentVector.Y) );
+		bool bKeyZ = EnumHasAnyFlags(ChannelsToKey, EMovieSceneTransformChannel::TranslationZ) && ( !bLastVectorIsValid || !FMath::IsNearlyEqual(LastVector.Z, CurrentVector.Z) );
+
+		if (GetSequencer()->GetKeyGroupMode() == EKeyGroupMode::KeyGroup && (bKeyX || bKeyY || bKeyZ))
+		{
+			bKeyX = bKeyY = bKeyZ = true;
+		}
 
 		OutGeneratedKeys.Add(FMovieSceneChannelValueSetter::Create<FMovieSceneFloatChannel>(0, CurrentVector.X, bKeyX));
 		OutGeneratedKeys.Add(FMovieSceneChannelValueSetter::Create<FMovieSceneFloatChannel>(1, CurrentVector.Y, bKeyY));
@@ -684,9 +695,14 @@ void F3DTransformTrackEditor::GetTransformKeys( const TOptional<FTransformData>&
 			CurrentRotator.Roll = UnwindChannel(LastRotator.Roll, CurrentRotator.Roll);
 		}
 
-		const bool bKeyX = EnumHasAnyFlags(ChannelsToKey, EMovieSceneTransformChannel::RotationX) && ( !bLastVectorIsValid || !FMath::IsNearlyEqual(LastVector.X, CurrentVector.X) );
-		const bool bKeyY = EnumHasAnyFlags(ChannelsToKey, EMovieSceneTransformChannel::RotationY) && ( !bLastVectorIsValid || !FMath::IsNearlyEqual(LastVector.Y, CurrentVector.Y) );
-		const bool bKeyZ = EnumHasAnyFlags(ChannelsToKey, EMovieSceneTransformChannel::RotationZ) && ( !bLastVectorIsValid || !FMath::IsNearlyEqual(LastVector.Z, CurrentVector.Z) );
+		bool bKeyX = EnumHasAnyFlags(ChannelsToKey, EMovieSceneTransformChannel::RotationX) && ( !bLastVectorIsValid || !FMath::IsNearlyEqual(LastVector.X, CurrentVector.X) );
+		bool bKeyY = EnumHasAnyFlags(ChannelsToKey, EMovieSceneTransformChannel::RotationY) && ( !bLastVectorIsValid || !FMath::IsNearlyEqual(LastVector.Y, CurrentVector.Y) );
+		bool bKeyZ = EnumHasAnyFlags(ChannelsToKey, EMovieSceneTransformChannel::RotationZ) && ( !bLastVectorIsValid || !FMath::IsNearlyEqual(LastVector.Z, CurrentVector.Z) );
+
+		if (GetSequencer()->GetKeyGroupMode() == EKeyGroupMode::KeyGroup && ( bKeyX || bKeyY || bKeyZ) )
+		{
+			bKeyX = bKeyY = bKeyZ = true;
+		}
 
 		OutGeneratedKeys.Add(FMovieSceneChannelValueSetter::Create<FMovieSceneFloatChannel>(3, CurrentVector.X, bKeyX));
 		OutGeneratedKeys.Add(FMovieSceneChannelValueSetter::Create<FMovieSceneFloatChannel>(4, CurrentVector.Y, bKeyY));
@@ -696,9 +712,14 @@ void F3DTransformTrackEditor::GetTransformKeys( const TOptional<FTransformData>&
 	// Set scale keys/defaults
 	{
 		FVector LastVector = bLastVectorIsValid ? LastTransform->Scale : FVector(), CurrentVector = CurrentTransform.Scale;
-		const bool bKeyX = EnumHasAnyFlags(ChannelsToKey, EMovieSceneTransformChannel::ScaleX) && ( !bLastVectorIsValid || !FMath::IsNearlyEqual(LastVector.X, CurrentVector.X) );
-		const bool bKeyY = EnumHasAnyFlags(ChannelsToKey, EMovieSceneTransformChannel::ScaleY) && ( !bLastVectorIsValid || !FMath::IsNearlyEqual(LastVector.Y, CurrentVector.Y) );
-		const bool bKeyZ = EnumHasAnyFlags(ChannelsToKey, EMovieSceneTransformChannel::ScaleZ) && ( !bLastVectorIsValid || !FMath::IsNearlyEqual(LastVector.Z, CurrentVector.Z) );
+		bool bKeyX = EnumHasAnyFlags(ChannelsToKey, EMovieSceneTransformChannel::ScaleX) && ( !bLastVectorIsValid || !FMath::IsNearlyEqual(LastVector.X, CurrentVector.X) );
+		bool bKeyY = EnumHasAnyFlags(ChannelsToKey, EMovieSceneTransformChannel::ScaleY) && ( !bLastVectorIsValid || !FMath::IsNearlyEqual(LastVector.Y, CurrentVector.Y) );
+		bool bKeyZ = EnumHasAnyFlags(ChannelsToKey, EMovieSceneTransformChannel::ScaleZ) && ( !bLastVectorIsValid || !FMath::IsNearlyEqual(LastVector.Z, CurrentVector.Z) );
+
+		if (GetSequencer()->GetKeyGroupMode() == EKeyGroupMode::KeyGroup && (bKeyX || bKeyY || bKeyZ))
+		{
+			bKeyX = bKeyY = bKeyZ = true;
+		}
 
 		OutGeneratedKeys.Add(FMovieSceneChannelValueSetter::Create<FMovieSceneFloatChannel>(6, CurrentVector.X, bKeyX));
 		OutGeneratedKeys.Add(FMovieSceneChannelValueSetter::Create<FMovieSceneFloatChannel>(7, CurrentVector.Y, bKeyY));
@@ -723,6 +744,12 @@ void F3DTransformTrackEditor::AddTransformKeysForObject( UObject* Object, EMovie
 	if ( Actor != nullptr && SceneComponent != nullptr )
 	{
 		FTransformData CurrentTransform( SceneComponent );
+
+		if (MovieSceneToolHelpers::HasHiddenMobility(Object->GetClass()))
+		{
+			return;
+		}
+
 		if(Object->GetClass()->IsChildOf(AActor::StaticClass()))
 		{
 			AddTransformKeys( Actor, TOptional<FTransformData>(), CurrentTransform, ChannelToKey, KeyMode );
@@ -830,7 +857,7 @@ void F3DTransformTrackEditor::ImportAnimSequenceTransforms(const FAssetData& Ass
 		TransformTrack->Modify();
 
 		UMovieScene3DTransformSection* Section = Cast<UMovieScene3DTransformSection>(TransformTrack->CreateNewSection());
-		FFrameRate FrameResolution = Section->GetTypedOuter<UMovieScene>()->GetFrameResolution();
+		FFrameRate TickResolution = Section->GetTypedOuter<UMovieScene>()->GetTickResolution();
 
 		TArrayView<FMovieSceneFloatChannel*> FloatChannels = Section->GetChannelProxy().GetChannels<FMovieSceneFloatChannel>();
 
@@ -914,7 +941,7 @@ void F3DTransformTrackEditor::ImportAnimSequenceTransforms(const FAssetData& Ass
 			TRange<FFrameNumber> Range = Section->GetRange();
 			for(const FTempTransformKey& TempKey : TempKeys)
 			{
-				FFrameNumber KeyTime = (TempKey.Time * FrameResolution).RoundToFrame();
+				FFrameNumber KeyTime = (TempKey.Time * TickResolution).RoundToFrame();
 
 				Range = TRange<FFrameNumber>::Hull(Range, TRange<FFrameNumber>(KeyTime));
 
@@ -945,4 +972,12 @@ void F3DTransformTrackEditor::ImportAnimSequenceTransforms(const FAssetData& Ass
 	}
 }
 
+void F3DTransformTrackEditor::ImportAnimSequenceTransformsEnterPressed(const TArray<FAssetData>& Asset, TSharedRef<ISequencer> Sequencer, UMovieScene3DTransformTrack* TransformTrack)
+{
+	if (Asset.Num() > 0)
+	{
+		ImportAnimSequenceTransforms(Asset[0].GetAsset(), Sequencer, TransformTrack);
+	}
+}
+
 #undef LOCTEXT_NAMESPACE
diff --git a/Engine/Source/Editor/MovieSceneTools/Private/TrackEditors/TransformTrackEditor.h b/Engine/Source/Editor/MovieSceneTools/Private/TrackEditors/TransformTrackEditor.h
--- a/Engine/Source/Editor/MovieSceneTools/Private/TrackEditors/TransformTrackEditor.h
+++ b/Engine/Source/Editor/MovieSceneTools/Private/TrackEditors/TransformTrackEditor.h
@@ -138,6 +138,9 @@ private:
 	/** Import an animation sequence's root transforms into a transform section */
 	static void ImportAnimSequenceTransforms(const FAssetData& Asset, TSharedRef<class ISequencer> Sequencer, UMovieScene3DTransformTrack* TransformTrack);
 
+	/** Import an animation sequence's root transforms into a transform section */
+	static void ImportAnimSequenceTransformsEnterPressed(const TArray<FAssetData>& Asset, TSharedRef<class ISequencer> Sequencer, UMovieScene3DTransformTrack* TransformTrack);
+
 private:
 
 	static FName TransformPropertyName;
diff --git a/Engine/Source/Editor/MovieSceneTools/Public/AutomatedLevelSequenceCapture.h b/Engine/Source/Editor/MovieSceneTools/Public/AutomatedLevelSequenceCapture.h
new file mode 100644
--- /dev/null
+++ b/Engine/Source/Editor/MovieSceneTools/Public/AutomatedLevelSequenceCapture.h
@@ -0,0 +1,189 @@
+// Copyright 1998-2018 Epic Games, Inc. All Rights Reserved.
+
+#pragma once
+
+#include "CoreMinimal.h"
+#include "UObject/ObjectMacros.h"
+#include "UObject/SoftObjectPath.h"
+#include "LevelSequencePlayer.h"
+#include "MovieSceneCapture.h"
+#include "AutomatedLevelSequenceCapture.generated.h"
+
+class ALevelSequenceActor;
+class FJsonObject;
+class FSceneViewport;
+class ULevelSequenceBurnInOptions;
+
+UCLASS(config=EditorSettings)
+class MOVIESCENETOOLS_API UAutomatedLevelSequenceCapture : public UMovieSceneCapture
+{
+public:
+	UAutomatedLevelSequenceCapture(const FObjectInitializer&);
+
+	GENERATED_BODY()
+
+	/** Set the level sequence asset that we are to record. We will spawn a new actor at runtime for this asset for playback. */
+	void SetLevelSequenceAsset(FString InAssetPath);
+
+	/** When enabled, the StartFrame setting will override the default starting frame number */
+	UPROPERTY(config, EditAnywhere, Category=Animation, AdvancedDisplay)
+	bool bUseCustomStartFrame;
+
+	/** Frame number to start capturing. */
+	UPROPERTY(config, EditAnywhere, Category=Animation, AdvancedDisplay, meta=(EditCondition="bUseCustomStartFrame", DisplayName="Start Frame"))
+	FFrameNumber CustomStartFrame;
+
+	/** When enabled, the EndFrame setting will override the default ending frame number */
+	UPROPERTY(config, EditAnywhere, Category=Animation, AdvancedDisplay)
+	bool bUseCustomEndFrame;
+
+	/** Frame number to end capturing. */
+	UPROPERTY(config, EditAnywhere, Category=Animation, AdvancedDisplay, meta=(EditCondition="bUseCustomEndFrame", DisplayName="End Frame"))
+	FFrameNumber CustomEndFrame;
+
+	/** The number of extra frames to play before the sequence's start frame, to "warm up" the animation.  This is useful if your
+	    animation contains particles or other runtime effects that are spawned into the scene earlier than your capture start frame */
+	UPROPERTY(config, EditAnywhere, Category=Animation, AdvancedDisplay)
+	int32 WarmUpFrameCount;
+
+	/** The number of seconds to wait (in real-time) before we start playing back the warm up frames.  Useful for allowing post processing effects to settle down before capturing the animation. */
+	UPROPERTY(config, EditAnywhere, Category=Animation, AdvancedDisplay, meta=(Units=Seconds, ClampMin=0))
+	float DelayBeforeWarmUp;
+
+	/** The number of seconds to wait (in real-time) at shot boundaries.  Useful for allowing post processing effects to settle down before capturing the animation. */
+	UPROPERTY(config, EditAnywhere, Category=Animation, AdvancedDisplay, meta=(Units=Seconds, ClampMin=0))
+	float DelayBeforeShotWarmUp;
+
+	UPROPERTY(EditAnywhere, Category=CaptureSettings, AdvancedDisplay, meta=(EditInline))
+	ULevelSequenceBurnInOptions* BurnInOptions;
+
+	/** Whether to write edit decision lists (EDLs) if the sequence contains shots */
+	UPROPERTY(config, EditAnywhere, Category=Sequence)
+	bool bWriteEditDecisionList;
+
+public:
+	// UMovieSceneCapture interface
+	virtual void Initialize(TSharedPtr<FSceneViewport> InViewport, int32 PIEInstance = -1) override;
+
+	virtual void LoadFromConfig() override;
+
+	virtual void SaveToConfig() override;
+
+	virtual void Close() override;
+
+	/** Override the render frames with the given start/end frames. Restore the values when done rendering. */
+	void SetFrameOverrides(FFrameNumber InStartFrame, FFrameNumber InEndFrame);
+
+protected:
+
+	virtual void AddFormatMappings(TMap<FString, FStringFormatArg>& OutFormatMappings, const FFrameMetrics& FrameMetrics) const override;
+
+	/** Custom, additional json serialization */
+	virtual void SerializeAdditionalJson(FJsonObject& Object);
+
+	/** Custom, additional json deserialization */
+	virtual void DeserializeAdditionalJson(const FJsonObject& Object);
+
+private:
+
+	/** Update any cached information we need from the level sequence actor */
+	void UpdateFrameState();
+
+	/** Called when the level sequence has updated the world */
+	void SequenceUpdated(const UMovieSceneSequencePlayer& Player, FFrameTime CurrentTime, FFrameTime PreviousTime);
+
+	/** Called to set up the player's playback range */
+	void SetupFrameRange();
+
+	/** Enable cinematic mode override */
+	void EnableCinematicMode();
+
+	/** Export EDL if requested */
+	void ExportEDL();
+
+	/** Delegate binding for the above callback */
+	FDelegateHandle OnPlayerUpdatedBinding;
+
+private:
+
+	virtual void Tick(float DeltaSeconds) override;
+
+	/** Initialize all the shots to be recorded, ie. expand section ranges with handle frames */
+	bool InitializeShots();
+
+	/** Set up the current shot to be recorded, ie. expand playback range to the section range */
+	bool SetupShot(FFrameNumber& StartTime, FFrameNumber& EndTime);
+
+	/** Restore any modification to shots */
+	void RestoreShots();
+
+	/** Restore frame settings from overridden shot frames */
+	bool RestoreFrameOverrides();
+
+	void DelayBeforeWarmupFinished();
+
+	void PauseFinished();
+
+	/** A level sequence asset to playback at runtime - used where the level sequence does not already exist in the world. */
+	UPROPERTY()
+	FSoftObjectPath LevelSequenceAsset;
+
+	/** The pre-existing level sequence actor to use for capture that specifies playback settings */
+	UPROPERTY()
+	TWeakObjectPtr<ALevelSequenceActor> LevelSequenceActor;
+
+	/** The viewport being captured. */
+	TWeakPtr<FSceneViewport> Viewport;
+
+	/** Which state we're in right now */
+	enum class ELevelSequenceCaptureState
+	{
+		Setup,
+		DelayBeforeWarmUp,
+		ReadyToWarmUp,
+		WarmingUp,
+		FinishedWarmUp,
+		Paused,
+		FinishedPause,
+	} CaptureState;
+
+	/** The number of warm up frames left before we actually start saving out images */
+	int32 RemainingWarmUpFrames;
+	
+	/** The number of individual shot movies to render */
+	int32 NumShots;
+
+	/** The current shot movie that is rendering */
+	int32 ShotIndex;
+
+	FLevelSequencePlayerSnapshot CachedState;
+
+	TOptional<float> CachedPlayRate;
+
+	FTimerHandle DelayTimer;
+
+	struct FCinematicShotCache
+	{
+		FCinematicShotCache(bool bInActive, bool bInLocked, const TRange<FFrameNumber>& InShotRange, const TRange<FFrameNumber>& InMovieSceneRange) : 
+			  bActive(bInActive)
+			, bLocked(bInLocked)
+			, ShotRange(InShotRange)
+			, MovieSceneRange(InMovieSceneRange)
+		{
+		}
+
+		bool bActive;
+		bool bLocked;
+		TRange<FFrameNumber> ShotRange;
+		TRange<FFrameNumber> MovieSceneRange;
+	};
+
+	TArray<FCinematicShotCache> CachedShotStates;
+	TRange<FFrameNumber> CachedPlaybackRange;
+
+	TOptional<FFrameNumber> CachedStartFrame;
+	TOptional<FFrameNumber> CachedEndFrame;
+	TOptional<bool> bCachedUseCustomStartFrame;
+	TOptional<bool> bCachedUseCustomEndFrame;
+};
+
diff --git a/Engine/Source/Editor/MovieSceneTools/Public/CommonMovieSceneTools.h b/Engine/Source/Editor/MovieSceneTools/Public/CommonMovieSceneTools.h
--- a/Engine/Source/Editor/MovieSceneTools/Public/CommonMovieSceneTools.h
+++ b/Engine/Source/Editor/MovieSceneTools/Public/CommonMovieSceneTools.h
@@ -12,9 +12,9 @@
 struct FTimeToPixel
 {
 public:
-	FTimeToPixel( const FGeometry& AllottedGeometry, const TRange<double>& InLocalViewRange, const FFrameRate& InFrameResolution )
+	FTimeToPixel( const FGeometry& AllottedGeometry, const TRange<double>& InLocalViewRange, const FFrameRate& InTickResolution )
 		: ViewRangeStartSeconds( InLocalViewRange.GetLowerBoundValue() )
-		, FrameResolution( InFrameResolution )
+		, TickResolution( InTickResolution )
 	{
 		double VisibleWidth = InLocalViewRange.Size<double>();
 
@@ -52,7 +52,7 @@ public:
 	 */
 	float FrameToPixel( const FFrameTime& Time ) const
 	{
-		return (Time / FrameResolution - ViewRangeStartSeconds) * PixelsPerSecond;
+		return (Time / TickResolution - ViewRangeStartSeconds) * PixelsPerSecond;
 	}
 
 	/**
@@ -63,7 +63,7 @@ public:
 	 */
 	FFrameTime PixelToFrame( float PixelX ) const
 	{
-		return ( PixelX / PixelsPerSecond + ViewRangeStartSeconds ) * FrameResolution;
+		return ( PixelX / PixelsPerSecond + ViewRangeStartSeconds ) * TickResolution;
 	}
 
 	/**
@@ -74,23 +74,23 @@ public:
 	 */
 	FFrameTime PixelDeltaToFrame( float PixelDelta ) const
 	{
-		return ( PixelDelta / PixelsPerSecond ) * FrameResolution;
+		return ( PixelDelta / PixelsPerSecond ) * TickResolution;
 	}
 
 	/**
-	 * Retrieve the frame resolution of the current sequence
+	 * Retrieve the tick resolution of the current sequence
 	 */
-	FFrameRate GetFrameResolution() const
+	FFrameRate GetTickResolution() const
 	{
-		return FrameResolution;
+		return TickResolution;
 	}
 
 private:
 
 	/** time range of the sequencer in seconds */
 	double ViewRangeStartSeconds;
-	/** The frame resolution of the current timeline */
-	FFrameRate FrameResolution;
+	/** The tick resolution of the current timeline */
+	FFrameRate TickResolution;
 	/** The number of pixels in the view range */
 	float PixelsPerSecond;
 };
diff --git a/Engine/Source/Editor/MovieSceneTools/Public/FCPXML/FCPXMLMovieSceneTranslator.h b/Engine/Source/Editor/MovieSceneTools/Public/FCPXML/FCPXMLMovieSceneTranslator.h
new file mode 100644
--- /dev/null
+++ b/Engine/Source/Editor/MovieSceneTools/Public/FCPXML/FCPXMLMovieSceneTranslator.h
@@ -0,0 +1,85 @@
+// Copyright 1998-2018 Epic Games, Inc. All Rights Reserved.
+
+#pragma once
+
+#include "CoreMinimal.h"
+#include "Misc/FileHelper.h"
+#include "Misc/Paths.h"
+#include "MovieScene.h"
+#include "MovieSceneTranslator.h"
+
+class FMovieSceneTranslatorContext;
+
+/** 
+ * The FFCPXMLImporter class is the entry point for launching an import of data from an XML file into Sequencer. 
+ */
+class MOVIESCENETOOLS_API FFCPXMLImporter : public FMovieSceneImporter
+{
+public:
+	FFCPXMLImporter();
+
+	virtual ~FFCPXMLImporter();
+
+	/** Format description. */
+	virtual FText GetFileTypeDescription() const;
+	/** Import window title. */
+	virtual FText GetDialogTitle() const;
+	/** Scoped transaction description. */
+	virtual FText GetTransactionDescription() const;
+	/** Message log window title. */
+	virtual FName GetMessageLogWindowTitle() const;
+	/** Message log list label. */
+	virtual FText GetMessageLogLabel() const;
+
+public:
+	/*
+	 * Import FCP 7 XML
+	 *
+	 * @param InMovieScene The movie scene to import the XML file into
+	 * @param InFrameRate The frame rate to import the XML at
+	 * @param InFilename The filename to import
+	 * @param OutError The return error message
+	 * @return Whether the import was successful
+	 */
+	bool Import(UMovieScene* InMovieScene, FFrameRate InFrameRate, FString InFilename, TSharedRef<FMovieSceneTranslatorContext> InContext);
+};
+
+/** 
+ * The FFCPXMLExporter class is the entry point for launching an export of data from Sequencer into an XML file. 
+ */
+class MOVIESCENETOOLS_API FFCPXMLExporter : public FMovieSceneExporter
+{
+public:
+	/** Constructor */
+	FFCPXMLExporter();
+	/** Destructor */
+	virtual ~FFCPXMLExporter();
+
+	/** Format description. */
+	virtual FText GetFileTypeDescription() const;
+	/** Export dialog window title. */
+	virtual FText GetDialogTitle() const;
+	/** Default format file extension. */
+	virtual FText GetDefaultFileExtension() const;
+	/** Notification when export completes. */
+	virtual FText GetNotificationExportFinished() const;
+	/** Notification hyperlink to exported file path. */
+	virtual FText GetNotificationHyperlinkText() const;
+	/** Message log window title. */
+	virtual FName GetMessageLogWindowTitle() const;
+	/** Message log list label. */
+	virtual FText GetMessageLogLabel() const;
+
+public:
+	/*
+	 * Export FCP 7 XML
+	 *
+	 * @param InMovieScene The movie scene with the cinematic shot track and audio tracks to export
+	 * @param InFrameRate The frame rate to export the XML at
+	 * @param InSaveFilename The file path to save to.
+	 * @param InHandleFrames The number of handle frames to include for each shot.
+	 * @param OutError The return error message
+	 * @return Whether the export was successful
+	 */
+	virtual bool Export(const UMovieScene* InMovieScene, FFrameRate InFrameRate, FString InSaveFilename, int32 InHandleFrames, TSharedRef<FMovieSceneTranslatorContext> InContext);
+};
diff --git a/Engine/Source/Editor/MovieSceneTools/Public/FrameNumberDetailsCustomization.h b/Engine/Source/Editor/MovieSceneTools/Public/FrameNumberDetailsCustomization.h
--- a/Engine/Source/Editor/MovieSceneTools/Public/FrameNumberDetailsCustomization.h
+++ b/Engine/Source/Editor/MovieSceneTools/Public/FrameNumberDetailsCustomization.h
@@ -25,14 +25,14 @@ public:
 	}
 
 	/** IPropertyTypeCustomization interface */
-	virtual void CustomizeHeader(TSharedRef<IPropertyHandle> PropertyHandle, FDetailWidgetRow& HeaderRow, IPropertyTypeCustomizationUtils& CustomizationUtils) override {};
+	virtual void CustomizeHeader(TSharedRef<IPropertyHandle> PropertyHandle, FDetailWidgetRow& HeaderRow, IPropertyTypeCustomizationUtils& CustomizationUtils) override {}
 	virtual void CustomizeChildren(TSharedRef<IPropertyHandle> PropertyHandle, IDetailChildrenBuilder& ChildBuilder, IPropertyTypeCustomizationUtils& CustomizationUtils) override;
 
 private:
 	FText OnGetTimeText() const;
 	void OnTimeTextCommitted(const FText& InText, ETextCommit::Type CommitInfo);
 
-	/** The Numeric Type interface used to convert between display formats and internal frame resolution. */
+	/** The Numeric Type interface used to convert between display formats and internal tick resolution. */
 	TSharedPtr<INumericTypeInterface<double>> NumericTypeInterface;
 
 	/** Store the property handle to the FrameNumber field so we can get/set the value on the object via text box callbacks. */
diff --git a/Engine/Source/Editor/MovieSceneTools/Public/KeyframeTrackEditor.h b/Engine/Source/Editor/MovieSceneTools/Public/KeyframeTrackEditor.h
--- a/Engine/Source/Editor/MovieSceneTools/Public/KeyframeTrackEditor.h
+++ b/Engine/Source/Editor/MovieSceneTools/Public/KeyframeTrackEditor.h
@@ -46,9 +46,9 @@ struct FMovieSceneChannelValueSetter
 		virtual ~IImpl(){}
 
 		/* Returns whether a key was created */
-		virtual bool Apply(FMovieSceneChannelProxy& Proxy, FFrameNumber InTime, EMovieSceneKeyInterpolation InterpolationMode, bool bKeyEvenIfUnchanged, bool bKeyEvenIfEmpty) const { return false; }
+		virtual bool Apply(UMovieSceneSection* Section, FMovieSceneChannelProxy& Proxy, FFrameNumber InTime, EMovieSceneKeyInterpolation InterpolationMode, bool bKeyEvenIfUnchanged, bool bKeyEvenIfEmpty) const { return false; }
 
-		virtual void ApplyDefault(FMovieSceneChannelProxy& Proxy) const { }
+		virtual void ApplyDefault(UMovieSceneSection* Section, FMovieSceneChannelProxy& Proxy) const { }
 	};
 
 	IImpl* operator->()
@@ -76,12 +76,15 @@ private:
 			: ChannelIndex(InChannelIndex), ValueToSet(InValue)
 		{}
 
-		virtual void ApplyDefault(FMovieSceneChannelProxy& Proxy) const override
+		virtual void ApplyDefault(UMovieSceneSection* Section, FMovieSceneChannelProxy& Proxy) const override
 		{
 			ChannelType* Channel = Proxy.GetChannel<ChannelType>(ChannelIndex);
-			if (Channel)
+			if (Channel && Channel->GetInterface().GetTimes().Num() == 0)
 			{
-				Channel->SetDefault(ValueToSet);
+				if (Section->TryModify())
+				{
+					Channel->SetDefault(ValueToSet);
+				}
 			}
 		}
 	};
@@ -96,7 +99,7 @@ private:
 			: ChannelIndex(InChannelIndex), ValueToSet(InValue)
 		{}
 
-		virtual bool Apply(FMovieSceneChannelProxy& Proxy, FFrameNumber InTime, EMovieSceneKeyInterpolation InterpolationMode, bool bKeyEvenIfUnchanged, bool bKeyEvenIfEmpty) const override
+		virtual bool Apply(UMovieSceneSection* Section, FMovieSceneChannelProxy& Proxy, FFrameNumber InTime, EMovieSceneKeyInterpolation InterpolationMode, bool bKeyEvenIfUnchanged, bool bKeyEvenIfEmpty) const override
 		{
 			bool bKeyCreated = false;
 			using namespace MovieScene;
@@ -114,8 +117,11 @@ private:
 				{
 					if (Channel->GetInterface().GetTimes().Num() != 0 || bKeyEvenIfEmpty)
 					{
-						AddKeyToChannel(Channel, InTime, ValueToSet, InterpolationMode);
-						bKeyCreated = true;
+						if (Section->TryModify())
+						{
+							AddKeyToChannel(Channel, InTime, ValueToSet, InterpolationMode);
+							bKeyCreated = true;
+						}
 					}
 				}
 			}
@@ -123,14 +129,17 @@ private:
 			return bKeyCreated;
 		}
 
-		virtual void ApplyDefault(FMovieSceneChannelProxy& Proxy) const override
+		virtual void ApplyDefault(UMovieSceneSection* Section, FMovieSceneChannelProxy& Proxy) const override
 		{
 			using namespace MovieScene;
 
 			ChannelType* Channel = Proxy.GetChannel<ChannelType>(ChannelIndex);
-			if (Channel)
+			if (Channel && Channel->GetInterface().GetTimes().Num() == 0)
 			{
-				Channel->SetDefault(ValueToSet);
+				if (Section->TryModify())
+				{
+					Channel->SetDefault(ValueToSet);
+				}
 			}
 		}
 	};
@@ -330,10 +339,6 @@ private:
 	FKeyPropertyResult AddKeysToSection(UMovieSceneSection* Section, FFrameNumber KeyTime, const FGeneratedTrackKeys& Keys, ESequencerKeyMode KeyMode)
 	{
 		FKeyPropertyResult KeyPropertyResult;
-		if (!Section->TryModify())
-		{
-			return KeyPropertyResult;
-		}
 
 		EAutoChangeMode AutoChangeMode = GetSequencer()->GetAutoChangeMode();
 
@@ -341,13 +346,14 @@ private:
 			
 		const bool bSetDefaults = GetSequencer()->GetAutoSetTrackDefaults();
 		
-		if ( KeyMode != ESequencerKeyMode::AutoKey || AutoChangeMode == EAutoChangeMode::AutoKey || AutoChangeMode == EAutoChangeMode::All || bSetDefaults)
+		if ( KeyMode != ESequencerKeyMode::AutoKey || AutoChangeMode == EAutoChangeMode::AutoKey || AutoChangeMode == EAutoChangeMode::All)
 		{
 			EMovieSceneKeyInterpolation InterpolationMode = GetSequencer()->GetKeyInterpolation();
 
 			const bool bKeyEvenIfUnchanged =
 				KeyMode == ESequencerKeyMode::ManualKeyForced ||
-				GetSequencer()->GetKeyAllEnabled();
+				GetSequencer()->GetKeyGroupMode() == EKeyGroupMode::KeyAll ||
+				GetSequencer()->GetKeyGroupMode() == EKeyGroupMode::KeyGroup;
 
 			const bool bKeyEvenIfEmpty =
 				(KeyMode == ESequencerKeyMode::AutoKey && AutoChangeMode == EAutoChangeMode::All) ||
@@ -355,15 +361,15 @@ private:
 
 			for (const FMovieSceneChannelValueSetter& GeneratedKey : Keys)
 			{
-				KeyPropertyResult.bKeyCreated |= GeneratedKey->Apply(Proxy, KeyTime, InterpolationMode, bKeyEvenIfUnchanged, bKeyEvenIfEmpty);
+				KeyPropertyResult.bKeyCreated |= GeneratedKey->Apply(Section, Proxy, KeyTime, InterpolationMode, bKeyEvenIfUnchanged, bKeyEvenIfEmpty);
 			}
-
-			if (bSetDefaults)
+		}
+			
+		if (bSetDefaults)
+		{
+			for (const FMovieSceneChannelValueSetter& GeneratedKey : Keys)
 			{
-				for (const FMovieSceneChannelValueSetter& GeneratedKey : Keys)
-				{
-					GeneratedKey->ApplyDefault(Proxy);
-				}
+				GeneratedKey->ApplyDefault(Section, Proxy);
 			}
 		}
 
diff --git a/Engine/Source/Editor/MovieSceneTools/Public/MovieSceneSectionDetailsCustomization.h b/Engine/Source/Editor/MovieSceneTools/Public/MovieSceneSectionDetailsCustomization.h
--- a/Engine/Source/Editor/MovieSceneTools/Public/MovieSceneSectionDetailsCustomization.h
+++ b/Engine/Source/Editor/MovieSceneTools/Public/MovieSceneSectionDetailsCustomization.h
@@ -70,7 +70,7 @@ private:
 	FReply ToggleRangeEndBounded();
 private:
 
-	/** The Numeric Type interface used to convert between display formats and internal frame resolution. */
+	/** The Numeric Type interface used to convert between display formats and internal tick resolution. */
 	TSharedPtr<INumericTypeInterface<double>> NumericTypeInterface;
 
 	/** Store the property handle to the FrameNumber field so we can get/set the value on the object via text box callbacks. */
diff --git a/Engine/Source/Editor/MovieSceneTools/Public/MovieSceneToolHelpers.h b/Engine/Source/Editor/MovieSceneTools/Public/MovieSceneToolHelpers.h
--- a/Engine/Source/Editor/MovieSceneTools/Public/MovieSceneToolHelpers.h
+++ b/Engine/Source/Editor/MovieSceneTools/Public/MovieSceneToolHelpers.h
@@ -11,6 +11,8 @@
 #include "MovieSceneObjectBindingID.h"
 #include "MovieSceneObjectBindingIDPicker.h"
 #include "ISequencer.h"
+#include "Logging/TokenizedMessage.h"
+#include "MovieSceneTranslator.h"
 
 class ISequencer;
 class UMovieScene;
@@ -35,7 +37,7 @@ public:
 	 * @param Time	The time at which to trim
 	 * @param bTrimLeft Trim left or trim right
 	 */
-	static void TrimSection(const TSet<TWeakObjectPtr<UMovieSceneSection>>& Sections, FFrameNumber Time, bool bTrimLeft);
+	static void TrimSection(const TSet<TWeakObjectPtr<UMovieSceneSection>>& Sections, FQualifiedFrameTime Time, bool bTrimLeft);
 
 	/**
 	 * Splits sections at the given time
@@ -43,7 +45,7 @@ public:
 	 * @param Sections The sections to split
 	 * @param Time	The time at which to split
 	 */
-	static void SplitSection(const TSet<TWeakObjectPtr<UMovieSceneSection>>& Sections, FFrameNumber Time);
+	static void SplitSection(const TSet<TWeakObjectPtr<UMovieSceneSection>>& Sections, FQualifiedFrameTime Time);
 
 	/**
 	 * Parse a shot name into its components.
@@ -144,6 +146,37 @@ public:
 	 */
 	static bool ShowExportEDLDialog(const UMovieScene* InMovieScene, FFrameRate InFrameRate, FString InSaveDirectory = TEXT(""), int32 InHandleFrames = 8);
 
+	/**
+	* Import movie scene formats
+	*
+	* @param InImporter The movie scene importer.
+	* @param InMovieScene The movie scene to import the format into
+	* @param InFrameRate The frame rate to import the format at
+	* @param InOpenDirectory Optional directory path to open from. If none given, a dialog will pop up to prompt the user
+	* @return Whether the import was successful
+	*/
+	static bool MovieSceneTranslatorImport(FMovieSceneImporter* InImporter, UMovieScene* InMovieScene, FFrameRate InFrameRate, FString InOpenDirectory = TEXT(""));
+
+	/**
+	* Export movie scene formats
+	*
+	* @param InExporter The movie scene exporter.
+	* @param InMovieScene The movie scene with the cinematic shot track and audio tracks to export
+	* @param InFrameRate The frame rate to export the AAF at
+	* @param InSaveDirectory Optional directory path to save to. If none given, a dialog will pop up to prompt the user
+	* @param InHandleFrames The number of handle frames to include for each shot.
+	* @return Whether the export was successful
+	*/
+	static bool MovieSceneTranslatorExport(FMovieSceneExporter* InExporter, const UMovieScene* InMovieScene, FFrameRate InFrameRate, FString InSaveDirectory = TEXT(""), int32 InHandleFrames = 8);
+
+	/** 
+	* Display error message window for MovieScene translators
+	*
+	* @param InTranslator The movie scene importer or exporter.
+	* @param InContext The context used to gather error, warning or info messages during import or export.
+	*/
+	static void MovieSceneTranslatorDisplayMessages(FMovieSceneTranslator* InTranslator,  TSharedRef<FMovieSceneTranslatorContext> InContext);
+
 	/**
 	 * Import FBX
 	 *
@@ -178,6 +211,11 @@ public:
 	 * @return The exported camera anim asset
 	 */
 	static UObject* ExportToCameraAnim(UMovieScene* InMovieScene, FGuid& InObjectBinding);
+
+	/*
+	 * @return Whether this object class has hidden mobility and can't be animated
+	 */
+	static bool HasHiddenMobility(const UClass* ObjectClass);
 };
 
 class FTrackEditorBindingIDPicker : public FMovieSceneObjectBindingIDPicker
@@ -206,3 +244,4 @@ private:
 	FOnBindingPicked OnBindingPickedEvent;
 };
 
+
diff --git a/Engine/Source/Editor/MovieSceneTools/Public/MovieSceneTranslator.h b/Engine/Source/Editor/MovieSceneTools/Public/MovieSceneTranslator.h
new file mode 100644
--- /dev/null
+++ b/Engine/Source/Editor/MovieSceneTools/Public/MovieSceneTranslator.h
@@ -0,0 +1,331 @@
+// Copyright 1998-2018 Epic Games, Inc. All Rights Reserved.
+
+#pragma once
+
+#include "Misc/FileHelper.h"
+#include "Misc/Paths.h"
+#include "Modules/ModuleManager.h"
+#include "MovieScene.h"
+#include "Tracks/MovieSceneAudioTrack.h"
+#include "Sections/MovieSceneAudioSection.h"
+#include "Tracks/MovieSceneCinematicShotTrack.h"
+#include "Sections/MovieSceneCinematicShotSection.h"
+#include "Logging/TokenizedMessage.h"
+#include "UObject/MetaData.h"
+
+/**
+* MovieSceneTranslator context class.
+*/
+class MOVIESCENETOOLS_API FMovieSceneTranslatorContext : public TSharedFromThis<FMovieSceneTranslatorContext>
+{
+public:
+	/** Constructor */
+	FMovieSceneTranslatorContext() {}
+	/** Destructor */
+	virtual ~FMovieSceneTranslatorContext() {}
+
+	/** Initialize the context */
+	void Init();
+
+	/** Add message. */
+	void AddMessage(EMessageSeverity::Type MessageSeverity, FText ErrorMessage);
+
+	/** Reset all messages. */
+	void ClearMessages();
+
+	/** Returns true if specified type of message */
+	bool ContainsMessageType(EMessageSeverity::Type InSeverity) const;
+
+	/** Get error messages. */
+	const TArray<TSharedRef<FTokenizedMessage>>& GetMessages() const;
+
+private:
+
+	/** Error messages **/
+	TArray<TSharedRef<FTokenizedMessage>> Messages;
+};
+
+struct FMovieSceneExportSectionData
+{
+	// FString Name;
+	int32 RowIndex;
+	FFrameNumber StartFrame;
+	FFrameNumber EndFrame;
+	bool bWithinPlaybackRange;
+	bool bEnabled;
+};
+
+struct FMovieSceneExportCinematicSectionData : public FMovieSceneExportSectionData
+{
+	FString ShotDisplayName;
+	FString ShotFilename;
+	const UMovieSceneCinematicShotSection* CinematicShotSection;
+};
+
+struct FMovieSceneExportAudioSectionData : public FMovieSceneExportSectionData
+{
+	uint32 NumChannels;
+	float Samples;
+};
+
+struct FMovieSceneExportAudioTrackData
+{
+	TArray< TSharedPtr<FMovieSceneExportAudioSectionData> > AudioSections;
+};
+
+struct FMovieSceneExportCinematicTrackData
+{
+	FMovieSceneExportCinematicTrackData(int32 InRowIndex) : RowIndex(InRowIndex) {}
+
+	/** This indicates the sub-track's row index in the master track */
+	int32 RowIndex;
+	TArray< TSharedPtr<FMovieSceneExportCinematicSectionData> > CinematicSections;
+};
+
+struct FMovieSceneExportCinematicMasterTrackData
+{
+	/** Array of all sections in order they appear in UMovieSceneCinematicTrack*/
+	TArray< TSharedPtr<FMovieSceneExportCinematicSectionData> > CinematicSections;
+
+	/** Array of sorted movie tracks, containing pointers to sections within the track */
+	TArray< TSharedPtr<FMovieSceneExportCinematicTrackData> > CinematicTracks;
+};
+
+struct FMovieSceneExportMovieSceneData
+{
+	FString Name;
+	FString Path;
+	FFrameRate Resolution;
+	int32 Duration;
+	FFrameNumber PlaybackRangeStartFrame;
+	FFrameNumber PlaybackRangeEndFrame;
+	TSharedPtr<FMovieSceneExportCinematicMasterTrackData> CinematicMasterTrack;
+	TArray< TSharedPtr<FMovieSceneExportAudioTrackData> > AudioTracks;
+};
+
+enum class EMovieSceneTranslatorSectionType : uint32
+{
+	Cinematic = 0,
+	Audio = 1,
+};
+
+/** 
+ * The FMovieSceneExportData class aggregates intermediate data from Sequencer classes to be used for timeline exports 
+ */
+class MOVIESCENETOOLS_API FMovieSceneExportData : public TSharedFromThis<FMovieSceneExportData>
+{
+public:
+	/** Constructor */
+	FMovieSceneExportData(const UMovieScene* InMovieScene, FFrameRate InFrameRate, int32 InHandleFrames, FString InSaveFilename, TSharedPtr<FMovieSceneTranslatorContext> InContext);
+	/** Default constructor, necessary for shared ref - should not be used */
+	FMovieSceneExportData();
+	/** Destructor */
+	~FMovieSceneExportData();
+
+	/** Gets export filename */
+	FString GetFilename() const;
+	/** Gets export filename with full path */
+	FString GetFilenamePath() const;
+
+	/** Gets export frame rate */
+	FFrameRate GetFrameRate() const;
+	/** Returns true if frame rate is a non-integral frame rate */
+	bool GetFrameRateIsNTSC() const;
+	/** Returns the nearest integral frame rate */
+	uint32 GetNearestWholeFrameRate() const;
+
+	/** Gets the frame handle which has been converted to the export frame rate */
+	FFrameTime GetHandleFrames() const;
+
+	/** True when the export data was successfully constructed. */
+	bool IsExportDataValid() const;
+
+private:
+
+	/** Entry point for creating the intermediate data to use when exporting. */
+	bool ConstructData(const UMovieScene* InMovieScene);
+	/** Loads intermediate movie scene data from Sequencer. */
+	bool ConstructMovieSceneData(const UMovieScene* InMovieScene);
+	/** Loads intermediate cinematic master track data from Sequencer. */
+	bool ConstructCinematicMasterTrackData(const UMovieScene* InMovieScene, const UMovieSceneCinematicShotTrack* InCinematicMasterTrack);
+	/** Loads intermediate cinematic track data from Sequencer. */
+	bool ConstructCinematicTrackData(const UMovieScene* InMovieScene, TSharedPtr<FMovieSceneExportCinematicMasterTrackData> InCinematicMasterTrack, int32 RowIndex);
+	/** Loads intermediate audio track data from Sequencer. */
+	bool ConstructAudioTrackData(const UMovieScene* InMovieScene, TSharedPtr<FMovieSceneExportMovieSceneData> InMovieSceneData, const UMovieSceneAudioTrack* InAudioTrack);
+	/** Loads intermediate cinematic section data from Sequencer. */
+	bool ConstructCinematicSectionData(const UMovieScene* InMovieScene, TSharedPtr<FMovieSceneExportCinematicMasterTrackData> InMasterTrackData, const UMovieSceneCinematicShotSection* InCinematicSection);
+	/** Loads intermediate audio section data from Sequencer. */
+	bool ConstructAudioSectionData(const UMovieScene* InMovieScene, TSharedPtr<FMovieSceneExportAudioTrackData> InTrackData, const UMovieSceneAudioSection* InAudioSection);
+	/** Loads intermediate common section data from Sequencer. */
+	bool ConstructSectionData(const UMovieScene* InMovieScene, TSharedPtr<FMovieSceneExportSectionData> InSectionData, const UMovieSceneSection* InSection, EMovieSceneTranslatorSectionType InSectionType, const FString& InSectionDisplayName);
+
+	/** Context for messages */
+	TSharedPtr<FMovieSceneTranslatorContext> ExportContext;
+
+public:
+
+	/** Intermediate data loaded from Sequencer to be used for export */
+	TSharedPtr<FMovieSceneExportMovieSceneData> MovieSceneData;
+
+private:
+
+	FFrameRate FrameRate;
+	FFrameTime HandleFrames;
+	FString SaveFilename;
+	FString SaveFilenamePath;
+	bool bExportDataIsValid;
+};
+
+struct FMovieSceneImportCinematicSectionData 
+{
+	UMovieSceneCinematicShotSection* CinematicSection;
+};
+
+struct FMovieSceneImportAudioSectionData
+{
+	UMovieSceneAudioSection* AudioSection;
+};
+
+struct FMovieSceneImportAudioTrackData
+{
+	UMovieSceneAudioTrack* AudioTrack;
+	TArray< TSharedPtr<FMovieSceneImportAudioSectionData> > AudioSectionsData;
+};
+
+struct FMovieSceneImportCinematicMasterTrackData
+{
+	UMovieSceneCinematicShotTrack* CinematicMasterTrack;
+	TArray< TSharedPtr<FMovieSceneImportCinematicSectionData> > CinematicSectionsData;
+};
+
+struct FMovieSceneImportMovieSceneData
+{
+	UMovieScene* MovieScene;
+	TSharedPtr<FMovieSceneImportCinematicMasterTrackData> CinematicMasterTrackData;
+	TArray< TSharedPtr<FMovieSceneImportAudioTrackData> > AudioTracksData;
+};
+
+/**
+* The FMovieSceneImportData class aggregates intermediate data from Sequencer classes to be used for timeline imports
+*/
+class FMovieSceneImportData : public TSharedFromThis<FMovieSceneImportData>
+{
+public:
+	/** Constructor */
+	FMovieSceneImportData(UMovieScene* InMovieScene, TSharedPtr<FMovieSceneTranslatorContext> InContext);
+	/** Default constructor, necessary for shared ref - should not be used */
+	FMovieSceneImportData();
+	/** Destructor */
+	~FMovieSceneImportData();
+
+	/** True when the export data was successfully constructed. */
+	bool IsImportDataValid() const;
+
+	/** Returns the cinematic master track data pointer or nullptr if one does not exist */
+	TSharedPtr<FMovieSceneImportCinematicMasterTrackData> GetCinematicMasterTrackData(bool CreateTrackIfNull);
+	/** Find cinematic section */
+	TSharedPtr<FMovieSceneImportCinematicSectionData> FindCinematicSection(const FString& InName);
+	/** Create cinematic section */
+	TSharedPtr<FMovieSceneImportCinematicSectionData> CreateCinematicSection(FString InName, int32 InRow, FFrameRate InFrameRate, FFrameNumber InStartFrame, FFrameNumber InEndFrame, FFrameNumber InStartOffsetFrame);
+	/** Set cinematic section */
+	bool SetCinematicSection(TSharedPtr<FMovieSceneImportCinematicSectionData> InSection, int32 InRow, FFrameRate InFrameRate, FFrameNumber InStartFrame, FFrameNumber InEndFrame, FFrameNumber InStartOffsetFrame);
+	/** Finds audio section */
+	TSharedPtr<FMovieSceneImportAudioSectionData> FindAudioSection(FString InName) const;
+
+private:
+
+	/** Entry point for setting up intermediate data for use when importing. */
+	TSharedPtr<FMovieSceneImportMovieSceneData> ConstructMovieSceneData(UMovieScene* InMovieScene);
+	/** Gets cinematic master track data from Sequencer. */
+	TSharedPtr<FMovieSceneImportCinematicMasterTrackData> ConstructCinematicMasterTrackData(UMovieSceneCinematicShotTrack* InCinematicMasterTrack);
+	/** Gets audio track data from Sequencer. */
+	TSharedPtr<FMovieSceneImportAudioTrackData> ConstructAudioTrackData(UMovieSceneAudioTrack* InAudioTrack);
+	/** Gets cinematic section data from Sequencer. */
+	TSharedPtr<FMovieSceneImportCinematicSectionData> ConstructCinematicSectionData(UMovieSceneCinematicShotSection* InCinematicSection);
+	/** Gets audio section data from Sequencer. */
+	TSharedPtr<FMovieSceneImportAudioSectionData>  ConstructAudioSectionData(UMovieSceneAudioSection* InAudioSection);
+
+	/** Context for messages */
+	TSharedPtr<FMovieSceneTranslatorContext> ImportContext;
+
+public:
+
+	/** Intermediate data loaded from Sequencer to be used for export */
+	TSharedPtr<FMovieSceneImportMovieSceneData> MovieSceneData;
+};
+
+/** Abstract base class for importer/exporter */
+class MOVIESCENETOOLS_API FMovieSceneTranslator
+{
+public:
+	FMovieSceneTranslator() {}
+	virtual ~FMovieSceneTranslator() {}
+
+	/** Error log window title. */
+	virtual FName GetMessageLogWindowTitle() const = 0;
+	/** Error log list label. */
+	virtual FText GetMessageLogLabel() const = 0;
+};
+
+/** Abstract base class for movie scene importers */
+class MOVIESCENETOOLS_API FMovieSceneImporter : public FMovieSceneTranslator
+{
+public:
+	FMovieSceneImporter() : FMovieSceneTranslator() {}
+
+	virtual ~FMovieSceneImporter() {}
+
+	/** Format description. */
+	virtual FText GetFileTypeDescription() const = 0;
+	/** Import window title. */
+	virtual FText GetDialogTitle() const = 0;
+	/** Scoped transaction description. */
+	virtual FText GetTransactionDescription() const = 0;
+
+public:
+	/*
+	* Import movie scene
+	*
+	* @param InMovieScene The movie scene to import the XML file into
+	* @param InFrameRate The frame rate to import the XML at
+	* @param InFilename The filename to import
+	* @param OutError The return error message
+	* @return Whether the import was successful
+	*/
+	virtual bool Import(UMovieScene* InMovieScene, FFrameRate InFrameRate, FString InFilename, TSharedRef<FMovieSceneTranslatorContext> InContext) = 0;
+};
+
+/** Abstract base class for movie scene exporters */
+class MOVIESCENETOOLS_API FMovieSceneExporter : public FMovieSceneTranslator
+{
+public:
+	FMovieSceneExporter() : FMovieSceneTranslator() {}
+
+	virtual ~FMovieSceneExporter() {}
+
+	/** Format description. */
+	virtual FText GetFileTypeDescription() const = 0;
+	/** Export dialog window title. */
+	virtual FText GetDialogTitle() const = 0;
+	/** Default format file extension. */
+	virtual FText GetDefaultFileExtension() const = 0;
+	/** Notification when export completes. */
+	virtual FText GetNotificationExportFinished() const = 0;
+	/** Notification hyperlink to exported file path. */
+	virtual FText GetNotificationHyperlinkText() const = 0;
+
+
+public:
+	/*
+	* Export movie scene
+	*
+	* @param InMovieScene The movie scene with the cinematic shot track and audio tracks to export
+	* @param InFrameRate The frame rate to export the XML at
+	* @param InSaveFilename The file path to save to.
+	* @param InHandleFrames The number of handle frames to include for each shot.
+	* @param OutError The return error message
+	* @return Whether the export was successful
+	*/
+	virtual bool Export(const UMovieScene* InMovieScene, FFrameRate InFrameRate, FString InSaveFilename, int32 InHandleFrames, TSharedRef<FMovieSceneTranslatorContext> InContext) = 0;
+};
+
diff --git a/Engine/Source/Editor/MovieSceneTools/Public/MovieSceneTranslatorEDL.h b/Engine/Source/Editor/MovieSceneTools/Public/MovieSceneTranslatorEDL.h
new file mode 100644
--- /dev/null
+++ b/Engine/Source/Editor/MovieSceneTools/Public/MovieSceneTranslatorEDL.h
@@ -0,0 +1,35 @@
+// Copyright 1998-2018 Epic Games, Inc. All Rights Reserved.
+
+#pragma once
+
+#include "CoreMinimal.h"
+
+struct FFrameRate;
+struct FFrameNumber;
+class UMovieScene;
+
+class MOVIESCENETOOLS_API MovieSceneTranslatorEDL
+{
+public:
+
+	/**
+	 * Import EDL
+	 *
+	 * @param InMovieScene The movie scene to import the edl into
+	 * @param InFrameRate The frame rate to import the EDL at
+	 * @param InFilename The filename to import
+	 * @return Whether the import was successful
+	 */
+	static bool ImportEDL(UMovieScene* InMovieScene, FFrameRate InFrameRate, FString InFilename);
+
+	/**
+	 * Export EDL
+	 *
+	 * @param InMovieScene The movie scene with the cinematic shot track and audio tracks to export
+	 * @param InFrameRate The frame rate to export the EDL at
+	 * @param InSaveFilename The file path to save to.
+	 * @param InHandleFrames The number of handle frames to include for each shot.
+	 * @return Whether the export was successful
+	 */
+	static bool ExportEDL(const UMovieScene* InMovieScene, FFrameRate InFrameRate, FString InSaveFilename, int32 InHandleFrames);
+};
diff --git a/Engine/Source/Editor/MovieSceneTools/Public/TrackEditors/AudioTrackEditor.h b/Engine/Source/Editor/MovieSceneTools/Public/TrackEditors/AudioTrackEditor.h
--- a/Engine/Source/Editor/MovieSceneTools/Public/TrackEditors/AudioTrackEditor.h
+++ b/Engine/Source/Editor/MovieSceneTools/Public/TrackEditors/AudioTrackEditor.h
@@ -79,6 +79,9 @@ private:
 
 	/** Audio asset selected */
 	void OnAudioAssetSelected(const FAssetData& AssetData, UMovieSceneTrack* Track);
+
+	/** Audio asset enter pressed */
+	void OnAudioAssetEnterPressed(const TArray<FAssetData>& AssetData, UMovieSceneTrack* Track);
 };
 
 
diff --git a/Engine/Source/Editor/MovieSceneTools/Public/TrackEditors/SubTrackEditor.h b/Engine/Source/Editor/MovieSceneTools/Public/TrackEditors/SubTrackEditor.h
--- a/Engine/Source/Editor/MovieSceneTools/Public/TrackEditors/SubTrackEditor.h
+++ b/Engine/Source/Editor/MovieSceneTools/Public/TrackEditors/SubTrackEditor.h
@@ -83,6 +83,9 @@ private:
 	/** Callback for executing a menu entry in the "Add Sequence" combo button. */
 	void HandleAddSubSequenceComboButtonMenuEntryExecute(const FAssetData& AssetData, UMovieSceneTrack* InTrack);
 
+	/** Callback for executing a menu entry in the "Add Sequence" combo button when enter pressed. */
+	void HandleAddSubSequenceComboButtonMenuEntryEnterPressed(const TArray<FAssetData>& AssetData, UMovieSceneTrack* InTrack);
+
 	/** Delegate for AnimatablePropertyChanged in AddKey */
 	FKeyPropertyResult AddKeyInternal(FFrameNumber KeyTime, UMovieSceneSequence* InMovieSceneSequence, UMovieSceneTrack* InTrack, int32 RowIndex);
 
diff --git a/Engine/Source/Editor/PropertyEditor/Private/SDetailsView.cpp b/Engine/Source/Editor/PropertyEditor/Private/SDetailsView.cpp
--- a/Engine/Source/Editor/PropertyEditor/Private/SDetailsView.cpp
+++ b/Engine/Source/Editor/PropertyEditor/Private/SDetailsView.cpp
@@ -43,10 +43,12 @@ void SDetailsView::Construct(const FArguments& InArgs)
 
 	PropertyUtilities = MakeShareable( new FPropertyDetailsUtilities( *this ) );
 	
+	ColumnWidth = DetailsViewArgs.ColumnWidth;
+
 	ColumnSizeData.LeftColumnWidth = TAttribute<float>( this, &SDetailsView::OnGetLeftColumnWidth );
 	ColumnSizeData.RightColumnWidth = TAttribute<float>( this, &SDetailsView::OnGetRightColumnWidth );
 	ColumnSizeData.OnWidthChanged = SSplitter::FOnSlotResized::CreateSP( this, &SDetailsView::OnSetColumnWidth );
-
+	
 	// We want the scrollbar to always be visible when objects are selected, but not when there is no selection - however:
 	//  - We can't use AlwaysShowScrollbar for this, as this will also show the scrollbar when nothing is selected
 	//  - We can't use the Visibility construction parameter, as it gets translated into user visibility and can hide the scrollbar even when objects are selected
diff --git a/Engine/Source/Editor/PropertyEditor/Public/IDetailsView.h b/Engine/Source/Editor/PropertyEditor/Public/IDetailsView.h
--- a/Engine/Source/Editor/PropertyEditor/Public/IDetailsView.h
+++ b/Engine/Source/Editor/PropertyEditor/Public/IDetailsView.h
@@ -94,6 +94,8 @@ struct FDetailsViewArgs
 	uint32 bShowKeyablePropertiesOption : 1;
 	/** True if you want to show the 'Show Only Animated Properties'. Only valid in conjunction with bShowOptions */
 	uint32 bShowAnimatedPropertiesOption: 1;
+	/** The default column width */
+	float ColumnWidth;
 
 public:
 	/** Default constructor */
@@ -121,12 +123,13 @@ public:
 		, bShowDifferingPropertiesOption(false)
 		, bCustomNameAreaLocation(false)
 		, bCustomFilterAreaLocation(false)
-		, bAllowFavoriteSystem(true)
+		, bAllowFavoriteSystem(true)	
 		, bAllowMultipleTopLevelObjects(false)
 		, bShowScrollBar(true)
 		, bForceHiddenPropertyVisibility(false)
 		, bShowKeyablePropertiesOption(true)
 		, bShowAnimatedPropertiesOption(true)
+		, ColumnWidth(.65f)
 	{
 	}
 };
diff --git a/Engine/Source/Editor/SequenceRecorder/Private/ActorRecording.h b/Engine/Source/Editor/SequenceRecorder/Private/ActorRecording.h
--- a/Engine/Source/Editor/SequenceRecorder/Private/ActorRecording.h
+++ b/Engine/Source/Editor/SequenceRecorder/Private/ActorRecording.h
@@ -102,7 +102,7 @@ public:
 	FAnimationRecordingSettings AnimationSettings;
 
 	/** Whether to record to 'possessable' (i.e. level-owned) or 'spawnable' (i.e. sequence-owned) actors. Defaults to the global setting. */
-	UPROPERTY(EditAnywhere, AdvancedDisplay, Category = "Animation Recording")
+	UPROPERTY(EditAnywhere, Category = "Actor Recording")
 	bool bRecordToPossessable;
 
 	/** Whether this actor recording was triggered from an actor spawn */
diff --git a/Engine/Source/Editor/SequenceRecorder/Private/SSequenceRecorder.cpp b/Engine/Source/Editor/SequenceRecorder/Private/SSequenceRecorder.cpp
--- a/Engine/Source/Editor/SequenceRecorder/Private/SSequenceRecorder.cpp
+++ b/Engine/Source/Editor/SequenceRecorder/Private/SSequenceRecorder.cpp
@@ -207,6 +207,7 @@ void SSequenceRecorder::Construct(const FArguments& Args)
 	ToolBarBuilder.BeginSection(TEXT("RecordingManagement"));
 	{
 		ToolBarBuilder.AddToolBarButton(FSequenceRecorderCommands::Get().AddRecording);
+		ToolBarBuilder.AddToolBarButton(FSequenceRecorderCommands::Get().AddCurrentPlayerRecording);
 		ToolBarBuilder.AddToolBarButton(FSequenceRecorderCommands::Get().RemoveRecording);
 		ToolBarBuilder.AddToolBarButton(FSequenceRecorderCommands::Get().RemoveAllRecordings);
 	}
@@ -377,6 +378,11 @@ void SSequenceRecorder::BindCommands()
 		FCanExecuteAction::CreateSP(this, &SSequenceRecorder::CanAddRecording)
 		);
 
+	CommandList->MapAction(FSequenceRecorderCommands::Get().AddCurrentPlayerRecording,
+		FExecuteAction::CreateSP(this, &SSequenceRecorder::HandleAddCurrentPlayerRecording),
+		FCanExecuteAction::CreateSP(this, &SSequenceRecorder::CanAddCurrentPlayerRecording)
+		);
+
 	CommandList->MapAction(FSequenceRecorderCommands::Get().RemoveRecording,
 		FExecuteAction::CreateSP(this, &SSequenceRecorder::HandleRemoveRecording),
 		FCanExecuteAction::CreateSP(this, &SSequenceRecorder::CanRemoveRecording)
@@ -466,6 +472,16 @@ bool SSequenceRecorder::CanAddRecording() const
 	return !FAnimationRecorderManager::Get().IsRecording();
 }
 
+void SSequenceRecorder::HandleAddCurrentPlayerRecording()
+{
+	FSequenceRecorder::Get().AddNewQueuedRecordingForCurrentPlayer();
+}
+
+bool SSequenceRecorder::CanAddCurrentPlayerRecording() const
+{
+	return FSequenceRecorder::Get().CanAddNewQueuedRecordingForCurrentPlayer();
+}
+
 void SSequenceRecorder::HandleRemoveRecording()
 {
 	TArray<UActorRecording*> SelectedRecordings;
diff --git a/Engine/Source/Editor/SequenceRecorder/Private/SSequenceRecorder.h b/Engine/Source/Editor/SequenceRecorder/Private/SSequenceRecorder.h
--- a/Engine/Source/Editor/SequenceRecorder/Private/SSequenceRecorder.h
+++ b/Engine/Source/Editor/SequenceRecorder/Private/SSequenceRecorder.h
@@ -62,6 +62,10 @@ private:
 
 	bool CanAddRecording() const;
 
+	void HandleAddCurrentPlayerRecording();
+
+	bool CanAddCurrentPlayerRecording() const;
+
 	void HandleRemoveRecording();
 
 	bool CanRemoveRecording() const;
diff --git a/Engine/Source/Editor/SequenceRecorder/Private/Sections/MovieScene3DTransformSectionRecorder.cpp b/Engine/Source/Editor/SequenceRecorder/Private/Sections/MovieScene3DTransformSectionRecorder.cpp
--- a/Engine/Source/Editor/SequenceRecorder/Private/Sections/MovieScene3DTransformSectionRecorder.cpp
+++ b/Engine/Source/Editor/SequenceRecorder/Private/Sections/MovieScene3DTransformSectionRecorder.cpp
@@ -131,8 +131,8 @@ void FMovieScene3DTransformSectionRecorder::FinalizeSection()
 
 				check(RootIndex != INDEX_NONE);
 
-				FFrameRate FrameResolution = MovieSceneSection->GetTypedOuter<UMovieScene>()->GetFrameResolution();
-				const FFrameNumber StartTime = (RecordingStartTime * FrameResolution).FloorToFrame();
+				FFrameRate TickResolution = MovieSceneSection->GetTypedOuter<UMovieScene>()->GetTickResolution();
+				const FFrameNumber StartTime = (RecordingStartTime * TickResolution).FloorToFrame();
 
 				// we may need to offset the transform here if the animation was not recorded on the root component
 				FTransform InvComponentTransform = AnimRecorder->GetComponentTransform().Inverse();
@@ -169,7 +169,7 @@ void FMovieScene3DTransformSectionRecorder::FinalizeSection()
 						Transform.SetScale3D(RawTrack.ScaleKeys[0]);
 					}
 
-					FFrameNumber AnimationFrame = (AnimSequence->GetTimeAtFrame(KeyIndex) * FrameResolution).FloorToFrame();
+					FFrameNumber AnimationFrame = (AnimSequence->GetTimeAtFrame(KeyIndex) * TickResolution).FloorToFrame();
 					BufferedTransforms.Add(InvComponentTransform * Transform, StartTime + AnimationFrame);
 				}
 			}
@@ -326,8 +326,8 @@ void FMovieScene3DTransformSectionRecorder::Record(float CurrentTime)
 			}
 		}
 
-		FFrameRate   FrameResolution = MovieSceneSection->GetTypedOuter<UMovieScene>()->GetFrameResolution();
-		FFrameNumber CurrentFrame    = (CurrentTime * FrameResolution).FloorToFrame();
+		FFrameRate   TickResolution  = MovieSceneSection->GetTypedOuter<UMovieScene>()->GetTickResolution();
+		FFrameNumber CurrentFrame    = (CurrentTime * TickResolution).FloorToFrame();
 
 		if(bRecording)
 		{
diff --git a/Engine/Source/Editor/SequenceRecorder/Private/Sections/MovieSceneAnimationSectionRecorder.cpp b/Engine/Source/Editor/SequenceRecorder/Private/Sections/MovieSceneAnimationSectionRecorder.cpp
--- a/Engine/Source/Editor/SequenceRecorder/Private/Sections/MovieSceneAnimationSectionRecorder.cpp
+++ b/Engine/Source/Editor/SequenceRecorder/Private/Sections/MovieSceneAnimationSectionRecorder.cpp
@@ -103,8 +103,8 @@ void FMovieSceneAnimationSectionRecorder::CreateSection(UObject* InObjectToRecor
 					UMovieSceneSkeletalAnimationTrack* AnimTrack = MovieScene->AddTrack<UMovieSceneSkeletalAnimationTrack>(Guid);
 					if (AnimTrack)
 					{
-						FFrameRate   FrameResolution = AnimTrack->GetTypedOuter<UMovieScene>()->GetFrameResolution();
-						FFrameNumber CurrentFrame    = (Time * FrameResolution).FloorToFrame();
+						FFrameRate   TickResolution  = AnimTrack->GetTypedOuter<UMovieScene>()->GetTickResolution();
+						FFrameNumber CurrentFrame    = (Time * TickResolution).FloorToFrame();
 
 						AnimTrack->AddNewAnimation(CurrentFrame, AnimSequence.Get());
 						MovieSceneSection = Cast<UMovieSceneSkeletalAnimationSection>(AnimTrack->GetAllSections()[0]);
@@ -135,8 +135,8 @@ void FMovieSceneAnimationSectionRecorder::FinalizeSection()
 
 	if(MovieSceneSection.IsValid() && AnimSequence.IsValid() && MovieSceneSection->HasStartFrame())
 	{
-		FFrameRate   FrameResolution = MovieSceneSection->GetTypedOuter<UMovieScene>()->GetFrameResolution();
-		FFrameNumber SequenceLength  = (AnimSequence->GetPlayLength() * FrameResolution).FloorToFrame();
+		FFrameRate   TickResolution  = MovieSceneSection->GetTypedOuter<UMovieScene>()->GetTickResolution();
+		FFrameNumber SequenceLength  = (AnimSequence->GetPlayLength() * TickResolution).FloorToFrame();
 		
 		MovieSceneSection->SetEndFrame(TRangeBound<FFrameNumber>::Exclusive(MovieSceneSection->GetInclusiveStartFrame() + SequenceLength));
 	}
diff --git a/Engine/Source/Editor/SequenceRecorder/Private/Sections/MovieScenePropertyRecorder.cpp b/Engine/Source/Editor/SequenceRecorder/Private/Sections/MovieScenePropertyRecorder.cpp
--- a/Engine/Source/Editor/SequenceRecorder/Private/Sections/MovieScenePropertyRecorder.cpp
+++ b/Engine/Source/Editor/SequenceRecorder/Private/Sections/MovieScenePropertyRecorder.cpp
@@ -37,8 +37,8 @@ UMovieSceneSection* FMovieScenePropertyRecorder<bool>::AddSection(UObject* InObj
 
 		UMovieSceneBoolSection* Section         = Cast<UMovieSceneBoolSection>(Track->CreateNewSection());
 
-		FFrameRate   FrameResolution = Section->GetTypedOuter<UMovieScene>()->GetFrameResolution();
-		FFrameNumber CurrentFrame    = (InTime * FrameResolution).FloorToFrame();
+		FFrameRate   TickResolution  = Section->GetTypedOuter<UMovieScene>()->GetTickResolution();
+		FFrameNumber CurrentFrame    = (InTime * TickResolution).FloorToFrame();
 
 		Section->SetRange(TRange<FFrameNumber>(CurrentFrame));
 
@@ -84,8 +84,8 @@ UMovieSceneSection* FMovieScenePropertyRecorder<uint8>::AddSection(UObject* InOb
 
 		UMovieSceneByteSection* Section = Cast<UMovieSceneByteSection>(Track->CreateNewSection());
 
-		FFrameRate   FrameResolution = Section->GetTypedOuter<UMovieScene>()->GetFrameResolution();
-		FFrameNumber CurrentFrame    = (InTime * FrameResolution).FloorToFrame();
+		FFrameRate   TickResolution  = Section->GetTypedOuter<UMovieScene>()->GetTickResolution();
+		FFrameNumber CurrentFrame    = (InTime * TickResolution).FloorToFrame();
 
 		Section->SetRange(TRange<FFrameNumber>::Inclusive(CurrentFrame, CurrentFrame));
 
@@ -129,8 +129,8 @@ UMovieSceneSection* FMovieScenePropertyRecorderEnum::AddSection(UObject* InObjec
 
 		UMovieSceneEnumSection* Section = Cast<UMovieSceneEnumSection>(Track->CreateNewSection());
 
-		FFrameRate   FrameResolution = Section->GetTypedOuter<UMovieScene>()->GetFrameResolution();
-		FFrameNumber CurrentFrame    = (InTime * FrameResolution).FloorToFrame();
+		FFrameRate   TickResolution  = Section->GetTypedOuter<UMovieScene>()->GetTickResolution();
+		FFrameNumber CurrentFrame    = (InTime * TickResolution).FloorToFrame();
 
 		Section->SetRange(TRange<FFrameNumber>::Inclusive(CurrentFrame, CurrentFrame));
 
@@ -174,8 +174,8 @@ UMovieSceneSection* FMovieScenePropertyRecorder<float>::AddSection(UObject* InOb
 
 		UMovieSceneFloatSection* Section = Cast<UMovieSceneFloatSection>(Track->CreateNewSection());
 
-		FFrameRate   FrameResolution = Section->GetTypedOuter<UMovieScene>()->GetFrameResolution();
-		FFrameNumber CurrentFrame    = (InTime * FrameResolution).FloorToFrame();
+		FFrameRate   TickResolution  = Section->GetTypedOuter<UMovieScene>()->GetTickResolution();
+		FFrameNumber CurrentFrame    = (InTime * TickResolution).FloorToFrame();
 
 		Section->SetRange(TRange<FFrameNumber>::Inclusive(CurrentFrame, CurrentFrame));
 
@@ -225,8 +225,8 @@ UMovieSceneSection* FMovieScenePropertyRecorder<FColor>::AddSection(UObject* InO
 
 		UMovieSceneColorSection* Section = Cast<UMovieSceneColorSection>(Track->CreateNewSection());
 
-		FFrameRate   FrameResolution = Section->GetTypedOuter<UMovieScene>()->GetFrameResolution();
-		FFrameNumber CurrentFrame    = (InTime * FrameResolution).FloorToFrame();
+		FFrameRate   TickResolution  = Section->GetTypedOuter<UMovieScene>()->GetTickResolution();
+		FFrameNumber CurrentFrame    = (InTime * TickResolution).FloorToFrame();
 
 		Section->SetRange(TRange<FFrameNumber>::Inclusive(CurrentFrame, CurrentFrame));
 
@@ -294,8 +294,8 @@ UMovieSceneSection* FMovieScenePropertyRecorder<FVector>::AddSection(UObject* In
 
 		UMovieSceneVectorSection* Section = Cast<UMovieSceneVectorSection>(Track->CreateNewSection());
 
-		FFrameRate   FrameResolution = Section->GetTypedOuter<UMovieScene>()->GetFrameResolution();
-		FFrameNumber CurrentFrame    = (InTime * FrameResolution).FloorToFrame();
+		FFrameRate   TickResolution  = Section->GetTypedOuter<UMovieScene>()->GetTickResolution();
+		FFrameNumber CurrentFrame    = (InTime * TickResolution).FloorToFrame();
 
 		Section->SetRange(TRange<FFrameNumber>::Inclusive(CurrentFrame, CurrentFrame));
 
diff --git a/Engine/Source/Editor/SequenceRecorder/Private/Sections/MovieScenePropertyRecorder.h b/Engine/Source/Editor/SequenceRecorder/Private/Sections/MovieScenePropertyRecorder.h
--- a/Engine/Source/Editor/SequenceRecorder/Private/Sections/MovieScenePropertyRecorder.h
+++ b/Engine/Source/Editor/SequenceRecorder/Private/Sections/MovieScenePropertyRecorder.h
@@ -53,8 +53,8 @@ public:
 	{
 		if (InObjectToRecord != nullptr)
 		{
-			FFrameRate   FrameResolution = MovieSceneSection->GetTypedOuter<UMovieScene>()->GetFrameResolution();
-			FFrameNumber CurrentFrame    = (InCurrentTime * FrameResolution).FloorToFrame();
+			FFrameRate   TickResolution  = MovieSceneSection->GetTypedOuter<UMovieScene>()->GetTickResolution();
+			FFrameNumber CurrentFrame    = (InCurrentTime * TickResolution).FloorToFrame();
 
 			MovieSceneSection->ExpandToFrame(CurrentFrame);
 
@@ -140,8 +140,8 @@ public:
 	{
 		if (InObjectToRecord != nullptr)
 		{
-			FFrameRate   FrameResolution = MovieSceneSection->GetTypedOuter<UMovieScene>()->GetFrameResolution();
-			FFrameNumber CurrentFrame    = (InCurrentTime * FrameResolution).FloorToFrame();
+			FFrameRate   TickResolution  = MovieSceneSection->GetTypedOuter<UMovieScene>()->GetTickResolution();
+			FFrameNumber CurrentFrame    = (InCurrentTime * TickResolution).FloorToFrame();
 
 			MovieSceneSection->ExpandToFrame(CurrentFrame);
 
diff --git a/Engine/Source/Editor/SequenceRecorder/Private/SequenceRecorder.cpp b/Engine/Source/Editor/SequenceRecorder/Private/SequenceRecorder.cpp
--- a/Engine/Source/Editor/SequenceRecorder/Private/SequenceRecorder.cpp
+++ b/Engine/Source/Editor/SequenceRecorder/Private/SequenceRecorder.cpp
@@ -153,6 +153,57 @@ void FSequenceRecorder::AddNewQueuedRecordingsForSelectedActors()
 	}
 }
 
+/** Helper function - get the first PIE world (or first PIE client world if there is more than one) */
+static UWorld* GetFirstPIEWorld()
+{
+	for (const FWorldContext& Context : GEngine->GetWorldContexts())
+	{
+		if (Context.World()->IsPlayInEditor())
+		{
+			if(Context.World()->GetNetMode() == ENetMode::NM_Standalone ||
+				(Context.World()->GetNetMode() == ENetMode::NM_Client && Context.PIEInstance == 2))
+			{
+				return Context.World();
+			}
+		}
+	}
+
+	return nullptr;
+}
+
+void FSequenceRecorder::AddNewQueuedRecordingForCurrentPlayer()
+{
+	if (UWorld* PIEWorld = GetFirstPIEWorld())
+	{
+		APlayerController* Controller = GEngine->GetFirstLocalPlayerController(PIEWorld);
+		if(Controller && Controller->GetPawn())
+		{
+			APawn* CurrentPlayer = Controller->GetPawn();
+			if (!FindRecording(CurrentPlayer))
+			{
+				AddNewQueuedRecording(CurrentPlayer);
+			}
+		}
+	}
+}
+
+bool FSequenceRecorder::CanAddNewQueuedRecordingForCurrentPlayer() const
+{
+	if (UWorld* PIEWorld = GetFirstPIEWorld())
+	{
+		APlayerController* Controller = GEngine->GetFirstLocalPlayerController(PIEWorld);
+		if(Controller && Controller->GetPawn())
+		{
+			APawn* CurrentPlayer = Controller->GetPawn();
+			if (!FindRecording(CurrentPlayer))
+			{
+				return true;
+			}
+		}
+	}
+	return false;
+}
+
 UActorRecording* FSequenceRecorder::AddNewQueuedRecording(AActor* Actor, UAnimSequence* AnimSequence, float Length)
 {
 	const USequenceRecorderSettings* Settings = GetDefault<USequenceRecorderSettings>();
@@ -288,6 +339,11 @@ void FSequenceRecorder::Tick(float DeltaSeconds)
 		{
 			CurrentDelay = 0.0f;
 			StartRecordingInternal(nullptr);
+
+			if (!IsRecording())
+			{
+				RestoreImmersive();
+			}
 		}
 	}
 
@@ -474,23 +530,7 @@ bool FSequenceRecorder::StartRecording(const FString& InPathToRecordTo, const FS
 
 	CurrentTime = 0.0f;
 
-	if (Settings->bImmersiveMode)
-	{
-		FLevelEditorModule& LevelEditorModule = FModuleManager::Get().LoadModuleChecked<FLevelEditorModule>("LevelEditor");
-		TSharedPtr< ILevelViewport > ActiveLevelViewport = LevelEditorModule.GetFirstActiveViewport();
-
-		if( ActiveLevelViewport.IsValid() )
-		{
-			bWasImmersive = ActiveLevelViewport->IsImmersive();
-
-			if (!ActiveLevelViewport->IsImmersive())
-			{
-				const bool bWantImmersive = true;
-				const bool bAllowAnimation = false;
-				ActiveLevelViewport->MakeImmersive( bWantImmersive, bAllowAnimation );
-			}
-		}
-	}
+	SetImmersive();
 
 	RefreshNextSequence();
 
@@ -732,6 +772,10 @@ void FSequenceRecorder::HandleEndPIE(bool bSimulating)
 
 bool FSequenceRecorder::StopRecording()
 {
+	const USequenceRecorderSettings* Settings = GetDefault<USequenceRecorderSettings>();
+
+	RestoreImmersive();
+
 	if (!IsRecording())
 	{
 		CurrentDelay = 0.0f;
@@ -739,23 +783,6 @@ bool FSequenceRecorder::StopRecording()
 		return false;
 	}
 
-	const USequenceRecorderSettings* Settings = GetDefault<USequenceRecorderSettings>();
-
-	if (Settings->bImmersiveMode)
-	{
-		FLevelEditorModule& LevelEditorModule = FModuleManager::Get().LoadModuleChecked<FLevelEditorModule>("LevelEditor");
-		TSharedPtr< ILevelViewport > ActiveLevelViewport = LevelEditorModule.GetFirstActiveViewport();
-
-		if( ActiveLevelViewport.IsValid() )
-		{
-			if (ActiveLevelViewport->IsImmersive() != bWasImmersive)
-			{
-				const bool bAllowAnimation = false;
-				ActiveLevelViewport->MakeImmersive(bWasImmersive, bAllowAnimation);
-			}
-		}
-	}
-
 	// 1 step for the audio processing
 	static const uint8 NumAdditionalSteps = 1;
 
@@ -792,11 +819,11 @@ bool FSequenceRecorder::StopRecording()
 
 			UMovieSceneAudioSection* NewAudioSection = NewObject<UMovieSceneAudioSection>(AudioTrack, UMovieSceneAudioSection::StaticClass());
 
-			FFrameRate FrameResolution = AudioTrack->GetTypedOuter<UMovieScene>()->GetFrameResolution();
+			FFrameRate TickResolution = AudioTrack->GetTypedOuter<UMovieScene>()->GetTickResolution();
 
 			NewAudioSection->SetRowIndex(RowIndex + 1);
 			NewAudioSection->SetSound(RecordedAudio);
-			NewAudioSection->SetRange(TRange<FFrameNumber>(FFrameNumber(0), (RecordedAudio->GetDuration() * FrameResolution).CeilToFrame()));
+			NewAudioSection->SetRange(TRange<FFrameNumber>(FFrameNumber(0), (RecordedAudio->GetDuration() * TickResolution).CeilToFrame()));
 
 			AudioTrack->AddSection(*NewAudioSection);
 		}
@@ -893,7 +920,7 @@ bool FSequenceRecorder::StopRecording()
 				UPackage::SavePackage(Package, NULL, RF_Standalone, *PackageFileName, GError, nullptr, false, true, SAVE_NoError);
 			}
 
-			if(FSlateApplication::IsInitialized())
+			if(FSlateApplication::IsInitialized() && GIsEditor)
 			{
 				const FText NotificationText = FText::Format(LOCTEXT("RecordSequence", "'{0}' has been successfully recorded."), FText::FromString(LevelSequence->GetName()));
 					
@@ -955,11 +982,11 @@ void FSequenceRecorder::UpdateSequencePlaybackRange()
 			MovieScene->SetPlaybackRange(PlayRange);
 
 			// Initialize the working and view range with a little bit more space
-			FFrameRate  FrameResolution = MovieScene->GetFrameResolution();
-			const double OutputViewSize = PlayRange.Size<FFrameNumber>() / FrameResolution;
+			FFrameRate  TickResolution = MovieScene->GetTickResolution();
+			const double OutputViewSize = PlayRange.Size<FFrameNumber>() / TickResolution;
 			const double OutputChange   = OutputViewSize * 0.1;
 
-			TRange<double> NewRange = MovieScene::ExpandRange(PlayRange / FrameResolution, OutputChange);
+			TRange<double> NewRange = MovieScene::ExpandRange(PlayRange / TickResolution, OutputChange);
 			FMovieSceneEditorData& EditorData = MovieScene->GetEditorData();
 			EditorData.ViewStart = EditorData.WorkStart = NewRange.GetLowerBoundValue();
 			EditorData.ViewEnd   = EditorData.WorkEnd   = NewRange.GetUpperBoundValue();
@@ -1055,6 +1082,17 @@ void FSequenceRecorder::HandleActorDespawned(AActor* Actor)
 }
 
 void FSequenceRecorder::RefreshNextSequence()
+{
+	if (SequenceName.IsEmpty())
+	{
+		SequenceName = GetSequenceRecordingName().Len() > 0 ? GetSequenceRecordingName() : TEXT("RecordedSequence");
+	}
+
+	// Cache the name of the next sequence we will try to record to
+	NextSequenceName = SequenceRecorderUtils::MakeNewAssetName<ULevelSequence>(GetSequenceRecordingBasePath(), SequenceName);
+}
+
+void FSequenceRecorder::ForceRefreshNextSequence()
 {
 	SequenceName = GetSequenceRecordingName().Len() > 0 ? GetSequenceRecordingName() : TEXT("RecordedSequence");
 
@@ -1243,4 +1281,49 @@ FString FSequenceRecorder::GetSequenceRecordingName() const
 	// If no profile is loaded, just return the default value.
 	return GetDefault<USequenceRecorderActorGroup>()->SequenceName;
 }
+
+void FSequenceRecorder::SetImmersive()
+{
+	const USequenceRecorderSettings* Settings = GetDefault<USequenceRecorderSettings>();
+
+	if (Settings->bImmersiveMode)
+	{
+		FLevelEditorModule& LevelEditorModule = FModuleManager::Get().LoadModuleChecked<FLevelEditorModule>("LevelEditor");
+		TSharedPtr< ILevelViewport > ActiveLevelViewport = LevelEditorModule.GetFirstActiveViewport();
+
+		if( ActiveLevelViewport.IsValid() )
+		{
+			bWasImmersive = ActiveLevelViewport->IsImmersive();
+
+			if (!ActiveLevelViewport->IsImmersive())
+			{
+				const bool bWantImmersive = true;
+				const bool bAllowAnimation = false;
+				ActiveLevelViewport->MakeImmersive( bWantImmersive, bAllowAnimation );
+			}
+		}
+	}
+
+}
+
+void FSequenceRecorder::RestoreImmersive()
+{
+	const USequenceRecorderSettings* Settings = GetDefault<USequenceRecorderSettings>();
+
+	if (Settings->bImmersiveMode)
+	{
+		FLevelEditorModule& LevelEditorModule = FModuleManager::Get().LoadModuleChecked<FLevelEditorModule>("LevelEditor");
+		TSharedPtr< ILevelViewport > ActiveLevelViewport = LevelEditorModule.GetFirstActiveViewport();
+
+		if( ActiveLevelViewport.IsValid() )
+		{
+			if (ActiveLevelViewport->IsImmersive() != bWasImmersive)
+			{
+				const bool bAllowAnimation = false;
+				ActiveLevelViewport->MakeImmersive(bWasImmersive, bAllowAnimation);
+			}
+		}
+	}
+}
+
 #undef LOCTEXT_NAMESPACE
diff --git a/Engine/Source/Editor/SequenceRecorder/Private/SequenceRecorder.h b/Engine/Source/Editor/SequenceRecorder/Private/SequenceRecorder.h
--- a/Engine/Source/Editor/SequenceRecorder/Private/SequenceRecorder.h
+++ b/Engine/Source/Editor/SequenceRecorder/Private/SequenceRecorder.h
@@ -63,6 +63,10 @@ public:
 
 	void AddNewQueuedRecordingsForSelectedActors();
 
+	void AddNewQueuedRecordingForCurrentPlayer();
+
+	bool CanAddNewQueuedRecordingForCurrentPlayer() const;
+
 	class UActorRecording* AddNewQueuedRecording(AActor* Actor = nullptr, UAnimSequence* AnimSequence = nullptr, float Length = 0.0f);
 
 	void RemoveQueuedRecording(AActor* Actor);
@@ -125,6 +129,9 @@ public:
 	/** Refresh the name of the next sequence we will be recording */
 	void RefreshNextSequence();
 
+	/** Force refresh the name of the next sequence, disregard the current sequence name */
+	void ForceRefreshNextSequence();
+
 	/** Multicast delegate fired when recording is started */
 	FOnRecordingStarted OnRecordingStartedDelegate;
 
@@ -144,6 +151,12 @@ private:
 	/** Keep sequence range up to date with sections that are being recorded */
 	void UpdateSequencePlaybackRange();
 
+	/** Set immersive mode and store whether viewports were immersive */
+	void SetImmersive();
+
+	/** Restore immersive mode to stored value */
+	void RestoreImmersive();
+
 private:
 	/** Constructor, private - use Get() function */
 	FSequenceRecorder();
diff --git a/Engine/Source/Editor/SequenceRecorder/Private/SequenceRecorderActorGroup.h b/Engine/Source/Editor/SequenceRecorder/Private/SequenceRecorderActorGroup.h
--- a/Engine/Source/Editor/SequenceRecorder/Private/SequenceRecorderActorGroup.h
+++ b/Engine/Source/Editor/SequenceRecorder/Private/SequenceRecorderActorGroup.h
@@ -42,7 +42,7 @@ public:
 			if (PropertyChangedEvent.Property->GetFName() == GET_MEMBER_NAME_CHECKED(USequenceRecorderActorGroup, SequenceName) ||
 				PropertyChangedEvent.Property->GetFName() == GET_MEMBER_NAME_CHECKED(USequenceRecorderActorGroup, SequenceRecordingBasePath))
 			{
-				FSequenceRecorder::Get().RefreshNextSequence();
+				FSequenceRecorder::Get().ForceRefreshNextSequence();
 			}
 		}
 	}
diff --git a/Engine/Source/Editor/SequenceRecorder/Private/SequenceRecorderCommands.cpp b/Engine/Source/Editor/SequenceRecorder/Private/SequenceRecorderCommands.cpp
--- a/Engine/Source/Editor/SequenceRecorder/Private/SequenceRecorderCommands.cpp
+++ b/Engine/Source/Editor/SequenceRecorder/Private/SequenceRecorderCommands.cpp
@@ -15,6 +15,7 @@ void FSequenceRecorderCommands::RegisterCommands()
 	UI_COMMAND(RecordAll, "Record", "Record all recordings", EUserInterfaceActionType::Button, FInputChord(EModifierKey::Control | EModifierKey::Shift, EKeys::R));
 	UI_COMMAND(StopAll, "StopAll", "Stop all recordings", EUserInterfaceActionType::Button, FInputChord(EModifierKey::Control | EModifierKey::Shift, EKeys::S));
 	UI_COMMAND(AddRecording, "Add", "Add a new recording for selected actors. If nothing selected, an empty recording will be added", EUserInterfaceActionType::Button, FInputChord());
+	UI_COMMAND(AddCurrentPlayerRecording, "Add Player", "Add a new recording for the current player.", EUserInterfaceActionType::Button, FInputChord());
 	UI_COMMAND(RemoveRecording, "Remove", "Remove selected recording", EUserInterfaceActionType::Button, FInputChord());
 	UI_COMMAND(RemoveAllRecordings, "Remove All", "Remove all recordings", EUserInterfaceActionType::Button, FInputChord());
 	UI_COMMAND(AddRecordingGroup, "Add Group", "Adds a new group for Actor recordings to be stored with", EUserInterfaceActionType::Button, FInputChord());
diff --git a/Engine/Source/Editor/SequenceRecorder/Private/SequenceRecorderCommands.h b/Engine/Source/Editor/SequenceRecorder/Private/SequenceRecorderCommands.h
--- a/Engine/Source/Editor/SequenceRecorder/Private/SequenceRecorderCommands.h
+++ b/Engine/Source/Editor/SequenceRecorder/Private/SequenceRecorderCommands.h
@@ -13,6 +13,7 @@ public:
 	TSharedPtr<FUICommandInfo> RecordAll;
 	TSharedPtr<FUICommandInfo> StopAll;
 	TSharedPtr<FUICommandInfo> AddRecording;
+	TSharedPtr<FUICommandInfo> AddCurrentPlayerRecording;
 	TSharedPtr<FUICommandInfo> RemoveRecording;
 	TSharedPtr<FUICommandInfo> RemoveAllRecordings;
 	TSharedPtr<FUICommandInfo> AddRecordingGroup;
