Sequencer: Use non-throttled spin box for numeric key editor

diff --git a/Engine/Source/Editor/MovieSceneTools/Private/CurveKeyEditors/SNumericKeyEditor.h b/Engine/Source/Editor/MovieSceneTools/Private/CurveKeyEditors/SNumericKeyEditor.h
--- a/Engine/Source/Editor/MovieSceneTools/Private/CurveKeyEditors/SNumericKeyEditor.h
+++ b/Engine/Source/Editor/MovieSceneTools/Private/CurveKeyEditors/SNumericKeyEditor.h
@@ -14,6 +14,20 @@
 
 #define LOCTEXT_NAMESPACE "NumericKeyEditor"
 
+template<typename T>
+struct SNonThrottledSpinBox : SSpinBox<T>
+{
+	virtual FReply OnMouseButtonDown(const FGeometry& MyGeometry, const FPointerEvent& MouseEvent) override
+	{
+		FReply Reply = SSpinBox<T>::OnMouseButtonDown(MyGeometry, MouseEvent);
+		if (Reply.IsEventHandled())
+		{
+			Reply.PreventThrottling();
+		}
+		return Reply;
+	}
+};
+
 /**
  * A widget for editing a curve representing integer keys.
  */
@@ -29,7 +43,7 @@ public:
 		KeyEditor = InKeyEditor;
 		ChildSlot
 		[
-			SNew(SSpinBox<NumericType>)
+			SNew(SNonThrottledSpinBox<NumericType>)
 			.Style(&FEditorStyle::GetWidgetStyle<FSpinBoxStyle>("Sequencer.HyperlinkSpinBox"))
 			.Font(FEditorStyle::GetFontStyle("Sequencer.AnimationOutliner.RegularFont"))
 			.MinValue(TOptional<NumericType>())
