Sequencer: Fix missing buttons in the particle toggle track.

Fixed CreateKeyEditor and make the the channel handles reflect the channel's inheritance hierarchy.

diff --git a/Engine/Source/Editor/MovieSceneTools/Private/Channels/BuiltInChannelEditors.cpp b/Engine/Source/Editor/MovieSceneTools/Private/Channels/BuiltInChannelEditors.cpp
--- a/Engine/Source/Editor/MovieSceneTools/Private/Channels/BuiltInChannelEditors.cpp
+++ b/Engine/Source/Editor/MovieSceneTools/Private/Channels/BuiltInChannelEditors.cpp
@@ -96,7 +96,7 @@ bool CanCreateKeyEditor(const FMovieSceneStringChannel*  Channel)
 	return true;
 }
 
-TSharedRef<SWidget> CreateKeyEditor(TMovieSceneChannelHandle<FMovieSceneBoolChannel>    Channel, UMovieSceneSection* Section, const FGuid& InObjectBindingID, TWeakPtr<FTrackInstancePropertyBindings> PropertyBindings, TWeakPtr<ISequencer> InSequencer)
+TSharedRef<SWidget> CreateKeyEditor(const TMovieSceneChannelHandle<FMovieSceneBoolChannel>&    Channel, UMovieSceneSection* Section, const FGuid& InObjectBindingID, TWeakPtr<FTrackInstancePropertyBindings> PropertyBindings, TWeakPtr<ISequencer> InSequencer)
 {
 	const TMovieSceneExternalValue<bool>* ExternalValue = Channel.GetExtendedEditorData();
 	if (!ExternalValue)
@@ -113,7 +113,7 @@ TSharedRef<SWidget> CreateKeyEditor(TMovieSceneChannelHandle<FMovieSceneBoolChan
 }
 
 
-TSharedRef<SWidget> CreateKeyEditor(TMovieSceneChannelHandle<FMovieSceneIntegerChannel> Channel, UMovieSceneSection* Section, const FGuid& InObjectBindingID, TWeakPtr<FTrackInstancePropertyBindings> PropertyBindings, TWeakPtr<ISequencer> InSequencer)
+TSharedRef<SWidget> CreateKeyEditor(const TMovieSceneChannelHandle<FMovieSceneIntegerChannel>& Channel, UMovieSceneSection* Section, const FGuid& InObjectBindingID, TWeakPtr<FTrackInstancePropertyBindings> PropertyBindings, TWeakPtr<ISequencer> InSequencer)
 {
 	const TMovieSceneExternalValue<int32>* ExternalValue = Channel.GetExtendedEditorData();
 	if (!ExternalValue)
@@ -131,7 +131,7 @@ TSharedRef<SWidget> CreateKeyEditor(TMovieSceneChannelHandle<FMovieSceneIntegerC
 }
 
 
-TSharedRef<SWidget> CreateKeyEditor(TMovieSceneChannelHandle<FMovieSceneFloatChannel>   Channel, UMovieSceneSection* Section, const FGuid& InObjectBindingID, TWeakPtr<FTrackInstancePropertyBindings> PropertyBindings, TWeakPtr<ISequencer> InSequencer)
+TSharedRef<SWidget> CreateKeyEditor(const TMovieSceneChannelHandle<FMovieSceneFloatChannel>&   Channel, UMovieSceneSection* Section, const FGuid& InObjectBindingID, TWeakPtr<FTrackInstancePropertyBindings> PropertyBindings, TWeakPtr<ISequencer> InSequencer)
 {
 	const TMovieSceneExternalValue<float>* ExternalValue = Channel.GetExtendedEditorData();
 	if (!ExternalValue)
@@ -149,7 +149,7 @@ TSharedRef<SWidget> CreateKeyEditor(TMovieSceneChannelHandle<FMovieSceneFloatCha
 }
 
 
-TSharedRef<SWidget> CreateKeyEditor(TMovieSceneChannelHandle<FMovieSceneStringChannel>  Channel, UMovieSceneSection* Section, const FGuid& InObjectBindingID, TWeakPtr<FTrackInstancePropertyBindings> PropertyBindings, TWeakPtr<ISequencer> InSequencer)
+TSharedRef<SWidget> CreateKeyEditor(const TMovieSceneChannelHandle<FMovieSceneStringChannel>&  Channel, UMovieSceneSection* Section, const FGuid& InObjectBindingID, TWeakPtr<FTrackInstancePropertyBindings> PropertyBindings, TWeakPtr<ISequencer> InSequencer)
 {
 	const TMovieSceneExternalValue<FString>* ExternalValue = Channel.GetExtendedEditorData();
 	if (!ExternalValue)
@@ -166,7 +166,7 @@ TSharedRef<SWidget> CreateKeyEditor(TMovieSceneChannelHandle<FMovieSceneStringCh
 }
 
 
-TSharedRef<SWidget> CreateKeyEditor(TMovieSceneChannelHandle<FMovieSceneByteChannel>    Channel, UMovieSceneSection* Section, const FGuid& InObjectBindingID, TWeakPtr<FTrackInstancePropertyBindings> PropertyBindings, TWeakPtr<ISequencer> InSequencer)
+TSharedRef<SWidget> CreateKeyEditor(const TMovieSceneChannelHandle<FMovieSceneByteChannel>&    Channel, UMovieSceneSection* Section, const FGuid& InObjectBindingID, TWeakPtr<FTrackInstancePropertyBindings> PropertyBindings, TWeakPtr<ISequencer> InSequencer)
 {
 	const TMovieSceneExternalValue<uint8>* ExternalValue = Channel.GetExtendedEditorData();
 	const FMovieSceneByteChannel* RawChannel = Channel.Get();
@@ -191,23 +191,23 @@ TSharedRef<SWidget> CreateKeyEditor(TMovieSceneChannelHandle<FMovieSceneByteChan
 	}
 }
 
-TSharedPtr<FStructOnScope> GetKeyStruct(TMovieSceneChannelHandle<FMovieSceneBoolChannel>     ChannelHandle, FKeyHandle InHandle)
+TSharedPtr<FStructOnScope> GetKeyStruct(const TMovieSceneChannelHandle<FMovieSceneBoolChannel>&     ChannelHandle, FKeyHandle InHandle)
 {
 	return CreateKeyStruct<FMovieSceneBoolKeyStruct>(ChannelHandle, InHandle);
 }
-TSharedPtr<FStructOnScope> GetKeyStruct(TMovieSceneChannelHandle<FMovieSceneByteChannel>     ChannelHandle, FKeyHandle InHandle)
+TSharedPtr<FStructOnScope> GetKeyStruct(const TMovieSceneChannelHandle<FMovieSceneByteChannel>&     ChannelHandle, FKeyHandle InHandle)
 {
 	return CreateKeyStruct<FMovieSceneByteKeyStruct>(ChannelHandle, InHandle);
 }
-TSharedPtr<FStructOnScope> GetKeyStruct(TMovieSceneChannelHandle<FMovieSceneIntegerChannel>  ChannelHandle, FKeyHandle InHandle)
+TSharedPtr<FStructOnScope> GetKeyStruct(const TMovieSceneChannelHandle<FMovieSceneIntegerChannel>&  ChannelHandle, FKeyHandle InHandle)
 {
 	return CreateKeyStruct<FMovieSceneIntegerKeyStruct>(ChannelHandle, InHandle);
 }
-TSharedPtr<FStructOnScope> GetKeyStruct(TMovieSceneChannelHandle<FMovieSceneStringChannel>   ChannelHandle, FKeyHandle InHandle)
+TSharedPtr<FStructOnScope> GetKeyStruct(const TMovieSceneChannelHandle<FMovieSceneStringChannel>&   ChannelHandle, FKeyHandle InHandle)
 {
 	return CreateKeyStruct<FMovieSceneStringKeyStruct>(ChannelHandle, InHandle);
 }
-TSharedPtr<FStructOnScope> GetKeyStruct(TMovieSceneChannelHandle<FMovieSceneParticleChannel> ChannelHandle, FKeyHandle InHandle)
+TSharedPtr<FStructOnScope> GetKeyStruct(const TMovieSceneChannelHandle<FMovieSceneParticleChannel>& ChannelHandle, FKeyHandle InHandle)
 {
 	FMovieSceneParticleChannel* Channel = ChannelHandle.Get();
 	if (!Channel)
@@ -688,7 +688,7 @@ void ExtendKeyMenu(FMenuBuilder& OuterMenuBuilder, TArray<TExtendKeyMenuParams<F
 	OuterMenuBuilder.PushExtender(Extension);
 }
 
-TUniquePtr<FCurveModel> CreateCurveEditorModel(TMovieSceneChannelHandle<FMovieSceneFloatChannel> FloatChannel, UMovieSceneSection* OwningSection, TSharedRef<ISequencer> InSequencer)
+TUniquePtr<FCurveModel> CreateCurveEditorModel(const TMovieSceneChannelHandle<FMovieSceneFloatChannel>& FloatChannel, UMovieSceneSection* OwningSection, TSharedRef<ISequencer> InSequencer)
 {
 	return MakeUnique<FFloatChannelCurveModel>(FloatChannel, OwningSection, InSequencer);
 }
diff --git a/Engine/Source/Editor/MovieSceneTools/Private/Channels/BuiltInChannelEditors.h b/Engine/Source/Editor/MovieSceneTools/Private/Channels/BuiltInChannelEditors.h
--- a/Engine/Source/Editor/MovieSceneTools/Private/Channels/BuiltInChannelEditors.h
+++ b/Engine/Source/Editor/MovieSceneTools/Private/Channels/BuiltInChannelEditors.h
@@ -44,18 +44,18 @@ bool CanCreateKeyEditor(const FMovieSceneIntegerChannel* Channel);
 bool CanCreateKeyEditor(const FMovieSceneFloatChannel*   Channel);
 bool CanCreateKeyEditor(const FMovieSceneStringChannel*  Channel);
 
-TSharedRef<SWidget> CreateKeyEditor(TMovieSceneChannelHandle<FMovieSceneBoolChannel>    Channel, const TMovieSceneExternalValue<bool>&    EditorData, UMovieSceneSection* Section, const FGuid& InObjectBindingID, TWeakPtr<FTrackInstancePropertyBindings> PropertyBindings, TWeakPtr<ISequencer> InSequencer);
-TSharedRef<SWidget> CreateKeyEditor(TMovieSceneChannelHandle<FMovieSceneByteChannel>    Channel, const TMovieSceneExternalValue<uint8>&   EditorData, UMovieSceneSection* Section, const FGuid& InObjectBindingID, TWeakPtr<FTrackInstancePropertyBindings> PropertyBindings, TWeakPtr<ISequencer> InSequencer);
-TSharedRef<SWidget> CreateKeyEditor(TMovieSceneChannelHandle<FMovieSceneIntegerChannel> Channel, const TMovieSceneExternalValue<int32>&   EditorData, UMovieSceneSection* Section, const FGuid& InObjectBindingID, TWeakPtr<FTrackInstancePropertyBindings> PropertyBindings, TWeakPtr<ISequencer> InSequencer);
-TSharedRef<SWidget> CreateKeyEditor(TMovieSceneChannelHandle<FMovieSceneFloatChannel>   Channel, const TMovieSceneExternalValue<float>&   EditorData, UMovieSceneSection* Section, const FGuid& InObjectBindingID, TWeakPtr<FTrackInstancePropertyBindings> PropertyBindings, TWeakPtr<ISequencer> InSequencer);
-TSharedRef<SWidget> CreateKeyEditor(TMovieSceneChannelHandle<FMovieSceneStringChannel>  Channel, const TMovieSceneExternalValue<FString>& EditorData, UMovieSceneSection* Section, const FGuid& InObjectBindingID, TWeakPtr<FTrackInstancePropertyBindings> PropertyBindings, TWeakPtr<ISequencer> InSequencer);
+TSharedRef<SWidget> CreateKeyEditor(const TMovieSceneChannelHandle<FMovieSceneBoolChannel>&    Channel, UMovieSceneSection* Section, const FGuid& InObjectBindingID, TWeakPtr<FTrackInstancePropertyBindings> PropertyBindings, TWeakPtr<ISequencer> InSequencer);
+TSharedRef<SWidget> CreateKeyEditor(const TMovieSceneChannelHandle<FMovieSceneByteChannel>&    Channel, UMovieSceneSection* Section, const FGuid& InObjectBindingID, TWeakPtr<FTrackInstancePropertyBindings> PropertyBindings, TWeakPtr<ISequencer> InSequencer);
+TSharedRef<SWidget> CreateKeyEditor(const TMovieSceneChannelHandle<FMovieSceneIntegerChannel>& Channel, UMovieSceneSection* Section, const FGuid& InObjectBindingID, TWeakPtr<FTrackInstancePropertyBindings> PropertyBindings, TWeakPtr<ISequencer> InSequencer);
+TSharedRef<SWidget> CreateKeyEditor(const TMovieSceneChannelHandle<FMovieSceneFloatChannel>&   Channel, UMovieSceneSection* Section, const FGuid& InObjectBindingID, TWeakPtr<FTrackInstancePropertyBindings> PropertyBindings, TWeakPtr<ISequencer> InSequencer);
+TSharedRef<SWidget> CreateKeyEditor(const TMovieSceneChannelHandle<FMovieSceneStringChannel>&  Channel, UMovieSceneSection* Section, const FGuid& InObjectBindingID, TWeakPtr<FTrackInstancePropertyBindings> PropertyBindings, TWeakPtr<ISequencer> InSequencer);
 
 /** KeyStruct overrides */
-TSharedPtr<FStructOnScope> GetKeyStruct(TMovieSceneChannelHandle<FMovieSceneBoolChannel>     Channel, FKeyHandle InHandle);
-TSharedPtr<FStructOnScope> GetKeyStruct(TMovieSceneChannelHandle<FMovieSceneByteChannel>     Channel, FKeyHandle InHandle);
-TSharedPtr<FStructOnScope> GetKeyStruct(TMovieSceneChannelHandle<FMovieSceneIntegerChannel>  Channel, FKeyHandle InHandle);
-TSharedPtr<FStructOnScope> GetKeyStruct(TMovieSceneChannelHandle<FMovieSceneStringChannel>   Channel, FKeyHandle InHandle);
-TSharedPtr<FStructOnScope> GetKeyStruct(TMovieSceneChannelHandle<FMovieSceneParticleChannel> Channel, FKeyHandle InHandle);
+TSharedPtr<FStructOnScope> GetKeyStruct(const TMovieSceneChannelHandle<FMovieSceneBoolChannel>&     Channel, FKeyHandle InHandle);
+TSharedPtr<FStructOnScope> GetKeyStruct(const TMovieSceneChannelHandle<FMovieSceneByteChannel>&     Channel, FKeyHandle InHandle);
+TSharedPtr<FStructOnScope> GetKeyStruct(const TMovieSceneChannelHandle<FMovieSceneIntegerChannel>&  Channel, FKeyHandle InHandle);
+TSharedPtr<FStructOnScope> GetKeyStruct(const TMovieSceneChannelHandle<FMovieSceneStringChannel>&   Channel, FKeyHandle InHandle);
+TSharedPtr<FStructOnScope> GetKeyStruct(const TMovieSceneChannelHandle<FMovieSceneParticleChannel>& Channel, FKeyHandle InHandle);
 
 /** Key drawing overrides */
 void DrawKeys(FMovieSceneFloatChannel*    Channel, TArrayView<const FKeyHandle> InKeyHandles, TArrayView<FKeyDrawParams> OutKeyDrawParams);
@@ -66,7 +66,7 @@ void ExtendSectionMenu(FMenuBuilder& OuterMenuBuilder, TArray<TMovieSceneChannel
 void ExtendKeyMenu(FMenuBuilder& OuterMenuBuilder, TArray<TExtendKeyMenuParams<FMovieSceneFloatChannel>>&& Channels, TWeakPtr<ISequencer> InSequencer);
 
 /** Curve editor models */
-TUniquePtr<FCurveModel> CreateCurveEditorModel(TMovieSceneChannelHandle<FMovieSceneFloatChannel> FloatChannel, UMovieSceneSection* OwningSection, TSharedRef<ISequencer> InSequencer);
+TUniquePtr<FCurveModel> CreateCurveEditorModel(const TMovieSceneChannelHandle<FMovieSceneFloatChannel>& FloatChannel, UMovieSceneSection* OwningSection, TSharedRef<ISequencer> InSequencer);
 
 USTRUCT()
 struct FMovieSceneIntegerKeyStruct : public FMovieSceneKeyTimeStruct
diff --git a/Engine/Source/Editor/Sequencer/Private/SequencerChannelTraits.cpp b/Engine/Source/Editor/Sequencer/Private/SequencerChannelTraits.cpp
--- a/Engine/Source/Editor/Sequencer/Private/SequencerChannelTraits.cpp
+++ b/Engine/Source/Editor/Sequencer/Private/SequencerChannelTraits.cpp
@@ -2,6 +2,7 @@
 
 #include "SequencerChannelTraits.h"
 #include "EditorStyleSet.h"
+#include "CurveModel.h"
 
 namespace Sequencer
 {
@@ -19,6 +20,9 @@ void DrawKeys(FMovieSceneChannel* Channel, TArrayView<const FKeyHandle> InHandle
 	}
 }
 
-
+TUniquePtr<FCurveModel> CreateCurveEditorModel(const FMovieSceneChannelHandle& ChannelHandle, UMovieSceneSection* OwningSection, TSharedRef<ISequencer> InSequencer)
+{
+	return nullptr;
+}
 
 }	// namespace Sequencer
\ No newline at end of file
diff --git a/Engine/Source/Editor/Sequencer/Public/SequencerChannelTraits.h b/Engine/Source/Editor/Sequencer/Public/SequencerChannelTraits.h
--- a/Engine/Source/Editor/Sequencer/Public/SequencerChannelTraits.h
+++ b/Engine/Source/Editor/Sequencer/Public/SequencerChannelTraits.h
@@ -74,7 +74,7 @@ namespace Sequencer
 	 * @return A shared struct object, or nullptr
 	 */
 	template<typename ChannelType>
-	TSharedPtr<FStructOnScope> GetKeyStruct(TMovieSceneChannelHandle<ChannelType> ChannelHandle, FKeyHandle KeyHandle)
+	TSharedPtr<FStructOnScope> GetKeyStruct(const TMovieSceneChannelHandle<ChannelType>& ChannelHandle, FKeyHandle KeyHandle)
 	{
 		ChannelType* Channel = ChannelHandle.Get();
 		if (Channel)
@@ -121,9 +121,8 @@ namespace Sequencer
 	 * @param InSequencer        The sequencer currently active
 	 * @return The key editor widget
 	 */
-	template<typename ChannelType>
-	TSharedRef<SWidget> CreateKeyEditor(
-		TMovieSceneChannelHandle<ChannelType>    InChannel,
+	inline TSharedRef<SWidget> CreateKeyEditor(
+		const FMovieSceneChannelHandle&          InChannel,
 		UMovieSceneSection*                      InOwningSection,
 		const FGuid&                             InObjectBindingID,
 		TWeakPtr<FTrackInstancePropertyBindings> InPropertyBindings,
@@ -360,11 +359,7 @@ namespace Sequencer
 	 *
 	 * @return (Optional) A new model to be added to a curve editor
 	 */
-	template<typename ChannelType>
-	TUniquePtr<FCurveModel> CreateCurveEditorModel(TMovieSceneChannelHandle<ChannelType> ChannelHandle, UMovieSceneSection* OwningSection, TSharedRef<ISequencer> InSequencer)
-	{
-		return nullptr;
-	}
+	SEQUENCER_API TUniquePtr<FCurveModel> CreateCurveEditorModel(const FMovieSceneChannelHandle& ChannelHandle, UMovieSceneSection* OwningSection, TSharedRef<ISequencer> InSequencer);
 
 }	// namespace Sequencer
 
diff --git a/Engine/Source/Runtime/MovieScene/Public/Channels/MovieSceneChannelHandle.h b/Engine/Source/Runtime/MovieScene/Public/Channels/MovieSceneChannelHandle.h
--- a/Engine/Source/Runtime/MovieScene/Public/Channels/MovieSceneChannelHandle.h
+++ b/Engine/Source/Runtime/MovieScene/Public/Channels/MovieSceneChannelHandle.h
@@ -12,8 +12,6 @@ template<typename> struct TMovieSceneChannelTraits;
 
 /**
  * Handle to a specific channel in a UMovieSceneSection. Will become nullptr when the FMovieSceneChannelProxy it was created with is reallocated.
- * Implemented internally as a weak pointer, but not advertised as such to ensure that no shared references are held to FMovieSceneChannelProxy
- * outside of the owning MovieSceneSection's reference.
  */
 struct MOVIESCENE_API FMovieSceneChannelHandle
 {
@@ -75,13 +73,31 @@ private:
 	int32 ChannelIndex;
 };
 
+
+/**
+ * Forward declaration of a typed channel handle.
+ * Typed channel handles replicate the inheritance hierarchy of their channels in order that overload resolution on them can happen without user-conversion or templating
+ */
+template<typename ChannelType> struct TMovieSceneChannelHandle;
+
+/** Base implementation that is just the untyped handle type */
+template<> struct TMovieSceneChannelHandle<FMovieSceneChannel> : FMovieSceneChannelHandle
+{
+protected:
+	TMovieSceneChannelHandle()
+	{}
+
+	TMovieSceneChannelHandle(TWeakPtr<FMovieSceneChannelProxy> InWeakChannelProxy, FName InName, int32 InChannelIndex)
+		: FMovieSceneChannelHandle(InWeakChannelProxy, InName, InChannelIndex)
+	{}
+};
+
 /**
  * Handle to a specific channel in a UMovieSceneSection. Will become nullptr when the FMovieSceneChannelProxy it was created with is reallocated.
- * Implemented internally as a weak pointer, but not advertised as such to ensure that no shared references are held to FMovieSceneChannelProxy
- * outside of the owning MovieSceneSection's reference.
+ * The handle class hierarchy matches that of the channel itself so that functions can be overloaded for any given base/derived handle type
  */
 template<typename ChannelType>
-struct TMovieSceneChannelHandle : FMovieSceneChannelHandle
+struct TMovieSceneChannelHandle : TMovieSceneChannelHandle<typename ChannelType::Super>
 {
 	TMovieSceneChannelHandle()
 	{}
@@ -90,7 +106,7 @@ struct TMovieSceneChannelHandle : FMovieSceneChannelHandle
 	 * Construction from a weak channel proxy, and the channel's index
 	 */
 	TMovieSceneChannelHandle(TWeakPtr<FMovieSceneChannelProxy> InWeakChannelProxy, int32 InChannelIndex)
-		: FMovieSceneChannelHandle(InWeakChannelProxy, ChannelType::StaticStruct()->GetFName(), InChannelIndex)
+		: TMovieSceneChannelHandle<typename ChannelType::Super>(InWeakChannelProxy, ChannelType::StaticStruct()->GetFName(), InChannelIndex)
 	{}
 
 	/**
@@ -116,6 +132,12 @@ struct TMovieSceneChannelHandle : FMovieSceneChannelHandle
 	}
 
 #endif // WITH_EDITOR
+
+protected:
+
+	TMovieSceneChannelHandle(TWeakPtr<FMovieSceneChannelProxy> InWeakChannelProxy, FName InChannelTypeName, int32 InChannelIndex)
+		: TMovieSceneChannelHandle<typename ChannelType::Super>(InWeakChannelProxy, InChannelTypeName, InChannelIndex)
+	{}
 };
 
 
