Sequencer: Modify the section when toggling the bounded range. This fixes an issue where the evaluation would be incorrect because the template wouldn't regenerate when toggling the section bounds.

diff --git a/Engine/Source/Editor/MovieSceneTools/Private/MovieSceneSectionDetailsCustomization.cpp b/Engine/Source/Editor/MovieSceneTools/Private/MovieSceneSectionDetailsCustomization.cpp
--- a/Engine/Source/Editor/MovieSceneTools/Private/MovieSceneSectionDetailsCustomization.cpp
+++ b/Engine/Source/Editor/MovieSceneTools/Private/MovieSceneSectionDetailsCustomization.cpp
@@ -280,6 +280,18 @@ FText FMovieSceneSectionDetailsCustomization::GetRangeStartButtonIcon() const
 /** Called when the button is pressed to toggle the current state. */
 FReply FMovieSceneSectionDetailsCustomization::ToggleRangeStartBounded()
 {
+	FScopedTransaction Transaction(LOCTEXT("ToggleRangeStartBounded", "Toggle Range Start Bounded"));
+	TArray<UObject*> Objects;
+	MovieSceneSectionPropertyHandle->GetOuterObjects(Objects);
+
+	for (auto Outer : Objects)
+	{
+		for (UObject* Obj : Objects)
+		{
+			Obj->Modify();
+		}
+	}
+
 	if (IsRangeStartTextboxEnabled())
 	{
 		SetRangeStartBounded(false);
@@ -468,6 +480,18 @@ FText FMovieSceneSectionDetailsCustomization::GetRangeEndButtonIcon() const
 /** Called when the button is pressed to toggle the current state. */
 FReply FMovieSceneSectionDetailsCustomization::ToggleRangeEndBounded()
 {
+	FScopedTransaction Transaction(LOCTEXT("ToggleRangeEndBounded", "Toggle Range End Bounded"));
+	TArray<UObject*> Objects;
+	MovieSceneSectionPropertyHandle->GetOuterObjects(Objects);
+
+	for (auto Outer : Objects)
+	{
+		for (UObject* Obj : Objects)
+		{
+			Obj->Modify();
+		}
+	}
+
 	if (IsRangeEndTextboxEnabled())
 	{
 		SetRangeEndBounded(false);
