Sequencer: Added 'Write FCPXML' option to Render Movie Settings dialog.

Changed track menu FCPXML export to get default resolution and frame rate from MovieSceneCaptureSettings.

diff --git a/Engine/Source/Editor/MovieSceneTools/Private/AutomatedLevelSequenceCapture.cpp b/Engine/Source/Editor/MovieSceneTools/Private/AutomatedLevelSequenceCapture.cpp
--- a/Engine/Source/Editor/MovieSceneTools/Private/AutomatedLevelSequenceCapture.cpp
+++ b/Engine/Source/Editor/MovieSceneTools/Private/AutomatedLevelSequenceCapture.cpp
@@ -10,6 +10,7 @@
 #include "JsonObjectConverter.h"
 #include "Tracks/MovieSceneCinematicShotTrack.h"
 #include "MovieSceneTranslatorEDL.h"
+#include "FCPXML/FCPXMLMovieSceneTranslator.h"
 #include "EngineUtils.h"
 #include "Sections/MovieSceneCinematicShotSection.h"
 #include "TimerManager.h"
@@ -17,6 +18,7 @@
 #include "Camera/CameraComponent.h"
 #include "MovieSceneCaptureModule.h"
 #include "MovieSceneTimeHelpers.h"
+#include "MovieSceneToolHelpers.h"
 
 const FName UAutomatedLevelSequenceCapture::AutomatedLevelSequenceCaptureUIName = FName(TEXT("AutomatedLevelSequenceCaptureUIInstance"));
 
@@ -86,6 +88,7 @@ UAutomatedLevelSequenceCapture::UAutomatedLevelSequenceCapture(const FObjectInit
 	DelayBeforeWarmUp = 0.0f;
 	DelayBeforeShotWarmUp = 0.0f;
 	bWriteEditDecisionList = true;
+	bWriteFinalCutProXML = true;
 
 	RemainingWarmUpFrames = 0;
 
@@ -241,6 +244,7 @@ void UAutomatedLevelSequenceCapture::Initialize(TSharedPtr<FSceneViewport> InVie
 	}
 
 	ExportEDL();
+	ExportFCPXML();
 
 	CaptureState = ELevelSequenceCaptureState::Setup;
 	CaptureStrategy = MakeShareable(new FFixedTimeStepCaptureStrategy(Settings.FrameRate));
@@ -800,5 +804,39 @@ void UAutomatedLevelSequenceCapture::ExportEDL()
 	MovieSceneTranslatorEDL::ExportEDL(MovieScene, Settings.FrameRate, SaveFilename, HandleFrames);
 }
 
+void UAutomatedLevelSequenceCapture::ExportFCPXML()
+{
+	if (!bWriteFinalCutProXML)
+	{
+		return;
+	}
+
+	UMovieScene* MovieScene = GetMovieScene(LevelSequenceActor);
+	if (!MovieScene)
+	{
+		return;
+	}
+
+	UMovieSceneCinematicShotTrack* ShotTrack = MovieScene->FindMasterTrack<UMovieSceneCinematicShotTrack>();
+	if (!ShotTrack)
+	{
+		return;
+	}
+
+	FString SaveFilename = Settings.OutputDirectory.Path / MovieScene->GetOuter()->GetName() + TEXT(".xml");
+	int32 HandleFrames = Settings.HandleFrames;
+	FFrameRate FrameRate = Settings.FrameRate;
+	uint32 ResX = Settings.Resolution.ResX;
+	uint32 ResY = Settings.Resolution.ResY;
+
+	FFCPXMLExporter *Exporter = new FFCPXMLExporter;
+
+	TSharedRef<FMovieSceneTranslatorContext> ExportContext(new FMovieSceneTranslatorContext);
+	ExportContext->Init();
+	bool bSuccess = Exporter->Export(MovieScene, FrameRate, ResX, ResY, HandleFrames, SaveFilename, ExportContext);
+
+	delete Exporter;
+}
+
 
 #endif
diff --git a/Engine/Source/Editor/MovieSceneTools/Private/FCPXML/FCPXMLExport.cpp b/Engine/Source/Editor/MovieSceneTools/Private/FCPXML/FCPXMLExport.cpp
--- a/Engine/Source/Editor/MovieSceneTools/Private/FCPXML/FCPXMLExport.cpp
+++ b/Engine/Source/Editor/MovieSceneTools/Private/FCPXML/FCPXMLExport.cpp
@@ -153,19 +153,6 @@ bool FFCPXMLExportVisitor::ConstructMasterClipNodes(TSharedRef<FFCPXMLNode> InPa
 			continue;
 		}
 
-		// skip sections if media file does not exist
-		FString FilePathName = GetFilePathName(CinematicSection->ShotFilename);
-		if (!FPaths::FileExists(FilePathName))
-		{
-			// add warning message and skip this section
-			ExportContext->AddMessage(
-				EMessageSeverity::Warning,
-				FText::Format(LOCTEXT("SkippingSection", "Warning: Skipping section {0}, media file does not exist: {1}."), FText::FromString(CinematicSection->ShotDisplayName), FText::FromString(FilePathName)));
-			CinematicSection->bEnabled = false;
-			continue;
-		}
-
-
 		if (!ConstructMasterClipNode(InParentNode, CinematicSection, ++MasterClipId))
 		{
 			return false;
@@ -370,10 +357,8 @@ bool FFCPXMLExportVisitor::ConstructVideoNode(TSharedRef<FFCPXMLNode> InParentNo
 	TSharedRef<FFCPXMLNode> VideoNode = InParentNode->CreateChildNode(TEXT("video"));
 
 	TSharedRef<FFCPXMLNode> FormatNode = VideoNode->CreateChildNode(TEXT("format"));
-	int32 Width{ 0 };
-	int32 Height{ 0 };
-	GetDefaultImageResolution(Width, Height);
-	if (!ConstructSampleCharacteristicsNode(FormatNode, Width, Height))
+
+	if (!ConstructSampleCharacteristicsNode(FormatNode, ExportData->GetResX(), ExportData->GetResY()))
 	{
 		return false;
 	}
@@ -553,7 +538,7 @@ bool FFCPXMLExportVisitor::ConstructAudioClipItemNode(TSharedRef<FFCPXMLNode> In
 
 bool FFCPXMLExportVisitor::ConstructVideoFileNode(TSharedRef<FFCPXMLNode> InParentNode, const TSharedPtr<FMovieSceneExportCinematicSectionData> InCinematicSectionData, uint32 Duration, bool bInMasterClip)
 {
-	if (!InCinematicSectionData.IsValid())
+	if (!ExportData->IsExportDataValid() || !InCinematicSectionData.IsValid())
 	{
 		return false;
 	}
@@ -596,10 +581,7 @@ bool FFCPXMLExportVisitor::ConstructVideoFileNode(TSharedRef<FFCPXMLNode> InPare
 		TSharedRef<FFCPXMLNode> MediaNode = FileNode->CreateChildNode(TEXT("media"));
 		TSharedRef<FFCPXMLNode> VideoNode = MediaNode->CreateChildNode(TEXT("video"));
 
-		int32 Width{ 0 };
-		int32 Height{ 0 };
-		GetDefaultImageResolution(Width, Height);
-		if (!ConstructSampleCharacteristicsNode(VideoNode, Width, Height))
+		if (!ConstructSampleCharacteristicsNode(VideoNode, ExportData->GetResX(), ExportData->GetResY()))
 		{
 			return false;
 		}
@@ -739,12 +721,6 @@ bool FFCPXMLExportVisitor::GetMasterClipFileIdName(const TSharedPtr<FMovieSceneE
 	return GetIdName(InCinematicSectionData, FString(TEXT("file")), OutName);
 }
 
-void FFCPXMLExportVisitor::GetDefaultImageResolution(int32& OutWidth, int32& OutHeight)
-{
-	OutWidth = 1280;
-	OutHeight = 720;
-}
-
 FString FFCPXMLExportVisitor::GetFilePathName(const FString& InSectionName) const
 {
 	return (SaveFilePath + TEXT("/") + InSectionName);
diff --git a/Engine/Source/Editor/MovieSceneTools/Private/FCPXML/FCPXMLExport.h b/Engine/Source/Editor/MovieSceneTools/Private/FCPXML/FCPXMLExport.h
--- a/Engine/Source/Editor/MovieSceneTools/Private/FCPXML/FCPXMLExport.h
+++ b/Engine/Source/Editor/MovieSceneTools/Private/FCPXML/FCPXMLExport.h
@@ -105,8 +105,6 @@ private:
 	bool GetMasterClipItemIdName(const TSharedPtr<FMovieSceneExportCinematicSectionData> InCinematicSectionData, FString& OutName);
 	/** Get master clip's file id name based on input name and id */
 	bool GetMasterClipFileIdName(const TSharedPtr<FMovieSceneExportCinematicSectionData> InCinematicSectionData, FString& OutName);
-	/** Get the default image resolution */
-	void GetDefaultImageResolution(int32& OutWidth, int32& OutHeight);
 	/** Return filename with path */
 	FString GetFilePathName(const FString& InSectionName) const;
 
diff --git a/Engine/Source/Editor/MovieSceneTools/Private/FCPXML/FCPXMLMovieSceneTranslator.cpp b/Engine/Source/Editor/MovieSceneTools/Private/FCPXML/FCPXMLMovieSceneTranslator.cpp
--- a/Engine/Source/Editor/MovieSceneTools/Private/FCPXML/FCPXMLMovieSceneTranslator.cpp
+++ b/Engine/Source/Editor/MovieSceneTools/Private/FCPXML/FCPXMLMovieSceneTranslator.cpp
@@ -146,7 +146,7 @@ FText FFCPXMLExporter::GetMessageLogLabel() const
 	return LOCTEXT("FCPXMLExportLogLabel", "FCP 7 XML Export Log");
 }
 
-bool FFCPXMLExporter::Export(const UMovieScene* InMovieScene, FFrameRate InFrameRate, FString InSaveFilename, int32 InHandleFrames, TSharedRef<FMovieSceneTranslatorContext> InContext)
+bool FFCPXMLExporter::Export(const UMovieScene* InMovieScene, FFrameRate InFrameRate, uint32 InResX, uint32 InResY, int32 InHandleFrames, FString InSaveFilename, TSharedRef<FMovieSceneTranslatorContext> InContext)
 {
 	InContext->Init();
 
@@ -154,7 +154,7 @@ bool FFCPXMLExporter::Export(const UMovieScene* InMovieScene, FFrameRate InFrame
 	TSharedRef<FFCPXMLFile> FCPXMLFile = MakeShared<FFCPXMLFile>();
 	FCPXMLFile->ConstructFile(FPaths::GetBaseFilename(InSaveFilename, true));
 
-	TSharedRef<FMovieSceneExportData> ExportData = MakeShared<FMovieSceneExportData>(InMovieScene, InFrameRate, InHandleFrames, InSaveFilename, InContext);
+	TSharedRef<FMovieSceneExportData> ExportData = MakeShared<FMovieSceneExportData>(InMovieScene, InFrameRate, InResX, InResY, InHandleFrames, InSaveFilename, InContext);
 
 	// Export sequencer movie scene, merging with existing Xml structure.
 	FFCPXMLExportVisitor ExportVisitor(InSaveFilename, ExportData, InContext);
diff --git a/Engine/Source/Editor/MovieSceneTools/Private/MovieSceneToolHelpers.cpp b/Engine/Source/Editor/MovieSceneTools/Private/MovieSceneToolHelpers.cpp
--- a/Engine/Source/Editor/MovieSceneTools/Private/MovieSceneToolHelpers.cpp
+++ b/Engine/Source/Editor/MovieSceneTools/Private/MovieSceneToolHelpers.cpp
@@ -728,13 +728,21 @@ bool MovieSceneToolHelpers::MovieSceneTranslatorImport(FMovieSceneImporter* InIm
 	return bSuccess;
 }
 
-bool MovieSceneToolHelpers::MovieSceneTranslatorExport(FMovieSceneExporter* InExporter, const UMovieScene* InMovieScene, FFrameRate InFrameRate, FString InSaveDirectory, int32 InHandleFrames)
+bool MovieSceneToolHelpers::MovieSceneTranslatorExport(FMovieSceneExporter* InExporter, const UMovieScene* InMovieScene, const FMovieSceneCaptureSettings& Settings)
 {
 	if (InExporter == nullptr || InMovieScene == nullptr)
 	{
 		return false;
 	}
 
+	FString SaveDirectory = FPaths::ConvertRelativePathToFull(Settings.OutputDirectory.Path);
+	int32 HandleFrames = Settings.HandleFrames;
+	// @todo: generate filename based on filename format, currently outputs {shot}.avi
+	FString FilenameFormat = Settings.OutputFormat;
+	FFrameRate FrameRate = Settings.FrameRate;
+	uint32 ResX = Settings.Resolution.ResX;
+	uint32 ResY = Settings.Resolution.ResY;
+
 	TArray<FString> SaveFilenames;
 	FString SequenceName = InMovieScene->GetOuter()->GetName();
 
@@ -750,7 +758,7 @@ bool MovieSceneToolHelpers::MovieSceneTranslatorExport(FMovieSceneExporter* InEx
 		bSave = DesktopPlatform->SaveFileDialog(
 			FSlateApplication::Get().FindBestParentWindowHandleForDialogs(nullptr),
 			DialogTitle,
-			InSaveDirectory,
+			SaveDirectory,
 			SequenceName + TEXT(".") + FileExtension,
 			FileTypeDescription,
 			EFileDialogFlags::None,
@@ -766,7 +774,7 @@ bool MovieSceneToolHelpers::MovieSceneTranslatorExport(FMovieSceneExporter* InEx
 	TSharedRef<FMovieSceneTranslatorContext> ExportContext(new FMovieSceneTranslatorContext);
 	ExportContext->Init();
 
-	bool bSuccess = InExporter->Export(InMovieScene, InFrameRate, SaveFilenames[0], InHandleFrames, ExportContext);
+	bool bSuccess = InExporter->Export(InMovieScene, FrameRate, ResX, ResY, HandleFrames, SaveFilenames[0], ExportContext);
 	
 	// Display any messages in context
 	MovieSceneTranslatorDisplayMessages(InExporter, ExportContext);
@@ -774,11 +782,11 @@ bool MovieSceneToolHelpers::MovieSceneTranslatorExport(FMovieSceneExporter* InEx
 	if (bSuccess)
 	{
 		const FString AbsoluteFilename = FPaths::ConvertRelativePathToFull(SaveFilenames[0]);
-		const FString SaveDirectory = FPaths::GetPath(AbsoluteFilename);
+		const FString ActualSaveDirectory = FPaths::GetPath(AbsoluteFilename);
 
 		FNotificationInfo NotificationInfo(InExporter->GetNotificationExportFinished());
 		NotificationInfo.ExpireDuration = 5.f;
-		NotificationInfo.Hyperlink = FSimpleDelegate::CreateStatic([](FString InDirectory) { FPlatformProcess::ExploreFolder(*InDirectory); }, SaveDirectory);
+		NotificationInfo.Hyperlink = FSimpleDelegate::CreateStatic([](FString InDirectory) { FPlatformProcess::ExploreFolder(*InDirectory); }, ActualSaveDirectory);
 		NotificationInfo.HyperlinkText = InExporter->GetNotificationHyperlinkText();
 		FSlateNotificationManager::Get().AddNotification(NotificationInfo);
 	}
diff --git a/Engine/Source/Editor/MovieSceneTools/Private/MovieSceneTranslator.cpp b/Engine/Source/Editor/MovieSceneTools/Private/MovieSceneTranslator.cpp
--- a/Engine/Source/Editor/MovieSceneTools/Private/MovieSceneTranslator.cpp
+++ b/Engine/Source/Editor/MovieSceneTools/Private/MovieSceneTranslator.cpp
@@ -16,7 +16,7 @@
 
 #define LOCTEXT_NAMESPACE "MovieSceneTranslator"
 
-FMovieSceneExportData::FMovieSceneExportData(const UMovieScene* InMovieScene, FFrameRate InFrameRate, int32 InHandleFrames, FString InSaveFilename, TSharedPtr<FMovieSceneTranslatorContext> InContext)
+FMovieSceneExportData::FMovieSceneExportData(const UMovieScene* InMovieScene, FFrameRate InFrameRate, uint32 InResX, uint32 InResY, int32 InHandleFrames, FString InSaveFilename, TSharedPtr<FMovieSceneTranslatorContext> InContext)
 {
 	if (InMovieScene == nullptr)
 	{
@@ -31,6 +31,8 @@ FMovieSceneExportData::FMovieSceneExportData(const UMovieScene* InMovieScene, FF
 	}
 
 	FrameRate = InFrameRate;
+	ResX = InResX;
+	ResY = InResY;
 	HandleFrames = ConvertFrameTime(FFrameTime(InHandleFrames), InMovieScene->GetTickResolution(), FrameRate);
 	SaveFilename = InSaveFilename;
 
@@ -76,13 +78,13 @@ bool FMovieSceneExportData::ConstructMovieSceneData(const UMovieScene* InMovieSc
 
 	MovieSceneData = MakeShared<FMovieSceneExportMovieSceneData>();
 
-	FFrameRate Resolution = InMovieScene->GetTickResolution();
+	FFrameRate TickResolution = InMovieScene->GetTickResolution();
 
 	TRange<FFrameNumber> PlaybackRange = InMovieScene->GetPlaybackRange();
 
 	if (PlaybackRange.HasLowerBound())
 	{
-		MovieSceneData->PlaybackRangeStartFrame = ConvertFrameTime(PlaybackRange.GetLowerBoundValue(), Resolution, FrameRate).CeilToFrame();
+		MovieSceneData->PlaybackRangeStartFrame = ConvertFrameTime(PlaybackRange.GetLowerBoundValue(), TickResolution, FrameRate).CeilToFrame();
 	}
 	else
 	{
@@ -92,7 +94,7 @@ bool FMovieSceneExportData::ConstructMovieSceneData(const UMovieScene* InMovieSc
 
 	if (PlaybackRange.HasUpperBound())
 	{
-		MovieSceneData->PlaybackRangeEndFrame = ConvertFrameTime(PlaybackRange.GetUpperBoundValue(), Resolution, FrameRate).CeilToFrame();
+		MovieSceneData->PlaybackRangeEndFrame = ConvertFrameTime(PlaybackRange.GetUpperBoundValue(), TickResolution, FrameRate).CeilToFrame();
 	}
 	else
 	{
@@ -102,8 +104,8 @@ bool FMovieSceneExportData::ConstructMovieSceneData(const UMovieScene* InMovieSc
 
 	MovieSceneData->Name = InMovieScene->GetOuter()->GetName();
 	MovieSceneData->Path = InMovieScene->GetOuter()->GetPathName();
-	MovieSceneData->Resolution = Resolution;
-	MovieSceneData->Duration = ConvertFrameTime(MovieScene::DiscreteSize(PlaybackRange), Resolution, FrameRate).FrameNumber.Value;
+	MovieSceneData->TickResolution = TickResolution;
+	MovieSceneData->Duration = ConvertFrameTime(MovieScene::DiscreteSize(PlaybackRange), TickResolution, FrameRate).FrameNumber.Value;
 
 	bool bFoundCinematicMasterTrack = false;
 
@@ -268,7 +270,7 @@ bool FMovieSceneExportData::ConstructSectionData(const UMovieScene* InMovieScene
 	if (InSection->HasStartFrame())
 	{
 		FFrameTime InclusiveStartFrame = InSection->GetInclusiveStartFrame();
-		FFrameTime ConvertedStartFrame = ConvertFrameTime(InclusiveStartFrame, MovieSceneData->Resolution, FrameRate);
+		FFrameTime ConvertedStartFrame = ConvertFrameTime(InclusiveStartFrame, MovieSceneData->TickResolution, FrameRate);
 		InSectionData->StartFrame = ConvertedStartFrame.CeilToFrame();
 
 		if (ExportContext.IsValid() && InSectionType == EMovieSceneTranslatorSectionType::Cinematic && ConvertedStartFrame.GetSubFrame() > 0.0f)
@@ -292,7 +294,7 @@ bool FMovieSceneExportData::ConstructSectionData(const UMovieScene* InMovieScene
 	if (InSection->HasEndFrame())
 	{
 		FFrameTime ExclusiveEndFrame = InSection->GetExclusiveEndFrame();
-		FFrameTime ConvertedEndFrame = ConvertFrameTime(ExclusiveEndFrame, MovieSceneData->Resolution, FrameRate);
+		FFrameTime ConvertedEndFrame = ConvertFrameTime(ExclusiveEndFrame, MovieSceneData->TickResolution, FrameRate);
 		InSectionData->EndFrame = ConvertedEndFrame.CeilToFrame();
 
 		if (ExportContext.IsValid() && InSectionType == EMovieSceneTranslatorSectionType::Cinematic && ConvertedEndFrame.GetSubFrame() > 0.0f)
@@ -338,6 +340,16 @@ FFrameRate FMovieSceneExportData::GetFrameRate() const
 	return FrameRate;
 }
 
+uint32 FMovieSceneExportData::GetResX() const
+{
+	return ResX;
+}
+
+uint32 FMovieSceneExportData::GetResY() const
+{
+	return ResY;
+}
+
 uint32 FMovieSceneExportData::GetNearestWholeFrameRate() const
 {
 	if (GetFrameRateIsNTSC())
@@ -451,10 +463,10 @@ TSharedPtr<FMovieSceneImportCinematicSectionData> FMovieSceneImportData::CreateC
 	}
 
 	// both FCP XML and Sequencer have inclusive start frame, exclusive end frame
-	FFrameRate Resolution = MovieSceneData->MovieScene->GetTickResolution();
-	FFrameNumber StartFrame = ConvertFrameTime(InStartFrame, InFrameRate, Resolution).RoundToFrame();
-	FFrameNumber StartOffsetFrame = ConvertFrameTime(InStartOffsetFrame, InFrameRate, Resolution).RoundToFrame();
-	FFrameNumber EndFrame = ConvertFrameTime(InEndFrame, InFrameRate, Resolution).RoundToFrame();
+	FFrameRate TickResolution = MovieSceneData->MovieScene->GetTickResolution();
+	FFrameNumber StartFrame = ConvertFrameTime(InStartFrame, InFrameRate, TickResolution).RoundToFrame();
+	FFrameNumber StartOffsetFrame = ConvertFrameTime(InStartOffsetFrame, InFrameRate, TickResolution).RoundToFrame();
+	FFrameNumber EndFrame = ConvertFrameTime(InEndFrame, InFrameRate, TickResolution).RoundToFrame();
 	int32 Duration = (EndFrame - StartFrame).Value;
 	
 	MasterTrack->Modify();
@@ -481,10 +493,10 @@ bool FMovieSceneImportData::SetCinematicSection(TSharedPtr<FMovieSceneImportCine
 		return false;
 	}
 
-	FFrameRate Resolution = MovieSceneData->MovieScene->GetTickResolution();
-	FFrameNumber StartFrame = ConvertFrameTime(InStartFrame, InFrameRate, Resolution).GetFrame();
-	FFrameNumber StartOffsetFrame = ConvertFrameTime(InStartOffsetFrame, InFrameRate, Resolution).GetFrame();
-	FFrameNumber EndFrame = ConvertFrameTime(InEndFrame, InFrameRate, Resolution).GetFrame();
+	FFrameRate TickResolution = MovieSceneData->MovieScene->GetTickResolution();
+	FFrameNumber StartFrame = ConvertFrameTime(InStartFrame, InFrameRate, TickResolution).GetFrame();
+	FFrameNumber StartOffsetFrame = ConvertFrameTime(InStartOffsetFrame, InFrameRate, TickResolution).GetFrame();
+	FFrameNumber EndFrame = ConvertFrameTime(InEndFrame, InFrameRate, TickResolution).GetFrame();
 
 	InSection->CinematicSection->Modify();
 	InSection->CinematicSection->Parameters.SetStartFrameOffset(StartOffsetFrame.Value);
diff --git a/Engine/Source/Editor/MovieSceneTools/Private/TrackEditors/CinematicShotTrackEditor.cpp b/Engine/Source/Editor/MovieSceneTools/Private/TrackEditors/CinematicShotTrackEditor.cpp
--- a/Engine/Source/Editor/MovieSceneTools/Private/TrackEditors/CinematicShotTrackEditor.cpp
+++ b/Engine/Source/Editor/MovieSceneTools/Private/TrackEditors/CinematicShotTrackEditor.cpp
@@ -36,6 +36,7 @@
 #include "MovieSceneTimeHelpers.h"
 #include "Framework/Notifications/NotificationManager.h"
 #include "Widgets/Notifications/SNotificationList.h"
+#include "MovieSceneCaptureSettings.h"
 
 #define LOCTEXT_NAMESPACE "FCinematicShotTrackEditor"
 
@@ -941,12 +942,10 @@ void FCinematicShotTrackEditor::ExportFCPXML()
 	}
 
 	const FMovieSceneCaptureSettings& Settings = MovieSceneCapture->GetSettings();
-	FString SaveDirectory = FPaths::ConvertRelativePathToFull(Settings.OutputDirectory.Path);
-	int32 HandleFrames = Settings.HandleFrames;
 
 	FFCPXMLExporter *Exporter = new FFCPXMLExporter;
 
-	MovieSceneToolHelpers::MovieSceneTranslatorExport(Exporter, MovieScene, MovieScene->GetDisplayRate(), SaveDirectory, HandleFrames);
+	MovieSceneToolHelpers::MovieSceneTranslatorExport(Exporter, MovieScene, Settings);
 
 	delete Exporter;
 }
diff --git a/Engine/Source/Editor/MovieSceneTools/Public/AutomatedLevelSequenceCapture.h b/Engine/Source/Editor/MovieSceneTools/Public/AutomatedLevelSequenceCapture.h
--- a/Engine/Source/Editor/MovieSceneTools/Public/AutomatedLevelSequenceCapture.h
+++ b/Engine/Source/Editor/MovieSceneTools/Public/AutomatedLevelSequenceCapture.h
@@ -66,6 +66,10 @@ public:
 	UPROPERTY(config, EditAnywhere, BlueprintReadWrite, Category=Sequence)
 	bool bWriteEditDecisionList;
 
+	/** Whether to write Final Cut Pro XML files (XMLs) if the sequence contains shots */
+	UPROPERTY(config, EditAnywhere, Category=Sequence)
+	bool bWriteFinalCutProXML;
+
 public:
 	// UMovieSceneCapture interface
 	virtual void Initialize(TSharedPtr<FSceneViewport> InViewport, int32 PIEInstance = -1) override;
@@ -106,6 +110,9 @@ private:
 	/** Export EDL if requested */
 	void ExportEDL();
 
+	/** Export FCPXML if requested */
+	void ExportFCPXML();
+
 	/** Delegate binding for the above callback */
 	FDelegateHandle OnPlayerUpdatedBinding;
 
diff --git a/Engine/Source/Editor/MovieSceneTools/Public/FCPXML/FCPXMLMovieSceneTranslator.h b/Engine/Source/Editor/MovieSceneTools/Public/FCPXML/FCPXMLMovieSceneTranslator.h
--- a/Engine/Source/Editor/MovieSceneTools/Public/FCPXML/FCPXMLMovieSceneTranslator.h
+++ b/Engine/Source/Editor/MovieSceneTools/Public/FCPXML/FCPXMLMovieSceneTranslator.h
@@ -76,10 +76,12 @@ public:
 	 *
 	 * @param InMovieScene The movie scene with the cinematic shot track and audio tracks to export
 	 * @param InFrameRate The frame rate to export the XML at
-	 * @param InSaveFilename The file path to save to.
+	 * @param InResX Sequence resolution x.
+	 * @param InResY Sequence resolution y.
 	 * @param InHandleFrames The number of handle frames to include for each shot.
+	 * @param InSaveFilename The file path to save to.
 	 * @param OutError The return error message
 	 * @return Whether the export was successful
 	 */
-	virtual bool Export(const UMovieScene* InMovieScene, FFrameRate InFrameRate, FString InSaveFilename, int32 InHandleFrames, TSharedRef<FMovieSceneTranslatorContext> InContext);
+	virtual bool Export(const UMovieScene* InMovieScene, FFrameRate InFrameRate, uint32 InResX, uint32 InResY, int32 InHandleFrames, FString InSaveFilename, TSharedRef<FMovieSceneTranslatorContext> InContext);
 };
diff --git a/Engine/Source/Editor/MovieSceneTools/Public/MovieSceneToolHelpers.h b/Engine/Source/Editor/MovieSceneTools/Public/MovieSceneToolHelpers.h
--- a/Engine/Source/Editor/MovieSceneTools/Public/MovieSceneToolHelpers.h
+++ b/Engine/Source/Editor/MovieSceneTools/Public/MovieSceneToolHelpers.h
@@ -13,6 +13,7 @@
 #include "ISequencer.h"
 #include "Logging/TokenizedMessage.h"
 #include "MovieSceneTranslator.h"
+#include "MovieSceneCaptureSettings.h"
 
 class ISequencer;
 class UMovieScene;
@@ -167,7 +168,7 @@ public:
 	* @param InHandleFrames The number of handle frames to include for each shot.
 	* @return Whether the export was successful
 	*/
-	static bool MovieSceneTranslatorExport(FMovieSceneExporter* InExporter, const UMovieScene* InMovieScene, FFrameRate InFrameRate, FString InSaveDirectory = TEXT(""), int32 InHandleFrames = 8);
+	static bool MovieSceneTranslatorExport(FMovieSceneExporter* InExporter, const UMovieScene* InMovieScene, const FMovieSceneCaptureSettings& Settings);
 
 	/** 
 	* Display error message window for MovieScene translators
diff --git a/Engine/Source/Editor/MovieSceneTools/Public/MovieSceneTranslator.h b/Engine/Source/Editor/MovieSceneTools/Public/MovieSceneTranslator.h
--- a/Engine/Source/Editor/MovieSceneTools/Public/MovieSceneTranslator.h
+++ b/Engine/Source/Editor/MovieSceneTools/Public/MovieSceneTranslator.h
@@ -95,7 +95,7 @@ struct FMovieSceneExportMovieSceneData
 {
 	FString Name;
 	FString Path;
-	FFrameRate Resolution;
+	FFrameRate TickResolution;
 	int32 Duration;
 	FFrameNumber PlaybackRangeStartFrame;
 	FFrameNumber PlaybackRangeEndFrame;
@@ -116,7 +116,7 @@ class MOVIESCENETOOLS_API FMovieSceneExportData : public TSharedFromThis<FMovieS
 {
 public:
 	/** Constructor */
-	FMovieSceneExportData(const UMovieScene* InMovieScene, FFrameRate InFrameRate, int32 InHandleFrames, FString InSaveFilename, TSharedPtr<FMovieSceneTranslatorContext> InContext);
+	FMovieSceneExportData(const UMovieScene* InMovieScene, FFrameRate InFrameRate, uint32 InResX, uint32 InResY, int32 InHandleFrames, FString InSaveFilename, TSharedPtr<FMovieSceneTranslatorContext> InContext);
 	/** Default constructor, necessary for shared ref - should not be used */
 	FMovieSceneExportData();
 	/** Destructor */
@@ -129,6 +129,10 @@ public:
 
 	/** Gets export frame rate */
 	FFrameRate GetFrameRate() const;
+	/** Gets x resolution */
+	uint32 GetResX() const;
+	/** Gets y resolution */
+	uint32 GetResY() const;
 	/** Returns true if frame rate is a non-integral frame rate */
 	bool GetFrameRateIsNTSC() const;
 	/** Returns the nearest integral frame rate */
@@ -170,6 +174,8 @@ public:
 private:
 
 	FFrameRate FrameRate;
+	uint32 ResX;
+	uint32 ResY;
 	FFrameTime HandleFrames;
 	FString SaveFilename;
 	FString SaveFilenamePath;
@@ -320,12 +326,14 @@ public:
 	* Export movie scene
 	*
 	* @param InMovieScene The movie scene with the cinematic shot track and audio tracks to export
-	* @param InFrameRate The frame rate to export the XML at
-	* @param InSaveFilename The file path to save to.
+	* @param InFrameRate The frame rate for export
+	* @param InResX Sequence resolution x.
+	* @param InResY Sequence resolution y.
 	* @param InHandleFrames The number of handle frames to include for each shot.
+	* @param InSaveFilename The file path to save to.
 	* @param OutError The return error message
 	* @return Whether the export was successful
 	*/
-	virtual bool Export(const UMovieScene* InMovieScene, FFrameRate InFrameRate, FString InSaveFilename, int32 InHandleFrames, TSharedRef<FMovieSceneTranslatorContext> InContext) = 0;
+	virtual bool Export(const UMovieScene* InMovieScene, FFrameRate InFrameRate, uint32 InResX, uint32 InResY, int32 InHandleFrames, FString InSaveFilename, TSharedRef<FMovieSceneTranslatorContext> InContext) = 0;
 };
 
