Sequencer: Upon export, log warnings when filename format is not '{shot}'

diff --git a/Engine/Source/Editor/MovieSceneTools/Private/AutomatedLevelSequenceCapture.cpp b/Engine/Source/Editor/MovieSceneTools/Private/AutomatedLevelSequenceCapture.cpp
--- a/Engine/Source/Editor/MovieSceneTools/Private/AutomatedLevelSequenceCapture.cpp
+++ b/Engine/Source/Editor/MovieSceneTools/Private/AutomatedLevelSequenceCapture.cpp
@@ -824,6 +824,7 @@ void UAutomatedLevelSequenceCapture::ExportFCPXML()
 	}
 
 	FString SaveFilename = Settings.OutputDirectory.Path / MovieScene->GetOuter()->GetName() + TEXT(".xml");
+	FString FilenameFormat = Settings.OutputFormat;
 	int32 HandleFrames = Settings.HandleFrames;
 	FFrameRate FrameRate = Settings.FrameRate;
 	uint32 ResX = Settings.Resolution.ResX;
@@ -833,7 +834,11 @@ void UAutomatedLevelSequenceCapture::ExportFCPXML()
 
 	TSharedRef<FMovieSceneTranslatorContext> ExportContext(new FMovieSceneTranslatorContext);
 	ExportContext->Init();
-	bool bSuccess = Exporter->Export(MovieScene, FrameRate, ResX, ResY, HandleFrames, SaveFilename, ExportContext);
+
+	bool bSuccess = Exporter->Export(MovieScene, FilenameFormat, FrameRate, ResX, ResY, HandleFrames, SaveFilename, ExportContext);
+
+	// Log any messages in context
+	MovieSceneToolHelpers::MovieSceneTranslatorLogMessages(Exporter, ExportContext, false);
 
 	delete Exporter;
 }
diff --git a/Engine/Source/Editor/MovieSceneTools/Private/FCPXML/FCPXMLMovieSceneTranslator.cpp b/Engine/Source/Editor/MovieSceneTools/Private/FCPXML/FCPXMLMovieSceneTranslator.cpp
--- a/Engine/Source/Editor/MovieSceneTools/Private/FCPXML/FCPXMLMovieSceneTranslator.cpp
+++ b/Engine/Source/Editor/MovieSceneTools/Private/FCPXML/FCPXMLMovieSceneTranslator.cpp
@@ -58,8 +58,6 @@ FText FFCPXMLImporter::GetMessageLogLabel() const
 
 bool FFCPXMLImporter::Import(UMovieScene* InMovieScene, FFrameRate InFrameRate, FString InFilename, TSharedRef<FMovieSceneTranslatorContext> InContext)
 {
-	InContext->Init();
-
 	// Create intermediate structure to assist with import
 	TSharedRef<FMovieSceneImportData> ImportData = MakeShared<FMovieSceneImportData>(InMovieScene, InContext);
 	if (!ImportData->IsImportDataValid())
@@ -146,9 +144,17 @@ FText FFCPXMLExporter::GetMessageLogLabel() const
 	return LOCTEXT("FCPXMLExportLogLabel", "FCP 7 XML Export Log");
 }
 
-bool FFCPXMLExporter::Export(const UMovieScene* InMovieScene, FFrameRate InFrameRate, uint32 InResX, uint32 InResY, int32 InHandleFrames, FString InSaveFilename, TSharedRef<FMovieSceneTranslatorContext> InContext)
+bool FFCPXMLExporter::Export(const UMovieScene* InMovieScene, FString InFilenameFormat, FFrameRate InFrameRate, uint32 InResX, uint32 InResY, int32 InHandleFrames, FString InSaveFilename, TSharedRef<FMovieSceneTranslatorContext> InContext)
 {
-	InContext->Init();
+	// add warning message if filename format is not "{shot}"
+	FString AcceptedFormat = TEXT("{shot}");
+	if (InFilenameFormat != AcceptedFormat)
+	{
+		InContext->AddMessage(EMessageSeverity::Warning,
+			FText::Format(LOCTEXT("FCPXMLExportFilenameFormatWarning", "FCP 7 XML export expects rendered video sequences with filename format '{0}' but most recent rendered filename format is '{1}'."),
+				FText::FromString(AcceptedFormat),
+				FText::FromString(InFilenameFormat)));
+	}
 
 	// Construct XML from file string
 	TSharedRef<FFCPXMLFile> FCPXMLFile = MakeShared<FFCPXMLFile>();
diff --git a/Engine/Source/Editor/MovieSceneTools/Private/MovieSceneToolHelpers.cpp b/Engine/Source/Editor/MovieSceneTools/Private/MovieSceneToolHelpers.cpp
--- a/Engine/Source/Editor/MovieSceneTools/Private/MovieSceneToolHelpers.cpp
+++ b/Engine/Source/Editor/MovieSceneTools/Private/MovieSceneToolHelpers.cpp
@@ -717,7 +717,7 @@ bool MovieSceneToolHelpers::MovieSceneTranslatorImport(FMovieSceneImporter* InIm
 	bool bSuccess = InImporter->Import(InMovieScene, InFrameRate, OpenFilenames[0], ImportContext);
 
 	// Display any messages in context
-	MovieSceneTranslatorDisplayMessages(InImporter, ImportContext);
+	MovieSceneTranslatorLogMessages(InImporter, ImportContext, true);
 
 	// Roll back transaction when import fails.
 	if (!bSuccess)
@@ -774,10 +774,10 @@ bool MovieSceneToolHelpers::MovieSceneTranslatorExport(FMovieSceneExporter* InEx
 	TSharedRef<FMovieSceneTranslatorContext> ExportContext(new FMovieSceneTranslatorContext);
 	ExportContext->Init();
 
-	bool bSuccess = InExporter->Export(InMovieScene, FrameRate, ResX, ResY, HandleFrames, SaveFilenames[0], ExportContext);
+	bool bSuccess = InExporter->Export(InMovieScene, FilenameFormat, FrameRate, ResX, ResY, HandleFrames, SaveFilenames[0], ExportContext);
 	
 	// Display any messages in context
-	MovieSceneTranslatorDisplayMessages(InExporter, ExportContext);
+	MovieSceneTranslatorLogMessages(InExporter, ExportContext, true);
 
 	if (bSuccess)
 	{
@@ -794,13 +794,13 @@ bool MovieSceneToolHelpers::MovieSceneTranslatorExport(FMovieSceneExporter* InEx
 	return bSuccess;
 }
 
-void MovieSceneToolHelpers::MovieSceneTranslatorDisplayMessages(FMovieSceneTranslator *InTranslator, TSharedRef<FMovieSceneTranslatorContext> InContext)
+void MovieSceneToolHelpers::MovieSceneTranslatorLogMessages(FMovieSceneTranslator *InTranslator, TSharedRef<FMovieSceneTranslatorContext> InContext, bool bDisplayMessages)
 {
 	if (InTranslator == nullptr || InContext->GetMessages().Num() == 0)
 	{
 		return;
 	}
-
+	
 	// Clear any old messages after an import or export
 	const FName LogTitle = InTranslator->GetMessageLogWindowTitle();
 	FMessageLogModule& MessageLogModule = FModuleManager::LoadModuleChecked<FMessageLogModule>("MessageLog");
@@ -813,7 +813,31 @@ void MovieSceneToolHelpers::MovieSceneTranslatorDisplayMessages(FMovieSceneTrans
 		LogListing->AddMessage(Message);
 	}
 
-	MessageLogModule.OpenMessageLog(LogTitle);
+	if (bDisplayMessages)
+	{
+		MessageLogModule.OpenMessageLog(LogTitle);
+	}
+}
+
+void MovieSceneToolHelpers::MovieSceneTranslatorLogOutput(FMovieSceneTranslator *InTranslator, TSharedRef<FMovieSceneTranslatorContext> InContext)
+{
+	if (InTranslator == nullptr || InContext->GetMessages().Num() == 0)
+	{
+		return;
+	}
+
+	for (TSharedRef<FTokenizedMessage> Message : InContext->GetMessages())
+	{
+		if (Message->GetSeverity() == EMessageSeverity::Error)
+		{
+			UE_LOG(LogMovieScene, Error, TEXT("%s"), *Message->ToText().ToString());
+		}
+		else if (Message->GetSeverity() == EMessageSeverity::Warning)
+		{
+			UE_LOG(LogMovieScene, Warning, TEXT("%s"), *Message->ToText().ToString());
+
+		}
+	}
 }
 
 bool ImportFBXProperty(FString NodeName, FString AnimatedPropertyName, FGuid ObjectBinding, UnFbx::FFbxCurvesAPI& CurveAPI, UMovieScene* InMovieScene, ISequencer& InSequencer)
diff --git a/Engine/Source/Editor/MovieSceneTools/Private/MovieSceneTranslator.cpp b/Engine/Source/Editor/MovieSceneTools/Private/MovieSceneTranslator.cpp
--- a/Engine/Source/Editor/MovieSceneTools/Private/MovieSceneTranslator.cpp
+++ b/Engine/Source/Editor/MovieSceneTools/Private/MovieSceneTranslator.cpp
@@ -25,11 +25,6 @@ FMovieSceneExportData::FMovieSceneExportData(const UMovieScene* InMovieScene, FF
 	}
 
 	ExportContext = InContext;
-	if (ExportContext.IsValid())
-	{
-		ExportContext->Init();
-	}
-
 	FrameRate = InFrameRate;
 	ResX = InResX;
 	ResY = InResY;
diff --git a/Engine/Source/Editor/MovieSceneTools/Public/FCPXML/FCPXMLMovieSceneTranslator.h b/Engine/Source/Editor/MovieSceneTools/Public/FCPXML/FCPXMLMovieSceneTranslator.h
--- a/Engine/Source/Editor/MovieSceneTools/Public/FCPXML/FCPXMLMovieSceneTranslator.h
+++ b/Engine/Source/Editor/MovieSceneTools/Public/FCPXML/FCPXMLMovieSceneTranslator.h
@@ -75,7 +75,8 @@ public:
 	 * Export FCP 7 XML
 	 *
 	 * @param InMovieScene The movie scene with the cinematic shot track and audio tracks to export
-	 * @param InFrameRate The frame rate to export the XML at
+	 * @param InFilenameFormat The last filename format used to render shots.
+	 * @param InFrameRate The frame rate for export.
 	 * @param InResX Sequence resolution x.
 	 * @param InResY Sequence resolution y.
 	 * @param InHandleFrames The number of handle frames to include for each shot.
@@ -83,5 +84,5 @@ public:
 	 * @param OutError The return error message
 	 * @return Whether the export was successful
 	 */
-	virtual bool Export(const UMovieScene* InMovieScene, FFrameRate InFrameRate, uint32 InResX, uint32 InResY, int32 InHandleFrames, FString InSaveFilename, TSharedRef<FMovieSceneTranslatorContext> InContext);
+	virtual bool Export(const UMovieScene* InMovieScene, FString InFilenameFormat, FFrameRate InFrameRate, uint32 InResX, uint32 InResY, int32 InHandleFrames, FString InSaveFilename, TSharedRef<FMovieSceneTranslatorContext> InContext);
 };
diff --git a/Engine/Source/Editor/MovieSceneTools/Public/MovieSceneToolHelpers.h b/Engine/Source/Editor/MovieSceneTools/Public/MovieSceneToolHelpers.h
--- a/Engine/Source/Editor/MovieSceneTools/Public/MovieSceneToolHelpers.h
+++ b/Engine/Source/Editor/MovieSceneTools/Public/MovieSceneToolHelpers.h
@@ -171,12 +171,21 @@ public:
 	static bool MovieSceneTranslatorExport(FMovieSceneExporter* InExporter, const UMovieScene* InMovieScene, const FMovieSceneCaptureSettings& Settings);
 
 	/** 
-	* Display error message window for MovieScene translators
+	* Log messages and display error message window for MovieScene translators
 	*
 	* @param InTranslator The movie scene importer or exporter.
 	* @param InContext The context used to gather error, warning or info messages during import or export.
+	* @param bDisplayMessages Whether to open the message log window after adding the message.
 	*/
-	static void MovieSceneTranslatorDisplayMessages(FMovieSceneTranslator* InTranslator,  TSharedRef<FMovieSceneTranslatorContext> InContext);
+	static void MovieSceneTranslatorLogMessages(FMovieSceneTranslator* InTranslator, TSharedRef<FMovieSceneTranslatorContext> InContext, bool bDisplayMessages);
+
+	/**
+	* Log error output for MovieScene translators
+	*
+	* @param InTranslator The movie scene importer or exporter.
+	* @param InContext The context used to gather error, warning or info messages during import or export.
+	*/
+	static void MovieSceneTranslatorLogOutput(FMovieSceneTranslator* InTranslator, TSharedRef<FMovieSceneTranslatorContext> InContext);
 
 	/**
 	 * Import FBX
diff --git a/Engine/Source/Editor/MovieSceneTools/Public/MovieSceneTranslator.h b/Engine/Source/Editor/MovieSceneTools/Public/MovieSceneTranslator.h
--- a/Engine/Source/Editor/MovieSceneTools/Public/MovieSceneTranslator.h
+++ b/Engine/Source/Editor/MovieSceneTools/Public/MovieSceneTranslator.h
@@ -326,7 +326,8 @@ public:
 	* Export movie scene
 	*
 	* @param InMovieScene The movie scene with the cinematic shot track and audio tracks to export
-	* @param InFrameRate The frame rate for export
+	* @param InFilenameFormat The specified filename format.
+	* @param InFrameRate The frame rate for export.
 	* @param InResX Sequence resolution x.
 	* @param InResY Sequence resolution y.
 	* @param InHandleFrames The number of handle frames to include for each shot.
@@ -334,6 +335,6 @@ public:
 	* @param OutError The return error message
 	* @return Whether the export was successful
 	*/
-	virtual bool Export(const UMovieScene* InMovieScene, FFrameRate InFrameRate, uint32 InResX, uint32 InResY, int32 InHandleFrames, FString InSaveFilename, TSharedRef<FMovieSceneTranslatorContext> InContext) = 0;
+	virtual bool Export(const UMovieScene* InMovieScene, FString InFilenameFormat, FFrameRate InFrameRate, uint32 InResX, uint32 InResY, int32 InHandleFrames, FString InSaveFilename, TSharedRef<FMovieSceneTranslatorContext> InContext) = 0;
 };
 
