Sequencer: Added audio support to FCP XML import/export

diff --git a/Engine/Source/Editor/MovieSceneTools/Private/FCPXML/FCPXMLExport.cpp b/Engine/Source/Editor/MovieSceneTools/Private/FCPXML/FCPXMLExport.cpp
--- a/Engine/Source/Editor/MovieSceneTools/Private/FCPXML/FCPXMLExport.cpp
+++ b/Engine/Source/Editor/MovieSceneTools/Private/FCPXML/FCPXMLExport.cpp
@@ -10,6 +10,7 @@
 #include "AssetRegistryModule.h"
 #include "UObject/MetaData.h"
 #include "Logging/TokenizedMessage.h"
+#include "Sound/SoundWave.h"
 
 #define LOCTEXT_NAMESPACE "FCPXMLExporter"
 
@@ -28,6 +29,9 @@ FFCPXMLExportVisitor::FFCPXMLExportVisitor(FString InSaveFilename, TSharedRef<FM
 	}
 
 	SequenceId = 0;
+	MasterClipId = 0;
+	ClipItemId = 0;
+	FileId = 0;
 }
 
 FFCPXMLExportVisitor::~FFCPXMLExportVisitor() {}
@@ -53,7 +57,7 @@ bool FFCPXMLExportVisitor::VisitNode(TSharedRef<FFCPXMLXmemlNode> InXmemlNode)
 	/** @todo - MERGE METADATA
 
 	   Merging the newly exported XML structure with pre-existing
-	   XML metadata will be implemented here. The traversal would proceed
+	   XML metadata may be implemented here. The traversal would proceed
 	   through the new XML structure, referring back to the metadata 
 	   XML structure to incorporate any missing or desired attributes or
 	   elements. Alternatively, the traversal might be invoked directly on
@@ -131,8 +135,6 @@ bool FFCPXMLExportVisitor::ConstructMasterClipNodes(TSharedRef<FFCPXMLNode> InPa
 		return false;
 	}
 
-	int32 MasterClipId = 0;
-
 	TSharedPtr<FMovieSceneExportCinematicMasterTrackData> CinematicMasterTrackData = ExportData->MovieSceneData->CinematicMasterTrack;
 	if (!CinematicMasterTrackData.IsValid())
 	{
@@ -153,15 +155,43 @@ bool FFCPXMLExportVisitor::ConstructMasterClipNodes(TSharedRef<FFCPXMLNode> InPa
 			continue;
 		}
 
-		if (!ConstructMasterClipNode(InParentNode, CinematicSection, ++MasterClipId))
+		if (!ConstructMasterClipNode(InParentNode, CinematicSection, CinematicMasterTrackData))
 		{
 			return false;
 		}
 	}
+
+	for (TSharedPtr<FMovieSceneExportAudioMasterTrackData> AudioMasterTrack : ExportData->MovieSceneData->AudioMasterTracks)
+	{
+		if (!AudioMasterTrack.IsValid())
+		{
+			return false;
+		}
+
+		for (TSharedPtr<FMovieSceneExportAudioTrackData> AudioTrack : AudioMasterTrack->AudioTracks)
+		{
+			if (!AudioTrack.IsValid())
+			{
+				return false;
+			}
+
+			for (TSharedPtr<FMovieSceneExportAudioSectionData> AudioSection : AudioTrack->AudioSections)
+			{
+				bool bMasterClipExists = false;
+				FString MasterClipName;
+				HasMasterClipIdName(AudioSection, MasterClipName, bMasterClipExists);
+
+				if (!bMasterClipExists)
+				{
+					ConstructMasterClipNode(InParentNode, AudioSection, AudioMasterTrack);
+				}
+			}
+		}
+	}
 	return true;
 }
 
-bool FFCPXMLExportVisitor::ConstructMasterClipNode(TSharedRef<FFCPXMLNode> InParentNode, const TSharedPtr<FMovieSceneExportCinematicSectionData> InCinematicSectionData, int32 InMasterClipId)
+bool FFCPXMLExportVisitor::ConstructMasterClipNode(TSharedRef<FFCPXMLNode> InParentNode, const TSharedPtr<FMovieSceneExportCinematicSectionData> InCinematicSectionData, const TSharedPtr<FMovieSceneExportCinematicMasterTrackData> InCinematicMasterTrackData)
 {
 	if (!InCinematicSectionData.IsValid())
 	{
@@ -173,16 +203,12 @@ bool FFCPXMLExportVisitor::ConstructMasterClipNode(TSharedRef<FFCPXMLNode> InPar
 	int32 EndFrame{ 0 };
 	int32 InFrame{ 0 };
 	int32 OutFrame{ 0 };
-	FString SectionName(TEXT(""));
-	GetSectionFrames(InCinematicSectionData, Duration, StartFrame, EndFrame, InFrame, OutFrame);
-	GetSectionName(InCinematicSectionData, SectionName);
+	FString SectionName = InCinematicSectionData->DisplayName;
+	GetCinematicSectionFrames(InCinematicSectionData, Duration, StartFrame, EndFrame, InFrame, OutFrame);
 
 	/** Construct a master clip id name based on the cinematic section and id */
-	FString MasterClipName(TEXT(""));
-	if (!CreateMasterClipIdName(InCinematicSectionData, InMasterClipId, MasterClipName))
-	{
-		return false;
-	}
+	FString MasterClipName{ TEXT("") };
+	GetMasterClipIdName(InCinematicSectionData, MasterClipName);
 
 	TSharedRef<FFCPXMLNode> ClipNode = InParentNode->CreateChildNode(TEXT("clip"));
 	ClipNode->AddAttribute(TEXT("id"), MasterClipName);
@@ -205,12 +231,12 @@ bool FFCPXMLExportVisitor::ConstructMasterClipNode(TSharedRef<FFCPXMLNode> InPar
 	TSharedRef<FFCPXMLNode> VideoNode = MediaNode->CreateChildNode(TEXT("video"));
 	TSharedRef<FFCPXMLNode> TrackNode = VideoNode->CreateChildNode(TEXT("track"));
 
-	if (!ConstructVideoClipItemNode(TrackNode, InCinematicSectionData, true))
+	if (!ConstructVideoClipItemNode(TrackNode, InCinematicSectionData, InCinematicMasterTrackData, true))
 	{
 		return false;
 	}
 
-	if (!ConstructSectionLoggingInfoNode(ClipNode, InCinematicSectionData, SectionName))
+	if (!ConstructLoggingInfoNode(ClipNode, InCinematicSectionData))
 	{
 		return false;
 	}
@@ -223,23 +249,174 @@ bool FFCPXMLExportVisitor::ConstructMasterClipNode(TSharedRef<FFCPXMLNode> InPar
 	return true;
 }
 
+bool FFCPXMLExportVisitor::ConstructMasterClipNode(TSharedRef<FFCPXMLNode> InParentNode, const TSharedPtr<FMovieSceneExportAudioSectionData> InAudioSectionData, const TSharedPtr<FMovieSceneExportAudioMasterTrackData> InAudioMasterTrackData)
+{
+	if (!InAudioSectionData.IsValid())
+	{
+		return false;
+	}
+
+	bool bIsStereo = (InAudioSectionData->NumChannels == 2);
+
+	int32 Duration{ 0 };
+	int32 StartFrame{ 0 };
+	int32 EndFrame{ 0 };
+	int32 InFrame{ 0 };
+	int32 OutFrame{ 0 };
+	FString SectionName = InAudioSectionData->DisplayName;
+	GetAudioSectionFrames(InAudioSectionData, Duration, StartFrame, EndFrame, InFrame, OutFrame);
+
+	/** Construct a master clip id name based on the audio section and id */
+	FString MasterClipName{ TEXT("") };
+	GetMasterClipIdName(InAudioSectionData, MasterClipName);
+
+	TSharedRef<FFCPXMLNode> ClipNode = InParentNode->CreateChildNode(TEXT("clip"));
+	ClipNode->AddAttribute(TEXT("id"), MasterClipName);
+	ClipNode->AddAttribute(TEXT("explodedTracks"), TEXT("true"));
+
+	// @todo add to file's masterclip and refidmap HERE
+
+	ClipNode->CreateChildNode(TEXT("masterclipid"))->SetContent(MasterClipName);
+	ClipNode->CreateChildNode(TEXT("ismasterclip"))->SetContent(true);
+	ClipNode->CreateChildNode(TEXT("duration"))->SetContent(Duration);
+
+	if (!ConstructRateNode(ClipNode))
+	{
+		return false;
+	}
+
+	ClipNode->CreateChildNode(TEXT("in"))->SetContent(InFrame);
+	ClipNode->CreateChildNode(TEXT("out"))->SetContent(OutFrame);
+	ClipNode->CreateChildNode(TEXT("name"))->SetContent(SectionName);
+	TSharedRef<FFCPXMLNode> MediaNode = ClipNode->CreateChildNode(TEXT("media"));
+	TSharedRef<FFCPXMLNode> AudioNode = MediaNode->CreateChildNode(TEXT("audio"));
+	TSharedRef<FFCPXMLNode> TrackNode = AudioNode->CreateChildNode(TEXT("track"));
+
+	FString ClipItemIdName1{ TEXT("") };
+	FString ClipItemIdName2{ TEXT("") };
+	GetNextClipItemIdName(ClipItemIdName1);
+	if (bIsStereo)
+	{
+		GetNextClipItemIdName(ClipItemIdName2);
+	}
+	if (!ConstructAudioClipItemNode(TrackNode, InAudioSectionData, InAudioMasterTrackData, 1, true, ClipItemIdName1, ClipItemIdName2, 1, 1, 1, 2))
+	{
+		return false;
+	}
+
+	// handle stereo master clip
+	if (bIsStereo)
+	{
+		TrackNode = AudioNode->CreateChildNode(TEXT("track"));
+
+		if (!ConstructAudioClipItemNode(TrackNode, InAudioSectionData, InAudioMasterTrackData, 2, true, ClipItemIdName1, ClipItemIdName2, 1, 1, 1, 2))
+		{
+			return false;
+		}
+	}
+
+	if (!ConstructLoggingInfoNode(ClipNode, InAudioSectionData))
+	{
+		return false;
+	}
+
+
+	return true;
+}
+
+
+/** Creates logginginfo node. */
+bool FFCPXMLExportVisitor::ConstructLoggingInfoNode(TSharedRef<FFCPXMLNode> InParentNode, const TSharedPtr<FMovieSceneExportCinematicSectionData> InSectionData)
+{
+	if (!InSectionData.IsValid() || InSectionData->MovieSceneSection == nullptr)
+	{
+		return false;
+	}
+
+	TSharedRef<FFCPXMLNode> LoggingInfoNode = InParentNode->CreateChildNode(TEXT("logginginfo"));
+	ConstructLoggingInfoElements(LoggingInfoNode, InSectionData->MovieSceneSection);
+
+	TSharedPtr<FFCPXMLNode> LogNoteNode = LoggingInfoNode->GetChildNode(TEXT("lognote"), ENodeInherit::NoInherit, ENodeReference::NoReferences);
+	if (!LogNoteNode.IsValid())
+	{
+		LogNoteNode = LoggingInfoNode->CreateChildNode(TEXT("lognote"));
+	}
+
+	FString Metadata{ TEXT("") };
+	const UMovieSceneCinematicShotSection* ShotSection = Cast<UMovieSceneCinematicShotSection>(InSectionData->MovieSceneSection);
+	if (ShotSection == nullptr)
+	{
+		return false;
+	}
+
+	if (!CreateCinematicSectionMetadata(ShotSection, Metadata))
+	{
+		return false;
+	}
+	LogNoteNode->SetContent(Metadata);
+
+	return true;
+}
+
 /** Creates logginginfo node. */
-bool FFCPXMLExportVisitor::ConstructSectionLoggingInfoNode(TSharedRef<FFCPXMLNode> InParentNode, const TSharedPtr<FMovieSceneExportCinematicSectionData> InCinematicSectionData, const FString& InSectionName)
+bool FFCPXMLExportVisitor::ConstructLoggingInfoNode(TSharedRef<FFCPXMLNode> InParentNode, const TSharedPtr<FMovieSceneExportAudioSectionData> InSectionData)
 {
-	if (InCinematicSectionData->CinematicShotSection == nullptr)
+	if (!ExportData->MovieSceneData.IsValid() || !InSectionData.IsValid() || InSectionData->MovieSceneSection == nullptr)
 	{
 		return false;
 	}
+	const UMovieSceneAudioSection* AudioSection = Cast<UMovieSceneAudioSection>(InSectionData->MovieSceneSection);
+	if (AudioSection == nullptr)
+	{
+		return false;
+	}
+	USoundBase *Sound = AudioSection->GetSound();
+	if (Sound == nullptr)
+	{
+		// skip logging 
+		return true;
+	}
+	USoundWave* SoundWave = Cast<USoundWave>(Sound);
+	if (SoundWave == nullptr)
+	{
+		// skip logging 
+		return true;
+	}
 
 	TSharedRef<FFCPXMLNode> LoggingInfoNode = InParentNode->CreateChildNode(TEXT("logginginfo"));
-	ConstructLoggingInfoElements(LoggingInfoNode, InCinematicSectionData->CinematicShotSection);
+	ConstructLoggingInfoElements(LoggingInfoNode, InSectionData->MovieSceneSection);
 
 	TSharedPtr<FFCPXMLNode> LogNoteNode = LoggingInfoNode->GetChildNode(TEXT("lognote"), ENodeInherit::NoInherit, ENodeReference::NoReferences);
 	if (!LogNoteNode.IsValid())
 	{
 		LogNoteNode = LoggingInfoNode->CreateChildNode(TEXT("lognote"));
 	}
-	LogNoteNode->SetContent(FFCPXMLClipNode::GetMetadataSectionName(InSectionName));
+
+	TArray <TSharedPtr<FMovieSceneExportAudioSectionData> > AudioSectionsData;
+	ExportData->FindAudioSections(SoundWave->GetPathName(), AudioSectionsData);
+
+	TArray< const UMovieSceneAudioSection*> AudioSections;
+	for (TSharedPtr<FMovieSceneExportAudioSectionData> AudioSectionData : AudioSectionsData)
+	{
+		if (AudioSectionData.IsValid() && AudioSectionData->MovieSceneSection != nullptr)
+		{
+			const UMovieSceneAudioSection* Section = Cast<UMovieSceneAudioSection>(AudioSectionData->MovieSceneSection);
+			if (Section != nullptr && Section->GetSound() != nullptr)
+			{
+				if (Section->GetSound()->GetPathName() == SoundWave->GetPathName())
+				{
+					AudioSections.Add(Section);
+				}
+			}
+		}
+	}
+
+	FString Metadata{ TEXT("") };
+	if (!CreateSoundWaveMetadata(SoundWave, AudioSections, Metadata))
+	{
+		return false;
+	}
+	LogNoteNode->SetContent(Metadata);
 
 	return true;
 }
@@ -285,7 +462,7 @@ void FFCPXMLExportVisitor::SetLoggingInfoElementValue(TSharedPtr<FFCPXMLNode> In
 }
 
 /** Creates colorinfo node. */
-bool FFCPXMLExportVisitor::ConstructColorInfoNode(TSharedRef<FFCPXMLNode> InParentNode, const TSharedPtr<FMovieSceneExportCinematicSectionData> InCinematicSectionData)
+bool FFCPXMLExportVisitor::ConstructColorInfoNode(TSharedRef<FFCPXMLNode> InParentNode, const TSharedPtr<FMovieSceneExportSectionData> InSectionData)
 {
 	TSharedPtr<FFCPXMLNode> ColorInfoNode = InParentNode->CreateChildNode(TEXT("colorinfo"));
 	ColorInfoNode->CreateChildNode(TEXT("lut"));
@@ -307,7 +484,7 @@ bool FFCPXMLExportVisitor::ConstructSequenceNode(TSharedRef<FFCPXMLNode> InParen
 	TSharedRef<FFCPXMLNode> SequenceNode = InParentNode->CreateChildNode("sequence");
 
 	// attributes
-	SequenceNode->AddAttribute(TEXT("id"), FString::Printf(TEXT("sequence%d"), ++SequenceId));
+	SequenceNode->AddAttribute(TEXT("id"), FString::Printf(TEXT("sequence-%d"), ++SequenceId));
 
 	// required elements
 	TSharedRef<FFCPXMLNode> DurationNode = SequenceNode->CreateChildNode(TEXT("duration"));
@@ -358,14 +535,14 @@ bool FFCPXMLExportVisitor::ConstructVideoNode(TSharedRef<FFCPXMLNode> InParentNo
 
 	TSharedRef<FFCPXMLNode> FormatNode = VideoNode->CreateChildNode(TEXT("format"));
 
-	if (!ConstructSampleCharacteristicsNode(FormatNode, ExportData->GetResX(), ExportData->GetResY()))
+	if (!ConstructVideoSampleCharacteristicsNode(FormatNode, ExportData->GetResX(), ExportData->GetResY()))
 	{
 		return false;
 	}
  
 	for (TSharedPtr<FMovieSceneExportCinematicTrackData> CinematicTrack : CinematicMasterTrackData->CinematicTracks)
 	{
-		if (!ConstructVideoTrackNode(VideoNode, CinematicTrack)) 
+		if (!ConstructVideoTrackNode(VideoNode, CinematicTrack, CinematicMasterTrackData)) 
 		{ 
 			return false;
 		}
@@ -383,20 +560,68 @@ bool FFCPXMLExportVisitor::ConstructAudioNode(TSharedRef<FFCPXMLNode> InParentNo
 
 	TSharedRef<FFCPXMLNode> AudioNode = InParentNode->CreateChildNode(TEXT("audio"));
 
-	//TSharedRef<FFCPXMLNode> FormatNode = InVideoNode->CreateChildNode(TEXT("format"));
+	int32 NumChannels = 1;
+	for (TSharedPtr<FMovieSceneExportAudioMasterTrackData> AudioMasterTrack : ExportData->MovieSceneData->AudioMasterTracks)
+	{
+		if (HasStereoAudioSections(AudioMasterTrack->AudioSections))
+		{
+			NumChannels = 2;
+			break;
+		}
+	}
+	AudioNode->CreateChildNode(TEXT("numOutputChannels"))->SetContent(NumChannels);
+
+	TSharedRef<FFCPXMLNode> FormatNode = AudioNode->CreateChildNode(TEXT("format"));
+
+	if (!ConstructAudioSampleCharacteristicsNode(FormatNode, ExportData->GetDefaultAudioDepth(), ExportData->GetDefaultAudioSampleRate()))
+	{
+		return false;
+	}
+
+	int32 Downmix = 0;
+	TSharedRef<FFCPXMLNode> OutputsNode = AudioNode->CreateChildNode(TEXT("outputs"));
+
+	TSharedRef<FFCPXMLNode> GroupNode = OutputsNode->CreateChildNode(TEXT("group"));
+	GroupNode->CreateChildNode(TEXT("index"))->SetContent(1);
+	GroupNode->CreateChildNode(TEXT("numchannels"))->SetContent(1);
+	GroupNode->CreateChildNode(TEXT("downmix"))->SetContent(Downmix);
+	TSharedRef<FFCPXMLNode> ChannelNode = GroupNode->CreateChildNode(TEXT("channel"));
+	ChannelNode->CreateChildNode(TEXT("index"))->SetContent(1);
 
-	for (TSharedPtr<FMovieSceneExportAudioTrackData> AudioTrack : ExportData->MovieSceneData->AudioTracks)
+	if (NumChannels == 2)
 	{
-		if (!ConstructAudioTrackNode(AudioNode, AudioTrack))
+		GroupNode = OutputsNode->CreateChildNode(TEXT("group"));
+		GroupNode->CreateChildNode(TEXT("index"))->SetContent(2);
+		GroupNode->CreateChildNode(TEXT("numchannels"))->SetContent(1);
+		GroupNode->CreateChildNode(TEXT("downmix"))->SetContent(Downmix);
+		ChannelNode = GroupNode->CreateChildNode(TEXT("channel"));
+		ChannelNode->CreateChildNode(TEXT("index"))->SetContent(2);
+	}
+
+	uint32 TrackIndex = 1;
+
+	for (TSharedPtr<FMovieSceneExportAudioMasterTrackData> AudioMasterTrack : ExportData->MovieSceneData->AudioMasterTracks)
+	{
+		if (!AudioMasterTrack.IsValid())
 		{
 			return false;
 		}
+
+		for (TSharedPtr<FMovieSceneExportAudioTrackData> AudioTrack : AudioMasterTrack->AudioTracks)
+		{
+			uint32 OutNumTracks{ 0 };
+			if (!ConstructAudioTrackNode(AudioNode, AudioTrack, AudioMasterTrack, TrackIndex, OutNumTracks))
+			{
+				return false;
+			}
+			TrackIndex += OutNumTracks;
+		}
 	}
 
 	return true;
 }
 
-bool FFCPXMLExportVisitor::ConstructVideoTrackNode(TSharedRef<FFCPXMLNode> InParentNode, const TSharedPtr<FMovieSceneExportCinematicTrackData> InCinematicTrackData)
+bool FFCPXMLExportVisitor::ConstructVideoTrackNode(TSharedRef<FFCPXMLNode> InParentNode, const TSharedPtr<FMovieSceneExportCinematicTrackData> InCinematicTrackData, const TSharedPtr<FMovieSceneExportCinematicMasterTrackData> InCinematicMasterTrackData)
 {
 	if (!ExportData->IsExportDataValid() || !InCinematicTrackData.IsValid())
 	{
@@ -413,7 +638,7 @@ bool FFCPXMLExportVisitor::ConstructVideoTrackNode(TSharedRef<FFCPXMLNode> InPar
 			continue;
 		}
 		
-		if (!ConstructVideoClipItemNode(TrackNode, CinematicSection, false))
+		if (!ConstructVideoClipItemNode(TrackNode, CinematicSection, InCinematicMasterTrackData, false))
 		{
 			return false;
 		}
@@ -428,33 +653,158 @@ bool FFCPXMLExportVisitor::ConstructVideoTrackNode(TSharedRef<FFCPXMLNode> InPar
 	return true;
 }
 
-bool FFCPXMLExportVisitor::ConstructAudioTrackNode(TSharedRef<FFCPXMLNode> InParentNode, const TSharedPtr<FMovieSceneExportAudioTrackData> InAudioTrackData)
+bool FFCPXMLExportVisitor::HasStereoAudioSections(const TArray<TSharedPtr<FMovieSceneExportAudioSectionData>>& InAudioSections) const
+{
+	for (TSharedPtr<FMovieSceneExportAudioSectionData> AudioSection : InAudioSections)
+	{
+		if (AudioSection.IsValid() && AudioSection->NumChannels == 2)
+		{
+			return true;
+		}
+	}
+	return false;
+}
+
+bool FFCPXMLExportVisitor::ConstructAudioTrackNode(TSharedRef<FFCPXMLNode> InParentNode, const TSharedPtr<FMovieSceneExportAudioTrackData> InAudioTrackData, const TSharedPtr<FMovieSceneExportAudioMasterTrackData> InAudioMasterTrackData, uint32 InTrackIndex, uint32 OutNumTracks)
 {
 	if (!ExportData->IsExportDataValid() || !InAudioTrackData.IsValid())
 	{
 		return false;
 	}
 
+	bool bTrackHasStereoClips = HasStereoAudioSections(InAudioTrackData->AudioSections);
+	int32 TrackIndex1{ 0 };
+	int32 TrackIndex2{ 0 };
+
+	if (bTrackHasStereoClips) 
+	{
+		OutNumTracks = 2;
+		TrackIndex1 = InTrackIndex;
+		TrackIndex2 = InTrackIndex + 1;
+	}
+	else
+	{
+		OutNumTracks = 1;
+		TrackIndex1 = InTrackIndex;
+		TrackIndex2 = InTrackIndex;
+	}
+
+	// Generate all clipitem names for this track so that linked clipitems can be associated
+	FString ClipItemIdName{ TEXT("") };
+	TArray<FString> ClipItem1;
+	TArray<FString> ClipItem2;
+	TArray<int32> ClipIndex1;
+	TArray<int32> ClipIndex2;
+
+	int32 Index = 0;
+	for (TSharedPtr<FMovieSceneExportAudioSectionData> AudioSection : InAudioTrackData->AudioSections)
+	{
+		if (AudioSection->NumChannels < 1)
+		{
+			ExportContext->AddMessage(EMessageSeverity::Warning,
+				FText::Format(LOCTEXT("FCPXMLAudioChannelsInvalidWarning", "FCP XML export only supports mono or stereo audio. Skipping audio section '{0}' which an invalid number of channels: '{1}'."),
+					FText::FromString(AudioSection->DisplayName),
+					FText::FromString(FString::FromInt(AudioSection->NumChannels))));
+
+			continue;
+		}
+		else if (AudioSection->NumChannels > 2)
+		{
+			ExportContext->AddMessage(EMessageSeverity::Warning,
+				FText::Format(LOCTEXT("FCPXMLAudioChannelsUnsupportedWarning", "FCP XML export only supports mono or stereo audio. Skipping audio section '{0}' which has '{1}' channels."),
+					FText::FromString(AudioSection->DisplayName),
+					FText::FromString(FString::FromInt(AudioSection->NumChannels))));
+
+			continue;
+		}
+
+		GetNextClipItemIdName(ClipItemIdName);
+		ClipItem1.Add(ClipItemIdName);
+		ClipIndex1.Add(++Index);
+	}
+
+	Index = 0;
+	for (TSharedPtr<FMovieSceneExportAudioSectionData> AudioSection : InAudioTrackData->AudioSections)
+	{
+		if (AudioSection->NumChannels == 1)
+		{
+			ClipItem2.Add(FString(TEXT("")));
+			ClipIndex2.Add(-1);
+		}
+		else if (AudioSection->NumChannels == 2)
+		{
+			GetNextClipItemIdName(ClipItemIdName);
+			ClipItem2.Add(ClipItemIdName);
+			ClipIndex2.Add(++Index);
+		}
+	}
+
+	// construct track 1
 	TSharedRef<FFCPXMLNode> TrackNode = InParentNode->CreateChildNode(TEXT("track"));
+	TrackNode->AddAttribute(TEXT("currentExplodedTrackIndex"), TEXT("0"));
+	TrackNode->AddAttribute(TEXT("totalExplodedTrackCount"), bTrackHasStereoClips ? TEXT("2") : TEXT("1"));
+	TrackNode->AddAttribute(TEXT("premiereTrackType"), bTrackHasStereoClips ? TEXT("Stereo") : TEXT("Mono"));
 
+	Index = 0;
 	for (TSharedPtr<FMovieSceneExportAudioSectionData> AudioSection : InAudioTrackData->AudioSections)
 	{
-		if (!ConstructAudioClipItemNode(TrackNode, AudioSection))
+		if (AudioSection->NumChannels < 1 || AudioSection->NumChannels > 2)
+		{
+			continue;
+		}
+
+		if (!ConstructAudioClipItemNode(TrackNode, AudioSection, InAudioMasterTrackData, 1, false, ClipItem1[Index], ClipItem2[Index], ClipIndex1[Index], ClipIndex2[Index], TrackIndex1, TrackIndex2))
 		{
 			return false;
 		}
+		Index++;
 	}
 
+
 	TSharedRef<FFCPXMLNode> EnabledNode = TrackNode->CreateChildNode(TEXT("enabled"));
 	EnabledNode->SetContent(true);
 
 	TSharedRef<FFCPXMLNode> LockedNode = TrackNode->CreateChildNode(TEXT("locked"));
 	LockedNode->SetContent(false);
 
+	// construct track 2, if stereo clipitems exist
+	if (bTrackHasStereoClips)
+	{
+		TrackNode = InParentNode->CreateChildNode(TEXT("track"));
+		TrackNode->AddAttribute(TEXT("currentExplodedTrackIndex"), TEXT("1"));
+		TrackNode->AddAttribute(TEXT("totalExplodedTrackCount"), TEXT("2"));
+		TrackNode->AddAttribute(TEXT("premiereTrackType"), TEXT("Stereo"));
+
+		Index = 0;
+		for (TSharedPtr<FMovieSceneExportAudioSectionData> AudioSection : InAudioTrackData->AudioSections)
+		{
+			if (AudioSection->NumChannels < 1 || AudioSection->NumChannels > 2)
+			{
+				continue;
+			}
+
+			if (AudioSection->NumChannels == 2)
+			{
+				if (!ConstructAudioClipItemNode(TrackNode, AudioSection, InAudioMasterTrackData, 2, false, ClipItem1[Index], ClipItem2[Index], ClipIndex1[Index], ClipIndex2[Index], TrackIndex1, TrackIndex2))
+				{
+					return false;
+				}
+			}
+			Index++;
+
+		}
+
+		EnabledNode = TrackNode->CreateChildNode(TEXT("enabled"));
+		EnabledNode->SetContent(true);
+
+		LockedNode = TrackNode->CreateChildNode(TEXT("locked"));
+		LockedNode->SetContent(false);
+	}
+
 	return true;
 }
 
-bool FFCPXMLExportVisitor::ConstructVideoClipItemNode(TSharedRef<FFCPXMLNode> InParentNode, const TSharedPtr<FMovieSceneExportCinematicSectionData> InCinematicSectionData, bool bInMasterClip)
+bool FFCPXMLExportVisitor::ConstructVideoClipItemNode(TSharedRef<FFCPXMLNode> InParentNode, const TSharedPtr<FMovieSceneExportCinematicSectionData> InCinematicSectionData, const TSharedPtr<FMovieSceneExportCinematicMasterTrackData> InCinematicMasterTrackData, bool bInMasterClip)
 {
 	if (!ExportData->IsExportDataValid() || !InCinematicSectionData.IsValid())
 	{
@@ -468,22 +818,13 @@ bool FFCPXMLExportVisitor::ConstructVideoClipItemNode(TSharedRef<FFCPXMLNode> In
 	int32 Out{ 0 };
 	int32 Start{ 0 };
 	int32 End{ 0 };
-	FString SectionName(TEXT(""));
-
-	GetSectionFrames(InCinematicSectionData, Duration, Start, End, In, Out);
-	GetSectionName(InCinematicSectionData, SectionName);
-
-	FString ClipItemIdName = TEXT("");
-	if (!GetMasterClipItemIdName(InCinematicSectionData, ClipItemIdName))
-	{
-		return false;
-	}
+	GetCinematicSectionFrames(InCinematicSectionData, Duration, Start, End, In, Out);
 
 	FString MasterClipIdName = TEXT("");
-	if (!GetMasterClipIdName(InCinematicSectionData, MasterClipIdName))
-	{
-		return false;
-	}
+	GetMasterClipIdName(InCinematicSectionData, MasterClipIdName);
+
+	FString ClipItemIdName{ TEXT("") };
+	GetNextClipItemIdName(ClipItemIdName);
 
 	// attributes
 	ClipItemNode->AddAttribute(TEXT("id"), ClipItemIdName);
@@ -491,7 +832,7 @@ bool FFCPXMLExportVisitor::ConstructVideoClipItemNode(TSharedRef<FFCPXMLNode> In
 	// elements
 	ClipItemNode->CreateChildNode(TEXT("masterclipid"))->SetContent(MasterClipIdName);
 	ClipItemNode->CreateChildNode(TEXT("ismasterclip"))->SetContent(bInMasterClip);
-	ClipItemNode->CreateChildNode(TEXT("name"))->SetContent(SectionName);
+	ClipItemNode->CreateChildNode(TEXT("name"))->SetContent(InCinematicSectionData->DisplayName);
 	ClipItemNode->CreateChildNode(TEXT("enabled"))->SetContent(true);
 	ClipItemNode->CreateChildNode(TEXT("duration"))->SetContent(Duration);
 
@@ -512,8 +853,8 @@ bool FFCPXMLExportVisitor::ConstructVideoClipItemNode(TSharedRef<FFCPXMLNode> In
 	if (bInMasterClip)
 	{
 		ClipItemNode->CreateChildNode(TEXT("anamorphic"))->SetContent(false);
-		ClipItemNode->CreateChildNode(TEXT("pixelaspectratio"))->SetContent(TEXT("square"));
-		ClipItemNode->CreateChildNode(TEXT("fielddominance"))->SetContent(TEXT("lower"));
+		ClipItemNode->CreateChildNode(TEXT("pixelaspectratio"))->SetContent(FString(TEXT("square")));
+		ClipItemNode->CreateChildNode(TEXT("fielddominance"))->SetContent(FString(TEXT("lower")));
 	}
 
 	if (!ConstructVideoFileNode(ClipItemNode, InCinematicSectionData, Duration, bInMasterClip))
@@ -524,43 +865,113 @@ bool FFCPXMLExportVisitor::ConstructVideoClipItemNode(TSharedRef<FFCPXMLNode> In
 	return true;
 }
 
-bool FFCPXMLExportVisitor::ConstructAudioClipItemNode(TSharedRef<FFCPXMLNode> InParentNode, const TSharedPtr<FMovieSceneExportAudioSectionData> InAudioSectionData)
+bool FFCPXMLExportVisitor::ConstructAudioClipItemNode(TSharedRef<FFCPXMLNode> InParentNode, const TSharedPtr<FMovieSceneExportAudioSectionData> InAudioSectionData, 
+	const TSharedPtr<FMovieSceneExportAudioMasterTrackData> InAudioMasterTrackData, int32 InChannel, bool bInMasterClip, 
+	const FString& InClipItemIdName1, const FString& InClipItemIdName2, int32 InClipIndex1, int32 InClipIndex2, int32 InTrackIndex1, int32 InTrackIndex2)
 {
 	if (!ExportData->IsExportDataValid() || !InAudioSectionData.IsValid())
 	{
 		return false;
 	}
 
-	// @todo - audio
+	bool bIsStereo = (InAudioSectionData->NumChannels == 2);
+
+	int32 Duration{ 0 };
+	int32 In{ 0 };
+	int32 Out{ 0 };
+	int32 Start{ 0 };
+	int32 End{ 0 };
+	GetAudioSectionFrames(InAudioSectionData, Duration, Start, End, In, Out);
+
+	FString MasterClipIdName;
+	GetMasterClipIdName(InAudioSectionData, MasterClipIdName);
+
+	FString ClipItemIdName = (InChannel == 1 ? InClipItemIdName1 : InClipItemIdName2);
+
+	TSharedRef<FFCPXMLNode> ClipItemNode = InParentNode->CreateChildNode(TEXT("clipitem"));
+	ClipItemNode->AddAttribute(TEXT("id"), ClipItemIdName);
+	if (!bInMasterClip)
+	{
+		ClipItemNode->AddAttribute(TEXT("premiereChannelType"), InAudioSectionData->NumChannels == 2 ? TEXT("stereo") : TEXT("mono"));
+	}
+
+	// elements
+	ClipItemNode->CreateChildNode(TEXT("masterclipid"))->SetContent(MasterClipIdName);
+	ClipItemNode->CreateChildNode(TEXT("name"))->SetContent(InAudioSectionData->DisplayName);
+
+	if (!bInMasterClip)
+	{
+		ClipItemNode->CreateChildNode(TEXT("enabled"))->SetContent(true);
+		ClipItemNode->CreateChildNode(TEXT("duration"))->SetContent(Duration);
+	}
+
+	if (!ConstructRateNode(ClipItemNode))
+	{
+		return false;
+	}
+
+	if (!bInMasterClip)
+	{
+		ClipItemNode->CreateChildNode(TEXT("start"))->SetContent(Start);
+		ClipItemNode->CreateChildNode(TEXT("end"))->SetContent(End);
+		ClipItemNode->CreateChildNode(TEXT("in"))->SetContent(In);
+		ClipItemNode->CreateChildNode(TEXT("out"))->SetContent(Out);
+	}
+
+	if (!ConstructAudioFileNode(ClipItemNode, InAudioSectionData, InChannel))
+	{
+		return false;
+	}
+
+	TSharedRef<FFCPXMLNode> SourceTrackNode = ClipItemNode->CreateChildNode(TEXT("sourcetrack"));
+	SourceTrackNode->CreateChildNode(TEXT("mediatype"))->SetContent(FString(TEXT("audio")));
+	SourceTrackNode->CreateChildNode(TEXT("trackindex"))->SetContent(InChannel);
+
+	// stereo track clipitems must be linked using the linkclipref element
+	if (bIsStereo)
+	{
+		TSharedRef<FFCPXMLNode> LinkNode = ClipItemNode->CreateChildNode(TEXT("link"));
+		LinkNode->CreateChildNode(TEXT("linkclipref"))->SetContent(InClipItemIdName1);
+		LinkNode->CreateChildNode(TEXT("mediatype"))->SetContent(FString(TEXT("audio")));
+		LinkNode->CreateChildNode(TEXT("trackindex"))->SetContent(InTrackIndex1);
+		LinkNode->CreateChildNode(TEXT("clipindex"))->SetContent(InClipIndex1);
+		LinkNode->CreateChildNode(TEXT("groupindex"))->SetContent(1);
+
+		LinkNode = ClipItemNode->CreateChildNode(TEXT("link"));
+		LinkNode->CreateChildNode(TEXT("linkclipref"))->SetContent(InClipItemIdName2);
+		LinkNode->CreateChildNode(TEXT("mediatype"))->SetContent(FString(TEXT("audio")));
+		LinkNode->CreateChildNode(TEXT("trackindex"))->SetContent(InTrackIndex2);
+		LinkNode->CreateChildNode(TEXT("clipindex"))->SetContent(InClipIndex1);
+		LinkNode->CreateChildNode(TEXT("groupindex"))->SetContent(1);
+	}
 
 	return true;
 }
 
-bool FFCPXMLExportVisitor::ConstructVideoFileNode(TSharedRef<FFCPXMLNode> InParentNode, const TSharedPtr<FMovieSceneExportCinematicSectionData> InCinematicSectionData, uint32 Duration, bool bInMasterClip)
+bool FFCPXMLExportVisitor::ConstructVideoFileNode(TSharedRef<FFCPXMLNode> InParentNode, const TSharedPtr<FMovieSceneExportCinematicSectionData> InCinematicSectionData, int32 Duration, bool bInMasterClip)
 {
 	if (!ExportData->IsExportDataValid() || !InCinematicSectionData.IsValid())
 	{
 		return false;
 	}
 
-	FString FileIdName = TEXT("");
-	if (!GetMasterClipFileIdName(InCinematicSectionData, FileIdName))
-	{
-		return false;
-	}
+	FString FileIdName{ TEXT("") };
+	bool bFileExists = false;
+	GetFileIdName(InCinematicSectionData, FileIdName, bFileExists);
 
 	// attributes
 	TSharedRef<FFCPXMLNode> FileNode = InParentNode->CreateChildNode(TEXT("file"));
 	FileNode->AddAttribute(TEXT("id"), FileIdName);
 
-	if (bInMasterClip)
+	if (!bFileExists)
 	{
-		FString FilePathName = SaveFilePath + TEXT("/") + InCinematicSectionData->ShotFilename;
+		FString FilePath = InCinematicSectionData->SourceFilePath.IsEmpty() ? SaveFilePath : InCinematicSectionData->SourceFilePath;
+		FString FilePathName = SaveFilePath + TEXT("/") + InCinematicSectionData->SourceFilename;
 		FString FilePathUrl = FString(TEXT("file://localhost/")) + FilePathName.Replace(TEXT(" "), TEXT("%20")).Replace(TEXT(":"), TEXT("%3a"));
 
 		// required elements
 		TSharedRef<FFCPXMLNode> NameNode = FileNode->CreateChildNode(TEXT("name"));
-		NameNode->SetContent(InCinematicSectionData->ShotFilename);
+		NameNode->SetContent(InCinematicSectionData->SourceFilename);
 
 		TSharedRef<FFCPXMLNode> PathUrlNode = FileNode->CreateChildNode(TEXT("pathurl"));
 		PathUrlNode->SetContent(FilePathUrl);
@@ -581,7 +992,7 @@ bool FFCPXMLExportVisitor::ConstructVideoFileNode(TSharedRef<FFCPXMLNode> InPare
 		TSharedRef<FFCPXMLNode> MediaNode = FileNode->CreateChildNode(TEXT("media"));
 		TSharedRef<FFCPXMLNode> VideoNode = MediaNode->CreateChildNode(TEXT("video"));
 
-		if (!ConstructSampleCharacteristicsNode(VideoNode, ExportData->GetResX(), ExportData->GetResY()))
+		if (!ConstructVideoSampleCharacteristicsNode(VideoNode, ExportData->GetResX(), ExportData->GetResY()))
 		{
 			return false;
 		}
@@ -590,14 +1001,84 @@ bool FFCPXMLExportVisitor::ConstructVideoFileNode(TSharedRef<FFCPXMLNode> InPare
 	return true;
 }
 
-bool FFCPXMLExportVisitor::ConstructAudioFileNode(TSharedRef<FFCPXMLNode> InParentNode, const TSharedPtr<FMovieSceneExportAudioSectionData> InAudioSectionData)
+bool FFCPXMLExportVisitor::ConstructAudioFileNode(TSharedRef<FFCPXMLNode> InParentNode, const TSharedPtr<FMovieSceneExportAudioSectionData> InAudioSectionData, int32 InChannel)
 {
-	// @todo - audio
+	FString FileIdName{ TEXT("") };
+	bool bFileExists = false;
+	GetFileIdName(InAudioSectionData, FileIdName, bFileExists);
+
+	int32 Duration{ 0 };
+	int32 In{ 0 };
+	int32 Out{ 0 };
+	int32 Start{ 0 };
+	int32 End{ 0 };
+	GetAudioSectionFrames(InAudioSectionData, Duration, Start, End, In, Out);
+
+	// attributes
+	TSharedRef<FFCPXMLNode> FileNode = InParentNode->CreateChildNode(TEXT("file"));
+	FileNode->AddAttribute(TEXT("id"), FileIdName);
+
+	// only add details if file id did not already exist
+	if (!bFileExists)
+	{
+		// FPaths
+		FString FilePathName = InAudioSectionData->SourceFilePath + TEXT("/") + InAudioSectionData->SourceFilename;
+		FString FilePathUrl = FString(TEXT("file://localhost/")) + FilePathName.Replace(TEXT(" "), TEXT("%20")).Replace(TEXT(":"), TEXT("%3a"));
+
+		// required elements
+		TSharedRef<FFCPXMLNode> NameNode = FileNode->CreateChildNode(TEXT("name"));
+		NameNode->SetContent(InAudioSectionData->SourceFilename);
+
+		TSharedRef<FFCPXMLNode> PathUrlNode = FileNode->CreateChildNode(TEXT("pathurl"));
+		PathUrlNode->SetContent(FilePathUrl);
+
+		if (!ConstructRateNode(FileNode))
+		{
+			return false;
+		}
+
+		if (!ConstructTimecodeNode(FileNode))
+		{
+			return false;
+		}
+
+		FileNode->CreateChildNode(TEXT("duration"))->SetContent(static_cast<int32>(Duration));
+
+		TSharedRef<FFCPXMLNode> MediaNode = FileNode->CreateChildNode(TEXT("media"));
+		TSharedRef<FFCPXMLNode> AudioNode = MediaNode->CreateChildNode(TEXT("audio"));
+
+		if (!ConstructAudioSampleCharacteristicsNode(AudioNode, InAudioSectionData->Depth, InAudioSectionData->SampleRate))
+		{
+			return false;
+		}
+		AudioNode->CreateChildNode(TEXT("channelcount"))->SetContent(1);
+		
+		if (InAudioSectionData->NumChannels == 2)
+		{
+			AudioNode->CreateChildNode(TEXT("layout"))->SetContent(TEXT("stereo"));
+			TSharedRef<FFCPXMLNode> AudioChannelNode = AudioNode->CreateChildNode(TEXT("audiochannel"));
+			AudioChannelNode->CreateChildNode(TEXT("sourcechannel"))->SetContent(1);
+			AudioChannelNode->CreateChildNode(TEXT("channellabel"))->SetContent(FString(TEXT("left")));
+
+			// second audio channel
+			AudioNode = MediaNode->CreateChildNode(TEXT("audio"));
+			if (!ConstructAudioSampleCharacteristicsNode(AudioNode, InAudioSectionData->Depth, InAudioSectionData->SampleRate))
+			{
+				return false;
+			}
+			AudioNode->CreateChildNode(TEXT("channelcount"))->SetContent(1);
+			AudioNode->CreateChildNode(TEXT("layout"))->SetContent(TEXT("stereo"));
+
+			AudioChannelNode = AudioNode->CreateChildNode(TEXT("audiochannel"));
+			AudioChannelNode->CreateChildNode(TEXT("sourcechannel"))->SetContent(2);
+			AudioChannelNode->CreateChildNode(TEXT("channellabel"))->SetContent(FString(TEXT("right")));
+		}
+	}
 
 	return true;
 }
 
-bool FFCPXMLExportVisitor::ConstructSampleCharacteristicsNode(TSharedRef<FFCPXMLNode> InParentNode, int InWidth, int InHeight)
+bool FFCPXMLExportVisitor::ConstructVideoSampleCharacteristicsNode(TSharedRef<FFCPXMLNode> InParentNode, int InWidth, int InHeight)
 {
 	TSharedRef<FFCPXMLNode> SampleCharacteristicsNode = InParentNode->CreateChildNode(TEXT("samplecharacteristics"));
 
@@ -609,8 +1090,17 @@ bool FFCPXMLExportVisitor::ConstructSampleCharacteristicsNode(TSharedRef<FFCPXML
 	SampleCharacteristicsNode->CreateChildNode(TEXT("width"))->SetContent(InWidth);
 	SampleCharacteristicsNode->CreateChildNode(TEXT("height"))->SetContent(InHeight);
 	SampleCharacteristicsNode->CreateChildNode(TEXT("anamorphic"))->SetContent(false);
-	SampleCharacteristicsNode->CreateChildNode(TEXT("pixelaspectratio"))->SetContent(TEXT("square"));
-	SampleCharacteristicsNode->CreateChildNode(TEXT("fielddominance"))->SetContent(TEXT("lower"));
+	SampleCharacteristicsNode->CreateChildNode(TEXT("pixelaspectratio"))->SetContent(FString(TEXT("square")));
+	SampleCharacteristicsNode->CreateChildNode(TEXT("fielddominance"))->SetContent(FString(TEXT("lower")));
+
+	return true;
+}
+
+bool FFCPXMLExportVisitor::ConstructAudioSampleCharacteristicsNode(TSharedRef<FFCPXMLNode> InParentNode, int InDepth, int InSampleRate)
+{
+	TSharedRef<FFCPXMLNode> SampleCharacteristicsNode = InParentNode->CreateChildNode(TEXT("samplecharacteristics"));
+	SampleCharacteristicsNode->CreateChildNode(TEXT("depth"))->SetContent(InDepth);
+	SampleCharacteristicsNode->CreateChildNode(TEXT("samplerate"))->SetContent(InSampleRate);
 
 	return true;
 }
@@ -637,18 +1127,15 @@ bool FFCPXMLExportVisitor::ConstructTimecodeNode(TSharedRef<FFCPXMLNode> InParen
 		return false;
 	}
 
-	TimecodeNode->CreateChildNode(TEXT("string"))->SetContent(TEXT("00:00:00:00"));
+	TimecodeNode->CreateChildNode(TEXT("string"))->SetContent(FString(TEXT("00:00:00:00")));
 	TimecodeNode->CreateChildNode(TEXT("frame"))->SetContent(0);
-	TimecodeNode->CreateChildNode(TEXT("displayformat"))->SetContent(TEXT("NDF"));
-
-	TSharedRef<FFCPXMLNode> ReelNode = TimecodeNode->CreateChildNode(TEXT("reel"));
-	ReelNode->CreateChildNode(TEXT("name"))->SetContent(TEXT(""));
 
 	return true;
 }
 
-/** Get duration, in and out frames for a given shot section */
-bool FFCPXMLExportVisitor::GetSectionFrames(const TSharedPtr<FMovieSceneExportCinematicSectionData> InCinematicSectionData, int32& OutDuration, int32& OutStartFrame, int32&OutEndFrame, int32& OutInFrame, int32& OutOutFrame)
+
+/** Get duration, in and out frames for a given video shot section */
+bool FFCPXMLExportVisitor::GetCinematicSectionFrames(const TSharedPtr<FMovieSceneExportCinematicSectionData> InCinematicSectionData, int32& OutDuration, int32& OutStartFrame, int32&OutEndFrame, int32& OutInFrame, int32& OutOutFrame)
 {
 	if (!InCinematicSectionData.IsValid() || !ExportData->MovieSceneData.IsValid())
 	{
@@ -665,65 +1152,178 @@ bool FFCPXMLExportVisitor::GetSectionFrames(const TSharedPtr<FMovieSceneExportCi
 	return true;
 }
 
-/** Get name of a given shot section */
-bool FFCPXMLExportVisitor::GetSectionName(const TSharedPtr<FMovieSceneExportCinematicSectionData> InCinematicSectionData, FString& OutSectionName)
+/** Get duration, in and out frames for a given audio shot section */
+bool FFCPXMLExportVisitor::GetAudioSectionFrames(const TSharedPtr<FMovieSceneExportAudioSectionData> InAudioSectionData, int32& OutDuration, int32& OutStartFrame, int32&OutEndFrame, int32& OutInFrame, int32& OutOutFrame)
 {
-	if (!InCinematicSectionData.IsValid())
+	if (!InAudioSectionData.IsValid() || !ExportData->MovieSceneData.IsValid())
 	{
 		return false;
 	}
-	OutSectionName = InCinematicSectionData->ShotDisplayName;
+
+	OutStartFrame = InAudioSectionData->StartFrame.Value;
+	OutEndFrame = InAudioSectionData->EndFrame.Value;
+	OutDuration = OutEndFrame - OutStartFrame;
+	OutInFrame = 0;
+	OutOutFrame = OutDuration;
+
 	return true;
 }
 
-bool FFCPXMLExportVisitor::CreateMasterClipIdName(const TSharedPtr<FMovieSceneExportCinematicSectionData> InCinematicSectionData, uint32 InId, FString& OutName)
+bool FFCPXMLExportVisitor::HasMasterClipIdName(const TSharedPtr<FMovieSceneExportSectionData> InSection, FString& OutName, bool& bOutMasterClipExists)
 {
-	if (!InCinematicSectionData.IsValid())
+	if (!InSection.IsValid())
 	{
 		return false;
 	}
-	FString Name(TEXT(""));
-	GetSectionName(InCinematicSectionData, Name);
-	MasterClipIdMap.Add(Name, InId);
-	OutName = FString::Printf(TEXT("masterclip%d_"), InId) + Name;
+
+	FString Key;
+	if (!ComposeFileKey(InSection, Key))
+	{
+		return false;
+	}
+
+	if (MasterClipIdMap.Num() > 0)
+	{
+		uint32 *FoundId = MasterClipIdMap.Find(Key);
+		if (FoundId != nullptr)
+		{
+			OutName = FString::Printf(TEXT("masterclip-%d"), *FoundId);
+			bOutMasterClipExists = true;
+			return true;
+		}
+	}
+
+	bOutMasterClipExists = false;
 	return true;
 }
 
-bool FFCPXMLExportVisitor::GetIdName(const TSharedPtr<FMovieSceneExportCinematicSectionData> InCinematicSectionData, FString InPrefix, FString &OutName)
+bool FFCPXMLExportVisitor::GetMasterClipIdName(const TSharedPtr<FMovieSceneExportSectionData> InSection, FString& OutName)
 {
-	if (!InCinematicSectionData.IsValid())
+	if (!InSection.IsValid())
+	{
+		return false;
+	}
+
+	bool bMasterClipExists = false;
+	if (!HasMasterClipIdName(InSection, OutName, bMasterClipExists))
+	{
+		return false;
+	}
+
+	FString Key;
+	if (!ComposeFileKey(InSection, Key))
+	{
+		return false;
+	}
+
+	if (!bMasterClipExists)
+	{
+		++MasterClipId;
+		MasterClipIdMap.Add(Key, MasterClipId);
+		OutName = FString::Printf(TEXT("masterclip-%d"), MasterClipId);
+	}
+	return true;
+}
+
+bool FFCPXMLExportVisitor::GetFileIdName(const TSharedPtr<FMovieSceneExportSectionData> InSection, FString& OutFileIdName, bool& OutFileExists)
+{
+	if (!InSection.IsValid())
 	{
 		return false;
 	}
-	FString Name(TEXT(""));
-	GetSectionName(InCinematicSectionData, Name);
-	uint32* Id = MasterClipIdMap.Find(Name);
-	if (Id == nullptr)
+
+	FString Key;
+	if (!ComposeFileKey(InSection, Key))
 	{
 		return false;
 	}
-	OutName = InPrefix + FString::Printf(TEXT("%d_"), *Id) + Name;
+
+	if (FileIdMap.Num() > 0)
+	{
+		uint32 *FoundFileId = FileIdMap.Find(Key);
+		if (FoundFileId != nullptr)
+		{
+			OutFileIdName = FString::Printf(TEXT("file-%d"), *FoundFileId);
+			OutFileExists = true;
+			return true;
+		}
+	}
+
+	++FileId;
+	FileIdMap.Add(Key, FileId);
+	OutFileIdName = FString::Printf(TEXT("file-%d"), FileId);
+	OutFileExists = false;
 	return true;
 }
 
-bool FFCPXMLExportVisitor::GetMasterClipIdName(const TSharedPtr<FMovieSceneExportCinematicSectionData> InCinematicSectionData, FString& OutName)
+void FFCPXMLExportVisitor::GetNextClipItemIdName(FString& OutName)
 {
-	return GetIdName(InCinematicSectionData, FString(TEXT("masterclip")), OutName);
+	++ClipItemId;
+	OutName = FString::Printf(TEXT("clipitem-%d"), ClipItemId);
 }
 
-bool FFCPXMLExportVisitor::GetMasterClipItemIdName(const TSharedPtr<FMovieSceneExportCinematicSectionData> InCinematicSectionData, FString& OutName)
+/** Compose a unique key string for audio sections based on channel */
+bool FFCPXMLExportVisitor::ComposeFileKey(const TSharedPtr<FMovieSceneExportSectionData> InSection, FString& OutName)
 {
-	return GetIdName(InCinematicSectionData, FString(TEXT("clipitem")), OutName);
+	if (!InSection.IsValid())
+	{
+		return false;
+	}
+
+	OutName = InSection->SourceFilePath + InSection->SourceFilename;
+
+	return true;
 }
 
-bool FFCPXMLExportVisitor::GetMasterClipFileIdName(const TSharedPtr<FMovieSceneExportCinematicSectionData> InCinematicSectionData, FString& OutName)
+bool FFCPXMLExportVisitor::CreateCinematicSectionMetadata(const UMovieSceneCinematicShotSection* InSection, FString& OutMetadata) const
 {
-	return GetIdName(InCinematicSectionData, FString(TEXT("file")), OutName);
+	if (InSection == nullptr)
+	{
+		return false;
+	}
+	OutMetadata = TEXT("[UE4ShotSection=") + InSection->GetPathName() + TEXT("]");
+	return true;
 }
 
-FString FFCPXMLExportVisitor::GetFilePathName(const FString& InSectionName) const
+/** Get metadata section name from sequencer shot name - format is "[UE4SoundWave=soundwaveobjectname][UE4SoundSectionTopLevel=toplevelobjectname][UE4SoundSection=sectionobjectname]", whitespace ok. */
+bool FFCPXMLExportVisitor::CreateSoundWaveMetadata(const USoundWave* InSoundWave, const TArray<const UMovieSceneAudioSection*> InAudioSections, FString& OutMetadata) const
 {
-	return (SaveFilePath + TEXT("/") + InSectionName);
+	if (InSoundWave == nullptr)
+	{
+		return false;
+	}
+
+	TArray<FString> SectionsAdded;
+	bool bTopLevelAdded = false;
+	OutMetadata = TEXT("[UE4SoundWave=") + InSoundWave->GetPathName() + TEXT("]");
+	for (const UMovieSceneAudioSection* AudioSection : InAudioSections)
+	{
+		if (!bTopLevelAdded)
+		{
+			OutMetadata += TEXT("[UE4AudioSectionTopLevel=") + FFCPXMLExportVisitor::GetAudioSectionTopLevelName(AudioSection) + TEXT("]");
+			bTopLevelAdded = true;
+		}
+
+		// skip duplicate section names
+		FString SectionName = FFCPXMLExportVisitor::GetAudioSectionName(AudioSection);
+		if (SectionsAdded.Num() == 0 || !SectionsAdded.Contains(SectionName))
+		{
+			OutMetadata += TEXT("[UE4AudioSection=") + SectionName + TEXT("]");
+			SectionsAdded.Add(SectionName);
+		}
+	}
+	return true;
+}
+
+FString FFCPXMLExportVisitor::GetAudioSectionTopLevelName(const UMovieSceneAudioSection* InAudioSection)
+{
+	return InAudioSection->GetOutermost()->GetName();
 }
 
+FString FFCPXMLExportVisitor::GetAudioSectionName(const UMovieSceneSection* InAudioSection)
+{
+	return InAudioSection->GetFullGroupName(false);
+}
+
+
 #undef LOCTEXT_NAMESPACE
diff --git a/Engine/Source/Editor/MovieSceneTools/Private/FCPXML/FCPXMLExport.h b/Engine/Source/Editor/MovieSceneTools/Private/FCPXML/FCPXMLExport.h
--- a/Engine/Source/Editor/MovieSceneTools/Private/FCPXML/FCPXMLExport.h
+++ b/Engine/Source/Editor/MovieSceneTools/Private/FCPXML/FCPXMLExport.h
@@ -53,66 +53,90 @@ public:
 	/** Creates master clip nodes. */
 	bool ConstructMasterClipNodes(TSharedRef<FFCPXMLNode> InParentNode);
 	/** Creates master clip node. */
-	bool ConstructMasterClipNode(TSharedRef<FFCPXMLNode> InParentNode, const TSharedPtr<FMovieSceneExportCinematicSectionData> InCinematicSectionData, int32 InMasterClipId);
+	bool ConstructMasterClipNode(TSharedRef<FFCPXMLNode> InParentNode, const TSharedPtr<FMovieSceneExportCinematicSectionData> InCinematicSectionData, const TSharedPtr<FMovieSceneExportCinematicMasterTrackData> InCinematicMasterTrackData);
+	/** Creates master clip node. */
+	bool ConstructMasterClipNode(TSharedRef<FFCPXMLNode> InParentNode, const TSharedPtr<FMovieSceneExportAudioSectionData> InAudioSectionData, const TSharedPtr<FMovieSceneExportAudioMasterTrackData> InCinematicMasterTrackData);
 	/** Creates colorinfo node. */
-	bool ConstructColorInfoNode(TSharedRef<FFCPXMLNode> InParentNode, const TSharedPtr<FMovieSceneExportCinematicSectionData> InCinematicSectionData);
+	bool ConstructColorInfoNode(TSharedRef<FFCPXMLNode> InParentNode, const TSharedPtr<FMovieSceneExportSectionData> InSectionData);
+	/** Creates logginginfo node. */
+	bool ConstructLoggingInfoNode(TSharedRef<FFCPXMLNode> InParentNode, const TSharedPtr<FMovieSceneExportCinematicSectionData> InSectionData);
 	/** Creates logginginfo node. */
-	bool ConstructSectionLoggingInfoNode(TSharedRef<FFCPXMLNode> InParentNode, const TSharedPtr<FMovieSceneExportCinematicSectionData> InCinematicSectionData, const FString& InSectionName);
+	bool ConstructLoggingInfoNode(TSharedRef<FFCPXMLNode> InParentNode, const TSharedPtr<FMovieSceneExportAudioSectionData> InSectionData);
 	/** Creates sequence node. */
 	bool ConstructSequenceNode(TSharedRef<FFCPXMLNode> InParentNode);
 	/** Creates video node. */
 	bool ConstructVideoNode(TSharedRef<FFCPXMLNode> InParentNode);
 	/** Creates audio node. */
 	bool ConstructAudioNode(TSharedRef<FFCPXMLNode> InParentNode);
-	/** Creates cinematic track node. */
-	bool ConstructVideoTrackNode(TSharedRef<FFCPXMLNode> InParentNode, const TSharedPtr<FMovieSceneExportCinematicTrackData> InCinematicTrackData);
+	/** Creates video track node. */
+	bool ConstructVideoTrackNode(TSharedRef<FFCPXMLNode> InParentNode, const TSharedPtr<FMovieSceneExportCinematicTrackData> InCinematicTrackData, const TSharedPtr<FMovieSceneExportCinematicMasterTrackData> InCinematicMasterTrackData);
 	/** Creates audio track node. */
-	bool ConstructAudioTrackNode(TSharedRef<FFCPXMLNode> InParentNode, const TSharedPtr<FMovieSceneExportAudioTrackData> InAudioTrackData);
+	bool ConstructAudioTrackNode(TSharedRef<FFCPXMLNode> InParentNode, const TSharedPtr<FMovieSceneExportAudioTrackData> InAudioTrackData, const TSharedPtr<FMovieSceneExportAudioMasterTrackData> InAudioMasterTrackData, uint32 InTrackIndex, uint32 OutNumTracks);
 	/** Creates video clip item node. */
-	bool ConstructVideoClipItemNode(TSharedRef<FFCPXMLNode> InParentNode, const TSharedPtr<FMovieSceneExportCinematicSectionData> InCinematicSectionData, bool bInMasterClip);
+	bool ConstructVideoClipItemNode(TSharedRef<FFCPXMLNode> InParentNode, const TSharedPtr<FMovieSceneExportCinematicSectionData> InCinematicSectionData, const TSharedPtr<FMovieSceneExportCinematicMasterTrackData> InCinematicMasterTrackData, bool bInMasterClip);
 	/** Creates audio clip item node. */
-	bool ConstructAudioClipItemNode(TSharedRef<FFCPXMLNode> InParentNode, const TSharedPtr<FMovieSceneExportAudioSectionData> InAudioSectionData);
+	bool ConstructAudioClipItemNode(TSharedRef<FFCPXMLNode> InParentNode, const TSharedPtr<FMovieSceneExportAudioSectionData> InAudioSectionData, 
+		const TSharedPtr<FMovieSceneExportAudioMasterTrackData> InAudioMasterTrackData, int32 InChannel, bool bInMasterClip, 
+		const FString& InClipItemIdName1, const FString& InClipItemIdName2, int32 InClipIndex1, int32 InClipIndex2, int32 InTrackIndex1, int32 InTrackIndex2);
 	/** Creates video file node. */
-	bool ConstructVideoFileNode(TSharedRef<FFCPXMLNode> InParentNode, const TSharedPtr<FMovieSceneExportCinematicSectionData> InCinematicSectionData, uint32 Duration, bool bInMasterClip);
+	bool ConstructVideoFileNode(TSharedRef<FFCPXMLNode> InParentNode, const TSharedPtr<FMovieSceneExportCinematicSectionData> InCinematicSectionData, int32 InDuration, bool bInMasterClip);
 	/** Creates audio file node. */
-	bool ConstructAudioFileNode(TSharedRef<FFCPXMLNode> InParentNode, const TSharedPtr<FMovieSceneExportAudioSectionData> InAudioSectionData);
-	/** Creates sample characteristics node */
-	bool ConstructSampleCharacteristicsNode(TSharedRef<FFCPXMLNode> InParentNode, int InWidth, int InHeight);
+	bool ConstructAudioFileNode(TSharedRef<FFCPXMLNode> InParentNode, const TSharedPtr<FMovieSceneExportAudioSectionData> InAudioSectionData, int32 InChannel);
+	/** Creates video sample characteristics node */
+	bool ConstructVideoSampleCharacteristicsNode(TSharedRef<FFCPXMLNode> InParentNode, int InWidth, int InHeight);
+	/** Creates audio video sample characteristics node */
+	bool ConstructAudioSampleCharacteristicsNode(TSharedRef<FFCPXMLNode> InParentNode, int InDepth, int InSampleRate);
 	/** Creates rate node. */
 	bool ConstructRateNode(TSharedRef<FFCPXMLNode> InParentNode);
 	/** Creates timecode node. */
 	bool ConstructTimecodeNode(TSharedRef<FFCPXMLNode> InParentNode);
 	/** Creates logginginfo elements. */
-	void ConstructLoggingInfoElements(TSharedRef<FFCPXMLNode> InLoggingInfoNode, const UObject* InObject) ;
+	void ConstructLoggingInfoElements(TSharedRef<FFCPXMLNode> InLoggingInfoNode, const UObject* InObject = nullptr);
 	/** Set logginginfo element value. */
-	void SetLoggingInfoElementValue(TSharedPtr<FFCPXMLNode> InNode, const UObject* InObject, const FString& InElement) ;
-	/** Get duration, start, end, in, out frames for a given shot section */
-	bool GetSectionFrames(const TSharedPtr<FMovieSceneExportCinematicSectionData> InCinematicSectionData, int32& OutDuration, int32& OutStartFrame, int32& OutEndFrame, int32& OutInFrame, int32& OutOutFrame);
-	/** Get start and end frames for a given shot section */
-	//bool GetSectionStartEnd(const TSharedPtr<FMovieSceneExportCinematicSectionData> InCinematicSectionData, int32& OutStartFrame, int32& OutEndFrame);
-	/** Get name of a given shot section */
-	bool GetSectionName(const TSharedPtr<FMovieSceneExportCinematicSectionData> InCinematicSectionData, FString& OutSectionName);
+	void SetLoggingInfoElementValue(TSharedPtr<FFCPXMLNode> InNode, const UObject* InObject, const FString& InElement);
+	/** Get cinematic duration, start, end, in, out frames for a given cinematic shot section */
+	bool GetCinematicSectionFrames(const TSharedPtr<FMovieSceneExportCinematicSectionData> InCinematicSectionData, int32& OutDuration, int32& OutStartFrame, int32& OutEndFrame, int32& OutInFrame, int32& OutOutFrame);
+	/** Get audio duration, start, end, in, out frames for a given cinematic shot section */
+	bool GetAudioSectionFrames(const TSharedPtr<FMovieSceneExportAudioSectionData> InAudioSectionData, int32& OutDuration, int32& OutStartFrame, int32& OutEndFrame, int32& OutInFrame, int32& OutOutFrame);
 
 private:
 
-	/** Construct a master clip id name based on input name and id */
-	bool CreateMasterClipIdName(const TSharedPtr<FMovieSceneExportCinematicSectionData> InCinematicSectionData, uint32 InId, FString& OutName);
-	/** Get id name for given section and prefix */
-	bool GetIdName(const TSharedPtr<FMovieSceneExportCinematicSectionData> InCinematicSectionData, FString InPrefix, FString &OutName);
-	/** Get master clip id name for given section */
-	bool GetMasterClipIdName(const TSharedPtr<FMovieSceneExportCinematicSectionData> InCinematicSectionData, FString& OutName);
-	/** Get master clip's clipitem id name for given section */
-	bool GetMasterClipItemIdName(const TSharedPtr<FMovieSceneExportCinematicSectionData> InCinematicSectionData, FString& OutName);
-	/** Get master clip's file id name based on input name and id */
-	bool GetMasterClipFileIdName(const TSharedPtr<FMovieSceneExportCinematicSectionData> InCinematicSectionData, FString& OutName);
-	/** Return filename with path */
-	FString GetFilePathName(const FString& InSectionName) const;
+	/** Has master clip id for given section */
+	bool HasMasterClipIdName(const TSharedPtr<FMovieSceneExportSectionData> InSection, FString& OutName, bool& bOutMasterClipExists);
+	/** Get master clip id name for section */
+	bool GetMasterClipIdName(const TSharedPtr<FMovieSceneExportSectionData> InSection, FString& OutName);
+	/** Get file id name for section, adds to file map if a new id name is created */
+	bool GetFileIdName(const TSharedPtr<FMovieSceneExportSectionData> InSection, FString& OutName, bool& bFileExists);
+	/** Get next clip item name */
+	void GetNextClipItemIdName(FString& OutName);
+
+	/** Compose key for section file */
+	bool ComposeFileKey(const TSharedPtr<FMovieSceneExportSectionData> InSection, FString& OutName);
+
+	/** Returns true if audio track data contains sections with 2 channels */
+	bool HasStereoAudioSections(const  TArray<TSharedPtr<FMovieSceneExportAudioSectionData>>& InAudioTrackData) const;
+
+	/** Get metadata section name from sequencer shot name - format is "[UE4Section=sectionobjectname]", whitespace ok. */
+	bool CreateCinematicSectionMetadata(const UMovieSceneCinematicShotSection* InSection, FString& OutMetadata) const;
+	/** Get metadata section name from sequencer shot name - format is "[UE4SoundWave=soundwaveobjectname]", whitespace ok. */
+	bool CreateSoundWaveMetadata(const USoundWave* InSoundWave, const TArray<const UMovieSceneAudioSection*> InAudioSections, FString& OutMetadata) const;
+	/** Get id for audio top level section */
+	static FString GetAudioSectionTopLevelName(const UMovieSceneAudioSection* InAudioSection);
+	/** Get audio section group name */
+	static FString GetAudioSectionName(const UMovieSceneSection* InAudioSection);
 
 private:
 
 	TSharedRef<FMovieSceneExportData> ExportData;
 	TSharedRef<FMovieSceneTranslatorContext> ExportContext;
 	FString SaveFilePath;
-	TMap<FString, uint32> MasterClipIdMap;
 	uint32 SequenceId;
+	uint32 MasterClipId;
+	uint32 ClipItemId;
+	uint32 FileId;
+
+	/** Map section's unique key string to the id used for masterclip element names */
+	TMap<FString, uint32> MasterClipIdMap;
+	/** Map section's source file name to its file element name */
+	TMap<FString, uint32> FileIdMap;
 };
diff --git a/Engine/Source/Editor/MovieSceneTools/Private/FCPXML/FCPXMLImport.cpp b/Engine/Source/Editor/MovieSceneTools/Private/FCPXML/FCPXMLImport.cpp
--- a/Engine/Source/Editor/MovieSceneTools/Private/FCPXML/FCPXMLImport.cpp
+++ b/Engine/Source/Editor/MovieSceneTools/Private/FCPXML/FCPXMLImport.cpp
@@ -9,27 +9,71 @@
 #include "Tracks/MovieSceneCinematicShotTrack.h"
 #include "AssetRegistryModule.h"
 #include "UObject/MetaData.h"
+#include "Sound/SoundWave.h"
+
+/**
+ METADATA NOTES:
+
+ Cinematic sections have a 1:1 relationships with their source .avi files. Audio sections, by contrast, have a many:1 relationship with
+ their source .wav files. Sections are represented by clipitem nodes in FCP XML while their source files are represented by clip nodes. 
+
+ Although the LoggingInfo node which stores metadata can be attached to any FCP XML node, Premiere only recognizes them for clip and 
+ file nodes. There is no way to associate metadata at the section (clipitem node) or track (track node) level that Premiere will roundtrip. 
+ Confusingly, Premiere exports the clip LoggingInfo on the clipitem. But it is always the same on all clipitems associated with a 
+ given clip.
+
+ For each audio asset, the clip metadata includes all the section path names associated with that audio asset. These are then used
+ on import by associated a clip item with the next path name from metadata and then marking that path name as used.
+
+ TRACK NOTES:
+
+ Because it is not possible to roundtrip track node metadata through Premiere, the importer creates a one-to-one ordered correspondence 
+ between Sequencer tracks and incoming tracks from FCP XML. 
+
+ In FCP XML, each track with stereo clipitems is encoded as 2 separate but linked tracks. But in Premiere, this only appears as a single
+ track. The importer skips reading the second track, representing the second channel, since the relevant duration, start, end 
+ times will be the same as on the clipitems in the first channel track.
+
+ CURRENT LIMITATIONS:
+
+ - Rendered movie files must be named '{shot}.avi
+ - Cinematic and audio sections can be updated/moved/add by import but never removed.
+ - Sound cues are not supported
+ - Nested sequences are not supported.
+
+ */
 
 #define LOCTEXT_NAMESPACE "FCPXMLImporter"
-// DEFINE_LOG_CATEGORY_STATIC(FCPXMLImporter, Log, All);
 
 #define INHERIT true
 #define NO_INHERIT false
 
 FFCPXMLImportVisitor::FFCPXMLImportVisitor(TSharedRef<FMovieSceneImportData> InImportData, TSharedRef<FMovieSceneTranslatorContext> InImportContext) :
 	FFCPXMLNodeVisitor()
-	, ImportData(InImportData)
-	, ImportContext(InImportContext)
-	, bInSequenceNode(false)
-	, bInVideoNode(false)
-	, bInAudioNode(false)
-	, bInVideoTrackNode(false)
-	, bInAudioTrackNode(false)
-	, CurrVideoTrackRowIndex(0)
-	, CurrAudioTrackIndex(0)
+, ImportData(InImportData)
+, ImportContext(InImportContext)
+, bInSequenceNode(false)
+, bInVideoNode(false)
+, bInAudioNode(false)
+, bInVideoTrackNode(false)
+, bInAudioTrackNode(false)
+, CurrVideoTrackRowIndex(0)
+, bCurrImportAudioTrackIsStereoChannel(false)
+, CurrAudioTrackListIndex(0)
+, CurrAudioMasterTrack(nullptr)
+, CurrAudioTrackRowIndex(0)
 {
+	ConstructAudioTrackList();
+	if (AudioTrackList.Num() > 0)
+	{
+		if (AudioTrackList[0].IsValid())
+		{
+			CurrAudioMasterTrack = AudioTrackList[0]->AudioTrackData;
+			CurrAudioTrackRowIndex = AudioTrackList[0]->RowIndex;
+		}
+	}
 }
-	
+
 FFCPXMLImportVisitor::~FFCPXMLImportVisitor() {}
 
 bool FFCPXMLImportVisitor::VisitNode(TSharedRef<FFCPXMLBasicNode> InBasicNode)
@@ -43,13 +87,13 @@ bool FFCPXMLImportVisitor::VisitNode(TSharedRef<FFCPXMLXmemlNode> InXmemlNode)
 }
 
 bool FFCPXMLImportVisitor::VisitNode(TSharedRef<FFCPXMLSequenceNode> InSequenceNode)
-{ 
+{
 	bool bPrev = bInSequenceNode;
 	bInSequenceNode = true;
 
 	// Sequences can be referenced so flag to visit reference node children
 	bool bSuccess = InSequenceNode->VisitChildren(*this, true);
-	
+
 	bInSequenceNode = false;
 
 	return bSuccess;
@@ -57,19 +101,13 @@ bool FFCPXMLImportVisitor::VisitNode(TSharedRef<FFCPXMLSequenceNode> InSequenceN
 
 bool FFCPXMLImportVisitor::VisitNode(TSharedRef<FFCPXMLVideoNode> InVideoNode)
 {
-	/*
-	if (!ImportData->GetCinematicMasterTrackData(true).IsValid())
-	{
-		return false;
-	}
-	*/
 	bool bPrev = bInVideoNode;
 	bInVideoNode = true;
 	int32 PrevVideoTrackRowIndex = CurrVideoTrackRowIndex;
 	CurrVideoTrackRowIndex = 0;
 
 	bool bSuccess = InVideoNode->VisitChildren(*this);
-	
+
 	bInVideoNode = bPrev;
 	CurrVideoTrackRowIndex = PrevVideoTrackRowIndex;
 
@@ -78,9 +116,14 @@ bool FFCPXMLImportVisitor::VisitNode(TSharedRef<FFCPXMLVideoNode> InVideoNode)
 
 bool FFCPXMLImportVisitor::VisitNode(TSharedRef<FFCPXMLAudioNode> InAudioNode)
 {
-	// @todo
-	CurrAudioTrackIndex = 0;
-	return InAudioNode->VisitChildren(*this);
+	bool bPrev = bInAudioNode;
+	bInAudioNode = true;
+
+	bool bSuccess = InAudioNode->VisitChildren(*this);
+
+	bInAudioNode = bPrev;
+
+	return bSuccess;
 }
 
 bool FFCPXMLImportVisitor::VisitNode(TSharedRef<FFCPXMLTrackNode> InTrackNode)
@@ -97,8 +140,9 @@ bool FFCPXMLImportVisitor::VisitNode(TSharedRef<FFCPXMLTrackNode> InTrackNode)
 	{
 		bPrevInTrackNode = bInAudioTrackNode;
 		bInAudioTrackNode = true;
+		bCurrImportAudioTrackIsStereoChannel = false;
 	}
-	
+
 	bool bSuccess = InTrackNode->VisitChildren(*this);
 
 	if (bInSequenceNode && bInVideoNode)
@@ -109,7 +153,23 @@ bool FFCPXMLImportVisitor::VisitNode(TSharedRef<FFCPXMLTrackNode> InTrackNode)
 	else if (bInSequenceNode && bInAudioNode)
 	{
 		bInAudioTrackNode = bPrevInTrackNode;
-		CurrAudioTrackIndex++;
+		if (!bCurrImportAudioTrackIsStereoChannel)
+		{
+			CurrAudioTrackListIndex++;
+			if (CurrAudioTrackListIndex < AudioTrackList.Num())
+			{
+				if (!AudioTrackList[CurrAudioTrackListIndex].IsValid())
+				{
+					return false;
+				}
+				CurrAudioMasterTrack = AudioTrackList[CurrAudioTrackListIndex]->AudioTrackData;
+				CurrAudioTrackRowIndex = AudioTrackList[CurrAudioTrackListIndex]->RowIndex;
+			}
+			else
+			{
+				CurrAudioTrackRowIndex++;
+			}
+		}
 	}
 
 	return bSuccess;
@@ -120,29 +180,31 @@ bool FFCPXMLImportVisitor::VisitNode(TSharedRef<FFCPXMLClipNode> InClipNode)
 	// masterclip
 	bool bIsMasterClip = false;
 	FString MasterClipName(TEXT(""));
-	FString MetadataSectionName(TEXT(""));
+	FString LogNoteMetadata(TEXT(""));
+
+	InClipNode->GetChildValue<bool>("ismasterclip", bIsMasterClip, ENodeInherit::NoInherit);
+
+	bool bHasMasterClipName = InClipNode->GetChildValue<FString>("masterclipid", MasterClipName, ENodeInherit::NoInherit);
 
-	if (InClipNode->GetChildValue<bool>("ismasterclip", bIsMasterClip, ENodeInherit::NoInherit) && bIsMasterClip)
+	TSharedPtr<FFCPXMLNode> LoggingInfoNode = InClipNode->GetChildNode("logginginfo", ENodeInherit::NoInherit, ENodeReference::NoReferences);
+	if (LoggingInfoNode.IsValid())
 	{
-		bool bHasMasterClipName = InClipNode->GetChildValue<FString>("masterclipid", MasterClipName, ENodeInherit::NoInherit);
+		bool bHasLoggingShotTrack = LoggingInfoNode->GetChildValue<FString>("lognote", LogNoteMetadata, ENodeInherit::NoInherit, ENodeReference::NoReferences);
 
-		TSharedPtr<FFCPXMLNode> LoggingInfoNode = InClipNode->GetChildNode("logginginfo", ENodeInherit::NoInherit, ENodeReference::NoReferences);
-		if (LoggingInfoNode.IsValid())
+		if (!LogNoteMetadata.IsEmpty() && bHasMasterClipName && !MasterClipName.IsEmpty())
 		{
-			bool bHasLoggingShotTrack = LoggingInfoNode->GetChildValue<FString>("lognote", MetadataSectionName, ENodeInherit::NoInherit, ENodeReference::NoReferences);
+			FString SectionName{ TEXT("") };
+			TSharedPtr<FFCPXMLImportAudioMetadata> AudioMetadata;
 
-			if (!MetadataSectionName.IsEmpty())
+			if (GetCinematicSectionPathNameFromMetadata(LogNoteMetadata, SectionName))
 			{
-				if (bHasMasterClipName)
-				{
-					FString CinematicSectionName;
-					if (FFCPXMLClipNode::GetSequencerSectionName(MetadataSectionName, CinematicSectionName))
-					{
-						AddMasterClipSectionName(MasterClipName, CinematicSectionName);
-						AddSectionMetaData(CinematicSectionName, LoggingInfoNode);
-					}
-				}
+				AddMasterClipCinematicSectionPathName(MasterClipName, SectionName);
 			}
+			else if (GetAudioFromMetadata(LogNoteMetadata, AudioMetadata))
+			{
+				AddMasterClipAudioMetadata(MasterClipName, AudioMetadata);
+			}
+			AddMasterClipLoggingNode(MasterClipName, LoggingInfoNode);
 		}
 	}
 
@@ -152,154 +214,493 @@ bool FFCPXMLImportVisitor::VisitNode(TSharedRef<FFCPXMLClipNode> InClipNode)
 
 bool FFCPXMLImportVisitor::VisitNode(TSharedRef<FFCPXMLClipItemNode> InClipItemNode)
 {
-	if (bInSequenceNode && bInVideoTrackNode)
+	if (bInSequenceNode)
 	{
-		// Sequencer section name 
-		FString CinematicSectionName(TEXT(""));
-		bool bHasCinematicSectionName = false;
-
-		// masterclip
-		FString MasterClipIdName(TEXT(""));
-		if (InClipItemNode->GetChildValue<FString>("masterclipid", MasterClipIdName, ENodeInherit::NoInherit))
+		if (bInVideoTrackNode && !VisitVideoClipItemNode(InClipItemNode))
 		{
-			bHasCinematicSectionName = GetMasterClipSectionName(MasterClipIdName, CinematicSectionName);
+			return false;
 		}
-
-		// Frame rate 
-		int32 FrameRateTimebase = 0;
-		if (!InClipItemNode->GetChildSubValue<int32>("rate", "timebase", FrameRateTimebase, ENodeInherit::CheckInherit))
+		else if (bInAudioTrackNode && !VisitAudioClipItemNode(InClipItemNode))
 		{
 			return false;
 		}
+	}
 
-		bool bFrameRateIsNTSC = false;
-		if (!InClipItemNode->GetChildSubValue<bool>("rate", "ntsc", bFrameRateIsNTSC, ENodeInherit::CheckInherit))
-		{
-			bFrameRateIsNTSC = false;
-		}
+	// Clip items can be referenced so flag to visit reference node children
+	return InClipItemNode->VisitChildren(*this, true);
+}
+
+bool FFCPXMLImportVisitor::VisitVideoClipItemNode(TSharedRef<FFCPXMLClipItemNode> InClipItemNode)
+{
+	if (!bInSequenceNode || !bInVideoTrackNode)
+	{
+		return false;
+	}
+
+	FString ClipItemName{ TEXT("") };
+	FString ClipItemId{ TEXT("") };
+	FString MasterClipId{ TEXT("") };
+	FString LogNote{ TEXT("") };
+	FString Filename{ TEXT("") };
+	TSharedPtr<FFCPXMLNode> LoggingInfoNode = nullptr;
+	FFrameRate FrameRate;
+	FFrameNumber StartOffset;
+	FFrameNumber Start;
+	FFrameNumber End;
+
+	if (!GetClipItemNodeData(InClipItemNode, ClipItemName, ClipItemId, MasterClipId, LoggingInfoNode, LogNote, Filename, FrameRate, StartOffset,Start, End))
+	{
+		return false;
+	}
+
+	// check for metadata
+	if (!LoggingInfoNode.IsValid() && !MasterClipId.IsEmpty())
+	{
+		GetMasterClipLoggingNode(MasterClipId, LoggingInfoNode);
+	}
+
+	FString SectionPathName = GetCinematicSectionPathName(LogNote, MasterClipId);
+
+	// Find actual section
+	TSharedPtr<FMovieSceneImportCinematicSectionData> SectionData = nullptr;
+	if (!SectionPathName.IsEmpty())
+	{
+		SectionData = ImportData->FindCinematicSection(SectionPathName);
+	}
 
-		// Shot name
-		FString ShotName;
-		if (!InClipItemNode->GetChildValue("name", ShotName, ENodeInherit::NoInherit) || ShotName.IsEmpty())
+	if (SectionData.IsValid())
+	{
+		// Update existing cinematic section
+		if (!ImportData->SetCinematicSection(SectionData, CurrVideoTrackRowIndex, FrameRate, Start, End, StartOffset))
 		{
 			return false;
 		}
-
-		if (!bHasCinematicSectionName)
+	}
+	else
+	{
+		// Add new cinematic section
+		SectionData = ImportData->CreateCinematicSection(ClipItemName, CurrVideoTrackRowIndex, FrameRate, Start, End, StartOffset);
+		if (!SectionData.IsValid())
 		{
-			CinematicSectionName = ShotName;
+			return false;
 		}
+	}
+
+	// Import metadata to cinematic section
+	if (LoggingInfoNode.IsValid())
+	{
+		ImportSectionMetaData(LoggingInfoNode, SectionData->CinematicSection);
+	}
+	
+	return true;
+}
+
 
-		// Start frame 
-		int32 StartVal = 0;
-		if (!InClipItemNode->GetChildValue("start", StartVal, ENodeInherit::NoInherit))
+bool FFCPXMLImportVisitor::VisitAudioClipItemNode(TSharedRef<FFCPXMLClipItemNode> InClipItemNode)
+{
+	// This should only have been called if clipitem is in an audio track node within a sequence node
+	if (!bInSequenceNode || !bInAudioTrackNode)
+	{
+		return false;
+	}
+
+	FString ClipItemName{ TEXT("") };
+	FString ClipItemId{ TEXT("") };
+	FString MasterClipId{ TEXT("") };
+	FString LogNote{ TEXT("") };
+	FString Filename{ TEXT("") };
+	TSharedPtr<FFCPXMLNode> LoggingInfoNode = nullptr;
+	FFrameRate FrameRate;
+	FFrameNumber StartOffset;
+	FFrameNumber Start;
+	FFrameNumber End;
+	
+	// Get relevant data from this clip item node
+	if (!GetClipItemNodeData(InClipItemNode, ClipItemName, ClipItemId, MasterClipId, LoggingInfoNode, LogNote, Filename, FrameRate, StartOffset, Start, End))
+	{
+		return false;
+	}
+
+	// Check if we are in track containing second channel of stereo clips.
+	if (!bCurrImportAudioTrackIsStereoChannel && GetAudioClipItemNodeChannel(InClipItemNode, ClipItemId) == 2)
+	{
+		bCurrImportAudioTrackIsStereoChannel = true;
+	}
+
+	// Skip track holding second channel of stereo clips.
+	if (bCurrImportAudioTrackIsStereoChannel)
+	{
+		return true;
+	}
+
+	// Get master clip logging info, if no logging info here
+	if (!LoggingInfoNode.IsValid() && !MasterClipId.IsEmpty())
+	{
+		GetMasterClipLoggingNode(MasterClipId, LoggingInfoNode);
+	}
+
+	// Get audio metadata
+	TSharedPtr<FFCPXMLImportAudioMetadata> AudioMetadata = GetAudioMetadataObject(LogNote, MasterClipId);
+
+	// Get next audio section based on the audio metadata
+	TSharedPtr<FMovieSceneImportAudioSectionData> AudioSectionData = nullptr;
+	TSharedPtr<FMovieSceneImportAudioMasterTrackData> AudioMasterTrackData = nullptr;
+	if (AudioMetadata.IsValid())
+	{
+		if (!GetNextAudioSection(AudioMetadata, AudioMasterTrackData, AudioSectionData))
 		{
 			return false;
 		}
-		FFrameNumber StartFrame = StartVal;
+	}
 
-		// End frame 
-		int32 EndVal = 0;
-		if (!InClipItemNode->GetChildValue("end", EndVal, ENodeInherit::NoInherit))
+	if (AudioSectionData.IsValid())
+	{
+		if (CurrAudioMasterTrack.IsValid() && CurrAudioMasterTrack->MovieSceneTrack != nullptr &&
+			AudioMasterTrackData.IsValid() && AudioMasterTrackData->MovieSceneTrack != nullptr &&
+			CurrAudioMasterTrack->MovieSceneTrack->GetFullName() != AudioMasterTrackData->MovieSceneTrack->GetFullName())
+		{
+			// Move audio section
+			if (!ImportData->MoveAudioSection(AudioSectionData, AudioMasterTrackData, CurrAudioMasterTrack, CurrAudioTrackRowIndex))
+			{
+				return false;
+			}
+		}
+		// Update existing audio section
+		if (!ImportData->SetAudioSection(AudioSectionData, CurrAudioTrackRowIndex, FrameRate, Start, End, StartOffset))
 		{
 			return false;
 		}
+	}
+	else
+	{
+		bool bUseSoundPathName = (AudioMetadata.IsValid() && !AudioMetadata->SoundPathName.IsEmpty());
+		FString SoundWaveName =  bUseSoundPathName ? AudioMetadata->SoundPathName : Filename;
+		if (!SoundWaveName.IsEmpty())
+		{
+			// Add new audio section
+			AudioSectionData = ImportData->CreateAudioSection(SoundWaveName, bUseSoundPathName, CurrAudioMasterTrack, CurrAudioTrackRowIndex, FrameRate, Start, End, StartOffset);
+			if (!AudioSectionData.IsValid())
+			{
+				return false;
+			}
+		}
+	}
+
+	// Import metadata to cinematic section
+	if (LoggingInfoNode.IsValid() && AudioSectionData.IsValid())
+	{
+		ImportSectionMetaData(LoggingInfoNode, AudioSectionData->AudioSection);
+	}
+
+	return true;
+}
+
+bool FFCPXMLImportVisitor::VisitNode(TSharedRef<FFCPXMLFileNode> InFileNode)
+{
+	// Files can be referenced so flag to visit reference node children
+	return InFileNode->VisitChildren(*this, true);
+}
+
+
+bool FFCPXMLImportVisitor::GetClipItemNodeData(TSharedRef<FFCPXMLClipItemNode> InClipItemNode, FString& OutClipItemName, FString& OutClipItemId, FString& OutMasterClipId,
+	TSharedPtr<FFCPXMLNode>& OutLoggingInfoNode, FString& OutLogNote, FString& OutFilename, FFrameRate& OutFrameRate, FFrameNumber& OutStartOffset,
+	FFrameNumber& OutStart, FFrameNumber &OutEnd)
+{
+	// Shot name
+	if (!InClipItemNode->GetChildValue("name", OutClipItemName, ENodeInherit::NoInherit) || OutClipItemName.IsEmpty())
+	{
+		return false;
+	}
 
-		// In frame 
-		int32 InVal = 0;
-		bool bInDefault = false;
-		if (!InClipItemNode->GetChildValue("in", InVal, ENodeInherit::NoInherit))
+	// clip item id
+	if (!InClipItemNode->GetAttributeValue("id", OutClipItemId) || OutClipItemId.IsEmpty())
+	{
+		return false;
+	}
+
+	// masterclip id
+	if (!InClipItemNode->GetChildValue<FString>("masterclipid", OutMasterClipId, ENodeInherit::NoInherit))
+	{
+		OutMasterClipId = TEXT("");
+	}
+
+	// logging info node
+	OutLoggingInfoNode = InClipItemNode->GetChildNode("logginginfo", ENodeInherit::NoInherit, ENodeReference::NoReferences);
+
+	// log note
+	if (OutLoggingInfoNode.IsValid())
+	{
+		if (!OutLoggingInfoNode->GetChildValue<FString>("lognote", OutLogNote, ENodeInherit::NoInherit, ENodeReference::NoReferences))
 		{
-			bInDefault = true;
+			OutLogNote = TEXT("");
 		}
+	}
+
+	// Frame rate 
+	int32 FrameRateTimebase = 0;
+	if (!InClipItemNode->GetChildSubValue<int32>("rate", "timebase", FrameRateTimebase, ENodeInherit::CheckInherit))
+	{
+		return false;
+	}
+
+	bool bFrameRateIsNTSC = false;
+	if (!InClipItemNode->GetChildSubValue<bool>("rate", "ntsc", bFrameRateIsNTSC, ENodeInherit::CheckInherit))
+	{
+		bFrameRateIsNTSC = false;
+	}
+
+	// Start frame 
+	int32 StartVal = 0;
+	if (!InClipItemNode->GetChildValue("start", StartVal, ENodeInherit::NoInherit))
+	{
+		return false;
+	}
+	FFrameNumber StartFrame = StartVal;
+
+	// End frame 
+	int32 EndVal = 0;
+	if (!InClipItemNode->GetChildValue("end", EndVal, ENodeInherit::NoInherit))
+	{
+		return false;
+	}
+
+	// In frame 
+	int32 InVal = 0;
+	bool bInDefault = false;
+	if (!InClipItemNode->GetChildValue("in", InVal, ENodeInherit::NoInherit))
+	{
+		bInDefault = true;
+	}
 
-		// Out frame 
-		int32 OutVal = 0;
-		bool bOutDefault = false;
-		if (!InClipItemNode->GetChildValue("out", OutVal, ENodeInherit::NoInherit))
+	// Out frame 
+	int32 OutVal = 0;
+	bool bOutDefault = false;
+	if (!InClipItemNode->GetChildValue("out", OutVal, ENodeInherit::NoInherit))
+	{
+		bOutDefault = true;
+	}
+
+	uint32 Numerator = static_cast<uint32>(FrameRateTimebase);
+	OutFrameRate = FFrameRate(Numerator, 1);
+	OutStartOffset = bInDefault ? 0 : InVal;
+	OutStart = StartVal;
+	OutEnd = EndVal;
+
+	TSharedPtr<FFCPXMLNode> LinkNode = InClipItemNode->GetChildNode("file", ENodeInherit::NoInherit, ENodeReference::CheckReferences);
+	if (LinkNode.IsValid())
+	{
+		TSharedPtr<FFCPXMLNode> NameNode = LinkNode->GetChildNode("name", ENodeInherit::NoInherit, ENodeReference::CheckReferences);
+		if (NameNode.IsValid())
 		{
-			bOutDefault = true;
+			NameNode->GetContent(OutFilename);
 		}
+	}
 
-		uint32 Numerator = static_cast<uint32>(FrameRateTimebase);
-		FFrameRate FrameRate(Numerator, 1);
-		FFrameNumber StartOffset = bInDefault ? 0 : InVal;
-		FFrameNumber Start = StartVal;
-		FFrameNumber End = EndVal;
+	return true;
+}
 
-		// @todo handle xml subclips?
+int32 FFCPXMLImportVisitor::GetAudioClipItemNodeChannel(TSharedRef<FFCPXMLClipItemNode> InClipItemNode, const FString& InClipItemId)
+{
+	TSharedPtr<FFCPXMLNode> LinkNode = InClipItemNode->GetChildNode("link", ENodeInherit::NoInherit, ENodeReference::NoReferences);
+	if (LinkNode.IsValid())
+	{
+		TSharedPtr<FFCPXMLNode>LinkClipRefNode = LinkNode->GetChildNode("linkclipref", ENodeInherit::NoInherit, ENodeReference::NoReferences);
+		if (LinkClipRefNode.IsValid())
+		{
+			FString CompareClipItemId;
+			LinkClipRefNode->GetContent(CompareClipItemId);
+			return (InClipItemId == CompareClipItemId) ? 1 : 2;
+		}
+	}
+	return 1;
+}
 
-		TSharedPtr<FMovieSceneImportCinematicSectionData> SectionData = ImportData->FindCinematicSection(CinematicSectionName);
-		if (!SectionData.IsValid())
+bool FFCPXMLImportVisitor::ConstructAudioTrackList()
+{
+	if (!ImportData->MovieSceneData.IsValid())
+	{
+		return false;
+	}
+
+	for (TSharedPtr<FMovieSceneImportAudioMasterTrackData> MasterTrackData : ImportData->MovieSceneData->AudioMasterTracks)
+	{
+		for (TSharedPtr<FMovieSceneImportAudioTrackData> TrackData : MasterTrackData->AudioTracks)
 		{
-			SectionData = ImportData->CreateCinematicSection(CinematicSectionName, CurrVideoTrackRowIndex, FrameRate, Start, End, StartOffset);
-			if (!SectionData.IsValid())
+			if (TrackData.IsValid())
 			{
-				return false;
+				TSharedPtr<FFCPXMLImportAudioTrackListItem> ListItem = MakeShared<FFCPXMLImportAudioTrackListItem>(MasterTrackData, TrackData->RowIndex);
+				AudioTrackList.Add(ListItem);
 			}
 		}
-		else
+	}
+	return true;
+}
+
+/** Add entry to master clip section name map */
+bool FFCPXMLImportVisitor::AddMasterClipCinematicSectionPathName(const FString& InMasterClipIdName, const FString& InSectionPathName)
+{
+	if (MasterClipCinematicSectionMap.Contains(InMasterClipIdName))
+	{
+		return false;
+	}
+
+	MasterClipCinematicSectionMap.Add(InMasterClipIdName, InSectionPathName);
+	return true;
+}
+
+/** Query master clip section name map */
+bool FFCPXMLImportVisitor::GetMasterClipCinematicSectionPathName(const FString& InMasterClipIdName, FString& OutSectionPathName) const
+{
+	const FString* Name = MasterClipCinematicSectionMap.Find(InMasterClipIdName);
+	if (Name == nullptr)
+	{
+		return false;
+	}
+
+	OutSectionPathName = *Name;
+	return true;
+}
+
+/** Get cinematic section based on metadata and master clip id */
+FString FFCPXMLImportVisitor::GetCinematicSectionPathName(const FString& InMetadata, const FString& InMasterClipId)
+{
+	if (!InMetadata.IsEmpty())
+	{
+		FString SectionPathName(TEXT(""));
+		if (GetCinematicSectionPathNameFromMetadata(InMetadata, SectionPathName))
 		{
-			if (!ImportData->SetCinematicSection(SectionData, CurrVideoTrackRowIndex, FrameRate, Start, End, StartOffset))
+			if (!InMasterClipId.IsEmpty())
 			{
-				return false;
+				AddMasterClipCinematicSectionPathName(InMasterClipId, SectionPathName);
 			}
+			return SectionPathName;
 		}
-
-		ImportSectionMetaData(CinematicSectionName, SectionData->CinematicSection);
 	}
 
-	// Clip items can be referenced so flag to visit reference node children
-	return InClipItemNode->VisitChildren(*this, true);
+	if (!InMasterClipId.IsEmpty())
+	{
+		FString SectionPathName(TEXT(""));
+		if (GetMasterClipCinematicSectionPathName(InMasterClipId, SectionPathName))
+		{
+			return SectionPathName;
+		}
+	}
+	return TEXT("");
 }
 
-bool FFCPXMLImportVisitor::VisitNode(TSharedRef<FFCPXMLFileNode> InFileNode)
+/** Add entry to master clip section name map */
+bool FFCPXMLImportVisitor::AddMasterClipAudioMetadata(const FString& InMasterClipIdName, TSharedPtr<FFCPXMLImportAudioMetadata> InAudioMetadata)
 {
-	// Files can be referenced so flag to visit reference node children
-	return InFileNode->VisitChildren(*this, true);
+	if (MasterClipAudioSectionMap.Contains(InMasterClipIdName))
+	{
+		return false;
+	}
+
+	MasterClipAudioSectionMap.Add(InMasterClipIdName, InAudioMetadata);
+
+	return true;
 }
 
-bool FFCPXMLImportVisitor::AddMasterClipSectionName(const FString& InMasterClipIdName, const FString& InSequencerSectionName)
+/** Query master clip section name map */
+bool FFCPXMLImportVisitor::GetMasterClipAudioMetadata(const FString& InMasterClipIdName, TSharedPtr<FFCPXMLImportAudioMetadata>& OutAudioMetadata) const
 {
-	if (MasterClipSectionNameMap.Contains(InMasterClipIdName))
+	const TSharedPtr<FFCPXMLImportAudioMetadata>* AudioMetadata = MasterClipAudioSectionMap.Find(InMasterClipIdName);
+	if (AudioMetadata == nullptr)
 	{
 		return false;
 	}
 
-	MasterClipSectionNameMap.Add(InMasterClipIdName, InSequencerSectionName);
+	OutAudioMetadata = *AudioMetadata;
 	return true;
 }
 
-bool FFCPXMLImportVisitor::GetMasterClipSectionName(const FString& InMasterClipIdName, FString& OutSequencerSectionName) const
+
+TSharedPtr<FFCPXMLImportAudioMetadata> FFCPXMLImportVisitor::GetAudioMetadataObject(const FString& InLogNote, const FString& InMasterClipId)
 {
-	const FString* Name = MasterClipSectionNameMap.Find(InMasterClipIdName);
-	if (Name == nullptr)
+	// Premiere exports the masterclip's logging info onto each clipitem node.
+	// So check the master clip FIRST and use this audio metadata if it exists.
+
+	TSharedPtr<FFCPXMLImportAudioMetadata> AudioMetadata = nullptr;
+
+	if (!InMasterClipId.IsEmpty())
+	{
+		GetMasterClipAudioMetadata(InMasterClipId, AudioMetadata);
+		if (AudioMetadata.IsValid())
+		{
+			return AudioMetadata;
+		}
+	}
+
+	if (!InLogNote.IsEmpty())
+	{
+		GetAudioFromMetadata(InLogNote, AudioMetadata);
+
+		if (!InMasterClipId.IsEmpty() && AudioMetadata.IsValid())
+		{
+			AddMasterClipAudioMetadata(InMasterClipId, AudioMetadata);
+		}
+	}
+
+	return AudioMetadata;
+}
+
+/** Get audio section path name based on node metadata and masterclip id */
+bool FFCPXMLImportVisitor::GetNextAudioSection(TSharedPtr<FFCPXMLImportAudioMetadata> InAudioMetadata, TSharedPtr<FMovieSceneImportAudioMasterTrackData>& OutAudioMasterTrackData, TSharedPtr<FMovieSceneImportAudioSectionData>& OutAudioSectionData)
+{
+	if (!InAudioMetadata.IsValid())
 	{
 		return false;
 	}
 
-	OutSequencerSectionName = *Name;
+	// if audio metadata exists, look for matching section
+	FString AudioSectionPathName{ TEXT("") };
+	for (TSharedPtr<FFCPXMLImportAudioSectionMetadata> Section : InAudioMetadata->AudioSections)
+	{
+		if (Section.IsValid() && !Section->bAudioSectionUpdated)
+		{
+			AudioSectionPathName = Section->AudioSectionPathName;
+			Section->bAudioSectionUpdated = true;
+			break;
+		}
+	}
+
+	// Find actual audio section
+	if (!AudioSectionPathName.IsEmpty())
+	{
+		OutAudioSectionData = ImportData->FindAudioSection(AudioSectionPathName, OutAudioMasterTrackData);
+	}
+
 	return true;
 }
 
+bool FFCPXMLImportVisitor::AddMasterClipLoggingNode(const FString& InMasterClipName, TSharedPtr<FFCPXMLNode> InLoggingInfoNode)
+{
+	if (MasterClipLoggingNodeMap.Contains(InMasterClipName))
+	{
+		return false;
+	}
 
-/** Add metadata entry */
-bool FFCPXMLImportVisitor::AddSectionMetaData(const FString& InSequencerSectionName, TSharedPtr<FFCPXMLNode> InLoggingInfoNode)
+	MasterClipLoggingNodeMap.Add(InMasterClipName, InLoggingInfoNode);
+	return true;
+}
+
+bool FFCPXMLImportVisitor::GetMasterClipLoggingNode(const FString& InMasterClipName, TSharedPtr<FFCPXMLNode>& OutLoggingInfoNode) const
 {
-	if (SectionMetaDataMap.Contains(InSequencerSectionName))
+	const TSharedPtr<FFCPXMLNode>* LoggingNode = MasterClipLoggingNodeMap.Find(InMasterClipName);
+	if (LoggingNode == nullptr)
 	{
 		return false;
 	}
 
-	SectionMetaDataMap.Add(InSequencerSectionName, InLoggingInfoNode);
+	OutLoggingInfoNode = *LoggingNode;
 	return true;
 }
 
-void SetMetaDataValue(TSharedPtr<FFCPXMLNode> InNode, const FString& InElement, UMetaData* InMetaData, const UMovieSceneCinematicShotSection* InSection)
+bool FFCPXMLImportVisitor::SetMetaDataValue(TSharedPtr<FFCPXMLNode> InNode, const FString& InElement, UMetaData* InMetaData, const UMovieSceneSection* InSection) const
 {
 	if (!InNode.IsValid() || InMetaData == nullptr || InSection == nullptr)
 	{
-		return;
+		return false;
 	}
 
 	FString Value;
@@ -308,18 +709,13 @@ void SetMetaDataValue(TSharedPtr<FFCPXMLNode> InNode, const FString& InElement,
 	{
 		InMetaData->SetValue(InSection, *InElement, *Value);
 	}
+	return true;
 }
 
 /** Store clip metadata */
-bool FFCPXMLImportVisitor::ImportSectionMetaData(const FString& InSequencerSectionName, const UMovieSceneCinematicShotSection* InSection)
+bool FFCPXMLImportVisitor::ImportSectionMetaData(const TSharedPtr<FFCPXMLNode>& InLoggingInfoNode, const UMovieSceneSection* InSection) const
 {
-	if (InSequencerSectionName.IsEmpty() || InSection == nullptr)
-	{
-		return false;
-	}
-
-	TSharedPtr<FFCPXMLNode>* LoggingInfoNode = SectionMetaDataMap.Find(InSequencerSectionName);
-	if (LoggingInfoNode != nullptr && (*LoggingInfoNode).IsValid())
+	if (InLoggingInfoNode.IsValid())
 	{
 		UPackage* Package = InSection->GetOutermost();
 		check(Package);
@@ -327,22 +723,77 @@ bool FFCPXMLImportVisitor::ImportSectionMetaData(const FString& InSequencerSecti
 		UMetaData* MetaData = Package->GetMetaData();
 		check(MetaData);
 
-		SetMetaDataValue(*LoggingInfoNode, TEXT("description"), MetaData, InSection);
-		SetMetaDataValue(*LoggingInfoNode, TEXT("scene"), MetaData, InSection);
-		SetMetaDataValue(*LoggingInfoNode, TEXT("shottake"), MetaData, InSection);
-		SetMetaDataValue(*LoggingInfoNode, TEXT("good"), MetaData, InSection);
-		SetMetaDataValue(*LoggingInfoNode, TEXT("originalvideofilename"), MetaData, InSection);
-		SetMetaDataValue(*LoggingInfoNode, TEXT("originalaudiofilename"), MetaData, InSection);
-/*
-		FString Value;
-		(*LoggingInfoNode)->GetChildValue(TEXT("description"), Value, ENodeInherit::NoInherit, ENodeReference::NoReferences);
-		if (!Value.IsEmpty())
+		SetMetaDataValue(InLoggingInfoNode, TEXT("description"), MetaData, InSection);
+		SetMetaDataValue(InLoggingInfoNode, TEXT("scene"), MetaData, InSection);
+		SetMetaDataValue(InLoggingInfoNode, TEXT("shottake"), MetaData, InSection);
+		SetMetaDataValue(InLoggingInfoNode, TEXT("good"), MetaData, InSection);
+		SetMetaDataValue(InLoggingInfoNode, TEXT("originalvideofilename"), MetaData, InSection);
+		SetMetaDataValue(InLoggingInfoNode, TEXT("originalaudiofilename"), MetaData, InSection);
+	}
+
+	return true;
+}
+
+/** parse metadata of the format "[key=value]", whitespace ok.  */
+bool FFCPXMLImportVisitor::ParseMetadata(const FString& InMetadata, const FString& InKey, FString& OutValue, FString& OutMetadata) const
+{
+	FString A{ TEXT("") };
+	FString B{ TEXT("") };
+	FString C{ TEXT("") };
+
+	if (InMetadata.Split(InKey, &A, &B))
+	{
+		if (B.Split(TEXT("="), &A, &C))
 		{
-			MetaData->SetValue(InSection, TEXT("description"), *Value);
+			if (C.Split(TEXT("]"), &A, &OutMetadata))
+			{
+				OutValue = A.TrimStartAndEnd();
+				return true;
+			}
 		}
-		*/
 	}
+	return false;
+}
+
+/** parse metadata of the format "[key=value]", whitespace ok.  */
+bool FFCPXMLImportVisitor::ParseMetadata(const FString& InMetadata, const FString& InKey, FString& OutValue) const
+{
+	FString MetadataRemaining{ TEXT("") };
+	return ParseMetadata(InMetadata, InKey, OutValue, MetadataRemaining);
+}
+
+/** Get sequencer section name from section metadata. Format is "[UE4ShotSection=sectionobjectname]", whitespace ok. */
+bool FFCPXMLImportVisitor::GetCinematicSectionPathNameFromMetadata(const FString& InMetadata, FString& OutSectionObjectName) const
+{
+	return ParseMetadata(InMetadata, TEXT("UE4ShotSection"), OutSectionObjectName);
+}
 
+/** Get sequencer track name from track metadata. Format is "[UE4Track=trackobjectname][UE4Row=rowindex]", whitespace ok. */
+bool FFCPXMLImportVisitor::GetAudioFromMetadata(const FString& InMetadata, TSharedPtr<FFCPXMLImportAudioMetadata>& OutAudioMetadata) const
+{
+	FString Metadata1{ TEXT("") };
+	FString SoundWavePathName;
+	bool bSuccess = ParseMetadata(InMetadata, TEXT("UE4SoundWave"), SoundWavePathName, Metadata1);
+	OutAudioMetadata = MakeShared<FFCPXMLImportAudioMetadata>(SoundWavePathName);
+
+	if (bSuccess)
+	{
+		FString Metadata2{ TEXT("") };
+		FString AudioSectionTopLevel{ TEXT("") };
+		bSuccess = ParseMetadata(Metadata1, TEXT("UE4AudioSectionTopLevel"), AudioSectionTopLevel, Metadata2);
+
+		if (bSuccess)
+		{
+			FString AudioSection{ TEXT("") };
+			while (ParseMetadata(Metadata2, TEXT("UE4AudioSection"), AudioSection, Metadata1))
+			{
+				FString AudioSectionPathName = AudioSectionTopLevel + TEXT(".") + AudioSection;
+				TSharedPtr<FFCPXMLImportAudioSectionMetadata> AudioSectionMetadata = MakeShared<FFCPXMLImportAudioSectionMetadata>(AudioSectionPathName);
+				OutAudioMetadata->AudioSections.Add(AudioSectionMetadata);
+				Metadata2 = Metadata1;
+			}
+		}
+	}
 	return true;
 }
 
diff --git a/Engine/Source/Editor/MovieSceneTools/Private/FCPXML/FCPXMLImport.h b/Engine/Source/Editor/MovieSceneTools/Private/FCPXML/FCPXMLImport.h
--- a/Engine/Source/Editor/MovieSceneTools/Private/FCPXML/FCPXMLImport.h
+++ b/Engine/Source/Editor/MovieSceneTools/Private/FCPXML/FCPXMLImport.h
@@ -9,6 +9,31 @@
 
 /** FCP XML node visitor classes. */
 
+struct FFCPXMLImportAudioTrackListItem
+{
+	FFCPXMLImportAudioTrackListItem(TSharedPtr<FMovieSceneImportAudioMasterTrackData> InAudioTrackData, int InRowIndex) : AudioTrackData(InAudioTrackData), RowIndex(InRowIndex) {}
+
+	TSharedPtr<FMovieSceneImportAudioMasterTrackData> AudioTrackData;
+	int32 RowIndex;
+};
+
+struct FFCPXMLImportAudioSectionMetadata
+{
+	FFCPXMLImportAudioSectionMetadata(const FString& InAudioSectionPathName) : AudioSectionPathName(InAudioSectionPathName), bAudioSectionUpdated(false) {}
+
+	FString AudioSectionPathName;
+	bool bAudioSectionUpdated;
+};
+
+struct FFCPXMLImportAudioMetadata
+{
+	FFCPXMLImportAudioMetadata(const FString& InSoundPathName) : SoundPathName(InSoundPathName) {}
+	FFCPXMLImportAudioMetadata() : SoundPathName(TEXT("")) {}
+
+	FString SoundPathName;
+	TArray< TSharedPtr<FFCPXMLImportAudioSectionMetadata> > AudioSections;
+};
+
 /** The FFCPXMLImportVisitor class imports from the FCP 7 XML structure into Sequencer classes. */
 
 class FFCPXMLImportVisitor : public FFCPXMLNodeVisitor
@@ -40,32 +65,81 @@ public:
 	virtual bool VisitNode(TSharedRef<FFCPXMLFileNode> InFileNode) override final;
 
 private:
-	/** Add entry to master clip section name map */
-	bool AddMasterClipSectionName(const FString& InMasterClipIdName, const FString& InSequencerSectionName);
-	/** Query master clip section name map */
-	bool GetMasterClipSectionName(const FString& InMasterClipIdName, FString& OutSectionName) const;
+	/** Called when visiting a video FFCPXMLClipItemNode during visitor traversal. */
+	bool VisitVideoClipItemNode(TSharedRef<FFCPXMLClipItemNode> InClipItemNode);
+	/** Called when visiting an audio FFCPXMLClipItemNode during visitor traversal. */
+	bool VisitAudioClipItemNode(TSharedRef<FFCPXMLClipItemNode> InClipItemNode);
+	/** Get clipitem node data used by both video and audio clipitems */
+	bool GetClipItemNodeData(TSharedRef<FFCPXMLClipItemNode> InClipItemNode, FString& OutClipItemName, FString& OutClipItemId, FString& OutMasterClipId, TSharedPtr<FFCPXMLNode>& OutLoggingInfoNode, 
+							 FString& OutLogNote, FString& OutFilename, FFrameRate& OutFrameRate, FFrameNumber& StartOffset, FFrameNumber& Start, FFrameNumber &End);
+	/** Query which channel audio clipitem node represents */
+	int32 GetAudioClipItemNodeChannel(TSharedRef<FFCPXMLClipItemNode> InClipItemNode, const FString& InClipItemId);
+
+	/** Construct list of master tracks and row indexes */
+	bool ConstructAudioTrackList();
+
+	/** Add entry to master cinematic clip section name map */
+	bool AddMasterClipCinematicSectionPathName(const FString& InMasterClipIdName, const FString& InSectionPathName);
+	/** Query master cinematic clip section name map */
+	bool GetMasterClipCinematicSectionPathName(const FString& InMasterClipIdName, FString& OutSectionPathName) const;
+	/** Get cinematic section path name based on node metadata and masterclip id */
+	FString GetCinematicSectionPathName(const FString& InMetadata, const FString& InMasterClipId);
+
+	/** Add entry to master clip audio section map */
+	bool AddMasterClipAudioMetadata(const FString& InMasterClipIdName, TSharedPtr<FFCPXMLImportAudioMetadata> InAudioMetadata);
+	/** Query master clip audio section map */
+	bool GetMasterClipAudioMetadata(const FString& InMasterClipIdName, TSharedPtr<FFCPXMLImportAudioMetadata>& OutAudioMetadata) const;
+	/** Get audio metadata object based on log note and master clip id */
+	TSharedPtr<FFCPXMLImportAudioMetadata> GetAudioMetadataObject(const FString& InLogNote, const FString& InMasterClipId);
+	/** Get next audio section based on audio metadata */
+	bool GetNextAudioSection(TSharedPtr<FFCPXMLImportAudioMetadata> InAudioMetadata, TSharedPtr<FMovieSceneImportAudioMasterTrackData>& OutAudioMasterTrackData, TSharedPtr<FMovieSceneImportAudioSectionData>& OutAudioSectionData);
+
 	/** Add entry to metadata map */
-	bool AddSectionMetaData(const FString& InSequencerSectionName, TSharedPtr<FFCPXMLNode> InLoggingInfoNode);
-	/** Import clip metadata to sequencer section */
-	bool ImportSectionMetaData(const FString& InSequencerSectionName, const UMovieSceneCinematicShotSection* InSection);
+	bool AddMasterClipLoggingNode(const FString& InMasterClipIdName, TSharedPtr<FFCPXMLNode> InLoggingInfoNode);
+	/** Query master clip logging node */
+	bool GetMasterClipLoggingNode(const FString& InMasterClipIdName, TSharedPtr<FFCPXMLNode>& InLoggingInfoNode) const;
 
+	/** Helper method */
+	bool SetMetaDataValue(TSharedPtr<FFCPXMLNode> InNode, const FString& InElement, UMetaData* InMetaData, const UMovieSceneSection* InSection) const;
+	/** Import logging info metadata to sequencer section */
+	bool ImportSectionMetaData(const TSharedPtr<FFCPXMLNode>& InLoggingInfoNode, const UMovieSceneSection* InSection) const;
+
+	/** parse metadata of the format "[tag=value]", whitespace ok. */
+	bool ParseMetadata(const FString& InMetadata, const FString& InKey, FString& OutValue, FString& OutMetadata) const;
+	/** parse metadata of the format "[tag=value]", whitespace ok. */
+	bool ParseMetadata(const FString& InMetadata, const FString& InKey, FString& OutValue) const;
+	/** Get sequencer section id from section metadata. Format is "[UE4ShotSection=sectionobjectname]", whitespace ok. */
+	bool GetCinematicSectionPathNameFromMetadata(const FString& InMetadata, FString& OutSectionObjectName) const;
+	/** Get sequencer sound wave id and audio section ids from metadata. Format is "[UE4SoundWave=trackobjectname][UE4AudioSectionTopLevel=toplevelobjectname][UE4AudioSection=audiosectionobjectname]", whitespace ok. */
+	bool GetAudioFromMetadata(const FString& InMetadata, TSharedPtr<FFCPXMLImportAudioMetadata>& OutAudioMetadata) const;
+	
 private:
 	TSharedRef<FMovieSceneImportData> ImportData;
 	TSharedRef<FMovieSceneTranslatorContext> ImportContext;
 
-	/** Map of masterclip names to sequencer section names. */
-	TMap<FString, FString> MasterClipSectionNameMap;
+	/** Map of masterclip names to sequencer section id names. */
+	TMap<FString, FString> MasterClipCinematicSectionMap;
+
+	/** Map of masterclip names to sequencer audio metadata */
+	TMap<FString, TSharedPtr<FFCPXMLImportAudioMetadata>> MasterClipAudioSectionMap;
+
+	/** Map of masterclip names to logging info nodes */
+	TMap<FString, TSharedPtr<FFCPXMLNode>> MasterClipLoggingNodeMap;
 
-	/** Map of sequencer section name to logging info */
-	TMap<FString, TSharedPtr<FFCPXMLNode>> SectionMetaDataMap;
+	/** List of existing master tracks including row index */
+	TArray< TSharedPtr<FFCPXMLImportAudioTrackListItem> > AudioTrackList;
 
+	/** Traversal state variables */
 	bool bInSequenceNode;
 	bool bInVideoNode;
 	bool bInAudioNode;
 	bool bInVideoTrackNode;
 	bool bInAudioTrackNode;
-	int32 MaxVideoTrackRowIndex;
+
 	int32 CurrVideoTrackRowIndex;
-	int32 MaxAudioTrackIndex;
-	int32 CurrAudioTrackIndex;
+	int32 CurrAudioTrackListIndex;
+	TSharedPtr<FMovieSceneImportAudioMasterTrackData> CurrAudioMasterTrack;
+	int32 CurrAudioTrackRowIndex;
+	bool bCurrImportAudioTrackIsStereoChannel;
 };
+
diff --git a/Engine/Source/Editor/MovieSceneTools/Private/FCPXML/FCPXMLNode.cpp b/Engine/Source/Editor/MovieSceneTools/Private/FCPXML/FCPXMLNode.cpp
--- a/Engine/Source/Editor/MovieSceneTools/Private/FCPXML/FCPXMLNode.cpp
+++ b/Engine/Source/Editor/MovieSceneTools/Private/FCPXML/FCPXMLNode.cpp
@@ -573,31 +573,6 @@ bool FFCPXMLClipNode::DoAccept(FFCPXMLNodeVisitor& NodeVisitor)
 	return NodeVisitor.VisitNode(StaticCastSharedRef<FFCPXMLClipNode>(AsShared()));
 }
 
-/** Get metadata section name from sequencer shot name - format is "[UE4ShotSection=shotname]", whitespace ok. */
-FString FFCPXMLClipNode::GetMetadataSectionName(const FString& InSequencerSectionName)
-{
-	return (TEXT("[UE4ShotSection=") + InSequencerSectionName + TEXT("]"));
-}
-
-/** Get sequencer section name from metadata shot name - format is "[UE4ShotSection=shotname]", whitespace ok. */
-bool FFCPXMLClipNode::GetSequencerSectionName(const FString& InMetadataSectionName, FString& OutSequencerSectionName)
-{
-	FString A, B, C;
-	if (InMetadataSectionName.Split(TEXT("UE4ShotSection"), &A, &B))
-	{
-		if (B.Split(TEXT("="), &A, &C))
-		{
-			if (C.Split(TEXT("]"), &A, &B))
-			{
-				OutSequencerSectionName = A.TrimStartAndEnd();
-				return true;
-			}
-		}
-	}
-
-	return false;
-}
-
 FFCPXMLClipItemNode::FFCPXMLClipItemNode(TSharedPtr<FFCPXMLNode> InParent, TSharedPtr<FFCPXMLFile> InFile) : FFCPXMLNode(FString(TEXT("clipitem")), InParent, InFile) {}
 
 bool FFCPXMLClipItemNode::DoAccept(FFCPXMLNodeVisitor& NodeVisitor)
diff --git a/Engine/Source/Editor/MovieSceneTools/Private/FCPXML/FCPXMLNode.h b/Engine/Source/Editor/MovieSceneTools/Private/FCPXML/FCPXMLNode.h
--- a/Engine/Source/Editor/MovieSceneTools/Private/FCPXML/FCPXMLNode.h
+++ b/Engine/Source/Editor/MovieSceneTools/Private/FCPXML/FCPXMLNode.h
@@ -345,12 +345,6 @@ public:
 	/** Constructor */
 	FFCPXMLClipNode(TSharedPtr<FFCPXMLNode> InParent, TSharedPtr<FFCPXMLFile> InFile);
 
-	/** Get metadata section name from sequencer section name */
-	static FString GetMetadataSectionName(const FString& InSequencerSectionName);
-
-	/** Get sequencer section name from metadata shot name */
-	static bool GetSequencerSectionName(const FString& InMetadataSectionName, FString &OutSequencerSectionName);
-
 private:
 	/** Accepts a node visitor object */
 	virtual bool DoAccept(FFCPXMLNodeVisitor& NodeVisitor) override final;
diff --git a/Engine/Source/Editor/MovieSceneTools/Private/MovieSceneTranslator.cpp b/Engine/Source/Editor/MovieSceneTools/Private/MovieSceneTranslator.cpp
--- a/Engine/Source/Editor/MovieSceneTools/Private/MovieSceneTranslator.cpp
+++ b/Engine/Source/Editor/MovieSceneTools/Private/MovieSceneTranslator.cpp
@@ -13,9 +13,31 @@
 #include "Sections/MovieSceneCinematicShotSection.h"
 #include "MovieSceneTimeHelpers.h"
 #include "AssetRegistryModule.h"
+#include "Sound/SoundWave.h"
+#include "Sound/SoundCue.h"
+#include "EditorFramework/AssetImportData.h"
 
 #define LOCTEXT_NAMESPACE "MovieSceneTranslator"
 
+namespace
+{
+	bool SupportMultipleAudioMasterTracks()
+	{
+		return true;
+	}
+
+	bool AudioSectionIsSoundWave(const UMovieSceneAudioSection* InAudioSection)
+	{
+		USoundBase* SoundBase = InAudioSection->GetSound();
+		if (SoundBase == nullptr)
+		{
+			return false;
+		}
+
+		return (SoundBase->IsA<USoundWave>());
+	}
+}
+
 FMovieSceneExportData::FMovieSceneExportData(const UMovieScene* InMovieScene, FFrameRate InFrameRate, uint32 InResX, uint32 InResY, int32 InHandleFrames, FString InSaveFilename, TSharedPtr<FMovieSceneTranslatorContext> InContext)
 {
 	if (InMovieScene == nullptr)
@@ -31,6 +53,11 @@ FMovieSceneExportData::FMovieSceneExportData(const UMovieScene* InMovieScene, FF
 	HandleFrames = ConvertFrameTime(FFrameTime(InHandleFrames), InMovieScene->GetTickResolution(), FrameRate);
 	SaveFilename = InSaveFilename;
 
+	// preferred sample rate in UE4
+	DefaultAudioSampleRate = 44100;
+	// all audio in UE4 is has depth 16
+	DefaultAudioDepth = 16;
+
 	bExportDataIsValid = ConstructData(InMovieScene);
 }
 
@@ -104,13 +131,16 @@ bool FMovieSceneExportData::ConstructMovieSceneData(const UMovieScene* InMovieSc
 
 	bool bFoundCinematicMasterTrack = false;
 
+	// sort audio master tracks
+	TMap<int32, TSharedPtr<FMovieSceneExportAudioMasterTrackData>> AudioTrackMap;
+
 	const TArray<UMovieSceneTrack*> MasterTracks = InMovieScene->GetMasterTracks();
 	for (UMovieSceneTrack* MasterTrack : MasterTracks)
 	{
 		if (!bFoundCinematicMasterTrack && MasterTrack->IsA(UMovieSceneCinematicShotTrack::StaticClass()))
 		{
 			const UMovieSceneCinematicShotTrack* CinematicTrack = Cast<UMovieSceneCinematicShotTrack>(MasterTrack);
-			if (!ConstructCinematicMasterTrackData(InMovieScene, CinematicTrack))
+			if (CinematicTrack == nullptr || !ConstructCinematicMasterTrackData(InMovieScene, CinematicTrack))
 			{ 
 				return false; 
 			}
@@ -119,13 +149,34 @@ bool FMovieSceneExportData::ConstructMovieSceneData(const UMovieScene* InMovieSc
 		else if (MasterTrack->IsA(UMovieSceneAudioTrack::StaticClass()))
 		{
 			const UMovieSceneAudioTrack* AudioTrack = Cast<UMovieSceneAudioTrack>(MasterTrack);
-			if (!ConstructAudioTrackData(InMovieScene, MovieSceneData, AudioTrack))
+			if (AudioTrack == nullptr || !ConstructAudioMasterTrackData(InMovieScene, AudioTrack, AudioTrackMap))
 			{ 
 				return false; 
 			}
 		}
 	}
 
+	// sort the audio master tracks by their sorting index and add to the AudioMasterTracks array
+	if (AudioTrackMap.Num() > 0)
+	{
+		AudioTrackMap.KeySort([](int32 A, int32 B) {
+			return A < B; // sort keys in order
+		});
+
+		for (auto& Elem : AudioTrackMap)
+		{
+			if (Elem.Value.IsValid())
+			{
+				MovieSceneData->AudioMasterTracks.Add(Elem.Value);
+
+				if (!SupportMultipleAudioMasterTracks())
+				{
+					break;
+				}
+			}
+		}
+	}
+
 	return true;
 }
 
@@ -137,6 +188,7 @@ bool FMovieSceneExportData::ConstructCinematicMasterTrackData(const UMovieScene*
 	}
 
 	TSharedPtr<FMovieSceneExportCinematicMasterTrackData> MasterTrackData = MakeShared<FMovieSceneExportCinematicMasterTrackData>();
+	MasterTrackData->MovieSceneTrack = InCinematicMasterTrack;
 	MovieSceneData->CinematicMasterTrack = MasterTrackData;
 
 	// Construct sections & create track row index array
@@ -175,19 +227,20 @@ bool FMovieSceneExportData::ConstructCinematicMasterTrackData(const UMovieScene*
 	return true;
 }
 
-bool FMovieSceneExportData::ConstructCinematicTrackData(const UMovieScene* InMovieScene, TSharedPtr<FMovieSceneExportCinematicMasterTrackData> InCinematicMasterTrackData, int32 RowIndex)
+bool FMovieSceneExportData::ConstructCinematicTrackData(const UMovieScene* InMovieScene, TSharedPtr<FMovieSceneExportCinematicMasterTrackData> InCinematicMasterTrackData, int32 InRowIndex)
 {
 	if (InMovieScene == nullptr || !InCinematicMasterTrackData.IsValid() || !MovieSceneData.IsValid() || !MovieSceneData->CinematicMasterTrack.IsValid())
 	{
 		return false;
 	}
 
-	TSharedPtr<FMovieSceneExportCinematicTrackData> TrackData = MakeShared<FMovieSceneExportCinematicTrackData>(RowIndex);
+	TSharedPtr<FMovieSceneExportCinematicTrackData> TrackData = MakeShared<FMovieSceneExportCinematicTrackData>();
+	TrackData->RowIndex = InRowIndex;
 	MovieSceneData->CinematicMasterTrack->CinematicTracks.Add(TrackData);
 
 	for (TSharedPtr<FMovieSceneExportCinematicSectionData> Section : InCinematicMasterTrackData->CinematicSections)
 	{
-		if (Section.IsValid() && Section->RowIndex == RowIndex)
+		if (Section.IsValid() && Section->RowIndex == InRowIndex)
 		{
 			TrackData->CinematicSections.Add(Section);
 		}
@@ -196,21 +249,58 @@ bool FMovieSceneExportData::ConstructCinematicTrackData(const UMovieScene* InMov
 	return true;
 }
 
-bool FMovieSceneExportData::ConstructAudioTrackData(const UMovieScene* InMovieScene, TSharedPtr<FMovieSceneExportMovieSceneData> InMovieSceneData, const UMovieSceneAudioTrack* InAudioTrack)
+bool FMovieSceneExportData::ConstructAudioMasterTrackData(const UMovieScene* InMovieScene, const UMovieSceneAudioTrack* InAudioMasterTrack, TMap<int32, TSharedPtr<FMovieSceneExportAudioMasterTrackData>>& InAudioTrackMap)
 {
-	if (InMovieScene == nullptr || InAudioTrack == nullptr || !InMovieSceneData.IsValid())
+	if (InMovieScene == nullptr || !MovieSceneData.IsValid())
 	{
 		return false;
 	}
 
-	TSharedPtr<FMovieSceneExportAudioTrackData> TrackData = MakeShared<FMovieSceneExportAudioTrackData>();
-	InMovieSceneData->AudioTracks.Add(TrackData);
+	TSharedPtr<FMovieSceneExportAudioMasterTrackData> MasterTrackData = MakeShared<FMovieSceneExportAudioMasterTrackData>();
+	MasterTrackData->MovieSceneTrack = InAudioMasterTrack;
+	InAudioTrackMap.Add(InAudioMasterTrack->GetSortingOrder(), MasterTrackData);
+
+	// Construct sections & create track row index array
+	TArray<int32> AudioTrackRowIndices;
+
+	TArray<FString> SectionPathNames;
 
-	for (UMovieSceneSection* ShotSection : InAudioTrack->GetAllSections())
+	for (UMovieSceneSection* Section : InAudioMasterTrack->GetAudioSections())
 	{
-		const UMovieSceneAudioSection* AudioSection = Cast<UMovieSceneAudioSection>(ShotSection);
+		const UMovieSceneAudioSection* AudioSection = Cast<UMovieSceneAudioSection>(Section);
+		if (AudioSection == nullptr)
+		{
+			continue;
+		}
 
-		if (!ConstructAudioSectionData(InMovieScene, TrackData, AudioSection))
+		// skip duplicate sections
+		if (SectionPathNames.Num() > 0 && SectionPathNames.Contains(AudioSection->GetPathName()))
+		{
+			continue;
+		}
+
+		if (AudioSectionIsSoundWave(AudioSection))
+		{
+			if (!ConstructAudioSectionData(InMovieScene, MasterTrackData, AudioSection))
+			{
+				return false;
+			}
+
+			int32 RowIndex = AudioSection->GetRowIndex();
+			if (RowIndex >= 0)
+			{
+				AudioTrackRowIndices.AddUnique(RowIndex);
+			}
+		}
+		SectionPathNames.Add(AudioSection->GetPathName());
+	}
+
+	// Construct tracks and point to sections
+	AudioTrackRowIndices.Sort();
+
+	for (int32 AudioTrackRowIndex : AudioTrackRowIndices)
+	{
+		if (!ConstructAudioTrackData(InMovieScene, MasterTrackData, AudioTrackRowIndex))
 		{
 			return false;
 		}
@@ -219,6 +309,29 @@ bool FMovieSceneExportData::ConstructAudioTrackData(const UMovieScene* InMovieSc
 	return true;
 }
 
+bool FMovieSceneExportData::ConstructAudioTrackData(const UMovieScene* InMovieScene, TSharedPtr<FMovieSceneExportAudioMasterTrackData> InAudioMasterTrackData, int32 InRowIndex)
+{
+	if (InMovieScene == nullptr || !InAudioMasterTrackData.IsValid())
+	{
+		return false;
+	}
+
+	TSharedPtr<FMovieSceneExportAudioTrackData> TrackData = MakeShared<FMovieSceneExportAudioTrackData>();
+	TrackData->SampleRate = DefaultAudioSampleRate;
+	TrackData->RowIndex = InRowIndex;
+	InAudioMasterTrackData->AudioTracks.Add(TrackData);
+
+	for (TSharedPtr<FMovieSceneExportAudioSectionData> Section : InAudioMasterTrackData->AudioSections)
+	{
+		if (Section.IsValid() && Section->RowIndex == InRowIndex)
+		{
+			TrackData->AudioSections.Add(Section);
+		}
+	}
+
+	return true;
+}
+
 bool FMovieSceneExportData::ConstructCinematicSectionData(const UMovieScene* InMovieScene, TSharedPtr<FMovieSceneExportCinematicMasterTrackData> InMasterTrackData, const UMovieSceneCinematicShotSection* InCinematicSection)
 {
 	if (InMovieScene == nullptr || !InMasterTrackData.IsValid() || InCinematicSection == nullptr)
@@ -229,25 +342,58 @@ bool FMovieSceneExportData::ConstructCinematicSectionData(const UMovieScene* InM
 	TSharedPtr<FMovieSceneExportCinematicSectionData> SectionData = MakeShared<FMovieSceneExportCinematicSectionData>();
 	InMasterTrackData->CinematicSections.Add(SectionData);
 
-	SectionData->ShotDisplayName = InCinematicSection->GetShotDisplayName();
-	SectionData->ShotFilename = SectionData->ShotDisplayName + TEXT(".avi");
-	SectionData->CinematicShotSection = InCinematicSection;
+	SectionData->DisplayName = InCinematicSection->GetShotDisplayName();
+	SectionData->SourceFilename = SectionData->DisplayName + TEXT(".avi");
+	SectionData->SourceFilePath = TEXT("");
 
-	ConstructSectionData(InMovieScene, SectionData, InCinematicSection, EMovieSceneTranslatorSectionType::Cinematic, SectionData->ShotDisplayName);
+	ConstructSectionData(InMovieScene, SectionData, InCinematicSection, EMovieSceneTranslatorSectionType::Cinematic, SectionData->DisplayName);
 
 	return true;
 }
 
-bool FMovieSceneExportData::ConstructAudioSectionData(const UMovieScene* InMovieScene, TSharedPtr<FMovieSceneExportAudioTrackData> InTrackData, const UMovieSceneAudioSection* InAudioSection)
+bool FMovieSceneExportData::ConstructAudioSectionData(const UMovieScene* InMovieScene, TSharedPtr<FMovieSceneExportAudioMasterTrackData> InTrackData, const UMovieSceneAudioSection* InAudioSection)
 {
 	if (InMovieScene == nullptr || !InTrackData.IsValid() || InAudioSection == nullptr)
 	{
 		return false;
 	}
 
+	USoundBase* SoundBase = InAudioSection->GetSound();
+	if (SoundBase == nullptr || !SoundBase->IsA<USoundWave>())
+	{
+		return false;
+	}
+
+	USoundWave* SoundWave = nullptr;
+	SoundWave = Cast<USoundWave>(SoundBase);
+	if (SoundWave == nullptr || SoundWave->AssetImportData == nullptr)
+	{
+		return false;
+	}
+
 	TSharedPtr<FMovieSceneExportAudioSectionData> SectionData = MakeShared<FMovieSceneExportAudioSectionData>();
 	InTrackData->AudioSections.Add(SectionData);
 
+	TArray<FString> Filenames = SoundWave->AssetImportData->ExtractFilenames();
+	if (Filenames.Num() < 1)
+	{
+		return false;
+	}
+
+	int32 SampleRate = SoundWave->GetSampleRateForCurrentPlatform();
+	if (SampleRate != 48000 && SampleRate != 44100 && SampleRate != 32000)
+	{
+		// @todo - warning about invalid sample rate
+		SampleRate = 44100;
+	}
+	
+	SectionData->DisplayName = SoundWave->GetName();
+	SectionData->SourceFilename = FPaths::GetCleanFilename(Filenames[0]);
+	SectionData->SourceFilePath = FPaths::GetPath(Filenames[0]);
+	SectionData->Depth = GetDefaultAudioDepth();
+	SectionData->SampleRate = SampleRate;
+	SectionData->NumChannels = SoundWave->NumChannels;
+
 	ConstructSectionData(InMovieScene, SectionData, InAudioSection, EMovieSceneTranslatorSectionType::Audio, TEXT(""));
 
 	return true;
@@ -260,6 +406,7 @@ bool FMovieSceneExportData::ConstructSectionData(const UMovieScene* InMovieScene
 		return false;
 	}
 
+	InSectionData->MovieSceneSection = InSection;
 	InSectionData->RowIndex = InSection->GetRowIndex();
 
 	if (InSection->HasStartFrame())
@@ -310,7 +457,7 @@ bool FMovieSceneExportData::ConstructSectionData(const UMovieScene* InMovieScene
 		}
 	}
 
-	// @todo handle intersection with playback range
+	// @todo handle intersection with playback range?
 	TRange<FFrameNumber> PlaybackRange = InMovieScene->GetPlaybackRange();
 	TRange<FFrameNumber> EditRange = InSection->GetRange();
 	TRange<FFrameNumber> Intersection = TRange<FFrameNumber>::Intersection(PlaybackRange, EditRange);
@@ -320,6 +467,43 @@ bool FMovieSceneExportData::ConstructSectionData(const UMovieScene* InMovieScene
 	return true;
 }
 
+bool FMovieSceneExportData::FindAudioSections(const FString& InSoundPathName, TArray<TSharedPtr<FMovieSceneExportAudioSectionData>>& OutFoundSections) const
+{
+	if (!MovieSceneData.IsValid())
+	{
+		return false;
+	}
+
+	for (TSharedPtr<FMovieSceneExportAudioMasterTrackData> AudioMasterTrack : MovieSceneData->AudioMasterTracks)
+	{
+		for (TSharedPtr<FMovieSceneExportAudioSectionData> AudioSectionData : AudioMasterTrack->AudioSections)
+		{
+			if (!AudioSectionData.IsValid() || AudioSectionData->MovieSceneSection == nullptr)
+			{
+				continue;
+			}
+			const UMovieSceneAudioSection* AudioSection = Cast<UMovieSceneAudioSection>(AudioSectionData->MovieSceneSection);
+			if (AudioSection == nullptr)
+			{
+				continue;
+			}
+			USoundBase* Sound = AudioSection->GetSound();
+			if (Sound == nullptr || !Sound->IsA<USoundWave>())
+			{
+				continue;
+			}
+
+			FString SoundPathName = Sound->GetPathName();
+			if (SoundPathName == InSoundPathName)
+			{
+				OutFoundSections.Add(AudioSectionData);
+			}
+		}
+	}
+
+	return true;
+}
+
 FString FMovieSceneExportData::GetFilename() const
 {
 	return SaveFilename;
@@ -366,6 +550,16 @@ FFrameTime FMovieSceneExportData::GetHandleFrames() const
 	return HandleFrames;
 }
 
+int32 FMovieSceneExportData::GetDefaultAudioSampleRate() const
+{
+	return DefaultAudioSampleRate;
+}
+
+int32 FMovieSceneExportData::GetDefaultAudioDepth() const
+{
+	return DefaultAudioDepth;
+}
+
 FMovieSceneImportData::FMovieSceneImportData(UMovieScene* InMovieScene, TSharedPtr<FMovieSceneTranslatorContext> InContext)
 {
 	if (InMovieScene == nullptr)
@@ -391,7 +585,7 @@ bool FMovieSceneImportData::IsImportDataValid() const
 	return MovieSceneData.IsValid();
 }
 
-TSharedPtr<FMovieSceneImportCinematicSectionData> FMovieSceneImportData::FindCinematicSection(const FString& InName) 
+TSharedPtr<FMovieSceneImportCinematicSectionData> FMovieSceneImportData::FindCinematicSection(const FString& InSectionPathName)
 {
 	TSharedPtr<FMovieSceneImportCinematicMasterTrackData> MasterTrackData = GetCinematicMasterTrackData(false);
 	if (!MasterTrackData.IsValid())
@@ -399,16 +593,17 @@ TSharedPtr<FMovieSceneImportCinematicSectionData> FMovieSceneImportData::FindCin
 		return nullptr;
 	}
 
-	for (TSharedPtr<FMovieSceneImportCinematicSectionData> CinematicSectionData : MasterTrackData->CinematicSectionsData)
+	for (TSharedPtr<FMovieSceneImportCinematicSectionData> CinematicSection : MasterTrackData->CinematicSections)
 	{
-		UMovieSceneCinematicShotSection* CinematicShotSection = CinematicSectionData->CinematicSection;
+		UMovieSceneCinematicShotSection* CinematicShotSection = CinematicSection->CinematicSection;
 		if (CinematicShotSection != nullptr)
 		{
 			UMovieSceneSequence* ShotSequence = CinematicShotSection->GetSequence();
 
-			if (ShotSequence != nullptr && ShotSequence->GetName() == InName)
+			FString ShotSectionPathName = CinematicShotSection->GetPathName();
+			if (ShotSequence != nullptr && ShotSectionPathName == InSectionPathName)
 			{
-				return CinematicSectionData;
+				return CinematicSection;
 			}
 		}
 	}
@@ -424,12 +619,12 @@ TSharedPtr<FMovieSceneImportCinematicSectionData> FMovieSceneImportData::CreateC
 	}
 
 	TSharedPtr<FMovieSceneImportCinematicMasterTrackData> MasterTrackData = GetCinematicMasterTrackData(true);
-	if (!MasterTrackData.IsValid())
+	if (!MasterTrackData.IsValid() || MasterTrackData->MovieSceneTrack == nullptr)
 	{
 		return nullptr;
 	}
 
-	UMovieSceneCinematicShotTrack* MasterTrack = MasterTrackData->CinematicMasterTrack;
+	UMovieSceneCinematicShotTrack* MasterTrack = Cast<UMovieSceneCinematicShotTrack>(MasterTrackData->MovieSceneTrack);
 	if (MasterTrack == nullptr)
 	{
 		return nullptr;
@@ -476,8 +671,14 @@ TSharedPtr<FMovieSceneImportCinematicSectionData> FMovieSceneImportData::CreateC
 	Section->SetRange(TRange<FFrameNumber>(StartFrame, EndFrame));
 
 	TSharedPtr<FMovieSceneImportCinematicSectionData> SectionData = ConstructCinematicSectionData(Section);
-	MasterTrackData->CinematicSectionsData.Add(SectionData);
-
+	MasterTrackData->CinematicSections.Add(SectionData);
+	for (TSharedPtr<FMovieSceneImportCinematicTrackData> TrackData : MasterTrackData->CinematicTracks)
+	{
+		if (InRow == TrackData->RowIndex)
+		{
+			TrackData->CinematicSections.Add(SectionData);
+		}
+	}
 	return SectionData;
 }
 
@@ -504,13 +705,220 @@ bool FMovieSceneImportData::SetCinematicSection(TSharedPtr<FMovieSceneImportCine
 	return true;
 }
 
-TSharedPtr<FMovieSceneImportAudioSectionData> FMovieSceneImportData::FindAudioSection(FString InName) const
+TSharedPtr<FMovieSceneImportAudioSectionData> FMovieSceneImportData::FindAudioSection(const FString& InSectionPathName, TSharedPtr<FMovieSceneImportAudioMasterTrackData>& OutMasterTrackData)
 {
-	// @todo
+	if (!MovieSceneData.IsValid())
+	{
+		OutMasterTrackData = nullptr;
+		return nullptr;
+	}
+
+	for (TSharedPtr<FMovieSceneImportAudioMasterTrackData> MasterTrackData : MovieSceneData->AudioMasterTracks)
+	{
+		if (!MasterTrackData.IsValid())
+		{
+			continue;
+		}
+
+		for (TSharedPtr<FMovieSceneImportAudioSectionData> AudioSectionData : MasterTrackData->AudioSections)
+		{
+			if (AudioSectionData.IsValid() && AudioSectionData->AudioSection != nullptr)
+			{
+				FString SectionName = AudioSectionData->AudioSection->GetPathName();
+				if (SectionName == InSectionPathName)
+				{
+					OutMasterTrackData = MasterTrackData;
+					return AudioSectionData;
+				}
+			}
+		}
+	}
 
+	OutMasterTrackData = nullptr;
 	return nullptr;
 }
 
+/** Create audio section */
+TSharedPtr<FMovieSceneImportAudioSectionData> FMovieSceneImportData::CreateAudioSection(FString InFilenameOrAssetPathName, bool bIsPathName, TSharedPtr<FMovieSceneImportAudioMasterTrackData> InMasterTrackData, int32 InRow, FFrameRate InFrameRate, FFrameNumber InStartFrame, FFrameNumber InEndFrame, FFrameNumber InStartOffsetFrame)
+{
+	if (!MovieSceneData.IsValid() || MovieSceneData->MovieScene == nullptr || !InMasterTrackData.IsValid() || InMasterTrackData->MovieSceneTrack == nullptr)
+	{
+		return nullptr;
+	}
+
+	UMovieSceneAudioTrack* MasterTrack = Cast<UMovieSceneAudioTrack>(InMasterTrackData->MovieSceneTrack);
+	if (MasterTrack == nullptr)
+	{
+		return nullptr;
+	}
+
+	USoundWave* SoundToAdd = nullptr;
+
+	FAssetRegistryModule& AssetRegistryModule = FModuleManager::LoadModuleChecked<FAssetRegistryModule>(TEXT("AssetRegistry"));
+
+	// Collect a full list of assets with the specified class
+	TArray<FAssetData> AssetDataArray;
+	AssetRegistryModule.Get().GetAssetsByClass(USoundWave::StaticClass()->GetFName(), AssetDataArray);
+
+	for (FAssetData AssetData : AssetDataArray)
+	{
+		USoundWave* SoundWaveAsset = Cast<USoundWave>(AssetData.GetAsset());
+
+		if (SoundWaveAsset == nullptr)
+		{
+			continue;
+		}
+
+		if (bIsPathName)
+		{
+			if (InFilenameOrAssetPathName == SoundWaveAsset->GetPathName())
+			{
+				SoundToAdd = SoundWaveAsset;
+				break;
+			}
+		}
+		else
+		{
+			if (SoundWaveAsset->AssetImportData == nullptr)
+			{
+				continue;
+			}
+
+			TArray<FString> Filenames = SoundWaveAsset->AssetImportData->ExtractFilenames();
+			if (Filenames.Num() < 1)
+			{
+				continue;
+			}
+
+			FString Filename = FPaths::GetCleanFilename(Filenames[0]);
+			if (Filename == InFilenameOrAssetPathName)
+			{
+				SoundToAdd = SoundWaveAsset;
+				break;
+			}
+		}
+	}
+	
+	if (SoundToAdd == nullptr)
+	{
+		return nullptr;
+	}
+
+	// both FCP XML and Sequencer have inclusive start frame, exclusive end frame
+	FFrameRate TickResolution = MovieSceneData->MovieScene->GetTickResolution();
+	FFrameNumber StartFrame = ConvertFrameTime(InStartFrame, InFrameRate, TickResolution).RoundToFrame();
+	FFrameNumber StartOffsetFrame = ConvertFrameTime(InStartOffsetFrame, InFrameRate, TickResolution).RoundToFrame();
+	FFrameNumber EndFrame = ConvertFrameTime(InEndFrame, InFrameRate, TickResolution).RoundToFrame();
+	int32 Duration = (EndFrame - StartFrame).Value;
+
+	MasterTrack->Modify();
+	UMovieSceneAudioSection* AudioSection = Cast<UMovieSceneAudioSection>(MasterTrack->AddNewSoundOnRow(SoundToAdd, StartFrame, InRow));
+	if (AudioSection == nullptr)
+	{
+		return nullptr;
+	}
+	AudioSection->Modify();
+	AudioSection->SetRowIndex(InRow);
+	AudioSection->SetStartOffset(StartOffsetFrame.Value);
+	AudioSection->SetRange(TRange<FFrameNumber>(StartFrame, EndFrame));
+
+	TSharedPtr<FMovieSceneImportAudioSectionData> AudioSectionData = ConstructAudioSectionData(AudioSection);
+
+	InMasterTrackData->AudioSections.Add(AudioSectionData);
+	for (TSharedPtr<FMovieSceneImportAudioTrackData> TrackData : InMasterTrackData->AudioTracks)
+	{
+		if (InRow == TrackData->RowIndex)
+		{
+			TrackData->AudioSections.Add(AudioSectionData);
+		}
+	}
+
+	return AudioSectionData;
+}
+
+bool FMovieSceneImportData::SetAudioSection(TSharedPtr<FMovieSceneImportAudioSectionData> InSection, int32 InRow, FFrameRate InFrameRate, FFrameNumber InStartFrame, FFrameNumber InEndFrame, FFrameNumber InStartOffsetFrame)
+{
+	if (!MovieSceneData.IsValid() || MovieSceneData->MovieScene == nullptr || !InSection.IsValid() || InSection->AudioSection == nullptr)
+	{
+		return false;
+	}
+
+	FFrameRate TickResolution = MovieSceneData->MovieScene->GetTickResolution();
+	FFrameNumber StartFrame = ConvertFrameTime(InStartFrame, InFrameRate, TickResolution).GetFrame();
+	FFrameNumber StartOffsetFrame = ConvertFrameTime(InStartOffsetFrame, InFrameRate, TickResolution).GetFrame();
+	FFrameNumber EndFrame = ConvertFrameTime(InEndFrame, InFrameRate, TickResolution).GetFrame();
+
+	InSection->AudioSection->Modify();
+	InSection->AudioSection->SetStartOffset(StartOffsetFrame.Value);
+	InSection->AudioSection->SetRange(TRange<FFrameNumber>(StartFrame, EndFrame));
+	if (InRow != InSection->AudioSection->GetRowIndex())
+	{
+		InSection->AudioSection->SetRowIndex(InRow);
+	}
+
+	return true;
+}
+
+bool FMovieSceneImportData::MoveAudioSection(TSharedPtr<FMovieSceneImportAudioSectionData> InAudioSectionData, TSharedPtr<FMovieSceneImportAudioMasterTrackData> InFromMasterTrackData, TSharedPtr<FMovieSceneImportAudioMasterTrackData> InToMasterTrackData, int32 InToRowIndex)
+{
+	if (!MovieSceneData.IsValid() || !InAudioSectionData.IsValid() || !InFromMasterTrackData.IsValid() || !InToMasterTrackData.IsValid())
+	{
+		return false;
+	}
+
+	UMovieSceneAudioTrack* FromTrack = Cast<UMovieSceneAudioTrack>(InFromMasterTrackData->MovieSceneTrack);
+	if (FromTrack == nullptr)
+	{
+		return false;
+	}
+
+	UMovieSceneAudioTrack* ToTrack = Cast<UMovieSceneAudioTrack>(InToMasterTrackData->MovieSceneTrack);
+	if (ToTrack == nullptr)
+	{
+		return false;
+	}
+	UMovieSceneAudioSection* AudioSection = InAudioSectionData->AudioSection;
+	if (AudioSection == nullptr)
+	{
+		return false;
+	}
+
+	FromTrack->Modify();
+	FromTrack->RemoveSection(*AudioSection);
+	ToTrack->Modify();
+	ToTrack->AddSection(*AudioSection);
+
+	InFromMasterTrackData->AudioSections.Remove(InAudioSectionData);
+	for (TSharedPtr<FMovieSceneImportAudioTrackData> AudioTrackData : InFromMasterTrackData->AudioTracks)
+	{
+		if (AudioTrackData->AudioSections.Contains(InAudioSectionData))
+		{
+			AudioTrackData->AudioSections.Remove(InAudioSectionData);
+		}
+	}
+
+	bool bFoundTrack = false;
+	InToMasterTrackData->AudioSections.Add(InAudioSectionData);
+	for (TSharedPtr<FMovieSceneImportAudioTrackData> AudioTrackData : InToMasterTrackData->AudioTracks)
+	{
+		if (AudioTrackData->RowIndex == InToRowIndex)
+		{
+			AudioTrackData->AudioSections.Add(InAudioSectionData);
+			bFoundTrack = true;
+			break;
+		}
+	}
+
+	if (!bFoundTrack)
+	{
+		TSharedPtr<FMovieSceneImportAudioTrackData> TrackData = MakeShared<FMovieSceneImportAudioTrackData>();
+		TrackData->RowIndex = InToRowIndex;
+		TrackData->AudioSections.Add(InAudioSectionData);
+		InToMasterTrackData->AudioTracks.Add(TrackData);
+	}
+
+	return true;
+}
 
 TSharedPtr<FMovieSceneImportCinematicMasterTrackData> FMovieSceneImportData::GetCinematicMasterTrackData(bool CreateTrackIfNull) 
 {
@@ -518,12 +926,28 @@ TSharedPtr<FMovieSceneImportCinematicMasterTrackData> FMovieSceneImportData::Get
 	{
 		return nullptr;
 	}
-	if (!MovieSceneData->CinematicMasterTrackData.IsValid() && CreateTrackIfNull)
+	if (!MovieSceneData->CinematicMasterTrack.IsValid() && CreateTrackIfNull)
 	{
 		UMovieSceneCinematicShotTrack* CinematicMasterTrack = MovieSceneData->MovieScene->AddMasterTrack<UMovieSceneCinematicShotTrack>();
-		MovieSceneData->CinematicMasterTrackData = ConstructCinematicMasterTrackData(CinematicMasterTrack);
+		MovieSceneData->CinematicMasterTrack = ConstructCinematicMasterTrackData(CinematicMasterTrack);
+	}
+	return MovieSceneData->CinematicMasterTrack;
+}
+
+TSharedPtr<FMovieSceneImportAudioMasterTrackData> FMovieSceneImportData::GetAudioMasterTrackData()
+{
+	if (!MovieSceneData.IsValid())
+	{
+		return nullptr;
+	}
+	for (TSharedPtr<FMovieSceneImportAudioMasterTrackData> MasterTrack : MovieSceneData->AudioMasterTracks)
+	{
+		if (MasterTrack.IsValid())
+		{
+			return MasterTrack;
+		}
 	}
-	return MovieSceneData->CinematicMasterTrackData;
+	return nullptr;
 }
 
 TSharedPtr<FMovieSceneImportMovieSceneData> FMovieSceneImportData::ConstructMovieSceneData(UMovieScene* InMovieScene)
@@ -540,14 +964,16 @@ TSharedPtr<FMovieSceneImportMovieSceneData> FMovieSceneImportData::ConstructMovi
 	UMovieSceneCinematicShotTrack* CinematicMasterTrack = InMovieScene->FindMasterTrack<UMovieSceneCinematicShotTrack>();
 	if (CinematicMasterTrack != nullptr)
 	{
-		MovieSceneData->CinematicMasterTrackData = ConstructCinematicMasterTrackData(CinematicMasterTrack);
-		if (!MovieSceneData->CinematicMasterTrackData.IsValid())
+		MovieSceneData->CinematicMasterTrack = ConstructCinematicMasterTrackData(CinematicMasterTrack);
+		if (!MovieSceneData->CinematicMasterTrack.IsValid())
 		{
 			return nullptr;
 		}
 	}
 
 	// Get audio tracks
+	TMap<int32, TSharedPtr<FMovieSceneImportAudioMasterTrackData>> AudioMasterTrackMap;
+
 	const TArray<UMovieSceneTrack*> MasterTracks = InMovieScene->GetMasterTracks();
 	for (UMovieSceneTrack* MasterTrack : MasterTracks)
 	{
@@ -555,12 +981,40 @@ TSharedPtr<FMovieSceneImportMovieSceneData> FMovieSceneImportData::ConstructMovi
 		{
 
 			UMovieSceneAudioTrack* AudioTrack = Cast<UMovieSceneAudioTrack>(MasterTrack);
-			TSharedPtr<FMovieSceneImportAudioTrackData> AudioTrackData = ConstructAudioTrackData(AudioTrack);
-			if (!AudioTrackData.IsValid())
+			if (AudioTrack == nullptr)
 			{
-				return nullptr;
+				continue;
+			}
+
+			TSharedPtr<FMovieSceneImportAudioMasterTrackData> AudioMasterTrackData = ConstructAudioMasterTrackData(AudioTrack);
+			if (!AudioMasterTrackData.IsValid())
+			{
+				continue;
+			}
+
+			AudioMasterTrackMap.Add(AudioTrack->GetSortingOrder(), AudioMasterTrackData);
+		}
+	}
+
+
+	// sort the audio master tracks by their sorting index and add to the AudioMasterTracks array
+	if (AudioMasterTrackMap.Num() > 0)
+	{
+		AudioMasterTrackMap.KeySort([](int32 A, int32 B) {
+			return A < B; // sort keys in order
+		});
+
+		for (auto& Elem : AudioMasterTrackMap)
+		{
+			if (Elem.Value.IsValid())
+			{
+				MovieSceneData->AudioMasterTracks.Add(Elem.Value);
+
+				if (!SupportMultipleAudioMasterTracks())
+				{
+					break;
+				}
 			}
-			MovieSceneData->AudioTracksData.Add(AudioTrackData);
 		}
 	}
 
@@ -575,47 +1029,157 @@ TSharedPtr<FMovieSceneImportCinematicMasterTrackData> FMovieSceneImportData::Con
 		return nullptr;
 	}
 
-	TSharedPtr<FMovieSceneImportCinematicMasterTrackData> TrackData = MakeShared<FMovieSceneImportCinematicMasterTrackData>();
-	TrackData->CinematicMasterTrack = InCinematicMasterTrack;
+	TSharedPtr<FMovieSceneImportCinematicMasterTrackData> MasterTrackData = MakeShared<FMovieSceneImportCinematicMasterTrackData>();
+	MasterTrackData->MovieSceneTrack = InCinematicMasterTrack;
+
+	// Construct sections & create track row index array
+	TArray<int32> CinematicTrackRowIndices;
 
 	for (UMovieSceneSection* ShotSection : InCinematicMasterTrack->GetAllSections())
 	{
 		UMovieSceneCinematicShotSection* CinematicSection = Cast<UMovieSceneCinematicShotSection>(ShotSection);
 
 		if (CinematicSection != nullptr && CinematicSection->GetSequence() != nullptr)
+		{
+			int32 RowIndex = CinematicSection->GetRowIndex();
+			if (RowIndex >= 0)
+			{
+				CinematicTrackRowIndices.AddUnique(RowIndex);
+			}
+		}
+	}
+
+	// Construct tracks and point to sections
+	CinematicTrackRowIndices.Sort();
+
+	for (int32 CinematicTrackRowIndex : CinematicTrackRowIndices)
+	{
+		TSharedPtr<FMovieSceneImportCinematicTrackData> TrackData = ConstructCinematicTrackData(InCinematicMasterTrack, CinematicTrackRowIndex);
+		if (TrackData.IsValid())
+		{
+			MasterTrackData->CinematicTracks.Add(TrackData);
+
+			for (TSharedPtr<FMovieSceneImportCinematicSectionData> SectionData : TrackData->CinematicSections)
+			{
+				if (SectionData.IsValid())
+				{
+					MasterTrackData->CinematicSections.Add(SectionData);
+				}
+			}
+		}
+	}
+
+	return MasterTrackData;
+}
+
+TSharedPtr<FMovieSceneImportCinematicTrackData> FMovieSceneImportData::ConstructCinematicTrackData(UMovieSceneCinematicShotTrack* InCinematicMasterTrack, int32 InRowIndex)
+{
+	if (!MovieSceneData.IsValid() || InCinematicMasterTrack == nullptr)
+	{
+		return nullptr;
+	}
+
+	TSharedPtr<FMovieSceneImportCinematicTrackData> TrackData = MakeShared<FMovieSceneImportCinematicTrackData>();
+	TrackData->RowIndex = InRowIndex;
+
+	for (UMovieSceneSection* ShotSection : InCinematicMasterTrack->GetAllSections())
+	{
+		UMovieSceneCinematicShotSection* CinematicSection = Cast<UMovieSceneCinematicShotSection>(ShotSection);
+
+		if (CinematicSection != nullptr && CinematicSection->GetSequence() != nullptr && CinematicSection->GetRowIndex() == InRowIndex)
 		{
 			TSharedPtr<FMovieSceneImportCinematicSectionData> CinematicSectionData = ConstructCinematicSectionData(CinematicSection);
 			if (!CinematicSectionData.IsValid())
 			{
 				return nullptr;
 			}
-			TrackData->CinematicSectionsData.Add(CinematicSectionData);
+			TrackData->CinematicSections.Add(CinematicSectionData);
 		}
 	}
 
 	return TrackData;
 }
 
-TSharedPtr<FMovieSceneImportAudioTrackData> FMovieSceneImportData::ConstructAudioTrackData(UMovieSceneAudioTrack* InAudioTrack)
+
+TSharedPtr<FMovieSceneImportAudioMasterTrackData> FMovieSceneImportData::ConstructAudioMasterTrackData(UMovieSceneAudioTrack* InAudioMasterTrack)
 {
-	if (!MovieSceneData.IsValid() || InAudioTrack == nullptr)
+	if (!MovieSceneData.IsValid() || InAudioMasterTrack == nullptr)
 	{
 		return nullptr;
 	}
 
-	TSharedPtr<FMovieSceneImportAudioTrackData> TrackData = MakeShared<FMovieSceneImportAudioTrackData>();
-	TrackData->AudioTrack = InAudioTrack;
+	TSharedPtr<FMovieSceneImportAudioMasterTrackData> MasterTrackData = MakeShared<FMovieSceneImportAudioMasterTrackData>();
+	MasterTrackData->MovieSceneTrack = InAudioMasterTrack;
+	MasterTrackData->MaxRowIndex = 0;
+
+	// Construct sections & create track row index array
+	TArray<int32> AudioTrackRowIndices;
 
-	for (UMovieSceneSection* ShotSection : InAudioTrack->GetAllSections())
+	for (UMovieSceneSection* ShotSection : InAudioMasterTrack->GetAllSections())
 	{
 		UMovieSceneAudioSection* AudioSection = Cast<UMovieSceneAudioSection>(ShotSection);
 
-		TSharedPtr<FMovieSceneImportAudioSectionData> AudioSectionData = ConstructAudioSectionData(AudioSection);
-		if (!AudioSectionData.IsValid())
+		if (AudioSection != nullptr)
 		{
-			return nullptr;
+			int32 RowIndex = AudioSection->GetRowIndex();
+			if (RowIndex >= 0)
+			{
+				AudioTrackRowIndices.AddUnique(RowIndex);
+
+				if (RowIndex > MasterTrackData->MaxRowIndex)
+				{
+					MasterTrackData->MaxRowIndex = RowIndex;
+				}
+			}
+		}
+	}
+
+	// Construct tracks and point to sections
+	AudioTrackRowIndices.Sort();
+
+	for (int32 AudioTrackRowIndex : AudioTrackRowIndices)
+	{
+		TSharedPtr<FMovieSceneImportAudioTrackData> TrackData = ConstructAudioTrackData(InAudioMasterTrack, AudioTrackRowIndex);
+		if (TrackData.IsValid())
+		{
+			MasterTrackData->AudioTracks.Add(TrackData);
+
+			for (TSharedPtr<FMovieSceneImportAudioSectionData> SectionData : TrackData->AudioSections)
+			{
+				if (SectionData.IsValid())
+				{
+					MasterTrackData->AudioSections.Add(SectionData);
+				}
+			}
+		}
+	}
+
+	return MasterTrackData;
+}
+
+TSharedPtr<FMovieSceneImportAudioTrackData> FMovieSceneImportData::ConstructAudioTrackData(UMovieSceneAudioTrack* InAudioMasterTrack, int32 InRowIndex)
+{
+	if (!MovieSceneData.IsValid() || InAudioMasterTrack == nullptr)
+	{
+		return nullptr;
+	}
+
+	TSharedPtr<FMovieSceneImportAudioTrackData> TrackData = MakeShared<FMovieSceneImportAudioTrackData>();
+	TrackData->RowIndex = InRowIndex;
+
+	for (UMovieSceneSection* Section : InAudioMasterTrack->GetAllSections())
+	{
+		UMovieSceneAudioSection* AudioSection = Cast<UMovieSceneAudioSection>(Section);
+
+		if (AudioSection != nullptr && AudioSectionIsSoundWave(AudioSection) && AudioSection->GetRowIndex() == InRowIndex)
+		{
+			TSharedPtr<FMovieSceneImportAudioSectionData> AudioSectionData = ConstructAudioSectionData(AudioSection);
+			if (!AudioSectionData.IsValid())
+			{
+				continue;
+			}
+			TrackData->AudioSections.Add(AudioSectionData);
 		}
-		TrackData->AudioSectionsData.Add(AudioSectionData);
 	}
 
 	return TrackData;
@@ -641,13 +1205,34 @@ TSharedPtr<FMovieSceneImportAudioSectionData> FMovieSceneImportData::ConstructAu
 		return nullptr;
 	}
 
+	USoundBase* SoundBase = InAudioSection->GetSound();
+	if (SoundBase == nullptr || !SoundBase->IsA<USoundWave>())
+	{
+		return false;
+	}
+
+	USoundWave* SoundWave = nullptr;
+	SoundWave = Cast<USoundWave>(SoundBase);
+	if (SoundWave == nullptr || SoundWave->AssetImportData == nullptr)
+	{
+		return false;
+	}
+
 	TSharedPtr<FMovieSceneImportAudioSectionData> SectionData = MakeShared<FMovieSceneImportAudioSectionData>();
 	SectionData->AudioSection = InAudioSection;
 
+	TArray<FString> Filenames = SoundWave->AssetImportData->ExtractFilenames();
+	if (Filenames.Num() < 1)
+	{
+		return false;
+	}
+
+	SectionData->SourceFilename = FPaths::GetCleanFilename(Filenames[0]);
+	SectionData->SourceFilePath = FPaths::GetPath(Filenames[0]);
+
 	return SectionData;
 }
 
-
 void FMovieSceneTranslatorContext::Init()
 {
 	ClearMessages();
diff --git a/Engine/Source/Editor/MovieSceneTools/Public/MovieSceneTranslator.h b/Engine/Source/Editor/MovieSceneTools/Public/MovieSceneTranslator.h
--- a/Engine/Source/Editor/MovieSceneTools/Public/MovieSceneTranslator.h
+++ b/Engine/Source/Editor/MovieSceneTools/Public/MovieSceneTranslator.h
@@ -12,6 +12,7 @@
 #include "Sections/MovieSceneCinematicShotSection.h"
 #include "Logging/TokenizedMessage.h"
 #include "UObject/MetaData.h"
+#include "Sound/SoundBase.h"
 
 /**
 * MovieSceneTranslator context class.
@@ -47,8 +48,13 @@ private:
 
 struct FMovieSceneExportSectionData
 {
-	// FString Name;
+	const UMovieSceneSection* MovieSceneSection;
 	int32 RowIndex;
+
+	FString DisplayName;
+	FString SourceFilename;
+	FString SourceFilePath;
+
 	FFrameNumber StartFrame;
 	FFrameNumber EndFrame;
 	bool bWithinPlaybackRange;
@@ -57,37 +63,54 @@ struct FMovieSceneExportSectionData
 
 struct FMovieSceneExportCinematicSectionData : public FMovieSceneExportSectionData
 {
-	FString ShotDisplayName;
-	FString ShotFilename;
-	const UMovieSceneCinematicShotSection* CinematicShotSection;
 };
 
 struct FMovieSceneExportAudioSectionData : public FMovieSceneExportSectionData
 {
-	uint32 NumChannels;
-	float Samples;
+	int32 NumChannels;
+	int32 Depth;
+	int32 SampleRate;
 };
 
-struct FMovieSceneExportAudioTrackData
+struct FMovieSceneExportCinematicTrackData
 {
-	TArray< TSharedPtr<FMovieSceneExportAudioSectionData> > AudioSections;
+	/** This indicates the sub-track's row index in the master track */
+	int32 RowIndex;
+	TArray< TSharedPtr<FMovieSceneExportCinematicSectionData> > CinematicSections;
 };
 
-struct FMovieSceneExportCinematicTrackData
+struct FMovieSceneExportAudioTrackData
 {
-	FMovieSceneExportCinematicTrackData(int32 InRowIndex) : RowIndex(InRowIndex) {}
+	int32 SampleRate;
 
 	/** This indicates the sub-track's row index in the master track */
 	int32 RowIndex;
-	TArray< TSharedPtr<FMovieSceneExportCinematicSectionData> > CinematicSections;
+	TArray< TSharedPtr<FMovieSceneExportAudioSectionData> > AudioSections;
 };
 
-struct FMovieSceneExportCinematicMasterTrackData
+struct FMovieSceneExportMasterTrackData
+{
+	const UMovieSceneTrack* MovieSceneTrack;
+};
+
+struct FMovieSceneExportAudioMasterTrackData : public FMovieSceneExportMasterTrackData
+{
+	int32 SampleRate;
+
+	/** Array of all sections in order they appear in UMovieSceneAudioTrack*/
+	TArray< TSharedPtr<FMovieSceneExportAudioSectionData> > AudioSections;
+
+	/** Array of sorted audio sub tracks, containing pointers to sections within the sub track row*/
+	TArray< TSharedPtr<FMovieSceneExportAudioTrackData> > AudioTracks;
+
+};
+
+struct FMovieSceneExportCinematicMasterTrackData : public FMovieSceneExportMasterTrackData
 {
 	/** Array of all sections in order they appear in UMovieSceneCinematicTrack*/
 	TArray< TSharedPtr<FMovieSceneExportCinematicSectionData> > CinematicSections;
 
-	/** Array of sorted movie tracks, containing pointers to sections within the track */
+	/** Array of sorted movie sub tracks, containing pointers to sections within the sub track row */
 	TArray< TSharedPtr<FMovieSceneExportCinematicTrackData> > CinematicTracks;
 };
 
@@ -100,13 +123,13 @@ struct FMovieSceneExportMovieSceneData
 	FFrameNumber PlaybackRangeStartFrame;
 	FFrameNumber PlaybackRangeEndFrame;
 	TSharedPtr<FMovieSceneExportCinematicMasterTrackData> CinematicMasterTrack;
-	TArray< TSharedPtr<FMovieSceneExportAudioTrackData> > AudioTracks;
+	TArray< TSharedPtr<FMovieSceneExportAudioMasterTrackData> > AudioMasterTracks;
 };
 
-enum class EMovieSceneTranslatorSectionType : uint32
+enum class EMovieSceneTranslatorSectionType : int32
 {
 	Cinematic = 0,
-	Audio = 1,
+	Audio = 1
 };
 
 /** 
@@ -140,10 +163,17 @@ public:
 
 	/** Gets the frame handle which has been converted to the export frame rate */
 	FFrameTime GetHandleFrames() const;
+	/** Gets default audio sample rate */
+	int32 GetDefaultAudioSampleRate() const;
+	/** Gets default audio depth */
+	int32 GetDefaultAudioDepth() const;
 
 	/** True when the export data was successfully constructed. */
 	bool IsExportDataValid() const;
 
+	/** Find audio sections */
+	bool FindAudioSections(const FString& InSoundPathName, TArray<TSharedPtr<FMovieSceneExportAudioSectionData>>& OutFoundSections) const;
+
 private:
 
 	/** Entry point for creating the intermediate data to use when exporting. */
@@ -153,13 +183,15 @@ private:
 	/** Loads intermediate cinematic master track data from Sequencer. */
 	bool ConstructCinematicMasterTrackData(const UMovieScene* InMovieScene, const UMovieSceneCinematicShotTrack* InCinematicMasterTrack);
 	/** Loads intermediate cinematic track data from Sequencer. */
-	bool ConstructCinematicTrackData(const UMovieScene* InMovieScene, TSharedPtr<FMovieSceneExportCinematicMasterTrackData> InCinematicMasterTrack, int32 RowIndex);
+	bool ConstructCinematicTrackData(const UMovieScene* InMovieScene, TSharedPtr<FMovieSceneExportCinematicMasterTrackData> InCinematicMasterTrack, int32 InRowIndex);
+	/** Loads intermediate audio track data from Sequencer. */
+	bool ConstructAudioMasterTrackData(const UMovieScene* InMovieScene, const UMovieSceneAudioTrack* InAudioMasterTrack, TMap<int32, TSharedPtr<FMovieSceneExportAudioMasterTrackData>>& InAudioTrackMap);
 	/** Loads intermediate audio track data from Sequencer. */
-	bool ConstructAudioTrackData(const UMovieScene* InMovieScene, TSharedPtr<FMovieSceneExportMovieSceneData> InMovieSceneData, const UMovieSceneAudioTrack* InAudioTrack);
+	bool ConstructAudioTrackData(const UMovieScene* InMovieScene, TSharedPtr<FMovieSceneExportAudioMasterTrackData> InAudioMasterTrack, int32 InRowIndex);
 	/** Loads intermediate cinematic section data from Sequencer. */
 	bool ConstructCinematicSectionData(const UMovieScene* InMovieScene, TSharedPtr<FMovieSceneExportCinematicMasterTrackData> InMasterTrackData, const UMovieSceneCinematicShotSection* InCinematicSection);
 	/** Loads intermediate audio section data from Sequencer. */
-	bool ConstructAudioSectionData(const UMovieScene* InMovieScene, TSharedPtr<FMovieSceneExportAudioTrackData> InTrackData, const UMovieSceneAudioSection* InAudioSection);
+	bool ConstructAudioSectionData(const UMovieScene* InMovieScene, TSharedPtr<FMovieSceneExportAudioMasterTrackData> InTrackData, const UMovieSceneAudioSection* InAudioSection);
 	/** Loads intermediate common section data from Sequencer. */
 	bool ConstructSectionData(const UMovieScene* InMovieScene, TSharedPtr<FMovieSceneExportSectionData> InSectionData, const UMovieSceneSection* InSection, EMovieSceneTranslatorSectionType InSectionType, const FString& InSectionDisplayName);
 
@@ -180,6 +212,8 @@ private:
 	FString SaveFilename;
 	FString SaveFilenamePath;
 	bool bExportDataIsValid;
+	int32 DefaultAudioSampleRate;
+	int32 DefaultAudioDepth;
 };
 
 struct FMovieSceneImportCinematicSectionData 
@@ -190,25 +224,53 @@ struct FMovieSceneImportCinematicSectionData
 struct FMovieSceneImportAudioSectionData
 {
 	UMovieSceneAudioSection* AudioSection;
+	FString SourceFilename;
+	FString SourceFilePath;
+};
+
+struct FMovieSceneImportCinematicTrackData
+{
+	int32 RowIndex;
+	TArray< TSharedPtr<FMovieSceneImportCinematicSectionData> > CinematicSections;
 };
 
 struct FMovieSceneImportAudioTrackData
 {
-	UMovieSceneAudioTrack* AudioTrack;
-	TArray< TSharedPtr<FMovieSceneImportAudioSectionData> > AudioSectionsData;
+	int32 RowIndex;
+	TArray< TSharedPtr<FMovieSceneImportAudioSectionData> > AudioSections;
+};
+
+struct FMovieSceneImportMasterTrackData
+{
+	UMovieSceneTrack* MovieSceneTrack;
+};
+        
+struct FMovieSceneImportAudioMasterTrackData : public FMovieSceneImportMasterTrackData
+{
+	/** Array of all sections in order they appear in UMovieSceneAudioTrack*/
+	TArray< TSharedPtr<FMovieSceneImportAudioSectionData> > AudioSections;
+
+	/** Array of sorted audio sub tracks, containing pointers to sections within the sub track row*/
+	TArray< TSharedPtr<FMovieSceneImportAudioTrackData> > AudioTracks;
+
+	/** Max row index existing in this master track */
+	int32 MaxRowIndex;
 };
 
-struct FMovieSceneImportCinematicMasterTrackData
+struct FMovieSceneImportCinematicMasterTrackData : public FMovieSceneImportMasterTrackData
 {
-	UMovieSceneCinematicShotTrack* CinematicMasterTrack;
-	TArray< TSharedPtr<FMovieSceneImportCinematicSectionData> > CinematicSectionsData;
+	/** Array of all sections in order they appear in UMovieSceneCinematicTrack*/
+	TArray< TSharedPtr<FMovieSceneImportCinematicSectionData> > CinematicSections;
+
+	/** Array of sorted movie sub tracks, containing pointers to sections within the sub track row */
+	TArray< TSharedPtr<FMovieSceneImportCinematicTrackData> > CinematicTracks;
 };
 
 struct FMovieSceneImportMovieSceneData
 {
 	UMovieScene* MovieScene;
-	TSharedPtr<FMovieSceneImportCinematicMasterTrackData> CinematicMasterTrackData;
-	TArray< TSharedPtr<FMovieSceneImportAudioTrackData> > AudioTracksData;
+	TSharedPtr<FMovieSceneImportCinematicMasterTrackData> CinematicMasterTrack;
+	TArray< TSharedPtr<FMovieSceneImportAudioMasterTrackData> > AudioMasterTracks;
 };
 
 /**
@@ -230,13 +292,22 @@ public:
 	/** Returns the cinematic master track data pointer or nullptr if one does not exist */
 	TSharedPtr<FMovieSceneImportCinematicMasterTrackData> GetCinematicMasterTrackData(bool CreateTrackIfNull);
 	/** Find cinematic section */
-	TSharedPtr<FMovieSceneImportCinematicSectionData> FindCinematicSection(const FString& InName);
+	TSharedPtr<FMovieSceneImportCinematicSectionData> FindCinematicSection(const FString& InSectionPathName);
 	/** Create cinematic section */
 	TSharedPtr<FMovieSceneImportCinematicSectionData> CreateCinematicSection(FString InName, int32 InRow, FFrameRate InFrameRate, FFrameNumber InStartFrame, FFrameNumber InEndFrame, FFrameNumber InStartOffsetFrame);
 	/** Set cinematic section */
 	bool SetCinematicSection(TSharedPtr<FMovieSceneImportCinematicSectionData> InSection, int32 InRow, FFrameRate InFrameRate, FFrameNumber InStartFrame, FFrameNumber InEndFrame, FFrameNumber InStartOffsetFrame);
-	/** Finds audio section */
-	TSharedPtr<FMovieSceneImportAudioSectionData> FindAudioSection(FString InName) const;
+
+	/** Returns the audio master track data pointer or nullptr if one does not exist */
+	TSharedPtr<FMovieSceneImportAudioMasterTrackData> GetAudioMasterTrackData();
+	/** Find audio sections */
+	TSharedPtr<FMovieSceneImportAudioSectionData>  FindAudioSection(const FString& InSectionPathName, TSharedPtr<FMovieSceneImportAudioMasterTrackData>& OutMasterTrackData);
+	/** Create audio section */
+	TSharedPtr<FMovieSceneImportAudioSectionData> CreateAudioSection(FString InFilenameOrAssetPathName, bool bIsPathName, TSharedPtr<FMovieSceneImportAudioMasterTrackData> InMasterTrack, int32 InRow, FFrameRate InFrameRate, FFrameNumber InStartFrame, FFrameNumber InEndFrame, FFrameNumber InStartOffsetFrame);
+	/** Set audio section */
+	bool SetAudioSection(TSharedPtr<FMovieSceneImportAudioSectionData> InSection, int32 InRow, FFrameRate InFrameRate, FFrameNumber InStartFrame, FFrameNumber InEndFrame, FFrameNumber InStartOffsetFrame);
+	/** Move audio section */
+	bool MoveAudioSection(TSharedPtr<FMovieSceneImportAudioSectionData> InAudioSectionData, TSharedPtr<FMovieSceneImportAudioMasterTrackData> InFromMasterTrackData, TSharedPtr<FMovieSceneImportAudioMasterTrackData> InToMasterTrackData, int32 InToRowIndex);
 
 private:
 
@@ -244,8 +315,12 @@ private:
 	TSharedPtr<FMovieSceneImportMovieSceneData> ConstructMovieSceneData(UMovieScene* InMovieScene);
 	/** Gets cinematic master track data from Sequencer. */
 	TSharedPtr<FMovieSceneImportCinematicMasterTrackData> ConstructCinematicMasterTrackData(UMovieSceneCinematicShotTrack* InCinematicMasterTrack);
+	/** Gets cinematic track data from Sequencer. */
+	TSharedPtr<FMovieSceneImportCinematicTrackData> ConstructCinematicTrackData(UMovieSceneCinematicShotTrack* InCinematicMasterTrack, int32 InRowIndex);
+	/** Gets audio master track data from Sequencer. */
+	TSharedPtr<FMovieSceneImportAudioMasterTrackData> ConstructAudioMasterTrackData(UMovieSceneAudioTrack* InAudioMasterTrack);
 	/** Gets audio track data from Sequencer. */
-	TSharedPtr<FMovieSceneImportAudioTrackData> ConstructAudioTrackData(UMovieSceneAudioTrack* InAudioTrack);
+	TSharedPtr<FMovieSceneImportAudioTrackData> ConstructAudioTrackData(UMovieSceneAudioTrack* InAudioMasterTrack, int32 InRowIndex);
 	/** Gets cinematic section data from Sequencer. */
 	TSharedPtr<FMovieSceneImportCinematicSectionData> ConstructCinematicSectionData(UMovieSceneCinematicShotSection* InCinematicSection);
 	/** Gets audio section data from Sequencer. */
