Sequencer: Fix CIS

diff --git a/Engine/Source/Editor/MovieSceneTools/Private/MovieSceneTranslator.cpp b/Engine/Source/Editor/MovieSceneTools/Private/MovieSceneTranslator.cpp
--- a/Engine/Source/Editor/MovieSceneTools/Private/MovieSceneTranslator.cpp
+++ b/Engine/Source/Editor/MovieSceneTools/Private/MovieSceneTranslator.cpp
@@ -1208,14 +1208,14 @@ TSharedPtr<FMovieSceneImportAudioSectionData> FMovieSceneImportData::ConstructAu
 	USoundBase* SoundBase = InAudioSection->GetSound();
 	if (SoundBase == nullptr || !SoundBase->IsA<USoundWave>())
 	{
-		return false;
+		return nullptr;
 	}
 
 	USoundWave* SoundWave = nullptr;
 	SoundWave = Cast<USoundWave>(SoundBase);
 	if (SoundWave == nullptr || SoundWave->AssetImportData == nullptr)
 	{
-		return false;
+		return nullptr;
 	}
 
 	TSharedPtr<FMovieSceneImportAudioSectionData> SectionData = MakeShared<FMovieSceneImportAudioSectionData>();
@@ -1224,7 +1224,7 @@ TSharedPtr<FMovieSceneImportAudioSectionData> FMovieSceneImportData::ConstructAu
 	TArray<FString> Filenames = SoundWave->AssetImportData->ExtractFilenames();
 	if (Filenames.Num() < 1)
 	{
-		return false;
+		return nullptr;
 	}
 
 	SectionData->SourceFilename = FPaths::GetCleanFilename(Filenames[0]);
