Sequencer: Fix some erroneous cases in import fbx. In particular, this fixes an issue where importing an fbx onto one node would behave differently if there are multiple nodes in the fbx since it was ambiguous as to which fbx node to map onto the single node in sequencer.

- Changed import operation to be more explicit - when invoking through the general import fbx menu, all object bindings are mapped onto. When invoking through the context menu, only the selected object bindings are mapped onto.
- Added a user toggle to match by name only. The default is true.
- When importing onto a selected node, disable the creation of cameras.
- Don't set range bounds on imported keys since external packages, ie. maya, most likely only have 1 equivalent infinite section range.
- Remove any existing sections/keys when importing onto existing tracks.

diff --git a/Engine/Source/Editor/MovieSceneTools/Private/MovieSceneToolHelpers.cpp b/Engine/Source/Editor/MovieSceneTools/Private/MovieSceneToolHelpers.cpp
--- a/Engine/Source/Editor/MovieSceneTools/Private/MovieSceneToolHelpers.cpp
+++ b/Engine/Source/Editor/MovieSceneTools/Private/MovieSceneToolHelpers.cpp
@@ -883,14 +883,15 @@ bool ImportFBXProperty(FString NodeName, FString AnimatedPropertyName, FGuid Obj
 			UMovieSceneFloatTrack* FloatTrack = InMovieScene->FindTrack<UMovieSceneFloatTrack>(PropertyOwnerGuid, *FbxSetting.PropertyPath.PropertyName);
 			if (!FloatTrack)
 			{
-				FString PropertyPath = FbxSetting.PropertyPath.ComponentName + TEXT(".") + FbxSetting.PropertyPath.PropertyName;
 				InMovieScene->Modify();
 				FloatTrack = InMovieScene->AddTrack<UMovieSceneFloatTrack>(PropertyOwnerGuid);
-				FloatTrack->SetPropertyNameAndPath(*FbxSetting.PropertyPath.PropertyName, PropertyPath);
+				FloatTrack->SetPropertyNameAndPath(*FbxSetting.PropertyPath.PropertyName, *FbxSetting.PropertyPath.PropertyName);
 			}
 
 			if (FloatTrack)
 			{
+				FloatTrack->RemoveAllAnimationData();
+
 				FFrameRate FrameRate = FloatTrack->GetTypedOuter<UMovieScene>()->GetTickResolution();
 
 				bool bSectionAdded = false;
@@ -945,11 +946,6 @@ bool ImportFBXProperty(FString NodeName, FString AnimatedPropertyName, FGuid Obj
 				}
 				Channel->AutoSetTangents();
 
-				TArrayView<const FFrameNumber> AllTimes = ChannelData.GetTimes();
-				if (AllTimes.Num() > 0)
-				{
-					FloatSection->SetRange(TRange<FFrameNumber>(AllTimes[0], TRangeBound<FFrameNumber>::Inclusive(AllTimes.Last())));
-				}
 				return true;
 			}
 		}
@@ -1013,6 +1009,7 @@ bool ImportFBXTransform(FString NodeName, FGuid ObjectBinding, UnFbx::FFbxCurves
 		InMovieScene->Modify();
 		TransformTrack = InMovieScene->AddTrack<UMovieScene3DTransformTrack>(ObjectBinding);
 	}
+	TransformTrack->RemoveAllAnimationData();
 
 	bool bSectionAdded = false;
 	UMovieScene3DTransformSection* TransformSection = Cast<UMovieScene3DTransformSection>(TransformTrack->FindOrAddSection(0, bSectionAdded));
@@ -1059,35 +1056,16 @@ bool ImportFBXTransform(FString NodeName, FGuid ObjectBinding, UnFbx::FFbxCurves
 	ImportTransformChannel(Scale[1],         Channels[7], FrameRate, false);
 	ImportTransformChannel(Scale[2],         Channels[8], FrameRate, false);
 
-	TOptional<TRangeBound<FFrameNumber>> MinLower;
-	TOptional<TRangeBound<FFrameNumber>> MaxUpper;
-
-	for (int32 ChannelIndex = 0; ChannelIndex < 9; ++ChannelIndex)
-	{
-		TArrayView<const FFrameNumber> AllTimes = Channels[ChannelIndex]->GetTimes();
-		if (AllTimes.Num() > 0)
-		{
-			MinLower = TRangeBound<FFrameNumber>::MinLower(AllTimes[0],     MinLower.Get(AllTimes[0]));
-			MaxUpper = TRangeBound<FFrameNumber>::MaxUpper(AllTimes.Last(), MaxUpper.Get(AllTimes.Last()));
-		}
-	}
-
-	// MinLower and MinUpper are both either set or unset
-	if (MinLower.IsSet())
-	{
-		TransformSection->SetRange(TRange<FFrameNumber>(MinLower.GetValue(), MaxUpper.GetValue()));
-	}
-
 	return true;
 }
 
-bool ImportFBXNode(FString NodeName, UnFbx::FFbxCurvesAPI& CurveAPI, UMovieScene* InMovieScene, ISequencer& InSequencer, const TMap<FGuid, FString>& InObjectBindingMap)
+bool ImportFBXNode(FString NodeName, UnFbx::FFbxCurvesAPI& CurveAPI, UMovieScene* InMovieScene, ISequencer& InSequencer, const TMap<FGuid, FString>& InObjectBindingMap, bool bMatchByNameOnly)
 {
-	// Find the matching object binding to apply this animation to. Defaults to the first.
+	// Find the matching object binding to apply this animation to. If not matching by name only, default to the first.
 	FGuid ObjectBinding;
 	for (auto It = InObjectBindingMap.CreateConstIterator(); It; ++It)
 	{
-		if (InObjectBindingMap.Num() == 1 || FCString::Strcmp(*It.Value().ToUpper(), *NodeName.ToUpper()) == 0)
+		if (!bMatchByNameOnly || FCString::Strcmp(*It.Value().ToUpper(), *NodeName.ToUpper()) == 0)
 		{
 			ObjectBinding = It.Key();
 			break;
@@ -1096,7 +1074,7 @@ bool ImportFBXNode(FString NodeName, UnFbx::FFbxCurvesAPI& CurveAPI, UMovieScene
 
 	if (!ObjectBinding.IsValid())
 	{
-		//@todo output warning
+		UE_LOG(LogMovieScene, Warning, TEXT("Fbx Import: Failed to find any matching node for (%s)."), *NodeName);
 		return false;
 	}
 
@@ -1215,7 +1193,7 @@ void CopyCameraProperties(FbxCamera* CameraNode, ACineCameraActor* CameraActor)
 	CineCameraComponent->CurrentFocalLength = FocalLength;
 }
 
-void ImportFBXCamera(UnFbx::FFbxImporter* FbxImporter, UMovieScene* InMovieScene, ISequencer& InSequencer, TMap<FGuid, FString>& InObjectBindingMap, bool bCreateCameras)
+void ImportFBXCamera(UnFbx::FFbxImporter* FbxImporter, UMovieScene* InMovieScene, ISequencer& InSequencer, TMap<FGuid, FString>& InObjectBindingMap, bool bMatchByNameOnly, bool bCreateCameras)
 {
 	if (bCreateCameras)
 	{
@@ -1293,21 +1271,17 @@ void ImportFBXCamera(UnFbx::FFbxImporter* FbxImporter, UMovieScene* InMovieScene
 
 		if (!CameraNode)
 		{
-			CameraNode = FindCamera(FbxImporter->Scene->GetRootNode());
-			if (CameraNode)
-			{
-				UE_LOG(LogMovieScene, Warning, TEXT("Fbx Import: Failed to find exact matching camera for (%s). Using first camera from fbx (%s)"), *ObjectName, UTF8_TO_TCHAR(CameraNode->GetName()));
-			}
-			else
+			if (bMatchByNameOnly)
 			{
 				UE_LOG(LogMovieScene, Error, TEXT("Fbx Import: Failed to find any matching camera for (%s)."), *ObjectName);
 				continue;
 			}
-		}
 
-		if (!CameraNode)
-		{
-			continue;
+			CameraNode = FindCamera(FbxImporter->Scene->GetRootNode());
+			if (CameraNode)
+			{
+				UE_LOG(LogMovieScene, Warning, TEXT("Fbx Import: Failed to find exact matching camera for (%s). Using first camera from fbx (%s)"), *ObjectName, UTF8_TO_TCHAR(CameraNode->GetName()));
+			}
 		}
 
 		float FieldOfView;
@@ -1343,6 +1317,8 @@ void ImportFBXCamera(UnFbx::FFbxImporter* FbxImporter, UMovieScene* InMovieScene
 				UMovieSceneFloatTrack* FloatTrack = InMovieScene->FindTrack<UMovieSceneFloatTrack>(PropertyOwnerGuid, TEXT("CurrentFocalLength"));
 				if (FloatTrack)
 				{
+					FloatTrack->RemoveAllAnimationData();
+
 					bool bSectionAdded = false;
 					UMovieSceneFloatSection* FloatSection = Cast<UMovieSceneFloatSection>(FloatTrack->FindOrAddSection(0, bSectionAdded));
 					if (!FloatSection)
@@ -1505,6 +1481,11 @@ class SMovieSceneImportFBXSettings : public SCompoundWidget, public FGCObject
 		ObjectBindingMap = InObjectBindingMap;
 	}
 
+	void SetCreateCameras(TOptional<bool> bInCreateCameras)
+	{
+		bCreateCameras = bInCreateCameras;
+	}
+
 private:
 
 	FReply OnImportFBXClicked()
@@ -1520,7 +1501,6 @@ private:
 		bool bConvertSceneUnitBackup = ImportOptions->bConvertSceneUnit;
 		bool bForceFrontXAxisBackup = ImportOptions->bForceFrontXAxis;
 
-
 		ImportOptions->bConvertScene = true;
 		ImportOptions->bConvertSceneUnit = true;
 		ImportOptions->bForceFrontXAxis = ImportFBXSettings->bForceFrontXAxis;
@@ -1535,12 +1515,13 @@ private:
 			ImportOptions->bForceFrontXAxis = bForceFrontXAxisBackup;
 			return FReply::Unhandled();
 		}
+		
+		const bool bMatchByNameOnly = ImportFBXSettings->bMatchByNameOnly;
 
 		const FScopedTransaction Transaction( NSLOCTEXT( "MovieSceneTools", "ImportFBXTransaction", "Import FBX" ) );
 
-		// Import static cameras first if not mapping to any particular object binding nodes
-		bool bCreateCameras = ImportFBXSettings->bCreateCameras && ObjectBindingMap.Num() == 0;
-		ImportFBXCamera(FbxImporter, MovieScene, *Sequencer, ObjectBindingMap, bCreateCameras);
+		// Import static cameras first
+		ImportFBXCamera(FbxImporter, MovieScene, *Sequencer, ObjectBindingMap, bMatchByNameOnly, bCreateCameras.IsSet() ? bCreateCameras.GetValue() : ImportFBXSettings->bCreateCameras);
 
 		UnFbx::FFbxCurvesAPI CurveAPI;
 		FbxImporter->PopulateAnimatedCurveData(CurveAPI);
@@ -1552,9 +1533,11 @@ private:
 
 		for (FString NodeName : AllNodeNames)
 		{
-			ImportFBXNode(NodeName, CurveAPI, MovieScene, *Sequencer, ObjectBindingMap);
+			ImportFBXNode(NodeName, CurveAPI, MovieScene, *Sequencer, ObjectBindingMap, bMatchByNameOnly);
 		}
 
+		Sequencer->NotifyMovieSceneDataChanged(EMovieSceneDataChangeType::MovieSceneStructureItemAdded);
+
 		FbxImporter->ReleaseScene();
 		ImportOptions->bConvertScene = bConvertSceneBackup;
 		ImportOptions->bConvertSceneUnit = bConvertSceneUnitBackup;
@@ -1576,10 +1559,11 @@ private:
 	UMovieScene* MovieScene;
 	ISequencer* Sequencer;
 	TMap<FGuid, FString> ObjectBindingMap;
+	TOptional<bool> bCreateCameras;
 };
 
 
-bool MovieSceneToolHelpers::ImportFBX(UMovieScene* InMovieScene, ISequencer& InSequencer, const TMap<FGuid, FString>& InObjectBindingMap)
+bool MovieSceneToolHelpers::ImportFBX(UMovieScene* InMovieScene, ISequencer& InSequencer, const TMap<FGuid, FString>& InObjectBindingMap, TOptional<bool> bCreateCameras)
 {
 	TArray<FString> OpenFilenames;
 	IDesktopPlatform* DesktopPlatform = FDesktopPlatformModule::Get();
@@ -1625,6 +1609,8 @@ bool MovieSceneToolHelpers::ImportFBX(UMovieScene* InMovieScene, ISequencer& InS
 		.MovieScene(InMovieScene)
 		.Sequencer(&InSequencer);
 	DialogWidget->SetObjectBindingMap(InObjectBindingMap);
+	DialogWidget->SetCreateCameras(bCreateCameras);
+	
 	Window->SetContent(DialogWidget);
 
 	FSlateApplication::Get().AddWindow(Window);
diff --git a/Engine/Source/Editor/MovieSceneTools/Private/MovieSceneToolsUserSettings.cpp b/Engine/Source/Editor/MovieSceneTools/Private/MovieSceneToolsUserSettings.cpp
--- a/Engine/Source/Editor/MovieSceneTools/Private/MovieSceneToolsUserSettings.cpp
+++ b/Engine/Source/Editor/MovieSceneTools/Private/MovieSceneToolsUserSettings.cpp
@@ -31,6 +31,7 @@ void UMovieSceneUserThumbnailSettings::PostEditChangeProperty(FPropertyChangedEv
 UMovieSceneUserImportFBXSettings::UMovieSceneUserImportFBXSettings(const FObjectInitializer& Initializer)
 	: Super(Initializer)
 {
+	bMatchByNameOnly = true;
 	bForceFrontXAxis = false;
 	bCreateCameras = true;
 	bReduceKeys = true;
diff --git a/Engine/Source/Editor/MovieSceneTools/Public/MovieSceneToolHelpers.h b/Engine/Source/Editor/MovieSceneTools/Public/MovieSceneToolHelpers.h
--- a/Engine/Source/Editor/MovieSceneTools/Public/MovieSceneToolHelpers.h
+++ b/Engine/Source/Editor/MovieSceneTools/Public/MovieSceneToolHelpers.h
@@ -192,9 +192,10 @@ public:
 	 *
 	 * @param InMovieScene The movie scene to import the fbx into
 	 * @param InObjectBindingNameMap The object binding to name map to map import fbx animation onto
+	 * @param bCreateCameras Whether to allow creation of cameras if found in the fbx file.
 	 * @return Whether the import was successful
 	 */
-	static bool ImportFBX(UMovieScene* InMovieScene, ISequencer& InSequencer, const TMap<FGuid, FString>& InObjectBindingNameMap);
+	static bool ImportFBX(UMovieScene* InMovieScene, ISequencer& InSequencer, const TMap<FGuid, FString>& InObjectBindingNameMap, TOptional<bool> bCreateCameras);
 
 	/*
 	 * Rich curve interpolation to matinee interpolation
diff --git a/Engine/Source/Editor/MovieSceneTools/Public/MovieSceneToolsUserSettings.h b/Engine/Source/Editor/MovieSceneTools/Public/MovieSceneToolsUserSettings.h
--- a/Engine/Source/Editor/MovieSceneTools/Public/MovieSceneToolsUserSettings.h
+++ b/Engine/Source/Editor/MovieSceneTools/Public/MovieSceneToolsUserSettings.h
@@ -57,6 +57,10 @@ public:
 	
 	GENERATED_BODY()
 
+	/** Whether to match fbx node names to sequencer node names. */
+	UPROPERTY(EditAnywhere, config, Category=Import, meta= (ToolTip = "Match fbx node names to sequencer node names"))
+	bool bMatchByNameOnly;
+
 	/** Whether to force the front axis to be align with X instead of -Y. */
 	UPROPERTY(EditAnywhere, config, Category=Import, meta= (ToolTip = "Convert the scene from FBX coordinate system to UE4 coordinate system with front X axis instead of -Y"))
 	bool bForceFrontXAxis;
diff --git a/Engine/Source/Editor/Sequencer/Private/DisplayNodes/SequencerObjectBindingNode.cpp b/Engine/Source/Editor/Sequencer/Private/DisplayNodes/SequencerObjectBindingNode.cpp
--- a/Engine/Source/Editor/Sequencer/Private/DisplayNodes/SequencerObjectBindingNode.cpp
+++ b/Engine/Source/Editor/Sequencer/Private/DisplayNodes/SequencerObjectBindingNode.cpp
@@ -207,7 +207,7 @@ void FSequencerObjectBindingNode::BuildContextMenu(FMenuBuilder& MenuBuilder)
 			LOCTEXT("ImportFBXTooltip", "Import FBX animation to this object"),
 			FSlateIcon(),
 			FUIAction(
-				FExecuteAction::CreateLambda([=]{ GetSequencer().ImportFBX(); })
+				FExecuteAction::CreateLambda([=]{ GetSequencer().ImportFBXOntoSelectedNodes(); })
 			));
 
 		MenuBuilder.AddMenuEntry(
diff --git a/Engine/Source/Editor/Sequencer/Private/Sequencer.cpp b/Engine/Source/Editor/Sequencer/Private/Sequencer.cpp
--- a/Engine/Source/Editor/Sequencer/Private/Sequencer.cpp
+++ b/Engine/Source/Editor/Sequencer/Private/Sequencer.cpp
@@ -1385,7 +1385,7 @@ void FSequencer::BakeTransform()
 			TArray<FFrameNumber> KeyTimes;
 
 			FFrameRate   Resolution  = FocusedMovieScene->GetTickResolution();
-			FFrameRate   SnapRate    = FocusedMovieScene->GetEvaluationType() == EMovieSceneEvaluationType::FrameLocked ? FocusedMovieScene->GetDisplayRate() : Resolution;
+			FFrameRate   SnapRate    = FocusedMovieScene->GetDisplayRate();
 
 			FFrameNumber InFrame     = MovieScene::DiscreteInclusiveLower(GetPlaybackRange());
 			FFrameNumber OutFrame    = MovieScene::DiscreteExclusiveUpper(GetPlaybackRange());
@@ -7480,6 +7480,25 @@ void FSequencer::ImportFBX()
 {
 	UMovieScene* MovieScene = GetFocusedMovieSceneSequence()->GetMovieScene();
 
+	TMap<FGuid, FString> ObjectBindingNameMap;
+
+	TArray<TSharedRef<FSequencerObjectBindingNode>> RootObjectBindingNodes;
+	GetRootObjectBindingNodes( NodeTree->GetRootNodes(), RootObjectBindingNodes );
+
+	for (auto RootObjectBindingNode : RootObjectBindingNodes)
+	{
+		FGuid ObjectBinding = RootObjectBindingNode.Get().GetObjectBinding();
+
+		ObjectBindingNameMap.Add(ObjectBinding, RootObjectBindingNode.Get().GetDisplayName().ToString());
+	}
+
+	MovieSceneToolHelpers::ImportFBX(MovieScene, *this, ObjectBindingNameMap, TOptional<bool>());
+}
+
+void FSequencer::ImportFBXOntoSelectedNodes()
+{
+	UMovieScene* MovieScene = GetFocusedMovieSceneSequence()->GetMovieScene();
+
 	// The object binding and names to match when importing from fbx
 	TMap<FGuid, FString> ObjectBindingNameMap;
 
@@ -7495,21 +7514,7 @@ void FSequencer::ImportFBX()
 		}
 	}
 
-	// If nothing selected, try to map onto everything
-	if (ObjectBindingNameMap.Num() == 0)
-	{
-		TArray<TSharedRef<FSequencerObjectBindingNode>> RootObjectBindingNodes;
-		GetRootObjectBindingNodes( NodeTree->GetRootNodes(), RootObjectBindingNodes );
-
-		for (auto RootObjectBindingNode : RootObjectBindingNodes)
-		{
-			FGuid ObjectBinding = RootObjectBindingNode.Get().GetObjectBinding();
-
-			ObjectBindingNameMap.Add(ObjectBinding, RootObjectBindingNode.Get().GetDisplayName().ToString());
-		}
-	}
-	
-	MovieSceneToolHelpers::ImportFBX(MovieScene, *this, ObjectBindingNameMap);
+	MovieSceneToolHelpers::ImportFBX(MovieScene, *this, ObjectBindingNameMap, TOptional<bool>(false));
 }
 
 
diff --git a/Engine/Source/Editor/Sequencer/Private/Sequencer.h b/Engine/Source/Editor/Sequencer/Private/Sequencer.h
--- a/Engine/Source/Editor/Sequencer/Private/Sequencer.h
+++ b/Engine/Source/Editor/Sequencer/Private/Sequencer.h
@@ -555,6 +555,7 @@ public:
 
 	/** Imports the animation from an fbx file. */
 	void ImportFBX();
+	void ImportFBXOntoSelectedNodes();
 
 	/** Exports the animation to an fbx file. */
 	void ExportFBX();
