Sequencer: Fix dereferencing null pointer - CameraNode

diff --git a/Engine/Source/Editor/MovieSceneTools/Private/MovieSceneToolHelpers.cpp b/Engine/Source/Editor/MovieSceneTools/Private/MovieSceneToolHelpers.cpp
--- a/Engine/Source/Editor/MovieSceneTools/Private/MovieSceneToolHelpers.cpp
+++ b/Engine/Source/Editor/MovieSceneTools/Private/MovieSceneToolHelpers.cpp
@@ -1284,6 +1284,11 @@ void ImportFBXCamera(UnFbx::FFbxImporter* FbxImporter, UMovieScene* InMovieScene
 			}
 		}
 
+		if (!CameraNode)
+		{
+			continue;
+		}
+
 		float FieldOfView;
 		float FocalLength;
 
