Do not import a camera track if we aren't creating cameras or importing on selected objects

diff --git a/Engine/Source/Editor/MovieSceneTools/Private/MovieSceneToolHelpers.cpp b/Engine/Source/Editor/MovieSceneTools/Private/MovieSceneToolHelpers.cpp
--- a/Engine/Source/Editor/MovieSceneTools/Private/MovieSceneToolHelpers.cpp
+++ b/Engine/Source/Editor/MovieSceneTools/Private/MovieSceneToolHelpers.cpp
@@ -1193,8 +1193,9 @@ void CopyCameraProperties(FbxCamera* CameraNode, ACineCameraActor* CameraActor)
 	CineCameraComponent->CurrentFocalLength = FocalLength;
 }
 
-void ImportFBXCamera(UnFbx::FFbxImporter* FbxImporter, UMovieScene* InMovieScene, ISequencer& InSequencer, TMap<FGuid, FString>& InObjectBindingMap, bool bMatchByNameOnly, bool bCreateCameras)
+bool ImportFBXCamera(UnFbx::FFbxImporter* FbxImporter, UMovieScene* InMovieScene, ISequencer& InSequencer, TMap<FGuid, FString>& InObjectBindingMap, bool bMatchByNameOnly, bool bCreateCameras)
 {
+	bool bCamsCreated = false;
 	if (bCreateCameras)
 	{
 		TArray<FbxCamera*> AllCameras;
@@ -1230,6 +1231,7 @@ void ImportFBXCamera(UnFbx::FFbxImporter* FbxImporter, UMovieScene* InMovieScene
 		if (UnmatchedCameras.Num() != 0)
 		{
 			InObjectBindingMap.Reset();
+			bCamsCreated = true;
 		}
 
 		for (auto UnmatchedCamera : UnmatchedCameras)
@@ -1343,6 +1345,7 @@ void ImportFBXCamera(UnFbx::FFbxImporter* FbxImporter, UMovieScene* InMovieScene
 			}
 		}
 	}
+	return bCamsCreated;
 }
 
 FGuid FindCameraGuid(FbxCamera* Camera, TMap<FGuid, FString>& InObjectBindingMap)
@@ -1371,12 +1374,6 @@ UMovieSceneCameraCutTrack* GetCameraCutTrack(UMovieScene* InMovieScene)
 
 void ImportCameraCut(UnFbx::FFbxImporter* FbxImporter, UMovieScene* InMovieScene, ISequencer& InSequencer, TMap<FGuid, FString>& InObjectBindingMap)
 {
-	// The camera switcher refers to the fbx camera
-	const UMovieSceneUserImportFBXSettings* ImportFBXSettings = GetDefault<UMovieSceneUserImportFBXSettings>();
-	if (!ImportFBXSettings->bCreateCameras)
-	{
-		return;
-	}
 	// Find a camera switcher
 	FbxCameraSwitcher* CameraSwitcher = FbxImporter->Scene->GlobalCameraSettings().GetCameraSwitcher();
 	if (CameraSwitcher == nullptr)
@@ -1526,15 +1523,18 @@ private:
 		const FScopedTransaction Transaction( NSLOCTEXT( "MovieSceneTools", "ImportFBXTransaction", "Import FBX" ) );
 
 		// Import static cameras first
-		ImportFBXCamera(FbxImporter, MovieScene, *Sequencer, ObjectBindingMap, bMatchByNameOnly, bCreateCameras.IsSet() ? bCreateCameras.GetValue() : ImportFBXSettings->bCreateCameras);
+		bool bCamsCreated = ImportFBXCamera(FbxImporter, MovieScene, *Sequencer, ObjectBindingMap, bMatchByNameOnly, bCreateCameras.IsSet() ? bCreateCameras.GetValue() : ImportFBXSettings->bCreateCameras);
 
 		UnFbx::FFbxCurvesAPI CurveAPI;
 		FbxImporter->PopulateAnimatedCurveData(CurveAPI);
 		TArray<FString> AllNodeNames;
 		CurveAPI.GetAllNodeNameArray(AllNodeNames);
 
-		// Import a camera cut track, do it after populating curve data ensure only one animation layer, if any
-		ImportCameraCut(FbxImporter, MovieScene, *Sequencer, ObjectBindingMap);
+		// Import a camera cut track if cams were created, do it after populating curve data ensure only one animation layer, if any
+		if (bCamsCreated)
+		{
+			ImportCameraCut(FbxImporter, MovieScene, *Sequencer, ObjectBindingMap);
+		}
 
 		for (FString NodeName : AllNodeNames)
 		{
