Movie Scene Capture: Fix to use automated level sequence capture ui name so that it's the same as the dialog settings.

diff --git a/Engine/Source/Editor/MovieSceneTools/Private/TrackEditors/CinematicShotTrackEditor.cpp b/Engine/Source/Editor/MovieSceneTools/Private/TrackEditors/CinematicShotTrackEditor.cpp
--- a/Engine/Source/Editor/MovieSceneTools/Private/TrackEditors/CinematicShotTrackEditor.cpp
+++ b/Engine/Source/Editor/MovieSceneTools/Private/TrackEditors/CinematicShotTrackEditor.cpp
@@ -822,7 +822,7 @@ void FCinematicShotTrackEditor::ImportEDL()
 	UAutomatedLevelSequenceCapture* MovieSceneCapture = Cast<UAutomatedLevelSequenceCapture>(IMovieSceneCaptureModule::Get().GetFirstActiveMovieSceneCapture());
 	if (!MovieSceneCapture)
 	{
-		MovieSceneCapture = NewObject<UAutomatedLevelSequenceCapture>(GetTransientPackage(), UAutomatedLevelSequenceCapture::StaticClass(), NAME_None, RF_Transient);
+		MovieSceneCapture = NewObject<UAutomatedLevelSequenceCapture>(GetTransientPackage(), UAutomatedLevelSequenceCapture::StaticClass(), UAutomatedLevelSequenceCapture::AutomatedLevelSequenceCaptureUIName, RF_Transient);
 		MovieSceneCapture->LoadFromConfig();
 	}
 
@@ -892,7 +892,7 @@ void FCinematicShotTrackEditor::ImportFCPXML()
 	UAutomatedLevelSequenceCapture* MovieSceneCapture = Cast<UAutomatedLevelSequenceCapture>(IMovieSceneCaptureModule::Get().GetFirstActiveMovieSceneCapture());
 	if (!MovieSceneCapture)
 	{
-		MovieSceneCapture = NewObject<UAutomatedLevelSequenceCapture>(GetTransientPackage(), UAutomatedLevelSequenceCapture::StaticClass(), NAME_None, RF_Transient);
+		MovieSceneCapture = NewObject<UAutomatedLevelSequenceCapture>(GetTransientPackage(), UAutomatedLevelSequenceCapture::StaticClass(), UAutomatedLevelSequenceCapture::AutomatedLevelSequenceCaptureUIName, RF_Transient);
 		MovieSceneCapture->LoadFromConfig();
 	}
 
@@ -932,7 +932,7 @@ void FCinematicShotTrackEditor::ExportFCPXML()
 	UAutomatedLevelSequenceCapture* MovieSceneCapture = Cast<UAutomatedLevelSequenceCapture>(IMovieSceneCaptureModule::Get().GetFirstActiveMovieSceneCapture());
 	if (!MovieSceneCapture)
 	{
-		MovieSceneCapture = NewObject<UAutomatedLevelSequenceCapture>(GetTransientPackage(), UAutomatedLevelSequenceCapture::StaticClass(), NAME_None, RF_Transient);
+		MovieSceneCapture = NewObject<UAutomatedLevelSequenceCapture>(GetTransientPackage(), UAutomatedLevelSequenceCapture::StaticClass(), UAutomatedLevelSequenceCapture::AutomatedLevelSequenceCaptureUIName, RF_Transient);
 		MovieSceneCapture->LoadFromConfig();
 	}
 
