Sequencer: Fix crash in import/export edl/fcpxml dialogs which was creating another automated level sequence capture object and nullifying the one created by the movie scene capture dialog.

diff --git a/Engine/Source/Editor/MovieSceneTools/Private/TrackEditors/CinematicShotTrackEditor.cpp b/Engine/Source/Editor/MovieSceneTools/Private/TrackEditors/CinematicShotTrackEditor.cpp
--- a/Engine/Source/Editor/MovieSceneTools/Private/TrackEditors/CinematicShotTrackEditor.cpp
+++ b/Engine/Source/Editor/MovieSceneTools/Private/TrackEditors/CinematicShotTrackEditor.cpp
@@ -804,6 +804,22 @@ FKeyPropertyResult FCinematicShotTrackEditor::HandleSequenceAdded(FFrameNumber K
 	return KeyPropertyResult;
 }
 
+UAutomatedLevelSequenceCapture* GetMovieSceneCapture()
+{
+	UAutomatedLevelSequenceCapture* MovieSceneCapture = Cast<UAutomatedLevelSequenceCapture>(IMovieSceneCaptureModule::Get().GetFirstActiveMovieSceneCapture());
+	if (!MovieSceneCapture)
+	{
+		MovieSceneCapture = FindObject<UAutomatedLevelSequenceCapture>(GetTransientPackage(), *UAutomatedLevelSequenceCapture::AutomatedLevelSequenceCaptureUIName.ToString());
+	}
+	
+	if (!MovieSceneCapture)
+	{
+		MovieSceneCapture = NewObject<UAutomatedLevelSequenceCapture>(GetTransientPackage(), UAutomatedLevelSequenceCapture::StaticClass(), UAutomatedLevelSequenceCapture::AutomatedLevelSequenceCaptureUIName, RF_Transient);
+		MovieSceneCapture->LoadFromConfig();
+	}
+
+	return MovieSceneCapture;
+}
 
 void FCinematicShotTrackEditor::ImportEDL()
 {
@@ -819,13 +835,7 @@ void FCinematicShotTrackEditor::ImportEDL()
 		return;
 	}
 
-	UAutomatedLevelSequenceCapture* MovieSceneCapture = Cast<UAutomatedLevelSequenceCapture>(IMovieSceneCaptureModule::Get().GetFirstActiveMovieSceneCapture());
-	if (!MovieSceneCapture)
-	{
-		MovieSceneCapture = NewObject<UAutomatedLevelSequenceCapture>(GetTransientPackage(), UAutomatedLevelSequenceCapture::StaticClass(), UAutomatedLevelSequenceCapture::AutomatedLevelSequenceCaptureUIName, RF_Transient);
-		MovieSceneCapture->LoadFromConfig();
-	}
-
+	UAutomatedLevelSequenceCapture* MovieSceneCapture = GetMovieSceneCapture();
 	if (!MovieSceneCapture)
 	{
 		return;
@@ -855,13 +865,7 @@ void FCinematicShotTrackEditor::ExportEDL()
 		return;
 	}
 		
-	UAutomatedLevelSequenceCapture* MovieSceneCapture = Cast<UAutomatedLevelSequenceCapture>(IMovieSceneCaptureModule::Get().GetFirstActiveMovieSceneCapture());
-	if (!MovieSceneCapture)
-	{
-		MovieSceneCapture = NewObject<UAutomatedLevelSequenceCapture>(GetTransientPackage(), UAutomatedLevelSequenceCapture::StaticClass(), UAutomatedLevelSequenceCapture::AutomatedLevelSequenceCaptureUIName, RF_Transient);
-		MovieSceneCapture->LoadFromConfig();
-	}
-
+	UAutomatedLevelSequenceCapture* MovieSceneCapture = GetMovieSceneCapture();
 	if (!MovieSceneCapture)
 	{
 		return;
@@ -889,13 +893,7 @@ void FCinematicShotTrackEditor::ImportFCPXML()
 		return;
 	}
 
-	UAutomatedLevelSequenceCapture* MovieSceneCapture = Cast<UAutomatedLevelSequenceCapture>(IMovieSceneCaptureModule::Get().GetFirstActiveMovieSceneCapture());
-	if (!MovieSceneCapture)
-	{
-		MovieSceneCapture = NewObject<UAutomatedLevelSequenceCapture>(GetTransientPackage(), UAutomatedLevelSequenceCapture::StaticClass(), UAutomatedLevelSequenceCapture::AutomatedLevelSequenceCaptureUIName, RF_Transient);
-		MovieSceneCapture->LoadFromConfig();
-	}
-
+	UAutomatedLevelSequenceCapture* MovieSceneCapture = GetMovieSceneCapture();
 	if (!MovieSceneCapture)
 	{
 		return;
@@ -929,13 +927,7 @@ void FCinematicShotTrackEditor::ExportFCPXML()
 		return;
 	}
 
-	UAutomatedLevelSequenceCapture* MovieSceneCapture = Cast<UAutomatedLevelSequenceCapture>(IMovieSceneCaptureModule::Get().GetFirstActiveMovieSceneCapture());
-	if (!MovieSceneCapture)
-	{
-		MovieSceneCapture = NewObject<UAutomatedLevelSequenceCapture>(GetTransientPackage(), UAutomatedLevelSequenceCapture::StaticClass(), UAutomatedLevelSequenceCapture::AutomatedLevelSequenceCaptureUIName, RF_Transient);
-		MovieSceneCapture->LoadFromConfig();
-	}
-
+	UAutomatedLevelSequenceCapture* MovieSceneCapture = GetMovieSceneCapture();
 	if (!MovieSceneCapture)
 	{
 		return;
