Sequencer: Fix crash on undo.

Add transaction for "Add Record New Sequence"

diff --git a/Engine/Source/Editor/MovieSceneTools/Private/TrackEditors/SubTrackEditor.cpp b/Engine/Source/Editor/MovieSceneTools/Private/TrackEditors/SubTrackEditor.cpp
--- a/Engine/Source/Editor/MovieSceneTools/Private/TrackEditors/SubTrackEditor.cpp
+++ b/Engine/Source/Editor/MovieSceneTools/Private/TrackEditors/SubTrackEditor.cpp
@@ -738,6 +738,8 @@ void FSubTrackEditor::HandleRecordNewSequence(AActor* InActorToRecord, UMovieSce
 {
 	FSlateApplication::Get().DismissAllMenus();
 
+	const FScopedTransaction Transaction(LOCTEXT("AddRecordNewSequence_Transaction", "Add Record New Sequence"));
+
 	AnimatablePropertyChanged( FOnKeyProperty::CreateRaw( this, &FSubTrackEditor::HandleRecordNewSequenceInternal, InActorToRecord, InTrack) );
 }
 
diff --git a/Engine/Source/Runtime/MovieScene/Private/Tracks/MovieSceneSubTrack.cpp b/Engine/Source/Runtime/MovieScene/Private/Tracks/MovieSceneSubTrack.cpp
--- a/Engine/Source/Runtime/MovieScene/Private/Tracks/MovieSceneSubTrack.cpp
+++ b/Engine/Source/Runtime/MovieScene/Private/Tracks/MovieSceneSubTrack.cpp
@@ -73,6 +73,8 @@ UMovieSceneSubSection* UMovieSceneSubTrack::AddSequenceOnRow(UMovieSceneSequence
 
 UMovieSceneSubSection* UMovieSceneSubTrack::AddSequenceToRecord()
 {
+	Modify();
+
 	UMovieScene* MovieScene = CastChecked<UMovieScene>(GetOuter());
 	TRange<FFrameNumber> PlaybackRange = MovieScene->GetPlaybackRange();
 
