Sequencer: Convert attach bindings to local space so that they can be editable from a master/parent sequence

diff --git a/Engine/Source/Editor/MovieSceneTools/Private/TrackEditors/AttachTrackEditor.cpp b/Engine/Source/Editor/MovieSceneTools/Private/TrackEditors/AttachTrackEditor.cpp
--- a/Engine/Source/Editor/MovieSceneTools/Private/TrackEditors/AttachTrackEditor.cpp
+++ b/Engine/Source/Editor/MovieSceneTools/Private/TrackEditors/AttachTrackEditor.cpp
@@ -10,6 +10,7 @@
 #include "ActorEditorUtils.h"
 #include "MovieSceneObjectBindingIDPicker.h"
 #include "MovieSceneToolHelpers.h"
+#include "Evaluation/MovieSceneEvaluationTemplateInstance.h"
 
 
 #define LOCTEXT_NAMESPACE "F3DAttachTrackEditor"
@@ -41,7 +42,15 @@ public:
 			TSharedPtr<ISequencer> Sequencer = AttachTrackEditor->GetSequencer();
 			if (Sequencer.IsValid())
 			{
-				TArrayView<TWeakObjectPtr<UObject>> RuntimeObjects = Sequencer->FindBoundObjects(AttachSection->GetConstraintBindingID().GetGuid(), AttachSection->GetConstraintBindingID().GetSequenceID());
+				FMovieSceneSequenceID SequenceID = Sequencer->GetFocusedTemplateID();
+				if (AttachSection->GetConstraintBindingID().GetSequenceID().IsValid())
+				{
+					// Ensure that this ID is resolvable from the root, based on the current local sequence ID
+					FMovieSceneObjectBindingID RootBindingID = AttachSection->GetConstraintBindingID().ResolveLocalToRoot(SequenceID, Sequencer->GetEvaluationTemplate().GetHierarchy());
+					SequenceID = RootBindingID.GetSequenceID();
+				}
+
+				TArrayView<TWeakObjectPtr<UObject>> RuntimeObjects = Sequencer->FindBoundObjects(AttachSection->GetConstraintBindingID().GetGuid(), SequenceID);
 				if (RuntimeObjects.Num() == 1 && RuntimeObjects[0].IsValid())
 				{
 					if (AActor* Actor = Cast<AActor>(RuntimeObjects[0].Get()))
@@ -189,7 +198,7 @@ void F3DAttachTrackEditor::ActorSocketPicked(const FName SocketName, USceneCompo
 		else if (ActorPickerID.ActorPicked.IsValid())
 		{
 			FGuid ParentActorId = FindOrCreateHandleToObject(ActorPickerID.ActorPicked.Get()).Handle;
-			ConstraintBindingID = FMovieSceneObjectBindingID(ParentActorId, GetSequencer()->GetFocusedTemplateID());
+			ConstraintBindingID = FMovieSceneObjectBindingID(ParentActorId, MovieSceneSequenceID::Root, EMovieSceneObjectBindingSpace::Local);
 		}
 
 		if (ConstraintBindingID.IsValid())
@@ -227,7 +236,7 @@ FKeyPropertyResult F3DAttachTrackEditor::AddKeyInternal( FFrameNumber KeyTime, c
 		FFindOrCreateHandleResult HandleResult = FindOrCreateHandleToObject(ActorPickerID.ActorPicked.Get());
 		FGuid ParentActorId = HandleResult.Handle;
 		KeyPropertyResult.bHandleCreated |= HandleResult.bWasCreated;
-		ConstraintBindingID = FMovieSceneObjectBindingID(ParentActorId, GetSequencer()->GetFocusedTemplateID());
+		ConstraintBindingID = FMovieSceneObjectBindingID(ParentActorId, MovieSceneSequenceID::Root, EMovieSceneObjectBindingSpace::Local);
 	}
 
 	if (!ConstraintBindingID.IsValid())
diff --git a/Engine/Source/Editor/MovieSceneTools/Private/TrackEditors/PathTrackEditor.cpp b/Engine/Source/Editor/MovieSceneTools/Private/TrackEditors/PathTrackEditor.cpp
--- a/Engine/Source/Editor/MovieSceneTools/Private/TrackEditors/PathTrackEditor.cpp
+++ b/Engine/Source/Editor/MovieSceneTools/Private/TrackEditors/PathTrackEditor.cpp
@@ -11,6 +11,7 @@
 #include "ActorEditorUtils.h"
 #include "Components/SplineComponent.h"
 #include "MovieSceneToolHelpers.h"
+#include "Evaluation/MovieSceneEvaluationTemplateInstance.h"
 
 
 #define LOCTEXT_NAMESPACE "FPathTrackEditor"
@@ -41,7 +42,15 @@ public:
 			TSharedPtr<ISequencer> Sequencer = PathTrackEditor->GetSequencer();
 			if (Sequencer.IsValid())
 			{
-				TArrayView<TWeakObjectPtr<UObject>> RuntimeObjects = Sequencer->FindBoundObjects(PathSection->GetConstraintBindingID().GetGuid(), PathSection->GetConstraintBindingID().GetSequenceID());
+				FMovieSceneSequenceID SequenceID = Sequencer->GetFocusedTemplateID();
+				if (PathSection->GetConstraintBindingID().GetSequenceID().IsValid())
+				{
+					// Ensure that this ID is resolvable from the root, based on the current local sequence ID
+					FMovieSceneObjectBindingID RootBindingID = PathSection->GetConstraintBindingID().ResolveLocalToRoot(SequenceID, Sequencer->GetEvaluationTemplate().GetHierarchy());
+					SequenceID = RootBindingID.GetSequenceID();
+				}
+
+				TArrayView<TWeakObjectPtr<UObject>> RuntimeObjects = Sequencer->FindBoundObjects(PathSection->GetConstraintBindingID().GetGuid(), SequenceID);
 				if (RuntimeObjects.Num() == 1 && RuntimeObjects[0].IsValid())
 				{
 					if (AActor* Actor = Cast<AActor>(RuntimeObjects[0].Get()))
@@ -133,7 +142,15 @@ bool F3DPathTrackEditor::IsActorPickable(const AActor* const ParentActor, FGuid
 	UMovieScene3DPathSection* PathSection = Cast<UMovieScene3DPathSection>(InSection);
 	if (PathSection != nullptr)
 	{
-		TArrayView<TWeakObjectPtr<UObject>> RuntimeObjects = GetSequencer()->FindBoundObjects(PathSection->GetConstraintBindingID().GetGuid(), PathSection->GetConstraintBindingID().GetSequenceID());
+		FMovieSceneSequenceID SequenceID = GetSequencer()->GetFocusedTemplateID();
+		if (PathSection->GetConstraintBindingID().GetSequenceID().IsValid())
+		{
+			// Ensure that this ID is resolvable from the root, based on the current local sequence ID
+			FMovieSceneObjectBindingID RootBindingID = PathSection->GetConstraintBindingID().ResolveLocalToRoot(SequenceID, GetSequencer()->GetEvaluationTemplate().GetHierarchy());
+			SequenceID = RootBindingID.GetSequenceID();
+		}
+
+		TArrayView<TWeakObjectPtr<UObject>> RuntimeObjects = GetSequencer()->FindBoundObjects(PathSection->GetConstraintBindingID().GetGuid(), SequenceID);
 		
 		if (RuntimeObjects.Contains(ParentActor))
 		{
@@ -174,7 +191,7 @@ void F3DPathTrackEditor::ActorSocketPicked(const FName SocketName, USceneCompone
 		else if (ActorPickerID.ActorPicked.IsValid())
 		{
 			FGuid ParentActorId = FindOrCreateHandleToObject(ActorPickerID.ActorPicked.Get()).Handle;
-			ConstraintBindingID = FMovieSceneObjectBindingID(ParentActorId, GetSequencer()->GetFocusedTemplateID());
+			ConstraintBindingID = FMovieSceneObjectBindingID(ParentActorId, MovieSceneSequenceID::Root, EMovieSceneObjectBindingSpace::Local);
 		}
 
 		if (ConstraintBindingID.IsValid())
@@ -208,7 +225,7 @@ FKeyPropertyResult F3DPathTrackEditor::AddKeyInternal( FFrameNumber KeyTime, con
 		FFindOrCreateHandleResult HandleResult = FindOrCreateHandleToObject(ActorPickerID.ActorPicked.Get());
 		FGuid ParentActorId = HandleResult.Handle;
 		KeyPropertyResult.bHandleCreated |= HandleResult.bWasCreated;
-		ConstraintBindingID = FMovieSceneObjectBindingID(ParentActorId, GetSequencer()->GetFocusedTemplateID());
+		ConstraintBindingID = FMovieSceneObjectBindingID(ParentActorId, MovieSceneSequenceID::Root, EMovieSceneObjectBindingSpace::Local);
 	}
 
 	if (!ConstraintBindingID.IsValid())
