Sequencer: Fix to prevent new attach/path section from being created when changing an existing section binding

diff --git a/Engine/Source/Editor/MovieSceneTools/Private/TrackEditors/ActorPickerTrackEditor.cpp b/Engine/Source/Editor/MovieSceneTools/Private/TrackEditors/ActorPickerTrackEditor.cpp
--- a/Engine/Source/Editor/MovieSceneTools/Private/TrackEditors/ActorPickerTrackEditor.cpp
+++ b/Engine/Source/Editor/MovieSceneTools/Private/TrackEditors/ActorPickerTrackEditor.cpp
@@ -129,7 +129,7 @@ void FActorPickerTrackEditor::ShowActorSubMenu(FMenuBuilder& MenuBuilder, FGuid
 
 	// Always recreate the binding picker to ensure we have the correct sequence ID
 	BindingIDPicker = MakeShared<FTrackEditorBindingIDPicker>(SequencerPtr->GetFocusedTemplateID(), SequencerPtr);
-	BindingIDPicker->OnBindingPicked().AddRaw(this, &FActorPickerTrackEditor::ExistingBindingPicked, ObjectBinding);
+	BindingIDPicker->OnBindingPicked().AddRaw(this, &FActorPickerTrackEditor::ExistingBindingPicked, ObjectBinding, Section);
 
 	FText ExistingBindingText = LOCTEXT("ExistingBinding", "Existing Binding");
 	FText NewBindingText = LOCTEXT("NewBinding", "New Binding");
@@ -260,7 +260,7 @@ void FActorPickerTrackEditor::ActorPicked(AActor* ParentActor, FGuid ObjectGuid,
 }
 
 
-void FActorPickerTrackEditor::ExistingBindingPicked(FMovieSceneObjectBindingID ExistingBindingID, FGuid ObjectBinding)
+void FActorPickerTrackEditor::ExistingBindingPicked(FMovieSceneObjectBindingID ExistingBindingID, FGuid ObjectBinding, UMovieSceneSection* Section)
 {
 	TSharedPtr<ISequencer> SequencerPtr = GetSequencer();
 
@@ -272,13 +272,13 @@ void FActorPickerTrackEditor::ExistingBindingPicked(FMovieSceneObjectBindingID E
 			AActor* Actor = Cast<AActor>(RuntimeObject.Get());
 			if (Actor)
 			{
-				ActorPickerIDPicked(FActorPickerID(Actor, ExistingBindingID), ObjectBinding, nullptr);
+				ActorPickerIDPicked(FActorPickerID(Actor, ExistingBindingID), ObjectBinding, Section);
 				return;
 			}
 		}
 	}
 
-	ActorPickerIDPicked(FActorPickerID(nullptr, ExistingBindingID), ObjectBinding, nullptr);
+	ActorPickerIDPicked(FActorPickerID(nullptr, ExistingBindingID), ObjectBinding, Section);
 }
 
 
diff --git a/Engine/Source/Editor/MovieSceneTools/Private/TrackEditors/ActorPickerTrackEditor.h b/Engine/Source/Editor/MovieSceneTools/Private/TrackEditors/ActorPickerTrackEditor.h
--- a/Engine/Source/Editor/MovieSceneTools/Private/TrackEditors/ActorPickerTrackEditor.h
+++ b/Engine/Source/Editor/MovieSceneTools/Private/TrackEditors/ActorPickerTrackEditor.h
@@ -45,7 +45,7 @@ private:
 	void ActorPicked(AActor* ParentActor, FGuid ObjectBinding, UMovieSceneSection* Section);
 
 	/** Existing binding was picked */
-	void ExistingBindingPicked(FMovieSceneObjectBindingID ExistingBindingID, FGuid ObjectBinding);
+	void ExistingBindingPicked(FMovieSceneObjectBindingID ExistingBindingID, FGuid ObjectBinding, UMovieSceneSection* Section);
 
 	/** Actor and/or existing binding was picked */
 	void ActorPickerIDPicked(FActorPickerID, FGuid ObjectBinding, UMovieSceneSection* Section);
