When calculating vertical extents find feature points where slopes are zero and check them in addition to the keys if the curve is cubic. Curves now fit correclty vertically.

Also changed fudge to 5% from 10% to match up old editor. Tighter fits seems better.

AndrewR I left the adding of the snap vallue factor to the max calculation though I'm not 100% happy with it since it's not discoverable that's the reason you can't zoom in on your curve.  I would vote to remove it since you can always zoom out and then snap.

diff --git a/Engine/Source/Editor/CurveEditor/Private/CurveEditor.cpp b/Engine/Source/Editor/CurveEditor/Private/CurveEditor.cpp
--- a/Engine/Source/Editor/CurveEditor/Private/CurveEditor.cpp
+++ b/Engine/Source/Editor/CurveEditor/Private/CurveEditor.cpp
@@ -235,27 +235,13 @@ void FCurveEditor::ZoomToFitCurves(TArrayView<const FCurveModelID> CurveModelIDs
 		{
 			if (const FCurveModel* Curve = FindCurve(CurveModelID))
 			{
-				KeyHandlesScratch.Reset();
-				Curve->GetKeys(*this,
-					TNumericLimits<double>::Lowest(),
-					TNumericLimits<double>::Max(),
-					TNumericLimits<double>::Lowest(),
-					TNumericLimits<double>::Max(),
-					KeyHandlesScratch);
-
-				if (KeyHandlesScratch.Num())
-				{
-					KeyPositionsScratch.SetNum(KeyHandlesScratch.Num());
-					Curve->GetKeyPositions(KeyHandlesScratch, KeyPositionsScratch);
-
-					for (FKeyPosition Key : KeyPositionsScratch)
-					{
-						InputMin  = FMath::Min(InputMin, Key.InputValue);
-						InputMax  = FMath::Max(InputMax, Key.InputValue);
-						OutputMin = FMath::Min(OutputMin, Key.OutputValue);
-						OutputMax = FMath::Max(OutputMax, Key.OutputValue);
-					}
-				}
+				double LocalMin, LocalMax;
+				Curve->GetTimeRange(LocalMin, LocalMax);
+				InputMin = FMath::Min(InputMin, LocalMin);
+				InputMax = FMath::Max(InputMax, LocalMax);
+				Curve->GetValueRange(LocalMin, LocalMax);
+				OutputMin = FMath::Min(OutputMin, LocalMin);
+				OutputMax = FMath::Max(OutputMax, LocalMax);
 			}
 		}
 
@@ -314,7 +300,7 @@ void FCurveEditor::ZoomToFitInternal(EAxisList::Type Axes, double InputMin, doub
 	OutputMax = FMath::Max(OutputMin + MinOutputZoom, OutputMax);
 
 	double InputPadding  = (InputMax - InputMin) * 0.1;
-	double OutputPadding = (OutputMax - OutputMin) * 0.1;
+	double OutputPadding = (OutputMax - OutputMin) * 0.05;
 
 	InputMin -= InputPadding;
 	InputMax += InputPadding;
diff --git a/Engine/Source/Editor/CurveEditor/Public/CurveModel.h b/Engine/Source/Editor/CurveEditor/Public/CurveModel.h
--- a/Engine/Source/Editor/CurveEditor/Public/CurveModel.h
+++ b/Engine/Source/Editor/CurveEditor/Public/CurveModel.h
@@ -109,6 +109,19 @@ public:
 	 */
 	virtual void GetKeyDrawInfo(ECurvePointType PointType, FKeyDrawInfo& OutDrawInfo) const = 0;
 
+	/** Get range of input time.
+	* @param MinTime Minimum Time
+	* @param MaxTime Minimum Time
+	*
+	*/
+	virtual void GetTimeRange(double& MinTime, double& MaxTime) const = 0;
+
+	/** Get range of output values.
+	* @param MinValue Minimum Value
+	* @param MaxValue Minimum Value
+	*/
+	virtual void GetValueRange(double& MinValue, double& MaxValue) const = 0;
+
 	/**
 	 * Evaluate this curve at the specified time
 	 *
diff --git a/Engine/Source/Editor/MovieSceneTools/Private/Channels/FloatChannelCurveModel.cpp b/Engine/Source/Editor/MovieSceneTools/Private/Channels/FloatChannelCurveModel.cpp
--- a/Engine/Source/Editor/MovieSceneTools/Private/Channels/FloatChannelCurveModel.cpp
+++ b/Engine/Source/Editor/MovieSceneTools/Private/Channels/FloatChannelCurveModel.cpp
@@ -329,7 +329,8 @@ void FFloatChannelCurveModel::SetKeyAttributes(TArrayView<const FKeyHandle> InKe
 					if (KeyValue.Tangent.TangentWeightMode == RCTWM_WeightedNone) //set tangent weights to default use
 					{
 						TArrayView<const FFrameNumber> Times = Channel->GetTimes();
-						
+						const float OneThird = 1.0f / 3.0f;
+
 						//calculate a tangent weight based upon tangent and time difference
 						//calculate arrive tangent weight
 						if (KeyIndex > 0 )
@@ -337,7 +338,7 @@ void FFloatChannelCurveModel::SetKeyAttributes(TArrayView<const FKeyHandle> InKe
 							const float X = TickResolution.AsSeconds(Times[KeyIndex].Value - Times[KeyIndex - 1].Value);
 							const float ArriveTangentNormal = KeyValue.Tangent.ArriveTangent / (TimeInterval);
 							const float Y = ArriveTangentNormal * X;
-							KeyValue.Tangent.ArriveTangentWeight = FMath::Sqrt(X*X + Y*Y) * 0.3f;
+							KeyValue.Tangent.ArriveTangentWeight = FMath::Sqrt(X*X + Y * Y) * OneThird;
 						}
 						//calculate leave weight
 						if(KeyIndex < ( Times.Num() - 1))
@@ -345,7 +346,7 @@ void FFloatChannelCurveModel::SetKeyAttributes(TArrayView<const FKeyHandle> InKe
 							const float X = TickResolution.AsSeconds(Times[KeyIndex].Value - Times[KeyIndex + 1].Value);
 							const float LeaveTangentNormal = KeyValue.Tangent.LeaveTangent / (TimeInterval);
 							const float Y = LeaveTangentNormal * X;
-							KeyValue.Tangent.LeaveTangentWeight = FMath::Sqrt(X*X + Y*Y) * 0.3f;
+							KeyValue.Tangent.LeaveTangentWeight = FMath::Sqrt(X*X + Y*Y) * OneThird;
 						}
 					}
 					KeyValue.Tangent.TangentWeightMode = Attributes.GetTangentWeightMode();
@@ -468,4 +469,106 @@ void FFloatChannelCurveModel::CreateKeyProxies(TArrayView<const FKeyHandle> InKe
 		NewProxy->Initialize(InKeyHandles[Index], ChannelHandle, WeakSection);
 		OutObjects[Index] = NewProxy;
 	}
+}
+
+void FFloatChannelCurveModel::GetTimeRange(double& MinTime, double& MaxTime) const
+{
+	FMovieSceneFloatChannel* Channel = ChannelHandle.Get();
+	UMovieSceneSection*      Section = WeakSection.Get();
+
+	if (Channel && Section)
+	{
+		TMovieSceneChannel<FMovieSceneFloatValue> ChannelInterface = Channel->GetInterface();
+		TArrayView<const FFrameNumber> Times = ChannelInterface.GetTimes();
+		if (Times.Num() == 0)
+		{
+			MinTime = 0.f;
+			MaxTime = 0.f;
+		}
+		else
+		{
+			FFrameRate TickResolution = Section->GetTypedOuter<UMovieScene>()->GetTickResolution();
+			double ToTime = TickResolution.AsInterval();
+			MinTime = static_cast<double> (Times[0].Value) * ToTime;
+			MaxTime = static_cast<double>(Times[Times.Num() - 1].Value) * ToTime;
+		}
+	}
+}
+
+/*	 Finds min/max for cubic curves:
+Looks for feature points in the signal(determined by change in direction of local tangent), these locations are then re-examined in closer detail recursively
+Similar to function in RichCurve but usees the Channel::Evaluate function, instead of CurveModel::Eval*/
+
+void FFloatChannelCurveModel::FeaturePointMethod(double StartTime, double EndTime, double StartValue, double Mu, int Depth, int MaxDepth, double& MaxV, double& MinVal) const
+{
+	if (Depth >= MaxDepth)
+	{
+		return;
+	}
+	double PrevValue = StartValue;
+	double EvalValue;
+	Evaluate(StartTime - Mu, EvalValue);
+	double PrevTangent = StartValue - EvalValue;
+	EndTime += Mu;
+	for (double f = StartTime + Mu; f < EndTime; f += Mu)
+	{
+		double Value;
+		Evaluate(f, Value);
+
+		MaxV = FMath::Max(Value, MaxV);
+		MinVal = FMath::Min(Value, MinVal);
+		double CurTangent = Value - PrevValue;
+		if (FMath::Sign(CurTangent) != FMath::Sign(PrevTangent))
+		{
+			//feature point centered around the previous tangent
+			double FeaturePointTime = f - Mu * 2.0f;
+			double NewVal;
+			Evaluate(FeaturePointTime, NewVal);
+			FeaturePointMethod(FeaturePointTime, f,NewVal, Mu*0.4f, Depth + 1, MaxDepth, MaxV, MinVal);
+		}
+		PrevTangent = CurTangent;
+		PrevValue = Value;
+	}
+}
+
+void FFloatChannelCurveModel::GetValueRange(double& MinValue, double& MaxValue) const
+{
+	FMovieSceneFloatChannel* Channel = ChannelHandle.Get();
+	UMovieSceneSection*      Section = WeakSection.Get();
+
+	if (Channel && Section)
+	{
+		TMovieSceneChannel<FMovieSceneFloatValue> ChannelInterface = Channel->GetInterface();
+		TArrayView<const FFrameNumber> Times = ChannelInterface.GetTimes();
+		TArrayView<const FMovieSceneFloatValue>   Values = ChannelInterface.GetValues();
+
+		if (Times.Num() == 0)
+		{
+			MinValue = MaxValue = 0.f;
+		}
+		else
+		{
+			FFrameRate TickResolution = Section->GetTypedOuter<UMovieScene>()->GetTickResolution();
+			double ToTime = TickResolution.AsInterval();
+			int32 LastKeyIndex = Values.Num() - 1;
+			MinValue = MaxValue = Values[0].Value;
+
+			for (int32 i = 0; i < Values.Num(); i++)
+			{
+				const FMovieSceneFloatValue& Key = Values[i];
+
+				MinValue = FMath::Min(MinValue,(double) Key.Value);
+				MaxValue = FMath::Max(MaxValue, (double)Key.Value);
+
+				if (Key.InterpMode == RCIM_Cubic && i != LastKeyIndex)
+				{
+					const FMovieSceneFloatValue& NextKey = Values[i + 1];
+					double KeyTime = static_cast<double>(Times[i].Value) *ToTime;
+					double NextTime = static_cast<double>(Times[i +1].Value) *ToTime;
+					double TimeStep = (NextTime - KeyTime) * 0.2f;
+					FeaturePointMethod(KeyTime, NextTime, Key.Value, TimeStep, 0, 3, MaxValue, MinValue);
+				}
+			}
+		}
+	}
 }
\ No newline at end of file
diff --git a/Engine/Source/Editor/MovieSceneTools/Private/Channels/FloatChannelCurveModel.h b/Engine/Source/Editor/MovieSceneTools/Private/Channels/FloatChannelCurveModel.h
--- a/Engine/Source/Editor/MovieSceneTools/Private/Channels/FloatChannelCurveModel.h
+++ b/Engine/Source/Editor/MovieSceneTools/Private/Channels/FloatChannelCurveModel.h
@@ -34,13 +34,16 @@ public:
 
 	virtual void GetCurveAttributes(FCurveAttributes& OutCurveAttributes) const override;
 	virtual void SetCurveAttributes(const FCurveAttributes& InCurveAttributes) override;
-
+	virtual void GetTimeRange(double& MinTime, double& MaxTime) const override;
+	virtual void GetValueRange(double& MinValue, double& MaxValue) const override;
 	virtual bool Evaluate(double ProspectiveTime, double& OutValue) const override;
 	virtual void AddKeys(TArrayView<const FKeyPosition> InKeyPositions, TArrayView<const FKeyAttributes> InAttributes, TArrayView<TOptional<FKeyHandle>>* OutKeyHandles) override;
 	virtual void RemoveKeys(TArrayView<const FKeyHandle> InKeys) override;
 
 	virtual void CreateKeyProxies(TArrayView<const FKeyHandle> InKeyHandles, TArrayView<UObject*> OutObjects) override;
 
+private:
+	void FeaturePointMethod(double StartTime, double EndTime, double StartValue, double Mu, int Depth, int MaxDepth, double& MaxV, double& MinVal) const;
 private:
 
 	TMovieSceneChannelHandle<FMovieSceneFloatChannel> ChannelHandle;
