Sequencer: Resolve merge conflicts

diff --git a/Engine/Source/Editor/MovieSceneTools/Private/Channels/FloatChannelCurveModel.cpp b/Engine/Source/Editor/MovieSceneTools/Private/Channels/FloatChannelCurveModel.cpp
--- a/Engine/Source/Editor/MovieSceneTools/Private/Channels/FloatChannelCurveModel.cpp
+++ b/Engine/Source/Editor/MovieSceneTools/Private/Channels/FloatChannelCurveModel.cpp
@@ -52,7 +52,7 @@ void FFloatChannelCurveModel::AddKeys(TArrayView<const FKeyPosition> InKeyPositi
 	{
 		Section->Modify();
 
-		TMovieSceneChannel<FMovieSceneFloatValue> ChannelInterface = Channel->GetInterface();
+		TMovieSceneChannelData<FMovieSceneFloatValue> ChannelData = Channel->GetData();
 		FFrameRate TickResolution = Section->GetTypedOuter<UMovieScene>()->GetTickResolution();
 
 		for (int32 Index = 0; Index < InKeyPositions.Num(); ++Index)
@@ -70,10 +70,10 @@ void FFloatChannelCurveModel::AddKeys(TArrayView<const FKeyPosition> InKeyPositi
 			if (Attributes.HasArriveTangent()) { Value.Tangent.ArriveTangent = Attributes.GetArriveTangent(); }
 			if (Attributes.HasLeaveTangent())  { Value.Tangent.LeaveTangent  = Attributes.GetLeaveTangent();  }
 
-			int32 KeyIndex = ChannelInterface.AddKey(Time, Value);
+			int32 KeyIndex = ChannelData.AddKey(Time, Value);
 			if (OutKeyHandles)
 			{
-				(*OutKeyHandles)[Index] = ChannelInterface.GetHandle(KeyIndex);
+				(*OutKeyHandles)[Index] = ChannelData.GetHandle(KeyIndex);
 			}
 		}
 
@@ -109,14 +109,14 @@ void FFloatChannelCurveModel::RemoveKeys(TArrayView<const FKeyHandle> InKeys)
 	{
 		Section->Modify();
 
-		TMovieSceneChannel<FMovieSceneFloatValue> ChannelInterface = Channel->GetInterface();
+		TMovieSceneChannelData<FMovieSceneFloatValue> ChannelData = Channel->GetData();
 
 		for (FKeyHandle Handle : InKeys)
 		{
-			int32 KeyIndex = ChannelInterface.GetIndex(Handle);
+			int32 KeyIndex = ChannelData.GetIndex(Handle);
 			if (KeyIndex != INDEX_NONE)
 			{
-				ChannelInterface.RemoveKey(KeyIndex);
+				ChannelData.RemoveKey(KeyIndex);
 			}
 		}
 	}
@@ -152,9 +152,9 @@ void FFloatChannelCurveModel::GetKeys(const FCurveEditor& CurveEditor, double Mi
 	{
 		FFrameRate TickResolution = Section->GetTypedOuter<UMovieScene>()->GetTickResolution();
 
-		TMovieSceneChannel<FMovieSceneFloatValue> ChannelInterface = Channel->GetInterface();
-		TArrayView<const FFrameNumber>            Times  = ChannelInterface.GetTimes();
-		TArrayView<const FMovieSceneFloatValue>   Values = ChannelInterface.GetValues();
+		TMovieSceneChannelData<FMovieSceneFloatValue> ChannelData = Channel->GetData();
+		TArrayView<const FFrameNumber>          Times  = ChannelData.GetTimes();
+		TArrayView<const FMovieSceneFloatValue> Values = ChannelData.GetValues();
 
 		const FFrameNumber StartFrame = MinTime <= MIN_int32 ? MIN_int32 : (MinTime * TickResolution).CeilToFrame();
 		const FFrameNumber EndFrame   = MaxTime >= MAX_int32 ? MAX_int32 : (MaxTime * TickResolution).FloorToFrame();
@@ -166,7 +166,7 @@ void FFloatChannelCurveModel::GetKeys(const FCurveEditor& CurveEditor, double Mi
 		{
 			if (Values[KeyIndex].Value >= MinValue && Values[KeyIndex].Value <= MaxValue)
 			{
-				OutKeyHandles.Add(ChannelInterface.GetHandle(KeyIndex));
+				OutKeyHandles.Add(ChannelData.GetHandle(KeyIndex));
 			}
 		}
 	}
@@ -197,13 +197,13 @@ void FFloatChannelCurveModel::GetKeyPositions(TArrayView<const FKeyHandle> InKey
 	{
 		FFrameRate TickResolution = Section->GetTypedOuter<UMovieScene>()->GetTickResolution();
 
-		TMovieSceneChannel<FMovieSceneFloatValue> ChannelInterface = Channel->GetInterface();
-		TArrayView<const FFrameNumber>            Times            = ChannelInterface.GetTimes();
-		TArrayView<const FMovieSceneFloatValue>   Values           = ChannelInterface.GetValues();
+		TMovieSceneChannelData<FMovieSceneFloatValue> ChannelData = Channel->GetData();
+		TArrayView<const FFrameNumber>          Times  = ChannelData.GetTimes();
+		TArrayView<const FMovieSceneFloatValue> Values = ChannelData.GetValues();
 
 		for (int32 Index = 0; Index < InKeys.Num(); ++Index)
 		{
-			int32 KeyIndex = ChannelInterface.GetIndex(InKeys[Index]);
+			int32 KeyIndex = ChannelData.GetIndex(InKeys[Index]);
 			if (KeyIndex != INDEX_NONE)
 			{
 				OutKeyPositions[Index].InputValue  = Times[KeyIndex] / TickResolution;
@@ -224,16 +224,16 @@ void FFloatChannelCurveModel::SetKeyPositions(TArrayView<const FKeyHandle> InKey
 
 		FFrameRate TickResolution = Section->GetTypedOuter<UMovieScene>()->GetTickResolution();
 
-		TMovieSceneChannel<FMovieSceneFloatValue> ChannelInterface = Channel->GetInterface();
+		TMovieSceneChannelData<FMovieSceneFloatValue> ChannelData = Channel->GetData();
 		for (int32 Index = 0; Index < InKeys.Num(); ++Index)
 		{
-			int32 KeyIndex = ChannelInterface.GetIndex(InKeys[Index]);
+			int32 KeyIndex = ChannelData.GetIndex(InKeys[Index]);
 			if (KeyIndex != INDEX_NONE)
 			{
 				FFrameNumber NewTime = (InKeyPositions[Index].InputValue * TickResolution).FloorToFrame();
 
-				KeyIndex = ChannelInterface.MoveKey(KeyIndex, NewTime);
-				ChannelInterface.GetValues()[KeyIndex].Value = InKeyPositions[Index].OutputValue;
+				KeyIndex = ChannelData.MoveKey(KeyIndex, NewTime);
+				ChannelData.GetValues()[KeyIndex].Value = InKeyPositions[Index].OutputValue;
 
 				Section->ExpandToFrame(NewTime);
 			}
@@ -249,15 +249,15 @@ void FFloatChannelCurveModel::GetKeyAttributes(TArrayView<const FKeyHandle> InKe
 	UMovieSceneSection*      Section = WeakSection.Get();
 	if (Channel && Section)
 	{
-		TMovieSceneChannel<FMovieSceneFloatValue> ChannelInterface = Channel->GetInterface();
-		TArrayView<const FFrameNumber>    Times  = ChannelInterface.GetTimes();
-		TArrayView<FMovieSceneFloatValue> Values = ChannelInterface.GetValues();
+		TMovieSceneChannelData<FMovieSceneFloatValue> ChannelData = Channel->GetData();
+		TArrayView<const FFrameNumber>    Times  = ChannelData.GetTimes();
+		TArrayView<FMovieSceneFloatValue> Values = ChannelData.GetValues();
 
 		float TimeInterval = Section->GetTypedOuter<UMovieScene>()->GetTickResolution().AsInterval();
 
 		for (int32 Index = 0; Index < InKeys.Num(); ++Index)
 		{
-			const int32 KeyIndex = ChannelInterface.GetIndex(InKeys[Index]);
+			const int32 KeyIndex = ChannelData.GetIndex(InKeys[Index]);
 			if (KeyIndex != INDEX_NONE)
 			{
 				const FMovieSceneFloatValue& KeyValue    = Values[KeyIndex];
@@ -301,15 +301,15 @@ void FFloatChannelCurveModel::SetKeyAttributes(TArrayView<const FKeyHandle> InKe
 		bool bAutoSetTangents = false;
 		Section->MarkAsChanged();
 
-		TMovieSceneChannel<FMovieSceneFloatValue> ChannelInterface = Channel->GetInterface();
-		TArrayView<FMovieSceneFloatValue> Values = ChannelInterface.GetValues();
+		TMovieSceneChannelData<FMovieSceneFloatValue> ChannelData = Channel->GetData();
+		TArrayView<FMovieSceneFloatValue> Values = ChannelData.GetValues();
 
 		FFrameRate TickResolution = Section->GetTypedOuter<UMovieScene>()->GetTickResolution();
 		float TimeInterval = TickResolution.AsInterval();
 
 		for (int32 Index = 0; Index < InKeys.Num(); ++Index)
 		{
-			const int32 KeyIndex = ChannelInterface.GetIndex(InKeys[Index]);
+			const int32 KeyIndex = ChannelData.GetIndex(InKeys[Index]);
 			if (KeyIndex != INDEX_NONE)
 			{
 				const FKeyAttributes&  Attributes = InAttributes[Index];
@@ -478,8 +478,7 @@ void FFloatChannelCurveModel::GetTimeRange(double& MinTime, double& MaxTime) con
 
 	if (Channel && Section)
 	{
-		TMovieSceneChannel<FMovieSceneFloatValue> ChannelInterface = Channel->GetInterface();
-		TArrayView<const FFrameNumber> Times = ChannelInterface.GetTimes();
+		TArrayView<const FFrameNumber> Times = Channel->GetData().GetTimes();
 		if (Times.Num() == 0)
 		{
 			MinTime = 0.f;
@@ -538,9 +537,8 @@ void FFloatChannelCurveModel::GetValueRange(double& MinValue, double& MaxValue)
 
 	if (Channel && Section)
 	{
-		TMovieSceneChannel<FMovieSceneFloatValue> ChannelInterface = Channel->GetInterface();
-		TArrayView<const FFrameNumber> Times = ChannelInterface.GetTimes();
-		TArrayView<const FMovieSceneFloatValue>   Values = ChannelInterface.GetValues();
+		TArrayView<const FFrameNumber> Times = Channel->GetData().GetTimes();
+		TArrayView<const FMovieSceneFloatValue>   Values = Channel->GetData().GetValues();
 
 		if (Times.Num() == 0)
 		{
diff --git a/Engine/Source/Editor/MovieSceneTools/Private/MatineeImportTools.cpp b/Engine/Source/Editor/MovieSceneTools/Private/MatineeImportTools.cpp
--- a/Engine/Source/Editor/MovieSceneTools/Private/MatineeImportTools.cpp
+++ b/Engine/Source/Editor/MovieSceneTools/Private/MatineeImportTools.cpp
@@ -114,10 +114,10 @@ bool FMatineeImportTools::TryConvertMatineeToggleToOutParticleKey( ETrackToggleA
 }
 
 
-void FMatineeImportTools::SetOrAddKey(TMovieSceneChannel<FMovieSceneFloatValue>& ChannelInterface, FFrameNumber Time, float Value, float ArriveTangent, float LeaveTangent, EInterpCurveMode MatineeInterpMode
+void FMatineeImportTools::SetOrAddKey(TMovieSceneChannelData<FMovieSceneFloatValue>& ChannelData, FFrameNumber Time, float Value, float ArriveTangent, float LeaveTangent, EInterpCurveMode MatineeInterpMode
 	, ERichCurveTangentWeightMode WeightedMode, float ArriveTangentWeight, float LeaveTangentWeight )
 {
-	if (ChannelInterface.FindKey(Time) == INDEX_NONE)
+	if (ChannelData.FindKey(Time) == INDEX_NONE)
 	{
 		FMovieSceneFloatValue NewKey(Value);
 
@@ -128,7 +128,7 @@ void FMatineeImportTools::SetOrAddKey(TMovieSceneChannel<FMovieSceneFloatValue>&
 		NewKey.Tangent.TangentWeightMode = WeightedMode;
 		NewKey.Tangent.ArriveTangentWeight = ArriveTangentWeight;
 		NewKey.Tangent.LeaveTangentWeight = LeaveTangentWeight;
-		ChannelInterface.AddKey( Time, NewKey );
+		ChannelData.AddKey( Time, NewKey );
 	}
 }
 
@@ -155,12 +155,12 @@ bool FMatineeImportTools::CopyInterpBoolTrack( UInterpTrackBoolProp* MatineeBool
 	{
 		TRange<FFrameNumber> KeyRange = TRange<FFrameNumber>::Empty();
 
-		TMovieSceneChannel<bool> ChannelInterface = Section->GetChannelProxy().GetChannel<FMovieSceneBoolChannel>(0)->GetInterface();
+		TMovieSceneChannelData<bool> ChannelData = Section->GetChannelProxy().GetChannel<FMovieSceneBoolChannel>(0)->GetData();
 		for ( const auto& Point : MatineeBoolTrack->BoolTrack )
 		{
 			FFrameNumber KeyTime = (Point.Time * FrameRate).RoundToFrame();
 
-			ChannelInterface.UpdateOrAddKey(KeyTime, Point.Value);
+			ChannelData.UpdateOrAddKey(KeyTime, Point.Value);
 
 			KeyRange = TRange<FFrameNumber>::Hull(KeyRange, TRange<FFrameNumber>(KeyTime));
 		}
@@ -197,13 +197,13 @@ bool FMatineeImportTools::CopyInterpFloatTrack( UInterpTrackFloatBase* MatineeFl
 		TRange<FFrameNumber> KeyRange = TRange<FFrameNumber>::Empty();
 
 		FMovieSceneFloatChannel* Channel = Section->GetChannelProxy().GetChannel<FMovieSceneFloatChannel>(0);
-		TMovieSceneChannel<FMovieSceneFloatValue> ChannelInterface = Channel->GetInterface();
+		TMovieSceneChannelData<FMovieSceneFloatValue> ChannelData = Channel->GetData();
 
 		for ( const auto& Point : MatineeFloatTrack->FloatTrack.Points )
 		{
 			FFrameNumber KeyTime = (Point.InVal * FrameRate).RoundToFrame();
 
-			FMatineeImportTools::SetOrAddKey( ChannelInterface, KeyTime, Point.OutVal, Point.ArriveTangent, Point.LeaveTangent, Point.InterpMode );
+			FMatineeImportTools::SetOrAddKey( ChannelData, KeyTime, Point.OutVal, Point.ArriveTangent, Point.LeaveTangent, Point.InterpMode );
 
 			KeyRange = TRange<FFrameNumber>::Hull(KeyRange, TRange<FFrameNumber>(KeyTime));
 		}
@@ -244,15 +244,15 @@ bool FMatineeImportTools::CopyInterpVectorTrack( UInterpTrackVectorProp* Matinee
 		if (Section->GetChannelsUsed() == 3)
 		{
 			TArrayView<FMovieSceneFloatChannel*> Channels = Section->GetChannelProxy().GetChannels<FMovieSceneFloatChannel>();
-			TMovieSceneChannel<FMovieSceneFloatValue> ChannelInterfaces[3] = { Channels[0]->GetInterface(), Channels[1]->GetInterface(), Channels[2]->GetInterface() };
+			TMovieSceneChannelData<FMovieSceneFloatValue> ChannelData[3] = { Channels[0]->GetData(), Channels[1]->GetData(), Channels[2]->GetData() };
 
 			for ( const auto& Point : MatineeVectorTrack->VectorTrack.Points )
 			{
 				FFrameNumber KeyTime = (Point.InVal * FrameRate).RoundToFrame();
 
-				FMatineeImportTools::SetOrAddKey( ChannelInterfaces[0], KeyTime, Point.OutVal.X, Point.ArriveTangent.X, Point.LeaveTangent.X, Point.InterpMode );
-				FMatineeImportTools::SetOrAddKey( ChannelInterfaces[1], KeyTime, Point.OutVal.Y, Point.ArriveTangent.Y, Point.LeaveTangent.Y, Point.InterpMode );
-				FMatineeImportTools::SetOrAddKey( ChannelInterfaces[2], KeyTime, Point.OutVal.Z, Point.ArriveTangent.Z, Point.LeaveTangent.Z, Point.InterpMode );
+				FMatineeImportTools::SetOrAddKey( ChannelData[0], KeyTime, Point.OutVal.X, Point.ArriveTangent.X, Point.LeaveTangent.X, Point.InterpMode );
+				FMatineeImportTools::SetOrAddKey( ChannelData[1], KeyTime, Point.OutVal.Y, Point.ArriveTangent.Y, Point.LeaveTangent.Y, Point.InterpMode );
+				FMatineeImportTools::SetOrAddKey( ChannelData[2], KeyTime, Point.OutVal.Z, Point.ArriveTangent.Z, Point.LeaveTangent.Z, Point.InterpMode );
 
 				KeyRange = TRange<FFrameNumber>::Hull(KeyRange, TRange<FFrameNumber>(KeyTime));
 			}
@@ -301,7 +301,7 @@ bool FMatineeImportTools::CopyInterpColorTrack( UInterpTrackColorProp* ColorProp
 	if (Section->TryModify())
 	{
 		TArrayView<FMovieSceneFloatChannel*> Channels = Section->GetChannelProxy().GetChannels<FMovieSceneFloatChannel>();
-		TMovieSceneChannel<FMovieSceneFloatValue> ChannelInterfaces[3] = { Channels[0]->GetInterface(), Channels[1]->GetInterface(), Channels[2]->GetInterface() };
+		TMovieSceneChannelData<FMovieSceneFloatValue> ChannelData[3] = { Channels[0]->GetData(), Channels[1]->GetData(), Channels[2]->GetData() };
 
 		TRange<FFrameNumber> KeyRange = TRange<FFrameNumber>::Empty();
 
@@ -309,9 +309,9 @@ bool FMatineeImportTools::CopyInterpColorTrack( UInterpTrackColorProp* ColorProp
 		{
 			FFrameNumber KeyTime = (Point.InVal * FrameRate).RoundToFrame();
 
-			FMatineeImportTools::SetOrAddKey( ChannelInterfaces[0], KeyTime, Point.OutVal.X, Point.ArriveTangent.X, Point.LeaveTangent.X, Point.InterpMode );
-			FMatineeImportTools::SetOrAddKey( ChannelInterfaces[1], KeyTime, Point.OutVal.Y, Point.ArriveTangent.Y, Point.LeaveTangent.Y, Point.InterpMode );
-			FMatineeImportTools::SetOrAddKey( ChannelInterfaces[2], KeyTime, Point.OutVal.Z, Point.ArriveTangent.Z, Point.LeaveTangent.Z, Point.InterpMode );
+			FMatineeImportTools::SetOrAddKey( ChannelData[0], KeyTime, Point.OutVal.X, Point.ArriveTangent.X, Point.LeaveTangent.X, Point.InterpMode );
+			FMatineeImportTools::SetOrAddKey( ChannelData[1], KeyTime, Point.OutVal.Y, Point.ArriveTangent.Y, Point.LeaveTangent.Y, Point.InterpMode );
+			FMatineeImportTools::SetOrAddKey( ChannelData[2], KeyTime, Point.OutVal.Z, Point.ArriveTangent.Z, Point.LeaveTangent.Z, Point.InterpMode );
 
 			KeyRange = TRange<FFrameNumber>::Hull(KeyRange, TRange<FFrameNumber>(KeyTime));
 		}
@@ -361,16 +361,16 @@ bool FMatineeImportTools::CopyInterpLinearColorTrack( UInterpTrackLinearColorPro
 		TRange<FFrameNumber> KeyRange = TRange<FFrameNumber>::Empty();
 
 		TArrayView<FMovieSceneFloatChannel*> Channels = Section->GetChannelProxy().GetChannels<FMovieSceneFloatChannel>();
-		TMovieSceneChannel<FMovieSceneFloatValue> ChannelInterfaces[4] = { Channels[0]->GetInterface(), Channels[1]->GetInterface(), Channels[2]->GetInterface(), Channels[3]->GetInterface() };
+		TMovieSceneChannelData<FMovieSceneFloatValue> ChannelData[4] = { Channels[0]->GetData(), Channels[1]->GetData(), Channels[2]->GetData(), Channels[3]->GetData() };
 
 		for ( const auto& Point : LinearColorPropTrack->LinearColorTrack.Points )
 		{
 			FFrameNumber KeyTime = (Point.InVal * FrameRate).RoundToFrame();
 
-			FMatineeImportTools::SetOrAddKey( ChannelInterfaces[0], KeyTime, Point.OutVal.R, Point.ArriveTangent.R, Point.LeaveTangent.R, Point.InterpMode );
-			FMatineeImportTools::SetOrAddKey( ChannelInterfaces[1], KeyTime, Point.OutVal.G, Point.ArriveTangent.G, Point.LeaveTangent.G, Point.InterpMode );
-			FMatineeImportTools::SetOrAddKey( ChannelInterfaces[2], KeyTime, Point.OutVal.B, Point.ArriveTangent.B, Point.LeaveTangent.B, Point.InterpMode );
-			FMatineeImportTools::SetOrAddKey( ChannelInterfaces[3], KeyTime, Point.OutVal.A, Point.ArriveTangent.A, Point.LeaveTangent.A, Point.InterpMode );
+			FMatineeImportTools::SetOrAddKey( ChannelData[0], KeyTime, Point.OutVal.R, Point.ArriveTangent.R, Point.LeaveTangent.R, Point.InterpMode );
+			FMatineeImportTools::SetOrAddKey( ChannelData[1], KeyTime, Point.OutVal.G, Point.ArriveTangent.G, Point.LeaveTangent.G, Point.InterpMode );
+			FMatineeImportTools::SetOrAddKey( ChannelData[2], KeyTime, Point.OutVal.B, Point.ArriveTangent.B, Point.LeaveTangent.B, Point.InterpMode );
+			FMatineeImportTools::SetOrAddKey( ChannelData[3], KeyTime, Point.OutVal.A, Point.ArriveTangent.A, Point.LeaveTangent.A, Point.InterpMode );
 
 			KeyRange = TRange<FFrameNumber>::Hull(KeyRange, TRange<FFrameNumber>(KeyTime));
 		}
@@ -419,18 +419,18 @@ bool FMatineeImportTools::CopyInterpMoveTrack( UInterpTrackMove* MoveTrack, UMov
 		TRange<FFrameNumber> KeyRange = TRange<FFrameNumber>::Empty();
 
 		TArrayView<FMovieSceneFloatChannel*> Channels = Section->GetChannelProxy().GetChannels<FMovieSceneFloatChannel>();
-		TMovieSceneChannel<FMovieSceneFloatValue> ChannelInterfaces[6] = {
-			Channels[0]->GetInterface(), Channels[1]->GetInterface(), Channels[2]->GetInterface(),
-			Channels[3]->GetInterface(), Channels[4]->GetInterface(), Channels[5]->GetInterface()
+		TMovieSceneChannelData<FMovieSceneFloatValue> ChannelData[6] = {
+			Channels[0]->GetData(), Channels[1]->GetData(), Channels[2]->GetData(),
+			Channels[3]->GetData(), Channels[4]->GetData(), Channels[5]->GetData()
 		};
 
 		for ( const auto& Point : MoveTrack->PosTrack.Points )
 		{
 			FFrameNumber KeyTime = (Point.InVal * FrameRate).RoundToFrame();
 
-			FMatineeImportTools::SetOrAddKey( ChannelInterfaces[0], KeyTime, Point.OutVal.X, Point.ArriveTangent.X, Point.LeaveTangent.X, Point.InterpMode );
-			FMatineeImportTools::SetOrAddKey( ChannelInterfaces[1], KeyTime, Point.OutVal.Y, Point.ArriveTangent.Y, Point.LeaveTangent.Y, Point.InterpMode );
-			FMatineeImportTools::SetOrAddKey( ChannelInterfaces[2], KeyTime, Point.OutVal.Z, Point.ArriveTangent.Z, Point.LeaveTangent.Z, Point.InterpMode );
+			FMatineeImportTools::SetOrAddKey( ChannelData[0], KeyTime, Point.OutVal.X, Point.ArriveTangent.X, Point.LeaveTangent.X, Point.InterpMode );
+			FMatineeImportTools::SetOrAddKey( ChannelData[1], KeyTime, Point.OutVal.Y, Point.ArriveTangent.Y, Point.LeaveTangent.Y, Point.InterpMode );
+			FMatineeImportTools::SetOrAddKey( ChannelData[2], KeyTime, Point.OutVal.Z, Point.ArriveTangent.Z, Point.LeaveTangent.Z, Point.InterpMode );
 
 			KeyRange = TRange<FFrameNumber>::Hull(KeyRange, TRange<FFrameNumber>(KeyTime));
 		}
@@ -439,9 +439,9 @@ bool FMatineeImportTools::CopyInterpMoveTrack( UInterpTrackMove* MoveTrack, UMov
 		{
 			FFrameNumber KeyTime = (Point.InVal * FrameRate).RoundToFrame();
 
-			FMatineeImportTools::SetOrAddKey( ChannelInterfaces[3], KeyTime, Point.OutVal.X, Point.ArriveTangent.X, Point.LeaveTangent.X, Point.InterpMode );
-			FMatineeImportTools::SetOrAddKey( ChannelInterfaces[4], KeyTime, Point.OutVal.Y, Point.ArriveTangent.Y, Point.LeaveTangent.Y, Point.InterpMode );
-			FMatineeImportTools::SetOrAddKey( ChannelInterfaces[5], KeyTime, Point.OutVal.Z, Point.ArriveTangent.Z, Point.LeaveTangent.Z, Point.InterpMode );
+			FMatineeImportTools::SetOrAddKey( ChannelData[3], KeyTime, Point.OutVal.X, Point.ArriveTangent.X, Point.LeaveTangent.X, Point.InterpMode );
+			FMatineeImportTools::SetOrAddKey( ChannelData[4], KeyTime, Point.OutVal.Y, Point.ArriveTangent.Y, Point.LeaveTangent.Y, Point.InterpMode );
+			FMatineeImportTools::SetOrAddKey( ChannelData[5], KeyTime, Point.OutVal.Z, Point.ArriveTangent.Z, Point.LeaveTangent.Z, Point.InterpMode );
 
 			KeyRange = TRange<FFrameNumber>::Hull(KeyRange, TRange<FFrameNumber>(KeyTime));
 		}
@@ -486,7 +486,7 @@ bool FMatineeImportTools::CopyInterpMoveTrack( UInterpTrackMove* MoveTrack, UMov
 						{
 							FFrameNumber KeyTime = (Point.InVal * FrameRate).RoundToFrame();
 
-							FMatineeImportTools::SetOrAddKey( ChannelInterfaces[ChannelIndex], KeyTime, Point.OutVal, Point.ArriveTangent, Point.LeaveTangent, Point.InterpMode );
+							FMatineeImportTools::SetOrAddKey( ChannelData[ChannelIndex], KeyTime, Point.OutVal, Point.ArriveTangent, Point.LeaveTangent, Point.InterpMode );
 
 							KeyRange = TRange<FFrameNumber>::Hull(KeyRange, TRange<FFrameNumber>(KeyTime));
 						}
@@ -537,7 +537,7 @@ bool FMatineeImportTools::CopyInterpParticleTrack( UInterpTrackToggle* MatineeTo
 
 		FMovieSceneParticleChannel* ParticleChannel = Section->GetChannelProxy().GetChannel<FMovieSceneParticleChannel>(0);
 		check(ParticleChannel);
-		TMovieSceneChannel<uint8> ChannelInterface = ParticleChannel->GetInterface();
+		TMovieSceneChannelData<uint8> ChannelData = ParticleChannel->GetData();
 
 		for ( const auto& Key : MatineeToggleTrack->ToggleTrack )
 		{
@@ -546,7 +546,7 @@ bool FMatineeImportTools::CopyInterpParticleTrack( UInterpTrackToggle* MatineeTo
 			EParticleKey ParticleKey;
 			if ( TryConvertMatineeToggleToOutParticleKey( Key.ToggleAction, ParticleKey ) )
 			{
-				ChannelInterface.AddKey( KeyTime, (uint8)ParticleKey );
+				ChannelData.AddKey( KeyTime, (uint8)ParticleKey );
 			}
 
 			KeyRange = TRange<FFrameNumber>::Hull(KeyRange, TRange<FFrameNumber>(KeyTime));
@@ -672,7 +672,7 @@ bool FMatineeImportTools::CopyInterpFadeTrack( UInterpTrackFade* MatineeFadeTrac
 
 		FMovieSceneFloatChannel* FadeChannel = Section->GetChannelProxy().GetChannel<FMovieSceneFloatChannel>(0);
 		check(FadeChannel);
-		TMovieSceneChannel<FMovieSceneFloatValue> FadeInterface = FadeChannel->GetInterface();
+		TMovieSceneChannelData<FMovieSceneFloatValue> FadeInterface = FadeChannel->GetData();
 		for ( const auto& Point : MatineeFadeTrack->FloatTrack.Points )
 		{
 			FFrameNumber KeyTime = (Point.InVal * FrameRate).RoundToFrame();
@@ -755,13 +755,13 @@ bool FMatineeImportTools::CopyInterpEventTrack( UInterpTrackEvent* MatineeEventT
 
 			FMovieSceneEventSectionData* EventChannel = Section->GetChannelProxy().GetChannel<FMovieSceneEventSectionData>(0);
 			check(EventChannel);
-			TMovieSceneChannel<FEventPayload> ChannelInterface = EventChannel->GetInterface();
+			TMovieSceneChannelData<FEventPayload> ChannelData = EventChannel->GetData();
 
 			for (FEventTrackKey EventTrackKey : MatineeEventTrack->EventTrack)
 			{
 				FFrameNumber KeyTime = (EventTrackKey.Time * FrameRate).RoundToFrame();
 
-				ChannelInterface.UpdateOrAddKey(KeyTime, FEventPayload(EventTrackKey.EventName));
+				ChannelData.UpdateOrAddKey(KeyTime, FEventPayload(EventTrackKey.EventName));
 
 				KeyRange = TRange<FFrameNumber>::Hull(KeyRange, TRange<FFrameNumber>(KeyTime));
 			}
@@ -803,7 +803,7 @@ bool FMatineeImportTools::CopyInterpVisibilityTrack( UInterpTrackVisibility* Mat
 
 			FMovieSceneBoolChannel* VisibilityChannel = Section->GetChannelProxy().GetChannel<FMovieSceneBoolChannel>(0);
 			check(VisibilityChannel);
-			TMovieSceneChannel<bool> ChannelInterface = VisibilityChannel->GetInterface();
+			TMovieSceneChannelData<bool> ChannelData = VisibilityChannel->GetData();
 
 			for (FVisibilityTrackKey VisibilityTrackKey : MatineeVisibilityTrack->VisibilityTrack)
 			{
@@ -822,7 +822,7 @@ bool FMatineeImportTools::CopyInterpVisibilityTrack( UInterpTrackVisibility* Mat
 
 				FFrameNumber KeyTime = (VisibilityTrackKey.Time * FrameRate).RoundToFrame();
 
-				ChannelInterface.UpdateOrAddKey(KeyTime, bVisible);
+				ChannelData.UpdateOrAddKey(KeyTime, bVisible);
 
 				KeyRange = TRange<FFrameNumber>::Hull(KeyRange, TRange<FFrameNumber>(KeyTime));
 			}
diff --git a/Engine/Source/Editor/MovieSceneTools/Private/MovieSceneToolHelpers.cpp b/Engine/Source/Editor/MovieSceneTools/Private/MovieSceneToolHelpers.cpp
--- a/Engine/Source/Editor/MovieSceneTools/Private/MovieSceneToolHelpers.cpp
+++ b/Engine/Source/Editor/MovieSceneTools/Private/MovieSceneToolHelpers.cpp
@@ -914,9 +914,9 @@ bool ImportFBXProperty(FString NodeName, FString AnimatedPropertyName, FGuid Obj
 				CurveAPI.GetCurveData(NodeName, AnimatedPropertyName, ChannelIndex, CompositeIndex, Source, bNegative);
 
 				FMovieSceneFloatChannel* Channel = FloatSection->GetChannelProxy().GetChannel<FMovieSceneFloatChannel>(0);
-				TMovieSceneChannel<FMovieSceneFloatValue> ChannelInterface = Channel->GetInterface();
+				TMovieSceneChannelData<FMovieSceneFloatValue> ChannelData = Channel->GetData();
 
-				ChannelInterface.Reset();
+				ChannelData.Reset();
 				double DecimalRate = FrameRate.AsDecimal();
 
 				for (auto SourceIt = Source.GetKeyHandleIterator(); SourceIt; ++SourceIt)
@@ -939,7 +939,7 @@ bool ImportFBXProperty(FString NodeName, FString AnimatedPropertyName, FGuid Obj
 					}
 
 					FFrameNumber KeyTime = (Key.Time * FrameRate).RoundToFrame();
-					FMatineeImportTools::SetOrAddKey(ChannelInterface, KeyTime, Key.Value, ArriveTangent, LeaveTangent,
+					FMatineeImportTools::SetOrAddKey(ChannelData, KeyTime, Key.Value, ArriveTangent, LeaveTangent,
 						MovieSceneToolHelpers::RichCurveInterpolationToMatineeInterpolation(Key.InterpMode, Key.TangentMode), Key.TangentWeightMode,
 						Key.ArriveTangentWeight, Key.LeaveTangentWeight);
 
@@ -954,8 +954,7 @@ bool ImportFBXProperty(FString NodeName, FString AnimatedPropertyName, FGuid Obj
 				}
 				Channel->AutoSetTangents();
 
-
-				TArrayView<const FFrameNumber> AllTimes = ChannelInterface.GetTimes();
+				TArrayView<const FFrameNumber> AllTimes = ChannelData.GetTimes();
 				if (AllTimes.Num() > 0)
 				{
 					FloatSection->SetRange(TRange<FFrameNumber>(AllTimes[0], TRangeBound<FFrameNumber>::Inclusive(AllTimes.Last())));
@@ -969,8 +968,8 @@ bool ImportFBXProperty(FString NodeName, FString AnimatedPropertyName, FGuid Obj
 
 void ImportTransformChannel(const FRichCurve& Source, FMovieSceneFloatChannel* Dest, FFrameRate DestFrameRate, bool bNegateTangents)
 {
-	TMovieSceneChannel<FMovieSceneFloatValue> ChannelInterface = Dest->GetInterface();
-	ChannelInterface.Reset();
+	TMovieSceneChannelData<FMovieSceneFloatValue> ChannelData = Dest->GetData();
+	ChannelData.Reset();
 	double DecimalRate = DestFrameRate.AsDecimal();
 
 	for (auto SourceIt = Source.GetKeyHandleIterator(); SourceIt; ++SourceIt)
@@ -999,7 +998,7 @@ void ImportTransformChannel(const FRichCurve& Source, FMovieSceneFloatChannel* D
 		}
 
 		FFrameNumber KeyTime = (Key.Time * DestFrameRate).RoundToFrame();
-		FMatineeImportTools::SetOrAddKey(ChannelInterface, KeyTime, Key.Value, ArriveTangent, LeaveTangent,
+		FMatineeImportTools::SetOrAddKey(ChannelData, KeyTime, Key.Value, ArriveTangent, LeaveTangent,
 			MovieSceneToolHelpers::RichCurveInterpolationToMatineeInterpolation(Key.InterpMode, Key.TangentMode), Key.TangentWeightMode,
 			Key.ArriveTangentWeight, Key.LeaveTangentWeight);
 
@@ -1677,7 +1676,7 @@ EInterpCurveMode MovieSceneToolHelpers::RichCurveInterpolationToMatineeInterpola
 	}
 }
 
-void MovieSceneToolHelpers::CopyKeyDataToMoveAxis(const TMovieSceneChannel<FMovieSceneFloatValue>& Channel, UInterpTrackMoveAxis* MoveAxis, FFrameRate InFrameRate)
+void MovieSceneToolHelpers::CopyKeyDataToMoveAxis(const TMovieSceneChannelData<FMovieSceneFloatValue>& Channel, UInterpTrackMoveAxis* MoveAxis, FFrameRate InFrameRate)
 {
 	MoveAxis->FloatTrack.Points.Reset();
 
@@ -1768,12 +1767,12 @@ UObject* MovieSceneToolHelpers::ExportToCameraAnim(UMovieScene* InMovieScene, FG
 				UMovieScene3DTransformSection* TransformSection = Cast<UMovieScene3DTransformSection>(Sections[0]);
 				TArrayView<FMovieSceneFloatChannel*> FloatChannels = TransformSection->GetChannelProxy().GetChannels<FMovieSceneFloatChannel>();
 
-				CopyKeyDataToMoveAxis(FloatChannels[0]->GetInterface(), MoveAxies[AXIS_TranslationX], TickResolution);
-				CopyKeyDataToMoveAxis(FloatChannels[1]->GetInterface(), MoveAxies[AXIS_TranslationY], TickResolution);
-				CopyKeyDataToMoveAxis(FloatChannels[2]->GetInterface(), MoveAxies[AXIS_TranslationZ], TickResolution);
-				CopyKeyDataToMoveAxis(FloatChannels[3]->GetInterface(), MoveAxies[AXIS_RotationX],    TickResolution);
-				CopyKeyDataToMoveAxis(FloatChannels[4]->GetInterface(), MoveAxies[AXIS_RotationY],    TickResolution);
-				CopyKeyDataToMoveAxis(FloatChannels[5]->GetInterface(), MoveAxies[AXIS_RotationZ],    TickResolution);
+				CopyKeyDataToMoveAxis(FloatChannels[0]->GetData(), MoveAxies[AXIS_TranslationX], TickResolution);
+				CopyKeyDataToMoveAxis(FloatChannels[1]->GetData(), MoveAxies[AXIS_TranslationY], TickResolution);
+				CopyKeyDataToMoveAxis(FloatChannels[2]->GetData(), MoveAxies[AXIS_TranslationZ], TickResolution);
+				CopyKeyDataToMoveAxis(FloatChannels[3]->GetData(), MoveAxies[AXIS_RotationX],    TickResolution);
+				CopyKeyDataToMoveAxis(FloatChannels[4]->GetData(), MoveAxies[AXIS_RotationY],    TickResolution);
+				CopyKeyDataToMoveAxis(FloatChannels[5]->GetData(), MoveAxies[AXIS_RotationZ],    TickResolution);
 			}
 		}
 	}
diff --git a/Engine/Source/Editor/MovieSceneTools/Public/MatineeImportTools.h b/Engine/Source/Editor/MovieSceneTools/Public/MatineeImportTools.h
--- a/Engine/Source/Editor/MovieSceneTools/Public/MatineeImportTools.h
+++ b/Engine/Source/Editor/MovieSceneTools/Public/MatineeImportTools.h
@@ -47,7 +47,7 @@ public:
 	static bool TryConvertMatineeToggleToOutParticleKey( ETrackToggleAction ToggleAction, EParticleKey& OutParticleKey );
 
 	/** Adds a key to a rich curve based on matinee curve key data. */
-	static void SetOrAddKey( TMovieSceneChannel<FMovieSceneFloatValue>& Curve, FFrameNumber Time, float Value, float ArriveTangent, float LeaveTangent, EInterpCurveMode MatineeInterpMode 
+	static void SetOrAddKey( TMovieSceneChannelData<FMovieSceneFloatValue>& Curve, FFrameNumber Time, float Value, float ArriveTangent, float LeaveTangent, EInterpCurveMode MatineeInterpMode
 	, ERichCurveTangentWeightMode WeightedMode = RCTWM_WeightedNone, float ArriveTangentWeight = 0.0f, float LeaveTangentWeight = 0.0f);
 
 	/** Copies keys from a matinee bool track to a sequencer bool track. */
