Sequencer: fixing CIS warnings in FCP XML import/export

diff --git a/Engine/Source/Editor/MovieSceneTools/Private/MovieSceneTranslator.cpp b/Engine/Source/Editor/MovieSceneTools/Private/MovieSceneTranslator.cpp
--- a/Engine/Source/Editor/MovieSceneTools/Private/MovieSceneTranslator.cpp
+++ b/Engine/Source/Editor/MovieSceneTools/Private/MovieSceneTranslator.cpp
@@ -268,6 +268,10 @@ bool FMovieSceneExportData::ConstructAudioMasterTrackData(const UMovieScene* InM
 	for (UMovieSceneSection* Section : InAudioMasterTrack->GetAudioSections())
 	{
 		const UMovieSceneAudioSection* AudioSection = Cast<UMovieSceneAudioSection>(Section);
+		if (AudioSection == nullptr)
+		{
+			continue;
+		}
 
 		// skip duplicate sections
 		if (SectionPathNames.Num() > 0 && SectionPathNames.Contains(AudioSection->GetPathName()))
@@ -275,7 +279,7 @@ bool FMovieSceneExportData::ConstructAudioMasterTrackData(const UMovieScene* InM
 			continue;
 		}
 
-		if (AudioSection != nullptr && AudioSectionIsSoundWave(AudioSection))
+		if (AudioSectionIsSoundWave(AudioSection))
 		{
 			if (!ConstructAudioSectionData(InMovieScene, MasterTrackData, AudioSection))
 			{
@@ -377,7 +381,7 @@ bool FMovieSceneExportData::ConstructAudioSectionData(const UMovieScene* InMovie
 	}
 
 	int32 SampleRate = SoundWave->GetSampleRateForCurrentPlatform();
-	if (SampleRate != 48000 || SampleRate != 44100 || SampleRate != 32000)
+	if (SampleRate != 48000 && SampleRate != 44100 && SampleRate != 32000)
 	{
 		// @todo - warning about invalid sample rate
 		SampleRate = 44100;
