Sequencer: Refresh actor list to clean up "deleted actor" in the outliner when drawing thumbnails

diff --git a/Engine/Source/Editor/MovieSceneTools/Private/TrackEditorThumbnail/TrackEditorThumbnailPool.cpp b/Engine/Source/Editor/MovieSceneTools/Private/TrackEditorThumbnail/TrackEditorThumbnailPool.cpp
--- a/Engine/Source/Editor/MovieSceneTools/Private/TrackEditorThumbnail/TrackEditorThumbnailPool.cpp
+++ b/Engine/Source/Editor/MovieSceneTools/Private/TrackEditorThumbnail/TrackEditorThumbnailPool.cpp
@@ -4,6 +4,7 @@
 #include "TrackEditorThumbnail/TrackEditorThumbnail.h"
 #include "Framework/Application/SlateApplication.h"
 #include "ISequencer.h"
+#include "Editor.h"
 
 
 /* FShotSequencerSection structors
@@ -115,6 +116,12 @@ bool FTrackEditorThumbnailPool::DrawThumbnails()
 	if (ThumbnailsDrawn > 0)
 	{
 		ThumbnailsNeedingDraw.RemoveAt(0, ThumbnailsDrawn, false);
+
+		// If done drawing thumbnails, refresh actor list
+		if (ThumbnailsNeedingDraw.Num() == 0)
+		{
+			GEditor->BroadcastLevelActorListChanged();
+		}
 	}
 
 	TimeOfLastUpdate = CurrentTime;
