Sequencer: Fix file extensions for movie files on mac

diff --git a/Engine/Source/Editor/MovieSceneTools/Private/MovieSceneTranslator.cpp b/Engine/Source/Editor/MovieSceneTools/Private/MovieSceneTranslator.cpp
--- a/Engine/Source/Editor/MovieSceneTools/Private/MovieSceneTranslator.cpp
+++ b/Engine/Source/Editor/MovieSceneTools/Private/MovieSceneTranslator.cpp
@@ -343,7 +343,17 @@ bool FMovieSceneExportData::ConstructCinematicSectionData(const UMovieScene* InM
 	InMasterTrackData->CinematicSections.Add(SectionData);
 
 	SectionData->DisplayName = InCinematicSection->GetShotDisplayName();
-	SectionData->SourceFilename = SectionData->DisplayName + TEXT(".avi");
+
+#if PLATFORM_MAC
+	static const TCHAR* Extension = TEXT(".mov");
+#elif PLATFORM_UNIX
+	static const TCHAR* Extension = TEXT(".unsupp");
+	return false;
+#else
+	static const TCHAR* Extension = TEXT(".avi");
+#endif
+
+	SectionData->SourceFilename = SectionData->DisplayName + Extension;
 	SectionData->SourceFilePath = TEXT("");
 
 	ConstructSectionData(InMovieScene, SectionData, InCinematicSection, EMovieSceneTranslatorSectionType::Cinematic, SectionData->DisplayName);
diff --git a/Engine/Source/Editor/MovieSceneTools/Private/MovieSceneTranslatorEDL.cpp b/Engine/Source/Editor/MovieSceneTools/Private/MovieSceneTranslatorEDL.cpp
--- a/Engine/Source/Editor/MovieSceneTools/Private/MovieSceneTranslatorEDL.cpp
+++ b/Engine/Source/Editor/MovieSceneTools/Private/MovieSceneTranslatorEDL.cpp
@@ -475,8 +475,15 @@ bool MovieSceneTranslatorEDL::ExportEDL(const UMovieScene* InMovieScene, FFrameR
 				FFrameNumber EditInFrame    = CinematicShotSection->GetInclusiveStartFrame();
 				FFrameNumber EditOutFrame   = CinematicShotSection->GetExclusiveEndFrame();
 
-				//@todo can't assume avi's written out
-				ShotName += ".avi";
+#if PLATFORM_MAC
+				static const TCHAR* Extension = TEXT(".mov");
+#elif PLATFORM_UNIX
+				static const TCHAR* Extension = TEXT(".unsupp");
+				continue;
+#else
+				static const TCHAR* Extension = TEXT(".avi");
+#endif
+				ShotName += Extension;
 
 				//@todo shotpath should really be moviefile path
 				ShotPath = ShotName;
