Sequencer: Sequencer Event Improvements
  - New events mandate a Blueprint funtion entry node as an endpoint.
  - 2 new event sections exist: Triggers and Repeaters.
  - New director blueprint exists as a child of level sequences.
  - Added a new concept for editor-only customization of sequence types - IMovieSceneSequenceEditor, analogous to IMovieSceneTrackEditor
  - Outstanding tasks: FiB integration, Nativization
  - Potential future work:
     - Expose the director blueprint more prominently
     - Usability and UX polish

diff --git a/Engine/Config/BaseEditorPerProjectUserSettings.ini b/Engine/Config/BaseEditorPerProjectUserSettings.ini
--- a/Engine/Config/BaseEditorPerProjectUserSettings.ini
+++ b/Engine/Config/BaseEditorPerProjectUserSettings.ini
@@ -777,6 +777,8 @@ bKeepPlayRangeInSectionBounds=false
 ZeroPadFrames=4
 bInfiniteKeyAreas=true
 bAutoSetTrackDefaults=true
+bCompileDirectorOnEvaluate=false
+
 [NiagaraSequenceEditor SequencerSettings]
 bAutoScrollEnabled=true
 bKeepPlayRangeInSectionBounds=false
diff --git a/Engine/Plugins/MovieScene/ActorSequence/Source/ActorSequence/Private/ActorSequence.cpp b/Engine/Plugins/MovieScene/ActorSequence/Source/ActorSequence/Private/ActorSequence.cpp
--- a/Engine/Plugins/MovieScene/ActorSequence/Source/ActorSequence/Private/ActorSequence.cpp
+++ b/Engine/Plugins/MovieScene/ActorSequence/Source/ActorSequence/Private/ActorSequence.cpp
@@ -8,6 +8,7 @@
 #include "Engine/Blueprint.h"
 #include "GameFramework/Actor.h"
 #include "ActorSequenceComponent.h"
+#include "Engine/LevelScriptActor.h"
 
 IMPLEMENT_MODULE(FDefaultModuleImpl, ActorSequence);
 
@@ -164,6 +165,20 @@ void UActorSequence::UnbindPossessableObjects(const FGuid& ObjectId)
 	ObjectReferences.RemoveBinding(ObjectId);
 }
 
+UObject* UActorSequence::CreateDirectorInstance(IMovieScenePlayer& Player)
+{
+	AActor* Actor = CastChecked<AActor>(Player.GetPlaybackContext());
+
+	// If this sequence is inside a blueprint, or its component's archetype is from a blueprint, we use the actor as the instace (which will be an instance of the blueprint itself)
+	if (GetTypedOuter<UBlueprintGeneratedClass>() || GetTypedOuter<UActorSequenceComponent>()->GetArchetype() != GetDefault<UActorSequenceComponent>())
+	{
+		return Actor;
+	}
+
+	// Otherwise we use the level script actor as the instance
+	return Actor->GetLevel()->GetLevelScriptActor();
+}
+
 #if WITH_EDITOR
 FText UActorSequence::GetDisplayName() const
 {
diff --git a/Engine/Plugins/MovieScene/ActorSequence/Source/ActorSequence/Public/ActorSequence.h b/Engine/Plugins/MovieScene/ActorSequence/Source/ActorSequence/Public/ActorSequence.h
--- a/Engine/Plugins/MovieScene/ActorSequence/Source/ActorSequence/Public/ActorSequence.h
+++ b/Engine/Plugins/MovieScene/ActorSequence/Source/ActorSequence/Public/ActorSequence.h
@@ -26,6 +26,7 @@ public:
 	virtual UMovieScene* GetMovieScene() const override;
 	virtual UObject* GetParentObject(UObject* Object) const override;
 	virtual void UnbindPossessableObjects(const FGuid& ObjectId) override;
+	virtual UObject* CreateDirectorInstance(IMovieScenePlayer& Player) override;
 
 #if WITH_EDITOR
 	virtual FText GetDisplayName() const override;
diff --git a/Engine/Plugins/MovieScene/ActorSequenceEditor/Source/ActorSequenceEditor/ActorSequenceEditor.Build.cs b/Engine/Plugins/MovieScene/ActorSequenceEditor/Source/ActorSequenceEditor/ActorSequenceEditor.Build.cs
--- a/Engine/Plugins/MovieScene/ActorSequenceEditor/Source/ActorSequenceEditor/ActorSequenceEditor.Build.cs
+++ b/Engine/Plugins/MovieScene/ActorSequenceEditor/Source/ActorSequenceEditor/ActorSequenceEditor.Build.cs
@@ -9,6 +9,7 @@ public class ActorSequenceEditor : ModuleRules
 		PrivateDependencyModuleNames.AddRange(
 			new string[] {
 				"ActorSequence",
+				"BlueprintGraph",
 				"Core",
 				"CoreUObject",
 				"Engine",
diff --git a/Engine/Plugins/MovieScene/ActorSequenceEditor/Source/ActorSequenceEditor/Private/ActorSequenceEditor.cpp b/Engine/Plugins/MovieScene/ActorSequenceEditor/Source/ActorSequenceEditor/Private/ActorSequenceEditor.cpp
--- a/Engine/Plugins/MovieScene/ActorSequenceEditor/Source/ActorSequenceEditor/Private/ActorSequenceEditor.cpp
+++ b/Engine/Plugins/MovieScene/ActorSequenceEditor/Source/ActorSequenceEditor/Private/ActorSequenceEditor.cpp
@@ -5,6 +5,7 @@
 #include "BlueprintEditorModule.h"
 #include "BlueprintEditorTabs.h"
 #include "ActorSequenceComponentCustomization.h"
+#include "MovieSceneSequenceEditor_ActorSequence.h"
 #include "ActorSequenceEditorStyle.h"
 #include "ActorSequenceEditorTabSummoner.h"
 #include "Framework/Docking/LayoutExtender.h"
@@ -17,6 +18,7 @@
 #include "Widgets/Docking/SDockTab.h"
 #include "ISettingsModule.h"
 #include "SequencerSettings.h"
+#include "ISequencerModule.h"
 
 
 #define LOCTEXT_NAMESPACE "ActorSequenceEditor"
@@ -120,6 +122,9 @@ public:
 		RegisterCustomizations();
 		RegisterSettings();
 		OnInitializeSequenceHandle = UActorSequence::OnInitializeSequence().AddStatic(FActorSequenceEditorModule::OnInitializeSequence);
+
+		ISequencerModule& SequencerModule = FModuleManager::Get().LoadModuleChecked<ISequencerModule>("Sequencer");
+		SequenceEditorHandle = SequencerModule.RegisterSequenceEditor(UActorSequence::StaticClass(), MakeUnique<FMovieSceneSequenceEditor_ActorSequence>());
 	}
 	
 	virtual void ShutdownModule() override
@@ -127,6 +132,13 @@ public:
 		UActorSequence::OnInitializeSequence().Remove(OnInitializeSequenceHandle);
 		UnregisterCustomizations();
 		UnregisterSettings();
+
+		ISequencerModule* SequencerModule = FModuleManager::Get().GetModulePtr<ISequencerModule>("Sequencer");
+		if (SequencerModule)
+		{
+			SequencerModule->UnregisterSequenceEditor(SequenceEditorHandle);
+		}
+
 		BlueprintEditorTabBinding = nullptr;
 	}
 
@@ -195,6 +207,7 @@ public:
 		}
 	}
 
+	FDelegateHandle SequenceEditorHandle;
 	FDelegateHandle OnInitializeSequenceHandle;
 	TSharedPtr<FActorSequenceEditorTabBinding> BlueprintEditorTabBinding;
 	FName ActorSequenceComponentName;
diff --git a/Engine/Plugins/MovieScene/ActorSequenceEditor/Source/ActorSequenceEditor/Private/MovieSceneSequenceEditor_ActorSequence.h b/Engine/Plugins/MovieScene/ActorSequenceEditor/Source/ActorSequenceEditor/Private/MovieSceneSequenceEditor_ActorSequence.h
new file mode 100644
--- /dev/null
+++ b/Engine/Plugins/MovieScene/ActorSequenceEditor/Source/ActorSequenceEditor/Private/MovieSceneSequenceEditor_ActorSequence.h
@@ -0,0 +1,65 @@
+// Copyright 1998-2018 Epic Games, Inc. All Rights Reserved.
+
+#pragma once
+
+#include "MovieSceneSequenceEditor.h"
+#include "ActorSequence.h"
+#include "Engine/Level.h"
+#include "Engine/LevelScriptBlueprint.h"
+#include "Kismet2/KismetEditorUtilities.h"
+#include "Tracks/MovieSceneEventTrack.h"
+#include "K2Node_FunctionEntry.h"
+#include "EdGraphSchema_K2.h"
+
+struct FMovieSceneSequenceEditor_ActorSequence : FMovieSceneSequenceEditor
+{
+	virtual bool CanCreateEvents(UMovieSceneSequence* InSequence) const
+	{
+		return true;
+	}
+
+	virtual UBlueprint* GetBlueprintForSequence(UMovieSceneSequence* InSequence) const override
+	{
+		UActorSequence* ActorSequence = CastChecked<UActorSequence>(InSequence);
+		if (UBlueprint* Blueprint = ActorSequence->GetParentBlueprint())
+		{
+			return Blueprint;
+		}
+
+		UActorSequenceComponent* Component = ActorSequence->GetTypedOuter<UActorSequenceComponent>();
+		ULevel* Level = Component ? Component->GetOwner()->GetLevel() : nullptr;
+
+		bool bDontCreateNewBlueprint = true;
+		return Level ? Level->GetLevelScriptBlueprint(bDontCreateNewBlueprint) : nullptr;
+	}
+
+	virtual UBlueprint* CreateBlueprintForSequence(UMovieSceneSequence* InSequence) const override
+	{
+		UActorSequence* ActorSequence = CastChecked<UActorSequence>(InSequence);
+		check(!ActorSequence->GetParentBlueprint());
+
+		UActorSequenceComponent* Component = ActorSequence->GetTypedOuter<UActorSequenceComponent>();
+		ULevel* Level = Component ? Component->GetOwner()->GetLevel() : nullptr;
+
+		bool bDontCreateNewBlueprint = false;
+		return Level ? Level->GetLevelScriptBlueprint(bDontCreateNewBlueprint) : nullptr;
+	}
+
+	virtual void SetupDefaultPinForEndpoint(UMovieSceneEventTrack* EventTrack, UK2Node_FunctionEntry* Endpoint) const override
+	{
+		// By default for master tracks we create a pin (that receives the level blueprint for master event tracks, or the event receivers)
+		UClass* PinClass = nullptr;
+
+		// When event receivers are not set, we set the pin type to the type of the object binding
+		if (EventTrack->EventReceivers.Num() == 0)
+		{
+			PinClass = FindTrackObjectBindingClass(EventTrack);
+		}
+
+		FEdGraphPinType PinType;
+		PinType.PinCategory = UEdGraphSchema_K2::PC_Object;
+		PinType.PinSubCategoryObject = PinClass ? PinClass : UObject::StaticClass();
+
+		Endpoint->CreateUserDefinedPin(TargetPinName, PinType, EGPD_Output, true);
+	}
+};
\ No newline at end of file
diff --git a/Engine/Plugins/MovieScene/LevelSequenceEditor/Source/LevelSequenceEditor/LevelSequenceEditor.Build.cs b/Engine/Plugins/MovieScene/LevelSequenceEditor/Source/LevelSequenceEditor/LevelSequenceEditor.Build.cs
--- a/Engine/Plugins/MovieScene/LevelSequenceEditor/Source/LevelSequenceEditor/LevelSequenceEditor.Build.cs
+++ b/Engine/Plugins/MovieScene/LevelSequenceEditor/Source/LevelSequenceEditor/LevelSequenceEditor.Build.cs
@@ -25,6 +25,8 @@ public class LevelSequenceEditor : ModuleRules
                 "EditorStyle",
                 "Engine",
                 "InputCore",
+                "Kismet",
+                "KismetCompiler",
 				"LevelEditor",
 				"MovieScene",
                 "MovieSceneTools",
diff --git a/Engine/Plugins/MovieScene/LevelSequenceEditor/Source/LevelSequenceEditor/Private/LevelSequenceDirectorBlueprint.cpp b/Engine/Plugins/MovieScene/LevelSequenceEditor/Source/LevelSequenceEditor/Private/LevelSequenceDirectorBlueprint.cpp
new file mode 100644
--- /dev/null
+++ b/Engine/Plugins/MovieScene/LevelSequenceEditor/Source/LevelSequenceEditor/Private/LevelSequenceDirectorBlueprint.cpp
@@ -0,0 +1,66 @@
+// Copyright 1998-2018 Epic Games, Inc. All Rights Reserved.
+
+#include "LevelSequenceDirectorBlueprint.h"
+#include "LevelSequenceDirectorGeneratedClass.h"
+#include "Kismet2/KismetReinstanceUtilities.h"
+#include "Modules/ModuleManager.h"
+
+ULevelSequenceDirectorBlueprint::ULevelSequenceDirectorBlueprint(const FObjectInitializer& ObjInit)
+	: Super(ObjInit)
+{
+	ParentClass = ULevelSequenceDirector::StaticClass();
+}
+
+UClass* ULevelSequenceDirectorBlueprint::GetBlueprintClass() const
+{
+	return ULevelSequenceDirectorGeneratedClass::StaticClass();
+}
+
+bool ULevelSequenceDirectorBlueprint::SupportsNativization(FText* OutReason) const
+{
+	if (OutReason)
+	{
+		*OutReason = NSLOCTEXT("LevelSequenceDirectorBlueprint", "NativizationError", "Level Sequence Director Blueprints do not support nativization.");
+	}
+	return false;
+}
+
+bool ULevelSequenceDirectorBlueprint::SupportedByDefaultBlueprintFactory() const
+{
+	return false;
+}
+
+bool ULevelSequenceDirectorBlueprint::AlwaysCompileOnLoad() const
+{
+	return true;
+}
+
+void ULevelSequenceDirectorBlueprint::LoadModulesRequiredForCompilation()
+{
+	static const FName ModuleName(TEXT("LevelSequence"));
+	FModuleManager::Get().LoadModule(ModuleName);
+}
+
+void FLevelSequenceDirectorBlueprintCompiler::SpawnNewClass(const FString& NewClassName)
+{
+	NewClass = FindObject<ULevelSequenceDirectorGeneratedClass>(Blueprint->GetOutermost(), *NewClassName);
+
+	if (!NewClass)
+	{
+		NewClass = NewObject<ULevelSequenceDirectorGeneratedClass>(Blueprint->GetOutermost(), FName(*NewClassName), RF_Public | RF_Transactional);
+	}
+	else
+	{
+		// Already existed, but wasn't linked in the Blueprint yet due to load ordering issues
+		FBlueprintCompileReinstancer::Create(NewClass);
+	}
+}
+
+void FLevelSequenceDirectorBlueprintCompiler::EnsureProperGeneratedClass(UClass*& TargetUClass)
+{
+	if ( TargetUClass && !( (UObject*)TargetUClass )->IsA(ULevelSequenceDirectorGeneratedClass::StaticClass()) )
+	{
+		FKismetCompilerUtilities::ConsignToOblivion(TargetUClass, Blueprint->bIsRegeneratingOnLoad);
+		TargetUClass = nullptr;
+	}
+}
\ No newline at end of file
diff --git a/Engine/Plugins/MovieScene/LevelSequenceEditor/Source/LevelSequenceEditor/Private/LevelSequenceDirectorBlueprint.h b/Engine/Plugins/MovieScene/LevelSequenceEditor/Source/LevelSequenceEditor/Private/LevelSequenceDirectorBlueprint.h
new file mode 100644
--- /dev/null
+++ b/Engine/Plugins/MovieScene/LevelSequenceEditor/Source/LevelSequenceEditor/Private/LevelSequenceDirectorBlueprint.h
@@ -0,0 +1,53 @@
+// Copyright 1998-2018 Epic Games, Inc. All Rights Reserved.
+
+#pragma once
+
+#include "CoreTypes.h"
+#include "Engine/Blueprint.h"
+#include "KismetCompiler.h"
+
+#include "LevelSequenceDirectorBlueprint.generated.h"
+
+class ULevelSequence;
+
+UCLASS(BlueprintType)
+class ULevelSequenceDirectorBlueprint : public UBlueprint
+{
+	GENERATED_BODY()
+
+	ULevelSequenceDirectorBlueprint(const FObjectInitializer& ObjInit);
+
+public:
+
+	UPROPERTY()
+	ULevelSequence* OwnerSequence;
+
+	//~ Begin UBlueprint Interface
+	virtual UClass* GetBlueprintClass() const override;
+	virtual bool SupportedByDefaultBlueprintFactory() const override;
+	virtual bool AlwaysCompileOnLoad() const override;
+	virtual bool SupportsNativization(FText* OutReason) const override;
+
+private:
+
+#if WITH_EDITOR
+	virtual void LoadModulesRequiredForCompilation() override;
+#endif
+};
+
+class FLevelSequenceDirectorBlueprintCompiler : public FKismetCompilerContext
+{
+public:
+
+	FLevelSequenceDirectorBlueprintCompiler(UBlueprint* SourceSketch, FCompilerResultsLog& InMessageLog, const FKismetCompilerOptions& InCompilerOptions, TArray<UObject*>* InObjLoaded)
+		: FKismetCompilerContext(SourceSketch, InMessageLog, InCompilerOptions, InObjLoaded)
+	{}
+
+protected:
+
+	// FKismetCompilerContext
+	virtual void SpawnNewClass(const FString& NewClassName) override;
+	virtual void EnsureProperGeneratedClass(UClass*& TargetClass) override;
+	// End FKismetCompilerContext
+};
+
diff --git a/Engine/Plugins/MovieScene/LevelSequenceEditor/Source/LevelSequenceEditor/Private/LevelSequenceEditorModule.cpp b/Engine/Plugins/MovieScene/LevelSequenceEditor/Source/LevelSequenceEditor/Private/LevelSequenceEditorModule.cpp
--- a/Engine/Plugins/MovieScene/LevelSequenceEditor/Source/LevelSequenceEditor/Private/LevelSequenceEditorModule.cpp
+++ b/Engine/Plugins/MovieScene/LevelSequenceEditor/Source/LevelSequenceEditor/Private/LevelSequenceEditorModule.cpp
@@ -35,6 +35,9 @@
 #include "ILevelSequenceModule.h"
 #include "Misc/LevelSequenceEditorActorSpawner.h"
 #include "SequencerSettings.h"
+#include "Misc/MovieSceneSequenceEditor_LevelSequence.h"
+#include "KismetCompilerModule.h"
+#include "BlueprintAssetHandler.h"
 
 #define LOCTEXT_NAMESPACE "LevelSequenceEditor"
 
@@ -46,7 +49,7 @@ TSharedPtr<FLevelSequenceEditorStyle> FLevelSequenceEditorStyle::Singleton;
  * Implements the LevelSequenceEditor module.
  */
 class FLevelSequenceEditorModule
-	: public ILevelSequenceEditorModule, public FGCObject
+	: public ILevelSequenceEditorModule, public IBlueprintCompiler, public FGCObject
 {
 public:
 
@@ -68,6 +71,26 @@ public:
 		RegisterLevelEditorExtensions();
 		RegisterPlacementModeExtensions();
 		RegisterSettings();
+		RegisterSequenceEditor();
+
+		IKismetCompilerInterface& KismetCompilerModule = FModuleManager::LoadModuleChecked<IKismetCompilerInterface>("KismetCompiler");
+		KismetCompilerModule.GetCompilers().Add(this);
+
+		class FLevelSequenceAssetBlueprintHandler : public IBlueprintAssetHandler
+		{
+			virtual UBlueprint* RetrieveBlueprint(UObject* InObject) const override
+			{
+				return Cast<UBlueprint>(CastChecked<ULevelSequence>(InObject)->DirectorBlueprint);
+			}
+
+			virtual bool AssetContainsBlueprint(const FAssetData& InAssetData) const
+			{
+				// Only have a blueprint if it contains the BlueprintPathWithinPackage tag
+				return InAssetData.TagsAndValues.Find(FBlueprintTags::BlueprintPathWithinPackage) != nullptr;
+			}
+		};
+
+		FBlueprintAssetHandler::Get().RegisterHandler<FLevelSequenceAssetBlueprintHandler>(ULevelSequence::StaticClass()->GetFName());
 	}
 	
 	virtual void ShutdownModule() override
@@ -79,6 +102,13 @@ public:
 		UnregisterLevelEditorExtensions();
 		UnregisterPlacementModeExtensions();
 		UnregisterSettings();
+		UnregisterSequenceEditor();
+
+		IKismetCompilerInterface* KismetCompilerModulePtr = FModuleManager::GetModulePtr<IKismetCompilerInterface>("KismetCompiler");
+		if (KismetCompilerModulePtr)
+		{
+			KismetCompilerModulePtr->GetCompilers().Remove(this);
+		}
 	}
 
 protected:
@@ -189,6 +219,12 @@ protected:
 		}
 	}
 
+	void RegisterSequenceEditor()
+	{
+		ISequencerModule& SequencerModule = FModuleManager::LoadModuleChecked<ISequencerModule>("Sequencer");
+		SequenceEditorHandle = SequencerModule.RegisterSequenceEditor(ULevelSequence::StaticClass(), MakeUnique<FMovieSceneSequenceEditor_LevelSequence>());
+	}
+
 protected:
 
 	/** Unregisters sequencer editor object bindings */
@@ -273,6 +309,32 @@ protected:
 		}
 	}
 
+	void UnregisterSequenceEditor()
+	{
+		ISequencerModule* SequencerModule = FModuleManager::GetModulePtr<ISequencerModule>("Sequencer");
+		if (SequencerModule)
+		{
+			SequencerModule->UnregisterSequenceEditor(SequenceEditorHandle);
+		}
+	}
+
+protected:
+
+	virtual bool CanCompile(const UBlueprint* Blueprint) override
+	{
+		return Blueprint && Blueprint->IsA<ULevelSequenceDirectorBlueprint>();
+	}
+
+	virtual void Compile(UBlueprint* Blueprint, const FKismetCompilerOptions& CompileOptions, FCompilerResultsLog& Results, TArray<UObject*>* ObjLoaded) override
+	{
+		if (ULevelSequenceDirectorBlueprint* LevelSequenceDirectorBP = CastChecked<ULevelSequenceDirectorBlueprint>(Blueprint) )
+		{
+			FLevelSequenceDirectorBlueprintCompiler Compiler(LevelSequenceDirectorBP, Results, CompileOptions, ObjLoaded);
+			Compiler.Compile();
+			check(Compiler.NewClass);
+		}
+	}
+
 protected:
 
 	/** Callback for creating a new level sequence asset in the level. */
@@ -366,6 +428,8 @@ private:
 	FDelegateHandle EditorActorSpawnerDelegateHandle;
 
 	USequencerSettings* Settings;
+
+	FDelegateHandle SequenceEditorHandle;
 };
 
 
diff --git a/Engine/Plugins/MovieScene/LevelSequenceEditor/Source/LevelSequenceEditor/Private/Misc/MovieSceneSequenceEditor_LevelSequence.h b/Engine/Plugins/MovieScene/LevelSequenceEditor/Source/LevelSequenceEditor/Private/Misc/MovieSceneSequenceEditor_LevelSequence.h
new file mode 100644
--- /dev/null
+++ b/Engine/Plugins/MovieScene/LevelSequenceEditor/Source/LevelSequenceEditor/Private/Misc/MovieSceneSequenceEditor_LevelSequence.h
@@ -0,0 +1,67 @@
+// Copyright 1998-2018 Epic Games, Inc. All Rights Reserved.
+
+#pragma once
+
+#include "MovieSceneSequenceEditor.h"
+#include "LevelSequenceDirectorGeneratedClass.h"
+#include "LevelSequenceDirectorBlueprint.h"
+#include "Kismet2/KismetEditorUtilities.h"
+#include "Tracks/MovieSceneEventTrack.h"
+#include "K2Node_FunctionEntry.h"
+
+struct FMovieSceneSequenceEditor_LevelSequence : FMovieSceneSequenceEditor
+{
+	FName BlueprintName;
+
+	FMovieSceneSequenceEditor_LevelSequence()
+	{
+		BlueprintName = "SequenceDirectorBlueprint";
+	}
+
+	virtual bool CanCreateEvents(UMovieSceneSequence* InSequence) const
+	{
+		return true;
+	}
+
+	virtual UBlueprint* GetBlueprintForSequence(UMovieSceneSequence* InSequence) const override
+	{
+		ULevelSequence* LevelSequence = CastChecked<ULevelSequence>(InSequence);
+		return CastChecked<UBlueprint>(LevelSequence->DirectorBlueprint, ECastCheckedType::NullAllowed);
+	}
+
+	virtual UBlueprint* CreateBlueprintForSequence(UMovieSceneSequence* InSequence) const override
+	{
+		UBlueprint* Blueprint = GetBlueprintForSequence(InSequence);
+		if (!ensureMsgf(!Blueprint, TEXT("Should not call CreateBlueprintForSequence when one already exists")))
+		{
+			return Blueprint;
+		}
+
+		ULevelSequence* LevelSequence = CastChecked<ULevelSequence>(InSequence);
+
+		Blueprint = FKismetEditorUtilities::CreateBlueprint(ULevelSequenceDirector::StaticClass(), InSequence, NAME_None, BPTYPE_Normal, ULevelSequenceDirectorBlueprint::StaticClass(), ULevelSequenceDirectorGeneratedClass::StaticClass());
+		CastChecked<ULevelSequenceDirectorBlueprint>(Blueprint)->OwnerSequence = LevelSequence;
+
+		LevelSequence->DirectorBlueprint = Blueprint;
+		LevelSequence->DirectorClass = CastChecked<ULevelSequenceDirectorGeneratedClass>(Blueprint->GeneratedClass);
+		return Blueprint;
+	}
+
+	virtual void SetupDefaultPinForEndpoint(UMovieSceneEventTrack* EventTrack, UK2Node_FunctionEntry* Endpoint) const override
+	{
+		// By default for master tracks we create a pin (that receives the level blueprint for master event tracks, or the event receivers)
+		UClass* PinClass = nullptr;
+
+		// When event receivers are not set, we set the pin type to the type of the object binding
+		if (EventTrack->EventReceivers.Num() == 0)
+		{
+			PinClass = FindTrackObjectBindingClass(EventTrack);
+		}
+
+		FEdGraphPinType PinType;
+		PinType.PinCategory = UEdGraphSchema_K2::PC_Object;
+		PinType.PinSubCategoryObject = PinClass ? PinClass : UObject::StaticClass();
+
+		Endpoint->CreateUserDefinedPin(TargetPinName, PinType, EGPD_Output, true);
+	}
+};
\ No newline at end of file
diff --git a/Engine/Source/Editor/EditorStyle/Private/SlateEditorStyle.cpp b/Engine/Source/Editor/EditorStyle/Private/SlateEditorStyle.cpp
--- a/Engine/Source/Editor/EditorStyle/Private/SlateEditorStyle.cpp
+++ b/Engine/Source/Editor/EditorStyle/Private/SlateEditorStyle.cpp
@@ -1772,6 +1772,7 @@ void FSlateEditorStyle::FStyle::SetupGeneralStyles()
 
 		Set( "Sequencer.KeyCircle", new IMAGE_BRUSH( "Sequencer/KeyCircle", Icon12x12 ) );
 		Set( "Sequencer.KeyDiamond", new IMAGE_BRUSH( "Sequencer/KeyDiamond", Icon12x12 ) );
+		Set( "Sequencer.KeyDiamondBorder", new IMAGE_BRUSH( "Sequencer/KeyDiamondBorder", Icon12x12 ) );
 		Set( "Sequencer.KeySquare", new IMAGE_BRUSH( "Sequencer/KeySquare", Icon12x12 ) );
 		Set( "Sequencer.KeyTriangle", new IMAGE_BRUSH( "Sequencer/KeyTriangle", Icon12x12 ) );
 		Set( "Sequencer.KeyLeft", new IMAGE_BRUSH( "Sequencer/KeyLeft", Icon12x12 ) );
@@ -2007,7 +2008,12 @@ void FSlateEditorStyle::FStyle::SetupGeneralStyles()
 			)
 			.SetDownArrowImage(IMAGE_BRUSH("Common/ComboArrow", Icon8x8));
 		Set( "Sequencer.SectionComboButton", SequencerSectionComboButton );
-		
+
+		Set("Sequencer.CreateEventBinding", new IMAGE_BRUSH("Icons/icon_Blueprint_AddFunction_16px", Icon16x16));
+		Set("Sequencer.CreateQuickBinding", new IMAGE_BRUSH("Icons/icon_Blueprint_Node_16x", Icon16x16));
+		Set("Sequencer.ClearEventBinding", new IMAGE_BRUSH("Icons/Edit/icon_Edit_Delete_40x", Icon16x16));
+		Set("Sequencer.MultipleEvents", new IMAGE_BRUSH("Sequencer/MultipleEvents", Icon16x16));
+		Set("Sequencer.UnboundEvent", new IMAGE_BRUSH("Sequencer/UnboundEvent", Icon16x16));
 
 		// Sequencer Blending Iconography
 		Set( "EMovieSceneBlendType::Absolute", new IMAGE_BRUSH( "Sequencer/EMovieSceneBlendType_Absolute", FVector2D(32, 16) ) );
diff --git a/Engine/Source/Editor/MovieSceneTools/Private/Channels/BuiltInChannelEditors.cpp b/Engine/Source/Editor/MovieSceneTools/Private/Channels/BuiltInChannelEditors.cpp
--- a/Engine/Source/Editor/MovieSceneTools/Private/Channels/BuiltInChannelEditors.cpp
+++ b/Engine/Source/Editor/MovieSceneTools/Private/Channels/BuiltInChannelEditors.cpp
@@ -336,6 +336,32 @@ void DrawKeys(FMovieSceneParticleChannel* Channel, TArrayView<const FKeyHandle>
 	}
 }
 
+void DrawKeys(FMovieSceneEventChannel* Channel, TArrayView<const FKeyHandle> InKeyHandles, TArrayView<FKeyDrawParams> OutKeyDrawParams)
+{
+	FKeyDrawParams ValidEventParams, InvalidEventParams;
+
+	ValidEventParams.BorderBrush   = ValidEventParams.FillBrush   = FEditorStyle::Get().GetBrush("Sequencer.KeyDiamond");
+
+	InvalidEventParams.FillBrush   = FEditorStyle::Get().GetBrush("Sequencer.KeyDiamond");
+	InvalidEventParams.BorderBrush = FEditorStyle::Get().GetBrush("Sequencer.KeyDiamondBorder");
+	InvalidEventParams.FillTint    = FLinearColor(1.f,1.f,1.f,.2f);
+
+	TMovieSceneChannelData<FMovieSceneEvent> ChannelData = Channel->GetData();
+	TArrayView<const FMovieSceneEvent>       Events = ChannelData.GetValues();
+
+	for (int32 Index = 0; Index < InKeyHandles.Num(); ++Index)
+	{
+		int32 KeyIndex = ChannelData.GetIndex(InKeyHandles[Index]);
+		if (KeyIndex != INDEX_NONE && Events[KeyIndex].IsBoundToBlueprint())
+		{
+			OutKeyDrawParams[Index] = ValidEventParams;
+		}
+		else
+		{
+			OutKeyDrawParams[Index] = InvalidEventParams;
+		}
+	}
+}
 
 struct FFloatChannelKeyMenuExtension : FExtender, TSharedFromThis<FFloatChannelKeyMenuExtension>
 {
diff --git a/Engine/Source/Editor/MovieSceneTools/Private/Channels/BuiltInChannelEditors.h b/Engine/Source/Editor/MovieSceneTools/Private/Channels/BuiltInChannelEditors.h
--- a/Engine/Source/Editor/MovieSceneTools/Private/Channels/BuiltInChannelEditors.h
+++ b/Engine/Source/Editor/MovieSceneTools/Private/Channels/BuiltInChannelEditors.h
@@ -16,6 +16,7 @@
 #include "Channels/MovieSceneByteChannel.h"
 #include "Channels/MovieSceneFloatChannel.h"
 #include "Channels/MovieSceneIntegerChannel.h"
+#include "Channels/MovieSceneEventChannel.h"
 #include "Sections/MovieSceneStringSection.h"
 #include "Sections/MovieSceneParticleSection.h"
 #include "Sections/MovieSceneActorReferenceSection.h"
@@ -60,6 +61,7 @@ TSharedPtr<FStructOnScope> GetKeyStruct(const TMovieSceneChannelHandle<FMovieSce
 /** Key drawing overrides */
 void DrawKeys(FMovieSceneFloatChannel*    Channel, TArrayView<const FKeyHandle> InKeyHandles, TArrayView<FKeyDrawParams> OutKeyDrawParams);
 void DrawKeys(FMovieSceneParticleChannel* Channel, TArrayView<const FKeyHandle> InKeyHandles, TArrayView<FKeyDrawParams> OutKeyDrawParams);
+void DrawKeys(FMovieSceneEventChannel*    Channel, TArrayView<const FKeyHandle> InKeyHandles, TArrayView<FKeyDrawParams> OutKeyDrawParams);
 
 /** Context menu overrides */
 void ExtendSectionMenu(FMenuBuilder& OuterMenuBuilder, TArray<TMovieSceneChannelHandle<FMovieSceneFloatChannel>>&& Channels, TArrayView<UMovieSceneSection* const> Sections, TWeakPtr<ISequencer> InSequencer);
diff --git a/Engine/Source/Editor/MovieSceneTools/Private/MovieSceneEventCustomization.cpp b/Engine/Source/Editor/MovieSceneTools/Private/MovieSceneEventCustomization.cpp
new file mode 100644
--- /dev/null
+++ b/Engine/Source/Editor/MovieSceneTools/Private/MovieSceneEventCustomization.cpp
@@ -0,0 +1,623 @@
+// Copyright 1998-2018 Epic Games, Inc. All Rights Reserved.
+
+#include "MovieSceneEventCustomization.h"
+#include "Modules/ModuleManager.h"
+#include "Algo/Find.h"
+
+#include "UObject/UnrealType.h"
+#include "Channels/MovieSceneEvent.h"
+#include "Tracks/MovieSceneEventTrack.h"
+#include "Sections/MovieSceneEventSectionBase.h"
+#include "MovieSceneSequence.h"
+#include "ISequencerModule.h"
+#include "ScopedTransaction.h"
+
+#include "Framework/MultiBox/MultiBoxBuilder.h"
+#include "Widgets/DeclarativeSyntaxSupport.h"
+#include "Widgets/Text/STextBlock.h"
+#include "Widgets/Images/SImage.h"
+#include "Widgets/Input/SComboButton.h"
+#include "Widgets/Input/SHyperlink.h"
+#include "Widgets/Input/SEditableTextBox.h"
+
+#include "PropertyHandle.h"
+#include "PropertyCustomizationHelpers.h"
+#include "IDetailChildrenBuilder.h"
+
+#include "K2Node_Variable.h"
+#include "K2Node_FunctionEntry.h"
+#include "K2Node_CallFunction.h"
+#include "Kismet2/KismetEditorUtilities.h"
+#include "Kismet2/BlueprintEditorUtils.h"
+#include "Kismet2/Kismet2NameValidators.h"
+
+#include "EditorStyleSet.h"
+#include "EditorFontGlyphs.h"
+
+#define LOCTEXT_NAMESPACE "MovieSceneEventCustomization"
+
+TSharedRef<IPropertyTypeCustomization> FMovieSceneEventCustomization::MakeInstance()
+{
+	return MakeShared<FMovieSceneEventCustomization>();
+}
+
+TSharedRef<IPropertyTypeCustomization> FMovieSceneEventCustomization::MakeInstance(UMovieSceneSection* InSection)
+{
+	TSharedRef<FMovieSceneEventCustomization> Custo = MakeShared<FMovieSceneEventCustomization>();
+	Custo->WeakExternalSection = InSection;
+	return Custo;
+}
+
+void FMovieSceneEventCustomization::GetEditObjects(TArray<UObject*>& OutObjects) const
+{
+	UMovieSceneEventSectionBase* ExternalSection = Cast<UMovieSceneEventSectionBase>(WeakExternalSection.Get());
+	if (ExternalSection)
+	{
+		OutObjects.Add(ExternalSection);
+	}
+	else
+	{
+		PropertyHandle->GetOuterObjects(OutObjects);
+	}
+}
+
+void FMovieSceneEventCustomization::CustomizeHeader(TSharedRef<IPropertyHandle> InPropertyHandle, FDetailWidgetRow& HeaderRow, IPropertyTypeCustomizationUtils& CustomizationUtils)
+{}
+
+void FMovieSceneEventCustomization::CustomizeChildren(TSharedRef<IPropertyHandle> InPropertyHandle, IDetailChildrenBuilder& ChildBuilder, IPropertyTypeCustomizationUtils& CustomizationUtils)
+{
+	PropertyHandle = InPropertyHandle;
+
+	TArray<void*> RawData;
+	PropertyHandle->AccessRawData(RawData);
+
+	ChildBuilder.AddCustomRow(FText())
+	.NameContent()
+	[
+		SNew(STextBlock)
+		.Font(CustomizationUtils.GetRegularFont())
+		.Text(LOCTEXT("EventValueText", "Event"))
+	]
+	.ValueContent()
+	.MinDesiredWidth(200.f)
+	[
+		SNew(SHorizontalBox)
+
+		+ SHorizontalBox::Slot()
+		[
+			SNew(SComboButton)
+			.ButtonStyle(FEditorStyle::Get(), "HoverHintOnly")
+			.ForegroundColor(FSlateColor::UseForeground())
+			.OnGetMenuContent(this, &FMovieSceneEventCustomization::GetMenuContent)
+			.CollapseMenuOnParentFocus(true)
+			.ContentPadding(FMargin(4.f, 0.f))
+			.ButtonContent()
+			[
+				SNew(SHorizontalBox)
+
+				+ SHorizontalBox::Slot()
+				.Padding(FMargin(0.f, 0.f, 4.f, 0.f))
+				.VAlign(VAlign_Center)
+				.AutoWidth()
+				[
+					SNew(SImage)
+					.Image(this, &FMovieSceneEventCustomization::GetEventIcon)
+				]
+
+				+ SHorizontalBox::Slot()
+				.Padding(FMargin(0.f, 0.f, 4.f, 0.f))
+				.VAlign(VAlign_Center)
+				[
+					SNew(STextBlock)
+					.Font(CustomizationUtils.GetRegularFont())
+					.Text(this, &FMovieSceneEventCustomization::GetEventName)
+				]
+			]
+		]
+
+		+ SHorizontalBox::Slot()
+		.VAlign(VAlign_Center)
+		.AutoWidth()
+		[
+			PropertyCustomizationHelpers::MakeBrowseButton(FSimpleDelegate::CreateSP(this, &FMovieSceneEventCustomization::NavigateToDefinition), LOCTEXT("NavigateToDefinition_Tip", "Navigate to this event's definition"))
+		]
+
+		+ SHorizontalBox::Slot()
+		.VAlign(VAlign_Center)
+		.AutoWidth()
+		[
+			SNew(STextBlock)
+			.Visibility(this, &FMovieSceneEventCustomization::GetErrorVisibility)
+			.ToolTipText(this, &FMovieSceneEventCustomization::GetErrorTooltip)
+			.TextStyle(FEditorStyle::Get(), "Log.Warning")
+			.Font(FEditorStyle::Get().GetFontStyle("FontAwesome.9"))
+			.Text(FEditorFontGlyphs::Exclamation_Triangle)
+		]
+	];
+}
+
+UClass* FMovieSceneEventCustomization::FindObjectBindingClass()
+{
+	UMovieSceneSequence* CommonSequence = GetCommonSequence();
+	UMovieScene*         MovieScene     = CommonSequence ? CommonSequence->GetMovieScene() : nullptr;
+
+	if (!MovieScene)
+	{
+		return nullptr;
+	}
+
+	TArray<UObject*> EditObjects;
+	GetEditObjects(EditObjects);
+
+	UClass* BindingClass = nullptr;
+
+	for (UObject* EditObject : EditObjects)
+	{
+		UMovieSceneTrack* Track = EditObject->GetTypedOuter<UMovieSceneTrack>();
+		if (Track)
+		{
+			for (const FMovieSceneBinding& Binding : MovieScene->GetBindings())
+			{
+				if (Algo::Find(Binding.GetTracks(), Track))
+				{
+					UClass* ThisTrackClass = nullptr;
+					if (FMovieScenePossessable* Possessable = MovieScene->FindPossessable(Binding.GetObjectGuid()))
+					{
+						ThisTrackClass = const_cast<UClass*>(Possessable->GetPossessedObjectClass());
+					}
+					else if(FMovieSceneSpawnable* Spawnable = MovieScene->FindSpawnable(Binding.GetObjectGuid()))
+					{
+						ThisTrackClass = Spawnable->GetObjectTemplate()->GetClass();
+					}
+
+					if (!BindingClass)
+					{
+						BindingClass = ThisTrackClass;
+					}
+					else if (BindingClass != ThisTrackClass)
+					{
+						return nullptr;
+					}
+					break;
+				}
+			}
+		}
+	}
+
+	return BindingClass;
+}
+
+UMovieSceneSequence* FMovieSceneEventCustomization::GetCommonSequence() const
+{
+	TArray<UObject*> EditObjects;
+	GetEditObjects(EditObjects);
+
+	UMovieSceneSequence* CommonEventSequence = nullptr;
+
+	for (UObject* Obj : EditObjects)
+	{
+		UMovieSceneSequence* ThisSequence = Obj ? Obj->GetTypedOuter<UMovieSceneSequence>() : nullptr;
+		if (CommonEventSequence && CommonEventSequence != ThisSequence)
+		{
+			return nullptr;
+		}
+
+		CommonEventSequence = ThisSequence;
+	}
+	return CommonEventSequence;
+}
+
+UMovieSceneEventTrack* FMovieSceneEventCustomization::GetCommonTrack() const
+{
+	TArray<UObject*> EditObjects;
+	GetEditObjects(EditObjects);
+
+	UMovieSceneEventTrack* CommonEventTrack = nullptr;
+
+	for (UObject* Obj : EditObjects)
+	{
+		UMovieSceneEventTrack* ThisTrack = Obj ? Obj->GetTypedOuter<UMovieSceneEventTrack>() : nullptr;
+		if (CommonEventTrack && CommonEventTrack != ThisTrack)
+		{
+			return nullptr;
+		}
+
+		CommonEventTrack = ThisTrack;
+	}
+	return CommonEventTrack;
+}
+
+UK2Node_FunctionEntry* FMovieSceneEventCustomization::GetCommonEndpoint() const
+{
+	TArray<void*> RawData;
+	PropertyHandle->AccessRawData(RawData);
+
+	UK2Node_FunctionEntry* CommonEndpoint = nullptr;
+
+	for (void* Ptr : RawData)
+	{
+		if (Ptr)
+		{
+			UK2Node_FunctionEntry* ThisEndpoint = static_cast<FMovieSceneEvent*>(Ptr)->GetFunctionEntry();
+			if (CommonEndpoint && CommonEndpoint != ThisEndpoint)
+			{
+				return nullptr;
+			}
+
+			CommonEndpoint = ThisEndpoint;
+		}
+	}
+	return CommonEndpoint;
+}
+
+TSharedRef<SWidget> FMovieSceneEventCustomization::GetMenuContent()
+{
+	FMenuBuilder MenuBuilder(true, nullptr, nullptr, true);
+
+	UMovieSceneSequence*       Sequence       = GetCommonSequence();
+	UK2Node_FunctionEntry*     CommonEndpoint = GetCommonEndpoint();
+	FMovieSceneSequenceEditor* SequenceEditor = FMovieSceneSequenceEditor::Find(Sequence);
+
+	UBlueprint* DirectorBP = SequenceEditor ? SequenceEditor->GetDirectorBlueprint(Sequence) : nullptr;
+
+	CachedCommonEndpoint = CommonEndpoint;
+
+	MenuBuilder.AddMenuEntry(
+		LOCTEXT("CreateEventEndpoint_Text",    "Create New Endpoint"),
+		LOCTEXT("CreateEventEndpoint_Tooltip", "Creates a new event endpoint in this sequence's blueprint."),
+		FSlateIcon(FEditorStyle::GetStyleSetName(), "Sequencer.CreateEventBinding"),
+		FUIAction(
+			FExecuteAction::CreateSP(this, &FMovieSceneEventCustomization::CreateEventEndpoint)
+		)
+	);
+
+	UClass* TemplateClass = FindObjectBindingClass();
+	if (TemplateClass)
+	{
+		MenuBuilder.AddSubMenu(
+			LOCTEXT("CreateQuickBinding_Text",    "Create Quick Binding"),
+			LOCTEXT("CreateQuickBinding_Tooltip", "Shows a list of functions on this object binding that can be bound directly to this event."),
+			FNewMenuDelegate::CreateSP(this, &FMovieSceneEventCustomization::PopulateQuickBindSubMenu, TemplateClass),
+			false /* bInOpenSubMenuOnClick */,
+			FSlateIcon(FEditorStyle::GetStyleSetName(), "Sequencer.CreateQuickBinding"),
+			false /* bInShouldCloseWindowAfterMenuSelection */
+		);
+	}
+
+	MenuBuilder.AddMenuEntry(
+		LOCTEXT("ClearEventEndpoint_Text",    "Clear"),
+		LOCTEXT("ClearEventEndpoint_Tooltip", "Unbinds this event from its current binding."),
+		FSlateIcon(FEditorStyle::GetStyleSetName(), "Sequencer.ClearEventBinding"),
+		FUIAction(
+			FExecuteAction::CreateSP(this, &FMovieSceneEventCustomization::ClearEventEndpoint)
+		)
+	);
+
+	if (DirectorBP)
+	{
+		FSlateIcon Icon(FEditorStyle::GetStyleSetName(), "GraphEditor.Function_16x");
+
+		MenuBuilder.BeginSection(NAME_None, LOCTEXT("ExistingEndpoints", "Existing"));
+
+		TArray<UK2Node_FunctionEntry*> EntryNodes;
+
+		for (UEdGraph* FunctionGraph : DirectorBP->FunctionGraphs)
+		{
+			EntryNodes.Reset();
+			FunctionGraph->GetNodesOfClass<UK2Node_FunctionEntry>(EntryNodes);
+
+			if (EntryNodes.Num() == 1 && FMovieSceneEvent::IsValidFunction(EntryNodes[0]))
+			{
+				MenuBuilder.AddMenuEntry(
+					EntryNodes[0]->GetNodeTitle(ENodeTitleType::MenuTitle),
+					EntryNodes[0]->GetTooltipText(),
+					Icon,
+					FUIAction(
+						FExecuteAction::CreateSP(this, &FMovieSceneEventCustomization::SetEventEndpoint, EntryNodes[0]),
+						FCanExecuteAction::CreateLambda([]{ return true; }),
+						FIsActionChecked::CreateSP(this, &FMovieSceneEventCustomization::CompareCurrentEventEndpoint, EntryNodes[0])
+					),
+					NAME_None,
+					EUserInterfaceActionType::RadioButton
+				);
+			}
+		}
+
+		MenuBuilder.EndSection();
+	}
+
+	return MenuBuilder.MakeWidget();
+}
+
+void FMovieSceneEventCustomization::PopulateQuickBindSubMenu(FMenuBuilder& MenuBuilder, UClass* TemplateClass)
+{
+	FSlateIcon Icon(FEditorStyle::GetStyleSetName(), "GraphEditor.Function_16x");
+
+	TArray<UFunction*> Functions;
+
+	UClass* SuperClass = TemplateClass;
+	while (SuperClass)
+	{
+		MenuBuilder.BeginSection(NAME_None, SuperClass->GetDisplayNameText());
+
+		Functions.Reset();
+		for (UFunction* Function : TFieldRange<UFunction>(SuperClass, EFieldIteratorFlags::ExcludeSuper, EFieldIteratorFlags::ExcludeDeprecated))
+		{
+			if (Function->HasAllFunctionFlags(FUNC_BlueprintCallable|FUNC_Public))
+			{
+				Functions.Add(Function);
+			}
+		}
+	
+		Algo::SortBy(Functions, &UFunction::GetFName);
+
+		for (UFunction* Function : Functions)
+		{
+			MenuBuilder.AddMenuEntry(
+				FText::FromName(Function->GetFName()),
+				FText(),
+				Icon,
+				FUIAction(
+					FExecuteAction::CreateSP(this, &FMovieSceneEventCustomization::CreateEventEndpointFromFunction, Function, TemplateClass)
+				)
+			);
+		}
+
+		MenuBuilder.EndSection();
+
+		SuperClass = SuperClass->GetSuperClass();
+	}
+}
+
+const FSlateBrush* FMovieSceneEventCustomization::GetEventIcon() const
+{
+	UK2Node_FunctionEntry* CommonEndpoint = GetCommonEndpoint();
+
+	if (FMovieSceneEvent::IsValidFunction(CommonEndpoint))
+	{
+		return FEditorStyle::GetBrush("GraphEditor.Function_16x");
+	}
+	else if (!CommonEndpoint)
+	{
+		TArray<void*> RawData;
+		PropertyHandle->AccessRawData(RawData);
+		if (RawData.Num() > 1)
+		{
+			return FEditorStyle::GetBrush("Sequencer.MultipleEvents");
+		}
+	}
+
+	return FEditorStyle::GetBrush("Sequencer.UnboundEvent");
+}
+
+FText FMovieSceneEventCustomization::GetEventName() const
+{
+	UK2Node_FunctionEntry* CommonEndpoint = GetCommonEndpoint();
+
+	if (FMovieSceneEvent::IsValidFunction(CommonEndpoint))
+	{
+		return CommonEndpoint->GetNodeTitle(ENodeTitleType::MenuTitle);
+	}
+	else if (!CommonEndpoint)
+	{
+		TArray<void*> RawData;
+		PropertyHandle->AccessRawData(RawData);
+		if (RawData.Num() != 1)
+		{
+			return LOCTEXT("MultipleValuesText", "Multiple Values");
+		}
+	}
+
+	return LOCTEXT("UnboundText", "Unbound");
+}
+
+EVisibility FMovieSceneEventCustomization::GetErrorVisibility() const
+{
+	UK2Node_FunctionEntry* CommonEndpoint = GetCommonEndpoint();
+	if (CommonEndpoint)
+	{
+		return FMovieSceneEvent::IsValidFunction(CommonEndpoint) ? EVisibility::Collapsed : EVisibility::Visible;
+	}
+	return EVisibility::Collapsed;
+}
+
+FText FMovieSceneEventCustomization::GetErrorTooltip() const
+{
+	UK2Node_FunctionEntry* CommonEndpoint = GetCommonEndpoint();
+	if (CommonEndpoint)
+	{
+		return FText::Format(LOCTEXT("ErrorToolTipFormat", "The currently assigned function '{0}' is not a valid event endpoint. Event endpoints must have a single object or interface parameter that is not passed by-reference."), CommonEndpoint->GetNodeTitle(ENodeTitleType::MenuTitle));
+	}
+	return FText::GetEmpty();
+}
+
+void FMovieSceneEventCustomization::SetEventEndpoint(UK2Node_FunctionEntry* NewEndpoint)
+{
+	FScopedTransaction Transaction(LOCTEXT("SetEventEndpoint", "Set Event Endpoint"));
+
+	UBlueprint* Blueprint = NewEndpoint->GetBlueprint();
+
+	CachedCommonEndpoint = NewEndpoint;
+
+	// Modify and assign the blueprint for outer sections
+	TArray<UObject*> EditObjects;
+	GetEditObjects(EditObjects);
+
+	for (UObject* Outer : EditObjects)
+	{
+		UMovieSceneEventSectionBase* BaseEventSection = Cast<UMovieSceneEventSectionBase>(Outer);
+		if (BaseEventSection)
+		{
+			BaseEventSection->Modify();
+			BaseEventSection->SetDirectorBlueprint(Blueprint);
+		}
+	}
+
+	// Assign the endpoints to all events
+	TArray<void*> RawData;
+	PropertyHandle->AccessRawData(RawData);
+
+	for (void* Ptr : RawData)
+	{
+		FMovieSceneEvent* Event = static_cast<FMovieSceneEvent*>(Ptr);
+		if (Event)
+		{
+			Event->SetFunctionEntry(NewEndpoint);
+		}
+	}
+
+	// Ensure that anything listening for property changed notifications are notified of the new binding
+	PropertyHandle->NotifyFinishedChangingProperties();
+}
+
+bool FMovieSceneEventCustomization::CompareCurrentEventEndpoint(UK2Node_FunctionEntry* NewEndpoint)
+{
+	return CachedCommonEndpoint == NewEndpoint;
+}
+
+void FMovieSceneEventCustomization::CreateEventEndpoint()
+{
+	UMovieSceneSequence*       Sequence       = GetCommonSequence();
+	FMovieSceneSequenceEditor* SequenceEditor = FMovieSceneSequenceEditor::Find(Sequence);
+	if (!SequenceEditor)
+	{
+		return;
+	}
+
+	FScopedTransaction Transaction(LOCTEXT("CreateEventEndpoint", "Create Event Endpoint"));
+
+	// Create a single event binding and point all events in this property handle to it
+	UK2Node_FunctionEntry* NewFunctionEntry = SequenceEditor->CreateEventEndpoint(Sequence);
+
+	if (UMovieSceneEventTrack* CommonTrack = GetCommonTrack())
+	{
+		SequenceEditor->InitializeEndpointForTrack(CommonTrack, NewFunctionEntry);
+	}
+	
+	SetEventEndpoint(NewFunctionEntry);
+
+	if (NewFunctionEntry)
+	{
+		FKismetEditorUtilities::BringKismetToFocusAttentionOnObject(NewFunctionEntry, false);
+	}
+}
+
+void FMovieSceneEventCustomization::CreateEventEndpointFromFunction(UFunction* QuickBindFunction, UClass* PinClassType)
+{
+	UMovieSceneSequence*       Sequence       = GetCommonSequence();
+	FMovieSceneSequenceEditor* SequenceEditor = FMovieSceneSequenceEditor::Find(Sequence);
+
+	if (!SequenceEditor)
+	{
+		return;
+	}
+
+	FScopedTransaction Transaction(LOCTEXT("CreateEventEndpoint", "Create Event Endpoint"));
+
+	FString DesiredNewEventName = FString(TEXT("Call ")) + QuickBindFunction->GetName();
+
+	// Create a single event binding and point all events in this property handle to it
+	UK2Node_FunctionEntry* NewFunctionEntry = SequenceEditor->CreateEventEndpoint(Sequence, DesiredNewEventName);
+	if (!ensure(NewFunctionEntry))
+	{
+		return;
+	}
+
+	if (ensure(NewFunctionEntry->UserDefinedPins.Num() == 0))
+	{
+		FEdGraphPinType PinType;
+		PinType.PinCategory = PinClassType->IsChildOf(UInterface::StaticClass()) ? UEdGraphSchema_K2::PC_Interface : UEdGraphSchema_K2::PC_Object;
+		PinType.PinSubCategoryObject = PinClassType;
+
+		NewFunctionEntry->CreateUserDefinedPin(FMovieSceneSequenceEditor::TargetPinName, PinType, EGPD_Output, true);
+		NewFunctionEntry->ReconstructNode();
+	}
+
+	NewFunctionEntry->bCommentBubblePinned = false;
+	NewFunctionEntry->bCommentBubbleVisible = false;
+
+	UEdGraph* Graph = NewFunctionEntry->GetGraph();
+
+	// Make a call function template
+	UK2Node_CallFunction* CallFuncNode = NewObject<UK2Node_CallFunction>(Graph);
+	CallFuncNode->NodePosX = NewFunctionEntry->NodePosX + NewFunctionEntry->NodeWidth + 200;
+	CallFuncNode->NodePosY = NewFunctionEntry->NodePosY;
+	CallFuncNode->FunctionReference.SetFromField<UFunction>(QuickBindFunction, false);
+	CallFuncNode->ReconstructNode();
+
+	Graph->AddNode(CallFuncNode, false, false);
+
+	{
+		// Connect the exec pins together
+		UEdGraphPin* ThenPin = NewFunctionEntry->FindPin(UEdGraphSchema_K2::PN_Then);
+		UEdGraphPin* ExecPin = CallFuncNode->GetExecPin();
+
+		if (ThenPin && ExecPin)
+		{
+			const UEdGraphSchema* Schema = ThenPin->GetSchema();
+			Schema->TryCreateConnection(ThenPin, ExecPin);
+		}
+	}
+
+	// Connect the object target pin to the self (input) pin on the call function node
+	if (UEdGraphPin* SelfPin = CallFuncNode->FindPin(UEdGraphSchema_K2::PSC_Self))
+	{
+		for (UEdGraphPin* Pin : NewFunctionEntry->Pins)
+		{
+			if (Pin->PinType.PinCategory == UEdGraphSchema_K2::PC_Interface  || Pin->PinType.PinCategory == UEdGraphSchema_K2::PC_Object)
+			{
+				const UEdGraphSchema* Schema = Pin->GetSchema();
+				Schema->TryCreateConnection(Pin, SelfPin);
+			}
+		}
+	}
+
+	SetEventEndpoint(NewFunctionEntry);
+
+	if (NewFunctionEntry)
+	{
+		FKismetEditorUtilities::BringKismetToFocusAttentionOnObject(NewFunctionEntry, false);
+	}
+}
+
+void FMovieSceneEventCustomization::ClearEventEndpoint()
+{
+	FScopedTransaction Transaction(LOCTEXT("ClearEventEndpoint", "Clear Event Endpoint"));
+
+	TArray<UObject*> EditObjects;
+	GetEditObjects(EditObjects);
+
+	for (UObject* Outer : EditObjects)
+	{
+		if (Outer)
+		{
+			Outer->Modify();
+		}
+	}
+
+	TArray<void*> RawData;
+	PropertyHandle->AccessRawData(RawData);
+
+	for (void* Ptr : RawData)
+	{
+		FMovieSceneEvent* Event = static_cast<FMovieSceneEvent*>(Ptr);
+		if (Event)
+		{
+			Event->SetFunctionEntry(nullptr);
+		}
+	}
+
+	// Ensure that anything listening for property changed notifications are notified of the new binding
+	PropertyHandle->NotifyFinishedChangingProperties();
+}
+
+void FMovieSceneEventCustomization::NavigateToDefinition()
+{
+	UK2Node_FunctionEntry* CommonEndpoint = GetCommonEndpoint();
+	if (CommonEndpoint)
+	{
+		FKismetEditorUtilities::BringKismetToFocusAttentionOnObject(CommonEndpoint, false);
+	}
+}
+
+#undef LOCTEXT_NAMESPACE
diff --git a/Engine/Source/Editor/MovieSceneTools/Private/MovieSceneToolsModule.cpp b/Engine/Source/Editor/MovieSceneTools/Private/MovieSceneToolsModule.cpp
--- a/Engine/Source/Editor/MovieSceneTools/Private/MovieSceneToolsModule.cpp
+++ b/Engine/Source/Editor/MovieSceneTools/Private/MovieSceneToolsModule.cpp
@@ -40,6 +40,7 @@
 
 #include "MovieSceneBuiltInEasingFunctionCustomization.h"
 #include "MovieSceneObjectBindingIDCustomization.h"
+#include "MovieSceneEventCustomization.h"
 #include "SequencerClipboardReconciler.h"
 #include "ClipboardTypes.h"
 #include "ISettingsModule.h"
@@ -50,6 +51,7 @@
 #include "ISequencerChannelInterface.h"
 #include "SequencerChannelInterface.h"
 #include "Channels/BuiltInChannelEditors.h"
+#include "Channels/MovieSceneEventChannel.h"
 #include "Sections/MovieSceneEventSection.h"
 
 
@@ -120,6 +122,7 @@ public:
 		PropertyModule.RegisterCustomClassLayout("MovieSceneToolsProjectSettings", FOnGetDetailCustomizationInstance::CreateStatic(&FMovieSceneToolsProjectSettingsCustomization::MakeInstance));
 		PropertyModule.RegisterCustomClassLayout("MovieSceneBuiltInEasingFunction", FOnGetDetailCustomizationInstance::CreateLambda(&MakeShared<FMovieSceneBuiltInEasingFunctionCustomization>));
 		PropertyModule.RegisterCustomPropertyTypeLayout("MovieSceneObjectBindingID", FOnGetPropertyTypeCustomizationInstance::CreateLambda(&MakeShared<FMovieSceneObjectBindingIDCustomization>));
+		PropertyModule.RegisterCustomPropertyTypeLayout("MovieSceneEvent", FOnGetPropertyTypeCustomizationInstance::CreateStatic(&FMovieSceneEventCustomization::MakeInstance));
 
 		SequencerModule.RegisterChannelInterface<FMovieSceneBoolChannel>();
 		SequencerModule.RegisterChannelInterface<FMovieSceneByteChannel>();
@@ -129,6 +132,8 @@ public:
 		SequencerModule.RegisterChannelInterface<FMovieSceneParticleChannel>();
 		SequencerModule.RegisterChannelInterface<FMovieSceneActorReferenceData>();
 		SequencerModule.RegisterChannelInterface<FMovieSceneEventSectionData>();
+
+		SequencerModule.RegisterChannelInterface<FMovieSceneEventChannel>();
 	}
 
 	virtual void ShutdownModule() override
@@ -185,6 +190,7 @@ public:
 			PropertyModule.UnregisterCustomClassLayout("MovieSceneToolsProjectSettings");
 			PropertyModule.UnregisterCustomClassLayout("MovieSceneBuiltInEasingFunction");
 			PropertyModule.UnregisterCustomPropertyTypeLayout("MovieSceneObjectBindingID");
+			PropertyModule.UnregisterCustomPropertyTypeLayout("MovieSceneEvent");
 		}
 	}
 
diff --git a/Engine/Source/Editor/MovieSceneTools/Private/SKeyEditInterface.cpp b/Engine/Source/Editor/MovieSceneTools/Private/SKeyEditInterface.cpp
--- a/Engine/Source/Editor/MovieSceneTools/Private/SKeyEditInterface.cpp
+++ b/Engine/Source/Editor/MovieSceneTools/Private/SKeyEditInterface.cpp
@@ -12,6 +12,7 @@
 #include "PropertyEditorModule.h"
 #include "FrameNumberDetailsCustomization.h"
 #include "MovieSceneSectionDetailsCustomization.h"
+#include "MovieSceneEventCustomization.h"
 #include "Misc/FrameNumber.h"
 #include "Modules/ModuleManager.h"
 #include "IDetailCustomization.h"
@@ -77,6 +78,7 @@ void SKeyEditInterface::Initialize()
 
 	// register details customizations for this instance
 	StructureDetailsView->GetDetailsView()->RegisterInstancedCustomPropertyTypeLayout("MovieSceneObjectBindingID", FOnGetPropertyTypeCustomizationInstance::CreateSP(this, &SKeyEditInterface::CreateBindingIDCustomization));
+	StructureDetailsView->GetDetailsView()->RegisterInstancedCustomPropertyTypeLayout("MovieSceneEvent", FOnGetPropertyTypeCustomizationInstance::CreateSP(this, &SKeyEditInterface::CreateEventCustomization));
 	StructureDetailsView->GetDetailsView()->RegisterInstancedCustomPropertyTypeLayout("FrameNumber", FOnGetPropertyTypeCustomizationInstance::CreateSP(this, &SKeyEditInterface::CreateFrameNumberCustomization));
 
 	StructureDetailsView->SetStructureData(NewEditData.KeyStruct);
@@ -101,6 +103,11 @@ TSharedRef<IPropertyTypeCustomization> SKeyEditInterface::CreateFrameNumberCusto
 	return MakeShared<FFrameNumberDetailsCustomization>(SequencerPtr->GetNumericTypeInterface());
 }
 
+TSharedRef<IPropertyTypeCustomization> SKeyEditInterface::CreateEventCustomization()
+{
+	return FMovieSceneEventCustomization::MakeInstance(WeakSection.Get());
+}
+
 void SKeyEditInterface::OnFinishedChangingProperties(const FPropertyChangedEvent& ChangeEvent, TSharedPtr<FStructOnScope> KeyStruct)
 {
 	if (KeyStruct->GetStruct()->IsChildOf(FMovieSceneKeyStruct::StaticStruct()))
diff --git a/Engine/Source/Editor/MovieSceneTools/Private/Sections/EventSection.cpp b/Engine/Source/Editor/MovieSceneTools/Private/Sections/EventSection.cpp
--- a/Engine/Source/Editor/MovieSceneTools/Private/Sections/EventSection.cpp
+++ b/Engine/Source/Editor/MovieSceneTools/Private/Sections/EventSection.cpp
@@ -6,101 +6,250 @@
 #include "ISectionLayoutBuilder.h"
 #include "EditorStyleSet.h"
 #include "CommonMovieSceneTools.h"
+#include "Tracks/MovieSceneEventTrack.h"
 #include "Sections/MovieSceneEventSection.h"
+#include "Sections/MovieSceneEventTriggerSection.h"
+#include "Sections/MovieSceneEventRepeaterSection.h"
 #include "MovieSceneTrack.h"
 #include "SequencerTimeSliderController.h"
 #include "Fonts/FontMeasure.h"
 #include "Framework/Application/SlateApplication.h"
 #include "EditorStyleSet.h"
+#include "K2Node_FunctionEntry.h"
+#include "Kismet2/KismetEditorUtilities.h"
+#include "MovieSceneSequence.h"
+#include "EditorFontGlyphs.h"
+#include "ScopedTransaction.h"
 
-int32 FEventSection::OnPaintSection( FSequencerSectionPainter& Painter ) const
+#define LOCTEXT_NAMESPACE "EventSection"
+
+bool FEventSectionBase::IsSectionSelected() const
+{
+	TSharedPtr<ISequencer> SequencerPtr = Sequencer.Pin();
+
+	TArray<UMovieSceneTrack*> SelectedTracks;
+	SequencerPtr->GetSelectedTracks(SelectedTracks);
+
+	UMovieSceneSection* Section = WeakSection.Get();
+	UMovieSceneTrack*   Track   = Section ? CastChecked<UMovieSceneTrack>(Section->GetOuter()) : nullptr;
+	return Track && SelectedTracks.Contains(Track);
+}
+
+void FEventSectionBase::PaintEventName(FSequencerSectionPainter& Painter, int32 LayerId, const FString& InEventString, float PixelPos, bool bIsEventValid) const
+{
+	static const int32   FontSize      = 10;
+	static const float   BoxOffsetPx   = 10.f;
+	static const TCHAR*  WarningString = TEXT("\xf071");
+
+	const FSlateFontInfo FontAwesomeFont = FEditorStyle::Get().GetFontStyle("FontAwesome.10");
+	const FSlateFontInfo SmallLayoutFont = FCoreStyle::GetDefaultFontStyle("Bold", 10);
+	const FLinearColor   DrawColor       = FEditorStyle::GetSlateColor("SelectionColor").GetColor(FWidgetStyle());
+
+	TSharedRef<FSlateFontMeasure> FontMeasureService = FSlateApplication::Get().GetRenderer()->GetFontMeasureService();
+
+	// Setup the warning size. Static since it won't ever change
+	static FVector2D WarningSize    = FontMeasureService->Measure(WarningString, FontAwesomeFont);
+	const  FMargin   WarningPadding = (bIsEventValid || InEventString.Len() == 0) ? FMargin(0.f) : FMargin(0.f, 0.f, 4.f, 0.f);
+	const  FMargin   BoxPadding     = FMargin(4.0f, 2.0f);
+
+	const FVector2D  TextSize       = FontMeasureService->Measure(InEventString, SmallLayoutFont);
+	const FVector2D  IconSize       = bIsEventValid ? FVector2D::ZeroVector : WarningSize;
+	const FVector2D  PaddedIconSize = IconSize + WarningPadding.GetDesiredSize();
+	const FVector2D  BoxSize        = FVector2D(TextSize.X + PaddedIconSize.X, FMath::Max(TextSize.Y, PaddedIconSize.Y )) + BoxPadding.GetDesiredSize();
+
+	// Flip the text position if getting near the end of the view range
+	bool  bDrawLeft    = (Painter.SectionGeometry.Size.X - PixelPos) < (BoxSize.X + 22.f) - BoxOffsetPx;
+	float BoxPositionX = bDrawLeft ? PixelPos - BoxSize.X - BoxOffsetPx : PixelPos + BoxOffsetPx;
+
+	FVector2D BoxOffset  = FVector2D(BoxPositionX,                    Painter.SectionGeometry.Size.Y*.5f - BoxSize.Y*.5f);
+	FVector2D IconOffset = FVector2D(BoxPadding.Left,                 BoxSize.Y*.5f - IconSize.Y*.5f);
+	FVector2D TextOffset = FVector2D(IconOffset.X + PaddedIconSize.X, BoxSize.Y*.5f - TextSize.Y*.5f);
+
+	// Draw the background box
+	FSlateDrawElement::MakeBox(
+		Painter.DrawElements,
+		LayerId + 1,
+		Painter.SectionGeometry.ToPaintGeometry(BoxOffset, BoxSize),
+		FEditorStyle::GetBrush("WhiteBrush"),
+		ESlateDrawEffect::None,
+		FLinearColor::Black.CopyWithNewOpacity(0.5f)
+	);
+
+	if (!bIsEventValid)
+	{
+		// Draw a warning icon for unbound repeaters
+		FSlateDrawElement::MakeText(
+			Painter.DrawElements,
+			LayerId + 2,
+			Painter.SectionGeometry.ToPaintGeometry(BoxOffset + IconOffset, IconSize),
+			WarningString,
+			FontAwesomeFont,
+			Painter.bParentEnabled ? ESlateDrawEffect::None : ESlateDrawEffect::DisabledEffect,
+			FEditorStyle::GetWidgetStyle<FTextBlockStyle>("Log.Warning").ColorAndOpacity.GetSpecifiedColor()
+		);
+	}
+
+	FSlateDrawElement::MakeText(
+		Painter.DrawElements,
+		LayerId + 2,
+		Painter.SectionGeometry.ToPaintGeometry(BoxOffset + TextOffset, TextSize),
+		InEventString,
+		SmallLayoutFont,
+		Painter.bParentEnabled ? ESlateDrawEffect::None : ESlateDrawEffect::DisabledEffect,
+		DrawColor
+	);
+}
+
+int32 FEventSection::OnPaintSection(FSequencerSectionPainter& Painter) const
 {
 	int32 LayerId = Painter.PaintSectionBackground();
 	UMovieSceneEventSection* EventSection = Cast<UMovieSceneEventSection>( WeakSection.Get() );
-	if (!EventSection)
+	if (!EventSection || !IsSectionSelected())
 	{
 		return LayerId;
 	}
 
-	// Draw event names for keys on sections that are on selected tracks
-	TSharedPtr<ISequencer> SequencerPtr = Sequencer.Pin();
-	TArray<UMovieSceneTrack*> SelectedTracks;
-	SequencerPtr->GetSelectedTracks(SelectedTracks);
+	const FTimeToPixel& TimeToPixelConverter = Painter.GetTimeConverter();
 
-	bool bSelected = false;
-	for (auto SelectedTrack : SelectedTracks)
+	for (int32 KeyIndex = 0; KeyIndex < EventSection->GetEventData().GetKeyTimes().Num(); ++KeyIndex)
 	{
-		if (SelectedTrack->GetAllSections().Contains(EventSection))
+		FFrameNumber EventTime = EventSection->GetEventData().GetKeyTimes()[KeyIndex];
+		FEventPayload EventData = EventSection->GetEventData().GetKeyValues()[KeyIndex];
+
+		if (EventSection->GetRange().Contains(EventTime))
 		{
-			bSelected = true;
-			break;
+			FString EventString = EventData.EventName.ToString();
+			if (!EventString.IsEmpty())
+			{
+				const float PixelPos = TimeToPixelConverter.FrameToPixel(EventTime);
+				PaintEventName(Painter, LayerId, EventString, PixelPos);
+			}
 		}
 	}
 
-	if (!bSelected)
+	return LayerId + 3;
+}
+
+int32 FEventTriggerSection::OnPaintSection(FSequencerSectionPainter& Painter) const
+{
+	int32 LayerId = Painter.PaintSectionBackground();
+
+	UMovieSceneEventTriggerSection* EventTriggerSection = Cast<UMovieSceneEventTriggerSection>(WeakSection.Get());
+	if (!EventTriggerSection || !IsSectionSelected())
 	{
 		return LayerId;
 	}
 
 	const FTimeToPixel& TimeToPixelConverter = Painter.GetTimeConverter();
 
-	const ESlateDrawEffect DrawEffects = Painter.bParentEnabled
-		? ESlateDrawEffect::None
-		: ESlateDrawEffect::DisabledEffect;
+	TMovieSceneChannelData<FMovieSceneEvent> EventData = EventTriggerSection->EventChannel.GetData();
+	TArrayView<const FFrameNumber>     Times  = EventData.GetTimes();
+	TArrayView<const FMovieSceneEvent> Events = EventData.GetValues();
 
-	for (int32 KeyIndex = 0; KeyIndex < EventSection->GetEventData().GetKeyTimes().Num(); ++KeyIndex)
+	TRange<FFrameNumber> EventSectionRange = EventTriggerSection->GetRange();
+	for (int32 KeyIndex = 0; KeyIndex < Times.Num(); ++KeyIndex)
 	{
-		FFrameNumber EventTime = EventSection->GetEventData().GetKeyTimes()[KeyIndex];
-		FEventPayload EventData = EventSection->GetEventData().GetKeyValues()[KeyIndex];
+		FFrameNumber EventTime = Times[KeyIndex];
+		UK2Node_FunctionEntry* FunctionEntry = Events[KeyIndex].GetFunctionEntry();
 
-		if (!EventSection->GetRange().Contains(EventTime))
+		if (EventSectionRange.Contains(EventTime))
 		{
-			continue;
+			FString EventString = FunctionEntry ? FunctionEntry->GetNodeTitle(ENodeTitleType::MenuTitle).ToString() : FString();
+			bool bIsEventValid = FMovieSceneEvent::IsValidFunction(FunctionEntry);
+
+			const float PixelPos = TimeToPixelConverter.FrameToPixel(EventTime);
+			PaintEventName(Painter, LayerId, EventString, PixelPos, bIsEventValid);
 		}
+	}
 
-		FString EventString = EventData.EventName.ToString();
-		if (EventString.IsEmpty())
+	return LayerId + 3;
+}
+
+FReply FEventTriggerSection::OnKeyDoubleClicked(FKeyHandle KeyHandle)
+{
+	UMovieSceneEventTriggerSection* EventTriggerSection = Cast<UMovieSceneEventTriggerSection>( WeakSection.Get() );
+	if (EventTriggerSection)
+	{
+		TMovieSceneChannelData<FMovieSceneEvent> ChannelData = EventTriggerSection->EventChannel.GetData();
+		int32 EventIndex = ChannelData.GetIndex(KeyHandle);
+		if (EventIndex != INDEX_NONE)
 		{
-			continue;
+			FMovieSceneEvent& Event = ChannelData.GetValues()[EventIndex];
+
+			if (!Event.IsBoundToBlueprint())
+			{
+				FMovieSceneSequenceEditor* SequenceEditor = FMovieSceneSequenceEditor::Find(EventTriggerSection->GetTypedOuter<UMovieSceneSequence>());
+				if (ensure(SequenceEditor))
+				{
+					FScopedTransaction Transaction(LOCTEXT("BindTriggerEvent", "Create Event Endpoint"));
+
+					UK2Node_FunctionEntry* NewEndpoint = SequenceEditor->CreateEventEndpoint(EventTriggerSection->GetTypedOuter<UMovieSceneSequence>());
+					if (NewEndpoint)
+					{
+						SequenceEditor->InitializeEndpointForTrack(EventTriggerSection->GetTypedOuter<UMovieSceneEventTrack>(), NewEndpoint);
+						FMovieSceneSequenceEditor::BindEventToEndpoint(EventTriggerSection, &Event, NewEndpoint);
+					}
+				}
+			}
+
+			if (UK2Node_FunctionEntry* FunctionEntry = Event.GetFunctionEntry())
+			{
+				FKismetEditorUtilities::BringKismetToFocusAttentionOnObject(FunctionEntry, false);
+			}
 		}
+	}
 
-		const float Time = TimeToPixelConverter.FrameToPixel(EventTime); 
-
-		const FSlateFontInfo SmallLayoutFont = FCoreStyle::GetDefaultFontStyle("Bold", 10);
-		const TSharedRef< FSlateFontMeasure > FontMeasureService = FSlateApplication::Get().GetRenderer()->GetFontMeasureService();
-		FVector2D TextSize = FontMeasureService->Measure(EventString, SmallLayoutFont);
+	return FReply::Handled();
+}
 
-		// Flip the text position if getting near the end of the view range
-		static const float TextOffsetPx = 10.f;
-		bool  bDrawLeft = (Painter.SectionGeometry.Size.X - Time) < (TextSize.X + 22.f) - TextOffsetPx;
-		float TextPosition = bDrawLeft ? Time - TextSize.X - TextOffsetPx : Time + TextOffsetPx;
-		//handle mirrored labels
-		const float MajorTickHeight = 4.0f; 
-		FVector2D TextOffset(TextPosition, Painter.SectionGeometry.Size.Y - (MajorTickHeight + TextSize.Y));
+int32 FEventRepeaterSection::OnPaintSection(FSequencerSectionPainter& Painter) const
+{
+	int32 LayerId = Painter.PaintSectionBackground();
 
-		const FLinearColor DrawColor = FEditorStyle::GetSlateColor("SelectionColor").GetColor(FWidgetStyle());
-		const FVector2D BoxPadding = FVector2D(4.0f, 2.0f);
+	UMovieSceneEventRepeaterSection* EventRepeaterSection = Cast<UMovieSceneEventRepeaterSection>(WeakSection.Get());
+	if (EventRepeaterSection)
+	{
+		UK2Node_FunctionEntry* FunctionEntry = EventRepeaterSection->Event.GetFunctionEntry();
 
-		// draw event string
-		FSlateDrawElement::MakeBox(
-			Painter.DrawElements,
-			LayerId + 5,
-			Painter.SectionGeometry.ToPaintGeometry(TextOffset - BoxPadding, TextSize + 2.0f * BoxPadding),
-			FEditorStyle::GetBrush("WhiteBrush"),
-			ESlateDrawEffect::None,
-			FLinearColor::Black.CopyWithNewOpacity(0.5f)
-		);
+		float TextOffsetX = EventRepeaterSection->GetRange().GetLowerBound().IsClosed() ? FMath::Max(0.f, Painter.GetTimeConverter().FrameToPixel(EventRepeaterSection->GetRange().GetLowerBoundValue())) : 0.f;
 
-		FSlateDrawElement::MakeText(
-			Painter.DrawElements,
-			LayerId + 6,
-			Painter.SectionGeometry.ToPaintGeometry(TextOffset, TextSize),
-			EventString,
-			SmallLayoutFont,
-			DrawEffects,
-			DrawColor
-		);
+		FString EventString = FunctionEntry ? FunctionEntry->GetNodeTitle(ENodeTitleType::MenuTitle).ToString() : FString();
+		bool bIsEventValid = FMovieSceneEvent::IsValidFunction(FunctionEntry);
+		PaintEventName(Painter, LayerId, EventString, TextOffsetX, bIsEventValid);
 	}
 
 	return LayerId + 1;
 }
+
+FReply FEventRepeaterSection::OnSectionDoubleClicked(const FGeometry& SectionGeometry, const FPointerEvent& MouseEvent)
+{
+	UMovieSceneEventRepeaterSection* EventRepeaterSection = Cast<UMovieSceneEventRepeaterSection>( WeakSection.Get() );
+	if (EventRepeaterSection)
+	{
+		if (!EventRepeaterSection->Event.IsBoundToBlueprint())
+		{
+			FMovieSceneSequenceEditor* SequenceEditor = FMovieSceneSequenceEditor::Find(EventRepeaterSection->GetTypedOuter<UMovieSceneSequence>());
+			if (ensure(SequenceEditor))
+			{
+				FScopedTransaction Transaction(LOCTEXT("BindRepeaterEvent", "Create Event Endpoint"));
+
+				UK2Node_FunctionEntry* NewEndpoint = SequenceEditor->CreateEventEndpoint(EventRepeaterSection->GetTypedOuter<UMovieSceneSequence>());
+				if (NewEndpoint)
+				{
+					SequenceEditor->InitializeEndpointForTrack(EventRepeaterSection->GetTypedOuter<UMovieSceneEventTrack>(), NewEndpoint);
+					FMovieSceneSequenceEditor::BindEventToEndpoint(EventRepeaterSection, &EventRepeaterSection->Event, NewEndpoint);
+				}
+			}
+		}
+
+		if (UK2Node* FunctionEntry = EventRepeaterSection->Event.GetFunctionEntry())
+		{
+			FKismetEditorUtilities::BringKismetToFocusAttentionOnObject(FunctionEntry, false);
+		}
+	}
+
+	return FReply::Handled();
+}
+
+
+#undef LOCTEXT_NAMESPACE // "EventSection"
\ No newline at end of file
diff --git a/Engine/Source/Editor/MovieSceneTools/Private/Sections/EventSection.h b/Engine/Source/Editor/MovieSceneTools/Private/Sections/EventSection.h
--- a/Engine/Source/Editor/MovieSceneTools/Private/Sections/EventSection.h
+++ b/Engine/Source/Editor/MovieSceneTools/Private/Sections/EventSection.h
@@ -7,30 +7,70 @@
 
 class FSequencerSectionPainter;
 
+class FEventSectionBase
+	: public FSequencerSection
+{
+public:
+
+	FEventSectionBase(UMovieSceneSection& InSectionObject, TWeakPtr<ISequencer> InSequencer)
+		: FSequencerSection(InSectionObject)
+		, Sequencer(InSequencer)
+	{}
+
+protected:
+
+	void PaintEventName(FSequencerSectionPainter& Painter, int32 LayerId, const FString& EventString, float PixelPosition, bool bIsEventValid = true) const;
+
+	bool IsSectionSelected() const;
+
+protected:
+
+	TWeakPtr<ISequencer> Sequencer;
+};
+
 /**
 * An implementation of event sections.
 */
 class FEventSection
-	: public FSequencerSection
+	: public FEventSectionBase
 {
 public:
-
-	/**
-	* Creates a new event section for editing event sections.
-	*
-	* @param InSectionObject The section object which is being displayed and edited.
-	* @param InDisplayName A display name for the section being displayed and edited.
-	*/
 	FEventSection(UMovieSceneSection& InSectionObject, TWeakPtr<ISequencer> InSequencer)
-		: FSequencerSection(InSectionObject)
-		, Sequencer(InSequencer)
+		: FEventSectionBase(InSectionObject, InSequencer)
 	{}
 
 public:
 
 	virtual int32 OnPaintSection(FSequencerSectionPainter& Painter) const override;
+};
 
-private:
+class FEventTriggerSection
+	: public FEventSectionBase
+{
+public:
 
-	TWeakPtr<ISequencer> Sequencer;
+	FEventTriggerSection(UMovieSceneSection& InSectionObject, TWeakPtr<ISequencer> InSequencer)
+		: FEventSectionBase(InSectionObject, InSequencer)
+	{}
+
+public:
+
+	virtual int32 OnPaintSection(FSequencerSectionPainter& Painter) const override;
+	virtual FReply OnKeyDoubleClicked(FKeyHandle KeyHandle) override;
+};
+
+class FEventRepeaterSection
+	: public FEventSectionBase
+{
+public:
+
+	FEventRepeaterSection(UMovieSceneSection& InSectionObject, TWeakPtr<ISequencer> InSequencer)
+		: FEventSectionBase(InSectionObject, InSequencer)
+	{}
+
+
+public:
+
+	virtual int32 OnPaintSection(FSequencerSectionPainter& Painter) const override;
+	virtual FReply OnSectionDoubleClicked(const FGeometry& SectionGeometry, const FPointerEvent& MouseEvent) override;
 };
diff --git a/Engine/Source/Editor/MovieSceneTools/Private/TrackEditors/EventTrackEditor.cpp b/Engine/Source/Editor/MovieSceneTools/Private/TrackEditors/EventTrackEditor.cpp
--- a/Engine/Source/Editor/MovieSceneTools/Private/TrackEditors/EventTrackEditor.cpp
+++ b/Engine/Source/Editor/MovieSceneTools/Private/TrackEditors/EventTrackEditor.cpp
@@ -17,6 +17,11 @@
 #include "IDetailPropertyRow.h"
 #include "Widgets/Layout/SBox.h"
 #include "Sections/EventSection.h"
+#include "Sections/MovieSceneEventSection.h"
+#include "Sections/MovieSceneEventTriggerSection.h"
+#include "Sections/MovieSceneEventRepeaterSection.h"
+#include "SequencerUtilities.h"
+#include "MovieSceneSequenceEditor.h"
 
 #define LOCTEXT_NAMESPACE "FEventTrackEditor"
 
@@ -32,7 +37,20 @@ TSharedRef<ISequencerTrackEditor> FEventTrackEditor::CreateTrackEditor(TSharedRe
 
 TSharedRef<ISequencerSection> FEventTrackEditor::MakeSectionInterface(UMovieSceneSection& SectionObject, UMovieSceneTrack& Track, FGuid ObjectBinding)
 {
-	return MakeShared<FEventSection>(SectionObject, GetSequencer());
+	if (SectionObject.IsA<UMovieSceneEventSection>())
+	{
+		return MakeShared<FEventSection>(SectionObject, GetSequencer());
+	}
+	else if (SectionObject.IsA<UMovieSceneEventTriggerSection>())
+	{
+		return MakeShared<FEventTriggerSection>(SectionObject, GetSequencer());
+	}
+	else if (SectionObject.IsA<UMovieSceneEventRepeaterSection>())
+	{
+		return MakeShared<FEventRepeaterSection>(SectionObject, GetSequencer());
+	}
+
+	return MakeShared<FSequencerSection>(SectionObject);
 }
 
 
@@ -49,33 +67,89 @@ FEventTrackEditor::FEventTrackEditor(TSharedRef<ISequencer> InSequencer)
 
 void FEventTrackEditor::BuildAddTrackMenu(FMenuBuilder& MenuBuilder)
 {
-	UMovieSceneSequence* RootMovieSceneSequence = GetSequencer()->GetRootMovieSceneSequence();
+	UMovieSceneSequence*       RootMovieSceneSequence = GetSequencer()->GetRootMovieSceneSequence();
+	FMovieSceneSequenceEditor* SequenceEditor         = FMovieSceneSequenceEditor::Find(RootMovieSceneSequence);
 
-	if (RootMovieSceneSequence == nullptr)
+	if (SequenceEditor && SequenceEditor->SupportsEvents(RootMovieSceneSequence))
 	{
-		return;
+		MenuBuilder.AddMenuEntry(
+			LOCTEXT("AddEventTrack", "Event Track"),
+			LOCTEXT("AddEventTooltip", "Adds a new event track that can trigger events on the timeline."),
+			FSlateIcon(FEditorStyle::GetStyleSetName(), "Sequencer.Tracks.Event"),
+			FUIAction(
+				FExecuteAction::CreateRaw(this, &FEventTrackEditor::HandleAddEventTrackMenuEntryExecute, FGuid())
+			)
+		);
 	}
-
-	MenuBuilder.AddMenuEntry(
-		LOCTEXT("AddEventTrack", "Event Track"),
-		LOCTEXT("AddEventTooltip", "Adds a new event track that can trigger events on the timeline."),
-		FSlateIcon(FEditorStyle::GetStyleSetName(), "Sequencer.Tracks.Event"),
-		FUIAction(
-			FExecuteAction::CreateRaw(this, &FEventTrackEditor::HandleAddEventTrackMenuEntryExecute)
-		)
-	);
 }
 
-
 void FEventTrackEditor::BuildObjectBindingTrackMenu(FMenuBuilder& MenuBuilder, const FGuid& ObjectBinding, const UClass* ObjectClass)
 {
-	if (!ObjectClass->IsChildOf(AActor::StaticClass()))
+	UMovieSceneSequence*       RootMovieSceneSequence = GetSequencer()->GetRootMovieSceneSequence();
+	FMovieSceneSequenceEditor* SequenceEditor         = FMovieSceneSequenceEditor::Find(RootMovieSceneSequence);
+
+	if (SequenceEditor && SequenceEditor->SupportsEvents(RootMovieSceneSequence))
 	{
-		return;
+		MenuBuilder.AddMenuEntry(
+			LOCTEXT("AddEventTrack_ObjectBinding", "Event"),
+			LOCTEXT("AddEventTooltip_ObjectBinding", "Adds a new event track that will trigger events on this object binding."),
+			FSlateIcon(),
+			FUIAction( 
+				FExecuteAction::CreateSP( this, &FEventTrackEditor::HandleAddEventTrackMenuEntryExecute, ObjectBinding )
+			)
+		);
 	}
-
 }
 
+TSharedPtr<SWidget> FEventTrackEditor::BuildOutlinerEditWidget(const FGuid& ObjectBinding, UMovieSceneTrack* Track, const FBuildEditWidgetParams& Params)
+{
+	check(Track);
+
+	TSharedPtr<ISequencer> SequencerPtr = GetSequencer();
+	if (!SequencerPtr.IsValid())
+	{
+		return SNullWidget::NullWidget;
+	}
+
+	TWeakObjectPtr<UMovieSceneTrack> WeakTrack = Track;
+	const int32 RowIndex = Params.TrackInsertRowIndex;
+	auto SubMenuCallback = [this, WeakTrack, RowIndex]
+	{
+		FMenuBuilder MenuBuilder(true, nullptr);
+
+		UMovieSceneTrack* TrackPtr = WeakTrack.Get();
+		if (TrackPtr)
+		{
+			MenuBuilder.AddMenuEntry(
+				LOCTEXT("AddNewTriggerSection", "Trigger"),
+				LOCTEXT("AddNewTriggerSectionTooltip", "Adds a new section that can trigger a specific event at a specific time"),
+				FSlateIcon(),
+				FUIAction(FExecuteAction::CreateSP(this, &FEventTrackEditor::CreateNewSection, TrackPtr, RowIndex + 1, UMovieSceneEventTriggerSection::StaticClass()))
+			);
+
+			MenuBuilder.AddMenuEntry(
+				LOCTEXT("AddNewRepeaterSection", "Repeater"),
+				LOCTEXT("AddNewRepeaterSectionTooltip", "Adds a new section that triggers an event every time it's evaluated"),
+				FSlateIcon(),
+				FUIAction(FExecuteAction::CreateSP(this, &FEventTrackEditor::CreateNewSection, TrackPtr, RowIndex + 1, UMovieSceneEventRepeaterSection::StaticClass()))
+			);
+		}
+		else
+		{
+			MenuBuilder.AddWidget(SNew(STextBlock).Text(LOCTEXT("InvalidTrack", "Track is no longer valid")), FText(), true);
+		}
+
+		return MenuBuilder.MakeWidget();
+	};
+
+	return SNew(SHorizontalBox)
+	+ SHorizontalBox::Slot()
+	.AutoWidth()
+	.VAlign(VAlign_Center)
+	[
+		FSequencerUtilities::MakeAddButton(LOCTEXT("AddSection", "Section"), FOnGetContent::CreateLambda(SubMenuCallback), Params.NodeIsHovered)
+	];
+}
 
 void FEventTrackEditor::BuildTrackContextMenu(FMenuBuilder& MenuBuilder, UMovieSceneTrack* Track)
 {
@@ -155,7 +229,7 @@ const FSlateBrush* FEventTrackEditor::GetIconBrush() const
 /* FEventTrackEditor callbacks
  *****************************************************************************/
 
-void FEventTrackEditor::HandleAddEventTrackMenuEntryExecute()
+void FEventTrackEditor::HandleAddEventTrackMenuEntryExecute(FGuid InObjectBindingID)
 {
 	UMovieScene* FocusedMovieScene = GetFocusedMovieScene();
 
@@ -166,8 +240,17 @@ void FEventTrackEditor::HandleAddEventTrackMenuEntryExecute()
 
 	const FScopedTransaction Transaction(NSLOCTEXT("Sequencer", "AddEventTrack_Transaction", "Add Event Track"));
 	FocusedMovieScene->Modify();
-	
-	UMovieSceneEventTrack* NewTrack = FocusedMovieScene->AddMasterTrack<UMovieSceneEventTrack>();
+
+	UMovieSceneEventTrack* NewTrack = nullptr;
+	if (InObjectBindingID.IsValid())
+	{
+		NewTrack = FocusedMovieScene->AddTrack<UMovieSceneEventTrack>(InObjectBindingID);
+	}
+	else
+	{
+		NewTrack = FocusedMovieScene->AddMasterTrack<UMovieSceneEventTrack>();
+	}
+
 	check(NewTrack);
 
 	UMovieSceneSection* NewSection = NewTrack->CreateNewSection();
@@ -184,5 +267,41 @@ void FEventTrackEditor::HandleAddEventTrackMenuEntryExecute()
 	GetSequencer()->NotifyMovieSceneDataChanged( EMovieSceneDataChangeType::MovieSceneStructureItemAdded );
 }
 
+void FEventTrackEditor::CreateNewSection(UMovieSceneTrack* Track, int32 RowIndex, UClass* SectionType)
+{
+	TSharedPtr<ISequencer> SequencerPtr = GetSequencer();
+	if (SequencerPtr.IsValid())
+	{
+		FQualifiedFrameTime CurrentTime = SequencerPtr->GetLocalTime();
+		TRange<double> VisibleRange = SequencerPtr->GetViewRange();
+
+		FScopedTransaction Transaction(LOCTEXT("CreateNewSectionTransactionText", "Add Section"));
+
+		UMovieSceneSection* NewSection = NewObject<UMovieSceneSection>(Track, SectionType);
+		check(NewSection);
+
+		int32 OverlapPriority = 0;
+		for (UMovieSceneSection* Section : Track->GetAllSections())
+		{
+			if (Section->GetRowIndex() >= RowIndex)
+			{
+				Section->SetRowIndex(Section->GetRowIndex() + 1);
+			}
+			OverlapPriority = FMath::Max(Section->GetOverlapPriority() + 1, OverlapPriority);
+		}
+
+		Track->Modify();
+
+		int32 DurationFrames = ( (VisibleRange.Size<double>() * 0.75) * CurrentTime.Rate ).FloorToFrame().Value;
+		NewSection->SetRange(TRange<FFrameNumber>(CurrentTime.Time.FrameNumber, CurrentTime.Time.FrameNumber + DurationFrames));
+		NewSection->SetOverlapPriority(OverlapPriority);
+		NewSection->SetRowIndex(RowIndex);
+
+		Track->AddSection(*NewSection);
+		Track->UpdateEasing();
+
+		SequencerPtr->NotifyMovieSceneDataChanged(EMovieSceneDataChangeType::MovieSceneStructureItemAdded);
+	}
+}
 
 #undef LOCTEXT_NAMESPACE
diff --git a/Engine/Source/Editor/MovieSceneTools/Private/TrackEditors/EventTrackEditor.h b/Engine/Source/Editor/MovieSceneTools/Private/TrackEditors/EventTrackEditor.h
--- a/Engine/Source/Editor/MovieSceneTools/Private/TrackEditors/EventTrackEditor.h
+++ b/Engine/Source/Editor/MovieSceneTools/Private/TrackEditors/EventTrackEditor.h
@@ -48,6 +48,7 @@ public:
 	virtual bool SupportsSequence(UMovieSceneSequence* InSequence) const override;
 	virtual const FSlateBrush* GetIconBrush() const override;
 	virtual void BuildTrackContextMenu( FMenuBuilder& MenuBuilder, UMovieSceneTrack* Track ) override;
+	virtual TSharedPtr<SWidget> BuildOutlinerEditWidget(const FGuid& ObjectBinding, UMovieSceneTrack* Track, const FBuildEditWidgetParams& Params) override;
 
 	//~ FPropertyTrackEditor interface
 	virtual TSharedRef<ISequencerSection> MakeSectionInterface(UMovieSceneSection& SectionObject, UMovieSceneTrack& Track, FGuid ObjectBinding) override;
@@ -55,5 +56,7 @@ public:
 private:
 
 	/** Callback for executing the "Add Event Track" menu entry. */
-	void HandleAddEventTrackMenuEntryExecute();
+	void HandleAddEventTrackMenuEntryExecute(FGuid InObjectBindingID);
+
+	void CreateNewSection(UMovieSceneTrack* Track, int32 RowIndex, UClass* SectionType);
 };
diff --git a/Engine/Source/Editor/MovieSceneTools/Public/ClipboardTypes.h b/Engine/Source/Editor/MovieSceneTools/Public/ClipboardTypes.h
--- a/Engine/Source/Editor/MovieSceneTools/Public/ClipboardTypes.h
+++ b/Engine/Source/Editor/MovieSceneTools/Public/ClipboardTypes.h
@@ -8,6 +8,7 @@
 #if WITH_EDITOR
 
 #include "Channels/MovieSceneFloatChannel.h"
+#include "Channels/MovieSceneEvent.h"
 #include "Sections/MovieSceneEventSection.h"
 #include "Sections/MovieSceneActorReferenceSection.h"
 #include "MovieSceneClipboard.h"
@@ -58,6 +59,10 @@ namespace MovieSceneClipboard
 	{
 		return "EventPayload";
 	}
+	template<> inline FName GetKeyTypeName<FMovieSceneEvent>()
+	{
+		return "MovieSceneEvent";
+	}
 }
 
 #endif
diff --git a/Engine/Source/Editor/MovieSceneTools/Public/MovieSceneEventCustomization.h b/Engine/Source/Editor/MovieSceneTools/Public/MovieSceneEventCustomization.h
new file mode 100644
--- /dev/null
+++ b/Engine/Source/Editor/MovieSceneTools/Public/MovieSceneEventCustomization.h
@@ -0,0 +1,131 @@
+// Copyright 1998-2018 Epic Games, Inc. All Rights Reserved.
+
+#pragma once
+
+#include "CoreMinimal.h"
+#include "IPropertyTypeCustomization.h"
+#include "IStructureDetailsView.h"
+
+struct FAssetData;
+
+class FMenuBuilder;
+class IPropertyHandle;
+class FDetailWidgetRow;
+class UMovieSceneSection;
+class UMovieSceneSequence;
+class UMovieSceneEventTrack;
+class UK2Node_FunctionEntry;
+class IDetailChildrenBuilder;
+
+/**
+ * Customization for FMovieSceneEvent structs.
+ * Will deduce the event's section either from the outer objects on the details customization, or use the one provided on construction (for instanced property type customizations)
+ */
+class FMovieSceneEventCustomization : public IPropertyTypeCustomization
+{
+public:
+	static TSharedRef<IPropertyTypeCustomization> MakeInstance();
+	static TSharedRef<IPropertyTypeCustomization> MakeInstance(UMovieSceneSection* InSection);
+
+	virtual void CustomizeHeader( TSharedRef<IPropertyHandle> PropertyHandle, FDetailWidgetRow& HeaderRow, IPropertyTypeCustomizationUtils& CustomizationUtils ) override;
+	virtual void CustomizeChildren( TSharedRef<IPropertyHandle> PropertyHandle, IDetailChildrenBuilder& ChildBuilder, IPropertyTypeCustomizationUtils& CustomizationUtils ) override;
+
+private:
+
+	/**
+	 * Clear the endpoint for the event(s) represented by this property handle.
+	 * @note: Does not delete the endpoint in the blueprint itself.
+	 */
+	void ClearEventEndpoint();
+
+
+	/**
+	 * Creates a single new endpoint for the event(s) represented by this property handle.
+	 */
+	void CreateEventEndpoint();
+
+
+	/**
+	 * Creates a single new endpoint for the event(s) represented by this property handle using the specified function as a quick binding.
+	 */
+	void CreateEventEndpointFromFunction(UFunction* QuickBindFunction, UClass* PinClassType);
+
+
+	/**
+	 * Assigns the specified function entry to the event(s) represented by this property handle.
+	 */
+	void SetEventEndpoint(UK2Node_FunctionEntry* NewEndpoint);
+
+	/**
+	 * Compare the currently assigned endpoint with the specified entry. Used as the check-state for the radio buttons on the menu.
+	 */
+	bool CompareCurrentEventEndpoint(UK2Node_FunctionEntry* NewEndpoint);
+
+
+	/**
+	 * Navigate to the definition of the endpoint specified by the event(s) represented by this property handle.
+	 */
+	void NavigateToDefinition();
+
+
+	/**
+	 * Generate the content of the main combo button menu dropdown
+	 */
+	TSharedRef<SWidget> GetMenuContent();
+
+
+	/**
+	 * Generate the content of the creation shortcut sub-menu dropdown
+	 */
+	void PopulateQuickBindSubMenu(FMenuBuilder& MenuBuilder, UClass* TemplateClass);
+
+private:
+
+	/**
+	 * Get the sequence that is common to all the events represented by this property handle, or nullptr if they are not all the same.
+	 */
+	UMovieSceneSequence* GetCommonSequence() const;
+
+
+	/**
+	 * Get the track that is common to all the events represented by this property handle, or nullptr if they are not all the same.
+	 */
+	UMovieSceneEventTrack* GetCommonTrack() const;
+
+
+	/**
+	 * Get the endpoint that is common to all the events represented by this property handle, or nullptr if they are not all the same.
+	 */
+	UK2Node_FunctionEntry* GetCommonEndpoint() const;
+
+
+	/**
+	 * Get all the objects that the events reside within.
+	 */
+	void GetEditObjects(TArray<UObject*>& OutObjects) const;
+
+private:
+
+	/** Get the name of the event to display on the main combo button */
+	FText GetEventName() const;
+	/** Get the icon of the event to display on the main combo button */
+	const FSlateBrush* GetEventIcon() const;
+
+	/** Get the visibility of the error icon */
+	EVisibility GetErrorVisibility() const;
+	/** Get the tooltip text for the error icon */
+	FText GetErrorTooltip() const;
+
+	UClass* FindObjectBindingClass();
+
+private:
+
+	/** Externally supplied section that the event(s) we're reflecting reside within */
+	TWeakObjectPtr<UMovieSceneSection> WeakExternalSection;
+
+	/** A cache of the common endpoint that is only used when the menu is open to avoid re-computing it every frame. */
+	TWeakObjectPtr<UK2Node_FunctionEntry> CachedCommonEndpoint;
+
+	/** The property handle we're reflecting */
+	TSharedPtr<IPropertyHandle> PropertyHandle;
+};
diff --git a/Engine/Source/Editor/MovieSceneTools/Public/SKeyEditInterface.h b/Engine/Source/Editor/MovieSceneTools/Public/SKeyEditInterface.h
--- a/Engine/Source/Editor/MovieSceneTools/Public/SKeyEditInterface.h
+++ b/Engine/Source/Editor/MovieSceneTools/Public/SKeyEditInterface.h
@@ -47,6 +47,7 @@ private:
 	 */
 	TSharedRef<IPropertyTypeCustomization> CreateBindingIDCustomization();
 	TSharedRef<IPropertyTypeCustomization> CreateFrameNumberCustomization();
+	TSharedRef<IPropertyTypeCustomization> CreateEventCustomization();
 
 	/**
 	 * Called when a property has been changed on the UI
diff --git a/Engine/Source/Editor/Sequencer/Private/LevelEditorSequencerIntegration.cpp b/Engine/Source/Editor/Sequencer/Private/LevelEditorSequencerIntegration.cpp
--- a/Engine/Source/Editor/Sequencer/Private/LevelEditorSequencerIntegration.cpp
+++ b/Engine/Source/Editor/Sequencer/Private/LevelEditorSequencerIntegration.cpp
@@ -572,6 +572,7 @@ void FLevelEditorSequencerIntegration::OnPreBeginPIE(bool bIsSimulating)
 		{
 			if (Options.bRequiresLevelEvents)
 			{
+				In.GetEvaluationTemplate().ResetDirectorInstances();
 				In.RestorePreAnimatedState();
 				In.State.ClearObjectCaches(In);
 
@@ -593,6 +594,7 @@ void FLevelEditorSequencerIntegration::OnEndPlayMap()
 			if (Options.bRequiresLevelEvents)
 			{
 				// Update and clear any stale bindings 
+				In.GetEvaluationTemplate().ResetDirectorInstances();
 				In.State.ClearObjectCaches(In);
 				In.ForceEvaluate();
 			}
diff --git a/Engine/Source/Editor/Sequencer/Private/MovieSceneSequenceEditor.cpp b/Engine/Source/Editor/Sequencer/Private/MovieSceneSequenceEditor.cpp
new file mode 100644
--- /dev/null
+++ b/Engine/Source/Editor/Sequencer/Private/MovieSceneSequenceEditor.cpp
@@ -0,0 +1,152 @@
+
+#include "MovieSceneSequenceEditor.h"
+#include "Modules/ModuleManager.h"
+#include "MovieSceneSection.h"
+#include "MovieSceneSequence.h"
+#include "Channels/MovieSceneEvent.h"
+#include "Tracks/MovieSceneEventTrack.h"
+#include "Sections/MovieSceneEventSectionBase.h"
+#include "ISequencerModule.h"
+#include "ScopedTransaction.h"
+
+#include "Engine/Blueprint.h"
+#include "EdGraph/EdGraph.h"
+#include "K2Node_FunctionEntry.h"
+#include "Kismet2/KismetEditorUtilities.h"
+#include "Kismet2/BlueprintEditorUtils.h"
+
+#define LOCTEXT_NAMESPACE "MovieSceneSequenceEditor"
+
+namespace
+{
+	FText GetDefaultCommentText()
+	{
+		FText DefaultComment_01 = LOCTEXT("DefaultComment_01", "Sequencer Events can be bound to either of the following function signatures. Return values are not supported.");
+		FText DefaultComment_02 = LOCTEXT("DefaultComment_02", "\t1. A function with no parameters");
+		FText DefaultComment_03 = LOCTEXT("DefaultComment_03", "\t\tCompatible with master tracks or object bindings");
+		FText DefaultComment_04 = LOCTEXT("DefaultComment_04", "\t2. A function with a single object or interface parameter");
+		FText DefaultComment_05 = LOCTEXT("DefaultComment_05", "\t\tCompatible with master tracks or object bindings");
+		FText DefaultComment_06 = LOCTEXT("DefaultComment_06", "\t\tWill only trigger if the source object is of the same type as the parameter (or interface)");
+		FText DefaultComment_07 = LOCTEXT("DefaultComment_07", "\t\tWill be triggered with objects in the following order:");
+		FText DefaultComment_08 = LOCTEXT("DefaultComment_08", "\t\t\ta) Objects bound to the track's object binding, or:");
+		FText DefaultComment_09 = LOCTEXT("DefaultComment_09", "\t\t\tb) Objects specified on the track's event receivers array, or:");
+		FText DefaultComment_10 = LOCTEXT("DefaultComment_10", "\t\t\tc) Objects provided by the playback context (level blueprints, widgets etc)");
+		FText DefaultComment_11 = LOCTEXT("DefaultComment_11", "Tip: Trigger events on level blueprints by implementing an interface on it");
+
+		return FText::Format(LOCTEXT("DefaultComment_Format", "{0}\n{1}\n{2}\n{3}\n{4}\n{5}\n{6}\n{7}\n{8}\n{9}\n\n{10}"),
+			DefaultComment_01, DefaultComment_02, DefaultComment_03, DefaultComment_04, DefaultComment_05,
+			DefaultComment_06, DefaultComment_07, DefaultComment_08, DefaultComment_09, DefaultComment_10, DefaultComment_11);
+	}
+}
+
+FName FMovieSceneSequenceEditor::TargetPinName("Target");
+
+FMovieSceneSequenceEditor* FMovieSceneSequenceEditor::Find(UMovieSceneSequence* InSequence)
+{
+	ISequencerModule& SequencerModule = FModuleManager::LoadModuleChecked<ISequencerModule>("Sequencer");
+	return InSequence ? SequencerModule.FindSequenceEditor(InSequence->GetClass()) : nullptr;
+}
+
+bool FMovieSceneSequenceEditor::SupportsEvents(UMovieSceneSequence* InSequence) const
+{
+	return CanCreateEvents(InSequence);
+}
+
+UBlueprint* FMovieSceneSequenceEditor::GetDirectorBlueprint(UMovieSceneSequence* Sequence) const
+{
+	return GetBlueprintForSequence(Sequence);
+}
+
+UBlueprint* FMovieSceneSequenceEditor::AccessDirectorBlueprint(UMovieSceneSequence* Sequence) const
+{
+	UBlueprint* Blueprint = GetBlueprintForSequence(Sequence);
+	if (!Blueprint)
+	{
+		Blueprint = CreateBlueprintForSequence(Sequence);
+	}
+	return Blueprint;
+}
+
+UK2Node_FunctionEntry* FMovieSceneSequenceEditor::CreateEventEndpoint(UMovieSceneSequence* Sequence, const FString& DesiredName) const
+{
+	FScopedTransaction Transaction(LOCTEXT("CreateEventEndpoint", "Create Event Endpoint"));
+
+	UBlueprint* Blueprint = AccessDirectorBlueprint(Sequence);
+	if (Blueprint)
+	{
+		static FString DefaultEventName = "SequenceEvent";
+		FName UniqueGraphName = FBlueprintEditorUtils::FindUniqueKismetName(Blueprint, DesiredName.Len() == 0 ? DefaultEventName : DesiredName);
+
+		Blueprint->Modify();
+
+		UEdGraph* Graph = FBlueprintEditorUtils::CreateNewGraph(Blueprint, UniqueGraphName, UEdGraph::StaticClass(), UEdGraphSchema_K2::StaticClass());
+
+		const bool bIsUserCreated = false;
+		FBlueprintEditorUtils::AddFunctionGraph<UClass>(Blueprint, Graph, bIsUserCreated, nullptr);
+
+		TArray<UK2Node_FunctionEntry*> EntryNodes;
+		Graph->GetNodesOfClass<UK2Node_FunctionEntry>(EntryNodes);
+
+		if (ensure(EntryNodes.Num() == 1 && EntryNodes[0]))
+		{
+			int32 ExtraFunctionFlags = ( FUNC_BlueprintCallable | FUNC_BlueprintEvent | FUNC_Public );
+			EntryNodes[0]->AddExtraFlags(ExtraFunctionFlags);
+
+			EntryNodes[0]->bIsEditable = true;
+			EntryNodes[0]->MetaData.Category = LOCTEXT("DefaultCategory", "Sequencer Event Endpoints");
+
+			EntryNodes[0]->NodeComment = GetDefaultCommentText().ToString();
+			EntryNodes[0]->bCommentBubblePinned = true;
+			EntryNodes[0]->bCommentBubbleVisible = true;
+
+			return EntryNodes[0];
+		}
+	}
+
+	return nullptr;
+}
+
+void FMovieSceneSequenceEditor::InitializeEndpointForTrack(UMovieSceneEventTrack* EventTrack, UK2Node_FunctionEntry* Endpoint) const
+{
+	SetupDefaultPinForEndpoint(EventTrack, Endpoint);
+}
+
+UClass* FMovieSceneSequenceEditor::FindTrackObjectBindingClass(UMovieSceneTrack* Track)
+{
+	UMovieScene* MovieScene = Track->GetTypedOuter<UMovieScene>();
+
+	// We have to find the object binding that owns this track, which is potentially slow, but there is no back reference here.
+	for (const FMovieSceneBinding& Binding : MovieScene->GetBindings())
+	{
+		if (!Algo::Find(Binding.GetTracks(), Track))
+		{
+			continue;
+		}
+
+		if (FMovieScenePossessable* Possessable = MovieScene->FindPossessable(Binding.GetObjectGuid()))
+		{
+			return const_cast<UClass*>(Possessable->GetPossessedObjectClass());
+		}
+		else if(FMovieSceneSpawnable* Spawnable = MovieScene->FindSpawnable(Binding.GetObjectGuid()))
+		{
+			return Spawnable->GetObjectTemplate()->GetClass();
+		}
+	}
+
+	return nullptr;
+}
+
+void FMovieSceneSequenceEditor::BindEventToEndpoint(UMovieSceneEventSectionBase* EventSection, FMovieSceneEvent* Event, UK2Node_FunctionEntry* Endpoint)
+{
+	check(EventSection && Event);
+
+	EventSection->Modify();
+	Event->SetFunctionEntry(Endpoint);
+
+	if (UBlueprint* Blueprint = Endpoint ? Endpoint->GetBlueprint() : nullptr)
+	{
+		EventSection->SetDirectorBlueprint(Blueprint);
+	}
+}
+
+#undef LOCTEXT_NAMESPACE	// MovieSceneSequenceEditor
\ No newline at end of file
diff --git a/Engine/Source/Editor/Sequencer/Private/SSequencerSection.cpp b/Engine/Source/Editor/Sequencer/Private/SSequencerSection.cpp
--- a/Engine/Source/Editor/Sequencer/Private/SSequencerSection.cpp
+++ b/Engine/Source/Editor/Sequencer/Private/SSequencerSection.cpp
@@ -1916,8 +1916,14 @@ FReply SSequencerSection::OnMouseButtonDoubleClick( const FGeometry& MyGeometry,
 {
 	if( MouseEvent.GetEffectingButton() == EKeys::LeftMouseButton )
 	{
-		FReply Reply = SectionInterface->OnSectionDoubleClicked( MyGeometry, MouseEvent );
+		TArray<FSequencerSelectedKey> Keys;
+		GetKeysUnderMouse(MouseEvent.GetScreenSpacePosition(), MyGeometry, Keys);
+		if (Keys.Num() == 1 && Keys[0].KeyHandle.IsSet())
+		{
+			return SectionInterface->OnKeyDoubleClicked(Keys[0].KeyHandle.GetValue());
+		}
 
+		FReply Reply = SectionInterface->OnSectionDoubleClicked( MyGeometry, MouseEvent );
 		if (!Reply.IsEventHandled())
 		{
 			// Find the object binding this node is underneath
diff --git a/Engine/Source/Editor/Sequencer/Private/Sequencer.cpp b/Engine/Source/Editor/Sequencer/Private/Sequencer.cpp
--- a/Engine/Source/Editor/Sequencer/Private/Sequencer.cpp
+++ b/Engine/Source/Editor/Sequencer/Private/Sequencer.cpp
@@ -120,6 +120,9 @@
 #include "Fonts/FontMeasure.h"
 #include "MovieSceneTimeHelpers.h"
 #include "FrameNumberNumericInterface.h"
+#include "Engine/Blueprint.h"
+#include "MovieSceneSequenceEditor.h"
+#include "Kismet2/KismetEditorUtilities.h"
 
 #define LOCTEXT_NAMESPACE "Sequencer"
 
@@ -2203,6 +2206,11 @@ void FSequencer::ForceEvaluate()
 
 void FSequencer::EvaluateInternal(FMovieSceneEvaluationRange InRange, bool bHasJumped)
 {
+	if (Settings->ShouldCompileDirectorOnEvaluate())
+	{
+		RecompileDirtyDirectors();
+	}
+
 	bNeedsEvaluate = false;
 
 	if (PlaybackContextAttribute.IsBound())
@@ -2816,22 +2824,7 @@ void FSequencer::AddReferencedObjects( FReferenceCollector& Collector )
 		Collector.AddReferencedObject( RootSequencePtr );
 	}
 
-	if (RootTemplateInstance.IsValid())
-	{
-		const FMovieSceneSequenceHierarchy& Hierarchy = RootTemplateInstance.GetHierarchy();
-
-		// Sequencer references all active sub sequences
-		for (FMovieSceneSequenceIDRef SequenceID : RootTemplateInstance.GetThisFrameMetaData().ActiveSequences)
-		{
-			const FMovieSceneSubSequenceData* SubData  = Hierarchy.FindSubData(SequenceID);
-			UMovieSceneSequence*              Sequence = SubData ? SubData->GetLoadedSequence() : nullptr;
-
-			if (Sequence)
-			{
-				Collector.AddReferencedObject(Sequence);
-			}
-		}
-	}
+	FMovieSceneRootEvaluationTemplateInstance::StaticStruct()->SerializeBin(Collector.GetVerySlowReferenceCollectorArchive(), &RootTemplateInstance);
 }
 
 
@@ -8725,5 +8718,38 @@ double FSequencer::GetDisplayRateDeltaFrameCount() const
 	return GetFocusedTickResolution().AsDecimal() * GetFocusedDisplayRate().AsInterval();
 }
 
+void FSequencer::RecompileDirtyDirectors()
+{
+	ISequencerModule& SequencerModule = FModuleManager::LoadModuleChecked<ISequencerModule>("Sequencer");
+
+	TSet<UMovieSceneSequence*> AllSequences;
+
+	// Gather all sequences in the hierarchy
+	if (UMovieSceneSequence* Sequence = RootSequence.Get())
+	{
+		AllSequences.Add(Sequence);
+	}
+
+	for (const TTuple<FMovieSceneSequenceID, FMovieSceneSubSequenceData>& Pair : RootTemplateInstance.GetHierarchy().AllSubSequenceData())
+	{
+		if (UMovieSceneSequence* Sequence = Pair.Value.GetSequence())
+		{
+			AllSequences.Add(Sequence);
+		}
+	}
+
+	// Recompile them all if they are dirty
+	for (UMovieSceneSequence* Sequence : AllSequences)
+	{
+		FMovieSceneSequenceEditor* SequenceEditor = SequencerModule.FindSequenceEditor(Sequence->GetClass());
+		UBlueprint*                DirectorBP     = SequenceEditor ? SequenceEditor->GetDirectorBlueprint(Sequence) : nullptr;
+
+		if (DirectorBP && (DirectorBP->Status == BS_Unknown || DirectorBP->Status == BS_Dirty))
+		{
+			FKismetEditorUtilities::CompileBlueprint(DirectorBP);
+		}
+	}
+}
+
 
 #undef LOCTEXT_NAMESPACE
diff --git a/Engine/Source/Editor/Sequencer/Private/Sequencer.h b/Engine/Source/Editor/Sequencer/Private/Sequencer.h
--- a/Engine/Source/Editor/Sequencer/Private/Sequencer.h
+++ b/Engine/Source/Editor/Sequencer/Private/Sequencer.h
@@ -964,6 +964,9 @@ protected:
 	/** Get the default key attributes to apply to newly created keys on the curve editor */
 	FKeyAttributes GetDefaultKeyAttributes() const;
 
+	/** Recompile any dirty director blueprints in the sequence hierarchy */
+	void RecompileDirtyDirectors();
+
 public:
 
 	/** Reset the timing manager to the clock source specified by the root movie scene */
diff --git a/Engine/Source/Editor/Sequencer/Private/SequencerSettings.cpp b/Engine/Source/Editor/Sequencer/Private/SequencerSettings.cpp
--- a/Engine/Source/Editor/Sequencer/Private/SequencerSettings.cpp
+++ b/Engine/Source/Editor/Sequencer/Private/SequencerSettings.cpp
@@ -37,6 +37,7 @@ USequencerSettings::USequencerSettings( const FObjectInitializer& ObjectInitiali
 	bKeepCursorInPlayRangeWhileScrubbing = false;
 	bKeepCursorInPlayRange = true;
 	bKeepPlayRangeInSectionBounds = true;
+	bCompileDirectorOnEvaluate = true;
 	ZeroPadFrames = 0;
 	bShowCombinedKeyframes = true;
 	bInfiniteKeyAreas = false;
@@ -626,6 +627,21 @@ void USequencerSettings::SetShouldShowPrePostRoll(bool bInVisualizePreAndPostRol
 	}
 }
 
+
+bool USequencerSettings::ShouldCompileDirectorOnEvaluate() const
+{
+	return bCompileDirectorOnEvaluate;
+}
+
+void USequencerSettings::SetCompileDirectorOnEvaluate(bool bInCompileDirectorOnEvaluate)
+{
+	if (bInCompileDirectorOnEvaluate != bCompileDirectorOnEvaluate)
+	{
+		bCompileDirectorOnEvaluate = bInCompileDirectorOnEvaluate;
+		SaveConfig();
+	}
+}
+
 ECurveEditorCurveVisibility USequencerSettings::GetCurveVisibility() const
 {
 	return CurveVisibility;
diff --git a/Engine/Source/Editor/Sequencer/Public/ISequencerModule.h b/Engine/Source/Editor/Sequencer/Public/ISequencerModule.h
--- a/Engine/Source/Editor/Sequencer/Public/ISequencerModule.h
+++ b/Engine/Source/Editor/Sequencer/Public/ISequencerModule.h
@@ -8,6 +8,7 @@
 #include "Modules/ModuleInterface.h"
 #include "AnimatedPropertyKey.h"
 #include "ISequencerChannelInterface.h"
+#include "MovieSceneSequenceEditor.h"
 
 class IKeyArea;
 class FExtender;
@@ -237,6 +238,55 @@ public:
 	 */
 	ISequencerChannelInterface* FindChannelEditorInterface(FName ChannelTypeName) const;
 
+public:
+
+	/**
+	 * Register a sequence editor for the specified type of sequence. Sequence editors provide editor-only functionality for particular sequence types.
+	 */
+	FDelegateHandle RegisterSequenceEditor(UClass* SequenceClass, TUniquePtr<FMovieSceneSequenceEditor>&& InSequenceEditor)
+	{
+		check(SequenceClass);
+
+		FDelegateHandle NewHandle(FDelegateHandle::GenerateNewHandle);
+
+		SequenceEditors.Add(FSequenceEditorEntry{ NewHandle, SequenceClass, MoveTemp(InSequenceEditor) });
+
+		return NewHandle;
+	}
+
+	/**
+	 * Unregister a sequence editor for the specified type of sequence.
+	 */
+	void UnregisterSequenceEditor(FDelegateHandle Handle)
+	{
+		SequenceEditors.RemoveAll([Handle](const FSequenceEditorEntry& In){ return In.Handle == Handle; });
+	}
+
+	/**
+	 * Find a sequence editor for the specified sequence class
+	 */
+	FMovieSceneSequenceEditor* FindSequenceEditor(UClass* SequenceClass) const
+	{
+		check(SequenceClass);
+
+		UClass* MostRelevantClass = nullptr;
+		FMovieSceneSequenceEditor* SequenceEditor = nullptr;
+
+		for (const FSequenceEditorEntry& Entry : SequenceEditors)
+		{
+			if (SequenceClass->IsChildOf(Entry.ApplicableClass))
+			{
+				if (!MostRelevantClass || Entry.ApplicableClass->IsChildOf(MostRelevantClass))
+				{
+					MostRelevantClass = Entry.ApplicableClass;
+					SequenceEditor = Entry.Editor.Get();
+				}
+			}
+		}
+
+		return SequenceEditor;
+	}
+
 public:
 
 	/** 
@@ -270,6 +320,16 @@ private:
 
 	/** Map of sequencer interfaces for movie scene channel types, keyed on channel UStruct name */
 	TMap<FName, TUniquePtr<ISequencerChannelInterface>> ChannelToEditorInterfaceMap;
+
+	struct FSequenceEditorEntry
+	{
+		FDelegateHandle Handle;
+		UClass* ApplicableClass;
+		TUniquePtr<FMovieSceneSequenceEditor> Editor;
+	};
+
+	/** Array of sequence editor entries */
+	TArray<FSequenceEditorEntry> SequenceEditors;
 };
 
 
diff --git a/Engine/Source/Editor/Sequencer/Public/ISequencerSection.h b/Engine/Source/Editor/Sequencer/Public/ISequencerSection.h
--- a/Engine/Source/Editor/Sequencer/Public/ISequencerSection.h
+++ b/Engine/Source/Editor/Sequencer/Public/ISequencerSection.h
@@ -93,6 +93,14 @@ public:
 	 */
 	virtual FReply OnSectionDoubleClicked( const FGeometry& SectionGeometry, const FPointerEvent& MouseEvent, const FGuid& ObjectBinding) { return FReply::Unhandled(); }
 
+	/**
+	 * Called when a key on this section is double clicked
+	 *
+	 * @param KeyHandle			The key that was clicked
+	 * @return A reply in response to double clicking the key
+	 */
+	virtual FReply OnKeyDoubleClicked( FKeyHandle KeyHandle ) { return FReply::Unhandled(); }
+
 	/**
 	 * @return The display name of the section in the section view
 	 */
diff --git a/Engine/Source/Editor/Sequencer/Public/MovieSceneSequenceEditor.h b/Engine/Source/Editor/Sequencer/Public/MovieSceneSequenceEditor.h
new file mode 100644
--- /dev/null
+++ b/Engine/Source/Editor/Sequencer/Public/MovieSceneSequenceEditor.h
@@ -0,0 +1,107 @@
+// Copyright 1998-2018 Epic Games, Inc. All Rights Reserved.
+
+#pragma once
+
+#include "CoreTypes.h"
+#include "UObject/NameTypes.h"
+
+struct FMovieSceneEvent;
+
+class UBlueprint;
+class UMovieSceneTrack;
+class UMovieSceneSequence;
+class UMovieSceneEventTrack;
+class UK2Node_FunctionEntry;
+class UMovieSceneEventSectionBase;
+
+struct SEQUENCER_API FMovieSceneSequenceEditor
+{
+	virtual ~FMovieSceneSequenceEditor(){}
+
+
+	/**
+	 * The name of the target pin for event function entry nodes
+	 */
+	static FName TargetPinName;
+
+
+	/**
+	 * Attempt to find a sequence editor for the specified sequence
+	 *
+	 * @param Sequence        The sequence to get an editor for
+	 * @return The sequence editor ptr, or null if one is not available for the specified type of sequence
+	 */
+	static FMovieSceneSequenceEditor* Find(UMovieSceneSequence* InSequence);
+
+
+	/**
+	 * Bind an event to an endpoint. Ensures the necessary pointers are set up to enable automatic re-compilation on evaluation.
+	 *
+	 * @param EventSection    The section that the event belongs to
+	 * @param Event           The event to bind to the endpoint
+	 * @param Endpoint        The blueprint node to bind to
+	 */
+	static void BindEventToEndpoint(UMovieSceneEventSectionBase* EventSection, FMovieSceneEvent* Event, UK2Node_FunctionEntry* Endpoint);
+
+	/**
+	 * Check whether the specified sequence supports events
+	 *
+	 * @param Sequence        The sequence to test
+	 */
+	bool SupportsEvents(UMovieSceneSequence* InSequence) const;
+
+
+	/**
+	 * Access the director blueprint for the specified sequence
+	 *
+	 * @param Sequence        The sequence to access the director blueprint for
+	 * @return The sequence's director blueprint or nullptr if it does not have one (or cannot)
+	 */
+	UBlueprint* GetDirectorBlueprint(UMovieSceneSequence* Sequence) const;
+
+
+	/**
+	 * Create a new event endpoint for the specified sequence
+	 *
+	 * @param Sequence        The sequence to create an endpoint in
+	 * @param DesiredName     (Optional) A user-specified name to use for the function
+	 * @return The blueprint node for the event endpoint, or nullptr if one does not exist
+	 */
+	UK2Node_FunctionEntry* CreateEventEndpoint(UMovieSceneSequence* Sequence, const FString& DesiredName = FString()) const;
+
+
+	/**
+	 * Initializes the specified endpoint node by creating an appropriate input pin for the track's object binding, if necessary
+	 *
+	 * @param EventTrack      The event track that owns the event
+	 * @param Endpoint        The event endpoint to initialize
+	 */
+	void InitializeEndpointForTrack(UMovieSceneEventTrack* EventTrack, UK2Node_FunctionEntry* Endpoint) const;
+
+protected:
+
+
+	/**
+	 * Find the class of the object binding the specified track is on, or nullptr
+	 */
+	static UClass* FindTrackObjectBindingClass(UMovieSceneTrack* Track);
+
+
+private:
+
+	/**
+	 * Access the director blueprint for the specified sequence
+	 *
+	 * @param Sequence        The sequence to access the director blueprint for
+	 * @return The sequence's director blueprint or nullptr if it does not have one (or cannot)
+	 */
+	UBlueprint* AccessDirectorBlueprint(UMovieSceneSequence* Sequence) const;
+
+private:
+
+	virtual bool CanCreateEvents(UMovieSceneSequence* InSequence) const { return false; }
+	virtual UBlueprint* GetBlueprintForSequence(UMovieSceneSequence* InSequence) const { return nullptr; }
+	virtual UBlueprint* CreateBlueprintForSequence(UMovieSceneSequence* InSequence) const { return nullptr; }
+	virtual void SetupDefaultPinForEndpoint(UMovieSceneEventTrack* EventTrack, UK2Node_FunctionEntry* Endpoint) const {}
+};
+
diff --git a/Engine/Source/Editor/Sequencer/Public/SequencerGenericKeyStruct.inl b/Engine/Source/Editor/Sequencer/Public/SequencerGenericKeyStruct.inl
--- a/Engine/Source/Editor/Sequencer/Public/SequencerGenericKeyStruct.inl
+++ b/Engine/Source/Editor/Sequencer/Public/SequencerGenericKeyStruct.inl
@@ -42,7 +42,7 @@ void TMovieSceneKeyStructCustomization<ChannelType>::Extend(IDetailLayoutBuilder
 
 	if (KeyStruct.IsValid())
 	{
-		DetailBuilder.EditCategory("Key").AddAllExternalStructureProperties(KeyStruct.ToSharedRef());
+		DetailBuilder.EditCategory("Key").AddExternalStructure(KeyStruct.ToSharedRef());
 	}
 }
 
diff --git a/Engine/Source/Editor/Sequencer/Public/SequencerSettings.h b/Engine/Source/Editor/Sequencer/Public/SequencerSettings.h
--- a/Engine/Source/Editor/Sequencer/Public/SequencerSettings.h
+++ b/Engine/Source/Editor/Sequencer/Public/SequencerSettings.h
@@ -318,6 +318,11 @@ public:
 	/** Toggle whether to show pre and post roll in sequencer */
 	void SetShouldShowPrePostRoll(bool bInVisualizePreAndPostRoll);
 
+	/** Check whether whether to recompile the director blueprint when the sequence is evaluated (if one exists) */
+	bool ShouldCompileDirectorOnEvaluate() const;
+	/** Assign whether whether to recompile the director blueprint when the sequence is evaluated (if one exists) */
+	void SetCompileDirectorOnEvaluate(bool bInCompileDirectorOnEvaluate);
+
 	uint32 GetTrajectoryPathCap() const { return TrajectoryPathCap; }
 
 	/** Gets the current curve visibility. */
@@ -492,6 +497,10 @@ protected:
 	UPROPERTY( config, EditAnywhere, Category=General )
 	bool bVisualizePreAndPostRoll;
 
+	/** Whether to recompile the director blueprint when the sequence is evaluated (if one exists) */
+	UPROPERTY(config, EditAnywhere, Category=General)
+	bool bCompileDirectorOnEvaluate;
+
 	/** Specifies the maximum number of keys to draw when rendering trajectories in viewports */
 	UPROPERTY(config, EditAnywhere, Category=General)
 	uint32 TrajectoryPathCap;
diff --git a/Engine/Source/Editor/UMGEditor/Private/Animation/MovieSceneSequenceEditor_WidgetAnimation.h b/Engine/Source/Editor/UMGEditor/Private/Animation/MovieSceneSequenceEditor_WidgetAnimation.h
new file mode 100644
--- /dev/null
+++ b/Engine/Source/Editor/UMGEditor/Private/Animation/MovieSceneSequenceEditor_WidgetAnimation.h
@@ -0,0 +1,47 @@
+// Copyright 1998-2018 Epic Games, Inc. All Rights Reserved.
+
+#pragma once
+
+#include "MovieSceneSequenceEditor.h"
+#include "Tracks/MovieSceneEventTrack.h"
+#include "Engine/Blueprint.h"
+#include "K2Node_FunctionEntry.h"
+
+struct FMovieSceneSequenceEditor_WidgetAnimation : FMovieSceneSequenceEditor
+{
+	virtual bool CanCreateEvents(UMovieSceneSequence* InSequence) const
+	{
+		return true;
+	}
+
+	virtual UBlueprint* GetBlueprintForSequence(UMovieSceneSequence* InSequence) const override
+	{
+		return InSequence->GetTypedOuter<UBlueprint>();
+	}
+
+	virtual void SetupDefaultPinForEndpoint(UMovieSceneEventTrack* EventTrack, UK2Node_FunctionEntry* Endpoint) const override
+	{
+		UClass* PinClass = nullptr;
+
+		// When event receivers are set, we just add a generic object pin as it's impossible to know what type the event receiver is here
+		if (EventTrack->EventReceivers.Num() > 0)
+		{
+			PinClass = UObject::StaticClass();
+		}
+		// Otherwise we use the object binding's class, or nullptr
+		else
+		{
+			PinClass = FindTrackObjectBindingClass(EventTrack);
+		}
+
+		// Do not create a pin for nullptr since we're probably a master track that is just triggering on self
+		if (PinClass)
+		{
+			FEdGraphPinType PinType;
+			PinType.PinCategory = UEdGraphSchema_K2::PC_Object;
+			PinType.PinSubCategoryObject = PinClass;
+
+			Endpoint->CreateUserDefinedPin(TargetPinName, PinType, EGPD_Output, true);
+		}
+	}
+};
\ No newline at end of file
diff --git a/Engine/Source/Editor/UMGEditor/Private/UMGEditorModule.cpp b/Engine/Source/Editor/UMGEditor/Private/UMGEditorModule.cpp
--- a/Engine/Source/Editor/UMGEditor/Private/UMGEditorModule.cpp
+++ b/Engine/Source/Editor/UMGEditor/Private/UMGEditorModule.cpp
@@ -9,6 +9,7 @@
 #include "Blueprint/WidgetBlueprintGeneratedClass.h"
 #include "Settings/WidgetDesignerSettings.h"
 #include "WidgetBlueprint.h"
+#include "Animation/WidgetAnimation.h"
 
 #include "AssetToolsModule.h"
 #include "IAssetTypeActions.h"
@@ -20,6 +21,7 @@
 #include "Animation/MarginTrackEditor.h"
 #include "Animation/Sequencer2DTransformTrackEditor.h"
 #include "Animation/WidgetMaterialTrackEditor.h"
+#include "Animation/MovieSceneSequenceEditor_WidgetAnimation.h"
 #include "IUMGModule.h"
 #include "ComponentReregisterContext.h"
 #include "Components/WidgetComponent.h"
@@ -84,6 +86,7 @@ public:
 
 		// Register with the sequencer module that we provide auto-key handlers.
 		ISequencerModule& SequencerModule = FModuleManager::Get().LoadModuleChecked<ISequencerModule>("Sequencer");
+		SequenceEditorHandle                              = SequencerModule.RegisterSequenceEditor(UWidgetAnimation::StaticClass(), MakeUnique<FMovieSceneSequenceEditor_WidgetAnimation>());
 		MarginTrackEditorCreateTrackEditorHandle          = SequencerModule.RegisterPropertyTrackEditor<FMarginTrackEditor>();
 		TransformTrackEditorCreateTrackEditorHandle       = SequencerModule.RegisterPropertyTrackEditor<F2DTransformTrackEditor>();
 		WidgetMaterialTrackEditorCreateTrackEditorHandle  = SequencerModule.RegisterTrackEditor(FOnCreateTrackEditor::CreateStatic(&FWidgetMaterialTrackEditor::CreateTrackEditor));
@@ -117,6 +120,8 @@ public:
 		ISequencerModule* SequencerModule = FModuleManager::GetModulePtr<ISequencerModule>( "Sequencer" );
 		if ( SequencerModule != nullptr )
 		{
+			SequencerModule->UnregisterSequenceEditor(SequenceEditorHandle);
+
 			SequencerModule->UnRegisterTrackEditor( MarginTrackEditorCreateTrackEditorHandle );
 			SequencerModule->UnRegisterTrackEditor( TransformTrackEditorCreateTrackEditorHandle );
 			SequencerModule->UnRegisterTrackEditor( WidgetMaterialTrackEditorCreateTrackEditorHandle );
@@ -240,6 +245,7 @@ private:
 	TSharedPtr<FExtensibilityManager> MenuExtensibilityManager;
 	TSharedPtr<FExtensibilityManager> ToolBarExtensibilityManager;
 
+	FDelegateHandle SequenceEditorHandle;
 	FDelegateHandle MarginTrackEditorCreateTrackEditorHandle;
 	FDelegateHandle TransformTrackEditorCreateTrackEditorHandle;
 	FDelegateHandle WidgetMaterialTrackEditorCreateTrackEditorHandle;
diff --git a/Engine/Source/Runtime/LevelSequence/Private/LevelSequence.cpp b/Engine/Source/Runtime/LevelSequence/Private/LevelSequence.cpp
--- a/Engine/Source/Runtime/LevelSequence/Private/LevelSequence.cpp
+++ b/Engine/Source/Runtime/LevelSequence/Private/LevelSequence.cpp
@@ -5,7 +5,7 @@
 #include "HAL/IConsoleManager.h"
 #include "Components/ActorComponent.h"
 #include "GameFramework/Actor.h"
-#include "Engine/BlueprintGeneratedClass.h"
+#include "LevelSequenceDirectorGeneratedClass.h"
 #include "Engine/Engine.h"
 #include "MovieScene.h"
 #include "UObject/Package.h"
@@ -15,6 +15,7 @@
 #include "MovieSceneSpawnableAnnotation.h"
 #include "Tracks/MovieSceneSpawnTrack.h"
 #include "Modules/ModuleManager.h"
+#include "LevelSequencePlayer.h"
 #include "Evaluation/MovieSceneEvaluationTemplateInstance.h"
 
 #if WITH_EDITOR
@@ -94,6 +95,18 @@ bool ULevelSequence::CanAnimateObject(UObject& InObject) const
 
 #if WITH_EDITOR
 
+void ULevelSequence::GetAssetRegistryTags(TArray<FAssetRegistryTag>& OutTags) const
+{
+#if WITH_EDITORONLY_DATA
+	if (DirectorBlueprint)
+	{
+		DirectorBlueprint->GetAssetRegistryTags(OutTags);
+	}
+#endif
+
+	Super::GetAssetRegistryTags(OutTags);
+}
+
 void PurgeLegacyBlueprints(UObject* InObject, UPackage* Package)
 {
 	if (UBlueprint* BP = Cast<UBlueprint>(InObject))
@@ -456,4 +469,27 @@ FGuid ULevelSequence::CreateSpawnable(UObject* ObjectToSpawn)
 	return FGuid();
 }
 
-#endif // WITH_EDITOR
\ No newline at end of file
+#endif // WITH_EDITOR
+
+UObject* ULevelSequence::CreateDirectorInstance(IMovieScenePlayer& Player)
+{
+	ULevelSequencePlayer* LevelSequencePlayer = Cast<ULevelSequencePlayer>(Player.AsUObject());
+	UObject*              DirectorOuter       = LevelSequencePlayer ? LevelSequencePlayer : Player.GetPlaybackContext();
+
+	if (DirectorClass && DirectorOuter)
+	{
+		FName DirectorName = NAME_None;
+
+#if WITH_EDITOR
+		// Give it a pretty name so it shows up in the debug instances drop down nicely
+		DirectorName = MakeUniqueObjectName(DirectorOuter, DirectorClass, *(GetFName().ToString() + TEXT("_Director")));
+#endif
+
+		ULevelSequenceDirector* NewDirector = NewObject<ULevelSequenceDirector>(DirectorOuter, DirectorClass, DirectorName, RF_Transient);
+		NewDirector->Player = LevelSequencePlayer;
+		NewDirector->OnCreated();
+		return NewDirector;
+	}
+
+	return nullptr;
+}
diff --git a/Engine/Source/Runtime/LevelSequence/Private/LevelSequenceDirectorGeneratedClass.cpp b/Engine/Source/Runtime/LevelSequence/Private/LevelSequenceDirectorGeneratedClass.cpp
new file mode 100644
--- /dev/null
+++ b/Engine/Source/Runtime/LevelSequence/Private/LevelSequenceDirectorGeneratedClass.cpp
@@ -0,0 +1,9 @@
+// Copyright 1998-2018 Epic Games, Inc. All Rights Reserved.
+
+#include "LevelSequenceDirectorGeneratedClass.h"
+#include "LevelSequencePlayer.h"
+
+UWorld* ULevelSequenceDirector::GetWorld() const
+{
+	return GetTypedOuter<UWorld>();
+}
diff --git a/Engine/Source/Runtime/LevelSequence/Public/LevelSequence.h b/Engine/Source/Runtime/LevelSequence/Public/LevelSequence.h
--- a/Engine/Source/Runtime/LevelSequence/Public/LevelSequence.h
+++ b/Engine/Source/Runtime/LevelSequence/Public/LevelSequence.h
@@ -11,6 +11,7 @@
 #include "LevelSequence.generated.h"
 
 class UMovieScene;
+class ULevelSequenceDirectorGeneratedClass;
 
 /**
  * Movie scene animation for Actors.
@@ -49,9 +50,27 @@ public:
 	virtual bool CanRebindPossessable(const FMovieScenePossessable& InPossessable) const override;
 	virtual UObject* MakeSpawnableTemplateFromInstance(UObject& InSourceObject, FName ObjectName) override;
 	virtual bool CanAnimateObject(UObject& InObject) const override;
-
+	virtual UObject* CreateDirectorInstance(IMovieScenePlayer& Player) override;
+#if WITH_EDITOR
+	virtual void GetAssetRegistryTags(TArray<FAssetRegistryTag>& OutTags) const override;
+#endif
 	virtual void PostLoad() override;
 
+#if WITH_EDITORONLY_DATA
+
+	/** A pointer to the director blueprint that generates this sequence's DirectorClass. */
+	UPROPERTY()
+	UObject* DirectorBlueprint;
+
+#endif
+
+	/**
+	 * The class that is used to spawn this level sequence's director instance.
+	 * Director instances are allocated on-demand one per sequence during evaluation and are used by event tracks for triggering events.
+	 */
+	UPROPERTY()
+	ULevelSequenceDirectorGeneratedClass* DirectorClass;
+
 protected:
 
 #if WITH_EDITOR
diff --git a/Engine/Source/Runtime/LevelSequence/Public/LevelSequenceDirectorGeneratedClass.h b/Engine/Source/Runtime/LevelSequence/Public/LevelSequenceDirectorGeneratedClass.h
new file mode 100644
--- /dev/null
+++ b/Engine/Source/Runtime/LevelSequence/Public/LevelSequenceDirectorGeneratedClass.h
@@ -0,0 +1,34 @@
+// Copyright 1998-2018 Epic Games, Inc. All Rights Reserved.
+
+#pragma once
+
+#include "CoreTypes.h"
+#include "Engine/BlueprintGeneratedClass.h"
+
+#include "LevelSequenceDirectorGeneratedClass.generated.h"
+
+class ULevelSequencePlayer;
+
+UCLASS()
+class LEVELSEQUENCE_API ULevelSequenceDirector : public UObject
+{
+public:
+	GENERATED_BODY()
+
+	/** Called when this director is created */
+	UFUNCTION(BlueprintImplementableEvent, Category="Sequencer")
+	void OnCreated();
+
+	virtual UWorld* GetWorld() const override;
+
+	/** Pointer to the player that's playing back this director's sequence */
+	UPROPERTY(BlueprintReadOnly, Category="Cinematics")
+	ULevelSequencePlayer* Player;
+};
+
+UCLASS()
+class LEVELSEQUENCE_API ULevelSequenceDirectorGeneratedClass : public UBlueprintGeneratedClass
+{
+public:
+	GENERATED_BODY()
+};
diff --git a/Engine/Source/Runtime/LevelSequence/Public/LevelSequencePlayer.h b/Engine/Source/Runtime/LevelSequence/Public/LevelSequencePlayer.h
--- a/Engine/Source/Runtime/LevelSequence/Public/LevelSequencePlayer.h
+++ b/Engine/Source/Runtime/LevelSequence/Public/LevelSequencePlayer.h
@@ -133,13 +133,15 @@ public:
 	DEPRECATED(4.15, "Please use GetSequence instead.")
 	ULevelSequence* GetLevelSequence() const { return Cast<ULevelSequence>(Sequence); }
 
+	// IMovieScenePlayer interface
+	virtual UObject* GetPlaybackContext() const override;
+	virtual TArray<UObject*> GetEventContexts() const override;
+
 protected:
 
 	// IMovieScenePlayer interface
 	virtual void UpdateCameraCut(UObject* CameraObject, UObject* UnlockIfCameraObject, bool bJumpCut) override;
 	virtual void NotifyBindingUpdate(const FGuid& InGuid, FMovieSceneSequenceIDRef InSequenceID, TArrayView<TWeakObjectPtr<>> Objects) override;
-	virtual UObject* GetPlaybackContext() const override;
-	virtual TArray<UObject*> GetEventContexts() const override;
 
 	//~ UMovieSceneSequencePlayer interface
 	virtual bool CanPlay() const override;
diff --git a/Engine/Source/Runtime/MovieScene/Private/Evaluation/MovieSceneEvaluationTemplateInstance.cpp b/Engine/Source/Runtime/MovieScene/Private/Evaluation/MovieSceneEvaluationTemplateInstance.cpp
--- a/Engine/Source/Runtime/MovieScene/Private/Evaluation/MovieSceneEvaluationTemplateInstance.cpp
+++ b/Engine/Source/Runtime/MovieScene/Private/Evaluation/MovieSceneEvaluationTemplateInstance.cpp
@@ -50,15 +50,15 @@ private:
 };
 
 FMovieSceneEvaluationTemplateInstance::FMovieSceneEvaluationTemplateInstance()
-	: Sequence(nullptr), Template(nullptr), SubData(nullptr)
+	: Sequence(nullptr), Template(nullptr), SubData(nullptr), DirectorInstance(nullptr)
 {}
 
 FMovieSceneEvaluationTemplateInstance::FMovieSceneEvaluationTemplateInstance(UMovieSceneSequence& InSequence, FMovieSceneEvaluationTemplate* InTemplate)
-	: Sequence(&InSequence), Template(InTemplate), SubData(nullptr)
+	: Sequence(&InSequence), Template(InTemplate), SubData(nullptr), DirectorInstance(nullptr)
 {}
 
 FMovieSceneEvaluationTemplateInstance::FMovieSceneEvaluationTemplateInstance(const FMovieSceneSubSequenceData* InSubData, IMovieSceneSequenceTemplateStore& TemplateStore)
-	: Sequence(nullptr), Template(nullptr), SubData(InSubData)
+	: Sequence(nullptr), Template(nullptr), SubData(InSubData), DirectorInstance(nullptr)
 {
 	Sequence = InSubData ? InSubData->GetSequence() : nullptr;
 	if (Sequence)
@@ -102,12 +102,16 @@ void FMovieSceneRootEvaluationTemplateInstance::Initialize(UMovieSceneSequence&
 		ThisFrameMetaData.Reset();
 		ExecutionTokens = FMovieSceneExecutionTokens();
 
-		TransientInstances.ResetSubInstances();
+		Instances.Reset();
 	}
 
 	RootSequence = &InRootSequence;
 	RootTemplate = &TemplateStore->AccessTemplate(InRootSequence);
-	TransientInstances.RootInstance = FMovieSceneEvaluationTemplateInstance(InRootSequence, RootTemplate);
+
+	Instances.RootID = MovieSceneSequenceID::Root;
+	Instances.RootInstance = FMovieSceneEvaluationTemplateInstance(InRootSequence, RootTemplate);
+
+	Player.State.AssignSequence(MovieSceneSequenceID::Root, InRootSequence, Player);
 }
 
 void FMovieSceneRootEvaluationTemplateInstance::Finish(IMovieScenePlayer& Player)
@@ -117,7 +121,8 @@ void FMovieSceneRootEvaluationTemplateInstance::Finish(IMovieScenePlayer& Player
 
 	CallSetupTearDown(Player);
 
-	TransientInstances.Reset();
+	Instances.RootInstance.DirectorInstance = nullptr;
+	Instances.ResetSubInstances();
 }
 
 void FMovieSceneRootEvaluationTemplateInstance::Evaluate(FMovieSceneContext Context, IMovieScenePlayer& Player, FMovieSceneSequenceID InOverrideRootID)
@@ -127,15 +132,13 @@ void FMovieSceneRootEvaluationTemplateInstance::Evaluate(FMovieSceneContext Cont
 	Swap(ThisFrameMetaData, LastFrameMetaData);
 	ThisFrameMetaData.Reset();
 
-	if (TransientInstances.RootID != InOverrideRootID)
+	if (Instances.RootID != InOverrideRootID)
 	{
 		// Tear everything down if we're evaluating a different root sequence
 		CallSetupTearDown(Player);
 		LastFrameMetaData.Reset();
 	}
 
-	RootOverridePath.Set(InOverrideRootID, GetHierarchy());
-
 	const FMovieSceneEvaluationGroup* GroupToEvaluate = SetupFrame(Player, InOverrideRootID, Context);
 	if (!GroupToEvaluate)
 	{
@@ -162,6 +165,9 @@ void FMovieSceneRootEvaluationTemplateInstance::Evaluate(FMovieSceneContext Cont
 
 const FMovieSceneEvaluationGroup* FMovieSceneRootEvaluationTemplateInstance::SetupFrame(IMovieScenePlayer& Player, FMovieSceneSequenceID InOverrideRootID, FMovieSceneContext Context)
 {
+	Instances.RootID = InOverrideRootID;
+	RootOverridePath.Set(InOverrideRootID, GetHierarchy());
+
 	UMovieSceneSequence* RootSequencePtr = RootSequence.Get();
 	if (!RootSequencePtr)
 	{
@@ -169,56 +175,54 @@ const FMovieSceneEvaluationGroup* FMovieSceneRootEvaluationTemplateInstance::Set
 	}
 
 	// Create the instance for the current root override
-	FMovieSceneEvaluationTemplateInstance RootOverrideInstance;
+	FMovieSceneEvaluationTemplateInstance* RootOverrideInstance = Instances.Find(InOverrideRootID);
 
-	if (InOverrideRootID != MovieSceneSequenceID::Root)
+	if (InOverrideRootID != MovieSceneSequenceID::Root && !RootOverrideInstance)
 	{
 		if (const FMovieSceneSubSequenceData* OverrideSubData = GetHierarchy().FindSubData(InOverrideRootID))
 		{
-			RootOverrideInstance = FMovieSceneEvaluationTemplateInstance(OverrideSubData, *TemplateStore);
+			Instances.Add(InOverrideRootID, FMovieSceneEvaluationTemplateInstance(OverrideSubData, *TemplateStore));
+			RootOverrideInstance = Instances.Find(InOverrideRootID);
 		}
 	}
-	else
-	{
-		RootOverrideInstance = FMovieSceneEvaluationTemplateInstance(*RootSequencePtr, RootTemplate);
-	}
 
-	if (!ensureMsgf(RootOverrideInstance.IsValid(), TEXT("Could not find valid sequence or template for supplied sequence ID.")))
+
+	if (!ensureMsgf(RootOverrideInstance && RootOverrideInstance->IsValid(), TEXT("Could not find valid sequence or template for supplied sequence ID.")))
 	{
 		return nullptr;
 	}
 
 	// Ensure the root is up to date
-	if (RootOverrideInstance.Template->SequenceSignature != RootOverrideInstance.Sequence->GetSignature())
+	if (RootOverrideInstance->Template->SequenceSignature != RootOverrideInstance->Sequence->GetSignature())
 	{
-		FMovieSceneEvaluationTemplateGenerator(*RootOverrideInstance.Sequence, *RootOverrideInstance.Template).Generate();
+		FMovieSceneEvaluationTemplateGenerator(*RootOverrideInstance->Sequence, *RootOverrideInstance->Template).Generate();
 	}
 
 	FMovieSceneSequenceTransform MasterToRootOverrideTransform;
-	if (RootOverrideInstance.SubData)
+	if (RootOverrideInstance->SubData)
 	{
-		MasterToRootOverrideTransform = RootOverrideInstance.SubData->RootToSequenceTransform;
+		MasterToRootOverrideTransform = RootOverrideInstance->SubData->RootToSequenceTransform;
 	}
 
 	FFrameNumber RootOverrideTime = (Context.GetTime() * MasterToRootOverrideTransform).FloorToFrame();
 
 	// First off, attempt to find the evaluation group in the existing evaluation field data from the template
-	int32 TemplateFieldIndex = RootOverrideInstance.Template->EvaluationField.GetSegmentFromTime(RootOverrideTime);
+	int32 TemplateFieldIndex = RootOverrideInstance->Template->EvaluationField.GetSegmentFromTime(RootOverrideTime);
 
 	if (TemplateFieldIndex != INDEX_NONE)
 	{
-		const FMovieSceneEvaluationMetaData& FieldMetaData = RootOverrideInstance.Template->EvaluationField.GetMetaData(TemplateFieldIndex);
+		const FMovieSceneEvaluationMetaData& FieldMetaData = RootOverrideInstance->Template->EvaluationField.GetMetaData(TemplateFieldIndex);
 
 		// Verify that this field entry is still valid (all its cached signatures are still the same)
 		TRange<FFrameNumber> InvalidatedSubSequenceRange = TRange<FFrameNumber>::Empty();
-		if (FieldMetaData.IsDirty(RootOverrideInstance.Template->Hierarchy, *TemplateStore, &InvalidatedSubSequenceRange))
+		if (FieldMetaData.IsDirty(RootOverrideInstance->Template->Hierarchy, *TemplateStore, &InvalidatedSubSequenceRange))
 		{
 			TemplateFieldIndex = INDEX_NONE;
 
 			if (!InvalidatedSubSequenceRange.IsEmpty())
 			{
 				// Invalidate the evaluation field for the root template (it may not exist until we compile below)
-				RootOverrideInstance.Template->EvaluationField.Invalidate(InvalidatedSubSequenceRange);
+				RootOverrideInstance->Template->EvaluationField.Invalidate(InvalidatedSubSequenceRange);
 			}
 		}
 	}
@@ -229,17 +233,17 @@ const FMovieSceneEvaluationGroup* FMovieSceneRootEvaluationTemplateInstance::Set
 		static bool bFullCompile = false;
 		if (bFullCompile)
 		{
-			FMovieSceneCompiler::Compile(*RootOverrideInstance.Sequence, *TemplateStore);
-			TemplateFieldIndex = RootOverrideInstance.Template->EvaluationField.GetSegmentFromTime(RootOverrideTime);
+			FMovieSceneCompiler::Compile(*RootOverrideInstance->Sequence, *TemplateStore);
+			TemplateFieldIndex = RootOverrideInstance->Template->EvaluationField.GetSegmentFromTime(RootOverrideTime);
 		}
 		else
 		{
-			TOptional<FCompiledGroupResult> CompileResult = FMovieSceneCompiler::CompileTime(RootOverrideTime, *RootOverrideInstance.Sequence, *TemplateStore);
+			TOptional<FCompiledGroupResult> CompileResult = FMovieSceneCompiler::CompileTime(RootOverrideTime, *RootOverrideInstance->Sequence, *TemplateStore);
 
 			if (CompileResult.IsSet())
 			{
 				TRange<FFrameNumber> FieldRange = CompileResult->Range;
-				TemplateFieldIndex = RootOverrideInstance.Template->EvaluationField.Insert(
+				TemplateFieldIndex = RootOverrideInstance->Template->EvaluationField.Insert(
 					RootOverrideTime,
 					FieldRange,
 					MoveTemp(CompileResult->Group),
@@ -252,11 +256,11 @@ const FMovieSceneEvaluationGroup* FMovieSceneRootEvaluationTemplateInstance::Set
 	if (TemplateFieldIndex != INDEX_NONE)
 	{
 		// Set meta-data
-		ThisFrameMetaData = RootOverrideInstance.Template->EvaluationField.GetMetaData(TemplateFieldIndex);
+		ThisFrameMetaData = RootOverrideInstance->Template->EvaluationField.GetMetaData(TemplateFieldIndex);
 
-		RecreateInstances(RootOverrideInstance, InOverrideRootID, Player);
+		RecreateInstances(*RootOverrideInstance, InOverrideRootID, Player);
 
-		return &RootOverrideInstance.Template->EvaluationField.GetGroup(TemplateFieldIndex);
+		return &RootOverrideInstance->Template->EvaluationField.GetGroup(TemplateFieldIndex);
 	}
 
 	return nullptr;
@@ -285,7 +289,7 @@ void FMovieSceneRootEvaluationTemplateInstance::EvaluateGroup(const FMovieSceneE
 			// Ensure we're able to find the sequence instance in our root if we've overridden
 			SegmentPtr.SequenceID = RootOverridePath.Remap(SegmentPtr.SequenceID);
 
-			const FMovieSceneEvaluationTemplateInstance& Instance = TransientInstances.GetChecked(SegmentPtr.SequenceID);
+			const FMovieSceneEvaluationTemplateInstance& Instance = Instances.GetChecked(SegmentPtr.SequenceID);
 			const FMovieSceneEvaluationTrack* Track = Instance.Template->FindTrack(SegmentPtr.TrackIdentifier);
 
 			if (Track)
@@ -325,7 +329,7 @@ void FMovieSceneRootEvaluationTemplateInstance::EvaluateGroup(const FMovieSceneE
 			// Ensure we're able to find the sequence instance in our root if we've overridden
 			SegmentPtr.SequenceID = RootOverridePath.Remap(SegmentPtr.SequenceID);
 
-			const FMovieSceneEvaluationTemplateInstance& Instance = TransientInstances.GetChecked(SegmentPtr.SequenceID);
+			const FMovieSceneEvaluationTemplateInstance& Instance = Instances.GetChecked(SegmentPtr.SequenceID);
 			const FMovieSceneEvaluationTrack* Track = Instance.Template->FindTrack(SegmentPtr.TrackIdentifier);
 
 			if (Track)
@@ -372,7 +376,7 @@ void FMovieSceneRootEvaluationTemplateInstance::CallSetupTearDown(IMovieScenePla
 	TArray<FMovieSceneOrderedEvaluationKey> ExpiredEntities;
 	TArray<FMovieSceneOrderedEvaluationKey> NewEntities;
 	ThisFrameMetaData.DiffEntities(LastFrameMetaData, &NewEntities, &ExpiredEntities);
-	
+
 	for (const FMovieSceneOrderedEvaluationKey& OrderedKey : ExpiredEntities)
 	{
 		FMovieSceneEvaluationKey Key = OrderedKey.Key;
@@ -380,7 +384,7 @@ void FMovieSceneRootEvaluationTemplateInstance::CallSetupTearDown(IMovieScenePla
 		// Ensure we're able to find the sequence instance in our root if we've overridden
 		Key.SequenceID = RootOverridePath.Remap(Key.SequenceID);
 
-		const FMovieSceneEvaluationTemplateInstance* Instance = TransientInstances.Find(Key.SequenceID);
+		const FMovieSceneEvaluationTemplateInstance* Instance = Instances.Find(Key.SequenceID);
 		if (Instance && Instance->IsValid())
 		{
 			const FMovieSceneEvaluationTrack* Track = Instance->Template->FindTrack(Key.TrackIdentifier);
@@ -427,7 +431,7 @@ void FMovieSceneRootEvaluationTemplateInstance::CallSetupTearDown(IMovieScenePla
 		// Ensure we're able to find the sequence instance in our root if we've overridden
 		Key.SequenceID = RootOverridePath.Remap(Key.SequenceID);
 
-		const FMovieSceneEvaluationTemplateInstance& Instance = TransientInstances.GetChecked(Key.SequenceID);
+		const FMovieSceneEvaluationTemplateInstance& Instance = Instances.GetChecked(Key.SequenceID);
 		check(Instance.IsValid());
 
 		if (const FMovieSceneEvaluationTrack* Track = Instance.Template->FindTrack(Key.TrackIdentifier))
@@ -459,28 +463,53 @@ void FMovieSceneRootEvaluationTemplateInstance::CallSetupTearDown(IMovieScenePla
 
 bool FMovieSceneRootEvaluationTemplateInstance::IsDirty(TSet<UMovieSceneSequence*>* OutDirtySequences) const
 {
-	if (TransientInstances.RootInstance.IsValid())
+	bool bIsDirty = false;
+
+	// Dirty if our master sequence is no longer valid
+	if (!Instances.RootInstance.IsValid())
+	{
+		bIsDirty = true;
+	}
+	// Dirty if our master sequence signature doesn't match the template
+	else if (Instances.RootInstance.Template->SequenceSignature != Instances.RootInstance.Sequence->GetSignature())
 	{
-		bool bIsDirty = false;
-		if (TransientInstances.RootInstance.Template->SequenceSignature != TransientInstances.RootInstance.Sequence->GetSignature())
+		bIsDirty = true;
+
+		if (OutDirtySequences)
+		{
+			OutDirtySequences->Add(Instances.RootInstance.Sequence);
+		}
+	}
+
+	// Find the sequence we're actually evaluating (only != MovieSceneSequenceID::Root when "Evaluate Sequences in Isolation" is on)
+	const FMovieSceneEvaluationTemplateInstance* RootOverrideInstance = Instances.Find(Instances.RootID);
+
+	// Dirty if the root sequence is not valid
+	if (!RootOverrideInstance || !RootOverrideInstance->IsValid())
+	{
+		bIsDirty = true;
+	}
+	else
+	{
+		// Dirty if our root override template signature doesn't match the sequence
+		if (RootOverrideInstance->Template->SequenceSignature != RootOverrideInstance->Sequence->GetSignature())
 		{
+			bIsDirty = true;
+
 			if (OutDirtySequences)
 			{
-				OutDirtySequences->Add(TransientInstances.RootInstance.Sequence);
+				OutDirtySequences->Add(RootOverrideInstance->Sequence);
 			}
-
-			bIsDirty = true;
 		}
 
-		if (LastFrameMetaData.IsDirty(TransientInstances.RootInstance.Template->Hierarchy, *TemplateStore, nullptr, OutDirtySequences))
+		// Dirty if anything we evaluated last frame is dirty
+		if (LastFrameMetaData.IsDirty(RootOverrideInstance->Template->Hierarchy, *TemplateStore, nullptr, OutDirtySequences))
 		{
 			bIsDirty = true;
 		}
-
-		return bIsDirty;
 	}
 
-	return true;
+	return bIsDirty;
 }
 
 void FMovieSceneRootEvaluationTemplateInstance::CopyActuators(FMovieSceneBlendingAccumulator& Accumulator) const
@@ -512,16 +541,36 @@ FMovieSceneEvaluationTemplate* FMovieSceneRootEvaluationTemplateInstance::FindTe
 	return Sequence ? &TemplateStore->AccessTemplate(*Sequence) : nullptr;
 }
 
-void FMovieSceneRootEvaluationTemplateInstance::RecreateInstances(const FMovieSceneEvaluationTemplateInstance& RootOverrideInstance, FMovieSceneSequenceID InOverrideRootID, IMovieScenePlayer& Player)
+UObject* FMovieSceneRootEvaluationTemplateInstance::GetOrCreateDirectorInstance(FMovieSceneSequenceIDRef SequenceID, IMovieScenePlayer& Player)
 {
-	check(RootOverrideInstance.IsValid());
+	FMovieSceneEvaluationTemplateInstance* Instance = Instances.Find(SequenceID);
 
-	TransientInstances.RootID = InOverrideRootID;
-	TransientInstances.RootInstance = RootOverrideInstance;
-	
-	Player.State.AssignSequence(TransientInstances.RootID, *RootOverrideInstance.Sequence, Player);
+	if (Instance && Instance->Sequence)
+	{
+		if (!Instance->DirectorInstance)
+		{
+			Instance->DirectorInstance = Instance->Sequence->CreateDirectorInstance(Player);
+		}
+
+		return Instance->DirectorInstance;
+	}
+
+	return nullptr;
+}
+
+void FMovieSceneRootEvaluationTemplateInstance::ResetDirectorInstances()
+{
+	Instances.RootInstance.DirectorInstance = nullptr;
 
-	TransientInstances.ResetSubInstances();
+	for (TTuple<FMovieSceneSequenceID, FMovieSceneEvaluationTemplateInstance>& Pair : Instances.SubInstances)
+	{
+		Pair.Value.DirectorInstance = nullptr;
+	}
+}
+
+void FMovieSceneRootEvaluationTemplateInstance::RecreateInstances(const FMovieSceneEvaluationTemplateInstance& RootOverrideInstance, FMovieSceneSequenceID InOverrideRootID, IMovieScenePlayer& Player)
+{
+	check(RootOverrideInstance.IsValid());
 
 	// We recreate all necessary sequence instances for the current and previous frames by diffing the sequences active last frame, with this frame
 	TArray<FMovieSceneSequenceID> PreviousAndCurrentFrameSequenceIDs = LastFrameMetaData.ActiveSequences;
@@ -529,14 +578,24 @@ void FMovieSceneRootEvaluationTemplateInstance::RecreateInstances(const FMovieSc
 
 	for (FMovieSceneSequenceID SequenceID : PreviousAndCurrentFrameSequenceIDs)
 	{
+		if (Instances.Find(SequenceID))
+		{
+			continue;
+		}
+
 		const FMovieSceneSubSequenceData* SubData = RootOverrideInstance.Template->Hierarchy.FindSubData(SequenceID);
 		if (SubData)
 		{
 			FMovieSceneEvaluationTemplateInstance NewInstance(SubData, *TemplateStore);
-			
+
+			if (!NewInstance.DirectorInstance)
+			{
+				NewInstance.DirectorInstance = NewInstance.Sequence->CreateDirectorInstance(Player);
+			}
+
 			// Always add sequence IDs as the root
 			SequenceID = RootOverridePath.Remap(SequenceID);
-			TransientInstances.Add(SequenceID, NewInstance);
+			Instances.Add(SequenceID, NewInstance);
 
 			if (NewInstance.Sequence)
 			{
diff --git a/Engine/Source/Runtime/MovieScene/Public/Evaluation/MovieSceneEvaluationTemplateInstance.h b/Engine/Source/Runtime/MovieScene/Public/Evaluation/MovieSceneEvaluationTemplateInstance.h
--- a/Engine/Source/Runtime/MovieScene/Public/Evaluation/MovieSceneEvaluationTemplateInstance.h
+++ b/Engine/Source/Runtime/MovieScene/Public/Evaluation/MovieSceneEvaluationTemplateInstance.h
@@ -10,6 +10,7 @@
 #include "Evaluation/MovieSceneEvaluationTemplate.h"
 #include "Evaluation/MovieSceneExecutionTokens.h"
 #include "Evaluation/MovieSceneRootOverridePath.h"
+#include "MovieSceneEvaluationTemplateInstance.generated.h"
 
 class UMovieSceneSequence;
 struct FDelayedPreAnimatedStateRestore;
@@ -18,8 +19,11 @@ struct FCompileOnTheFlyData;
 /**
  * An instance of an evaluation template. Fast to initialize and evaluate.
  */
+USTRUCT()
 struct FMovieSceneEvaluationTemplateInstance
 {
+	GENERATED_BODY()
+
 	FMovieSceneEvaluationTemplateInstance();
 
 	FMovieSceneEvaluationTemplateInstance(UMovieSceneSequence& InSequence, FMovieSceneEvaluationTemplate* InTemplate);
@@ -32,6 +36,7 @@ struct FMovieSceneEvaluationTemplateInstance
 	}
 
 	/** The sequence this applies to */
+	UPROPERTY()
 	UMovieSceneSequence* Sequence;
 
 	/** Pointer to the evaluation template we're evaluating */
@@ -39,10 +44,17 @@ struct FMovieSceneEvaluationTemplateInstance
 
 	/** Sub Data or nullptr */
 	const FMovieSceneSubSequenceData* SubData;
+
+	/** The sequence this applies to */
+	UPROPERTY()
+	UObject* DirectorInstance;
 };
 
+USTRUCT()
 struct FMovieSceneEvaluationTemplateInstanceContainer
 {
+	GENERATED_BODY()
+
 	FMovieSceneEvaluationTemplateInstanceContainer()
 		: RootID(MovieSceneSequenceID::Invalid)
 	{}
@@ -67,26 +79,38 @@ struct FMovieSceneEvaluationTemplateInstanceContainer
 
 	const FMovieSceneEvaluationTemplateInstance& GetChecked(FMovieSceneSequenceID InID) const
 	{
-		return InID == RootID ? RootInstance : SubInstances.FindChecked(InID);
+		return InID == MovieSceneSequenceID::Root ? RootInstance : SubInstances.FindChecked(InID);
 	}
 
 	FMovieSceneEvaluationTemplateInstance* Find(FMovieSceneSequenceID InID)
 	{
-		return InID == RootID ? &RootInstance : SubInstances.Find(InID);
+		return InID == MovieSceneSequenceID::Root ? &RootInstance : SubInstances.Find(InID);
+	}
+
+	const FMovieSceneEvaluationTemplateInstance* Find(FMovieSceneSequenceID InID) const
+	{
+		return InID == MovieSceneSequenceID::Root ? &RootInstance : SubInstances.Find(InID);
 	}
 
 	/** Sequence ID that is used to evaluate from for the current frame */
 	FMovieSceneSequenceID RootID;
+	
+	UPROPERTY()
 	FMovieSceneEvaluationTemplateInstance RootInstance;
+
+	UPROPERTY()
 	TMap<FMovieSceneSequenceID, FMovieSceneEvaluationTemplateInstance> SubInstances;
 };
 
 /**
  * Root evaluation template instance used to play back any sequence
  */
+USTRUCT()
 struct FMovieSceneRootEvaluationTemplateInstance
 {
 public:
+	GENERATED_BODY()
+
 	MOVIESCENE_API FMovieSceneRootEvaluationTemplateInstance();
 	MOVIESCENE_API ~FMovieSceneRootEvaluationTemplateInstance();
 
@@ -162,6 +186,19 @@ public:
 	 */
 	MOVIESCENE_API FMovieSceneEvaluationTemplate* FindTemplate(FMovieSceneSequenceIDRef SequenceID);
 
+	/**
+	 * Locate a director instance object for the specified sequence ID, creating one if necessary
+	 *
+	 * @param SequenceID 			ID of the sequence template to locate a director instance for
+	 * @return A director instance as defined by the specific sequence, or nullptr if one could not be found or created
+	 */
+	MOVIESCENE_API UObject* GetOrCreateDirectorInstance(FMovieSceneSequenceIDRef SequenceID, IMovieScenePlayer& Player);
+
+	/**
+	 * Resets all the director instances currently stored by this template instance
+	 */
+	MOVIESCENE_API void ResetDirectorInstances();
+
 	/**
 	 * Access the master sequence's hierarchy data
 	 */
@@ -195,6 +232,14 @@ public:
 	 */
 	MOVIESCENE_API void CopyActuators(FMovieSceneBlendingAccumulator& Accumulator) const;
 
+	/**
+	 * Const access to all the instances contained within this template instance
+	 */
+	const FMovieSceneEvaluationTemplateInstanceContainer& GetInstances() const
+	{
+		return Instances;
+	}
+
 private:
 
 	/**
@@ -224,7 +269,8 @@ private:
 
 	FMovieSceneEvaluationTemplate* RootTemplate;
 
-	FMovieSceneEvaluationTemplateInstanceContainer TransientInstances;
+	UPROPERTY()
+	FMovieSceneEvaluationTemplateInstanceContainer Instances;
 
 	/** Cache of everything that was evaluated last frame */
 	FMovieSceneEvaluationMetaData LastFrameMetaData;
@@ -240,3 +286,12 @@ private:
 	/** Execution tokens that are used to apply animated state */
 	FMovieSceneExecutionTokens ExecutionTokens;
 };
+
+template<>
+struct TStructOpsTypeTraits<FMovieSceneRootEvaluationTemplateInstance> : public TStructOpsTypeTraitsBase2<FMovieSceneRootEvaluationTemplateInstance>
+{
+	enum
+	{
+		WithCopy = false
+	};
+};
\ No newline at end of file
diff --git a/Engine/Source/Runtime/MovieScene/Public/IMovieScenePlayer.h b/Engine/Source/Runtime/MovieScene/Public/IMovieScenePlayer.h
--- a/Engine/Source/Runtime/MovieScene/Public/IMovieScenePlayer.h
+++ b/Engine/Source/Runtime/MovieScene/Public/IMovieScenePlayer.h
@@ -59,6 +59,11 @@ public:
 	 */
 	virtual FMovieSceneRootEvaluationTemplateInstance& GetEvaluationTemplate() = 0;
 
+	/**
+	 * Cast this player instance as a UObject if possible
+	 */
+	virtual UObject* AsUObject() { return nullptr; }
+
 	/**
 	 * Whether this player can update the camera cut
 	 */
diff --git a/Engine/Source/Runtime/MovieScene/Public/MovieSceneSequence.h b/Engine/Source/Runtime/MovieScene/Public/MovieSceneSequence.h
--- a/Engine/Source/Runtime/MovieScene/Public/MovieSceneSequence.h
+++ b/Engine/Source/Runtime/MovieScene/Public/MovieSceneSequence.h
@@ -157,6 +157,11 @@ public:
 	 */
 	virtual FGuid CreateSpawnable(UObject* ObjectToSpawn) { return FGuid(); }
 
+	/**
+	 * Called to retrieve or construct a director instance to be used for the specified player
+	 */
+	virtual UObject* CreateDirectorInstance(IMovieScenePlayer& Player) { return nullptr; }
+
 public:
 
 	MOVIESCENE_API virtual void PostLoad() override;
diff --git a/Engine/Source/Runtime/MovieScene/Public/MovieSceneSequencePlayer.h b/Engine/Source/Runtime/MovieScene/Public/MovieSceneSequencePlayer.h
--- a/Engine/Source/Runtime/MovieScene/Public/MovieSceneSequencePlayer.h
+++ b/Engine/Source/Runtime/MovieScene/Public/MovieSceneSequencePlayer.h
@@ -444,6 +444,7 @@ protected:
 	virtual FMovieSceneRootEvaluationTemplateInstance& GetEvaluationTemplate() override { return RootTemplateInstance; }
 	virtual EMovieScenePlayerStatus::Type GetPlaybackStatus() const override;
 	virtual FMovieSceneSpawnRegister& GetSpawnRegister() override;
+	virtual UObject* AsUObject() override { return this; }
 
 	virtual void SetPlaybackStatus(EMovieScenePlayerStatus::Type InPlaybackStatus) override {}
 	virtual void SetViewportSettings(const TMap<FViewportClient*, EMovieSceneViewportParams>& ViewportParamsMap) override {}
@@ -526,6 +527,7 @@ protected:
 	FMovieSceneSequencePlaybackSettings PlaybackSettings;
 
 	/** The root template instance we're evaluating */
+	UPROPERTY(transient)
 	FMovieSceneRootEvaluationTemplateInstance RootTemplateInstance;
 
 	/** Play position helper */
diff --git a/Engine/Source/Runtime/MovieSceneTracks/MovieSceneTracks.Build.cs b/Engine/Source/Runtime/MovieSceneTracks/MovieSceneTracks.Build.cs
--- a/Engine/Source/Runtime/MovieSceneTracks/MovieSceneTracks.Build.cs
+++ b/Engine/Source/Runtime/MovieSceneTracks/MovieSceneTracks.Build.cs
@@ -30,5 +30,14 @@ public class MovieSceneTracks : ModuleRules
                 "AnimGraphRuntime"
 			}
 		);
+
+		if (Target.bBuildWithEditorOnlyData)
+		{
+			PublicDependencyModuleNames.AddRange(
+				new string[] {
+					"BlueprintGraph"
+				}
+			);
+		}
 	}
 }
diff --git a/Engine/Source/Runtime/MovieSceneTracks/Private/Channels/MovieSceneEvent.cpp b/Engine/Source/Runtime/MovieSceneTracks/Private/Channels/MovieSceneEvent.cpp
new file mode 100644
--- /dev/null
+++ b/Engine/Source/Runtime/MovieSceneTracks/Private/Channels/MovieSceneEvent.cpp
@@ -0,0 +1,96 @@
+// Copyright 1998-2018 Epic Games, Inc. All Rights Reserved.
+
+#include "Channels/MovieSceneEvent.h"
+
+#if WITH_EDITORONLY_DATA
+
+#include "EdGraph/EdGraph.h"
+#include "K2Node_FunctionEntry.h"
+#include "EdGraphSchema_K2.h"
+#include "UObject/PropertyPortFlags.h"
+
+UK2Node_FunctionEntry* FMovieSceneEvent::GetFunctionEntry() const
+{
+	return CastChecked<UK2Node_FunctionEntry>(FunctionEntry.Get(), ECastCheckedType::NullAllowed);
+}
+
+void FMovieSceneEvent::SetFunctionEntry(UK2Node_FunctionEntry* InFunctionEntry)
+{
+	FunctionEntry = InFunctionEntry;
+	CacheFunctionName();
+}
+
+bool FMovieSceneEvent::IsBoundToBlueprint() const
+{
+	return IsValidFunction(GetFunctionEntry());
+}
+
+bool FMovieSceneEvent::IsValidFunction(UK2Node_FunctionEntry* Function)
+{
+	// Must have a single non-reference parameter on the function
+	if (!Function || Function->IsPendingKill() || Function->GetGraph()->IsPendingKill())
+	{
+		return false;
+	}
+	else if (Function->UserDefinedPins.Num() == 0)
+	{
+		return true;
+	}
+	else if (Function->UserDefinedPins.Num() != 1 || Function->UserDefinedPins[0]->PinType.bIsReference)
+	{
+		return false;
+	}
+
+	// Pin must be an object or interface property
+	FName PinCategory = Function->UserDefinedPins[0]->PinType.PinCategory;
+	return PinCategory == UEdGraphSchema_K2::PC_Object || PinCategory == UEdGraphSchema_K2::PC_Interface;
+}
+
+void FMovieSceneEvent::CacheFunctionName()
+{
+	UK2Node_FunctionEntry* Node = GetFunctionEntry();
+
+	FunctionName = NAME_None;
+
+	if (IsValidFunction(Node))
+	{
+		UEdGraph* Graph = Node->GetGraph();
+		if (Graph)
+		{
+			FunctionName = Graph->GetFName();
+		}
+	}
+}
+
+#endif // WITH_EDITORONLY_DATA
+
+
+void FMovieSceneEvent::PostSerialize(const FArchive& Ar)
+{
+#if WITH_EDITORONLY_DATA
+	if (Ar.IsLoading() && !Ar.HasAnyPortFlags(PPF_Duplicate | PPF_DuplicateForPIE))
+	{
+		CacheFunctionName();
+	}
+#endif
+}
+
+bool FMovieSceneEvent::IsValidFunction(UFunction* Function)
+{
+	// Must have a single non-reference parameter on the function
+	if (!Function)
+	{
+		return false;
+	}
+	else if (Function->NumParms == 0)
+	{
+		return true;
+	}
+	else if (Function->NumParms != 1 || !Function->PropertyLink || (Function->PropertyLink->GetPropertyFlags() & CPF_ReferenceParm))
+	{
+		return false;
+	}
+
+	// Parameter must be an object or interface property
+	return Function->PropertyLink->IsA<UObjectProperty>() || Function->PropertyLink->IsA<UInterfaceProperty>();
+}
diff --git a/Engine/Source/Runtime/MovieSceneTracks/Private/Channels/MovieSceneEventChannel.cpp b/Engine/Source/Runtime/MovieSceneTracks/Private/Channels/MovieSceneEventChannel.cpp
new file mode 100644
--- /dev/null
+++ b/Engine/Source/Runtime/MovieSceneTracks/Private/Channels/MovieSceneEventChannel.cpp
@@ -0,0 +1,55 @@
+// Copyright 1998-2018 Epic Games, Inc. All Rights Reserved.
+
+#include "Channels/MovieSceneEventChannel.h"
+
+void FMovieSceneEventChannel::GetKeys(const TRange<FFrameNumber>& WithinRange, TArray<FFrameNumber>* OutKeyTimes, TArray<FKeyHandle>* OutKeyHandles)
+{
+	GetData().GetKeys(WithinRange, OutKeyTimes, OutKeyHandles);
+}
+
+void FMovieSceneEventChannel::GetKeyTimes(TArrayView<const FKeyHandle> InHandles, TArrayView<FFrameNumber> OutKeyTimes)
+{
+	GetData().GetKeyTimes(InHandles, OutKeyTimes);
+}
+
+void FMovieSceneEventChannel::SetKeyTimes(TArrayView<const FKeyHandle> InHandles, TArrayView<const FFrameNumber> InKeyTimes)
+{
+	GetData().SetKeyTimes(InHandles, InKeyTimes);
+}
+
+void FMovieSceneEventChannel::DuplicateKeys(TArrayView<const FKeyHandle> InHandles, TArrayView<FKeyHandle> OutNewHandles)
+{
+	GetData().DuplicateKeys(InHandles, OutNewHandles);
+}
+
+void FMovieSceneEventChannel::DeleteKeys(TArrayView<const FKeyHandle> InHandles)
+{
+	GetData().DeleteKeys(InHandles);
+}
+
+void FMovieSceneEventChannel::ChangeFrameResolution(FFrameRate SourceRate, FFrameRate DestinationRate)
+{
+	GetData().ChangeFrameResolution(SourceRate, DestinationRate);
+}
+
+TRange<FFrameNumber> FMovieSceneEventChannel::ComputeEffectiveRange() const
+{
+	return GetData().GetTotalRange();
+}
+
+int32 FMovieSceneEventChannel::GetNumKeys() const
+{
+	return KeyTimes.Num();
+}
+
+void FMovieSceneEventChannel::Reset()
+{
+	KeyTimes.Reset();
+	KeyValues.Reset();
+	KeyHandles.Reset();
+}
+
+void FMovieSceneEventChannel::Offset(FFrameNumber DeltaPosition)
+{
+	GetData().Offset(DeltaPosition);
+}
diff --git a/Engine/Source/Runtime/MovieSceneTracks/Private/Evaluation/MovieSceneEventTemplate.cpp b/Engine/Source/Runtime/MovieSceneTracks/Private/Evaluation/MovieSceneEventTemplate.cpp
--- a/Engine/Source/Runtime/MovieSceneTracks/Private/Evaluation/MovieSceneEventTemplate.cpp
+++ b/Engine/Source/Runtime/MovieSceneTracks/Private/Evaluation/MovieSceneEventTemplate.cpp
@@ -8,6 +8,7 @@
 #include "MovieScene.h"
 #include "Evaluation/MovieSceneEvaluation.h"
 #include "IMovieScenePlayer.h"
+#include "Algo/Accumulate.h"
 
 
 DECLARE_CYCLE_STAT(TEXT("Event Track Token Execute"), MovieSceneEval_EventTrack_TokenExecute, STATGROUP_MovieSceneEval);
@@ -159,14 +160,190 @@ struct FEventTrackExecutionToken
 	TArray<FMovieSceneObjectBindingID, TInlineAllocator<2>> EventReceivers;
 };
 
-FMovieSceneEventSectionTemplate::FMovieSceneEventSectionTemplate(const UMovieSceneEventSection& Section, const UMovieSceneEventTrack& Track)
-	: EventData(Section.GetEventData())
-	, EventReceivers(Track.EventReceivers)
+struct FEventTriggerExecutionToken
+	: IMovieSceneExecutionToken
+{
+	FEventTriggerExecutionToken(TArray<FName> InEvents, const TArray<FMovieSceneObjectBindingID>& InEventReceivers)
+		: Events(MoveTemp(InEvents)), EventReceivers(InEventReceivers)
+	{}
+
+	/** Execute this token, operating on all objects referenced by 'Operand' */
+	virtual void Execute(const FMovieSceneContext& Context, const FMovieSceneEvaluationOperand& Operand, FPersistentEvaluationData& PersistentData, IMovieScenePlayer& Player) override
+	{
+		MOVIESCENE_DETAILED_SCOPE_CYCLE_COUNTER(MovieSceneEval_EventTrack_TokenExecute)
+
+		UObject* DirectorInstance = Player.GetEvaluationTemplate().GetOrCreateDirectorInstance(Operand.SequenceID, Player);
+		if (!DirectorInstance)
+		{
+			UE_LOG(LogMovieScene, Warning, TEXT("Failed to trigger the following events because no director instance was available: %s."), *GenerateEventListString());
+			return;
+		}
+
+		// Resolve event contexts to trigger the event on
+		TArray<UObject*> EventContexts;
+
+		// If the event track resides within an object binding, add those to the event contexts
+		if (Operand.ObjectBindingID.IsValid())
+		{
+			for (TWeakObjectPtr<> WeakEventContext : Player.FindBoundObjects(Operand))
+			{
+				if (UObject* EventContext = WeakEventContext.Get())
+				{
+					EventContexts.Add(EventContext);
+				}
+			}
+		}
+		// If we have specified event receivers
+		else if (EventReceivers.Num())
+		{
+			EventContexts.Reserve(EventReceivers.Num());
+			for (FMovieSceneObjectBindingID ID : EventReceivers)
+			{
+				// Ensure that this ID is resolvable from the root, based on the current local sequence ID
+				ID = ID.ResolveLocalToRoot(Operand.SequenceID, Player.GetEvaluationTemplate().GetHierarchy());
+
+				// Lookup the object(s) specified by ID in the player
+				for (TWeakObjectPtr<> WeakEventContext : Player.FindBoundObjects(ID.GetGuid(), ID.GetSequenceID()))
+				{
+					if (UObject* EventContext = WeakEventContext.Get())
+					{
+						EventContexts.Add(EventContext);
+					}
+				}
+			}
+		}
+		// If we haven't specified event receivers, use the default set defined on the player
+		else
+		{
+			EventContexts = Player.GetEventContexts();
+		}
+
+#if WITH_EDITOR
+		const static FName NAME_CallInEditor(TEXT("CallInEditor"));
+
+		UWorld* World = DirectorInstance->GetWorld();
+		bool bIsGameWorld = World && World->IsGameWorld();
+#endif // WITH_EDITOR
+
+		for (FName EventName : Events)
+		{
+			if (EventName == NAME_None)
+			{
+				continue;
+			}
+
+			UFunction* Function = DirectorInstance->FindFunction(EventName);
+			// Event must have only a single object parameter, and the director instance must be an implementation of the function's class
+			if (!Function)
+			{
+				UE_LOG(LogMovieScene, Warning, TEXT("Failed to trigger the event '%s' because the function does not exist on director instance '%s'."), *EventName.ToString(), *DirectorInstance->GetName());
+				continue;
+			}
+#if WITH_EDITOR
+			else if (!bIsGameWorld && !Function->HasMetaData(NAME_CallInEditor))
+			{
+				UE_LOG(LogMovieScene, Log, TEXT("Refusing to trigger event '%s' in editor world when 'Call in Editor' is false."), *EventName.ToString());
+				continue;
+			}
+#endif // WITH_EDITOR
+
+			if (Function->NumParms == 0)
+			{
+				UE_LOG(LogMovieScene, VeryVerbose, TEXT("Triggering event '%s'."), *EventName.ToString());
+				DirectorInstance->ProcessEvent(Function, nullptr);
+			}
+			else if (Function->NumParms == 1 && Function->PropertyLink && (Function->PropertyLink->GetPropertyFlags() & CPF_ReferenceParm) == 0)
+			{
+				for (UObject* EventContextObject : EventContexts)
+				{
+					TriggerEvent(DirectorInstance, Function, EventContextObject);
+				}
+			}
+			else
+			{
+				UE_LOG(LogMovieScene, Warning, TEXT("Failed to trigger event '%s' because its signature is not compatible. Function signatures must have either 0 or 1 (non-ref) parameters."), *EventName.ToString());
+			}
+		}
+	}
+
+	void TriggerEvent(UObject* DirectorInstance, UFunction* Function, UObject* ObjectParamValue)
+	{
+		// We know by now that the parameter type is compatible because it wouldn't be added to the array if not
+		if (UObjectProperty* ObjectParameter = Cast<UObjectProperty>(Function->PropertyLink))
+		{
+			if (!ObjectParameter->PropertyClass || ObjectParamValue->IsA(ObjectParameter->PropertyClass))
+			{
+				DirectorInstance->ProcessEvent(Function, &ObjectParamValue);
+				return;
+			}
+
+
+			UE_LOG(LogMovieScene, VeryVerbose, TEXT("Failed to trigger event '%s' with object '%s' because it is not the correct type. Function expects a '%s' but target object is a '%s'."),
+				*Function->GetName(),
+				*ObjectParamValue->GetName(),
+				*ObjectParameter->PropertyClass->GetName(),
+				*ObjectParamValue->GetClass()->GetName()
+				);
+
+			return;
+		}
+
+		if (UInterfaceProperty* InterfaceParameter = Cast<UInterfaceProperty>(Function->PropertyLink))
+		{
+			if (ObjectParamValue->GetClass()->ImplementsInterface(InterfaceParameter->InterfaceClass))
+			{
+				DirectorInstance->ProcessEvent(Function, &ObjectParamValue);
+				return;
+			}
+
+
+			UE_LOG(LogMovieScene, VeryVerbose, TEXT("Failed to trigger event '%s' with object '%s' because it does not implement the necessary interface. Function expects a '%s'."),
+				*Function->GetName(),
+				*ObjectParamValue->GetName(),
+				*InterfaceParameter->InterfaceClass->GetName()
+				);
+
+			return;
+		}
+
+
+		UE_LOG(LogMovieScene, Warning, TEXT("Failed to trigger event '%s' because its signature is not compatible. Function expects a '%s' parameter, but only object and interface parameters are supported."),
+			*Function->GetName(),
+			*Function->PropertyLink->GetClass()->GetName()
+			);
+	}
+
+#if !NO_LOGGING
+	FString GenerateEventListString() const
+	{
+		return Algo::Accumulate(Events, FString(), [](FString&& InString, FName Event){
+			if (InString.Len() > 0)
+			{
+				InString += TEXT(", ");
+			}
+			return InString + Event.ToString();
+		});
+	}
+#endif
+
+	TArray<FName> Events;
+	TArray<FMovieSceneObjectBindingID, TInlineAllocator<2>> EventReceivers;
+};
+
+
+FMovieSceneEventTemplateBase::FMovieSceneEventTemplateBase(const UMovieSceneEventTrack& Track)
+	: EventReceivers(Track.EventReceivers)
 	, bFireEventsWhenForwards(Track.bFireEventsWhenForwards)
 	, bFireEventsWhenBackwards(Track.bFireEventsWhenBackwards)
 {
 }
 
+FMovieSceneEventSectionTemplate::FMovieSceneEventSectionTemplate(const UMovieSceneEventSection& Section, const UMovieSceneEventTrack& Track)
+	: FMovieSceneEventTemplateBase(Track)
+	, EventData(Section.GetEventData())
+{
+}
+
 void FMovieSceneEventSectionTemplate::EvaluateSwept(const FMovieSceneEvaluationOperand& Operand, const FMovieSceneContext& Context, const FPersistentEvaluationData& PersistentData, FMovieSceneExecutionTokens& ExecutionTokens) const
 {
 	// Don't allow events to fire when playback is in a stopped state. This can occur when stopping 
@@ -226,3 +403,105 @@ void FMovieSceneEventSectionTemplate::EvaluateSwept(const FMovieSceneEvaluationO
 		ExecutionTokens.Add(FEventTrackExecutionToken(MoveTemp(Events), EventReceivers));
 	}
 }
+
+
+
+FMovieSceneEventTriggerTemplate::FMovieSceneEventTriggerTemplate(const UMovieSceneEventTriggerSection& Section, const UMovieSceneEventTrack& Track)
+	: FMovieSceneEventTemplateBase(Track)
+{
+	TMovieSceneChannelData<const FMovieSceneEvent> EventData = Section.EventChannel.GetData();
+	TArrayView<const FFrameNumber>     Times  = EventData.GetTimes();
+	TArrayView<const FMovieSceneEvent> Events = EventData.GetValues();
+
+	EventTimes.Reserve(Times.Num());
+	EventFunctions.Reserve(Times.Num());
+
+	for (int32 Index = 0; Index < Times.Num(); ++Index)
+	{
+		EventTimes.Add(Times[Index]);
+		EventFunctions.Add(Events[Index].FunctionName);
+	}
+}
+
+void FMovieSceneEventTriggerTemplate::EvaluateSwept(const FMovieSceneEvaluationOperand& Operand, const FMovieSceneContext& Context, const FPersistentEvaluationData& PersistentData, FMovieSceneExecutionTokens& ExecutionTokens) const
+{
+	// Don't allow events to fire when playback is in a stopped state. This can occur when stopping 
+	// playback and returning the current position to the start of playback. It's not desireable to have 
+	// all the events from the last playback position to the start of playback be fired.
+	if (Context.GetStatus() == EMovieScenePlayerStatus::Stopped || Context.IsSilent())
+	{
+		return;
+	}
+
+	const bool bBackwards = Context.GetDirection() == EPlayDirection::Backwards;
+
+	if ((!bBackwards && !bFireEventsWhenForwards) ||
+		(bBackwards && !bFireEventsWhenBackwards))
+	{
+		return;
+	}
+
+	TArray<FName> Events;
+
+	TRange<FFrameNumber> SweptRange = Context.GetFrameNumberRange();
+
+	const int32 First = bBackwards ? EventTimes.Num() - 1 : 0;
+	const int32 Last = bBackwards ? 0 : EventTimes.Num() - 1;
+	const int32 Inc = bBackwards ? -1 : 1;
+
+	const float PositionInSeconds = Context.GetTime() * Context.GetRootToSequenceTransform().Inverse() / Context.GetFrameRate();
+
+	if (bBackwards)
+	{
+		// Trigger events backwards
+		for (int32 KeyIndex = EventTimes.Num() - 1; KeyIndex >= 0; --KeyIndex)
+		{
+			FFrameNumber Time = EventTimes[KeyIndex];
+			if (SweptRange.Contains(Time))
+			{
+				Events.Add(EventFunctions[KeyIndex]);
+			}
+		}
+	}
+	// Trigger events forwards
+	else for (int32 KeyIndex = 0; KeyIndex < EventTimes.Num(); ++KeyIndex)
+	{
+		FFrameNumber Time = EventTimes[KeyIndex];
+		if (SweptRange.Contains(Time))
+		{
+			Events.Add(EventFunctions[KeyIndex]);
+		}
+	}
+
+
+	if (Events.Num())
+	{
+		ExecutionTokens.Add(FEventTriggerExecutionToken(MoveTemp(Events), EventReceivers));
+	}
+}
+
+
+
+FMovieSceneEventRepeaterTemplate::FMovieSceneEventRepeaterTemplate(const UMovieSceneEventRepeaterSection& Section, const UMovieSceneEventTrack& Track)
+	: FMovieSceneEventTemplateBase(Track)
+	, EventToTrigger(Section.Event.FunctionName)
+{
+}
+
+void FMovieSceneEventRepeaterTemplate::EvaluateSwept(const FMovieSceneEvaluationOperand& Operand, const FMovieSceneContext& Context, const FPersistentEvaluationData& PersistentData, FMovieSceneExecutionTokens& ExecutionTokens) const
+{
+	// Don't allow events to fire when playback is in a stopped state. This can occur when stopping 
+	// playback and returning the current position to the start of playback. It's not desireable to have 
+	// all the events from the last playback position to the start of playback be fired.
+	if (Context.GetStatus() == EMovieScenePlayerStatus::Stopped || Context.IsSilent())
+	{
+		return;
+	}
+
+	const bool bBackwards = Context.GetDirection() == EPlayDirection::Backwards;
+	if ((!bBackwards && bFireEventsWhenForwards) || (bBackwards && bFireEventsWhenBackwards))
+	{
+		TArray<FName> Events = { EventToTrigger };
+		ExecutionTokens.Add(FEventTriggerExecutionToken(MoveTemp(Events), EventReceivers));
+	}
+}
\ No newline at end of file
diff --git a/Engine/Source/Runtime/MovieSceneTracks/Private/Evaluation/MovieSceneEventTemplate.h b/Engine/Source/Runtime/MovieSceneTracks/Private/Evaluation/MovieSceneEventTemplate.h
--- a/Engine/Source/Runtime/MovieSceneTracks/Private/Evaluation/MovieSceneEventTemplate.h
+++ b/Engine/Source/Runtime/MovieSceneTracks/Private/Evaluation/MovieSceneEventTemplate.h
@@ -5,6 +5,8 @@
 #include "CoreMinimal.h"
 #include "UObject/ObjectMacros.h"
 #include "Sections/MovieSceneEventSection.h"
+#include "Sections/MovieSceneEventRepeaterSection.h"
+#include "Sections/MovieSceneEventTriggerSection.h"
 #include "Evaluation/MovieSceneEvalTemplate.h"
 #include "MovieSceneObjectBindingID.h"
 
@@ -14,24 +16,75 @@ class UMovieSceneEventTrack;
 struct EventData;
 
 USTRUCT()
-struct FMovieSceneEventSectionTemplate : public FMovieSceneEvalTemplate
+struct FMovieSceneEventTemplateBase : public FMovieSceneEvalTemplate
 {
 	GENERATED_BODY()
 	
-	FMovieSceneEventSectionTemplate() : bFireEventsWhenForwards(false), bFireEventsWhenBackwards(false) {}
+	FMovieSceneEventTemplateBase() : bFireEventsWhenForwards(false), bFireEventsWhenBackwards(false) {}
+	FMovieSceneEventTemplateBase(const UMovieSceneEventTrack& Track);
+
+	UPROPERTY()
+	TArray<FMovieSceneObjectBindingID> EventReceivers;
+
+	UPROPERTY()
+	uint32 bFireEventsWhenForwards : 1;
+
+	UPROPERTY()
+	uint32 bFireEventsWhenBackwards : 1;
+
+private:
+
+	virtual UScriptStruct& GetScriptStructImpl() const override { return *StaticStruct(); }
+};
+
+USTRUCT()
+struct FMovieSceneEventSectionTemplate : public FMovieSceneEventTemplateBase
+{
+	GENERATED_BODY()
+	
+	FMovieSceneEventSectionTemplate() {}
 	FMovieSceneEventSectionTemplate(const UMovieSceneEventSection& Section, const UMovieSceneEventTrack& Track);
 
 	UPROPERTY()
 	FMovieSceneEventSectionData EventData;
 
+private:
+
+	virtual UScriptStruct& GetScriptStructImpl() const override { return *StaticStruct(); }
+	virtual void EvaluateSwept(const FMovieSceneEvaluationOperand& Operand, const FMovieSceneContext& Context, const FPersistentEvaluationData& PersistentData, FMovieSceneExecutionTokens& ExecutionTokens) const override;
+};
+
+
+USTRUCT()
+struct FMovieSceneEventTriggerTemplate : public FMovieSceneEventTemplateBase
+{
+	GENERATED_BODY()
+	
+	FMovieSceneEventTriggerTemplate() {}
+	FMovieSceneEventTriggerTemplate(const UMovieSceneEventTriggerSection& Section, const UMovieSceneEventTrack& Track);
+
 	UPROPERTY()
-	TArray<FMovieSceneObjectBindingID> EventReceivers;
+	TArray<FFrameNumber> EventTimes;
 
 	UPROPERTY()
-	uint32 bFireEventsWhenForwards : 1;
+	TArray<FName> EventFunctions;
+
+private:
+
+	virtual UScriptStruct& GetScriptStructImpl() const override { return *StaticStruct(); }
+	virtual void EvaluateSwept(const FMovieSceneEvaluationOperand& Operand, const FMovieSceneContext& Context, const FPersistentEvaluationData& PersistentData, FMovieSceneExecutionTokens& ExecutionTokens) const override;
+};
+
+USTRUCT()
+struct FMovieSceneEventRepeaterTemplate : public FMovieSceneEventTemplateBase
+{
+	GENERATED_BODY()
+	
+	FMovieSceneEventRepeaterTemplate() {}
+	FMovieSceneEventRepeaterTemplate(const UMovieSceneEventRepeaterSection& Section, const UMovieSceneEventTrack& Track);
 
 	UPROPERTY()
-	uint32 bFireEventsWhenBackwards : 1;
+	FName EventToTrigger;
 
 private:
 
diff --git a/Engine/Source/Runtime/MovieSceneTracks/Private/Sections/MovieSceneEventRepeaterSection.cpp b/Engine/Source/Runtime/MovieSceneTracks/Private/Sections/MovieSceneEventRepeaterSection.cpp
new file mode 100644
--- /dev/null
+++ b/Engine/Source/Runtime/MovieSceneTracks/Private/Sections/MovieSceneEventRepeaterSection.cpp
@@ -0,0 +1,20 @@
+// Copyright 1998-2018 Epic Games, Inc. All Rights Reserved.
+
+#include "Sections/MovieSceneEventRepeaterSection.h"
+
+#if WITH_EDITORONLY_DATA
+
+void UMovieSceneEventRepeaterSection::OnBlueprintRecompiled(UBlueprint*)
+{
+	FName OldFunctionName = Event.FunctionName;
+
+	Event.CacheFunctionName();
+
+	if (Event.FunctionName != OldFunctionName)
+	{
+		MarkAsChanged();
+		MarkPackageDirty();
+	}
+}
+
+#endif		// WITH_EDITORONLY_DATA
\ No newline at end of file
diff --git a/Engine/Source/Runtime/MovieSceneTracks/Private/Sections/MovieSceneEventSectionBase.cpp b/Engine/Source/Runtime/MovieSceneTracks/Private/Sections/MovieSceneEventSectionBase.cpp
new file mode 100644
--- /dev/null
+++ b/Engine/Source/Runtime/MovieSceneTracks/Private/Sections/MovieSceneEventSectionBase.cpp
@@ -0,0 +1,40 @@
+// Copyright 1998-2018 Epic Games, Inc. All Rights Reserved.
+
+#include "Sections/MovieSceneEventSectionBase.h"
+#include "Engine/Blueprint.h"
+
+#if WITH_EDITORONLY_DATA
+
+void UMovieSceneEventSectionBase::Serialize(FArchive& Ar)
+{
+	Super::Serialize(Ar);
+
+
+	if (Ar.IsLoading())
+	{
+		if (UBlueprint* Blueprint = DirectorBlueprint.Get())
+		{
+			OnBlueprintCompiledHandle = Blueprint->OnCompiled().AddUObject(this, &UMovieSceneEventSectionBase::OnBlueprintRecompiled);
+		}
+	}
+}
+
+void UMovieSceneEventSectionBase::SetDirectorBlueprint(UBlueprint* InBlueprint)
+{
+	if (DirectorBlueprint != InBlueprint)
+	{
+		UBlueprint* ExistingBlueprint = DirectorBlueprint.Get();
+		if (ExistingBlueprint)
+		{
+			ExistingBlueprint->OnCompiled().Remove(OnBlueprintCompiledHandle);
+		}
+
+		DirectorBlueprint = InBlueprint;
+		if (InBlueprint)
+		{
+			OnBlueprintCompiledHandle = InBlueprint->OnCompiled().AddUObject(this, &UMovieSceneEventSectionBase::OnBlueprintRecompiled);
+		}
+	}
+}
+
+#endif		// WITH_EDITORONLY_DATA
\ No newline at end of file
diff --git a/Engine/Source/Runtime/MovieSceneTracks/Private/Sections/MovieSceneEventTriggerSection.cpp b/Engine/Source/Runtime/MovieSceneTracks/Private/Sections/MovieSceneEventTriggerSection.cpp
new file mode 100644
--- /dev/null
+++ b/Engine/Source/Runtime/MovieSceneTracks/Private/Sections/MovieSceneEventTriggerSection.cpp
@@ -0,0 +1,51 @@
+// Copyright 1998-2018 Epic Games, Inc. All Rights Reserved.
+
+#include "Sections/MovieSceneEventTriggerSection.h"
+#include "Channels/MovieSceneChannelProxy.h"
+#include "Engine/Blueprint.h"
+
+
+UMovieSceneEventTriggerSection::UMovieSceneEventTriggerSection(const FObjectInitializer& ObjInit)
+	: Super(ObjInit)
+{
+	bSupportsInfiniteRange = true;
+	SetRange(TRange<FFrameNumber>::All());
+
+#if WITH_EDITOR
+
+	ChannelProxy = MakeShared<FMovieSceneChannelProxy>(EventChannel, FMovieSceneChannelMetaData());
+
+#else
+
+	ChannelProxy = MakeShared<FMovieSceneChannelProxy>(EventChannel);
+
+#endif
+}
+
+
+#if WITH_EDITORONLY_DATA
+
+void UMovieSceneEventTriggerSection::OnBlueprintRecompiled(UBlueprint* InBlueprint)
+{
+	bool bHasChanged = false;
+
+	for (FMovieSceneEvent& Event : EventChannel.GetData().GetValues())
+	{
+		FName OldFunctionName = Event.FunctionName;
+
+		Event.CacheFunctionName();
+
+		if (Event.FunctionName != OldFunctionName)
+		{
+			bHasChanged = true;
+		}
+	}
+
+	if (bHasChanged)
+	{
+		MarkAsChanged();
+		MarkPackageDirty();
+	}
+}
+
+#endif		// WITH_EDITORONLY_DATA
\ No newline at end of file
diff --git a/Engine/Source/Runtime/MovieSceneTracks/Private/Tracks/MovieSceneEventTrack.cpp b/Engine/Source/Runtime/MovieSceneTracks/Private/Tracks/MovieSceneEventTrack.cpp
--- a/Engine/Source/Runtime/MovieSceneTracks/Private/Tracks/MovieSceneEventTrack.cpp
+++ b/Engine/Source/Runtime/MovieSceneTracks/Private/Tracks/MovieSceneEventTrack.cpp
@@ -3,6 +3,8 @@
 #include "Tracks/MovieSceneEventTrack.h"
 #include "MovieSceneCommonHelpers.h"
 #include "Sections/MovieSceneEventSection.h"
+#include "Sections/MovieSceneEventTriggerSection.h"
+#include "Sections/MovieSceneEventRepeaterSection.h"
 #include "Evaluation/MovieSceneEventTemplate.h"
 #include "Evaluation/MovieSceneEvaluationTrack.h"
 #include "Compilation/IMovieSceneTemplateGenerator.h"
@@ -22,7 +24,7 @@ void UMovieSceneEventTrack::AddSection(UMovieSceneSection& Section)
 
 UMovieSceneSection* UMovieSceneEventTrack::CreateNewSection()
 {
-	return NewObject<UMovieSceneSection>(this, UMovieSceneEventSection::StaticClass(), NAME_None, RF_Transactional);
+	return NewObject<UMovieSceneSection>(this, UMovieSceneEventTriggerSection::StaticClass(), NAME_None, RF_Transactional);
 }
 
 
@@ -57,7 +59,22 @@ void UMovieSceneEventTrack::RemoveSection(UMovieSceneSection& Section)
 
 FMovieSceneEvalTemplatePtr UMovieSceneEventTrack::CreateTemplateForSection(const UMovieSceneSection& InSection) const
 {
-	return FMovieSceneEventSectionTemplate(*CastChecked<UMovieSceneEventSection>(&InSection), *this);
+	if (const UMovieSceneEventSection* LegacyEventSection = Cast<const UMovieSceneEventSection>(&InSection))
+	{
+		return FMovieSceneEventSectionTemplate(*LegacyEventSection, *this);
+	}
+	else if (const UMovieSceneEventTriggerSection* TriggerSection = Cast<const UMovieSceneEventTriggerSection>(&InSection))
+	{
+		return FMovieSceneEventTriggerTemplate(*TriggerSection, *this);
+	}
+	else if (const UMovieSceneEventRepeaterSection* RepeaterSection = Cast<const UMovieSceneEventRepeaterSection>(&InSection))
+	{
+		return FMovieSceneEventRepeaterTemplate(*RepeaterSection, *this);
+	}
+	else
+	{
+		return FMovieSceneEvalTemplatePtr();
+	}
 }
 
 void UMovieSceneEventTrack::PostCompile(FMovieSceneEvaluationTrack& Track, const FMovieSceneTrackCompilerArgs& Args) const
diff --git a/Engine/Source/Runtime/MovieSceneTracks/Public/Channels/MovieSceneEvent.h b/Engine/Source/Runtime/MovieSceneTracks/Public/Channels/MovieSceneEvent.h
new file mode 100644
--- /dev/null
+++ b/Engine/Source/Runtime/MovieSceneTracks/Public/Channels/MovieSceneEvent.h
@@ -0,0 +1,114 @@
+// Copyright 1998-2018 Epic Games, Inc. All Rights Reserved.
+
+#pragma once
+
+#include "CoreTypes.h"
+#include "MovieSceneEvent.generated.h"
+
+class UEdGraph;
+class UK2Node_FunctionEntry;
+class UMovieSceneEventSectionBase;
+
+
+/**
+ * Struct type that is bound to a blueprint function entry node, and resolved to a cached UFunction when the blueprint/track is recompiled.
+ * Events can be bound to either of the following function signatures:
+ *
+ *   1. A function with no parameters (and no return value).
+ *     - Compatible with master tracks or object bindings
+ *     - No context passed through to event
+
+ *   2. A function with a single object or interface parameter (and no return value).
+ *     - Compatible with master tracks or object bindings
+ *     - Will be triggered with objects in the following order:
+ *         - Objects bound to the track's object binding, or:
+ *         - Objects specified on the track's event receivers array, or:
+ *         - Objects provided by the playback context.
+ *     - Will only trigger if the specified object is of the same type as the parameter (or interface)
+ */
+USTRUCT()
+struct MOVIESCENETRACKS_API FMovieSceneEvent
+{
+	GENERATED_BODY()
+
+
+	/**
+	 * Called after this event has been serialized in order to cache the function pointer if necessary
+	 */
+	void PostSerialize(const FArchive& Ar);
+
+
+	/**
+	 * Check whether the specified function is valid for a movie scene event
+	 * Functions must have either no parameters, or a single, pass-by-value object/interface parameter, with no return parameter.
+	 */
+	static bool IsValidFunction(UFunction* Function);
+
+
+	/**
+	 * The function that should be called to invoke this event.
+	 * Functions must have either no parameters, or a single, pass-by-value object/interface parameter, with no return parameter.
+	 */
+	UPROPERTY()
+	FName FunctionName;
+
+#if WITH_EDITORONLY_DATA
+
+public:
+
+
+	/**
+	 * Cache the function name to call from the blueprint function entry node. Will only cache the function if it has a valid signature.
+	 */
+	void CacheFunctionName();
+
+
+	/**
+	 * Check whether this event is bound to a valid blueprint entry node
+	 *
+	 * @return true if this event is bound to a function entry node with a valid signature, false otherwise.
+	 */
+	bool IsBoundToBlueprint() const;
+
+
+	/**
+	 * Helper function to determine whether the specified function entry is valid for this event
+	 *
+	 * @param Node         The node to test
+	 * @return true if the function entry node is compatible with a moviescene event, false otherwise
+	 */
+	static bool IsValidFunction(UK2Node_FunctionEntry* Node);
+
+
+	/**
+	 * Retrieve the function entry node this event is bound to
+	 *
+	 * @note Events may be bound to invalid function entries if they have been changed since they were assigned.
+	 * @see SetFunctionEntry, IsValidFunction
+	 * @return The function entry node if still available, nullptr if it has been destroyed, or was never assigned.
+	 */
+	UK2Node_FunctionEntry* GetFunctionEntry() const;
+
+
+	/**
+	 * Set the function entry that this event should trigger
+	 *
+	 * @param Entry        The graph node to bind to
+	 */
+	void SetFunctionEntry(UK2Node_FunctionEntry* Entry);
+
+private:
+
+	/** Weak pointer to the function entry within the blueprint graph for this event. Stored as an editor-only UObject so UHT can parse it when building for non-editor. */
+	UPROPERTY()
+	TWeakObjectPtr<UObject> FunctionEntry;
+
+#endif // WITH_EDITORONLY_DATA
+};
+
+template<>
+struct TStructOpsTypeTraits<FMovieSceneEvent> : TStructOpsTypeTraitsBase2<FMovieSceneEvent>
+{
+	enum { WithPostSerialize = true };
+};
+
diff --git a/Engine/Source/Runtime/MovieSceneTracks/Public/Channels/MovieSceneEventChannel.h b/Engine/Source/Runtime/MovieSceneTracks/Public/Channels/MovieSceneEventChannel.h
new file mode 100644
--- /dev/null
+++ b/Engine/Source/Runtime/MovieSceneTracks/Public/Channels/MovieSceneEventChannel.h
@@ -0,0 +1,75 @@
+// Copyright 1998-2018 Epic Games, Inc. All Rights Reserved.
+
+#pragma once
+
+#include "CoreTypes.h"
+#include "Curves/KeyHandle.h"
+#include "Channels/MovieSceneEvent.h"
+#include "Channels/MovieSceneChannel.h"
+#include "Channels/MovieSceneChannelData.h"
+#include "Channels/MovieSceneChannelTraits.h"
+#include "MovieSceneEventChannel.generated.h"
+
+
+USTRUCT()
+struct MOVIESCENETRACKS_API FMovieSceneEventChannel : public FMovieSceneChannel
+{
+	GENERATED_BODY()
+
+	/**
+	 * Access a mutable interface for this channel's data
+	 *
+	 * @return An object that is able to manipulate this channel's data
+	 */
+	FORCEINLINE TMovieSceneChannelData<FMovieSceneEvent> GetData()
+	{
+		return TMovieSceneChannelData<FMovieSceneEvent>(&KeyTimes, &KeyValues, &KeyHandles);
+	}
+
+	/**
+	 * Access a constant interface for this channel's data
+	 *
+	 * @return An object that is able to interrogate this channel's data
+	 */
+	FORCEINLINE TMovieSceneChannelData<const FMovieSceneEvent> GetData() const
+	{
+		return TMovieSceneChannelData<const FMovieSceneEvent>(&KeyTimes, &KeyValues);
+	}
+
+public:
+
+	// ~ FMovieSceneChannel Interface
+	virtual void GetKeys(const TRange<FFrameNumber>& WithinRange, TArray<FFrameNumber>* OutKeyTimes, TArray<FKeyHandle>* OutKeyHandles) override;
+	virtual void GetKeyTimes(TArrayView<const FKeyHandle> InHandles, TArrayView<FFrameNumber> OutKeyTimes) override;
+	virtual void SetKeyTimes(TArrayView<const FKeyHandle> InHandles, TArrayView<const FFrameNumber> InKeyTimes) override;
+	virtual void DuplicateKeys(TArrayView<const FKeyHandle> InHandles, TArrayView<FKeyHandle> OutNewHandles) override;
+	virtual void DeleteKeys(TArrayView<const FKeyHandle> InHandles) override;
+	virtual void ChangeFrameResolution(FFrameRate SourceRate, FFrameRate DestinationRate) override;
+	virtual TRange<FFrameNumber> ComputeEffectiveRange() const override;
+	virtual int32 GetNumKeys() const override;
+	virtual void Reset() override;
+	virtual void Offset(FFrameNumber DeltaPosition) override;
+
+private:
+
+	/** Array of times for each key */
+	UPROPERTY()
+	TArray<FFrameNumber> KeyTimes;
+
+	/** Array of values that correspond to each key time */
+	UPROPERTY()
+	TArray<FMovieSceneEvent> KeyValues;
+
+	FMovieSceneKeyHandleMap KeyHandles;
+};
+
+template<>
+struct TMovieSceneChannelTraits<FMovieSceneEventChannel> : TMovieSceneChannelTraitsBase<FMovieSceneEventChannel>
+{
+	enum { SupportsDefaults = false };
+};
+
+inline bool EvaluateChannel(const FMovieSceneEventChannel* InChannel, FFrameTime InTime, FMovieSceneEvent& OutValue)
+{
+	return false;
+}
diff --git a/Engine/Source/Runtime/MovieSceneTracks/Public/Sections/MovieSceneEventRepeaterSection.h b/Engine/Source/Runtime/MovieSceneTracks/Public/Sections/MovieSceneEventRepeaterSection.h
new file mode 100644
--- /dev/null
+++ b/Engine/Source/Runtime/MovieSceneTracks/Public/Sections/MovieSceneEventRepeaterSection.h
@@ -0,0 +1,29 @@
+// Copyright 1998-2018 Epic Games, Inc. All Rights Reserved.
+
+#pragma once
+
+#include "Sections/MovieSceneEventSectionBase.h"
+#include "Channels/MovieSceneEvent.h"
+#include "MovieSceneEventRepeaterSection.generated.h"
+
+
+/**
+ * Event section that will trigger its event exactly once, every time it is evaluated.
+ */
+UCLASS(MinimalAPI)
+class UMovieSceneEventRepeaterSection
+	: public UMovieSceneEventSectionBase
+{
+public:
+	GENERATED_BODY()
+
+#if WITH_EDITORONLY_DATA
+	virtual void OnBlueprintRecompiled(UBlueprint*) override;
+#endif
+
+	/** The event that should be triggered each time this section is evaluated */
+	UPROPERTY(EditAnywhere, Category="Event")
+	FMovieSceneEvent Event;
+};
+
+
diff --git a/Engine/Source/Runtime/MovieSceneTracks/Public/Sections/MovieSceneEventSectionBase.h b/Engine/Source/Runtime/MovieSceneTracks/Public/Sections/MovieSceneEventSectionBase.h
new file mode 100644
--- /dev/null
+++ b/Engine/Source/Runtime/MovieSceneTracks/Public/Sections/MovieSceneEventSectionBase.h
@@ -0,0 +1,42 @@
+// Copyright 1998-2018 Epic Games, Inc. All Rights Reserved.
+
+#pragma once
+
+#include "CoreTypes.h"
+#include "MovieSceneSection.h"
+#include "MovieSceneEventSectionBase.generated.h"
+
+
+/**
+ * Base class for all event sections. Manages dirtying the section and track on recompilation of the director blueprint.
+ */
+UCLASS(MinimalAPI)
+class UMovieSceneEventSectionBase
+	: public UMovieSceneSection
+{
+public:
+	GENERATED_BODY()
+
+#if WITH_EDITORONLY_DATA
+
+	/**
+	 * Assign the director blueprint pointer for this event section.
+	 */
+	MOVIESCENETRACKS_API void SetDirectorBlueprint(UBlueprint* InBlueprint);
+
+protected:
+
+	virtual void Serialize(FArchive& Ar) override;
+	virtual void OnBlueprintRecompiled(UBlueprint*) {}
+
+private:
+
+	/** A handle to the blueprint recompiled delegate binding for this section's director blueprint */
+	UPROPERTY()
+	TWeakObjectPtr<UBlueprint> DirectorBlueprint;
+
+	/** A handle to the blueprint recompiled delegate binding for this section's director blueprint */
+	FDelegateHandle OnBlueprintCompiledHandle;
+
+#endif // WITH_EDITORONLY_DATA
+};
\ No newline at end of file
diff --git a/Engine/Source/Runtime/MovieSceneTracks/Public/Sections/MovieSceneEventTriggerSection.h b/Engine/Source/Runtime/MovieSceneTracks/Public/Sections/MovieSceneEventTriggerSection.h
new file mode 100644
--- /dev/null
+++ b/Engine/Source/Runtime/MovieSceneTracks/Public/Sections/MovieSceneEventTriggerSection.h
@@ -0,0 +1,32 @@
+// Copyright 1998-2018 Epic Games, Inc. All Rights Reserved.
+
+#pragma once
+
+#include "MovieSceneEventSectionBase.h"
+#include "Channels/MovieSceneEventChannel.h"
+#include "MovieSceneEventTriggerSection.generated.h"
+
+
+/**
+ * Event section that triggeres specific timed events.
+ */
+UCLASS(MinimalAPI)
+class UMovieSceneEventTriggerSection
+	: public UMovieSceneEventSectionBase
+{
+public:
+	GENERATED_BODY()
+
+	UMovieSceneEventTriggerSection(const FObjectInitializer& ObjInit);
+
+
+#if WITH_EDITORONLY_DATA
+
+	virtual void OnBlueprintRecompiled(UBlueprint*) override;
+
+#endif
+
+	/** The channel that defines this section's timed events */
+	UPROPERTY()
+	FMovieSceneEventChannel EventChannel;
+};
\ No newline at end of file
diff --git a/Engine/Source/Runtime/MovieSceneTracks/Public/Tracks/MovieSceneEventTrack.h b/Engine/Source/Runtime/MovieSceneTracks/Public/Tracks/MovieSceneEventTrack.h
--- a/Engine/Source/Runtime/MovieSceneTracks/Public/Tracks/MovieSceneEventTrack.h
+++ b/Engine/Source/Runtime/MovieSceneTracks/Public/Tracks/MovieSceneEventTrack.h
@@ -58,6 +58,7 @@ public:
 	virtual void RemoveSection(UMovieSceneSection& Section) override;
 	virtual FMovieSceneEvalTemplatePtr CreateTemplateForSection(const UMovieSceneSection& InSection) const override;
 	virtual void PostCompile(FMovieSceneEvaluationTrack& Track, const FMovieSceneTrackCompilerArgs& Args) const override;
+	virtual bool SupportsMultipleRows() const override { return true; }
 
 #if WITH_EDITORONLY_DATA
 	virtual FText GetDefaultDisplayName() const override;
diff --git a/Engine/Source/Runtime/UMG/Private/Animation/WidgetAnimation.cpp b/Engine/Source/Runtime/UMG/Private/Animation/WidgetAnimation.cpp
--- a/Engine/Source/Runtime/UMG/Private/Animation/WidgetAnimation.cpp
+++ b/Engine/Source/Runtime/UMG/Private/Animation/WidgetAnimation.cpp
@@ -179,6 +179,12 @@ UMovieScene* UWidgetAnimation::GetMovieScene() const
 	return MovieScene;
 }
 
+UObject* UWidgetAnimation::CreateDirectorInstance(IMovieScenePlayer& Player)
+{
+	// Widget animations do not create separate director instances, but just re-use the UUserWidget from the playback context
+	UUserWidget* WidgetContext = CastChecked<UUserWidget>(Player.GetPlaybackContext());
+	return WidgetContext;
+}
 
 UObject* UWidgetAnimation::GetParentObject(UObject* Object) const
 {
diff --git a/Engine/Source/Runtime/UMG/Public/Animation/UMGSequencePlayer.h b/Engine/Source/Runtime/UMG/Public/Animation/UMGSequencePlayer.h
--- a/Engine/Source/Runtime/UMG/Public/Animation/UMGSequencePlayer.h
+++ b/Engine/Source/Runtime/UMG/Public/Animation/UMGSequencePlayer.h
@@ -59,6 +59,7 @@ public:
 
 	/** IMovieScenePlayer interface */
 	virtual FMovieSceneRootEvaluationTemplateInstance& GetEvaluationTemplate() override { return RootTemplateInstance; }
+	virtual UObject* AsUObject() override { return this; }
 	virtual void UpdateCameraCut(UObject* CameraObject, UObject* UnlockIfCameraObject, bool bJumpCut) override {}
 	virtual void SetViewportSettings(const TMap<FViewportClient*, EMovieSceneViewportParams>& ViewportParamsMap) override {}
 	virtual void GetViewportSettings(TMap<FViewportClient*, EMovieSceneViewportParams>& ViewportParamsMap) const override {}
diff --git a/Engine/Source/Runtime/UMG/Public/Animation/WidgetAnimation.h b/Engine/Source/Runtime/UMG/Public/Animation/WidgetAnimation.h
--- a/Engine/Source/Runtime/UMG/Public/Animation/WidgetAnimation.h
+++ b/Engine/Source/Runtime/UMG/Public/Animation/WidgetAnimation.h
@@ -67,6 +67,7 @@ public:
 	virtual UObject* GetParentObject(UObject* Object) const override;
 	virtual void UnbindPossessableObjects(const FGuid& ObjectId) override;
 	virtual void LocateBoundObjects(const FGuid& ObjectId, UObject* Context, TArray<UObject*, TInlineAllocator<1>>& OutObjects) const override;
+	virtual UObject* CreateDirectorInstance(IMovieScenePlayer& Player) override;
 	// ~UMovieSceneAnimation overrides
 
 	//~ Begin UObject Interface. 
