Trivial changes to appease static analysis

diff --git a/Engine/Source/Editor/MovieSceneTools/Private/MovieSceneEventCustomization.cpp b/Engine/Source/Editor/MovieSceneTools/Private/MovieSceneEventCustomization.cpp
--- a/Engine/Source/Editor/MovieSceneTools/Private/MovieSceneEventCustomization.cpp
+++ b/Engine/Source/Editor/MovieSceneTools/Private/MovieSceneEventCustomization.cpp
@@ -488,18 +488,18 @@ void FMovieSceneEventCustomization::CreateEventEndpoint()
 
 	// Create a single event binding and point all events in this property handle to it
 	UK2Node_FunctionEntry* NewFunctionEntry = SequenceEditor->CreateEventEndpoint(Sequence);
+	if (!ensure(NewFunctionEntry))
+	{
+		return;
+	}
 
 	if (UMovieSceneEventTrack* CommonTrack = GetCommonTrack())
 	{
 		SequenceEditor->InitializeEndpointForTrack(CommonTrack, NewFunctionEntry);
 	}
-	
-	SetEventEndpoint(NewFunctionEntry);
 
-	if (NewFunctionEntry)
-	{
-		FKismetEditorUtilities::BringKismetToFocusAttentionOnObject(NewFunctionEntry, false);
-	}
+	SetEventEndpoint(NewFunctionEntry);
+	FKismetEditorUtilities::BringKismetToFocusAttentionOnObject(NewFunctionEntry, false);
 }
 
 void FMovieSceneEventCustomization::CreateEventEndpointFromFunction(UFunction* QuickBindFunction, UClass* PinClassType)
@@ -573,11 +573,7 @@ void FMovieSceneEventCustomization::CreateEventEndpointFromFunction(UFunction* Q
 	}
 
 	SetEventEndpoint(NewFunctionEntry);
-
-	if (NewFunctionEntry)
-	{
-		FKismetEditorUtilities::BringKismetToFocusAttentionOnObject(NewFunctionEntry, false);
-	}
+	FKismetEditorUtilities::BringKismetToFocusAttentionOnObject(NewFunctionEntry, false);
 }
 
 void FMovieSceneEventCustomization::ClearEventEndpoint()
diff --git a/Engine/Source/Editor/MovieSceneTools/Private/Sections/EventSection.cpp b/Engine/Source/Editor/MovieSceneTools/Private/Sections/EventSection.cpp
--- a/Engine/Source/Editor/MovieSceneTools/Private/Sections/EventSection.cpp
+++ b/Engine/Source/Editor/MovieSceneTools/Private/Sections/EventSection.cpp
@@ -174,7 +174,8 @@ FReply FEventTriggerSection::OnKeyDoubleClicked(FKeyHandle KeyHandle)
 		int32 EventIndex = ChannelData.GetIndex(KeyHandle);
 		if (EventIndex != INDEX_NONE)
 		{
-			FMovieSceneEvent& Event = ChannelData.GetValues()[EventIndex];
+			TArrayView<FMovieSceneEvent> Events = ChannelData.GetValues();
+			FMovieSceneEvent& Event = Events[EventIndex];
 
 			if (!Event.IsBoundToBlueprint())
 			{
diff --git a/Engine/Source/Runtime/MovieScene/Private/Evaluation/MovieSceneEvaluationTemplateInstance.cpp b/Engine/Source/Runtime/MovieScene/Private/Evaluation/MovieSceneEvaluationTemplateInstance.cpp
--- a/Engine/Source/Runtime/MovieScene/Private/Evaluation/MovieSceneEvaluationTemplateInstance.cpp
+++ b/Engine/Source/Runtime/MovieScene/Private/Evaluation/MovieSceneEvaluationTemplateInstance.cpp
@@ -588,7 +588,7 @@ void FMovieSceneRootEvaluationTemplateInstance::RecreateInstances(const FMovieSc
 		{
 			FMovieSceneEvaluationTemplateInstance NewInstance(SubData, *TemplateStore);
 
-			if (!NewInstance.DirectorInstance)
+			if (!NewInstance.DirectorInstance && NewInstance.Sequence)
 			{
 				NewInstance.DirectorInstance = NewInstance.Sequence->CreateDirectorInstance(Player);
 			}
