Sequencer: Compensate the added transform when an attach track is added. If there's no existing transform track,  create one and add the compensation to that transform section. If there's an existing transform section without keys, compensate to that transform section. Otherwise, create an additive transform section and compensate to that transform section.

diff --git a/Engine/Source/Editor/MovieSceneTools/Private/TrackEditors/AttachTrackEditor.cpp b/Engine/Source/Editor/MovieSceneTools/Private/TrackEditors/AttachTrackEditor.cpp
--- a/Engine/Source/Editor/MovieSceneTools/Private/TrackEditors/AttachTrackEditor.cpp
+++ b/Engine/Source/Editor/MovieSceneTools/Private/TrackEditors/AttachTrackEditor.cpp
@@ -6,10 +6,13 @@
 #include "SequencerSectionPainter.h"
 #include "GameFramework/WorldSettings.h"
 #include "Tracks/MovieScene3DAttachTrack.h"
+#include "Tracks/MovieScene3DTransformTrack.h"
 #include "Sections/MovieScene3DAttachSection.h"
+#include "Sections/MovieScene3DTransformSection.h"
 #include "ActorEditorUtils.h"
 #include "MovieSceneObjectBindingIDPicker.h"
 #include "MovieSceneToolHelpers.h"
+#include "Channels/MovieSceneChannelProxy.h"
 #include "Evaluation/MovieSceneEvaluationTemplateInstance.h"
 
 
@@ -244,7 +247,34 @@ FKeyPropertyResult F3DAttachTrackEditor::AddKeyInternal( FFrameNumber KeyTime, c
 		return KeyPropertyResult;
 	}
 
-	for( int32 ObjectIndex = 0; ObjectIndex < Objects.Num(); ++ObjectIndex )
+	UMovieScene* MovieScene = GetSequencer()->GetFocusedMovieSceneSequence()->GetMovieScene();
+
+	FMovieSceneSequenceID SequenceID = GetSequencer()->GetFocusedTemplateID();
+	if (ConstraintBindingID.GetSequenceID().IsValid())
+	{
+		// Ensure that this ID is resolvable from the root, based on the current local sequence ID
+		FMovieSceneObjectBindingID RootBindingID = ConstraintBindingID.ResolveLocalToRoot(SequenceID, GetSequencer()->GetEvaluationTemplate().GetHierarchy());
+		SequenceID = RootBindingID.GetSequenceID();
+	}
+
+	FTransform ParentTransform;
+	TArrayView<TWeakObjectPtr<UObject>> RuntimeObjects = GetSequencer()->FindBoundObjects(ConstraintBindingID.GetGuid(), SequenceID);
+
+	if (RuntimeObjects.Num() == 1 && RuntimeObjects[0].IsValid())
+	{
+		AActor* ParentActor = Cast<AActor>(RuntimeObjects[0].Get());
+		if (ParentActor)
+		{
+			ParentTransform = ParentActor->GetActorTransform();
+			
+			if (ParentActor->GetRootComponent()->DoesSocketExist(SocketName))
+			{
+				ParentTransform = ParentActor->GetRootComponent()->GetSocketTransform(SocketName);
+			}
+		}
+	}
+		
+	for (int32 ObjectIndex = 0; ObjectIndex < Objects.Num(); ++ObjectIndex)
 	{
 		UObject* Object = Objects[ObjectIndex].Get();
 
@@ -259,8 +289,10 @@ FKeyPropertyResult F3DAttachTrackEditor::AddKeyInternal( FFrameNumber KeyTime, c
 
 			if (ensure(Track))
 			{
-				// Clamp to next attach section's start time or the end of the current sequencer view range
-				FFrameNumber AttachEndTime = (GetSequencer()->GetViewRange().GetUpperBoundValue() * Track->GetTypedOuter<UMovieScene>()->GetTickResolution()).FrameNumber;
+				FFrameRate TickResolution = Track->GetTypedOuter<UMovieScene>()->GetTickResolution();
+
+				// Clamp to next attach section's start time or the end of the current movie scene range
+				FFrameNumber AttachEndTime = MovieScene->GetPlaybackRange().GetUpperBoundValue();
 
 				for (UMovieSceneSection* Section : Track->GetAllSections())
 				{
@@ -274,9 +306,122 @@ FKeyPropertyResult F3DAttachTrackEditor::AddKeyInternal( FFrameNumber KeyTime, c
 					}
 				}
 
+				AActor* Child = Cast<AActor>(Object);
+				FTransform RelativeTransform = Child->GetRootComponent()->GetComponentTransform().GetRelativeTransform(ParentTransform);
+
 				int32 Duration = FMath::Max(0, (AttachEndTime - KeyTime).Value);
 				Cast<UMovieScene3DAttachTrack>(Track)->AddConstraint( KeyTime, Duration, SocketName, ComponentName, ConstraintBindingID);
 				KeyPropertyResult.bTrackModified = true;
+
+				// Compensate
+				FName TransformPropertyName("Transform");
+				TRange<FFrameNumber> AttachRange(KeyTime, AttachEndTime);
+
+				MovieScene->Modify();
+
+				// Create a transform track if it doesn't exist
+				UMovieScene3DTransformTrack* TransformTrack = Cast<UMovieScene3DTransformTrack>(MovieScene->FindTrack<UMovieScene3DTransformTrack>(ObjectHandle));
+				if (!TransformTrack)
+				{
+					FFindOrCreateTrackResult TransformTrackResult = FindOrCreateTrackForObject(ObjectHandle, UMovieScene3DTransformTrack::StaticClass());
+					TransformTrack = Cast<UMovieScene3DTransformTrack>(TransformTrackResult.Track);
+					TransformTrack->SetPropertyNameAndPath(TransformPropertyName, TransformPropertyName.ToString());
+				}
+
+				if (!TransformTrack)
+				{
+					continue;
+				}
+
+				// Create a transform section if it doesn't exist
+				UMovieScene3DTransformSection* TransformSection = nullptr;
+				if (TransformTrack->IsEmpty())
+				{
+					TransformTrack->Modify();
+					TransformSection = Cast<UMovieScene3DTransformSection>(TransformTrack->CreateNewSection());
+					if (TransformSection)
+					{
+						TransformSection->SetRange(TRange<FFrameNumber>::All());
+
+						TransformTrack->AddSection(*TransformSection);
+					}
+				}
+
+				// Reuse the transform section if it overlaps and there's no keys
+				else if (TransformTrack->GetAllSections().Num() == 1)
+				{
+					TRange<FFrameNumber> TransformRange = TransformTrack->GetAllSections()[0]->GetRange();
+					if (TRange<FFrameNumber>::Intersection(AttachRange, TransformRange).IsEmpty())
+					{
+						continue;
+					}
+
+					bool bEmptyKeys = true;
+					TArrayView<FMovieSceneFloatChannel*> Channels = Cast<UMovieScene3DTransformSection>(TransformTrack->GetAllSections()[0])->GetChannelProxy().GetChannels<FMovieSceneFloatChannel>();
+					for (auto Channel : Channels)
+					{
+						if (Channel->GetTimes().Num() != 0)
+						{
+							bEmptyKeys = false;
+							break;
+						}
+					}
+
+					if (bEmptyKeys)
+					{
+						TransformSection = Cast<UMovieScene3DTransformSection>(TransformTrack->GetAllSections()[0]);
+					}
+				}
+
+				// Create a new additive transform section
+				if (!TransformSection)
+				{
+					TransformTrack->Modify();
+					TransformSection = Cast<UMovieScene3DTransformSection>(TransformTrack->CreateNewSection());
+					if (TransformSection)
+					{
+						for (auto Section : TransformTrack->GetAllSections())
+						{
+							Section->SetRowIndex(Section->GetRowIndex() + 1);
+						}
+
+						TransformSection->SetRange(AttachRange);
+						TransformSection->SetBlendType(EMovieSceneBlendType::Additive);
+						TransformSection->SetRowIndex(0);
+						TransformSection->SetMask(FMovieSceneTransformMask(EMovieSceneTransformChannel::Rotation | EMovieSceneTransformChannel::Translation));
+						TransformTrack->AddSection(*TransformSection);
+
+						RelativeTransform = ParentTransform.Inverse();
+					}
+				}
+
+				if (!TransformSection)
+				{
+					continue;
+				}
+
+				if (!TransformSection->TryModify())
+				{
+					continue;
+				}
+
+				TArrayView<FMovieSceneFloatChannel*> Channels = TransformSection->GetChannelProxy().GetChannels<FMovieSceneFloatChannel>();
+
+				FVector DefaultLocation = RelativeTransform.GetLocation();
+				FVector DefaultRotation = RelativeTransform.GetRotation().Euler();
+				FVector DefaultScale3D = RelativeTransform.GetScale3D();
+
+				Channels[0]->SetDefault(DefaultLocation.X);
+				Channels[1]->SetDefault(DefaultLocation.Y);
+				Channels[2]->SetDefault(DefaultLocation.Z);
+
+				Channels[3]->SetDefault(DefaultRotation.X);
+				Channels[4]->SetDefault(DefaultRotation.Y);
+				Channels[5]->SetDefault(DefaultRotation.Z);
+
+				Channels[6]->SetDefault(DefaultScale3D.X);
+				Channels[7]->SetDefault(DefaultScale3D.Y);
+				Channels[8]->SetDefault(DefaultScale3D.Z);
 			}
 		}
 	}
