Sequencer: Fix crash on importing fbx cameras with identical names to existing objects in the level. Add cameras if bound objects don't exist. Changed logs to popup toasts.

diff --git a/Engine/Source/Editor/MovieSceneTools/Private/MovieSceneToolHelpers.cpp b/Engine/Source/Editor/MovieSceneTools/Private/MovieSceneToolHelpers.cpp
--- a/Engine/Source/Editor/MovieSceneTools/Private/MovieSceneToolHelpers.cpp
+++ b/Engine/Source/Editor/MovieSceneTools/Private/MovieSceneToolHelpers.cpp
@@ -1221,6 +1221,8 @@ bool ImportFBXCamera(UnFbx::FFbxImporter* FbxImporter, UMovieScene* InMovieScene
 		TArray<FbxCamera*> AllCameras;
 		GetCameras(FbxImporter->Scene->GetRootNode(), AllCameras);
 
+		UWorld* World = GCurrentLevelEditingViewportClient ? GCurrentLevelEditingViewportClient->GetWorld() : nullptr;
+
 		// Find unmatched cameras
 		TArray<FbxCamera*> UnmatchedCameras;
 		for (auto Camera : AllCameras)
@@ -1231,10 +1233,26 @@ bool ImportFBXCamera(UnFbx::FFbxImporter* FbxImporter, UMovieScene* InMovieScene
 			for (auto InObjectBinding : InObjectBindingMap)
 			{		
 				FString ObjectName = InObjectBinding.Value;
-				if ( !FCString::Strcmp(*ObjectName,UTF8_TO_TCHAR(Camera->GetName())))
+				if (!FCString::Strcmp(*ObjectName, UTF8_TO_TCHAR(Camera->GetName())))
 				{
-					bMatched = true;
-					break;
+					// Look for a valid bound object, otherwise need to create a new camera and assign this binding to it
+					bool bFoundBoundObject = false;
+					TArrayView<TWeakObjectPtr<>> BoundObjects = InSequencer.FindBoundObjects(InObjectBinding.Key, InSequencer.GetFocusedTemplateID());
+					for (auto BoundObject : BoundObjects)
+					{
+						if (BoundObject.IsValid())
+						{
+							bFoundBoundObject = true;
+							break;
+						}
+					}
+
+					if (!bFoundBoundObject)
+					{
+						FNotificationInfo Info(FText::Format(NSLOCTEXT("MovieSceneTools", "NoBoundObjectsError", "Existing binding has no objects. Creating a new camera and binding for {0}"), FText::FromString(ObjectName)));
+						Info.ExpireDuration = 5.0f;
+						FSlateNotificationManager::Get().AddNotification(Info)->SetCompletionState(SNotificationItem::CS_Fail);
+					}
 				}
 			}
 
@@ -1244,9 +1262,6 @@ bool ImportFBXCamera(UnFbx::FFbxImporter* FbxImporter, UMovieScene* InMovieScene
 			}
 		}
 
-		// Add any unmatched cameras
-		UWorld* World = GCurrentLevelEditingViewportClient ? GCurrentLevelEditingViewportClient->GetWorld() : nullptr;
-
 		// If there are new cameras, clear the object binding map so that we're only assigning values to the newly created cameras
 		if (UnmatchedCameras.Num() != 0)
 		{
@@ -1254,12 +1269,12 @@ bool ImportFBXCamera(UnFbx::FFbxImporter* FbxImporter, UMovieScene* InMovieScene
 			bCamsCreated = true;
 		}
 
+		// Add any unmatched cameras
 		for (auto UnmatchedCamera : UnmatchedCameras)
 		{
 			FString CameraName = FString(ANSI_TO_TCHAR(UnmatchedCamera->GetName()));
 
 			FActorSpawnParameters SpawnParams;
-			SpawnParams.Name = *CameraName;
 			ACineCameraActor* NewCamera = World->SpawnActor<ACineCameraActor>(SpawnParams);
 			NewCamera->SetActorLabel(*CameraName);
 
@@ -1295,14 +1310,20 @@ bool ImportFBXCamera(UnFbx::FFbxImporter* FbxImporter, UMovieScene* InMovieScene
 		{
 			if (bMatchByNameOnly)
 			{
-				UE_LOG(LogMovieScene, Error, TEXT("Fbx Import: Failed to find any matching camera for (%s)."), *ObjectName);
+				FNotificationInfo Info(FText::Format(NSLOCTEXT("MovieSceneTools", "NoMatchingCameraError", "Failed to find any matching camera for {0}"), FText::FromString(ObjectName)));
+				Info.ExpireDuration = 5.0f;
+				FSlateNotificationManager::Get().AddNotification(Info)->SetCompletionState(SNotificationItem::CS_Fail);
+
 				continue;
 			}
 
 			CameraNode = FindCamera(FbxImporter->Scene->GetRootNode());
 			if (CameraNode)
 			{
-				UE_LOG(LogMovieScene, Warning, TEXT("Fbx Import: Failed to find exact matching camera for (%s). Using first camera from fbx (%s)"), *ObjectName, UTF8_TO_TCHAR(CameraNode->GetName()));
+				FString CameraName = FString(ANSI_TO_TCHAR(CameraNode->GetName()));
+				FNotificationInfo Info(FText::Format(NSLOCTEXT("MovieSceneTools", "NoMatchingCameraWarning", "Failed to find any matching camera for {0}. Importing onto first camera from fbx {1}"), FText::FromString(ObjectName), FText::FromString(CameraName)));
+				Info.ExpireDuration = 5.0f;
+				FSlateNotificationManager::Get().AddNotification(Info)->SetCompletionState(SNotificationItem::CS_Fail);
 			}
 		}
 
