Fixing scaling of audio tracks in sequencer. Changed to multiplying duration by 1/pitch instead of pitch.

diff --git a/Engine/Source/Editor/MovieSceneTools/Private/TrackEditors/AudioTrackEditor.cpp b/Engine/Source/Editor/MovieSceneTools/Private/TrackEditors/AudioTrackEditor.cpp
--- a/Engine/Source/Editor/MovieSceneTools/Private/TrackEditors/AudioTrackEditor.cpp
+++ b/Engine/Source/Editor/MovieSceneTools/Private/TrackEditors/AudioTrackEditor.cpp
@@ -331,7 +331,7 @@ void FAudioThumbnail::GenerateWaveformPreview(TArray<uint8>& OutData, TRange<flo
 	// @todo Sequencer This fixes looping drawing by pretending we are only dealing with a SoundWave
 	TRange<double> AudioTrueRange = TRange<double>(
 		SectionStartTime - AudioSection->GetStartOffset(),
-		SectionStartTime - AudioSection->GetStartOffset() + DeriveUnloopedDuration(AudioSection) * PitchMultiplierValue);
+		SectionStartTime - AudioSection->GetStartOffset() + DeriveUnloopedDuration(AudioSection) * (1.0f / PitchMultiplierValue));
 
 	float TrueRangeSize = AudioTrueRange.Size<float>();
 	float DrawRangeSize = DrawRange.Size<float>();
