Fixed GUIDs not being created for call function nodes created from a movie scene event quick binding

diff --git a/Engine/Source/Editor/MovieSceneTools/Private/MovieSceneEventCustomization.cpp b/Engine/Source/Editor/MovieSceneTools/Private/MovieSceneEventCustomization.cpp
--- a/Engine/Source/Editor/MovieSceneTools/Private/MovieSceneEventCustomization.cpp
+++ b/Engine/Source/Editor/MovieSceneTools/Private/MovieSceneEventCustomization.cpp
@@ -542,7 +542,9 @@ void FMovieSceneEventCustomization::CreateEventEndpointFromFunction(UFunction* Q
 	UK2Node_CallFunction* CallFuncNode = NewObject<UK2Node_CallFunction>(Graph);
 	CallFuncNode->NodePosX = NewFunctionEntry->NodePosX + NewFunctionEntry->NodeWidth + 200;
 	CallFuncNode->NodePosY = NewFunctionEntry->NodePosY;
-	CallFuncNode->FunctionReference.SetFromField<UFunction>(QuickBindFunction, false);
+	CallFuncNode->CreateNewGuid();
+	CallFuncNode->SetFromFunction(QuickBindFunction);
+	CallFuncNode->PostPlacedNewNode();
 	CallFuncNode->ReconstructNode();
 
 	Graph->AddNode(CallFuncNode, false, false);
