Sequencer: Add double click to open the shot asset directly

diff --git a/Engine/Source/Editor/MovieSceneTools/Private/Sections/CinematicShotSection.cpp b/Engine/Source/Editor/MovieSceneTools/Private/Sections/CinematicShotSection.cpp
--- a/Engine/Source/Editor/MovieSceneTools/Private/Sections/CinematicShotSection.cpp
+++ b/Engine/Source/Editor/MovieSceneTools/Private/Sections/CinematicShotSection.cpp
@@ -14,6 +14,7 @@
 #include "EditorStyleSet.h"
 #include "MovieSceneToolHelpers.h"
 #include "MovieSceneTimeHelpers.h"
+#include "Toolkits/AssetEditorManager.h"
 
 
 #define LOCTEXT_NAMESPACE "FCinematicShotSection"
@@ -363,7 +364,14 @@ FReply FCinematicShotSection::OnSectionDoubleClicked(const FGeometry& SectionGeo
 	{
 		if (SectionObject.GetSequence())
 		{
-			Sequencer.Pin()->FocusSequenceInstance(SectionObject);
+			if (MouseEvent.IsControlDown())
+			{
+				FAssetEditorManager::Get().OpenEditorForAsset(SectionObject.GetSequence());
+			}
+			else
+			{
+				Sequencer.Pin()->FocusSequenceInstance(SectionObject);
+			}
 		}
 	}
 
