Sequencer: Fix logic issue in newly added key not expanding section range.

diff --git a/Engine/Source/Editor/MovieSceneTools/Private/CurveKeyEditors/SequencerKeyEditor.h b/Engine/Source/Editor/MovieSceneTools/Private/CurveKeyEditors/SequencerKeyEditor.h
--- a/Engine/Source/Editor/MovieSceneTools/Private/CurveKeyEditors/SequencerKeyEditor.h
+++ b/Engine/Source/Editor/MovieSceneTools/Private/CurveKeyEditors/SequencerKeyEditor.h
@@ -120,13 +120,14 @@ struct TSequencerKeyEditor
 		}
 		else
 		{
-			const bool bHasAnyKeys = Channel->GetNumKeys() != 0;
+			bool bHasAnyKeys = Channel->GetNumKeys() != 0;
 
 			if (bHasAnyKeys || bAutoSetTrackDefaults == false)
 			{
 				// When auto setting track defaults are disabled, add a key even when it's empty so that the changed
 				// value is saved and is propagated to the property.
 				AddKeyToChannel(Channel, CurrentTime, InValue, Interpolation);
+				bHasAnyKeys = Channel->GetNumKeys() != 0;
 			}
 
 			if (bHasAnyKeys)
