Sequencer: Do not clear keyboard focus on commit for frame number text boxes. This fixes an issue where the right click properties popup disappears when you type in a frame number. This was marked as a regression when editing the Start Frame Offset because it used to be a float property but when changing it to a frame number property, it inherited this behavior.

diff --git a/Engine/Source/Editor/MovieSceneTools/Private/FrameNumberDetailsCustomization.cpp b/Engine/Source/Editor/MovieSceneTools/Private/FrameNumberDetailsCustomization.cpp
--- a/Engine/Source/Editor/MovieSceneTools/Private/FrameNumberDetailsCustomization.cpp
+++ b/Engine/Source/Editor/MovieSceneTools/Private/FrameNumberDetailsCustomization.cpp
@@ -59,6 +59,7 @@ void FFrameNumberDetailsCustomization::CustomizeChildren(TSharedRef<IPropertyHan
 			.Text(this, &FFrameNumberDetailsCustomization::OnGetTimeText)
 			.OnTextCommitted(this, &FFrameNumberDetailsCustomization::OnTimeTextCommitted)
 			.SelectAllTextWhenFocused(true)
+			.ClearKeyboardFocusOnCommit(false)
 			.RevertTextOnEscape(true)
 			.Font(IDetailLayoutBuilder::GetDetailFont())
 		];
