Sequencer: Followup change for clearing focus on section start/end range

diff --git a/Engine/Source/Editor/MovieSceneTools/Private/MovieSceneSectionDetailsCustomization.cpp b/Engine/Source/Editor/MovieSceneTools/Private/MovieSceneSectionDetailsCustomization.cpp
--- a/Engine/Source/Editor/MovieSceneTools/Private/MovieSceneSectionDetailsCustomization.cpp
+++ b/Engine/Source/Editor/MovieSceneTools/Private/MovieSceneSectionDetailsCustomization.cpp
@@ -68,6 +68,7 @@ void FMovieSceneSectionDetailsCustomization::CustomizeDetails(IDetailLayoutBuild
 			.IsEnabled(this, &FMovieSceneSectionDetailsCustomization::IsRangeStartTextboxEnabled)
 			.SelectAllTextWhenFocused(true)
 			.RevertTextOnEscape(true)
+			.ClearKeyboardFocusOnCommit(false)
 			.Font(IDetailLayoutBuilder::GetDetailFont())
 		]
 		+ SHorizontalBox::Slot()
@@ -112,6 +113,7 @@ void FMovieSceneSectionDetailsCustomization::CustomizeDetails(IDetailLayoutBuild
 			.IsEnabled(this, &FMovieSceneSectionDetailsCustomization::IsRangeEndTextboxEnabled)
 			.SelectAllTextWhenFocused(true)
 			.RevertTextOnEscape(true)
+			.ClearKeyboardFocusOnCommit(false)
 			.Font(IDetailLayoutBuilder::GetDetailFont())
 		]
 		+ SHorizontalBox::Slot()
