PR #5681: Fix command line sequence render (Contributed by [REDACTED_PERSON])

diff --git a/Engine/Source/Editor/MovieSceneTools/Private/AutomatedLevelSequenceCapture.cpp b/Engine/Source/Editor/MovieSceneTools/Private/AutomatedLevelSequenceCapture.cpp
--- a/Engine/Source/Editor/MovieSceneTools/Private/AutomatedLevelSequenceCapture.cpp
+++ b/Engine/Source/Editor/MovieSceneTools/Private/AutomatedLevelSequenceCapture.cpp
@@ -184,7 +184,7 @@ void UAutomatedLevelSequenceCapture::Initialize(TSharedPtr<FSceneViewport> InVie
 			{
 				for( auto It = TActorIterator<ALevelSequenceActor>( InViewport->GetClient()->GetWorld() ); It; ++It )
 				{
-					if( It->LevelSequence == LevelSequenceAsset )
+					if( It->LevelSequence == Asset )
 					{
 						// Found it!
 						Actor = *It;
