Made it impossible for movie captures to capture the same frame twice

  - We now automatically move on to the next frame when a frame has been captured which ensures we cannot possibly capture the same frame twice.

diff --git a/Engine/Source/Editor/MovieSceneTools/Private/AutomatedLevelSequenceCapture.cpp b/Engine/Source/Editor/MovieSceneTools/Private/AutomatedLevelSequenceCapture.cpp
--- a/Engine/Source/Editor/MovieSceneTools/Private/AutomatedLevelSequenceCapture.cpp
+++ b/Engine/Source/Editor/MovieSceneTools/Private/AutomatedLevelSequenceCapture.cpp
@@ -26,20 +26,16 @@ const FName UAutomatedLevelSequenceCapture::AutomatedLevelSequenceCaptureUIName
 struct FMovieSceneTimeController_FrameStep : FMovieSceneTimeController
 {
 	FMovieSceneTimeController_FrameStep()
-		: DeltaTime(0)
-		, CurrentTime(-1)
+		: StartTimeOffset(0)
 	{}
 
 	virtual void OnTick(float DeltaSeconds, float InPlayRate) override
-	{
-		// Move onto the next frame in the sequence. Play rate dilation occurs in OnRequestCurrentTime, since this InPlayRate does not consider the global world settings dilation
-		DeltaTime = FFrameTime(1);
-	}
+	{}
 
 	virtual void OnStartPlaying(const FQualifiedFrameTime& InStartTime)
 	{
-		DeltaTime   = FFrameTime(0);
-		CurrentTime = FFrameTime(-1);
+		DeltaTime = FFrameTime(0);
+		StartTimeOffset = FFrameTime(0);
 	} 
 
 	virtual FFrameTime OnRequestCurrentTime(const FQualifiedFrameTime& InCurrentTime, float InPlayRate) override
@@ -47,6 +43,7 @@ struct FMovieSceneTimeController_FrameStep : FMovieSceneTimeController
 		TOptional<FQualifiedFrameTime> StartTimeIfPlaying = GetPlaybackStartTime();
 		if (!StartTimeIfPlaying.IsSet())
 		{
+			UE_LOG(LogMovieSceneCapture, VeryVerbose, TEXT("FMovieSceneTimeController_FrameStep::OnRequestCurrentTime = Paused at frame %d subframe %f"), InCurrentTime.Time.FrameNumber.Value, InCurrentTime.Time.GetSubFrame());
 			return InCurrentTime.Time;
 		}
 		else
@@ -54,22 +51,23 @@ struct FMovieSceneTimeController_FrameStep : FMovieSceneTimeController
 			// Scale the delta time (should be one frame) by this frame's play rate, and add it to the current time offset
 			if (InPlayRate == 1.f)
 			{
-				CurrentTime += DeltaTime;
+				StartTimeOffset += DeltaTime;
 			}
 			else
 			{
-				CurrentTime += DeltaTime * InPlayRate;
+				StartTimeOffset += DeltaTime * InPlayRate;
 			}
 
 			DeltaTime = FFrameTime(0);
 
-			ensure(CurrentTime >= 0);
-			return StartTimeIfPlaying->ConvertTo(InCurrentTime.Rate) + CurrentTime;
+			FFrameTime NewTime = StartTimeIfPlaying->ConvertTo(InCurrentTime.Rate) + StartTimeOffset;
+			UE_LOG(LogMovieSceneCapture, VeryVerbose, TEXT("FMovieSceneTimeController_FrameStep::OnRequestCurrentTime = Playing at frame %d subframe %f"), NewTime.FrameNumber.Value, NewTime.GetSubFrame());
+			return NewTime;
 		}
 	}
 
 	FFrameTime DeltaTime;
-	FFrameTime CurrentTime;
+	FFrameTime StartTimeOffset;
 };
 
 UAutomatedLevelSequenceCapture::UAutomatedLevelSequenceCapture(const FObjectInitializer& Init)
@@ -115,6 +113,7 @@ void UAutomatedLevelSequenceCapture::AddFormatMappings(TMap<FString, FStringForm
 
 void UAutomatedLevelSequenceCapture::Initialize(TSharedPtr<FSceneViewport> InViewport, int32 PIEInstance)
 {
+	TimeController = MakeShared<FMovieSceneTimeController_FrameStep>();
 	Viewport = InViewport;
 
 	// Apply command-line overrides from parent class first. This needs to be called before setting up the capture strategy with the desired frame rate.
@@ -221,7 +220,7 @@ void UAutomatedLevelSequenceCapture::Initialize(TSharedPtr<FSceneViewport> InVie
 	{
 		// Ensure it doesn't loop (-1 is indefinite)
 		Actor->PlaybackSettings.LoopCount.Value = 0;
-		Actor->SequencePlayer->SetTimeController(MakeShared<FMovieSceneTimeController_FrameStep>());
+		Actor->SequencePlayer->SetTimeController(TimeController);
 		Actor->PlaybackSettings.bPauseAtEnd = true;
 		Actor->PlaybackSettings.bAutoPlay = false;
 
@@ -581,17 +580,25 @@ void UAutomatedLevelSequenceCapture::OnTick(float DeltaSeconds)
 	}
 
 	// Count down our warm up frames.
-	// The post increment is important - it ensures we capture the very first frame if there are no warm up frames,
-	// but correctly skip n frames if there are n warmup frames
-	if( CaptureState == ELevelSequenceCaptureState::WarmingUp && RemainingWarmUpFrames-- == 0)
+	if( CaptureState == ELevelSequenceCaptureState::WarmingUp)
 	{
-		// Start capturing - this will capture the *next* update from sequencer
-		CaptureState = ELevelSequenceCaptureState::FinishedWarmUp;
-		UpdateFrameState();
-		StartCapture();
+		// The post increment is important - it ensures we capture the very first frame if there are no warm up frames,
+		// but correctly skip n frames if there are n warmup frames
+		if (RemainingWarmUpFrames-- == 0)
+		{
+			// Start capturing - this will capture the *next* update from sequencer
+			CaptureState = ELevelSequenceCaptureState::FinishedWarmUp;
+			UpdateFrameState();
+			StartCapture();
+		}
+		else
+		{
+			// Move onto the next frame
+			TimeController->DeltaTime = FFrameTime(1);
+		}
 	}
 
-	if( bCapturing && !Actor->SequencePlayer->IsPlaying() )
+	if( bCapturing && !Actor->SequencePlayer->IsPlaying() && CaptureState != ELevelSequenceCaptureState::Paused )
 	{
 		++ShotIndex;
 
@@ -652,21 +659,11 @@ void UAutomatedLevelSequenceCapture::DelayBeforeWarmupFinished()
 
 void UAutomatedLevelSequenceCapture::PauseFinished()
 {
-	CaptureState = ELevelSequenceCaptureState::FinishedWarmUp;
-
-	if (CachedPlayRate.IsSet())
-	{
-		ALevelSequenceActor* Actor = LevelSequenceActor.Get();
+	ALevelSequenceActor* Actor = LevelSequenceActor.Get();
+	Actor->SequencePlayer->Play();
 
-		// Force an evaluation to capture this frame
-		Actor->SequencePlayer->JumpToFrame(Actor->SequencePlayer->GetCurrentTime().Time);
+	CaptureState = ELevelSequenceCaptureState::FinishedWarmUp;
 
-		// Continue playing forwards
-		Actor->SequencePlayer->SetPlayRate(CachedPlayRate.GetValue());
-		CachedPlayRate.Reset();
-	}
-	
-	
 	if (bIsAudioCapturePass)
 	{
 		UE_LOG(LogMovieSceneCapture, Log, TEXT("WarmUp pause finished. Resuming the capture of audio."));
@@ -718,8 +715,7 @@ void UAutomatedLevelSequenceCapture::SequenceUpdated(const UMovieSceneSequencePl
 				CaptureState = ELevelSequenceCaptureState::Paused;
 
 				Actor->GetWorld()->GetTimerManager().SetTimer(DelayTimer, FTimerDelegate::CreateUObject(this, &UAutomatedLevelSequenceCapture::PauseFinished), DelayBeforeShotWarmUp + DelayEveryFrame, false);
-				CachedPlayRate = Actor->SequencePlayer->GetPlayRate();
-				Actor->SequencePlayer->SetPlayRate(0.f);
+				Actor->SequencePlayer->Pause();
 			}
 			else if (CaptureState == ELevelSequenceCaptureState::FinishedWarmUp)
 			{
@@ -761,6 +757,10 @@ void UAutomatedLevelSequenceCapture::SequenceUpdated(const UMovieSceneSequencePl
 					}
 					Actor->SequencePlayer->OnSequenceUpdated().Remove(OnPlayerUpdatedBinding);
 				}
+
+
+				// Move onto the next frame now that we've captured this one
+				TimeController->DeltaTime = FFrameTime(1);
 			}
 		}
 	}
diff --git a/Engine/Source/Editor/MovieSceneTools/Public/AutomatedLevelSequenceCapture.h b/Engine/Source/Editor/MovieSceneTools/Public/AutomatedLevelSequenceCapture.h
--- a/Engine/Source/Editor/MovieSceneTools/Public/AutomatedLevelSequenceCapture.h
+++ b/Engine/Source/Editor/MovieSceneTools/Public/AutomatedLevelSequenceCapture.h
@@ -13,7 +13,7 @@ class ALevelSequenceActor;
 class FJsonObject;
 class FSceneViewport;
 class ULevelSequenceBurnInOptions;
-
+struct FMovieSceneTimeController_FrameStep;
 
 UCLASS(config=EditorPerProjectUserSettings, PerObjectConfig)
 class MOVIESCENETOOLS_API UAutomatedLevelSequenceCapture : public UMovieSceneCapture
@@ -172,8 +172,6 @@ private:
 
 	FLevelSequencePlayerSnapshot CachedState;
 
-	TOptional<float> CachedPlayRate;
-
 	FTimerHandle DelayTimer;
 
 	struct FCinematicShotCache
@@ -199,5 +197,7 @@ private:
 	TOptional<FFrameNumber> CachedEndFrame;
 	TOptional<bool> bCachedUseCustomStartFrame;
 	TOptional<bool> bCachedUseCustomEndFrame;
+
+	TSharedPtr<FMovieSceneTimeController_FrameStep> TimeController;
 };
 
