Movie Rendering - Flush level streaming each frame during movie capture to ensure streaming levels are fully loaded to avoid pop-in mid shot.

diff --git a/Engine/Source/Editor/MovieSceneTools/Private/AutomatedLevelSequenceCapture.cpp b/Engine/Source/Editor/MovieSceneTools/Private/AutomatedLevelSequenceCapture.cpp
--- a/Engine/Source/Editor/MovieSceneTools/Private/AutomatedLevelSequenceCapture.cpp
+++ b/Engine/Source/Editor/MovieSceneTools/Private/AutomatedLevelSequenceCapture.cpp
@@ -538,6 +538,14 @@ void UAutomatedLevelSequenceCapture::OnTick(float DeltaSeconds)
 		return;
 	}
 
+	// Flush the level streaming system. This would cause hitches under normal gameplay, but because we already run at slower-than-real-time
+	// it doesn't matter for movie captures. This solves situations where games have systems that pull in sublevels via level streaming that cannot
+	// be normally controlled via the Sequencer Level Visibility track. If all levels are already loaded, flushing will have no effect.
+	if (GetWorld())
+	{
+		GetWorld()->FlushLevelStreaming(EFlushLevelStreamingType::Full);
+	}
+
 	// Setup the automated capture
 	if (CaptureState == ELevelSequenceCaptureState::Setup)
 	{
