Sequencer: Add shot name for movie scene capture

diff --git a/Engine/Source/Editor/MovieSceneTools/Private/AutomatedLevelSequenceCapture.cpp b/Engine/Source/Editor/MovieSceneTools/Private/AutomatedLevelSequenceCapture.cpp
--- a/Engine/Source/Editor/MovieSceneTools/Private/AutomatedLevelSequenceCapture.cpp
+++ b/Engine/Source/Editor/MovieSceneTools/Private/AutomatedLevelSequenceCapture.cpp
@@ -70,6 +70,34 @@ struct FMovieSceneTimeController_FrameStep : FMovieSceneTimeController
 	FFrameTime StartTimeOffset;
 };
 
+UMovieScene* GetMovieScene(TWeakObjectPtr<ALevelSequenceActor> LevelSequenceActor)
+{
+	ALevelSequenceActor* Actor = LevelSequenceActor.Get();
+	if (!Actor)
+	{
+		return nullptr;
+	}
+
+	ULevelSequence* LevelSequence = Cast<ULevelSequence>(Actor->LevelSequence.TryLoad());
+	if (!LevelSequence)
+	{
+		return nullptr;
+	}
+
+	return LevelSequence->GetMovieScene();
+}
+
+UMovieSceneCinematicShotTrack* GetCinematicShotTrack(TWeakObjectPtr<ALevelSequenceActor> LevelSequenceActor)
+{
+	UMovieScene* MovieScene = GetMovieScene(LevelSequenceActor);
+	if (!MovieScene)
+	{
+		return nullptr;
+	}
+
+	return MovieScene->FindMasterTrack<UMovieSceneCinematicShotTrack>();
+}
+
 UAutomatedLevelSequenceCapture::UAutomatedLevelSequenceCapture(const FObjectInitializer& Init)
 	: Super(Init)
 {
@@ -122,9 +150,15 @@ void UAutomatedLevelSequenceCapture::Initialize(TSharedPtr<FSceneViewport> InVie
 	// Apply command-line overrides
 	{
 		FString LevelSequenceAssetPath;
-		if( FParse::Value( FCommandLine::Get(), TEXT( "-LevelSequence=" ), LevelSequenceAssetPath ) )
+		if (FParse::Value(FCommandLine::Get(), TEXT("-LevelSequence="), LevelSequenceAssetPath))
 		{
-			LevelSequenceAsset.SetPath( LevelSequenceAssetPath );
+			LevelSequenceAsset.SetPath(LevelSequenceAssetPath);
+		}
+
+		FString ShotNameOverride;
+		if (FParse::Value(FCommandLine::Get(), TEXT("-Shot="), ShotNameOverride))
+		{
+			ShotName = ShotNameOverride;
 		}
 
 		int32 StartFrameOverride;
@@ -213,6 +247,50 @@ void UAutomatedLevelSequenceCapture::Initialize(TSharedPtr<FSceneViewport> InVie
 		}
 	}
 
+	if (Actor && !ShotName.IsEmpty())
+	{
+		UMovieScene* MovieScene = GetMovieScene(Actor);
+
+		UMovieSceneCinematicShotTrack* CinematicShotTrack = GetCinematicShotTrack(Actor);
+
+		if (CinematicShotTrack && MovieScene)
+		{
+			FFrameRate TickResolution = MovieScene->GetTickResolution();
+			FFrameRate DisplayRate = MovieScene->GetDisplayRate();
+
+			UMovieSceneCinematicShotSection* CinematicShotSection = nullptr;
+			for (UMovieSceneSection* Section : CinematicShotTrack->GetAllSections())
+			{
+				if (UMovieSceneCinematicShotSection* ShotSection = Cast<UMovieSceneCinematicShotSection>(Section))
+				{
+					if (ShotSection->GetShotDisplayName() == ShotName)
+					{
+						CinematicShotSection = ShotSection;
+						break;
+					}
+				}
+			}
+
+			if (CinematicShotSection)
+			{
+				bUseCustomStartFrame = true;
+				bUseCustomEndFrame = true;
+
+				FFrameNumber StartFrame = MovieScene::DiscreteInclusiveLower(CinematicShotSection->GetRange());
+				FFrameNumber EndFrame = MovieScene::DiscreteExclusiveUpper(CinematicShotSection->GetRange());
+
+				CustomStartFrame = FFrameRate::TransformTime(StartFrame, TickResolution, DisplayRate).CeilToFrame();
+				CustomEndFrame = FFrameRate::TransformTime(EndFrame, TickResolution, DisplayRate).CeilToFrame();
+				
+				UE_LOG(LogMovieSceneCapture, Log, TEXT("Found shot '%s' to capture. Setting custom start and end frames to %d and %d."), *ShotName, CustomStartFrame.Value, CustomEndFrame.Value);
+			}
+			else
+			{
+				UE_LOG(LogMovieSceneCapture, Error, TEXT("Could not find named shot '%s' to capture"), *ShotName);
+			}
+		}
+	}
+
 	ExportEDL();
 	ExportFCPXML();
 
@@ -262,34 +340,6 @@ void UAutomatedLevelSequenceCapture::Initialize(TSharedPtr<FSceneViewport> InVie
 	CaptureStrategy->OnInitialize();
 }
 
-UMovieScene* GetMovieScene(TWeakObjectPtr<ALevelSequenceActor> LevelSequenceActor)
-{
-	ALevelSequenceActor* Actor = LevelSequenceActor.Get();
-	if (!Actor)
-	{
-		return nullptr;
-	}
-
-	ULevelSequence* LevelSequence = Cast<ULevelSequence>( Actor->LevelSequence.TryLoad() );
-	if (!LevelSequence)
-	{
-		return nullptr;
-	}
-
-	return LevelSequence->GetMovieScene();
-}
-
-UMovieSceneCinematicShotTrack* GetCinematicShotTrack(TWeakObjectPtr<ALevelSequenceActor> LevelSequenceActor)
-{
-	UMovieScene* MovieScene = GetMovieScene(LevelSequenceActor);
-	if (!MovieScene)
-	{
-		return nullptr;
-	}
-
-	return MovieScene->FindMasterTrack<UMovieSceneCinematicShotTrack>();
-}
-
 bool UAutomatedLevelSequenceCapture::InitializeShots()
 {
 	NumShots = 0;
diff --git a/Engine/Source/Editor/MovieSceneTools/Public/AutomatedLevelSequenceCapture.h b/Engine/Source/Editor/MovieSceneTools/Public/AutomatedLevelSequenceCapture.h
--- a/Engine/Source/Editor/MovieSceneTools/Public/AutomatedLevelSequenceCapture.h
+++ b/Engine/Source/Editor/MovieSceneTools/Public/AutomatedLevelSequenceCapture.h
@@ -30,6 +30,10 @@ public:
 	UPROPERTY(BlueprintReadWrite, Category=Animation)
 	FSoftObjectPath LevelSequenceAsset;
 
+	/** Optional shot name to render. The frame range to render will be set to the shot frame range. */
+	UPROPERTY(BlueprintReadWrite, Category=Animation)
+	FString ShotName;
+
 	/** When enabled, the StartFrame setting will override the default starting frame number */
 	UPROPERTY(config, EditAnywhere, BlueprintReadWrite, Category=Animation, AdvancedDisplay)
 	bool bUseCustomStartFrame;
