Sequencer: Add -WriteFinalCutProXML and -WriteEditDecisionList commandline options

diff --git a/Engine/Source/Editor/MovieSceneTools/Private/AutomatedLevelSequenceCapture.cpp b/Engine/Source/Editor/MovieSceneTools/Private/AutomatedLevelSequenceCapture.cpp
--- a/Engine/Source/Editor/MovieSceneTools/Private/AutomatedLevelSequenceCapture.cpp
+++ b/Engine/Source/Editor/MovieSceneTools/Private/AutomatedLevelSequenceCapture.cpp
@@ -198,6 +198,18 @@ void UAutomatedLevelSequenceCapture::Initialize(TSharedPtr<FSceneViewport> InVie
 		{
 			DelayEveryFrame = DelayEveryFrameOverride;
 		}
+
+		bool bWriteEditDecisionListOverride;
+		if (FParse::Bool(FCommandLine::Get(), TEXT("-WriteEditDecisionList="), bWriteEditDecisionListOverride))
+		{
+			bWriteEditDecisionList = bWriteEditDecisionListOverride;
+		}
+
+		bool bWriteFinalCutProXMLOverride;
+		if (FParse::Bool(FCommandLine::Get(), TEXT("-WriteFinalCutProXML="), bWriteFinalCutProXMLOverride))
+		{
+			bWriteFinalCutProXML = bWriteFinalCutProXMLOverride;
+		}
 	}
 
 	if (Settings.bUsePathTracer)
diff --git a/Engine/Source/Editor/MovieSceneTools/Public/AutomatedLevelSequenceCapture.h b/Engine/Source/Editor/MovieSceneTools/Public/AutomatedLevelSequenceCapture.h
--- a/Engine/Source/Editor/MovieSceneTools/Public/AutomatedLevelSequenceCapture.h
+++ b/Engine/Source/Editor/MovieSceneTools/Public/AutomatedLevelSequenceCapture.h
@@ -71,7 +71,7 @@ public:
 	bool bWriteEditDecisionList;
 
 	/** Whether to write Final Cut Pro XML files (XMLs) if the sequence contains shots */
-	UPROPERTY(config, EditAnywhere, Category=Sequence)
+	UPROPERTY(config, EditAnywhere, BlueprintReadWrite, Category=Sequence)
 	bool bWriteFinalCutProXML;
 
 public:
