MovieSceneCapture: Add {sequence} to movie scene capture output

diff --git a/Engine/Source/Editor/MovieSceneTools/Private/AutomatedLevelSequenceCapture.cpp b/Engine/Source/Editor/MovieSceneTools/Private/AutomatedLevelSequenceCapture.cpp
--- a/Engine/Source/Editor/MovieSceneTools/Private/AutomatedLevelSequenceCapture.cpp
+++ b/Engine/Source/Editor/MovieSceneTools/Private/AutomatedLevelSequenceCapture.cpp
@@ -130,6 +130,8 @@ UAutomatedLevelSequenceCapture::UAutomatedLevelSequenceCapture(const FObjectInit
 #if WITH_EDITORONLY_DATA
 void UAutomatedLevelSequenceCapture::AddFormatMappings(TMap<FString, FStringFormatArg>& OutFormatMappings, const FFrameMetrics& FrameMetrics) const
 {
+	OutFormatMappings.Add(TEXT("sequence"), CachedState.MasterName);
+
 	OutFormatMappings.Add(TEXT("shot"), CachedState.CurrentShotName);
 	OutFormatMappings.Add(TEXT("shot_frame"), FString::Printf(TEXT("%0*d"), Settings.ZeroPadFrameNumbers, CachedState.CurrentShotLocalTime.Time.FrameNumber.Value));
 
diff --git a/Engine/Source/Runtime/MovieSceneCapture/Public/MovieSceneCaptureSettings.h b/Engine/Source/Runtime/MovieSceneCapture/Public/MovieSceneCaptureSettings.h
--- a/Engine/Source/Runtime/MovieSceneCapture/Public/MovieSceneCaptureSettings.h
+++ b/Engine/Source/Runtime/MovieSceneCapture/Public/MovieSceneCaptureSettings.h
@@ -50,6 +50,7 @@ struct MOVIESCENECAPTURE_API FMovieSceneCaptureSettings
 	 * {quality}	- The image compression quality setting
 	 * {material}   - The material/render pass
 	 * {shot}       - The name of the level sequence asset shot being played
+	 * {sequence}   - The name of the level sequence asset (ie. master) being played
 	 * {camera}     - The name of the current camera
 	 */
 	UPROPERTY(config, EditAnywhere, BlueprintReadWrite, Category=General, DisplayName="Filename Format")
