Movie Scene Capture: BlueprintType

diff --git a/Engine/Source/Editor/MovieSceneTools/Public/AutomatedLevelSequenceCapture.h b/Engine/Source/Editor/MovieSceneTools/Public/AutomatedLevelSequenceCapture.h
--- a/Engine/Source/Editor/MovieSceneTools/Public/AutomatedLevelSequenceCapture.h
+++ b/Engine/Source/Editor/MovieSceneTools/Public/AutomatedLevelSequenceCapture.h
@@ -15,7 +15,7 @@ class FSceneViewport;
 class ULevelSequenceBurnInOptions;
 struct FMovieSceneTimeController_FrameStep;
 
-UCLASS(config=EditorPerProjectUserSettings, PerObjectConfig)
+UCLASS(config=EditorPerProjectUserSettings, PerObjectConfig, BlueprintType)
 class MOVIESCENETOOLS_API UAutomatedLevelSequenceCapture : public UMovieSceneCapture
 {
 public:
