FCPXML: Look for level sequence asset without extension.

Add more log errors for clip or audio asset not found.

diff --git a/Engine/Source/Editor/MovieSceneTools/Private/MovieSceneTranslator.cpp b/Engine/Source/Editor/MovieSceneTools/Private/MovieSceneTranslator.cpp
--- a/Engine/Source/Editor/MovieSceneTools/Private/MovieSceneTranslator.cpp
+++ b/Engine/Source/Editor/MovieSceneTools/Private/MovieSceneTranslator.cpp
@@ -654,8 +654,22 @@ TSharedPtr<FMovieSceneImportCinematicSectionData> FMovieSceneImportData::CreateC
 		}
 	}
 
+	// If exact match wasn't found, look for the base filename (in case the clip name comes with an extension)
 	if (SequenceToAdd == nullptr)
 	{
+		for (FAssetData AssetData : AssetDataArray)
+		{
+			if (AssetData.AssetName == *FPaths::GetBaseFilename(InName))
+			{
+				SequenceToAdd = Cast<ULevelSequence>(AssetData.GetAsset());
+				break;
+			}
+		}
+	}
+
+	if (SequenceToAdd == nullptr)
+	{
+		UE_LOG(LogMovieScene, Error, TEXT("Can't find valid level sequence asset to map clip to: (%s)"), *InName);
 		return nullptr;
 	}
 
@@ -811,6 +825,7 @@ TSharedPtr<FMovieSceneImportAudioSectionData> FMovieSceneImportData::CreateAudio
 	
 	if (SoundToAdd == nullptr)
 	{
+		UE_LOG(LogMovieScene, Error, TEXT("Can't find valid asset asset to map clip to: (%s)"), *InFilenameOrAssetPathName);
 		return nullptr;
 	}
 
