Sequencer: Add check for active subsequence when getting subsequence IDs

diff --git a/Engine/Source/Editor/MovieSceneTools/Private/MovieSceneObjectBindingIDPicker.cpp b/Engine/Source/Editor/MovieSceneTools/Private/MovieSceneObjectBindingIDPicker.cpp
--- a/Engine/Source/Editor/MovieSceneTools/Private/MovieSceneObjectBindingIDPicker.cpp
+++ b/Engine/Source/Editor/MovieSceneTools/Private/MovieSceneObjectBindingIDPicker.cpp
@@ -213,7 +213,7 @@ private:
 				{
 					UMovieSceneSubSection* SubSection = Cast<UMovieSceneSubSection>(Section);
 					UMovieSceneSequence* SubSequence = SubSection ? SubSection->GetSequence() : nullptr;
-					if (SubSequence)
+					if (SubSequence && Section->IsActive())
 					{
 						// Hold onto the current parent ID before adding our ID onto the stack
 						FMovieSceneSequenceID ParentID = SequenceIDStack.GetCurrent();
