Sequencer: Fix vector property track editor improperly setting values

diff --git a/Engine/Source/Editor/MovieSceneTools/Private/TrackEditors/PropertyTrackEditors/VectorPropertyTrackEditor.cpp b/Engine/Source/Editor/MovieSceneTools/Private/TrackEditors/PropertyTrackEditors/VectorPropertyTrackEditor.cpp
--- a/Engine/Source/Editor/MovieSceneTools/Private/TrackEditors/PropertyTrackEditors/VectorPropertyTrackEditor.cpp
+++ b/Engine/Source/Editor/MovieSceneTools/Private/TrackEditors/PropertyTrackEditors/VectorPropertyTrackEditor.cpp
@@ -173,7 +173,7 @@ bool FVectorPropertyTrackEditor::ModifyGeneratedKeysByCurrentAndWeight(UObject *
 				FMovieSceneChannelProxy& Proxy = SectionToKey->GetChannelProxy();
 				GeneratedTotalKeys[0]->ModifyByCurrentAndWeight(Proxy, KeyTime, (void *)&Val.X, Weight);
 				GeneratedTotalKeys[1]->ModifyByCurrentAndWeight(Proxy, KeyTime, (void *)&Val.Y, Weight);
-				GeneratedTotalKeys[1]->ModifyByCurrentAndWeight(Proxy, KeyTime, (void *)&Val.Z, Weight);
+				GeneratedTotalKeys[2]->ModifyByCurrentAndWeight(Proxy, KeyTime, (void *)&Val.Z, Weight);
 			}
 			break;
 		case 4:
@@ -187,8 +187,8 @@ bool FVectorPropertyTrackEditor::ModifyGeneratedKeysByCurrentAndWeight(UObject *
 				FMovieSceneChannelProxy& Proxy = SectionToKey->GetChannelProxy();
 				GeneratedTotalKeys[0]->ModifyByCurrentAndWeight(Proxy, KeyTime, (void *)&Val.X, Weight);
 				GeneratedTotalKeys[1]->ModifyByCurrentAndWeight(Proxy, KeyTime, (void *)&Val.Y, Weight);
-				GeneratedTotalKeys[1]->ModifyByCurrentAndWeight(Proxy, KeyTime, (void *)&Val.Z, Weight);
-				GeneratedTotalKeys[1]->ModifyByCurrentAndWeight(Proxy, KeyTime, (void *)&Val.W, Weight);
+				GeneratedTotalKeys[2]->ModifyByCurrentAndWeight(Proxy, KeyTime, (void *)&Val.Z, Weight);
+				GeneratedTotalKeys[3]->ModifyByCurrentAndWeight(Proxy, KeyTime, (void *)&Val.W, Weight);
 
 			}
 			break;
