Sequencer: Search through all levels (persistent and streaming) to find actors with the given binding to export

diff --git a/Engine/Source/Editor/MovieSceneTools/Private/MovieSceneToolHelpers.cpp b/Engine/Source/Editor/MovieSceneTools/Private/MovieSceneToolHelpers.cpp
--- a/Engine/Source/Editor/MovieSceneTools/Private/MovieSceneToolHelpers.cpp
+++ b/Engine/Source/Editor/MovieSceneTools/Private/MovieSceneToolHelpers.cpp
@@ -56,6 +56,7 @@
 #include "Tracks/MovieSceneCinematicShotTrack.h"
 #include "Tracks/MovieSceneCameraCutTrack.h"
 #include "Sections/MovieSceneCameraCutSection.h"
+#include "Engine/LevelStreaming.h"
 #include "FbxExporter.h"
 
 
@@ -2055,25 +2056,22 @@ FMovieSceneEvaluationTrack* MovieSceneToolHelpers::GetEvaluationTrack(ISequencer
 	return nullptr;
 }
 
-bool MovieSceneToolHelpers::ExportFBX(UWorld* World, UMovieScene* MovieScene, IMovieScenePlayer* Player, TArray<FGuid>& Bindings, INodeNameAdapter& NodeNameAdapter, FMovieSceneSequenceIDRef& Template, const FString& InFBXFileName)
+void ExportLevelMesh(UnFbx::FFbxExporter* Exporter, ULevel* Level, IMovieScenePlayer* Player, TArray<FGuid>& Bindings, INodeNameAdapter& NodeNameAdapter, FMovieSceneSequenceIDRef& Template)
 {
-
-	UnFbx::FFbxExporter* Exporter = UnFbx::FFbxExporter::GetInstance();
-
-	Exporter->CreateDocument();
-	Exporter->SetTrasformBaking(false);
-	Exporter->SetKeepHierarchy(true);
-
 	// Get list of actors based upon bindings...
 	const bool bSelectedOnly = (Bindings.Num()) != 0;
+
+	const bool bSaveAnimSeq = false; //force off saving any AnimSequences since this can conflict when we export the level sequence animations.
+
 	TArray<AActor*> ActorToExport;
-	int32 ActorCount = World->PersistentLevel->Actors.Num();
+
+	int32 ActorCount = Level->Actors.Num();
 	for (int32 ActorIndex = 0; ActorIndex < ActorCount; ++ActorIndex)
 	{
-		AActor* Actor = World->PersistentLevel->Actors[ActorIndex];
+		AActor* Actor = Level->Actors[ActorIndex];
 		if (Actor != NULL)
 		{
-			FGuid ExistingGuid = Player->FindObjectId(*Actor, MovieSceneSequenceID::Root);
+			FGuid ExistingGuid = Player->FindObjectId(*Actor, Template);
 			if (ExistingGuid.IsValid() && (!bSelectedOnly || Bindings.Contains(ExistingGuid)))
 			{
 				ActorToExport.Add(Actor);
@@ -2082,16 +2080,28 @@ bool MovieSceneToolHelpers::ExportFBX(UWorld* World, UMovieScene* MovieScene, IM
 	}
 
 	// Export the persistent level and all of it's actors
-	const bool bSaveAnimSeq = false; //force off saving any AnimSequences since this can conflict when we export the level sequence animations.
-	Exporter->ExportLevelMesh(World->PersistentLevel, !bSelectedOnly, ActorToExport, NodeNameAdapter, bSaveAnimSeq);
+	Exporter->ExportLevelMesh(Level, !bSelectedOnly, ActorToExport, NodeNameAdapter, bSaveAnimSeq);
+}
+
+bool MovieSceneToolHelpers::ExportFBX(UWorld* World, UMovieScene* MovieScene, IMovieScenePlayer* Player, TArray<FGuid>& Bindings, INodeNameAdapter& NodeNameAdapter, FMovieSceneSequenceIDRef& Template, const FString& InFBXFileName)
+{
+	UnFbx::FFbxExporter* Exporter = UnFbx::FFbxExporter::GetInstance();
+
+	Exporter->CreateDocument();
+	Exporter->SetTrasformBaking(false);
+	Exporter->SetKeepHierarchy(true);
+
+	ExportLevelMesh(Exporter, World->PersistentLevel, Player, Bindings, NodeNameAdapter, Template);
 
 	// Export streaming levels and actors
-	for (int32 CurLevelIndex = 0; CurLevelIndex < World->GetNumLevels(); ++CurLevelIndex)
+	for (ULevelStreaming* StreamingLevel : World->GetStreamingLevels())
 	{
-		ULevel* CurLevel = World->GetLevel(CurLevelIndex);
-		if (CurLevel != NULL && CurLevel != (World->PersistentLevel))
+		if (StreamingLevel)
 		{
-			Exporter->ExportLevelMesh(CurLevel, !bSelectedOnly, ActorToExport, NodeNameAdapter, bSaveAnimSeq);
+			if (ULevel* Level = StreamingLevel->GetLoadedLevel())
+			{
+				ExportLevelMesh(Exporter, Level, Player, Bindings, NodeNameAdapter, Template);
+			}
 		}
 	}
 
