Sequencer: Added an option to disable superfluous playback data in cooked builds for sequences that are not playable directly

  - The option is exposed on the right click menu for sub sequences and when disabled will cause the sequence's playback data to be pruned at cook time potentially saving memory.

diff --git a/Engine/Source/Editor/MovieSceneTools/Private/TrackEditors/SubTrackEditor.cpp b/Engine/Source/Editor/MovieSceneTools/Private/TrackEditors/SubTrackEditor.cpp
--- a/Engine/Source/Editor/MovieSceneTools/Private/TrackEditors/SubTrackEditor.cpp
+++ b/Engine/Source/Editor/MovieSceneTools/Private/TrackEditors/SubTrackEditor.cpp
@@ -315,6 +315,78 @@ public:
 			LOCTEXT("TakesMenu", "Takes"),
 			LOCTEXT("TakesMenuTooltip", "Sub section takes"),
 			FNewMenuDelegate::CreateLambda([=](FMenuBuilder& InMenuBuilder){ AddTakesMenu(InMenuBuilder); }));
+
+		MenuBuilder.AddMenuEntry(
+			LOCTEXT("PlayableDirectly_Label", "Playable Directly"),
+			LOCTEXT("PlayableDirectly_Tip", "When enabled, this sequence will also support being played directly outside of the master sequence. Disable this to save some memory on complex hierarchies of sequences."),
+			FSlateIcon(),
+			FUIAction(
+				FExecuteAction::CreateRaw(this, &FSubSection::TogglePlayableDirectly),
+				FCanExecuteAction::CreateLambda([]{ return true; }),
+				FGetActionCheckState::CreateRaw(this, &FSubSection::IsPlayableDirectly)
+			),
+			NAME_None,
+			EUserInterfaceActionType::ToggleButton
+		);
+	}
+
+	void TogglePlayableDirectly()
+	{
+		TSharedPtr<ISequencer> SequencerPtr = Sequencer.Pin();
+		if (SequencerPtr)
+		{
+			FScopedTransaction Transaction(LOCTEXT("SetPlayableDirectly_Transaction", "Set Playable Directly"));
+
+			TArray<UMovieSceneSection*> SelectedSections;
+			SequencerPtr->GetSelectedSections(SelectedSections);
+
+			const bool bNewPlayableDirectly = IsPlayableDirectly() != ECheckBoxState::Checked;
+
+			for (UMovieSceneSection* Section : SelectedSections)
+			{
+				if (UMovieSceneSubSection* SubSection = Cast<UMovieSceneSubSection>(Section))
+				{
+					UMovieSceneSequence* Sequence = SubSection->GetSequence();
+					if (Sequence->IsPlayableDirectly() != bNewPlayableDirectly)
+					{
+						Sequence->SetPlayableDirectly(bNewPlayableDirectly);
+					}
+				}
+			}
+		}
+	}
+
+	ECheckBoxState IsPlayableDirectly() const
+	{
+		ECheckBoxState CheckboxState = ECheckBoxState::Undetermined;
+
+		TSharedPtr<ISequencer> SequencerPtr = Sequencer.Pin();
+		if (SequencerPtr)
+		{
+			TArray<UMovieSceneSection*> SelectedSections;
+			SequencerPtr->GetSelectedSections(SelectedSections);
+
+			for (UMovieSceneSection* Section : SelectedSections)
+			{
+				if (UMovieSceneSubSection* SubSection = Cast<UMovieSceneSubSection>(Section))
+				{
+					UMovieSceneSequence* Sequence = SubSection->GetSequence();
+					if (Sequence)
+					{
+						if (CheckboxState == ECheckBoxState::Undetermined)
+						{
+							CheckboxState = Sequence->IsPlayableDirectly() ? ECheckBoxState::Checked : ECheckBoxState::Unchecked;
+						}
+						else if (CheckboxState == ECheckBoxState::Checked != Sequence->IsPlayableDirectly())
+						{
+							return ECheckBoxState::Undetermined;
+						}
+					}
+				}
+			}
+		}
+
+		return CheckboxState;
 	}
 
 	void BeginResizeSection()
diff --git a/Engine/Source/Runtime/MovieScene/Private/MovieSceneSequence.cpp b/Engine/Source/Runtime/MovieScene/Private/MovieSceneSequence.cpp
--- a/Engine/Source/Runtime/MovieScene/Private/MovieSceneSequence.cpp
+++ b/Engine/Source/Runtime/MovieScene/Private/MovieSceneSequence.cpp
@@ -14,6 +14,7 @@ UMovieSceneSequence::UMovieSceneSequence(const FObjectInitializer& Init)
 	: Super(Init)
 {
 	bParentContextsAreSignificant = false;
+	bPlayableDirectly = true;
 }
 
 #if WITH_EDITORONLY_DATA
@@ -48,6 +49,12 @@ void UMovieSceneSequence::PreSave(const ITargetPlatform* TargetPlatform)
 		{
 			FMovieSceneSequencePrecompiledTemplateStore Store;
 			FMovieSceneCompiler::Compile(*this, Store);
+
+			if (!bPlayableDirectly)
+			{
+				PrecompiledEvaluationTemplate.EvaluationField = FMovieSceneEvaluationField();
+				PrecompiledEvaluationTemplate.ResetFieldData();
+			}
 		}
 		else
 		{
diff --git a/Engine/Source/Runtime/MovieScene/Public/MovieSceneSequence.h b/Engine/Source/Runtime/MovieScene/Public/MovieSceneSequence.h
--- a/Engine/Source/Runtime/MovieScene/Public/MovieSceneSequence.h
+++ b/Engine/Source/Runtime/MovieScene/Public/MovieSceneSequence.h
@@ -52,7 +52,7 @@ public:
 	/**
 	 * Locate all the objects that correspond to the specified object ID, using the specified context
 	 *
-	 * @param ObjectId				The unique identifier of the object.
+ * @param ObjectId				The unique identifier of the object.
 	 * @param Context				Optional context to use to find the required object (for instance, a parent spawnable object)
 	 * @param OutObjects			Destination array to add found objects to
 	 */
@@ -208,6 +208,27 @@ public:
 		return bParentContextsAreSignificant;
 	}
 
+	/**
+	 * Check whether this sequence is playable directly outside of a master sub sequence or not
+	 *
+	 * @return True if this sequences cooked data will include all the necessary information to be played back on its own, false if this data is not present in cooked builds
+	 */
+	bool IsPlayableDirectly() const
+	{
+		return bPlayableDirectly;
+	}
+
+	/**
+	 * Assign whether this sequence is playable directly outside of a master sub sequence or not
+	 *
+	 * @param bInPlayableDirectly   When true, this sequence's cooked data will include all the necessary information to be played back on its own. When false this data will be culled resulting in less memory usage.
+	 */
+	void SetPlayableDirectly(bool bInPlayableDirectly)
+	{
+		Modify();
+		bPlayableDirectly = bInPlayableDirectly;
+	}
+
 protected:
 
 	/**
@@ -218,6 +239,12 @@ protected:
 	UPROPERTY()
 	bool bParentContextsAreSignificant;
 
+	/**
+	 * When true, this sequence should be compiled as if it is playable directly (outside of a master sequence). When false, various compiled data will be omitted, preventing direct playback at runtime (although will still play as a sub sequence)
+	 */
+	UPROPERTY()
+	bool bPlayableDirectly;
+
 public:
 
 #if WITH_EDITOR
