Sequencer: Change cached state camera component to soft object pointer

diff --git a/Engine/Source/Editor/MovieSceneTools/Private/AutomatedLevelSequenceCapture.cpp b/Engine/Source/Editor/MovieSceneTools/Private/AutomatedLevelSequenceCapture.cpp
--- a/Engine/Source/Editor/MovieSceneTools/Private/AutomatedLevelSequenceCapture.cpp
+++ b/Engine/Source/Editor/MovieSceneTools/Private/AutomatedLevelSequenceCapture.cpp
@@ -137,9 +137,13 @@ void UAutomatedLevelSequenceCapture::AddFormatMappings(TMap<FString, FStringForm
 	OutFormatMappings.Add(TEXT("shot"), CachedState.CurrentShotName);
 	OutFormatMappings.Add(TEXT("shot_frame"), FString::Printf(TEXT("%0*d"), Settings.ZeroPadFrameNumbers, CachedState.CurrentShotLocalTime.Time.FrameNumber.Value));
 
-	if (CachedState.CameraComponent && CachedState.CameraComponent->GetOwner())
+	if (CachedState.CameraComponent.IsValid())
 	{
-		OutFormatMappings.Add(TEXT("camera"), CachedState.CameraComponent->GetOwner()->GetName());
+		AActor* OuterActor = Cast<AActor>(CachedState.CameraComponent.Get()->GetOuter());
+		if (OuterActor)
+		{
+			OutFormatMappings.Add(TEXT("camera"), OuterActor->GetActorLabel());
+		}
 	}
 }
 
diff --git a/Engine/Source/Runtime/LevelSequence/Public/LevelSequencePlayer.h b/Engine/Source/Runtime/LevelSequence/Public/LevelSequencePlayer.h
--- a/Engine/Source/Runtime/LevelSequence/Public/LevelSequencePlayer.h
+++ b/Engine/Source/Runtime/LevelSequence/Public/LevelSequencePlayer.h
@@ -80,7 +80,7 @@ struct FLevelSequencePlayerSnapshot
 	FString SourceTimecode;
 
 	UPROPERTY(VisibleAnywhere, BlueprintReadOnly, Category="General")
-	UCameraComponent* CameraComponent;
+	TSoftObjectPtr<UCameraComponent> CameraComponent;
 
 	UPROPERTY(VisibleAnywhere, BlueprintReadOnly, Category = "General")
 	FLevelSequenceSnapshotSettings Settings;
