Added support to the Curve Editor for 'S' to key transform channels individually like in the main Sequencer. This makes the Sequencer Curve Editor commands available externally so you can add additional commands to the Sequencer-specific curve editor without knowing about the actual curve editor.

diff --git a/Engine/Source/Editor/CurveEditor/Private/CurveEditorCommands.cpp b/Engine/Source/Editor/CurveEditor/Private/CurveEditorCommands.cpp
--- a/Engine/Source/Editor/CurveEditor/Private/CurveEditorCommands.cpp
+++ b/Engine/Source/Editor/CurveEditor/Private/CurveEditorCommands.cpp
@@ -21,7 +21,7 @@ void FCurveEditorCommands::RegisterCommands()
 	UI_COMMAND(InterpolationCubicBreak, "Break", "Cubic interpolation - User broken tangents", EUserInterfaceActionType::ToggleButton, FInputChord(EKeys::Three));
 	UI_COMMAND(InterpolationToggleWeighted, "Weighted Tangents", "Toggle weighted tangents for cubic interpolation modes. Only supported on some curve types", EUserInterfaceActionType::ToggleButton, FInputChord(EModifierKey::Control, EKeys::W));
 
-	UI_COMMAND(FlattenTangents, "Flatten", "Flatten Tangents", EUserInterfaceActionType::Button, FInputChord(EKeys::S));
+	UI_COMMAND(FlattenTangents, "Flatten", "Flatten Tangents", EUserInterfaceActionType::Button, FInputChord(EKeys::Six));
 	UI_COMMAND(StraightenTangents, "Straighten", "Straighten Tangents", EUserInterfaceActionType::Button, FInputChord());
 
 	UI_COMMAND(BakeCurve, "Bake", "Bake curve", EUserInterfaceActionType::Button, FInputChord());
diff --git a/Engine/Source/Editor/MovieSceneTools/Private/TrackEditors/TransformTrackEditor.cpp b/Engine/Source/Editor/MovieSceneTools/Private/TrackEditors/TransformTrackEditor.cpp
--- a/Engine/Source/Editor/MovieSceneTools/Private/TrackEditors/TransformTrackEditor.cpp
+++ b/Engine/Source/Editor/MovieSceneTools/Private/TrackEditors/TransformTrackEditor.cpp
@@ -428,23 +428,30 @@ void F3DTransformTrackEditor::BindCommands(TSharedRef<FUICommandList> SequencerC
 {
 	const F3DTransformTrackCommands& Commands = F3DTransformTrackCommands::Get();
 
-	SequencerCommandBindings->MapAction(
+	CommandBindings = MakeShared<FUICommandList>();
+	CommandBindings->MapAction(
 		Commands.AddTransformKey,
 		FExecuteAction::CreateSP( this, &F3DTransformTrackEditor::OnAddTransformKeysForSelectedObjects, EMovieSceneTransformChannel::All ) );
 
-	SequencerCommandBindings->MapAction(
+	CommandBindings->MapAction(
 		Commands.AddTranslationKey,
 		FExecuteAction::CreateSP( this, &F3DTransformTrackEditor::OnAddTransformKeysForSelectedObjects, EMovieSceneTransformChannel::Translation ) );
 
-	SequencerCommandBindings->MapAction(
+	CommandBindings->MapAction(
 		Commands.AddRotationKey,
 		FExecuteAction::CreateSP( this, &F3DTransformTrackEditor::OnAddTransformKeysForSelectedObjects, EMovieSceneTransformChannel::Rotation ) );
 
-	SequencerCommandBindings->MapAction(
+	CommandBindings->MapAction(
 		Commands.AddScaleKey,
 		FExecuteAction::CreateSP( this, &F3DTransformTrackEditor::OnAddTransformKeysForSelectedObjects, EMovieSceneTransformChannel::Scale ) );
 
 	Commands.BindingCount++;
+
+	// Add these bindings to Sequencer
+	SequencerCommandBindings->Append(CommandBindings.ToSharedRef());
+	
+	// Also add them to the Curve Editor 
+	GetSequencer()->GetCommandBindings(ESequencerCommandBindings::CurveEditor)->Append(CommandBindings.ToSharedRef());
 }
 
 
diff --git a/Engine/Source/Editor/MovieSceneTools/Private/TrackEditors/TransformTrackEditor.h b/Engine/Source/Editor/MovieSceneTools/Private/TrackEditors/TransformTrackEditor.h
--- a/Engine/Source/Editor/MovieSceneTools/Private/TrackEditors/TransformTrackEditor.h
+++ b/Engine/Source/Editor/MovieSceneTools/Private/TrackEditors/TransformTrackEditor.h
@@ -167,6 +167,9 @@ private:
 	/** Mapping of objects to their existing transform data (for comparing against new transform data) */
 	TMap< TWeakObjectPtr<UObject>, FTransformData > ObjectToExistingTransform;
 
+	/** Command Bindings added by the Transform Track Editor to Sequencer and curve editor. */
+	TSharedPtr<FUICommandList> CommandBindings;
+
 private:
 	/** 
 	 * Modify the passed in Generated Keys by the current tracks values and weight at the passed in time.
diff --git a/Engine/Source/Editor/Sequencer/Private/Sequencer.h b/Engine/Source/Editor/Sequencer/Private/Sequencer.h
--- a/Engine/Source/Editor/Sequencer/Private/Sequencer.h
+++ b/Engine/Source/Editor/Sequencer/Private/Sequencer.h
@@ -384,6 +384,10 @@ public:
 		{
 			return SequencerCommandBindings;
 		}
+		else if (Type == ESequencerCommandBindings::CurveEditor)
+		{
+			return CurveEditorSharedBindings;
+		}
 
 		return SequencerSharedBindings;
 	}
diff --git a/Engine/Source/Editor/Sequencer/Public/ISequencer.h b/Engine/Source/Editor/Sequencer/Public/ISequencer.h
--- a/Engine/Source/Editor/Sequencer/Public/ISequencer.h
+++ b/Engine/Source/Editor/Sequencer/Public/ISequencer.h
@@ -92,7 +92,10 @@ enum class ESequencerCommandBindings
 	Sequencer,
 
 	/** Bindings that are shared between Sequencer and non-Sequencer widgets (subset of Sequencer commands). */
-	Shared
+	Shared,
+
+	/** Bindings that are available in the Curve Editor. */
+	CurveEditor
 };
 
 
