Fixed adding keys to a curve model sometimes causing tangents to get incorrectly calculated. This was causing issues with Buffered curves.

diff --git a/Engine/Source/Editor/CurveEditor/Private/RichCurveEditorModel.cpp b/Engine/Source/Editor/CurveEditor/Private/RichCurveEditorModel.cpp
--- a/Engine/Source/Editor/CurveEditor/Private/RichCurveEditorModel.cpp
+++ b/Engine/Source/Editor/CurveEditor/Private/RichCurveEditorModel.cpp
@@ -41,6 +41,9 @@ void FRichCurveEditorModel::AddKeys(TArrayView<const FKeyPosition> InKeyPosition
 	{
 		Owner->Modify();
 
+		TArray<FKeyHandle> NewKeyHandles;
+		NewKeyHandles.SetNumUninitialized(InKeyPositions.Num());
+
 		for (int32 Index = 0; Index < InKeyPositions.Num(); ++Index)
 		{
 			FKeyPosition   Position   = InKeyPositions[Index];
@@ -49,20 +52,16 @@ void FRichCurveEditorModel::AddKeys(TArrayView<const FKeyPosition> InKeyPosition
 			FKeyHandle     NewHandle = RichCurve->AddKey(Position.InputValue, Position.OutputValue);
 			FRichCurveKey* NewKey    = &RichCurve->GetKey(NewHandle);
 
-			if (Attributes.HasInterpMode())          { NewKey->InterpMode          = Attributes.GetInterpMode();    }
-			if (Attributes.HasTangentMode())         { NewKey->TangentMode         = Attributes.GetTangentMode();   }
-			if (Attributes.HasArriveTangent())       { NewKey->ArriveTangent       = Attributes.GetArriveTangent(); }
-			if (Attributes.HasLeaveTangent())        { NewKey->LeaveTangent        = Attributes.GetLeaveTangent();  }
-			if (Attributes.HasArriveTangentWeight()) { NewKey->ArriveTangentWeight = Attributes.GetArriveTangent(); }
-			if (Attributes.HasLeaveTangentWeight())  { NewKey->LeaveTangentWeight  = Attributes.GetLeaveTangent();  }
-
+			NewKeyHandles[Index] = NewHandle;
 			if (OutKeyHandles)
 			{
 				(*OutKeyHandles)[Index] = NewHandle;
 			}
 		}
 
-		RichCurve->AutoSetTangents();
+		// We reuse SetKeyAttributes here as there is complex logic determining which parts of the attributes are valid to pass on.
+		// For now we need to duplicate the new key handle array due to API mismatch. This will auto-calculate tangents if required.
+		SetKeyAttributes(NewKeyHandles, InKeyAttributes);
 	}
 }
 
diff --git a/Engine/Source/Editor/MovieSceneTools/Private/Channels/FloatChannelCurveModel.cpp b/Engine/Source/Editor/MovieSceneTools/Private/Channels/FloatChannelCurveModel.cpp
--- a/Engine/Source/Editor/MovieSceneTools/Private/Channels/FloatChannelCurveModel.cpp
+++ b/Engine/Source/Editor/MovieSceneTools/Private/Channels/FloatChannelCurveModel.cpp
@@ -57,29 +57,31 @@ void FFloatChannelCurveModel::AddKeys(TArrayView<const FKeyPosition> InKeyPositi
 		TMovieSceneChannelData<FMovieSceneFloatValue> ChannelData = Channel->GetData();
 		FFrameRate TickResolution = Section->GetTypedOuter<UMovieScene>()->GetTickResolution();
 
+		TArray<FKeyHandle> NewKeyHandles;
+		NewKeyHandles.SetNumUninitialized(InKeyPositions.Num());
+
 		for (int32 Index = 0; Index < InKeyPositions.Num(); ++Index)
 		{
 			FKeyPosition   Position   = InKeyPositions[Index];
-			FKeyAttributes Attributes = InKeyAttributes[Index];
 
 			FFrameNumber Time = (Position.InputValue * TickResolution).RoundToFrame();
 			Section->ExpandToFrame(Time);
 
 			FMovieSceneFloatValue Value(Position.OutputValue);
 
-			if (Attributes.HasInterpMode())    { Value.InterpMode            = Attributes.GetInterpMode();    }
-			if (Attributes.HasTangentMode())   { Value.TangentMode           = Attributes.GetTangentMode();   }
-			if (Attributes.HasArriveTangent()) { Value.Tangent.ArriveTangent = Attributes.GetArriveTangent(); }
-			if (Attributes.HasLeaveTangent())  { Value.Tangent.LeaveTangent  = Attributes.GetLeaveTangent();  }
-
 			int32 KeyIndex = ChannelData.AddKey(Time, Value);
+			FKeyHandle NewHandle = ChannelData.GetHandle(KeyIndex);
+			NewKeyHandles[Index] = NewHandle;
+
 			if (OutKeyHandles)
 			{
-				(*OutKeyHandles)[Index] = ChannelData.GetHandle(KeyIndex);
+				(*OutKeyHandles)[Index] = NewHandle;
 			}
 		}
 
-		Channel->AutoSetTangents();
+		// We reuse SetKeyAttributes here as there is complex logic determining which parts of the attributes are valid to set.
+		// For now we need to duplicate the new key handle array due to API mismatch. This will auto calculate tangents if needed.
+		SetKeyAttributes(NewKeyHandles, InKeyAttributes);
 	}
 }
 
