Sequencer: Remove restriction on hidden mobility. It's too wide to restrict transform tracks.

diff --git a/Engine/Source/Editor/MovieSceneTools/Private/TrackEditors/AttachTrackEditor.cpp b/Engine/Source/Editor/MovieSceneTools/Private/TrackEditors/AttachTrackEditor.cpp
--- a/Engine/Source/Editor/MovieSceneTools/Private/TrackEditors/AttachTrackEditor.cpp
+++ b/Engine/Source/Editor/MovieSceneTools/Private/TrackEditors/AttachTrackEditor.cpp
@@ -132,11 +132,6 @@ void F3DAttachTrackEditor::BuildObjectBindingTrackMenu(FMenuBuilder& MenuBuilder
 {
 	if (ObjectClass != nullptr && ObjectClass->IsChildOf(AActor::StaticClass()))
 	{
-		if (MovieSceneToolHelpers::HasHiddenMobility(ObjectClass))
-		{
-			return;
-		}
-
 		UMovieSceneSection* DummySection = nullptr;
 
 		MenuBuilder.AddSubMenu(
diff --git a/Engine/Source/Editor/MovieSceneTools/Private/TrackEditors/PathTrackEditor.cpp b/Engine/Source/Editor/MovieSceneTools/Private/TrackEditors/PathTrackEditor.cpp
--- a/Engine/Source/Editor/MovieSceneTools/Private/TrackEditors/PathTrackEditor.cpp
+++ b/Engine/Source/Editor/MovieSceneTools/Private/TrackEditors/PathTrackEditor.cpp
@@ -117,11 +117,6 @@ void F3DPathTrackEditor::BuildObjectBindingTrackMenu(FMenuBuilder& MenuBuilder,
 {
 	if (ObjectClass && ObjectClass->IsChildOf(AActor::StaticClass()))
 	{
-		if (MovieSceneToolHelpers::HasHiddenMobility(ObjectClass))
-		{
-			return;
-		}
-
 		UMovieSceneSection* DummySection = nullptr;
 
 		MenuBuilder.AddSubMenu(
diff --git a/Engine/Source/Editor/MovieSceneTools/Private/TrackEditors/TransformTrackEditor.cpp b/Engine/Source/Editor/MovieSceneTools/Private/TrackEditors/TransformTrackEditor.cpp
--- a/Engine/Source/Editor/MovieSceneTools/Private/TrackEditors/TransformTrackEditor.cpp
+++ b/Engine/Source/Editor/MovieSceneTools/Private/TrackEditors/TransformTrackEditor.cpp
@@ -443,11 +443,6 @@ void F3DTransformTrackEditor::BuildObjectBindingTrackMenu(FMenuBuilder& MenuBuil
 {
 	if (ObjectClass != nullptr && (ObjectClass->IsChildOf(AActor::StaticClass()) || ObjectClass->IsChildOf(USceneComponent::StaticClass())))
 	{
-		if (MovieSceneToolHelpers::HasHiddenMobility(ObjectClass))
-		{
-			return;
-		}
-
 		MenuBuilder.AddMenuEntry(
 			NSLOCTEXT("Sequencer", "AddTransform", "Transform"),
 			NSLOCTEXT("Sequencer", "AddPTransformTooltip", "Adds a transform track."),
@@ -745,11 +740,6 @@ void F3DTransformTrackEditor::AddTransformKeysForObject( UObject* Object, EMovie
 	{
 		FTransformData CurrentTransform( SceneComponent );
 
-		if (MovieSceneToolHelpers::HasHiddenMobility(Object->GetClass()))
-		{
-			return;
-		}
-
 		if(Object->GetClass()->IsChildOf(AActor::StaticClass()))
 		{
 			AddTransformKeys( Actor, TOptional<FTransformData>(), CurrentTransform, ChannelToKey, KeyMode );
