Sequencer: Fix crash when converting matinee assets to sequences due to changes in events. Event keys will now have blueprint event endpoints created in the level sequence blueprint when converting.

diff --git a/Engine/Source/Editor/MovieSceneTools/Private/MatineeImportTools.cpp b/Engine/Source/Editor/MovieSceneTools/Private/MatineeImportTools.cpp
--- a/Engine/Source/Editor/MovieSceneTools/Private/MatineeImportTools.cpp
+++ b/Engine/Source/Editor/MovieSceneTools/Private/MatineeImportTools.cpp
@@ -8,6 +8,10 @@
 #include "IMovieScenePlayer.h"
 #include "Channels/MovieSceneChannelProxy.h"
 #include "MovieScene.h"
+#include "MovieSceneSequenceEditor.h"
+#include "Channels/MovieSceneEvent.h"
+#include "K2Node_FunctionEntry.h"
+#include "Kismet2/KismetEditorUtilities.h"
 
 #include "Matinee/MatineeActor.h"
 #include "Matinee/InterpData.h"
@@ -46,6 +50,7 @@
 #include "Sections/MovieSceneFadeSection.h"
 #include "Sections/MovieSceneCameraCutSection.h"
 #include "Sections/MovieSceneEventSection.h"
+#include "Sections/MovieSceneEventTriggerSection.h"
 #include "Sections/MovieSceneVectorSection.h"
 
 
@@ -742,27 +747,55 @@ bool FMatineeImportTools::CopyInterpEventTrack( UInterpTrackEvent* MatineeEventT
 		FFrameRate   FrameRate    = EventTrack->GetTypedOuter<UMovieScene>()->GetTickResolution();
 		FFrameNumber FirstKeyTime = (MatineeEventTrack->EventTrack[0].Time * FrameRate).RoundToFrame();
 
-		UMovieSceneEventSection* Section = Cast<UMovieSceneEventSection>( MovieSceneHelpers::FindSectionAtTime( EventTrack->GetAllSections(), FirstKeyTime ) );
+		UMovieSceneEventTriggerSection* Section = Cast<UMovieSceneEventTriggerSection>( MovieSceneHelpers::FindSectionAtTime( EventTrack->GetAllSections(), FirstKeyTime ) );
 		if ( Section == nullptr )
 		{
-			Section = Cast<UMovieSceneEventSection>( EventTrack->CreateNewSection() );
-			EventTrack->AddSection( *Section );
+			Section = Cast<UMovieSceneEventTriggerSection>( EventTrack->CreateNewSection() );
+			EventTrack->AddSection(*Section);
 			bSectionCreated = true;
+			
 		}
 		if (Section->TryModify())
 		{
 			TRange<FFrameNumber> KeyRange = TRange<FFrameNumber>::Empty();
 
-			FMovieSceneEventSectionData* EventChannel = Section->GetChannelProxy().GetChannel<FMovieSceneEventSectionData>(0);
+			FMovieSceneEventChannel* EventChannel = Section->GetChannelProxy().GetChannel<FMovieSceneEventChannel>(0);
 			check(EventChannel);
-			TMovieSceneChannelData<FEventPayload> ChannelData = EventChannel->GetData();
+			TMovieSceneChannelData<FMovieSceneEvent> ChannelData = EventChannel->GetData();
+			TMap<FName, FMovieSceneEvent> EventMap;
 
 			for (FEventTrackKey EventTrackKey : MatineeEventTrack->EventTrack)
 			{
 				FFrameNumber KeyTime = (EventTrackKey.Time * FrameRate).RoundToFrame();
+				
+				FMovieSceneEvent Event;
+
+				if (FMovieSceneEvent* FoundEvent = EventMap.Find(EventTrackKey.EventName))
+				{
+					Event = *FoundEvent;
+				}
+				else
+				{
+					FMovieSceneSequenceEditor* SequenceEditor = FMovieSceneSequenceEditor::Find(Section->GetTypedOuter<UMovieSceneSequence>());
+					if (ensure(SequenceEditor))
+					{
+						// Create and bind event endpoint for this event
+						UK2Node_FunctionEntry* NewEndpoint = SequenceEditor->CreateEventEndpoint(Section->GetTypedOuter<UMovieSceneSequence>(), EventTrackKey.EventName.ToString());
+						if (NewEndpoint)
+						{
+							SequenceEditor->InitializeEndpointForTrack(Section->GetTypedOuter<UMovieSceneEventTrack>(), NewEndpoint);
+							FMovieSceneSequenceEditor::BindEventToEndpoint(Section, &Event, NewEndpoint);
+						}
+					}
 
-				ChannelData.UpdateOrAddKey(KeyTime, FEventPayload(EventTrackKey.EventName));
+					if (UK2Node_FunctionEntry* FunctionEntry = Event.GetFunctionEntry())
+					{
+						FKismetEditorUtilities::BringKismetToFocusAttentionOnObject(FunctionEntry, false);
+					}
+					EventMap.Add(EventTrackKey.EventName, Event);
+				}
 
+				ChannelData.UpdateOrAddKey(KeyTime, Event);
 				KeyRange = TRange<FFrameNumber>::Hull(KeyRange, TRange<FFrameNumber>(KeyTime));
 			}
 
