Sequencer: Now accouting for frame rate when copying key tangents when converting from matinee to level sequence.

diff --git a/Engine/Source/Editor/MovieSceneTools/Private/MatineeImportTools.cpp b/Engine/Source/Editor/MovieSceneTools/Private/MatineeImportTools.cpp
--- a/Engine/Source/Editor/MovieSceneTools/Private/MatineeImportTools.cpp
+++ b/Engine/Source/Editor/MovieSceneTools/Private/MatineeImportTools.cpp
@@ -119,7 +119,7 @@ bool FMatineeImportTools::TryConvertMatineeToggleToOutParticleKey( ETrackToggleA
 }
 
 
-void FMatineeImportTools::SetOrAddKey(TMovieSceneChannelData<FMovieSceneFloatValue>& ChannelData, FFrameNumber Time, float Value, float ArriveTangent, float LeaveTangent, EInterpCurveMode MatineeInterpMode
+void FMatineeImportTools::SetOrAddKey(TMovieSceneChannelData<FMovieSceneFloatValue>& ChannelData, FFrameNumber Time, float Value, float ArriveTangent, float LeaveTangent, EInterpCurveMode MatineeInterpMode, FFrameRate FrameRate
 	, ERichCurveTangentWeightMode WeightedMode, float ArriveTangentWeight, float LeaveTangentWeight )
 {
 	if (ChannelData.FindKey(Time) == INDEX_NONE)
@@ -128,8 +128,8 @@ void FMatineeImportTools::SetOrAddKey(TMovieSceneChannelData<FMovieSceneFloatVal
 
 		NewKey.InterpMode = MatineeInterpolationToRichCurveInterpolation( MatineeInterpMode );
 		NewKey.TangentMode = MatineeInterpolationToRichCurveTangent( MatineeInterpMode );
-		NewKey.Tangent.ArriveTangent = ArriveTangent;
-		NewKey.Tangent.LeaveTangent = LeaveTangent;
+		NewKey.Tangent.ArriveTangent = ArriveTangent / FrameRate.AsDecimal();
+		NewKey.Tangent.LeaveTangent = LeaveTangent / FrameRate.AsDecimal();
 		NewKey.Tangent.TangentWeightMode = WeightedMode;
 		NewKey.Tangent.ArriveTangentWeight = ArriveTangentWeight;
 		NewKey.Tangent.LeaveTangentWeight = LeaveTangentWeight;
@@ -208,7 +208,7 @@ bool FMatineeImportTools::CopyInterpFloatTrack( UInterpTrackFloatBase* MatineeFl
 		{
 			FFrameNumber KeyTime = (Point.InVal * FrameRate).RoundToFrame();
 
-			FMatineeImportTools::SetOrAddKey( ChannelData, KeyTime, Point.OutVal, Point.ArriveTangent, Point.LeaveTangent, Point.InterpMode );
+			FMatineeImportTools::SetOrAddKey( ChannelData, KeyTime, Point.OutVal, Point.ArriveTangent, Point.LeaveTangent, Point.InterpMode, FrameRate);
 
 			KeyRange = TRange<FFrameNumber>::Hull(KeyRange, TRange<FFrameNumber>(KeyTime));
 		}
@@ -255,9 +255,9 @@ bool FMatineeImportTools::CopyInterpVectorTrack( UInterpTrackVectorProp* Matinee
 			{
 				FFrameNumber KeyTime = (Point.InVal * FrameRate).RoundToFrame();
 
-				FMatineeImportTools::SetOrAddKey( ChannelData[0], KeyTime, Point.OutVal.X, Point.ArriveTangent.X, Point.LeaveTangent.X, Point.InterpMode );
-				FMatineeImportTools::SetOrAddKey( ChannelData[1], KeyTime, Point.OutVal.Y, Point.ArriveTangent.Y, Point.LeaveTangent.Y, Point.InterpMode );
-				FMatineeImportTools::SetOrAddKey( ChannelData[2], KeyTime, Point.OutVal.Z, Point.ArriveTangent.Z, Point.LeaveTangent.Z, Point.InterpMode );
+				FMatineeImportTools::SetOrAddKey( ChannelData[0], KeyTime, Point.OutVal.X, Point.ArriveTangent.X, Point.LeaveTangent.X, Point.InterpMode, FrameRate);
+				FMatineeImportTools::SetOrAddKey( ChannelData[1], KeyTime, Point.OutVal.Y, Point.ArriveTangent.Y, Point.LeaveTangent.Y, Point.InterpMode, FrameRate);
+				FMatineeImportTools::SetOrAddKey( ChannelData[2], KeyTime, Point.OutVal.Z, Point.ArriveTangent.Z, Point.LeaveTangent.Z, Point.InterpMode, FrameRate);
 
 				KeyRange = TRange<FFrameNumber>::Hull(KeyRange, TRange<FFrameNumber>(KeyTime));
 			}
@@ -314,9 +314,9 @@ bool FMatineeImportTools::CopyInterpColorTrack( UInterpTrackColorProp* ColorProp
 		{
 			FFrameNumber KeyTime = (Point.InVal * FrameRate).RoundToFrame();
 
-			FMatineeImportTools::SetOrAddKey( ChannelData[0], KeyTime, Point.OutVal.X, Point.ArriveTangent.X, Point.LeaveTangent.X, Point.InterpMode );
-			FMatineeImportTools::SetOrAddKey( ChannelData[1], KeyTime, Point.OutVal.Y, Point.ArriveTangent.Y, Point.LeaveTangent.Y, Point.InterpMode );
-			FMatineeImportTools::SetOrAddKey( ChannelData[2], KeyTime, Point.OutVal.Z, Point.ArriveTangent.Z, Point.LeaveTangent.Z, Point.InterpMode );
+			FMatineeImportTools::SetOrAddKey( ChannelData[0], KeyTime, Point.OutVal.X, Point.ArriveTangent.X, Point.LeaveTangent.X, Point.InterpMode, FrameRate);
+			FMatineeImportTools::SetOrAddKey( ChannelData[1], KeyTime, Point.OutVal.Y, Point.ArriveTangent.Y, Point.LeaveTangent.Y, Point.InterpMode, FrameRate);
+			FMatineeImportTools::SetOrAddKey( ChannelData[2], KeyTime, Point.OutVal.Z, Point.ArriveTangent.Z, Point.LeaveTangent.Z, Point.InterpMode, FrameRate);
 
 			KeyRange = TRange<FFrameNumber>::Hull(KeyRange, TRange<FFrameNumber>(KeyTime));
 		}
@@ -372,10 +372,10 @@ bool FMatineeImportTools::CopyInterpLinearColorTrack( UInterpTrackLinearColorPro
 		{
 			FFrameNumber KeyTime = (Point.InVal * FrameRate).RoundToFrame();
 
-			FMatineeImportTools::SetOrAddKey( ChannelData[0], KeyTime, Point.OutVal.R, Point.ArriveTangent.R, Point.LeaveTangent.R, Point.InterpMode );
-			FMatineeImportTools::SetOrAddKey( ChannelData[1], KeyTime, Point.OutVal.G, Point.ArriveTangent.G, Point.LeaveTangent.G, Point.InterpMode );
-			FMatineeImportTools::SetOrAddKey( ChannelData[2], KeyTime, Point.OutVal.B, Point.ArriveTangent.B, Point.LeaveTangent.B, Point.InterpMode );
-			FMatineeImportTools::SetOrAddKey( ChannelData[3], KeyTime, Point.OutVal.A, Point.ArriveTangent.A, Point.LeaveTangent.A, Point.InterpMode );
+			FMatineeImportTools::SetOrAddKey( ChannelData[0], KeyTime, Point.OutVal.R, Point.ArriveTangent.R, Point.LeaveTangent.R, Point.InterpMode, FrameRate);
+			FMatineeImportTools::SetOrAddKey( ChannelData[1], KeyTime, Point.OutVal.G, Point.ArriveTangent.G, Point.LeaveTangent.G, Point.InterpMode, FrameRate);
+			FMatineeImportTools::SetOrAddKey( ChannelData[2], KeyTime, Point.OutVal.B, Point.ArriveTangent.B, Point.LeaveTangent.B, Point.InterpMode, FrameRate);
+			FMatineeImportTools::SetOrAddKey( ChannelData[3], KeyTime, Point.OutVal.A, Point.ArriveTangent.A, Point.LeaveTangent.A, Point.InterpMode, FrameRate);
 
 			KeyRange = TRange<FFrameNumber>::Hull(KeyRange, TRange<FFrameNumber>(KeyTime));
 		}
@@ -433,9 +433,9 @@ bool FMatineeImportTools::CopyInterpMoveTrack( UInterpTrackMove* MoveTrack, UMov
 		{
 			FFrameNumber KeyTime = (Point.InVal * FrameRate).RoundToFrame();
 
-			FMatineeImportTools::SetOrAddKey( ChannelData[0], KeyTime, Point.OutVal.X, Point.ArriveTangent.X, Point.LeaveTangent.X, Point.InterpMode );
-			FMatineeImportTools::SetOrAddKey( ChannelData[1], KeyTime, Point.OutVal.Y, Point.ArriveTangent.Y, Point.LeaveTangent.Y, Point.InterpMode );
-			FMatineeImportTools::SetOrAddKey( ChannelData[2], KeyTime, Point.OutVal.Z, Point.ArriveTangent.Z, Point.LeaveTangent.Z, Point.InterpMode );
+			FMatineeImportTools::SetOrAddKey( ChannelData[0], KeyTime, Point.OutVal.X, Point.ArriveTangent.X, Point.LeaveTangent.X, Point.InterpMode, FrameRate);
+			FMatineeImportTools::SetOrAddKey( ChannelData[1], KeyTime, Point.OutVal.Y, Point.ArriveTangent.Y, Point.LeaveTangent.Y, Point.InterpMode, FrameRate);
+			FMatineeImportTools::SetOrAddKey( ChannelData[2], KeyTime, Point.OutVal.Z, Point.ArriveTangent.Z, Point.LeaveTangent.Z, Point.InterpMode, FrameRate);
 
 			KeyRange = TRange<FFrameNumber>::Hull(KeyRange, TRange<FFrameNumber>(KeyTime));
 		}
@@ -444,9 +444,9 @@ bool FMatineeImportTools::CopyInterpMoveTrack( UInterpTrackMove* MoveTrack, UMov
 		{
 			FFrameNumber KeyTime = (Point.InVal * FrameRate).RoundToFrame();
 
-			FMatineeImportTools::SetOrAddKey( ChannelData[3], KeyTime, Point.OutVal.X, Point.ArriveTangent.X, Point.LeaveTangent.X, Point.InterpMode );
-			FMatineeImportTools::SetOrAddKey( ChannelData[4], KeyTime, Point.OutVal.Y, Point.ArriveTangent.Y, Point.LeaveTangent.Y, Point.InterpMode );
-			FMatineeImportTools::SetOrAddKey( ChannelData[5], KeyTime, Point.OutVal.Z, Point.ArriveTangent.Z, Point.LeaveTangent.Z, Point.InterpMode );
+			FMatineeImportTools::SetOrAddKey( ChannelData[3], KeyTime, Point.OutVal.X, Point.ArriveTangent.X, Point.LeaveTangent.X, Point.InterpMode, FrameRate);
+			FMatineeImportTools::SetOrAddKey( ChannelData[4], KeyTime, Point.OutVal.Y, Point.ArriveTangent.Y, Point.LeaveTangent.Y, Point.InterpMode, FrameRate);
+			FMatineeImportTools::SetOrAddKey( ChannelData[5], KeyTime, Point.OutVal.Z, Point.ArriveTangent.Z, Point.LeaveTangent.Z, Point.InterpMode, FrameRate);
 
 			KeyRange = TRange<FFrameNumber>::Hull(KeyRange, TRange<FFrameNumber>(KeyTime));
 		}
@@ -491,7 +491,7 @@ bool FMatineeImportTools::CopyInterpMoveTrack( UInterpTrackMove* MoveTrack, UMov
 						{
 							FFrameNumber KeyTime = (Point.InVal * FrameRate).RoundToFrame();
 
-							FMatineeImportTools::SetOrAddKey( ChannelData[ChannelIndex], KeyTime, Point.OutVal, Point.ArriveTangent, Point.LeaveTangent, Point.InterpMode );
+							FMatineeImportTools::SetOrAddKey( ChannelData[ChannelIndex], KeyTime, Point.OutVal, Point.ArriveTangent, Point.LeaveTangent, Point.InterpMode, FrameRate);
 
 							KeyRange = TRange<FFrameNumber>::Hull(KeyRange, TRange<FFrameNumber>(KeyTime));
 						}
@@ -682,7 +682,7 @@ bool FMatineeImportTools::CopyInterpFadeTrack( UInterpTrackFade* MatineeFadeTrac
 		{
 			FFrameNumber KeyTime = (Point.InVal * FrameRate).RoundToFrame();
 
-			FMatineeImportTools::SetOrAddKey( FadeInterface, KeyTime, Point.OutVal, Point.ArriveTangent, Point.LeaveTangent, Point.InterpMode );
+			FMatineeImportTools::SetOrAddKey( FadeInterface, KeyTime, Point.OutVal, Point.ArriveTangent, Point.LeaveTangent, Point.InterpMode, FrameRate);
 
 			KeyRange = TRange<FFrameNumber>::Hull(KeyRange, TRange<FFrameNumber>(KeyTime));
 		}
diff --git a/Engine/Source/Editor/MovieSceneTools/Private/MovieSceneToolHelpers.cpp b/Engine/Source/Editor/MovieSceneTools/Private/MovieSceneToolHelpers.cpp
--- a/Engine/Source/Editor/MovieSceneTools/Private/MovieSceneToolHelpers.cpp
+++ b/Engine/Source/Editor/MovieSceneTools/Private/MovieSceneToolHelpers.cpp
@@ -941,7 +941,7 @@ bool ImportFBXProperty(FString NodeName, FString AnimatedPropertyName, FGuid Obj
 
 					FFrameNumber KeyTime = (Key.Time * FrameRate).RoundToFrame();
 					FMatineeImportTools::SetOrAddKey(ChannelData, KeyTime, Key.Value, ArriveTangent, LeaveTangent,
-						MovieSceneToolHelpers::RichCurveInterpolationToMatineeInterpolation(Key.InterpMode, Key.TangentMode), Key.TangentWeightMode,
+						MovieSceneToolHelpers::RichCurveInterpolationToMatineeInterpolation(Key.InterpMode, Key.TangentMode), FrameRate, Key.TangentWeightMode,
 						Key.ArriveTangentWeight, Key.LeaveTangentWeight);
 
 				}
@@ -996,7 +996,7 @@ void ImportTransformChannel(const FRichCurve& Source, FMovieSceneFloatChannel* D
 
 		FFrameNumber KeyTime = (Key.Time * DestFrameRate).RoundToFrame();
 		FMatineeImportTools::SetOrAddKey(ChannelData, KeyTime, Key.Value, ArriveTangent, LeaveTangent,
-			MovieSceneToolHelpers::RichCurveInterpolationToMatineeInterpolation(Key.InterpMode, Key.TangentMode), Key.TangentWeightMode,
+			MovieSceneToolHelpers::RichCurveInterpolationToMatineeInterpolation(Key.InterpMode, Key.TangentMode), DestFrameRate, Key.TangentWeightMode,
 			Key.ArriveTangentWeight, Key.LeaveTangentWeight);
 
 	}
diff --git a/Engine/Source/Editor/MovieSceneTools/Public/MatineeImportTools.h b/Engine/Source/Editor/MovieSceneTools/Public/MatineeImportTools.h
--- a/Engine/Source/Editor/MovieSceneTools/Public/MatineeImportTools.h
+++ b/Engine/Source/Editor/MovieSceneTools/Public/MatineeImportTools.h
@@ -47,7 +47,7 @@ public:
 	static bool TryConvertMatineeToggleToOutParticleKey( ETrackToggleAction ToggleAction, EParticleKey& OutParticleKey );
 
 	/** Adds a key to a rich curve based on matinee curve key data. */
-	static void SetOrAddKey( TMovieSceneChannelData<FMovieSceneFloatValue>& Curve, FFrameNumber Time, float Value, float ArriveTangent, float LeaveTangent, EInterpCurveMode MatineeInterpMode
+	static void SetOrAddKey( TMovieSceneChannelData<FMovieSceneFloatValue>& Curve, FFrameNumber Time, float Value, float ArriveTangent, float LeaveTangent, EInterpCurveMode MatineeInterpMode, FFrameRate FrameRate
 	, ERichCurveTangentWeightMode WeightedMode = RCTWM_WeightedNone, float ArriveTangentWeight = 0.0f, float LeaveTangentWeight = 0.0f);
 
 	/** Copies keys from a matinee bool track to a sequencer bool track. */
