Sequencer: Fix ImportFBX undo

diff --git a/Engine/Source/Editor/MovieSceneTools/Private/MovieSceneToolHelpers.cpp b/Engine/Source/Editor/MovieSceneTools/Private/MovieSceneToolHelpers.cpp
--- a/Engine/Source/Editor/MovieSceneTools/Private/MovieSceneToolHelpers.cpp
+++ b/Engine/Source/Editor/MovieSceneTools/Private/MovieSceneToolHelpers.cpp
@@ -890,6 +890,7 @@ bool ImportFBXProperty(FString NodeName, FString AnimatedPropertyName, FGuid Obj
 
 			if (FloatTrack)
 			{
+				FloatTrack->Modify();
 				FloatTrack->RemoveAllAnimationData();
 
 				FFrameRate FrameRate = FloatTrack->GetTypedOuter<UMovieScene>()->GetTickResolution();
@@ -1029,6 +1030,7 @@ bool ImportFBXTransform(FString NodeName, FGuid ObjectBinding, UnFbx::FFbxCurves
 		InMovieScene->Modify();
 		TransformTrack = InMovieScene->AddTrack<UMovieScene3DTransformTrack>(ObjectBinding);
 	}
+	TransformTrack->Modify();
 	TransformTrack->RemoveAllAnimationData();
 
 	bool bSectionAdded = false;
@@ -1364,6 +1366,7 @@ void ImportFBXCamera(UnFbx::FFbxImporter* FbxImporter, UMovieScene* InMovieScene
 				UMovieSceneFloatTrack* FloatTrack = InMovieScene->FindTrack<UMovieSceneFloatTrack>(PropertyOwnerGuid, TEXT("CurrentFocalLength"));
 				if (FloatTrack)
 				{
+					FloatTrack->Modify();
 					FloatTrack->RemoveAllAnimationData();
 
 					bool bSectionAdded = false;
