Sequencer: Fix fbx transform import so that if there are no keys in the fbx channel, the existing destination channel is not cleared.

diff --git a/Engine/Source/Editor/MovieSceneTools/Private/MovieSceneToolHelpers.cpp b/Engine/Source/Editor/MovieSceneTools/Private/MovieSceneToolHelpers.cpp
--- a/Engine/Source/Editor/MovieSceneTools/Private/MovieSceneToolHelpers.cpp
+++ b/Engine/Source/Editor/MovieSceneTools/Private/MovieSceneToolHelpers.cpp
@@ -965,6 +965,12 @@ bool ImportFBXProperty(FString NodeName, FString AnimatedPropertyName, FGuid Obj
 
 void ImportTransformChannel(const FRichCurve& Source, FMovieSceneFloatChannel* Dest, FFrameRate DestFrameRate, bool bNegateTangents)
 {
+	// If there are no keys, don't clear the existing channel
+	if (!Source.GetNumKeys())
+	{
+		return;
+	}
+
 	TMovieSceneChannelData<FMovieSceneFloatValue> ChannelData = Dest->GetData();
 	ChannelData.Reset();
 	double DecimalRate = DestFrameRate.AsDecimal();
@@ -1031,7 +1037,6 @@ bool ImportFBXTransform(FString NodeName, FGuid ObjectBinding, UnFbx::FFbxCurves
 		TransformTrack = InMovieScene->AddTrack<UMovieScene3DTransformTrack>(ObjectBinding);
 	}
 	TransformTrack->Modify();
-	TransformTrack->RemoveAllAnimationData();
 
 	bool bSectionAdded = false;
 	UMovieScene3DTransformSection* TransformSection = Cast<UMovieScene3DTransformSection>(TransformTrack->FindOrAddSection(0, bSectionAdded));
@@ -1044,7 +1049,6 @@ bool ImportFBXTransform(FString NodeName, FGuid ObjectBinding, UnFbx::FFbxCurves
 
 	FFrameRate FrameRate = TransformSection->GetTypedOuter<UMovieScene>()->GetTickResolution();
 
-
 	if (bSectionAdded)
 	{
 		TransformSection->SetRange(TRange<FFrameNumber>::All());
