Sequencer: Remove restriction on hidden mobility. It's too wide to restrict transform tracks.

diff --git a/Engine/Source/Editor/MovieSceneTools/Private/TrackEditors/TransformTrackEditor.cpp b/Engine/Source/Editor/MovieSceneTools/Private/TrackEditors/TransformTrackEditor.cpp
--- a/Engine/Source/Editor/MovieSceneTools/Private/TrackEditors/TransformTrackEditor.cpp
+++ b/Engine/Source/Editor/MovieSceneTools/Private/TrackEditors/TransformTrackEditor.cpp
@@ -740,11 +740,6 @@ void F3DTransformTrackEditor::AddTransformKeysForObject( UObject* Object, EMovie
 	{
 		FTransformData CurrentTransform( SceneComponent );
 
-		if (MovieSceneToolHelpers::HasHiddenMobility(Object->GetClass()))
-		{
-			return;
-		}
-
 		if(Object->GetClass()->IsChildOf(AActor::StaticClass()))
 		{
 			AddTransformKeys( Actor, TOptional<FTransformData>(), CurrentTransform, ChannelToKey, KeyMode );
