Sequencer: Get animation skeleton from CDO if the actor hasn't been spawned

diff --git a/Engine/Source/Editor/MovieSceneTools/Private/TrackEditors/SkeletalAnimationTrackEditor.cpp b/Engine/Source/Editor/MovieSceneTools/Private/TrackEditors/SkeletalAnimationTrackEditor.cpp
--- a/Engine/Source/Editor/MovieSceneTools/Private/TrackEditors/SkeletalAnimationTrackEditor.cpp
+++ b/Engine/Source/Editor/MovieSceneTools/Private/TrackEditors/SkeletalAnimationTrackEditor.cpp
@@ -40,6 +40,9 @@
 #include "Framework/Notifications/NotificationManager.h"
 #include "Widgets/Notifications/SNotificationList.h"
 #include "Toolkits/AssetEditorManager.h"
+#include "Engine/SCS_Node.h"
+#include "Engine/SimpleConstructionScript.h"
+#include "Engine/Blueprint.h"
 
 namespace SkeletalAnimationEditorConstants
 {
@@ -74,6 +77,18 @@ USkeletalMeshComponent* AcquireSkeletalMeshFromObjectGuid(const FGuid& Guid, TSh
 	return nullptr;
 }
 
+USkeleton* GetSkeletonFromComponent(UActorComponent* InComponent)
+{
+	USkeletalMeshComponent* SkeletalMeshComp = Cast<USkeletalMeshComponent>(InComponent);
+	if (SkeletalMeshComp && SkeletalMeshComp->SkeletalMesh && SkeletalMeshComp->SkeletalMesh->Skeleton)
+	{
+		// @todo Multiple actors, multiple components
+		return SkeletalMeshComp->SkeletalMesh->Skeleton;
+	}
+
+	return nullptr;
+}
+
 USkeleton* AcquireSkeletonFromObjectGuid(const FGuid& Guid, TSharedPtr<ISequencer> SequencerPtr)
 {
 	UObject* BoundObject = SequencerPtr.IsValid() ? SequencerPtr->FindSpawnedObjectOrTemplate(Guid) : nullptr;
@@ -82,19 +97,46 @@ USkeleton* AcquireSkeletonFromObjectGuid(const FGuid& Guid, TSharedPtr<ISequence
 	{
 		for (UActorComponent* Component : Actor->GetComponents())
 		{
-			USkeletalMeshComponent* SkeletalMeshComp = Cast<USkeletalMeshComponent>(Component);
-			if (SkeletalMeshComp && SkeletalMeshComp->SkeletalMesh && SkeletalMeshComp->SkeletalMesh->Skeleton)
+			if (USkeleton* Skeleton = GetSkeletonFromComponent(Component))
 			{
-				// @todo Multiple actors, multiple components
-				return SkeletalMeshComp->SkeletalMesh->Skeleton;
+				return Skeleton;
+			}
+		}
+
+		AActor* ActorCDO = Cast<AActor>(Actor->GetClass()->GetDefaultObject());
+		if (ActorCDO)
+		{
+			for (UActorComponent* Component : ActorCDO->GetComponents())
+			{
+				if (USkeleton* Skeleton = GetSkeletonFromComponent(Component))
+				{
+					return Skeleton;
+				}
+			}
+		}
+
+		UBlueprintGeneratedClass* ActorBlueprintGeneratedClass = Cast<UBlueprintGeneratedClass>(Actor->GetClass());
+		if (ActorBlueprintGeneratedClass)
+		{
+			const TArray<USCS_Node*>& ActorBlueprintNodes = ActorBlueprintGeneratedClass->SimpleConstructionScript->GetAllNodes();
+
+			for (USCS_Node* Node : ActorBlueprintNodes)
+			{
+				if (Node->ComponentClass->IsChildOf(USkeletalMeshComponent::StaticClass()))
+				{
+					if (USkeleton* Skeleton = GetSkeletonFromComponent(Node->GetActualComponentTemplate(ActorBlueprintGeneratedClass)))
+					{
+						return Skeleton;
+					}
+				}
 			}
 		}
 	}
 	else if(USkeletalMeshComponent* SkeletalMeshComponent = Cast<USkeletalMeshComponent>(BoundObject))
 	{
-		if (SkeletalMeshComponent->SkeletalMesh)
+		if (USkeleton* Skeleton = GetSkeletonFromComponent(SkeletalMeshComponent))
 		{
-			return SkeletalMeshComponent->SkeletalMesh->Skeleton;
+			return Skeleton;
 		}
 	}
 
