Matinee material float and vector paramter conversion to sequencer. This addds support for converting material tracks form matinee to sequencer.

diff --git a/Engine/Plugins/MovieScene/MatineeToLevelSequence/Source/MatineeToLevelSequence/Private/MatineeToLevelSequenceModule.cpp b/Engine/Plugins/MovieScene/MatineeToLevelSequence/Source/MatineeToLevelSequence/Private/MatineeToLevelSequenceModule.cpp
--- a/Engine/Plugins/MovieScene/MatineeToLevelSequence/Source/MatineeToLevelSequence/Private/MatineeToLevelSequenceModule.cpp
+++ b/Engine/Plugins/MovieScene/MatineeToLevelSequence/Source/MatineeToLevelSequence/Private/MatineeToLevelSequenceModule.cpp
@@ -30,6 +30,8 @@
 #include "Matinee/InterpTrackColorScale.h"
 #include "Matinee/InterpTrackColorProp.h"
 #include "Matinee/InterpTrackFloatProp.h"
+#include "Matinee/InterpTrackFloatMaterialParam.h"
+#include "Matinee/InterpTrackVectorMaterialParam.h"
 #include "Matinee/InterpTrackMove.h"
 #include "Matinee/InterpTrackAnimControl.h"
 #include "Matinee/InterpTrackSound.h"
@@ -59,6 +61,7 @@
 #include "Tracks/MovieSceneBoolTrack.h"
 #include "Tracks/MovieSceneColorTrack.h"
 #include "Tracks/MovieSceneFloatTrack.h"
+#include "Tracks/MovieSceneMaterialTrack.h"
 #include "Tracks/MovieScene3DTransformTrack.h"
 #include "Tracks/MovieSceneParticleTrack.h"
 #include "Tracks/MovieSceneSkeletalAnimationTrack.h"
@@ -73,6 +76,8 @@
 #include "Framework/Notifications/NotificationManager.h"
 #include "Widgets/Notifications/SNotificationList.h"
 #include "MovieSceneTimeHelpers.h"
+#include "Components/SkeletalMeshComponent.h"
+#include "Engine/StaticMeshActor.h"
 
 #define LOCTEXT_NAMESPACE "MatineeToLevelSequence"
 
@@ -425,6 +430,104 @@ protected:
 		return PropertyTrack;
 	}
 
+	FGuid FindComponentGuid(AActor* GroupActor, UMovieSceneSequence* NewSequence, UMovieScene* NewMovieScene, FGuid PossessableGuid)
+	{
+		FGuid ComponentGUID;
+		// Skeletal and static mesh actors can both have material component tracks, and need to have their mesh component added to sequencer.
+		if (GroupActor->GetClass() == ASkeletalMeshActor::StaticClass())
+		{
+
+			ASkeletalMeshActor* SkelMeshActor = CastChecked<ASkeletalMeshActor>(GroupActor);
+			USkeletalMeshComponent* SkelMeshComponent = SkelMeshActor->GetSkeletalMeshComponent();
+			// In matinee a component may be referenced in multiple material tracks, so check to see if this one is already bound.
+			FGuid FoundGUID = NewSequence->FindPossessableObjectId(*SkelMeshComponent, SkelMeshActor);
+			if (FoundGUID != FGuid())
+			{
+				ComponentGUID = FoundGUID;
+			}
+			else
+			{
+				ComponentGUID = NewMovieScene->AddPossessable(SkelMeshComponent->GetName(), SkelMeshComponent->GetClass());
+				FMovieScenePossessable* ChildPossesable = NewMovieScene->FindPossessable(ComponentGUID);
+				ChildPossesable->SetParent(PossessableGuid);
+				NewSequence->BindPossessableObject(ComponentGUID, *SkelMeshComponent, GroupActor->GetWorld());
+			}
+		}
+		else if (GroupActor->GetClass() == AStaticMeshActor::StaticClass())
+		{
+			AStaticMeshActor* StaticMeshActor = CastChecked<AStaticMeshActor>(GroupActor);
+			UStaticMeshComponent* StaticMeshComponent = StaticMeshActor->GetStaticMeshComponent();
+			FGuid FoundGUID = NewSequence->FindPossessableObjectId(*StaticMeshComponent, StaticMeshActor);
+			if (FoundGUID != FGuid())
+			{
+				ComponentGUID = FoundGUID;
+			}
+			else
+			{
+				ComponentGUID = NewMovieScene->AddPossessable(StaticMeshComponent->GetName(), StaticMeshComponent->GetClass());
+				FMovieScenePossessable* ChildPossesable = NewMovieScene->FindPossessable(ComponentGUID);
+				ChildPossesable->SetParent(PossessableGuid);
+				NewSequence->BindPossessableObject(ComponentGUID, *StaticMeshComponent, GroupActor->GetWorld());
+			}
+		}
+		else
+		{
+			return FGuid();
+		}
+		return ComponentGUID;
+	}
+	
+	template<typename T>
+	void CopyMaterialsToComponents(int32 NumMaterials, FGuid ComponentGuid, UMovieScene* NewMovieScene, T* MatineeMaterialParamTrack)
+	{
+		// One matinee material track can change the same parameter for multiple materials, but sequencer binds them to individual tracks.
+		for (int32 MaterialIndex = 0; MaterialIndex < NumMaterials; ++MaterialIndex)
+		{
+			UMovieSceneComponentMaterialTrack* MaterialTrack = nullptr;
+
+			TArray<UMovieSceneComponentMaterialTrack *> BoundTracks;
+
+			// Find all tracks bound to the component we added.
+			for (const FMovieSceneBinding& Binding : NewMovieScene->GetBindings())
+			{
+				if (Binding.GetObjectGuid() == ComponentGuid)
+				{
+					for (auto Track : Binding.GetTracks())
+					{
+						UMovieSceneComponentMaterialTrack * MaterialTrack = Cast<UMovieSceneComponentMaterialTrack>(Track);
+						if (MaterialTrack)
+						{
+							BoundTracks.Add(MaterialTrack);
+						}
+					}
+					break;
+				}
+			}
+
+			// The material may have already been added to the component, so look first to see if there is a track with the current material index.
+			UMovieSceneComponentMaterialTrack** FoundTrack = BoundTracks.FindByPredicate([MaterialIndex](UMovieSceneComponentMaterialTrack* Track) -> bool { return Track && Track->GetMaterialIndex() == MaterialIndex; });
+			if (FoundTrack)
+			{
+				MaterialTrack = *FoundTrack;
+			}
+
+			if (MaterialTrack)
+			{
+				FMatineeImportTools::CopyInterpMaterialParamTrack(MatineeMaterialParamTrack, MaterialTrack);
+			}
+			else
+			{
+				UMovieSceneComponentMaterialTrack* MaterialTrack = NewMovieScene->AddTrack<UMovieSceneComponentMaterialTrack>(ComponentGuid);
+				if (MaterialTrack)
+				{
+					MaterialTrack->SetMaterialIndex(MaterialIndex);
+					FMatineeImportTools::CopyInterpMaterialParamTrack(MatineeMaterialParamTrack, MaterialTrack);
+				}
+			}
+			
+		}
+	}
+
 	/** Convert an interp group */
 	void ConvertInterpGroup(UInterpGroup* Group, AActor* GroupActor, IMovieScenePlayer& Player, UMovieSceneSequence* NewSequence, UMovieScene* NewMovieScene, int32& NumWarnings)
 	{
@@ -553,6 +656,38 @@ protected:
 					}
 				}
 			}
+			else if (Track->IsA(UInterpTrackFloatMaterialParam::StaticClass()))
+			{
+				UInterpTrackFloatMaterialParam* MatineeMaterialParamTrack = StaticCast<UInterpTrackFloatMaterialParam*>(Track);
+				if (MatineeMaterialParamTrack->GetNumKeyframes() != 0 && GroupActor && PossessableGuid.IsValid())
+				{
+					FGuid ComponentGuid = FindComponentGuid(GroupActor, NewSequence, NewMovieScene, PossessableGuid);
+
+					if (ComponentGuid == FGuid())
+					{
+						continue;
+					}
+
+					CopyMaterialsToComponents(MatineeMaterialParamTrack->TargetMaterials.Num(), ComponentGuid, NewMovieScene, MatineeMaterialParamTrack);
+					
+				}
+			}
+			else if (Track->IsA(UInterpTrackVectorMaterialParam::StaticClass()))
+			{
+				UInterpTrackVectorMaterialParam* MatineeMaterialParamTrack = StaticCast<UInterpTrackVectorMaterialParam*>(Track);
+				if (MatineeMaterialParamTrack->GetNumKeyframes() != 0 && GroupActor && PossessableGuid.IsValid())
+				{
+					FGuid ComponentGuid = FindComponentGuid(GroupActor, NewSequence, NewMovieScene, PossessableGuid);
+
+					if (ComponentGuid == FGuid())
+					{
+						continue;
+					}
+
+					CopyMaterialsToComponents(MatineeMaterialParamTrack->TargetMaterials.Num(), ComponentGuid, NewMovieScene, MatineeMaterialParamTrack);
+
+				}
+			}
 			else if (Track->IsA(UInterpTrackVectorProp::StaticClass()))
 			{
 				UInterpTrackVectorProp* MatineeVectorTrack = StaticCast<UInterpTrackVectorProp*>(Track);
diff --git a/Engine/Source/Editor/MovieSceneTools/Private/MatineeImportTools.cpp b/Engine/Source/Editor/MovieSceneTools/Private/MatineeImportTools.cpp
--- a/Engine/Source/Editor/MovieSceneTools/Private/MatineeImportTools.cpp
+++ b/Engine/Source/Editor/MovieSceneTools/Private/MatineeImportTools.cpp
@@ -12,10 +12,14 @@
 #include "Channels/MovieSceneEvent.h"
 #include "K2Node_FunctionEntry.h"
 #include "Kismet2/KismetEditorUtilities.h"
+#include "Animation/SkeletalMeshActor.h"
+#include "EngineUtils.h"
 
 #include "Matinee/MatineeActor.h"
 #include "Matinee/InterpData.h"
 #include "Matinee/InterpGroupInst.h"
+#include "Matinee/InterpTrackFloatMaterialParam.h"
+#include "Matinee/InterpTrackVectorMaterialParam.h"
 #include "Matinee/InterpTrackLinearColorProp.h"
 #include "Matinee/InterpTrackColorProp.h"
 #include "Matinee/InterpTrackBoolProp.h"
@@ -30,6 +34,7 @@
 
 #include "Tracks/MovieSceneBoolTrack.h"
 #include "Tracks/MovieSceneFloatTrack.h"
+#include "Tracks/MovieSceneMaterialTrack.h"
 #include "Tracks/MovieSceneColorTrack.h"
 #include "Tracks/MovieScene3DTransformTrack.h"
 #include "Tracks/MovieSceneParticleTrack.h"
@@ -52,6 +57,7 @@
 #include "Sections/MovieSceneEventSection.h"
 #include "Sections/MovieSceneEventTriggerSection.h"
 #include "Sections/MovieSceneVectorSection.h"
+#include "Sections/MovieSceneParameterSection.h"
 
 
 #include "Animation/AnimSequence.h"
@@ -120,7 +126,7 @@ bool FMatineeImportTools::TryConvertMatineeToggleToOutParticleKey( ETrackToggleA
 
 
 void FMatineeImportTools::SetOrAddKey(TMovieSceneChannelData<FMovieSceneFloatValue>& ChannelData, FFrameNumber Time, float Value, float ArriveTangent, float LeaveTangent, EInterpCurveMode MatineeInterpMode, FFrameRate FrameRate
-	, ERichCurveTangentWeightMode WeightedMode, float ArriveTangentWeight, float LeaveTangentWeight )
+	, ERichCurveTangentWeightMode WeightedMode, float ArriveTangentWeight, float LeaveTangentWeight)
 {
 	if (ChannelData.FindKey(Time) == INDEX_NONE)
 	{
@@ -224,6 +230,131 @@ bool FMatineeImportTools::CopyInterpFloatTrack( UInterpTrackFloatBase* MatineeFl
 	return bSectionCreated;
 }
 
+bool FMatineeImportTools::CopyInterpMaterialParamTrack(UInterpTrackFloatMaterialParam* MatineeMaterialParamTrack, UMovieSceneComponentMaterialTrack * MaterialTrack)
+{
+	const FScopedTransaction Transaction(NSLOCTEXT("Sequencer", "PasteMatineeFloatMaterialParamTrack", "Paste Matinee Float Material Param Track"));
+	bool bSectionCreated = false;
+
+	MaterialTrack->Modify();
+
+	FFrameRate   FrameRate = MaterialTrack->GetTypedOuter<UMovieScene>()->GetTickResolution();
+	FFrameNumber FirstKeyTime = (MatineeMaterialParamTrack->GetKeyframeTime(0) * FrameRate).RoundToFrame();
+
+	UMovieSceneParameterSection* Section = Cast<UMovieSceneParameterSection>(MovieSceneHelpers::FindSectionAtTime(MaterialTrack->GetAllSections(), FirstKeyTime));
+	if (Section == nullptr)
+	{
+		Section = Cast<UMovieSceneParameterSection>(MaterialTrack->CreateNewSection());
+		MaterialTrack->AddSection(*Section);
+		Section->SetRange(TRange<FFrameNumber>::All());
+		bSectionCreated = true;
+	}
+	if (Section->TryModify())
+	{
+		auto FirstPoint = MatineeMaterialParamTrack->FloatTrack.Points[0];
+		FFrameNumber KeyTime = (FirstPoint.InVal * FrameRate).RoundToFrame();
+
+		// The section needs a key added to be initialized, so add the first key from the matinee track.
+		Section->AddScalarParameterKey(MatineeMaterialParamTrack->ParamName, KeyTime, FirstPoint.OutVal);
+
+		TArray<FScalarParameterNameAndCurve> ScalarParms = Section->GetScalarParameterNamesAndCurves();
+
+		FMovieSceneFloatChannel* Channel = nullptr;
+		int32 CurveIndex = 0;
+		for (FScalarParameterNameAndCurve NameAndCurve : ScalarParms)
+		{
+			if (NameAndCurve.ParameterName == MatineeMaterialParamTrack->ParamName)
+			{
+				Channel = Section->GetChannelProxy().GetChannel<FMovieSceneFloatChannel>(CurveIndex);
+				break;
+			}
+			++CurveIndex;
+		}
+
+		if (!Channel)
+		{
+			return false;
+		}
+
+		
+		TMovieSceneChannelData<FMovieSceneFloatValue> ChannelData = Channel->GetData();
+
+
+		for (const auto& Point : MatineeMaterialParamTrack->FloatTrack.Points)
+		{
+			FFrameNumber KeyTime = (Point.InVal * FrameRate).RoundToFrame();
+
+			FMatineeImportTools::SetOrAddKey(ChannelData, KeyTime, Point.OutVal, Point.ArriveTangent, Point.LeaveTangent, Point.InterpMode, FrameRate);
+		}
+
+		CleanupCurveKeys(Channel);
+	}
+
+	return bSectionCreated;
+}
+
+bool FMatineeImportTools::CopyInterpMaterialParamTrack(UInterpTrackVectorMaterialParam * MatineeMaterialParamTrack, UMovieSceneComponentMaterialTrack * MaterialTrack)
+{
+	const FScopedTransaction Transaction(NSLOCTEXT("Sequencer", "PasteMatineeFloatMaterialParamTrack", "Paste Matinee Float Material Param Track"));
+	bool bSectionCreated = false;
+
+	MaterialTrack->Modify();
+
+	FFrameRate   FrameRate = MaterialTrack->GetTypedOuter<UMovieScene>()->GetTickResolution();
+	FFrameNumber FirstKeyTime = (MatineeMaterialParamTrack->GetKeyframeTime(0) * FrameRate).RoundToFrame();
+
+	UMovieSceneParameterSection* Section = Cast<UMovieSceneParameterSection>(MovieSceneHelpers::FindSectionAtTime(MaterialTrack->GetAllSections(), FirstKeyTime));
+	if (Section == nullptr)
+	{
+		Section = Cast<UMovieSceneParameterSection>(MaterialTrack->CreateNewSection());
+		MaterialTrack->AddSection(*Section);
+		Section->SetRange(TRange<FFrameNumber>::All());
+		bSectionCreated = true;
+	}
+	if (Section->TryModify())
+	{
+		FInterpCurvePointVector FirstPoint = MatineeMaterialParamTrack->VectorTrack.Points[0];
+		FFrameNumber KeyTime = (FirstPoint.InVal * FrameRate).RoundToFrame();
+
+		// The section needs a key added to be initialized, so add the first key from the matinee track.
+		Section->AddVectorParameterKey(MatineeMaterialParamTrack->ParamName, KeyTime, FirstPoint.OutVal);
+
+		TArray<FScalarParameterNameAndCurve> ScalarParms = Section->GetScalarParameterNamesAndCurves();
+
+		FMovieSceneFloatChannel* Channel = nullptr;
+		int32 CurveIndex = 0;
+		for (FScalarParameterNameAndCurve NameAndCurve : ScalarParms)
+		{
+			if (NameAndCurve.ParameterName == MatineeMaterialParamTrack->ParamName)
+			{
+				Channel = Section->GetChannelProxy().GetChannel<FMovieSceneFloatChannel>(CurveIndex);
+				break;
+			}
+			++CurveIndex;
+		}
+
+		TArrayView<FMovieSceneFloatChannel*> Channels = Section->GetChannelProxy().GetChannels<FMovieSceneFloatChannel>();
+		TMovieSceneChannelData<FMovieSceneFloatValue> ChannelData[3] = {
+			Channels[CurveIndex]->GetData(), Channels[CurveIndex + 1]->GetData(), Channels[CurveIndex + 2]->GetData()
+		};
+
+
+		for (const auto& Point : MatineeMaterialParamTrack->VectorTrack.Points)
+		{
+			FFrameNumber KeyTime = (Point.InVal * FrameRate).RoundToFrame();
+
+			FMatineeImportTools::SetOrAddKey(ChannelData[0], KeyTime, Point.OutVal.X, Point.ArriveTangent.X, Point.LeaveTangent.X, Point.InterpMode, FrameRate);
+			FMatineeImportTools::SetOrAddKey(ChannelData[1], KeyTime, Point.OutVal.Y, Point.ArriveTangent.Y, Point.LeaveTangent.Y, Point.InterpMode, FrameRate);
+			FMatineeImportTools::SetOrAddKey(ChannelData[2], KeyTime, Point.OutVal.Z, Point.ArriveTangent.Z, Point.LeaveTangent.Z, Point.InterpMode, FrameRate);
+		}
+
+		CleanupCurveKeys(Channels[0]);
+		CleanupCurveKeys(Channels[1]);
+		CleanupCurveKeys(Channels[2]);
+	}
+
+	return bSectionCreated;
+}
+
 bool FMatineeImportTools::CopyInterpVectorTrack( UInterpTrackVectorProp* MatineeVectorTrack, UMovieSceneVectorTrack* VectorTrack )
 {
 	const FScopedTransaction Transaction( NSLOCTEXT( "Sequencer", "PasteMatineeVectorTrack", "Paste Matinee Vector Track" ) );
@@ -499,7 +630,7 @@ bool FMatineeImportTools::CopyInterpMoveTrack( UInterpTrackMove* MoveTrack, UMov
 				}
 			}
 		}
-
+		
 		CleanupCurveKeys(Channels[0]);
 		CleanupCurveKeys(Channels[1]);
 		CleanupCurveKeys(Channels[2]);
diff --git a/Engine/Source/Editor/MovieSceneTools/Public/MatineeImportTools.h b/Engine/Source/Editor/MovieSceneTools/Public/MatineeImportTools.h
--- a/Engine/Source/Editor/MovieSceneTools/Public/MatineeImportTools.h
+++ b/Engine/Source/Editor/MovieSceneTools/Public/MatineeImportTools.h
@@ -11,6 +11,8 @@ class IMovieScenePlayer;
 class UInterpTrackAnimControl;
 class UInterpTrackBoolProp;
 class UInterpTrackColorProp;
+class UInterpTrackFloatMaterialParam;
+class UInterpTrackVectorMaterialParam;
 class UInterpTrackDirector;
 class UInterpTrackEvent;
 class UInterpTrackFade;
@@ -25,6 +27,7 @@ class UMovieSceneAudioTrack;
 class UMovieSceneBoolTrack;
 class UMovieSceneCameraCutTrack;
 class UMovieSceneColorTrack;
+class UMovieSceneComponentMaterialTrack;
 class UMovieSceneEventTrack;
 class UMovieSceneFadeTrack;
 class UMovieSceneFloatTrack;
@@ -56,6 +59,12 @@ public:
 	/** Copies keys from a matinee float track to a sequencer float track. */
 	static bool CopyInterpFloatTrack( UInterpTrackFloatBase* MatineeFloatTrack, UMovieSceneFloatTrack* FloatTrack );
 
+	/** Copies keys from a matinee material paramter track to a sequencer material track*/
+	static bool CopyInterpMaterialParamTrack( UInterpTrackFloatMaterialParam* MatineeMaterialParamTrack, UMovieSceneComponentMaterialTrack* MaterialTrack );
+
+	/** Copies keys from a matinee material paramter track to a sequencer material track */
+	static bool CopyInterpMaterialParamTrack( UInterpTrackVectorMaterialParam* MatineeMaterialParamTrack, UMovieSceneComponentMaterialTrack* MaterialTrack);
+
 	/** Copies keys from a matinee vector track to a sequencer vector track. */
 	static bool CopyInterpVectorTrack( UInterpTrackVectorProp* MatineeVectorTrack, UMovieSceneVectorTrack* VectorTrack );
 
diff --git a/Engine/Source/Runtime/Engine/Classes/Matinee/InterpTrackFloatMaterialParam.h b/Engine/Source/Runtime/Engine/Classes/Matinee/InterpTrackFloatMaterialParam.h
--- a/Engine/Source/Runtime/Engine/Classes/Matinee/InterpTrackFloatMaterialParam.h
+++ b/Engine/Source/Runtime/Engine/Classes/Matinee/InterpTrackFloatMaterialParam.h
@@ -10,7 +10,7 @@
 class UInterpTrackInst;
 struct FPropertyChangedEvent;
 
-UCLASS(meta=( DisplayName = "Float Material Parameter Track" ) )
+UCLASS(MinimalAPI, meta=( DisplayName = "Float Material Parameter Track" ) )
 class UInterpTrackFloatMaterialParam : public UInterpTrackFloatBase
 {
 	GENERATED_UCLASS_BODY()
