Fix CIS Warning

diff --git a/Engine/Plugins/MovieScene/MatineeToLevelSequence/Source/MatineeToLevelSequence/Private/MatineeToLevelSequenceModule.cpp b/Engine/Plugins/MovieScene/MatineeToLevelSequence/Source/MatineeToLevelSequence/Private/MatineeToLevelSequenceModule.cpp
--- a/Engine/Plugins/MovieScene/MatineeToLevelSequence/Source/MatineeToLevelSequence/Private/MatineeToLevelSequenceModule.cpp
+++ b/Engine/Plugins/MovieScene/MatineeToLevelSequence/Source/MatineeToLevelSequence/Private/MatineeToLevelSequenceModule.cpp
@@ -494,7 +494,7 @@ protected:
 				{
 					for (auto Track : Binding.GetTracks())
 					{
-						UMovieSceneComponentMaterialTrack * MaterialTrack = Cast<UMovieSceneComponentMaterialTrack>(Track);
+						MaterialTrack = Cast<UMovieSceneComponentMaterialTrack>(Track);
 						if (MaterialTrack)
 						{
 							BoundTracks.Add(MaterialTrack);
@@ -517,7 +517,7 @@ protected:
 			}
 			else
 			{
-				UMovieSceneComponentMaterialTrack* MaterialTrack = NewMovieScene->AddTrack<UMovieSceneComponentMaterialTrack>(ComponentGuid);
+				MaterialTrack = NewMovieScene->AddTrack<UMovieSceneComponentMaterialTrack>(ComponentGuid);
 				if (MaterialTrack)
 				{
 					MaterialTrack->SetMaterialIndex(MaterialIndex);
diff --git a/Engine/Source/Editor/MovieSceneTools/Private/MatineeImportTools.cpp b/Engine/Source/Editor/MovieSceneTools/Private/MatineeImportTools.cpp
--- a/Engine/Source/Editor/MovieSceneTools/Private/MatineeImportTools.cpp
+++ b/Engine/Source/Editor/MovieSceneTools/Private/MatineeImportTools.cpp
@@ -281,9 +281,9 @@ bool FMatineeImportTools::CopyInterpMaterialParamTrack(UInterpTrackFloatMaterial
 
 		for (const auto& Point : MatineeMaterialParamTrack->FloatTrack.Points)
 		{
-			FFrameNumber KeyTime = (Point.InVal * FrameRate).RoundToFrame();
+			FFrameNumber CurrentKeyTime = (Point.InVal * FrameRate).RoundToFrame();
 
-			FMatineeImportTools::SetOrAddKey(ChannelData, KeyTime, Point.OutVal, Point.ArriveTangent, Point.LeaveTangent, Point.InterpMode, FrameRate);
+			FMatineeImportTools::SetOrAddKey(ChannelData, CurrentKeyTime, Point.OutVal, Point.ArriveTangent, Point.LeaveTangent, Point.InterpMode, FrameRate);
 		}
 
 		CleanupCurveKeys(Channel);
@@ -340,11 +340,11 @@ bool FMatineeImportTools::CopyInterpMaterialParamTrack(UInterpTrackVectorMateria
 
 		for (const auto& Point : MatineeMaterialParamTrack->VectorTrack.Points)
 		{
-			FFrameNumber KeyTime = (Point.InVal * FrameRate).RoundToFrame();
+			FFrameNumber CurrentKeyTime = (Point.InVal * FrameRate).RoundToFrame();
 
-			FMatineeImportTools::SetOrAddKey(ChannelData[0], KeyTime, Point.OutVal.X, Point.ArriveTangent.X, Point.LeaveTangent.X, Point.InterpMode, FrameRate);
-			FMatineeImportTools::SetOrAddKey(ChannelData[1], KeyTime, Point.OutVal.Y, Point.ArriveTangent.Y, Point.LeaveTangent.Y, Point.InterpMode, FrameRate);
-			FMatineeImportTools::SetOrAddKey(ChannelData[2], KeyTime, Point.OutVal.Z, Point.ArriveTangent.Z, Point.LeaveTangent.Z, Point.InterpMode, FrameRate);
+			FMatineeImportTools::SetOrAddKey(ChannelData[0], CurrentKeyTime, Point.OutVal.X, Point.ArriveTangent.X, Point.LeaveTangent.X, Point.InterpMode, FrameRate);
+			FMatineeImportTools::SetOrAddKey(ChannelData[1], CurrentKeyTime, Point.OutVal.Y, Point.ArriveTangent.Y, Point.LeaveTangent.Y, Point.InterpMode, FrameRate);
+			FMatineeImportTools::SetOrAddKey(ChannelData[2], CurrentKeyTime, Point.OutVal.Z, Point.ArriveTangent.Z, Point.LeaveTangent.Z, Point.InterpMode, FrameRate);
 		}
 
 		CleanupCurveKeys(Channels[0]);
