Sequencer: Always try to match by name first. Otherwise, if the binding isn't found and not matching by name, default to first and output a warning.

diff --git a/Engine/Source/Editor/MovieSceneTools/Private/MovieSceneToolHelpers.cpp b/Engine/Source/Editor/MovieSceneTools/Private/MovieSceneToolHelpers.cpp
--- a/Engine/Source/Editor/MovieSceneTools/Private/MovieSceneToolHelpers.cpp
+++ b/Engine/Source/Editor/MovieSceneTools/Private/MovieSceneToolHelpers.cpp
@@ -1086,25 +1086,8 @@ bool ImportFBXTransform(FString NodeName, FGuid ObjectBinding, UnFbx::FFbxCurves
 	return true;
 }
 
-bool ImportFBXNode(FString NodeName, UnFbx::FFbxCurvesAPI& CurveAPI, UMovieScene* InMovieScene, ISequencer& InSequencer, const TMap<FGuid, FString>& InObjectBindingMap, bool bMatchByNameOnly)
+bool ImportFBXNode(FString NodeName, UnFbx::FFbxCurvesAPI& CurveAPI, UMovieScene* InMovieScene, ISequencer& InSequencer, FGuid ObjectBinding)
 {
-	// Find the matching object binding to apply this animation to. If not matching by name only, default to the first.
-	FGuid ObjectBinding;
-	for (auto It = InObjectBindingMap.CreateConstIterator(); It; ++It)
-	{
-		if (!bMatchByNameOnly || FCString::Strcmp(*It.Value().ToUpper(), *NodeName.ToUpper()) == 0)
-		{
-			ObjectBinding = It.Key();
-			break;
-		}
-	}
-
-	if (!ObjectBinding.IsValid())
-	{
-		UE_LOG(LogMovieScene, Warning, TEXT("Fbx Import: Failed to find any matching node for (%s)."), *NodeName);
-		return false;
-	}
-
 	// Look for animated float properties
 	TArray<FString> AnimatedPropertyNames;
 	CurveAPI.GetNodeAnimatedPropertyNameArray(NodeName, AnimatedPropertyNames);
@@ -1582,9 +1565,57 @@ private:
 		TArray<FString> AllNodeNames;
 		CurveAPI.GetAllNodeNameArray(AllNodeNames);
 
+		// First try matching by name
+		for (int32 NodeIndex = 0; NodeIndex < AllNodeNames.Num(); )
+		{
+			FString NodeName = AllNodeNames[NodeIndex];
+			bool bFoundMatch = false;
+			for (auto It = ObjectBindingMap.CreateConstIterator(); It; ++It)
+			{
+				if (FCString::Strcmp(*It.Value().ToUpper(), *NodeName.ToUpper()) == 0)
+				{
+					ImportFBXNode(NodeName, CurveAPI, MovieScene, *Sequencer, It.Key());
+
+					ObjectBindingMap.Remove(It.Key());
+					AllNodeNames.RemoveAt(NodeIndex);
+
+					bFoundMatch = true;
+					break;
+				}
+			}
+
+			if (bFoundMatch)
+			{
+				continue;
+			}
+
+			++NodeIndex;
+		}
+
+		// Otherwise, get the first available node that hasn't been imported onto yet
+		if (!bMatchByNameOnly)
+		{
+			for (int32 NodeIndex = 0; NodeIndex < AllNodeNames.Num(); )
+			{
+				FString NodeName = AllNodeNames[NodeIndex];
+				auto It = ObjectBindingMap.CreateConstIterator();
+				if (It)
+				{
+					ImportFBXNode(NodeName, CurveAPI, MovieScene, *Sequencer, It.Key());
+
+					UE_LOG(LogMovieScene, Warning, TEXT("Fbx Import: Failed to find any matching node for (%s). Defaulting to first available (%s)."), *NodeName, *It.Value());
+					ObjectBindingMap.Remove(It.Key());
+					AllNodeNames.RemoveAt(NodeIndex);
+					continue;
+				}
+
+				++NodeIndex;
+			}
+		}
+
 		for (FString NodeName : AllNodeNames)
 		{
-			ImportFBXNode(NodeName, CurveAPI, MovieScene, *Sequencer, ObjectBindingMap, bMatchByNameOnly);
+			UE_LOG(LogMovieScene, Warning, TEXT("Fbx Import: Failed to find any matching node for (%s)."), *NodeName);
 		}
 
 		Sequencer->NotifyMovieSceneDataChanged(EMovieSceneDataChangeType::MovieSceneStructureItemAdded);
