CL 4443114 Fixed Matinee FrameRate conversions but added a double frame rate conversion when importing FBX.  The SetOrAddKey function now does the conversion.

diff --git a/Engine/Source/Editor/MovieSceneTools/Private/MovieSceneToolHelpers.cpp b/Engine/Source/Editor/MovieSceneTools/Private/MovieSceneToolHelpers.cpp
--- a/Engine/Source/Editor/MovieSceneTools/Private/MovieSceneToolHelpers.cpp
+++ b/Engine/Source/Editor/MovieSceneTools/Private/MovieSceneToolHelpers.cpp
@@ -974,8 +974,6 @@ void ImportTransformChannel(const FRichCurve& Source, FMovieSceneFloatChannel* D
 
 	TMovieSceneChannelData<FMovieSceneFloatValue> ChannelData = Dest->GetData();
 	ChannelData.Reset();
-	double DecimalRate = DestFrameRate.AsDecimal();
-
 	for (auto SourceIt = Source.GetKeyHandleIterator(); SourceIt; ++SourceIt)
 	{
 		const FRichCurveKey Key = Source.GetKey(SourceIt.Key());
@@ -984,7 +982,7 @@ void ImportTransformChannel(const FRichCurve& Source, FMovieSceneFloatChannel* D
 		if (Source.IsKeyHandleValid(PrevKeyHandle))
 		{
 			const FRichCurveKey PrevKey = Source.GetKey(PrevKeyHandle);
-			ArriveTangent = ArriveTangent / ((Key.Time - PrevKey.Time) * DecimalRate);
+			ArriveTangent = ArriveTangent / (Key.Time - PrevKey.Time);
 
 		}
 		float LeaveTangent = Key.LeaveTangent;
@@ -992,7 +990,7 @@ void ImportTransformChannel(const FRichCurve& Source, FMovieSceneFloatChannel* D
 		if (Source.IsKeyHandleValid(NextKeyHandle))
 		{
 			const FRichCurveKey NextKey = Source.GetKey(NextKeyHandle);
-			LeaveTangent = LeaveTangent / ((NextKey.Time - Key.Time) * DecimalRate);
+			LeaveTangent = LeaveTangent / (NextKey.Time - Key.Time);
 		}
 
 		if (bNegateTangents)
