Renderpass updates.

TrackEditorThumbnail. self-contained.

diff --git a/Engine/Source/Editor/MovieSceneTools/Private/TrackEditorThumbnail/TrackEditorThumbnail.cpp b/Engine/Source/Editor/MovieSceneTools/Private/TrackEditorThumbnail/TrackEditorThumbnail.cpp
--- a/Engine/Source/Editor/MovieSceneTools/Private/TrackEditorThumbnail/TrackEditorThumbnail.cpp
+++ b/Engine/Source/Editor/MovieSceneTools/Private/TrackEditorThumbnail/TrackEditorThumbnail.cpp
@@ -136,57 +136,56 @@ void FTrackEditorThumbnail::CopyTextureIn(FTexture2DRHIRef SourceTexture)
 		const FSceneRenderTargetItem& DestRenderTarget = ResampleTexturePooledRenderTarget->GetRenderTargetItem();
 
 		FGraphicsPipelineStateInitializer GraphicsPSOInit;
-		SetRenderTarget(RHICmdList, DestRenderTarget.TargetableTexture, FTextureRHIRef());
-		RHICmdList.ApplyCachedRenderTargets(GraphicsPSOInit);
-		GraphicsPSOInit.PrimitiveType = PT_TriangleList;
-
-		RHICmdList.SetViewport(0, 0, 0.0f, TargetSize.X, TargetSize.Y, 1.0f);
-//		RHICmdList.ClearColorTexture(DestRenderTarget.TargetableTexture, FLinearColor(0.0f, 0.0f, 0.0f, 1.0f));
-
-		GraphicsPSOInit.BlendState = TStaticBlendState<>::GetRHI();
-		GraphicsPSOInit.RasterizerState = TStaticRasterizerState<>::GetRHI();
-		GraphicsPSOInit.DepthStencilState = TStaticDepthStencilState<false,CF_Always>::GetRHI();
-
-		const ERHIFeatureLevel::Type FeatureLevel = GMaxRHIFeatureLevel;
-		
-		TShaderMap<FGlobalShaderType>* ShaderMap = GetGlobalShaderMap(FeatureLevel);
-		TShaderMapRef<FScreenVS> VertexShader(ShaderMap);
-		TShaderMapRef<FScreenPS> PixelShader(ShaderMap);
-
-		GraphicsPSOInit.BoundShaderState.VertexDeclarationRHI = RendererModule->GetFilterVertexDeclaration().VertexDeclarationRHI;
-		GraphicsPSOInit.BoundShaderState.VertexShaderRHI = GETSAFERHISHADER_VERTEX(*VertexShader);
-		GraphicsPSOInit.BoundShaderState.PixelShaderRHI = GETSAFERHISHADER_PIXEL(*PixelShader);
-
-		SetGraphicsPipelineState(RHICmdList, GraphicsPSOInit);
-
-		PixelShader->SetParameters(RHICmdList, TStaticSamplerState<SF_Bilinear>::GetRHI(), SourceTexture);
-
-		const float Scale = FMath::Min(float(SourceTexture->GetSizeX()) / TargetTexture->GetWidth(), float(SourceTexture->GetSizeY()) / TargetTexture->GetHeight());
-		const float Left = (SourceTexture->GetSizeX() - TargetTexture->GetWidth()*Scale) * .5f;
-		const float Top = (SourceTexture->GetSizeY() - TargetTexture->GetHeight()*Scale) * .5f;
-
-		const float U = Left / float(SourceTexture->GetSizeX());
-		const float V = Top / float(SourceTexture->GetSizeY());
-		const float SizeU = float(TargetTexture->GetWidth() * Scale) / float(SourceTexture->GetSizeX());
-		const float SizeV = float(TargetTexture->GetHeight() * Scale) / float(SourceTexture->GetSizeY());
-
-		RendererModule->DrawRectangle(
-			RHICmdList,
-			0, 0,									// Dest X, Y
-			TargetSize.X,							// Dest Width
-			TargetSize.Y,							// Dest Height
-			U, V,									// Source U, V
-			SizeU, SizeV,							// Source USize, VSize
-			TargetSize,								// Target buffer size
-			FIntPoint(1, 1),						// Source texture size
-			*VertexShader,
-			EDRF_Default);
-
-		// Asynchronously copy render target from GPU to CPU
-		RHICmdList.CopyToResolveTarget(
-			DestRenderTarget.TargetableTexture,
-			TargetTexture->GetTypedResource(),
-			FResolveParams());
+		FRHIRenderPassInfo RPInfo(DestRenderTarget.TargetableTexture, ERenderTargetActions::Load_Store);
+		RHICmdList.BeginRenderPass(RPInfo, TEXT("CopyTextureIn"));
+		{
+			RHICmdList.ApplyCachedRenderTargets(GraphicsPSOInit);
+			GraphicsPSOInit.PrimitiveType = PT_TriangleList;
+
+			RHICmdList.SetViewport(0, 0, 0.0f, TargetSize.X, TargetSize.Y, 1.0f);
+			//		RHICmdList.ClearColorTexture(DestRenderTarget.TargetableTexture, FLinearColor(0.0f, 0.0f, 0.0f, 1.0f));
+
+			GraphicsPSOInit.BlendState = TStaticBlendState<>::GetRHI();
+			GraphicsPSOInit.RasterizerState = TStaticRasterizerState<>::GetRHI();
+			GraphicsPSOInit.DepthStencilState = TStaticDepthStencilState<false, CF_Always>::GetRHI();
+
+			const ERHIFeatureLevel::Type FeatureLevel = GMaxRHIFeatureLevel;
+
+			TShaderMap<FGlobalShaderType>* ShaderMap = GetGlobalShaderMap(FeatureLevel);
+			TShaderMapRef<FScreenVS> VertexShader(ShaderMap);
+			TShaderMapRef<FScreenPS> PixelShader(ShaderMap);
+
+			GraphicsPSOInit.BoundShaderState.VertexDeclarationRHI = RendererModule->GetFilterVertexDeclaration().VertexDeclarationRHI;
+			GraphicsPSOInit.BoundShaderState.VertexShaderRHI = GETSAFERHISHADER_VERTEX(*VertexShader);
+			GraphicsPSOInit.BoundShaderState.PixelShaderRHI = GETSAFERHISHADER_PIXEL(*PixelShader);
+
+			SetGraphicsPipelineState(RHICmdList, GraphicsPSOInit);
+
+			PixelShader->SetParameters(RHICmdList, TStaticSamplerState<SF_Bilinear>::GetRHI(), SourceTexture);
+
+			const float Scale = FMath::Min(float(SourceTexture->GetSizeX()) / TargetTexture->GetWidth(), float(SourceTexture->GetSizeY()) / TargetTexture->GetHeight());
+			const float Left = (SourceTexture->GetSizeX() - TargetTexture->GetWidth()*Scale) * .5f;
+			const float Top = (SourceTexture->GetSizeY() - TargetTexture->GetHeight()*Scale) * .5f;
+
+			const float U = Left / float(SourceTexture->GetSizeX());
+			const float V = Top / float(SourceTexture->GetSizeY());
+			const float SizeU = float(TargetTexture->GetWidth() * Scale) / float(SourceTexture->GetSizeX());
+			const float SizeV = float(TargetTexture->GetHeight() * Scale) / float(SourceTexture->GetSizeY());
+
+			RendererModule->DrawRectangle(
+				RHICmdList,
+				0, 0,									// Dest X, Y
+				TargetSize.X,							// Dest Width
+				TargetSize.Y,							// Dest Height
+				U, V,									// Source U, V
+				SizeU, SizeV,							// Source USize, VSize
+				TargetSize,								// Target buffer size
+				FIntPoint(1, 1),						// Source texture size
+				*VertexShader,
+				EDRF_Default);
+		}
+		RHICmdList.EndRenderPass();
+		RHICmdList.CopyToResolveTarget(DestRenderTarget.TargetableTexture, TargetTexture->GetTypedResource(), FResolveParams());
 
 		*bHasFinishedDrawingPtr = true;
 	};
