Make sure sequence is fully loaded before compilation begins

diff --git a/Engine/Source/Editor/MovieSceneTools/Private/K2Node_GetSequenceBinding.cpp b/Engine/Source/Editor/MovieSceneTools/Private/K2Node_GetSequenceBinding.cpp
--- a/Engine/Source/Editor/MovieSceneTools/Private/K2Node_GetSequenceBinding.cpp
+++ b/Engine/Source/Editor/MovieSceneTools/Private/K2Node_GetSequenceBinding.cpp
@@ -205,7 +205,7 @@ FNodeHandlingFunctor* UK2Node_GetSequenceBinding::CreateNodeHandler(FKismetCompi
 
 void UK2Node_GetSequenceBinding::PreloadRequiredAssets()
 {
-	GetSequence();
+	EnsureFullyLoaded(GetSequence());
 }
 
 FText UK2Node_GetSequenceBinding::GetSequenceName() const
