Sequencer: Specify precache frames when generating audio track thumbnails

diff --git a/Engine/Source/Editor/MovieSceneTools/Private/TrackEditors/AudioTrackEditor.cpp b/Engine/Source/Editor/MovieSceneTools/Private/TrackEditors/AudioTrackEditor.cpp
--- a/Engine/Source/Editor/MovieSceneTools/Private/TrackEditors/AudioTrackEditor.cpp
+++ b/Engine/Source/Editor/MovieSceneTools/Private/TrackEditors/AudioTrackEditor.cpp
@@ -306,7 +306,7 @@ void FAudioThumbnail::GenerateWaveformPreview(TArray<uint8>& OutData, TRange<flo
 
 			if ( SoundWave->InitAudioResource( AudioDevice->GetRuntimeFormat( SoundWave ) ) && (SoundWave->DecompressionType != DTYPE_RealTime || SoundWave->CachedRealtimeFirstBuffer == nullptr ) )
 			{
-				FAsyncAudioDecompress TempDecompress(SoundWave, 0);
+				FAsyncAudioDecompress TempDecompress(SoundWave, AudioDevice->NumPrecacheFrames);
 				TempDecompress.StartSynchronousTask();
 			}
 
