Sequencer: Fix field of view not getting imported through fbx.

There were actually a few issues here:

 - The sequencer import is always creating a cine camera actor. But the fbx camera in this case is not compatible - it's just a simple camera with field of view. Solution: Create a regular camera in this case.
 - The name for the camera doesn't match - that's causing the name matching to fail when copying properties from the fbx camera to the ue4 camera. Solution: Get the name from the camera node and not the camera itself.
 - Sequencer has a mapping that determines what fbx property names should match to what ue4 property names. That's defined in the ini and does not include FieldOfView. Solution: Add FieldOfView to ini.

diff --git a/Engine/Config/BaseEditorPerProjectUserSettings.ini b/Engine/Config/BaseEditorPerProjectUserSettings.ini
--- a/Engine/Config/BaseEditorPerProjectUserSettings.ini
+++ b/Engine/Config/BaseEditorPerProjectUserSettings.ini
@@ -812,6 +812,7 @@ LoopMode=SLM_Loop
 +TrackSettings=(MatchingActorClass=/Script/Engine.Light,DefaultTracks=(None),DefaultPropertyTracks=((ComponentPath="LightComponent0",PropertyPath="Intensity"),(ComponentPath="LightComponent0",PropertyPath="LightColor")))
 
 [/Script/MovieSceneTools.MovieSceneToolsProjectSettings]
++FbxSettings=(FbxPropertyName="FieldOfView", PropertyPath=(ComponentName="CameraComponent",PropertyName="FieldOfView"))
 +FbxSettings=(FbxPropertyName="FocalLength", PropertyPath=(ComponentName="CameraComponent",PropertyName="CurrentFocalLength"))
 +FbxSettings=(FbxPropertyName="FocusDistance", PropertyPath=(ComponentName="CameraComponent",PropertyName="FocusSettings.ManualFocusDistance"))
 
diff --git a/Engine/Source/Editor/MovieSceneTools/Private/MovieSceneToolHelpers.cpp b/Engine/Source/Editor/MovieSceneTools/Private/MovieSceneToolHelpers.cpp
--- a/Engine/Source/Editor/MovieSceneTools/Private/MovieSceneToolHelpers.cpp
+++ b/Engine/Source/Editor/MovieSceneTools/Private/MovieSceneToolHelpers.cpp
@@ -1161,7 +1161,7 @@ FbxNode* RetrieveObjectFromName(const TCHAR* ObjectName, FbxNode* Root)
 	return nullptr;
 }
 
-void CopyCameraProperties(FbxCamera* CameraNode, ACineCameraActor* CameraActor)
+void CopyCameraProperties(FbxCamera* CameraNode, AActor* InCameraActor)
 {
 	float FieldOfView;
 	float FocalLength;
@@ -1180,26 +1180,53 @@ void CopyCameraProperties(FbxCamera* CameraNode, ACineCameraActor* CameraActor)
 	float ApertureWidth = CameraNode->GetApertureWidth();
 	float ApertureHeight = CameraNode->GetApertureHeight();
 
-	UCineCameraComponent* CineCameraComponent = CameraActor->GetCineCameraComponent();
+	UCameraComponent* CameraComponent = nullptr;
 
-	CineCameraComponent->SetProjectionMode(CameraNode->ProjectionType.Get() == FbxCamera::ePerspective ? ECameraProjectionMode::Perspective : ECameraProjectionMode::Orthographic);
-	CineCameraComponent->SetAspectRatio(CameraNode->AspectWidth.Get() / CameraNode->AspectHeight.Get());
-	CineCameraComponent->SetOrthoNearClipPlane(CameraNode->NearPlane.Get());
-	CineCameraComponent->SetOrthoFarClipPlane(CameraNode->FarPlane.Get());
-	CineCameraComponent->SetOrthoWidth(CameraNode->OrthoZoom.Get());
-	CineCameraComponent->SetFieldOfView(FieldOfView);
-	CineCameraComponent->FilmbackSettings.SensorWidth = FUnitConversion::Convert(ApertureWidth, EUnit::Inches, EUnit::Millimeters);
-	CineCameraComponent->FilmbackSettings.SensorHeight = FUnitConversion::Convert(ApertureHeight, EUnit::Inches, EUnit::Millimeters);
-	CineCameraComponent->FocusSettings.ManualFocusDistance = CameraNode->FocusDistance;
-	if (FocalLength < CineCameraComponent->LensSettings.MinFocalLength)
+	if (ACineCameraActor* CineCameraActor = Cast<ACineCameraActor>(InCameraActor))
 	{
-		CineCameraComponent->LensSettings.MinFocalLength = FocalLength;
+		CameraComponent = CineCameraActor->GetCineCameraComponent();
+
+		UCineCameraComponent* CineCameraComponent = CineCameraActor->GetCineCameraComponent();
+		CineCameraComponent->FilmbackSettings.SensorWidth = FUnitConversion::Convert(ApertureWidth, EUnit::Inches, EUnit::Millimeters);
+		CineCameraComponent->FilmbackSettings.SensorHeight = FUnitConversion::Convert(ApertureHeight, EUnit::Inches, EUnit::Millimeters);
+		CineCameraComponent->FocusSettings.ManualFocusDistance = CameraNode->FocusDistance;
+		if (FocalLength < CineCameraComponent->LensSettings.MinFocalLength)
+		{
+			CineCameraComponent->LensSettings.MinFocalLength = FocalLength;
+		}
+		if (FocalLength > CineCameraComponent->LensSettings.MaxFocalLength)
+		{
+			CineCameraComponent->LensSettings.MaxFocalLength = FocalLength;
+		}
+		CineCameraComponent->CurrentFocalLength = FocalLength;
+	}
+	else if (ACameraActor* CameraActor = Cast<ACameraActor>(InCameraActor))
+	{
+		CameraComponent = CameraActor->GetCameraComponent();
 	}
-	if (FocalLength > CineCameraComponent->LensSettings.MaxFocalLength)
+
+	if (!CameraComponent)
+	{
+		return;
+	}
+
+	CameraComponent->SetProjectionMode(CameraNode->ProjectionType.Get() == FbxCamera::ePerspective ? ECameraProjectionMode::Perspective : ECameraProjectionMode::Orthographic);
+	CameraComponent->SetAspectRatio(CameraNode->AspectWidth.Get() / CameraNode->AspectHeight.Get());
+	CameraComponent->SetOrthoNearClipPlane(CameraNode->NearPlane.Get());
+	CameraComponent->SetOrthoFarClipPlane(CameraNode->FarPlane.Get());
+	CameraComponent->SetOrthoWidth(CameraNode->OrthoZoom.Get());
+	CameraComponent->SetFieldOfView(FieldOfView);
+}
+
+FString GetCameraName(FbxCamera* InCamera)
+{
+	FbxNode* CameraNode = InCamera->GetNode();
+	if (CameraNode)
 	{
-		CineCameraComponent->LensSettings.MaxFocalLength = FocalLength;
+		return CameraNode->GetName();
 	}
-	CineCameraComponent->CurrentFocalLength = FocalLength;
+
+	return InCamera->GetName();
 }
 
 void ImportFBXCamera(UnFbx::FFbxImporter* FbxImporter, UMovieScene* InMovieScene, ISequencer& InSequencer, TMap<FGuid, FString>& InObjectBindingMap, bool bMatchByNameOnly, bool bCreateCameras)
@@ -1215,13 +1242,13 @@ void ImportFBXCamera(UnFbx::FFbxImporter* FbxImporter, UMovieScene* InMovieScene
 		TArray<FbxCamera*> UnmatchedCameras;
 		for (auto Camera : AllCameras)
 		{
-			FString NodeName = FString(Camera->GetName());
+			FString NodeName = GetCameraName(Camera);
 
 			bool bMatched = false;
 			for (auto InObjectBinding : InObjectBindingMap)
 			{		
 				FString ObjectName = InObjectBinding.Value;
-				if (!FCString::Strcmp(*ObjectName, UTF8_TO_TCHAR(Camera->GetName())))
+				if (ObjectName == NodeName)
 				{
 					// Look for a valid bound object, otherwise need to create a new camera and assign this binding to it
 					bool bFoundBoundObject = false;
@@ -1260,11 +1287,21 @@ void ImportFBXCamera(UnFbx::FFbxImporter* FbxImporter, UMovieScene* InMovieScene
 		// Add any unmatched cameras
 		for (auto UnmatchedCamera : UnmatchedCameras)
 		{
-			FString CameraName = FString(ANSI_TO_TCHAR(UnmatchedCamera->GetName()));
+			FString CameraName = GetCameraName(UnmatchedCamera);
 
-			FActorSpawnParameters SpawnParams;
-			ACineCameraActor* NewCamera = World->SpawnActor<ACineCameraActor>(SpawnParams);
-			NewCamera->SetActorLabel(*CameraName);
+			AActor* NewCamera = nullptr;
+			if (UnmatchedCamera->GetApertureMode() == FbxCamera::eFocalLength)
+			{
+				FActorSpawnParameters SpawnParams;
+				NewCamera = World->SpawnActor<ACineCameraActor>(SpawnParams);
+				NewCamera->SetActorLabel(*CameraName);
+			}
+			else
+			{
+				FActorSpawnParameters SpawnParams;
+				NewCamera = World->SpawnActor<ACameraActor>(SpawnParams);
+				NewCamera->SetActorLabel(*CameraName);
+			}
 
 			// Copy camera properties before adding default tracks so that initial camera properties match and can be restored after sequencer finishes
 			CopyCameraProperties(UnmatchedCamera, NewCamera);
@@ -1308,7 +1345,7 @@ void ImportFBXCamera(UnFbx::FFbxImporter* FbxImporter, UMovieScene* InMovieScene
 			CameraNode = FindCamera(FbxImporter->Scene->GetRootNode());
 			if (CameraNode)
 			{
-				FString CameraName = FString(ANSI_TO_TCHAR(CameraNode->GetName()));
+				FString CameraName = GetCameraName(CameraNode);
 				FNotificationInfo Info(FText::Format(NSLOCTEXT("MovieSceneTools", "NoMatchingCameraWarning", "Failed to find any matching camera for {0}. Importing onto first camera from fbx {1}"), FText::FromString(ObjectName), FText::FromString(CameraName)));
 				Info.ExpireDuration = 5.0f;
 				FSlateNotificationManager::Get().AddNotification(Info)->SetCompletionState(SNotificationItem::CS_Fail);
@@ -1337,20 +1374,39 @@ void ImportFBXCamera(UnFbx::FFbxImporter* FbxImporter, UMovieScene* InMovieScene
 		for (TWeakObjectPtr<>& WeakObject : BoundObjects)
 		{
 			UObject* FoundObject = WeakObject.Get();
-			if (FoundObject && FoundObject->IsA(ACineCameraActor::StaticClass()))
+			if (FoundObject && FoundObject->GetClass()->IsChildOf(ACameraActor::StaticClass()))
 			{
-				ACineCameraActor* CineCameraActor = Cast<ACineCameraActor>(FoundObject);
-				UCineCameraComponent* CineCameraComponent = CineCameraActor->GetCineCameraComponent();
-				CopyCameraProperties(CameraNode, CineCameraActor);
+				CopyCameraProperties(CameraNode, Cast<AActor>(FoundObject));
 
-				// Set the default value of the current focal length section
-				FGuid PropertyOwnerGuid = InSequencer.GetHandleToObject(CineCameraComponent);
+				UCameraComponent* CameraComponent = nullptr;
+				FName TrackName;
+				float TrackValue;
+				
+				if (ACineCameraActor* CineCameraActor = Cast<ACineCameraActor>(FoundObject))
+				{
+					CameraComponent = CineCameraActor->GetCineCameraComponent();
+					TrackName = TEXT("CurrentFocalLength");
+					TrackValue = FocalLength;
+				}
+				else if (ACameraActor* CameraActor = Cast<ACameraActor>(FoundObject))
+				{
+					CameraComponent = CameraActor->GetCameraComponent();
+					TrackName = TEXT("FieldOfView");
+					TrackValue = FieldOfView;
+				}
+				else
+				{
+					continue;
+				}
+
+				// Set the default value of the current focal length or field of view section
+				FGuid PropertyOwnerGuid = InSequencer.GetHandleToObject(CameraComponent);
 				if (!PropertyOwnerGuid.IsValid())
 				{
 					continue;
 				}
 
-				UMovieSceneFloatTrack* FloatTrack = InMovieScene->FindTrack<UMovieSceneFloatTrack>(PropertyOwnerGuid, TEXT("CurrentFocalLength"));
+				UMovieSceneFloatTrack* FloatTrack = InMovieScene->FindTrack<UMovieSceneFloatTrack>(PropertyOwnerGuid, TrackName);
 				if (FloatTrack)
 				{
 					FloatTrack->Modify();
@@ -1370,7 +1426,7 @@ void ImportFBXCamera(UnFbx::FFbxImporter* FbxImporter, UMovieScene* InMovieScene
 						FloatSection->SetRange(TRange<FFrameNumber>::All());
 					}
 
-					FloatSection->GetChannelProxy().GetChannel<FMovieSceneFloatChannel>(0)->SetDefault(FocalLength);
+					FloatSection->GetChannelProxy().GetChannel<FMovieSceneFloatChannel>(0)->SetDefault(TrackValue);
 				}
 			}
 		}
@@ -1379,9 +1435,11 @@ void ImportFBXCamera(UnFbx::FFbxImporter* FbxImporter, UMovieScene* InMovieScene
 
 FGuid FindCameraGuid(FbxCamera* Camera, TMap<FGuid, FString>& InObjectBindingMap)
 {
+	FString CameraName = GetCameraName(Camera);
+
 	for (auto& Pair : InObjectBindingMap)
 	{
-		if (FCString::Strcmp(*Pair.Value, UTF8_TO_TCHAR(Camera->GetName())) == 0)
+		if (Pair.Value == CameraName)
 		{
 			return Pair.Key;
 		}
