Sequencer: Skip matching fbx root node

diff --git a/Engine/Source/Editor/MovieSceneTools/Private/MovieSceneToolHelpers.cpp b/Engine/Source/Editor/MovieSceneTools/Private/MovieSceneToolHelpers.cpp
--- a/Engine/Source/Editor/MovieSceneTools/Private/MovieSceneToolHelpers.cpp
+++ b/Engine/Source/Editor/MovieSceneTools/Private/MovieSceneToolHelpers.cpp
@@ -1621,10 +1621,18 @@ private:
 		TArray<FString> AllNodeNames;
 		CurveAPI.GetAllNodeNameArray(AllNodeNames);
 
+		FString RootNodeName = FbxImporter->Scene->GetRootNode()->GetName();
+
 		// First try matching by name
 		for (int32 NodeIndex = 0; NodeIndex < AllNodeNames.Num(); )
 		{
 			FString NodeName = AllNodeNames[NodeIndex];
+			if (RootNodeName == NodeName)
+			{
+				++NodeIndex;
+				continue;
+			}
+
 			bool bFoundMatch = false;
 			for (auto It = ObjectBindingMap.CreateConstIterator(); It; ++It)
 			{
@@ -1654,6 +1662,12 @@ private:
 			for (int32 NodeIndex = 0; NodeIndex < AllNodeNames.Num(); )
 			{
 				FString NodeName = AllNodeNames[NodeIndex];
+				if (RootNodeName == NodeName)
+				{
+					++NodeIndex;
+					continue;
+				}
+
 				auto It = ObjectBindingMap.CreateConstIterator();
 				if (It)
 				{
