Sequencer: Close Set Asset menu and rename MPC track when changing assets.

diff --git a/Engine/Source/Editor/MovieSceneTools/Private/TrackEditors/MaterialParameterCollectionTrackEditor.cpp b/Engine/Source/Editor/MovieSceneTools/Private/TrackEditors/MaterialParameterCollectionTrackEditor.cpp
--- a/Engine/Source/Editor/MovieSceneTools/Private/TrackEditors/MaterialParameterCollectionTrackEditor.cpp
+++ b/Engine/Source/Editor/MovieSceneTools/Private/TrackEditors/MaterialParameterCollectionTrackEditor.cpp
@@ -69,8 +69,11 @@ void FMaterialParameterCollectionTrackEditor::BuildTrackContextMenu(FMenuBuilder
 		{
 			FScopedTransaction Transaction(LOCTEXT("SetAssetTransaction", "Assign Material Parameter Collection"));
 			MPCTrack->Modify();
+			MPCTrack->SetDisplayName(FText::FromString(MPC->GetName()));
 			MPCTrack->MPC = MPC;
 		}
+
+		FSlateApplication::Get().DismissAllMenus();
 	};
 
 	auto AssignAssetEnterPressed = [AssignAsset](const TArray<FAssetData>& InAssetData)
