Sequencer: Fix switching takes so that the swapped in shot is swapped onto the original shot's track

diff --git a/Engine/Source/Editor/MovieSceneTools/Private/TrackEditors/CinematicShotTrackEditor.cpp b/Engine/Source/Editor/MovieSceneTools/Private/TrackEditors/CinematicShotTrackEditor.cpp
--- a/Engine/Source/Editor/MovieSceneTools/Private/TrackEditors/CinematicShotTrackEditor.cpp
+++ b/Engine/Source/Editor/MovieSceneTools/Private/TrackEditors/CinematicShotTrackEditor.cpp
@@ -525,8 +525,8 @@ void FCinematicShotTrackEditor::SwitchTake(uint32 TakeNumber)
 		if (TakeObject && TakeObject->IsA(UMovieSceneSequence::StaticClass()))
 		{
 			UMovieSceneSequence* MovieSceneSequence = CastChecked<UMovieSceneSequence>(TakeObject);
-
-			UMovieSceneCinematicShotTrack* CinematicShotTrack = FindOrCreateCinematicShotTrack();
+			
+			UMovieSceneCinematicShotTrack* CinematicShotTrack = CastChecked<UMovieSceneCinematicShotTrack>(Section->GetOuter());
 
 			TRange<FFrameNumber> NewShotRange         = Section->GetRange();
 			FFrameNumber         NewShotStartOffset   = Section->Parameters.StartFrameOffset;
diff --git a/Engine/Source/Editor/MovieSceneTools/Private/TrackEditors/SubTrackEditor.cpp b/Engine/Source/Editor/MovieSceneTools/Private/TrackEditors/SubTrackEditor.cpp
--- a/Engine/Source/Editor/MovieSceneTools/Private/TrackEditors/SubTrackEditor.cpp
+++ b/Engine/Source/Editor/MovieSceneTools/Private/TrackEditors/SubTrackEditor.cpp
@@ -833,8 +833,8 @@ void FSubTrackEditor::SwitchTake(uint32 TakeNumber)
 		{
 			UMovieSceneSequence* MovieSceneSequence = CastChecked<UMovieSceneSequence>(TakeObject);
 
-			UMovieSceneSubTrack* SubTrack = FindOrCreateMasterTrack<UMovieSceneSubTrack>().Track;
-		
+			UMovieSceneSubTrack* SubTrack = CastChecked<UMovieSceneSubTrack>(Section->GetOuter());
+
 			TRange<FFrameNumber> NewShotRange         = Section->GetRange();
 			FFrameNumber		 NewShotStartOffset   = Section->Parameters.StartFrameOffset;
 			float                NewShotTimeScale     = Section->Parameters.TimeScale;
