Sequencer EDL: EDL and rv session file fixes.

- Don't export inactive shots - move export to before InitializeShots (which can mutate active states)
- Fix handle frame tick resolution
- Fix handle frames for XML so that they're in int32
- Invert row index for XML import/export
- Store original start offset and handle frames in xml log for roundtripping
- Use automated level sequence capture ui name when exporting from shot track so that user config is loaded
- Add MovieExtension config to allow the extension of the shot names to be changed from the default .avi

diff --git a/Engine/Source/Editor/MovieSceneTools/Private/AutomatedLevelSequenceCapture.cpp b/Engine/Source/Editor/MovieSceneTools/Private/AutomatedLevelSequenceCapture.cpp
--- a/Engine/Source/Editor/MovieSceneTools/Private/AutomatedLevelSequenceCapture.cpp
+++ b/Engine/Source/Editor/MovieSceneTools/Private/AutomatedLevelSequenceCapture.cpp
@@ -202,6 +202,9 @@ void UAutomatedLevelSequenceCapture::Initialize(TSharedPtr<FSceneViewport> InVie
 		}
 	}
 
+	ExportEDL();
+	ExportFCPXML();
+
 	if (Actor)
 	{
 		// Ensure it doesn't loop (-1 is indefinite)
@@ -243,9 +246,6 @@ void UAutomatedLevelSequenceCapture::Initialize(TSharedPtr<FSceneViewport> InVie
 		UE_LOG(LogMovieSceneCapture, Error, TEXT("Could not find or create a Level Sequence Actor for this capture. Capturing will fail."));
 	}
 
-	ExportEDL();
-	ExportFCPXML();
-
 	CaptureState = ELevelSequenceCaptureState::Setup;
 	CaptureStrategy = MakeShareable(new FFixedTimeStepCaptureStrategy(Settings.FrameRate));
 	CaptureStrategy->OnInitialize();
@@ -903,8 +903,9 @@ void UAutomatedLevelSequenceCapture::ExportEDL()
 
 	FString SaveFilename = 	Settings.OutputDirectory.Path / MovieScene->GetOuter()->GetName();
 	int32 HandleFrames = Settings.HandleFrames;
+	FString MovieExtension = Settings.MovieExtension;
 
-	MovieSceneTranslatorEDL::ExportEDL(MovieScene, Settings.FrameRate, SaveFilename, HandleFrames);
+	MovieSceneTranslatorEDL::ExportEDL(MovieScene, Settings.FrameRate, SaveFilename, HandleFrames, MovieExtension);
 }
 
 double UAutomatedLevelSequenceCapture::GetEstimatedCaptureDurationSeconds() const
@@ -947,13 +948,14 @@ void UAutomatedLevelSequenceCapture::ExportFCPXML()
 	FFrameRate FrameRate = Settings.FrameRate;
 	uint32 ResX = Settings.Resolution.ResX;
 	uint32 ResY = Settings.Resolution.ResY;
+	FString MovieExtension = Settings.MovieExtension;
 
 	FFCPXMLExporter *Exporter = new FFCPXMLExporter;
 
 	TSharedRef<FMovieSceneTranslatorContext> ExportContext(new FMovieSceneTranslatorContext);
 	ExportContext->Init();
 
-	bool bSuccess = Exporter->Export(MovieScene, FilenameFormat, FrameRate, ResX, ResY, HandleFrames, SaveFilename, ExportContext);
+	bool bSuccess = Exporter->Export(MovieScene, FilenameFormat, FrameRate, ResX, ResY, HandleFrames, SaveFilename, ExportContext, MovieExtension);
 
 	// Log any messages in context
 	MovieSceneToolHelpers::MovieSceneTranslatorLogMessages(Exporter, ExportContext, false);
diff --git a/Engine/Source/Editor/MovieSceneTools/Private/FCPXML/FCPXMLExport.cpp b/Engine/Source/Editor/MovieSceneTools/Private/FCPXML/FCPXMLExport.cpp
--- a/Engine/Source/Editor/MovieSceneTools/Private/FCPXML/FCPXMLExport.cpp
+++ b/Engine/Source/Editor/MovieSceneTools/Private/FCPXML/FCPXMLExport.cpp
@@ -540,9 +540,10 @@ bool FFCPXMLExportVisitor::ConstructVideoNode(TSharedRef<FFCPXMLNode> InParentNo
 		return false;
 	}
  
-	for (TSharedPtr<FMovieSceneExportCinematicTrackData> CinematicTrack : CinematicMasterTrackData->CinematicTracks)
+	// Add in reverse order
+	for (int32 RowIndex = CinematicMasterTrackData->CinematicTracks.Num()-1; RowIndex >= 0; --RowIndex)
 	{
-		if (!ConstructVideoTrackNode(VideoNode, CinematicTrack, CinematicMasterTrackData)) 
+		if (!ConstructVideoTrackNode(VideoNode, CinematicMasterTrackData->CinematicTracks[RowIndex], CinematicMasterTrackData)) 
 		{ 
 			return false;
 		}
@@ -600,17 +601,18 @@ bool FFCPXMLExportVisitor::ConstructAudioNode(TSharedRef<FFCPXMLNode> InParentNo
 
 	uint32 TrackIndex = 1;
 
-	for (TSharedPtr<FMovieSceneExportAudioMasterTrackData> AudioMasterTrack : ExportData->MovieSceneData->AudioMasterTracks)
+	// Add in reverse order
+	for (int32 RowIndex = ExportData->MovieSceneData->AudioMasterTracks.Num() - 1; RowIndex >= 0; --RowIndex)
 	{
-		if (!AudioMasterTrack.IsValid())
+		if (!ExportData->MovieSceneData->AudioMasterTracks[RowIndex].IsValid())
 		{
 			return false;
 		}
 
-		for (TSharedPtr<FMovieSceneExportAudioTrackData> AudioTrack : AudioMasterTrack->AudioTracks)
+		for (TSharedPtr<FMovieSceneExportAudioTrackData> AudioTrack : ExportData->MovieSceneData->AudioMasterTracks[RowIndex]->AudioTracks)
 		{
 			uint32 OutNumTracks{ 0 };
-			if (!ConstructAudioTrackNode(AudioNode, AudioTrack, AudioMasterTrack, TrackIndex, OutNumTracks))
+			if (!ConstructAudioTrackNode(AudioNode, AudioTrack, ExportData->MovieSceneData->AudioMasterTracks[RowIndex], TrackIndex, OutNumTracks))
 			{
 				return false;
 			}
@@ -1142,12 +1144,12 @@ bool FFCPXMLExportVisitor::GetCinematicSectionFrames(const TSharedPtr<FMovieScen
 		return false;
 	}
 
-	int32 Handles = ExportData->GetHandleFrames().FloorToFrame().Value;
+	int32 HandleFrames = ExportData->GetHandleFrames();
 	OutStartFrame = InCinematicSectionData->StartFrame.Value;
 	OutEndFrame = InCinematicSectionData->EndFrame.Value;
 	OutDuration = OutEndFrame - OutStartFrame;
-	OutInFrame = Handles;
-	OutOutFrame = Handles + OutDuration;
+	OutInFrame = HandleFrames + 1;
+	OutOutFrame = HandleFrames + OutDuration;
 
 	return true;
 }
@@ -1281,7 +1283,16 @@ bool FFCPXMLExportVisitor::CreateCinematicSectionMetadata(const UMovieSceneCinem
 	{
 		return false;
 	}
+
 	OutMetadata = TEXT("[UE4ShotSection=") + InSection->GetPathName() + TEXT("]");
+
+	// Store the start offset and the handle frames for round-tripping to compute the new start offset
+	int32 HandleFrames = ExportData->GetHandleFrames();
+	FFrameRate TickResolution = InSection->GetTypedOuter<UMovieScene>()->GetTickResolution();
+	int32 StartFrameOffset = ConvertFrameTime(InSection->Parameters.StartFrameOffset, TickResolution, ExportData->GetFrameRate()).CeilToFrame().Value;
+
+	OutMetadata += TEXT("[UE4ShotStartOffset=") + FString::FromInt(StartFrameOffset) + TEXT("]");
+	OutMetadata += TEXT("[UE4ShotHandleFrames=") + FString::FromInt(HandleFrames) + TEXT("]");
 	return true;
 }
 
diff --git a/Engine/Source/Editor/MovieSceneTools/Private/FCPXML/FCPXMLImport.cpp b/Engine/Source/Editor/MovieSceneTools/Private/FCPXML/FCPXMLImport.cpp
--- a/Engine/Source/Editor/MovieSceneTools/Private/FCPXML/FCPXMLImport.cpp
+++ b/Engine/Source/Editor/MovieSceneTools/Private/FCPXML/FCPXMLImport.cpp
@@ -62,6 +62,8 @@ FFCPXMLImportVisitor::FFCPXMLImportVisitor(TSharedRef<FMovieSceneImportData> InI
 , CurrAudioMasterTrack(nullptr)
 , CurrAudioTrackRowIndex(0)
 , bCurrImportAudioTrackIsStereoChannel(false)
+, MaxVideoTrackRowIndex(0)
+, MaxAudioTrackRowIndex(0)
 {
 	ConstructAudioTrackList();
 	if (AudioTrackList.Num() > 0)
@@ -149,6 +151,10 @@ bool FFCPXMLImportVisitor::VisitNode(TSharedRef<FFCPXMLTrackNode> InTrackNode)
 	{
 		bInVideoTrackNode = bPrevInTrackNode;
 		CurrVideoTrackRowIndex++;
+		if (MaxVideoTrackRowIndex < CurrVideoTrackRowIndex)
+		{
+			MaxVideoTrackRowIndex = CurrVideoTrackRowIndex;
+		}
 	}
 	else if (bInSequenceNode && bInAudioNode)
 	{
@@ -170,6 +176,10 @@ bool FFCPXMLImportVisitor::VisitNode(TSharedRef<FFCPXMLTrackNode> InTrackNode)
 				CurrAudioTrackRowIndex++;
 			}
 		}
+		if (MaxAudioTrackRowIndex < CurrAudioTrackRowIndex)
+		{
+			MaxAudioTrackRowIndex = CurrAudioTrackRowIndex;
+		}
 	}
 
 	return bSuccess;
@@ -261,6 +271,14 @@ bool FFCPXMLImportVisitor::VisitVideoClipItemNode(TSharedRef<FFCPXMLClipItemNode
 
 	FString SectionPathName = GetCinematicSectionPathName(LogNote, MasterClipId);
 
+	int32 HandleFrames = 0;
+	int32 OriginalStartOffset = 0;
+	TOptional<FFrameNumber> NewStartOffset;
+	if (GetCinematicSectionHandleFramesFromMetadata(LogNote, HandleFrames) && GetCinematicSectionStartOffsetFromMetadata(LogNote, OriginalStartOffset))
+	{
+		NewStartOffset = OriginalStartOffset - ((1 + HandleFrames) - StartOffset);
+	}
+
 	// Find actual section
 	TSharedPtr<FMovieSceneImportCinematicSectionData> SectionData = nullptr;
 	if (!SectionPathName.IsEmpty())
@@ -271,7 +289,7 @@ bool FFCPXMLImportVisitor::VisitVideoClipItemNode(TSharedRef<FFCPXMLClipItemNode
 	if (SectionData.IsValid())
 	{
 		// Update existing cinematic section
-		if (!ImportData->SetCinematicSection(SectionData, CurrVideoTrackRowIndex, FrameRate, Start, End, StartOffset))
+		if (!ImportData->SetCinematicSection(SectionData, CurrVideoTrackRowIndex, FrameRate, Start, End, NewStartOffset))
 		{
 			return false;
 		}
@@ -279,7 +297,7 @@ bool FFCPXMLImportVisitor::VisitVideoClipItemNode(TSharedRef<FFCPXMLClipItemNode
 	else
 	{
 		// Add new cinematic section
-		SectionData = ImportData->CreateCinematicSection(ClipItemName, CurrVideoTrackRowIndex, FrameRate, Start, End, StartOffset);
+		SectionData = ImportData->CreateCinematicSection(ClipItemName, CurrVideoTrackRowIndex, FrameRate, Start, End, NewStartOffset.IsSet() ? NewStartOffset.GetValue() : 0);
 		if (!SectionData.IsValid())
 		{
 			return false;
@@ -314,7 +332,7 @@ bool FFCPXMLImportVisitor::VisitAudioClipItemNode(TSharedRef<FFCPXMLClipItemNode
 	FFrameNumber StartOffset;
 	FFrameNumber Start;
 	FFrameNumber End;
-	
+
 	// Get relevant data from this clip item node
 	if (!GetClipItemNodeData(InClipItemNode, ClipItemName, ClipItemId, MasterClipId, LoggingInfoNode, LogNote, Filename, FrameRate, StartOffset, Start, End))
 	{
@@ -768,6 +786,30 @@ bool FFCPXMLImportVisitor::GetCinematicSectionPathNameFromMetadata(const FString
 	return ParseMetadata(InMetadata, TEXT("UE4ShotSection"), OutSectionObjectName);
 }
 
+/** Get sequencer shot handle frames from section metadata. Format is "[UE4ShotHandleFrames=handleframes]", whitespace ok. */
+bool FFCPXMLImportVisitor::GetCinematicSectionHandleFramesFromMetadata(const FString& InMetadata, int32& OutHandleFrames) const
+{
+	FString HandleFrameData;
+	bool bSuccess = ParseMetadata(InMetadata, TEXT("UE4ShotHandleFrames"), HandleFrameData);
+	if (bSuccess)
+	{
+		OutHandleFrames = FCString::Atoi(*HandleFrameData);
+	}
+	return bSuccess;
+}
+
+/** Get sequencer shot start offset frame from section metadata. Format is "[UE4ShotStartOffset=startoffset]", whitespace ok. */
+bool FFCPXMLImportVisitor::GetCinematicSectionStartOffsetFromMetadata(const FString& InMetadata, int32& OutStartOffset) const
+{
+	FString StartOffsetData;
+	bool bSuccess = ParseMetadata(InMetadata, TEXT("UE4ShotStartOffset"), StartOffsetData);
+	if (bSuccess)
+	{
+		OutStartOffset = FCString::Atoi(*StartOffsetData);
+	}
+	return bSuccess;
+}
+
 /** Get sequencer track name from track metadata. Format is "[UE4Track=trackobjectname][UE4Row=rowindex]", whitespace ok. */
 bool FFCPXMLImportVisitor::GetAudioFromMetadata(const FString& InMetadata, TSharedPtr<FFCPXMLImportAudioMetadata>& OutAudioMetadata) const
 {
diff --git a/Engine/Source/Editor/MovieSceneTools/Private/FCPXML/FCPXMLImport.h b/Engine/Source/Editor/MovieSceneTools/Private/FCPXML/FCPXMLImport.h
--- a/Engine/Source/Editor/MovieSceneTools/Private/FCPXML/FCPXMLImport.h
+++ b/Engine/Source/Editor/MovieSceneTools/Private/FCPXML/FCPXMLImport.h
@@ -110,9 +110,18 @@ private:
 	bool ParseMetadata(const FString& InMetadata, const FString& InKey, FString& OutValue) const;
 	/** Get sequencer section id from section metadata. Format is "[UE4ShotSection=sectionobjectname]", whitespace ok. */
 	bool GetCinematicSectionPathNameFromMetadata(const FString& InMetadata, FString& OutSectionObjectName) const;
+	/** Get sequencer shot handle frames from section metadata. Format is "[UE4ShotHandleFrames=handleframes]", whitespace ok. */
+	bool GetCinematicSectionHandleFramesFromMetadata(const FString& InMetadata, int32& OutHandleFrames) const;
+	/** Get sequencer shot start offset frame from section metadata. Format is "[UE4ShotStartOffset=startoffset]", whitespace ok. */
+	bool GetCinematicSectionStartOffsetFromMetadata(const FString& InMetadata, int32& OutStartOffset) const;
 	/** Get sequencer sound wave id and audio section ids from metadata. Format is "[UE4SoundWave=trackobjectname][UE4AudioSectionTopLevel=toplevelobjectname][UE4AudioSection=audiosectionobjectname]", whitespace ok. */
 	bool GetAudioFromMetadata(const FString& InMetadata, TSharedPtr<FFCPXMLImportAudioMetadata>& OutAudioMetadata) const;
 	
+public:
+
+	int32 GetMaxVideoTrackRowIndex() const { return MaxVideoTrackRowIndex; }
+	int32 GetMaxAudioTrackRowIndex() const { return MaxAudioTrackRowIndex; }
+
 private:
 	TSharedRef<FMovieSceneImportData> ImportData;
 	TSharedRef<FMovieSceneTranslatorContext> ImportContext;
@@ -141,5 +150,8 @@ private:
 	TSharedPtr<FMovieSceneImportAudioMasterTrackData> CurrAudioMasterTrack;
 	int32 CurrAudioTrackRowIndex;
 	bool bCurrImportAudioTrackIsStereoChannel;
+
+	int32 MaxVideoTrackRowIndex;
+	int32 MaxAudioTrackRowIndex;
 };
 
diff --git a/Engine/Source/Editor/MovieSceneTools/Private/FCPXML/FCPXMLMovieSceneTranslator.cpp b/Engine/Source/Editor/MovieSceneTools/Private/FCPXML/FCPXMLMovieSceneTranslator.cpp
--- a/Engine/Source/Editor/MovieSceneTools/Private/FCPXML/FCPXMLMovieSceneTranslator.cpp
+++ b/Engine/Source/Editor/MovieSceneTools/Private/FCPXML/FCPXMLMovieSceneTranslator.cpp
@@ -81,6 +81,38 @@ bool FFCPXMLImporter::Import(UMovieScene* InMovieScene, FFrameRate InFrameRate,
 		// Import the loaded Xml structure into the Sequencer movie scene
 		FFCPXMLImportVisitor ImportVisitor(ImportData, InContext);
 		bSuccess = FCPXMLFile->Accept(ImportVisitor);
+
+		// Flip rows
+		if (ImportData->GetCinematicMasterTrackData(false).IsValid())
+		{
+			int32 MaxVideoTrackRowIndex = ImportVisitor.GetMaxVideoTrackRowIndex();
+			for (TSharedPtr<FMovieSceneImportCinematicSectionData> SectionData : ImportData->GetCinematicMasterTrackData(false)->CinematicSections)
+			{
+				if (SectionData.IsValid())
+				{
+					UMovieSceneCinematicShotSection* CinematicSection = SectionData.Get()->CinematicSection;
+					if (CinematicSection)
+					{
+						CinematicSection->SetRowIndex(MaxVideoTrackRowIndex - CinematicSection->GetRowIndex());
+					}
+				}
+			}
+		}
+		if (ImportData->GetAudioMasterTrackData().IsValid())
+		{
+			int32 MaxAudioTrackRowIndex = ImportVisitor.GetMaxAudioTrackRowIndex();
+			for (TSharedPtr<FMovieSceneImportAudioSectionData> SectionData : ImportData->GetAudioMasterTrackData()->AudioSections)
+			{
+				if (SectionData.IsValid())
+				{
+					UMovieSceneAudioSection* AudioSection = SectionData.Get()->AudioSection;
+					if (AudioSection)
+					{
+						AudioSection->SetRowIndex(MaxAudioTrackRowIndex - AudioSection->GetRowIndex());
+					}
+				}
+			}
+		}
 	}
 
 	// add error message if one does not exist in the context
@@ -144,7 +176,7 @@ FText FFCPXMLExporter::GetMessageLogLabel() const
 	return LOCTEXT("FCPXMLExportLogLabel", "FCP 7 XML Export Log");
 }
 
-bool FFCPXMLExporter::Export(const UMovieScene* InMovieScene, FString InFilenameFormat, FFrameRate InFrameRate, uint32 InResX, uint32 InResY, int32 InHandleFrames, FString InSaveFilename, TSharedRef<FMovieSceneTranslatorContext> InContext)
+bool FFCPXMLExporter::Export(const UMovieScene* InMovieScene, FString InFilenameFormat, FFrameRate InFrameRate, uint32 InResX, uint32 InResY, int32 InHandleFrames, FString InSaveFilename, TSharedRef<FMovieSceneTranslatorContext> InContext, FString InMovieExtension)
 {
 	// add warning message if filename format is not "{shot}"
 	FString AcceptedFormat = TEXT("{shot}");
@@ -160,7 +192,7 @@ bool FFCPXMLExporter::Export(const UMovieScene* InMovieScene, FString InFilename
 	TSharedRef<FFCPXMLFile> FCPXMLFile = MakeShared<FFCPXMLFile>();
 	FCPXMLFile->ConstructFile(FPaths::GetBaseFilename(InSaveFilename, true));
 
-	TSharedRef<FMovieSceneExportData> ExportData = MakeShared<FMovieSceneExportData>(InMovieScene, InFrameRate, InResX, InResY, InHandleFrames, InSaveFilename, InContext);
+	TSharedRef<FMovieSceneExportData> ExportData = MakeShared<FMovieSceneExportData>(InMovieScene, InFrameRate, InResX, InResY, InHandleFrames, InSaveFilename, InContext, InMovieExtension);
 
 	// Export sequencer movie scene, merging with existing Xml structure.
 	FFCPXMLExportVisitor ExportVisitor(InSaveFilename, ExportData, InContext);
diff --git a/Engine/Source/Editor/MovieSceneTools/Private/MovieSceneToolHelpers.cpp b/Engine/Source/Editor/MovieSceneTools/Private/MovieSceneToolHelpers.cpp
--- a/Engine/Source/Editor/MovieSceneTools/Private/MovieSceneToolHelpers.cpp
+++ b/Engine/Source/Editor/MovieSceneTools/Private/MovieSceneToolHelpers.cpp
@@ -643,7 +643,7 @@ bool MovieSceneToolHelpers::ShowImportEDLDialog(UMovieScene* InMovieScene, FFram
 	return MovieSceneTranslatorEDL::ImportEDL(InMovieScene, InFrameRate, OpenFilenames[0]);
 }
 
-bool MovieSceneToolHelpers::ShowExportEDLDialog(const UMovieScene* InMovieScene, FFrameRate InFrameRate, FString InSaveDirectory, int32 InHandleFrames)
+bool MovieSceneToolHelpers::ShowExportEDLDialog(const UMovieScene* InMovieScene, FFrameRate InFrameRate, FString InSaveDirectory, int32 InHandleFrames, FString InMovieExtension)
 {
 	TArray<FString> SaveFilenames;
 	FString SequenceName = InMovieScene->GetOuter()->GetName();
@@ -677,7 +677,7 @@ bool MovieSceneToolHelpers::ShowExportEDLDialog(const UMovieScene* InMovieScene,
 		return false;
 	}
 
-	if (MovieSceneTranslatorEDL::ExportEDL(InMovieScene, InFrameRate, SaveFilenames[0], InHandleFrames))
+	if (MovieSceneTranslatorEDL::ExportEDL(InMovieScene, InFrameRate, SaveFilenames[0], InHandleFrames, InMovieExtension))
 	{
 		const FString AbsoluteFilename = FPaths::ConvertRelativePathToFull(SaveFilenames[0]);
 		const FString SaveDirectory = FPaths::GetPath(AbsoluteFilename);
@@ -753,6 +753,7 @@ bool MovieSceneToolHelpers::MovieSceneTranslatorExport(FMovieSceneExporter* InEx
 	FFrameRate FrameRate = Settings.FrameRate;
 	uint32 ResX = Settings.Resolution.ResX;
 	uint32 ResY = Settings.Resolution.ResY;
+	FString MovieExtension = Settings.MovieExtension;
 
 	TArray<FString> SaveFilenames;
 	FString SequenceName = InMovieScene->GetOuter()->GetName();
@@ -785,7 +786,7 @@ bool MovieSceneToolHelpers::MovieSceneTranslatorExport(FMovieSceneExporter* InEx
 	TSharedRef<FMovieSceneTranslatorContext> ExportContext(new FMovieSceneTranslatorContext);
 	ExportContext->Init();
 
-	bool bSuccess = InExporter->Export(InMovieScene, FilenameFormat, FrameRate, ResX, ResY, HandleFrames, SaveFilenames[0], ExportContext);
+	bool bSuccess = InExporter->Export(InMovieScene, FilenameFormat, FrameRate, ResX, ResY, HandleFrames, SaveFilenames[0], ExportContext, MovieExtension);
 	
 	// Display any messages in context
 	MovieSceneTranslatorLogMessages(InExporter, ExportContext, true);
diff --git a/Engine/Source/Editor/MovieSceneTools/Private/MovieSceneTranslator.cpp b/Engine/Source/Editor/MovieSceneTools/Private/MovieSceneTranslator.cpp
--- a/Engine/Source/Editor/MovieSceneTools/Private/MovieSceneTranslator.cpp
+++ b/Engine/Source/Editor/MovieSceneTools/Private/MovieSceneTranslator.cpp
@@ -38,7 +38,7 @@ namespace
 	}
 }
 
-FMovieSceneExportData::FMovieSceneExportData(const UMovieScene* InMovieScene, FFrameRate InFrameRate, uint32 InResX, uint32 InResY, int32 InHandleFrames, FString InSaveFilename, TSharedPtr<FMovieSceneTranslatorContext> InContext)
+FMovieSceneExportData::FMovieSceneExportData(const UMovieScene* InMovieScene, FFrameRate InFrameRate, uint32 InResX, uint32 InResY, int32 InHandleFrames, FString InSaveFilename, TSharedPtr<FMovieSceneTranslatorContext> InContext, FString InMovieExtension)
 {
 	if (InMovieScene == nullptr)
 	{
@@ -50,8 +50,9 @@ FMovieSceneExportData::FMovieSceneExportData(const UMovieScene* InMovieScene, FF
 	FrameRate = InFrameRate;
 	ResX = InResX;
 	ResY = InResY;
-	HandleFrames = ConvertFrameTime(FFrameTime(InHandleFrames), InMovieScene->GetTickResolution(), FrameRate);
+	HandleFrames = InHandleFrames;
 	SaveFilename = InSaveFilename;
+	MovieExtension = InMovieExtension;
 
 	// preferred sample rate in UE4
 	DefaultAudioSampleRate = 44100;
@@ -344,16 +345,7 @@ bool FMovieSceneExportData::ConstructCinematicSectionData(const UMovieScene* InM
 
 	SectionData->DisplayName = InCinematicSection->GetShotDisplayName();
 
-#if PLATFORM_MAC
-	static const TCHAR* Extension = TEXT(".mov");
-#elif PLATFORM_UNIX
-	static const TCHAR* Extension = TEXT(".unsupp");
-	return false;
-#else
-	static const TCHAR* Extension = TEXT(".avi");
-#endif
-
-	SectionData->SourceFilename = SectionData->DisplayName + Extension;
+	SectionData->SourceFilename = SectionData->DisplayName + MovieExtension;
 	SectionData->SourceFilePath = TEXT("");
 
 	ConstructSectionData(InMovieScene, SectionData, InCinematicSection, EMovieSceneTranslatorSectionType::Cinematic, SectionData->DisplayName);
@@ -524,6 +516,11 @@ FString FMovieSceneExportData::GetFilenamePath() const
 	return SaveFilenamePath;
 }
 
+FString FMovieSceneExportData::GetMovieExtension() const
+{
+	return MovieExtension;
+}
+
 FFrameRate FMovieSceneExportData::GetFrameRate() const
 {
 	return FrameRate;
@@ -555,7 +552,7 @@ bool FMovieSceneExportData::GetFrameRateIsNTSC() const
 	return (!FMath::IsNearlyZero(FractionalPart));
 }
 
-FFrameTime FMovieSceneExportData::GetHandleFrames() const
+int32 FMovieSceneExportData::GetHandleFrames() const
 {
 	return HandleFrames;
 }
@@ -692,7 +689,7 @@ TSharedPtr<FMovieSceneImportCinematicSectionData> FMovieSceneImportData::CreateC
 	return SectionData;
 }
 
-bool FMovieSceneImportData::SetCinematicSection(TSharedPtr<FMovieSceneImportCinematicSectionData> InSection, int32 InRow, FFrameRate InFrameRate, FFrameNumber InStartFrame, FFrameNumber InEndFrame, FFrameNumber InStartOffsetFrame)
+bool FMovieSceneImportData::SetCinematicSection(TSharedPtr<FMovieSceneImportCinematicSectionData> InSection, int32 InRow, FFrameRate InFrameRate, FFrameNumber InStartFrame, FFrameNumber InEndFrame, TOptional<FFrameNumber> InStartOffsetFrame)
 {
 	if (!InSection.IsValid() || InSection->CinematicSection == nullptr)
 	{
@@ -701,11 +698,14 @@ bool FMovieSceneImportData::SetCinematicSection(TSharedPtr<FMovieSceneImportCine
 
 	FFrameRate TickResolution = MovieSceneData->MovieScene->GetTickResolution();
 	FFrameNumber StartFrame = ConvertFrameTime(InStartFrame, InFrameRate, TickResolution).GetFrame();
-	FFrameNumber StartOffsetFrame = ConvertFrameTime(InStartOffsetFrame, InFrameRate, TickResolution).GetFrame();
 	FFrameNumber EndFrame = ConvertFrameTime(InEndFrame, InFrameRate, TickResolution).GetFrame();
-
+		
 	InSection->CinematicSection->Modify();
-	InSection->CinematicSection->Parameters.StartFrameOffset = StartOffsetFrame;
+	if (InStartOffsetFrame.IsSet())
+	{
+		FFrameNumber StartOffsetFrame = ConvertFrameTime(InStartOffsetFrame.GetValue(), InFrameRate, TickResolution).GetFrame();
+		InSection->CinematicSection->Parameters.StartFrameOffset = StartOffsetFrame;
+	}
 	InSection->CinematicSection->SetRange(TRange<FFrameNumber>(StartFrame, EndFrame));
 	if (InRow != InSection->CinematicSection->GetRowIndex())
 	{
diff --git a/Engine/Source/Editor/MovieSceneTools/Private/MovieSceneTranslatorEDL.cpp b/Engine/Source/Editor/MovieSceneTools/Private/MovieSceneTranslatorEDL.cpp
--- a/Engine/Source/Editor/MovieSceneTools/Private/MovieSceneTranslatorEDL.cpp
+++ b/Engine/Source/Editor/MovieSceneTools/Private/MovieSceneTranslatorEDL.cpp
@@ -432,7 +432,7 @@ bool MovieSceneTranslatorEDL::ImportEDL(UMovieScene* InMovieScene, FFrameRate In
 	return true;
 }
 
-bool MovieSceneTranslatorEDL::ExportEDL(const UMovieScene* InMovieScene, FFrameRate InFrameRate, FString InSaveFilename, const int32 InHandleFrames)
+bool MovieSceneTranslatorEDL::ExportEDL(const UMovieScene* InMovieScene, FFrameRate InFrameRate, FString InSaveFilename, const int32 InHandleFrames, FString InMovieExtension)
 {
 	FString SequenceName = InMovieScene->GetOuter()->GetName();
 	FString SaveBasename = FPaths::GetPath(InSaveFilename) / FPaths::GetBaseFilename(InSaveFilename);
@@ -459,31 +459,25 @@ bool MovieSceneTranslatorEDL::ExportEDL(const UMovieScene* InMovieScene, FFrameR
 			{
 				UMovieSceneCinematicShotSection* CinematicShotSection = Cast<UMovieSceneCinematicShotSection>(ShotSection);
 
-				if ( CinematicShotSection->GetSequence() == nullptr || !CinematicShotSection->HasStartFrame() || !CinematicShotSection->HasEndFrame() )
+				if ( CinematicShotSection->GetSequence() == nullptr || !CinematicShotSection->HasStartFrame() || !CinematicShotSection->HasEndFrame() || !CinematicShotSection->IsActive())
 				{
 					// If the shot doesn't have a valid sequence skip it.  This is currently the case for filler sections.
 					// TODO: Handle this properly in the edl output.
 					continue;
 				}
 
+				FFrameRate TickResolution = CinematicShotSection->GetSequence()->GetMovieScene()->GetTickResolution();
+
 				FString ShotName = CinematicShotSection->GetShotDisplayName();
 				FString ShotPath = CinematicShotSection->GetSequence()->GetMovieScene()->GetOuter()->GetPathName();
 
-				FFrameNumber SourceInFrame  = InHandleFrames;
-				FFrameNumber SourceOutFrame = InHandleFrames + MovieScene::DiscreteSize(CinematicShotSection->GetRange());
+				FFrameNumber SourceInFrame = ConvertFrameTime(InHandleFrames + 1, InFrameRate, TickResolution).FrameNumber;
+				FFrameNumber SourceOutFrame = ConvertFrameTime(InHandleFrames, InFrameRate, TickResolution).FrameNumber + MovieScene::DiscreteSize(CinematicShotSection->GetRange());
 
 				FFrameNumber EditInFrame    = CinematicShotSection->GetInclusiveStartFrame();
 				FFrameNumber EditOutFrame   = CinematicShotSection->GetExclusiveEndFrame();
 
-#if PLATFORM_MAC
-				static const TCHAR* Extension = TEXT(".mov");
-#elif PLATFORM_UNIX
-				static const TCHAR* Extension = TEXT(".unsupp");
-				continue;
-#else
-				static const TCHAR* Extension = TEXT(".avi");
-#endif
-				ShotName += Extension;
+				ShotName += InMovieExtension;
 
 				//@todo shotpath should really be moviefile path
 				ShotPath = ShotName;
diff --git a/Engine/Source/Editor/MovieSceneTools/Private/TrackEditors/CinematicShotTrackEditor.cpp b/Engine/Source/Editor/MovieSceneTools/Private/TrackEditors/CinematicShotTrackEditor.cpp
--- a/Engine/Source/Editor/MovieSceneTools/Private/TrackEditors/CinematicShotTrackEditor.cpp
+++ b/Engine/Source/Editor/MovieSceneTools/Private/TrackEditors/CinematicShotTrackEditor.cpp
@@ -838,7 +838,7 @@ UAutomatedLevelSequenceCapture* GetMovieSceneCapture()
 	
 	if (!MovieSceneCapture)
 	{
-		MovieSceneCapture = NewObject<UAutomatedLevelSequenceCapture>(GetTransientPackage(), UAutomatedLevelSequenceCapture::StaticClass(), UAutomatedLevelSequenceCapture::AutomatedLevelSequenceCaptureUIName, RF_Transient);
+		MovieSceneCapture = NewObject<UAutomatedLevelSequenceCapture>(GetTransientPackage(), UAutomatedLevelSequenceCapture::StaticClass(), UMovieSceneCapture::MovieSceneCaptureUIName, RF_Transient);
 		MovieSceneCapture->LoadFromConfig();
 	}
 
@@ -898,8 +898,9 @@ void FCinematicShotTrackEditor::ExportEDL()
 	const FMovieSceneCaptureSettings& Settings = MovieSceneCapture->GetSettings();
 	FString SaveDirectory = FPaths::ConvertRelativePathToFull(Settings.OutputDirectory.Path);
 	int32 HandleFrames = Settings.HandleFrames;
+	FString MovieExtension = Settings.MovieExtension;
 
-	MovieSceneToolHelpers::ShowExportEDLDialog(MovieScene, MovieScene->GetDisplayRate(), SaveDirectory, HandleFrames);
+	MovieSceneToolHelpers::ShowExportEDLDialog(MovieScene, MovieScene->GetDisplayRate(), SaveDirectory, HandleFrames, MovieExtension);
 }
 
 
diff --git a/Engine/Source/Editor/MovieSceneTools/Public/FCPXML/FCPXMLMovieSceneTranslator.h b/Engine/Source/Editor/MovieSceneTools/Public/FCPXML/FCPXMLMovieSceneTranslator.h
--- a/Engine/Source/Editor/MovieSceneTools/Public/FCPXML/FCPXMLMovieSceneTranslator.h
+++ b/Engine/Source/Editor/MovieSceneTools/Public/FCPXML/FCPXMLMovieSceneTranslator.h
@@ -82,7 +82,8 @@ public:
 	 * @param InHandleFrames The number of handle frames to include for each shot.
 	 * @param InSaveFilename The file path to save to.
 	 * @param OutError The return error message
+	 * @param MovieExtension The movie extension for the shot filenames (ie. .avi, .mov, .mp4)
 	 * @return Whether the export was successful
 	 */
-	virtual bool Export(const UMovieScene* InMovieScene, FString InFilenameFormat, FFrameRate InFrameRate, uint32 InResX, uint32 InResY, int32 InHandleFrames, FString InSaveFilename, TSharedRef<FMovieSceneTranslatorContext> InContext);
+	virtual bool Export(const UMovieScene* InMovieScene, FString InFilenameFormat, FFrameRate InFrameRate, uint32 InResX, uint32 InResY, int32 InHandleFrames, FString InSaveFilename, TSharedRef<FMovieSceneTranslatorContext> InContext, FString InMovieExtension);
 };
diff --git a/Engine/Source/Editor/MovieSceneTools/Public/MovieSceneToolHelpers.h b/Engine/Source/Editor/MovieSceneTools/Public/MovieSceneToolHelpers.h
--- a/Engine/Source/Editor/MovieSceneTools/Public/MovieSceneToolHelpers.h
+++ b/Engine/Source/Editor/MovieSceneTools/Public/MovieSceneToolHelpers.h
@@ -143,9 +143,10 @@ public:
 	 * @param InFrameRate The frame rate to export the EDL at
 	 * @param InSaveDirectory Optional directory path to save to. If none given, a dialog will pop up to prompt the user
 	 * @param InHandleFrames The number of handle frames to include for each shot.
+	 * @param MovieExtension The movie extension for the shot filenames (ie. .avi, .mov, .mp4)
 	 * @return Whether the export was successful
 	 */
-	static bool ShowExportEDLDialog(const UMovieScene* InMovieScene, FFrameRate InFrameRate, FString InSaveDirectory = TEXT(""), int32 InHandleFrames = 8);
+	static bool ShowExportEDLDialog(const UMovieScene* InMovieScene, FFrameRate InFrameRate, FString InSaveDirectory = TEXT(""), int32 InHandleFrames = 8, FString InMovieExtension = TEXT(".avi"));
 
 	/**
 	* Import movie scene formats
diff --git a/Engine/Source/Editor/MovieSceneTools/Public/MovieSceneTranslator.h b/Engine/Source/Editor/MovieSceneTools/Public/MovieSceneTranslator.h
--- a/Engine/Source/Editor/MovieSceneTools/Public/MovieSceneTranslator.h
+++ b/Engine/Source/Editor/MovieSceneTools/Public/MovieSceneTranslator.h
@@ -124,6 +124,7 @@ struct FMovieSceneExportMovieSceneData
 	FFrameNumber PlaybackRangeEndFrame;
 	TSharedPtr<FMovieSceneExportCinematicMasterTrackData> CinematicMasterTrack;
 	TArray< TSharedPtr<FMovieSceneExportAudioMasterTrackData> > AudioMasterTracks;
+	FString MovieExtension;
 };
 
 enum class EMovieSceneTranslatorSectionType : int32
@@ -139,7 +140,7 @@ class MOVIESCENETOOLS_API FMovieSceneExportData : public TSharedFromThis<FMovieS
 {
 public:
 	/** Constructor */
-	FMovieSceneExportData(const UMovieScene* InMovieScene, FFrameRate InFrameRate, uint32 InResX, uint32 InResY, int32 InHandleFrames, FString InSaveFilename, TSharedPtr<FMovieSceneTranslatorContext> InContext);
+	FMovieSceneExportData(const UMovieScene* InMovieScene, FFrameRate InFrameRate, uint32 InResX, uint32 InResY, int32 InHandleFrames, FString InSaveFilename, TSharedPtr<FMovieSceneTranslatorContext> InContext, FString InMovieExtension);
 	/** Default constructor, necessary for shared ref - should not be used */
 	FMovieSceneExportData();
 	/** Destructor */
@@ -149,6 +150,8 @@ public:
 	FString GetFilename() const;
 	/** Gets export filename with full path */
 	FString GetFilenamePath() const;
+	/** Gets the shot movie extension */
+	FString GetMovieExtension() const;
 
 	/** Gets export frame rate */
 	FFrameRate GetFrameRate() const;
@@ -161,8 +164,8 @@ public:
 	/** Returns the nearest integral frame rate */
 	uint32 GetNearestWholeFrameRate() const;
 
-	/** Gets the frame handle which has been converted to the export frame rate */
-	FFrameTime GetHandleFrames() const;
+	/** Gets the frame handle */
+	int32 GetHandleFrames() const;
 	/** Gets default audio sample rate */
 	int32 GetDefaultAudioSampleRate() const;
 	/** Gets default audio depth */
@@ -208,12 +211,13 @@ private:
 	FFrameRate FrameRate;
 	uint32 ResX;
 	uint32 ResY;
-	FFrameTime HandleFrames;
+	int32 HandleFrames;
 	FString SaveFilename;
 	FString SaveFilenamePath;
 	bool bExportDataIsValid;
 	int32 DefaultAudioSampleRate;
 	int32 DefaultAudioDepth;
+	FString MovieExtension;
 };
 
 struct FMovieSceneImportCinematicSectionData 
@@ -296,7 +300,7 @@ public:
 	/** Create cinematic section */
 	TSharedPtr<FMovieSceneImportCinematicSectionData> CreateCinematicSection(FString InName, int32 InRow, FFrameRate InFrameRate, FFrameNumber InStartFrame, FFrameNumber InEndFrame, FFrameNumber InStartOffsetFrame);
 	/** Set cinematic section */
-	bool SetCinematicSection(TSharedPtr<FMovieSceneImportCinematicSectionData> InSection, int32 InRow, FFrameRate InFrameRate, FFrameNumber InStartFrame, FFrameNumber InEndFrame, FFrameNumber InStartOffsetFrame);
+	bool SetCinematicSection(TSharedPtr<FMovieSceneImportCinematicSectionData> InSection, int32 InRow, FFrameRate InFrameRate, FFrameNumber InStartFrame, FFrameNumber InEndFrame, TOptional<FFrameNumber> InStartOffsetFrame);
 
 	/** Returns the audio master track data pointer or nullptr if one does not exist */
 	TSharedPtr<FMovieSceneImportAudioMasterTrackData> GetAudioMasterTrackData();
@@ -408,8 +412,9 @@ public:
 	* @param InHandleFrames The number of handle frames to include for each shot.
 	* @param InSaveFilename The file path to save to.
 	* @param OutError The return error message
+	* @param MovieExtension The movie extension for the shot filenames (ie. .avi, .mov, .mp4)
 	* @return Whether the export was successful
 	*/
-	virtual bool Export(const UMovieScene* InMovieScene, FString InFilenameFormat, FFrameRate InFrameRate, uint32 InResX, uint32 InResY, int32 InHandleFrames, FString InSaveFilename, TSharedRef<FMovieSceneTranslatorContext> InContext) = 0;
+	virtual bool Export(const UMovieScene* InMovieScene, FString InFilenameFormat, FFrameRate InFrameRate, uint32 InResX, uint32 InResY, int32 InHandleFrames, FString InSaveFilename, TSharedRef<FMovieSceneTranslatorContext> InContext, FString InMovieExtension) = 0;
 };
 
diff --git a/Engine/Source/Editor/MovieSceneTools/Public/MovieSceneTranslatorEDL.h b/Engine/Source/Editor/MovieSceneTools/Public/MovieSceneTranslatorEDL.h
--- a/Engine/Source/Editor/MovieSceneTools/Public/MovieSceneTranslatorEDL.h
+++ b/Engine/Source/Editor/MovieSceneTools/Public/MovieSceneTranslatorEDL.h
@@ -29,7 +29,8 @@ public:
 	 * @param InFrameRate The frame rate to export the EDL at
 	 * @param InSaveFilename The file path to save to.
 	 * @param InHandleFrames The number of handle frames to include for each shot.
+	 * @param MovieExtension The movie extension for the shot filenames (ie. .avi, .mov, .mp4)
 	 * @return Whether the export was successful
 	 */
-	static bool ExportEDL(const UMovieScene* InMovieScene, FFrameRate InFrameRate, FString InSaveFilename, int32 InHandleFrames);
+	static bool ExportEDL(const UMovieScene* InMovieScene, FFrameRate InFrameRate, FString InSaveFilename, int32 InHandleFrames, FString InMovieExtension);
 };
diff --git a/Engine/Source/Runtime/MovieSceneCapture/Private/MovieSceneCapture.cpp b/Engine/Source/Runtime/MovieSceneCapture/Private/MovieSceneCapture.cpp
--- a/Engine/Source/Runtime/MovieSceneCapture/Private/MovieSceneCapture.cpp
+++ b/Engine/Source/Runtime/MovieSceneCapture/Private/MovieSceneCapture.cpp
@@ -95,6 +95,14 @@ FMovieSceneCaptureSettings::FMovieSceneCaptureSettings()
 	bAllowTurning = false;
 	bShowPlayer = false;
 	bShowHUD = false;
+
+#if PLATFORM_MAC
+	MovieExtension = TEXT(".mov");
+#elif PLATFORM_UNIX
+	MovieExtension = TEXT(".unsupp");
+#else
+	MovieExtension = TEXT(".avi");
+#endif
 }
 
 UMovieSceneCapture::UMovieSceneCapture(const FObjectInitializer& Initializer)
diff --git a/Engine/Source/Runtime/MovieSceneCapture/Public/MovieSceneCaptureSettings.h b/Engine/Source/Runtime/MovieSceneCapture/Public/MovieSceneCaptureSettings.h
--- a/Engine/Source/Runtime/MovieSceneCapture/Public/MovieSceneCaptureSettings.h
+++ b/Engine/Source/Runtime/MovieSceneCapture/Public/MovieSceneCaptureSettings.h
@@ -67,6 +67,10 @@ struct MOVIESCENECAPTURE_API FMovieSceneCaptureSettings
 	UPROPERTY(config, EditAnywhere, BlueprintReadWrite, Category=Sequence, AdvancedDisplay, meta=(ClampMin=0, UIMin=0))
 	int32 HandleFrames;
 
+	/** Filename extension for movies referenced in the XMLs/EDLs */
+	UPROPERTY(config, EditAnywhere, BlueprintReadWrite, Category=Sequence, AdvancedDisplay)
+	FString MovieExtension;
+
 	/** How much to zero-pad frame numbers on filenames */
 	UPROPERTY(config, EditAnywhere, BlueprintReadWrite, Category=General, AdvancedDisplay, meta=(ClampMin=0, UIMin=0))
 	uint8 ZeroPadFrameNumbers;
