Sequencer: Add toast when the tick resolution of the subsequence doesn't match the parent

diff --git a/Engine/Source/Editor/MovieSceneTools/Private/TrackEditors/CinematicShotTrackEditor.cpp b/Engine/Source/Editor/MovieSceneTools/Private/TrackEditors/CinematicShotTrackEditor.cpp
--- a/Engine/Source/Editor/MovieSceneTools/Private/TrackEditors/CinematicShotTrackEditor.cpp
+++ b/Engine/Source/Editor/MovieSceneTools/Private/TrackEditors/CinematicShotTrackEditor.cpp
@@ -658,9 +658,10 @@ FKeyPropertyResult FCinematicShotTrackEditor::AddKeyInternal(FFrameNumber KeyTim
 	{
 		UMovieSceneCinematicShotTrack* CinematicShotTrack = FindOrCreateCinematicShotTrack();
 		
+		const FFrameRate TickResolution = InMovieSceneSequence->GetMovieScene()->GetTickResolution();
 		const FQualifiedFrameTime InnerDuration = FQualifiedFrameTime(
 			MovieScene::DiscreteSize(InMovieSceneSequence->GetMovieScene()->GetPlaybackRange()),
-			InMovieSceneSequence->GetMovieScene()->GetTickResolution());
+			TickResolution);
 
 		const FFrameRate OuterFrameRate = CinematicShotTrack->GetTypedOuter<UMovieScene>()->GetTickResolution();
 		const int32      OuterDuration  = InnerDuration.ConvertTo(OuterFrameRate).FrameNumber.Value;
@@ -672,6 +673,13 @@ FKeyPropertyResult FCinematicShotTrackEditor::AddKeyInternal(FFrameNumber KeyTim
 		GetSequencer()->SelectSection(NewSection);
 		GetSequencer()->ThrobSectionSelection();
 
+		if (TickResolution != OuterFrameRate)
+		{
+			FNotificationInfo Info(FText::Format(LOCTEXT("TickResolutionMismatch", "The parent sequence has a different tick resolution {0} than the newly added sequence {1}"), OuterFrameRate.ToPrettyText(), TickResolution.ToPrettyText()));
+			Info.bUseLargeFont = false;
+			FSlateNotificationManager::Get().AddNotification(Info);
+		}
+
 		return KeyPropertyResult;
 	}
 
@@ -811,9 +819,10 @@ FKeyPropertyResult FCinematicShotTrackEditor::HandleSequenceAdded(FFrameNumber K
 
 	auto CinematicShotTrack = FindOrCreateCinematicShotTrack();
 
+	const FFrameRate TickResolution = Sequence->GetMovieScene()->GetTickResolution();
 	const FQualifiedFrameTime InnerDuration = FQualifiedFrameTime(
 		MovieScene::DiscreteSize(Sequence->GetMovieScene()->GetPlaybackRange()),
-		Sequence->GetMovieScene()->GetTickResolution());
+		TickResolution);
 
 	const FFrameRate OuterFrameRate = CinematicShotTrack->GetTypedOuter<UMovieScene>()->GetTickResolution();
 	const int32      OuterDuration  = InnerDuration.ConvertTo(OuterFrameRate).FrameNumber.Value;
@@ -825,6 +834,13 @@ FKeyPropertyResult FCinematicShotTrackEditor::HandleSequenceAdded(FFrameNumber K
 	GetSequencer()->SelectSection(NewSection);
 	GetSequencer()->ThrobSectionSelection();
 
+	if (TickResolution != OuterFrameRate)
+	{
+		FNotificationInfo Info(FText::Format(LOCTEXT("TickResolutionMismatch", "The parent sequence has a different tick resolution {0} than the newly added sequence {1}"), OuterFrameRate.ToPrettyText(), TickResolution.ToPrettyText()));
+		Info.bUseLargeFont = false;
+		FSlateNotificationManager::Get().AddNotification(Info);
+	}
+
 	return KeyPropertyResult;
 }
 
diff --git a/Engine/Source/Editor/MovieSceneTools/Private/TrackEditors/SubTrackEditor.cpp b/Engine/Source/Editor/MovieSceneTools/Private/TrackEditors/SubTrackEditor.cpp
--- a/Engine/Source/Editor/MovieSceneTools/Private/TrackEditors/SubTrackEditor.cpp
+++ b/Engine/Source/Editor/MovieSceneTools/Private/TrackEditors/SubTrackEditor.cpp
@@ -730,9 +730,10 @@ FKeyPropertyResult FSubTrackEditor::AddKeyInternal(FFrameNumber KeyTime, UMovieS
 	{
 		UMovieSceneSubTrack* SubTrack = Cast<UMovieSceneSubTrack>(InTrack);
 
+		const FFrameRate TickResolution = InMovieSceneSequence->GetMovieScene()->GetTickResolution();
 		const FQualifiedFrameTime InnerDuration = FQualifiedFrameTime(
 			MovieScene::DiscreteSize(InMovieSceneSequence->GetMovieScene()->GetPlaybackRange()),
-			InMovieSceneSequence->GetMovieScene()->GetTickResolution());
+			TickResolution);
 
 		const FFrameRate OuterFrameRate = SubTrack->GetTypedOuter<UMovieScene>()->GetTickResolution();
 		const int32      OuterDuration  = InnerDuration.ConvertTo(OuterFrameRate).FrameNumber.Value;
@@ -744,6 +745,13 @@ FKeyPropertyResult FSubTrackEditor::AddKeyInternal(FFrameNumber KeyTime, UMovieS
 		GetSequencer()->SelectSection(NewSection);
 		GetSequencer()->ThrobSectionSelection();
 
+		if (TickResolution != OuterFrameRate)
+		{
+			FNotificationInfo Info(FText::Format(LOCTEXT("TickResolutionMismatch", "The parent sequence has a different tick resolution {0} than the newly added sequence {1}"), OuterFrameRate.ToPrettyText(), TickResolution.ToPrettyText()));
+			Info.bUseLargeFont = false;
+			FSlateNotificationManager::Get().AddNotification(Info);
+		}
+
 		return KeyPropertyResult;
 	}
 
@@ -760,9 +768,10 @@ FKeyPropertyResult FSubTrackEditor::HandleSequenceAdded(FFrameNumber KeyTime, UM
 
 	auto SubTrack = FindOrCreateMasterTrack<UMovieSceneSubTrack>().Track;
 
+	const FFrameRate TickResolution = Sequence->GetMovieScene()->GetTickResolution();
 	const FQualifiedFrameTime InnerDuration = FQualifiedFrameTime(
 		MovieScene::DiscreteSize(Sequence->GetMovieScene()->GetPlaybackRange()),
-		Sequence->GetMovieScene()->GetTickResolution());
+		TickResolution);
 
 	const FFrameRate OuterFrameRate = SubTrack->GetTypedOuter<UMovieScene>()->GetTickResolution();
 	const int32      OuterDuration  = InnerDuration.ConvertTo(OuterFrameRate).FrameNumber.Value;
@@ -774,6 +783,13 @@ FKeyPropertyResult FSubTrackEditor::HandleSequenceAdded(FFrameNumber KeyTime, UM
 	GetSequencer()->SelectSection(NewSection);
 	GetSequencer()->ThrobSectionSelection();
 
+	if (TickResolution != OuterFrameRate)
+	{
+		FNotificationInfo Info(FText::Format(LOCTEXT("TickResolutionMismatch", "The parent sequence has a different tick resolution {0} than the newly added sequence {1}"), OuterFrameRate.ToPrettyText(), TickResolution.ToPrettyText()));
+		Info.bUseLargeFont = false;
+		FSlateNotificationManager::Get().AddNotification(Info);
+	}
+
 	return KeyPropertyResult;
 }
 
