PR #5506: Fix Crash when Cast to UStructProperty failed (Contributed by [REDACTED_PERSON])

diff --git a/Engine/Source/Editor/MovieSceneTools/Private/TrackEditors/PropertyTrackEditors/ColorPropertyTrackEditor.cpp b/Engine/Source/Editor/MovieSceneTools/Private/TrackEditors/PropertyTrackEditors/ColorPropertyTrackEditor.cpp
--- a/Engine/Source/Editor/MovieSceneTools/Private/TrackEditors/PropertyTrackEditors/ColorPropertyTrackEditor.cpp
+++ b/Engine/Source/Editor/MovieSceneTools/Private/TrackEditors/PropertyTrackEditors/ColorPropertyTrackEditor.cpp
@@ -38,6 +38,11 @@ void FColorPropertyTrackEditor::GenerateKeysFromPropertyChanged( const FProperty
 	}
 
 	const UStructProperty* StructProp = Cast<const UStructProperty>( Property );
+	if (!StructProp)
+	{
+		return;
+	}
+
 	FName StructName = StructProp->Struct->GetFName();
 	FName PropertyName = Property->GetFName();
 
diff --git a/Engine/Source/Editor/MovieSceneTools/Private/TrackEditors/PropertyTrackEditors/VectorPropertyTrackEditor.cpp b/Engine/Source/Editor/MovieSceneTools/Private/TrackEditors/PropertyTrackEditors/VectorPropertyTrackEditor.cpp
--- a/Engine/Source/Editor/MovieSceneTools/Private/TrackEditors/PropertyTrackEditors/VectorPropertyTrackEditor.cpp
+++ b/Engine/Source/Editor/MovieSceneTools/Private/TrackEditors/PropertyTrackEditors/VectorPropertyTrackEditor.cpp
@@ -24,7 +24,10 @@ TSharedRef<ISequencerTrackEditor> FVectorPropertyTrackEditor::CreateTrackEditor(
 void FVectorPropertyTrackEditor::GenerateKeysFromPropertyChanged( const FPropertyChangedParams& PropertyChangedParams, FGeneratedTrackKeys& OutGeneratedKeys )
 {
 	const UStructProperty* StructProp = Cast<const UStructProperty>( PropertyChangedParams.PropertyPath.GetLeafMostProperty().Property.Get() );
-
+	if (!StructProp)
+	{
+		return;
+	}
 	FName StructName = StructProp->Struct->GetFName();
 
 	bool bIsVector2D = StructName == NAME_Vector2D;
