Sequencer: Skip past deprecated functions for events

diff --git a/Engine/Source/Editor/MovieSceneTools/Private/MovieSceneEventCustomization.cpp b/Engine/Source/Editor/MovieSceneTools/Private/MovieSceneEventCustomization.cpp
--- a/Engine/Source/Editor/MovieSceneTools/Private/MovieSceneEventCustomization.cpp
+++ b/Engine/Source/Editor/MovieSceneTools/Private/MovieSceneEventCustomization.cpp
@@ -332,6 +332,7 @@ TSharedRef<SWidget> FMovieSceneEventCustomization::GetMenuContent()
 void FMovieSceneEventCustomization::PopulateQuickBindSubMenu(FMenuBuilder& MenuBuilder, UClass* TemplateClass)
 {
 	FSlateIcon Icon(FEditorStyle::GetStyleSetName(), "GraphEditor.Function_16x");
+	static const FName DeprecatedFunctionName(TEXT("DeprecatedFunction"));
 
 	TArray<UFunction*> Functions;
 
@@ -343,7 +344,7 @@ void FMovieSceneEventCustomization::PopulateQuickBindSubMenu(FMenuBuilder& MenuB
 		Functions.Reset();
 		for (UFunction* Function : TFieldRange<UFunction>(SuperClass, EFieldIteratorFlags::ExcludeSuper, EFieldIteratorFlags::ExcludeDeprecated))
 		{
-			if (Function->HasAllFunctionFlags(FUNC_BlueprintCallable|FUNC_Public))
+			if (Function->HasAllFunctionFlags(FUNC_BlueprintCallable|FUNC_Public) && !Function->HasMetaData(DeprecatedFunctionName))
 			{
 				Functions.Add(Function);
 			}
