Sequencer: Sort sections before rendering with handles

diff --git a/Engine/Source/Editor/MovieSceneTools/Private/AutomatedLevelSequenceCapture.cpp b/Engine/Source/Editor/MovieSceneTools/Private/AutomatedLevelSequenceCapture.cpp
--- a/Engine/Source/Editor/MovieSceneTools/Private/AutomatedLevelSequenceCapture.cpp
+++ b/Engine/Source/Editor/MovieSceneTools/Private/AutomatedLevelSequenceCapture.cpp
@@ -309,6 +309,8 @@ bool UAutomatedLevelSequenceCapture::InitializeShots()
 	// Compute handle frames in tick resolution space since that is what the section ranges are defined in
 	FFrameNumber HandleFramesResolutionSpace = ConvertFrameTime(Settings.HandleFrames, Settings.FrameRate, MovieScene->GetTickResolution()).FloorToFrame();
 
+	CinematicShotTrack->SortSections();
+
 	for (int32 SectionIndex = 0; SectionIndex < CinematicShotTrack->GetAllSections().Num(); ++SectionIndex)
 	{
 		UMovieSceneCinematicShotSection* ShotSection = Cast<UMovieSceneCinematicShotSection>(CinematicShotTrack->GetAllSections()[SectionIndex]);
diff --git a/Engine/Source/Runtime/MovieSceneTracks/Private/Tracks/MovieSceneCinematicShotTrack.cpp b/Engine/Source/Runtime/MovieSceneTracks/Private/Tracks/MovieSceneCinematicShotTrack.cpp
--- a/Engine/Source/Runtime/MovieSceneTracks/Private/Tracks/MovieSceneCinematicShotTrack.cpp
+++ b/Engine/Source/Runtime/MovieSceneTracks/Private/Tracks/MovieSceneCinematicShotTrack.cpp
@@ -164,4 +164,9 @@ FText UMovieSceneCinematicShotTrack::GetDefaultDisplayName() const
 }
 #endif
 
+void UMovieSceneCinematicShotTrack::SortSections()
+{
+	MovieSceneHelpers::SortConsecutiveSections(Sections);
+}
+
 #undef LOCTEXT_NAMESPACE
diff --git a/Engine/Source/Runtime/MovieSceneTracks/Public/Tracks/MovieSceneCinematicShotTrack.h b/Engine/Source/Runtime/MovieSceneTracks/Public/Tracks/MovieSceneCinematicShotTrack.h
--- a/Engine/Source/Runtime/MovieSceneTracks/Public/Tracks/MovieSceneCinematicShotTrack.h
+++ b/Engine/Source/Runtime/MovieSceneTracks/Public/Tracks/MovieSceneCinematicShotTrack.h
@@ -24,7 +24,9 @@ class UMovieSceneCinematicShotTrack
 public:
 
 	UMovieSceneCinematicShotTrack(const FObjectInitializer& ObjectInitializer);
-	
+
+	MOVIESCENETRACKS_API void SortSections();
+
 	// UMovieSceneSubTrack interface
 
 	MOVIESCENETRACKS_API virtual UMovieSceneSubSection* AddSequence(UMovieSceneSequence* Sequence, FFrameNumber StartTime, int32 Duration) { return AddSequenceOnRow(Sequence, StartTime, Duration, INDEX_NONE); }
