Sequencer: Draw section title at section clipping rect bottom left so that it's always on screen

diff --git a/Engine/Plugins/Media/MediaCompositing/Source/MediaCompositingEditor/Private/Sequencer/MediaThumbnailSection.cpp b/Engine/Plugins/Media/MediaCompositing/Source/MediaCompositingEditor/Private/Sequencer/MediaThumbnailSection.cpp
--- a/Engine/Plugins/Media/MediaCompositing/Source/MediaCompositingEditor/Private/Sequencer/MediaThumbnailSection.cpp
+++ b/Engine/Plugins/Media/MediaCompositing/Source/MediaCompositingEditor/Private/Sequencer/MediaThumbnailSection.cpp
@@ -62,20 +62,6 @@ void FMediaThumbnailSection::AddReferencedObjects(FReferenceCollector& Collector
 /* FThumbnailSection interface
  *****************************************************************************/
 
-TSharedRef<SWidget> FMediaThumbnailSection::GenerateSectionWidget()
-{
-	return SNew(SBox)
-		.HAlign(HAlign_Left)
-		.VAlign(VAlign_Top)
-		.Padding(GetContentPadding())
-		[
-			SNew(STextBlock)
-				.Text(this, &FMediaThumbnailSection::GetSectionTitle)
-				.ShadowOffset(FVector2D(1.0f, 1.0f))
-		];
-}
-
-
 FMargin FMediaThumbnailSection::GetContentPadding() const
 {
 	return FMargin(8.0f, 15.0f);
diff --git a/Engine/Plugins/Media/MediaCompositing/Source/MediaCompositingEditor/Private/Sequencer/MediaThumbnailSection.h b/Engine/Plugins/Media/MediaCompositing/Source/MediaCompositingEditor/Private/Sequencer/MediaThumbnailSection.h
--- a/Engine/Plugins/Media/MediaCompositing/Source/MediaCompositingEditor/Private/Sequencer/MediaThumbnailSection.h
+++ b/Engine/Plugins/Media/MediaCompositing/Source/MediaCompositingEditor/Private/Sequencer/MediaThumbnailSection.h
@@ -48,7 +48,6 @@ public:
 
 	//~ FThumbnailSection interface
 
-	virtual TSharedRef<SWidget> GenerateSectionWidget() override;
 	virtual FMargin GetContentPadding() const override;
 	virtual float GetSectionHeight() const override;
 	virtual FText GetSectionTitle() const override;
diff --git a/Engine/Source/Editor/MovieSceneTools/Private/Sections/CameraCutSection.cpp b/Engine/Source/Editor/MovieSceneTools/Private/Sections/CameraCutSection.cpp
--- a/Engine/Source/Editor/MovieSceneTools/Private/Sections/CameraCutSection.cpp
+++ b/Engine/Source/Editor/MovieSceneTools/Private/Sections/CameraCutSection.cpp
@@ -154,6 +154,11 @@ const AActor* FCameraCutSection::GetCameraForFrame(FFrameNumber Time) const
 	return nullptr;
 }
 
+FText FCameraCutSection::GetSectionTitle() const
+{
+	return HandleThumbnailTextBlockText();
+}
+
 float FCameraCutSection::GetSectionHeight() const
 {
 	return FViewportThumbnailSection::GetSectionHeight() + 10.f;
diff --git a/Engine/Source/Editor/MovieSceneTools/Private/Sections/CameraCutSection.h b/Engine/Source/Editor/MovieSceneTools/Private/Sections/CameraCutSection.h
--- a/Engine/Source/Editor/MovieSceneTools/Private/Sections/CameraCutSection.h
+++ b/Engine/Source/Editor/MovieSceneTools/Private/Sections/CameraCutSection.h
@@ -30,6 +30,7 @@ public:
 	// ISequencerSection interface
 	virtual void Tick(const FGeometry& AllottedGeometry, const FGeometry& ClippedGeometry, const double InCurrentTime, const float InDeltaTime) override;
 	virtual void BuildSectionContextMenu(FMenuBuilder& MenuBuilder, const FGuid& ObjectBinding) override;
+	virtual FText GetSectionTitle() const override;
 	virtual float GetSectionHeight() const override;
 	virtual int32 OnPaintSection(FSequencerSectionPainter& InPainter) const override;
 	virtual FMargin GetContentPadding() const override;
diff --git a/Engine/Source/Editor/MovieSceneTools/Private/Sections/CinematicShotSection.cpp b/Engine/Source/Editor/MovieSceneTools/Private/Sections/CinematicShotSection.cpp
--- a/Engine/Source/Editor/MovieSceneTools/Private/Sections/CinematicShotSection.cpp
+++ b/Engine/Source/Editor/MovieSceneTools/Private/Sections/CinematicShotSection.cpp
@@ -61,6 +61,11 @@ FCinematicShotSection::~FCinematicShotSection()
 {
 }
 
+FText FCinematicShotSection::GetSectionTitle() const
+{
+	return GetRenameVisibility() == EVisibility::Visible ? FText::GetEmpty() : HandleThumbnailTextBlockText();
+}
+
 float FCinematicShotSection::GetSectionHeight() const
 {
 	return FViewportThumbnailSection::GetSectionHeight() + 2*9.f;
diff --git a/Engine/Source/Editor/MovieSceneTools/Private/Sections/CinematicShotSection.h b/Engine/Source/Editor/MovieSceneTools/Private/Sections/CinematicShotSection.h
--- a/Engine/Source/Editor/MovieSceneTools/Private/Sections/CinematicShotSection.h
+++ b/Engine/Source/Editor/MovieSceneTools/Private/Sections/CinematicShotSection.h
@@ -35,6 +35,7 @@ public:
 	virtual int32 OnPaintSection( FSequencerSectionPainter& Painter ) const override;
 	virtual void BuildSectionContextMenu(FMenuBuilder& MenuBuilder, const FGuid& ObjectBinding) override;
 	virtual FReply OnSectionDoubleClicked(const FGeometry& SectionGeometry, const FPointerEvent& MouseEvent) override;
+	virtual FText GetSectionTitle() const override;
 	virtual float GetSectionHeight() const override;
 	virtual FMargin GetContentPadding() const override;
 	virtual void BeginResizeSection() override;
diff --git a/Engine/Source/Editor/MovieSceneTools/Private/Sections/ThumbnailSection.cpp b/Engine/Source/Editor/MovieSceneTools/Private/Sections/ThumbnailSection.cpp
--- a/Engine/Source/Editor/MovieSceneTools/Private/Sections/ThumbnailSection.cpp
+++ b/Engine/Source/Editor/MovieSceneTools/Private/Sections/ThumbnailSection.cpp
@@ -85,9 +85,19 @@ TSharedRef<SWidget> FThumbnailSection::GenerateSectionWidget()
 				.ShadowOffset(FVector2D(1,1))
 				.OnTextCommitted(this, &FThumbnailSection::HandleThumbnailTextBlockTextCommitted)
 				.IsReadOnly(!CanRename())
+				.Visibility(this, &FThumbnailSection::GetRenameVisibility)
 		];
 }
 
+EVisibility FThumbnailSection::GetRenameVisibility() const
+{
+	if (NameWidget.IsValid() && NameWidget->IsInEditMode())
+	{
+		return EVisibility::Visible;
+	}
+	return EVisibility::Hidden;
+}
+
 void FThumbnailSection::EnterRename()
 {
 	if (NameWidget.IsValid())
diff --git a/Engine/Source/Editor/MovieSceneTools/Private/TrackEditors/SkeletalAnimationTrackEditor.cpp b/Engine/Source/Editor/MovieSceneTools/Private/TrackEditors/SkeletalAnimationTrackEditor.cpp
--- a/Engine/Source/Editor/MovieSceneTools/Private/TrackEditors/SkeletalAnimationTrackEditor.cpp
+++ b/Engine/Source/Editor/MovieSceneTools/Private/TrackEditors/SkeletalAnimationTrackEditor.cpp
@@ -174,6 +174,12 @@ float FSkeletalAnimationSection::GetSectionHeight() const
 }
 
 
+FMargin FSkeletalAnimationSection::GetContentPadding() const
+{
+	return FMargin(8.0f, 8.0f);
+}
+
+
 int32 FSkeletalAnimationSection::OnPaintSection( FSequencerSectionPainter& Painter ) const
 {
 	const ESlateDrawEffect DrawEffects = Painter.bParentEnabled ? ESlateDrawEffect::None : ESlateDrawEffect::DisabledEffect;
diff --git a/Engine/Source/Editor/MovieSceneTools/Private/TrackEditors/SkeletalAnimationTrackEditor.h b/Engine/Source/Editor/MovieSceneTools/Private/TrackEditors/SkeletalAnimationTrackEditor.h
--- a/Engine/Source/Editor/MovieSceneTools/Private/TrackEditors/SkeletalAnimationTrackEditor.h
+++ b/Engine/Source/Editor/MovieSceneTools/Private/TrackEditors/SkeletalAnimationTrackEditor.h
@@ -96,6 +96,7 @@ public:
 	virtual UMovieSceneSection* GetSectionObject() override;
 	virtual FText GetSectionTitle() const override;
 	virtual float GetSectionHeight() const override;
+	virtual FMargin GetContentPadding() const override;
 	virtual int32 OnPaintSection( FSequencerSectionPainter& Painter ) const override;
 	virtual void BeginResizeSection() override;
 	virtual void ResizeSection(ESequencerSectionResizeMode ResizeMode, FFrameNumber ResizeTime) override;
diff --git a/Engine/Source/Editor/MovieSceneTools/Private/TrackEditors/SubTrackEditor.cpp b/Engine/Source/Editor/MovieSceneTools/Private/TrackEditors/SubTrackEditor.cpp
--- a/Engine/Source/Editor/MovieSceneTools/Private/TrackEditors/SubTrackEditor.cpp
+++ b/Engine/Source/Editor/MovieSceneTools/Private/TrackEditors/SubTrackEditor.cpp
@@ -202,6 +202,8 @@ public:
 				);
 			}
 
+			FMargin ContentPadding = GetContentPadding();
+
 			int32 NumTracks = MovieScene->GetPossessableCount() + MovieScene->GetSpawnableCount() + MovieScene->GetMasterTracks().Num();
 
 			FVector2D TopLeft = InPainter.SectionGeometry.AbsoluteToLocal(InPainter.SectionClippingRect.GetTopLeft()) + FVector2D(1.f, -1.f);
@@ -224,9 +226,8 @@ public:
 				++LayerId,
 				InPainter.SectionGeometry.MakeChild(
 					FVector2D(InPainter.SectionGeometry.Size.X, GetFontHeight()),
-					FSlateLayoutTransform(TopLeft + FVector2D(0.f, GetFontHeight()) + FVector2D(11.f, GetFontHeight()*2.f))
+					FSlateLayoutTransform(TopLeft + FVector2D(ContentPadding.Left, ContentPadding.Top) + FVector2D(11.f, GetFontHeight()*2.f))
 				).ToPaintGeometry(),
-				//InPainter.SectionGeometry.ToOffsetPaintGeometry(FVector2D(11.0f, 32.0f)),
 				FText::Format(LOCTEXT("NumTracksFormat", "{0} track(s)"), FText::AsNumber(NumTracks)),
 				FontInfo,
 				DrawEffects,
diff --git a/Engine/Source/Editor/MovieSceneTools/Public/Sections/ThumbnailSection.h b/Engine/Source/Editor/MovieSceneTools/Public/Sections/ThumbnailSection.h
--- a/Engine/Source/Editor/MovieSceneTools/Public/Sections/ThumbnailSection.h
+++ b/Engine/Source/Editor/MovieSceneTools/Public/Sections/ThumbnailSection.h
@@ -88,6 +88,9 @@ protected:
 	/** Get the total range that thumbnails are to be generated for in the section's time space */
 	TRange<double> GetTotalRange() const;
 
+	/** Get rename visibility */
+	EVisibility GetRenameVisibility() const;
+
 protected:
 
 	/** The section we are visualizing. */
diff --git a/Engine/Source/Editor/Sequencer/Private/SSequencerSection.cpp b/Engine/Source/Editor/Sequencer/Private/SSequencerSection.cpp
--- a/Engine/Source/Editor/Sequencer/Private/SSequencerSection.cpp
+++ b/Engine/Source/Editor/Sequencer/Private/SSequencerSection.cpp
@@ -1144,7 +1144,7 @@ int32 SSequencerSection::OnPaint( const FPaintArgs& Args, const FGeometry& Allot
 			LayerId,
 			SectionGeometry.MakeChild(
 				FVector2D(SectionGeometry.Size.X, GetFontHeight()),
-				FSlateLayoutTransform(TopLeft + FVector2D(0.f, GetFontHeight()) + FVector2D(1.f, 1.f))
+				FSlateLayoutTransform(TopLeft + FVector2D(ContentPadding.Left, ContentPadding.Top) + FVector2D(1.f, 1.f))
 			).ToPaintGeometry(),
 			SectionTitle,
 			FontInfo,
@@ -1157,7 +1157,7 @@ int32 SSequencerSection::OnPaint( const FPaintArgs& Args, const FGeometry& Allot
 			LayerId,
 			SectionGeometry.MakeChild(
 				FVector2D(SectionGeometry.Size.X, GetFontHeight()),
-				FSlateLayoutTransform(TopLeft + FVector2D(0.f, GetFontHeight()))
+				FSlateLayoutTransform(TopLeft + FVector2D(ContentPadding.Left, ContentPadding.Top))
 			).ToPaintGeometry(),
 			SectionTitle,
 			FontInfo,
