Sequencer: Expose for blueprints/python - skeletal animation, sub section, attach, constraint, path, camera cut, and level visibility sections

diff --git a/Engine/Plugins/VirtualProduction/Takes/Source/TakeRecorderSources/Private/TakeRecorderLevelVisibilitySource.cpp b/Engine/Plugins/VirtualProduction/Takes/Source/TakeRecorderSources/Private/TakeRecorderLevelVisibilitySource.cpp
--- a/Engine/Plugins/VirtualProduction/Takes/Source/TakeRecorderSources/Private/TakeRecorderLevelVisibilitySource.cpp
+++ b/Engine/Plugins/VirtualProduction/Takes/Source/TakeRecorderSources/Private/TakeRecorderLevelVisibilitySource.cpp
@@ -99,18 +99,12 @@ TArray<UTakeRecorderSource*> UTakeRecorderLevelVisibilitySource::PreRecording(cl
 	UMovieSceneLevelVisibilitySection* VisibleSection = NewObject<UMovieSceneLevelVisibilitySection>(CachedLevelVisibilityTrack.Get(), UMovieSceneLevelVisibilitySection::StaticClass());
 	VisibleSection->SetVisibility(ELevelVisibility::Visible);
 	VisibleSection->SetRowIndex(0);
-	for (FName VisibleLevelName : VisibleLevelNames)
-	{
-		VisibleSection->GetLevelNames()->AddUnique(VisibleLevelName);
-	}
+	VisibleSection->SetLevelNames(VisibleLevelNames);
 
 	UMovieSceneLevelVisibilitySection* HiddenSection = NewObject<UMovieSceneLevelVisibilitySection>(CachedLevelVisibilityTrack.Get(), UMovieSceneLevelVisibilitySection::StaticClass());
 	HiddenSection->SetVisibility(ELevelVisibility::Hidden);
 	HiddenSection->SetRowIndex(1);
-	for (FName HiddenLevelName : HiddenLevelNames)
-	{
-		HiddenSection->GetLevelNames()->AddUnique(HiddenLevelName);
-	}
+	HiddenSection->SetLevelNames(HiddenLevelNames);
 
 	CachedLevelVisibilityTrack->AddSection(*VisibleSection);
 	CachedLevelVisibilityTrack->AddSection(*HiddenSection);
diff --git a/Engine/Source/Editor/MovieSceneTools/Private/Sections/LevelVisibilitySection.cpp b/Engine/Source/Editor/MovieSceneTools/Private/Sections/LevelVisibilitySection.cpp
--- a/Engine/Source/Editor/MovieSceneTools/Private/Sections/LevelVisibilitySection.cpp
+++ b/Engine/Source/Editor/MovieSceneTools/Private/Sections/LevelVisibilitySection.cpp
@@ -61,13 +61,13 @@ FSlateColor FLevelVisibilitySection::GetBackgroundColor() const
 FText FLevelVisibilitySection::GetVisibilityText() const
 {
 	FText VisibilityText = SectionObject.GetVisibility() == ELevelVisibility::Visible ? VisibleText : HiddenText;
-	return FText::Format( NSLOCTEXT( "LevelVisibilitySection", "SectionTextFormat", "{0} ({1})" ), VisibilityText, FText::AsNumber( SectionObject.GetLevelNames()->Num() ) );
+	return FText::Format( NSLOCTEXT( "LevelVisibilitySection", "SectionTextFormat", "{0} ({1})" ), VisibilityText, FText::AsNumber( SectionObject.GetLevelNames().Num() ) );
 }
 
 FText FLevelVisibilitySection::GetVisibilityToolTip() const
 {
 	TArray<FString> LevelNameStrings;
-	for ( const FName& LevelName : *SectionObject.GetLevelNames() )
+	for ( const FName& LevelName : SectionObject.GetLevelNames() )
 	{
 		LevelNameStrings.Add( LevelName.ToString() );
 	}
@@ -93,14 +93,18 @@ FReply FLevelVisibilitySection::OnDrop( TSharedPtr<FDragDropOperation> DragDropO
 			FScopedTransaction Transaction(NSLOCTEXT("LevelVisibilitySection", "TransactionText", "Add Level(s) to Level Visibility Section"));
 			SectionObject.Modify();
 
+			TArray<FName> LevelNames = SectionObject.GetLevelNames();
 			for ( TWeakObjectPtr<ULevelStreaming> Level : LevelDragDropOperation->StreamingLevelsToDrop )
 			{
 				if ( Level.IsValid() )
 				{
 					FName ShortLevelName = FPackageName::GetShortFName( Level->GetWorldAssetPackageFName() );
-					SectionObject.GetLevelNames()->AddUnique( ShortLevelName );
+					LevelNames.AddUnique( ShortLevelName );
 				}
 			}
+
+			SectionObject.SetLevelNames(LevelNames);
+
 			return FReply::Handled();
 		}
 	}
diff --git a/Engine/Source/Runtime/MovieScene/Public/Evaluation/MovieSceneSectionParameters.h b/Engine/Source/Runtime/MovieScene/Public/Evaluation/MovieSceneSectionParameters.h
--- a/Engine/Source/Runtime/MovieScene/Public/Evaluation/MovieSceneSectionParameters.h
+++ b/Engine/Source/Runtime/MovieScene/Public/Evaluation/MovieSceneSectionParameters.h
@@ -24,16 +24,17 @@ struct FMovieSceneSectionParameters
 
 public:
 	/** Number of frames (in display rate) to skip at the beginning of the sub-sequence. */
-	UPROPERTY(EditAnywhere, BlueprintReadWrite, Category = "Clipping")
+	UPROPERTY(BlueprintReadWrite, EditAnywhere, BlueprintReadWrite, Category = "Clipping")
 	FFrameNumber StartFrameOffset;
 
 	/** Playback time scaling factor. */
-	UPROPERTY(EditAnywhere, BlueprintReadWrite, Category="Timing")
+	UPROPERTY(BlueprintReadWrite, EditAnywhere, BlueprintReadWrite, Category="Timing")
 	float TimeScale;
 
 	/** Hierachical bias. Higher bias will take precedence. */
-	UPROPERTY(EditAnywhere, BlueprintReadWrite, Category="Sequence")
+	UPROPERTY(BlueprintReadWrite, EditAnywhere, BlueprintReadWrite, Category="Sequence")
 	int32 HierarchicalBias;
+
 	UPROPERTY()
 	float StartOffset_DEPRECATED;
 	UPROPERTY()
diff --git a/Engine/Source/Runtime/MovieScene/Public/Sections/MovieSceneSubSection.h b/Engine/Source/Runtime/MovieScene/Public/Sections/MovieSceneSubSection.h
--- a/Engine/Source/Runtime/MovieScene/Public/Sections/MovieSceneSubSection.h
+++ b/Engine/Source/Runtime/MovieScene/Public/Sections/MovieSceneSubSection.h
@@ -142,7 +142,7 @@ public:
 
 public:
 
-	UPROPERTY(EditAnywhere, BlueprintReadWrite, Category="General", meta=(ShowOnlyInnerProperties))
+	UPROPERTY(BlueprintReadWrite, EditAnywhere, BlueprintReadWrite, Category="General", meta=(ShowOnlyInnerProperties))
 	FMovieSceneSectionParameters Parameters;
 
 private:
diff --git a/Engine/Source/Runtime/MovieSceneTracks/Private/Sections/MovieSceneLevelVisibilitySection.cpp b/Engine/Source/Runtime/MovieSceneTracks/Private/Sections/MovieSceneLevelVisibilitySection.cpp
--- a/Engine/Source/Runtime/MovieSceneTracks/Private/Sections/MovieSceneLevelVisibilitySection.cpp
+++ b/Engine/Source/Runtime/MovieSceneTracks/Private/Sections/MovieSceneLevelVisibilitySection.cpp
@@ -23,12 +23,6 @@ void UMovieSceneLevelVisibilitySection::SetVisibility( ELevelVisibility InVisibi
 }
 
 
-TArray<FName>* UMovieSceneLevelVisibilitySection::GetLevelNames()
-{
-	return &LevelNames;
-}
-
-
 FMovieSceneEvalTemplatePtr UMovieSceneLevelVisibilitySection::GenerateTemplate() const
 {
 	return FMovieSceneLevelVisibilitySectionTemplate(*this);
diff --git a/Engine/Source/Runtime/MovieSceneTracks/Public/Sections/MovieScene3DAttachSection.h b/Engine/Source/Runtime/MovieSceneTracks/Public/Sections/MovieScene3DAttachSection.h
--- a/Engine/Source/Runtime/MovieSceneTracks/Public/Sections/MovieScene3DAttachSection.h
+++ b/Engine/Source/Runtime/MovieSceneTracks/Public/Sections/MovieScene3DAttachSection.h
@@ -31,27 +31,27 @@ public:
 
 public:
 
-	UPROPERTY(EditAnywhere, Category="Attach")
+	UPROPERTY(BlueprintReadWrite, EditAnywhere, Category="Attach")
 	FName AttachSocketName;
 
-	UPROPERTY(EditAnywhere, Category="Attach")
+	UPROPERTY(BlueprintReadWrite, EditAnywhere, Category="Attach")
 	FName AttachComponentName;
 
-	UPROPERTY(EditAnywhere, Category="Attach")
+	UPROPERTY(BlueprintReadWrite, EditAnywhere, Category="Attach")
 	EAttachmentRule AttachmentLocationRule;
 
-	UPROPERTY(EditAnywhere, Category="Attach")
+	UPROPERTY(BlueprintReadWrite, EditAnywhere, Category="Attach")
 	EAttachmentRule AttachmentRotationRule;
 
-	UPROPERTY(EditAnywhere, Category="Attach")
+	UPROPERTY(BlueprintReadWrite, EditAnywhere, Category="Attach")
 	EAttachmentRule AttachmentScaleRule;
 
-	UPROPERTY(EditAnywhere, Category="Attach")
+	UPROPERTY(BlueprintReadWrite, EditAnywhere, Category="Attach")
 	EDetachmentRule DetachmentLocationRule;
 
-	UPROPERTY(EditAnywhere, Category="Attach")
+	UPROPERTY(BlueprintReadWrite, EditAnywhere, Category="Attach")
 	EDetachmentRule DetachmentRotationRule;
 
-	UPROPERTY(EditAnywhere, Category="Attach")
+	UPROPERTY(BlueprintReadWrite, EditAnywhere, Category="Attach")
 	EDetachmentRule DetachmentScaleRule;
 };
diff --git a/Engine/Source/Runtime/MovieSceneTracks/Public/Sections/MovieScene3DConstraintSection.h b/Engine/Source/Runtime/MovieSceneTracks/Public/Sections/MovieScene3DConstraintSection.h
--- a/Engine/Source/Runtime/MovieSceneTracks/Public/Sections/MovieScene3DConstraintSection.h
+++ b/Engine/Source/Runtime/MovieSceneTracks/Public/Sections/MovieScene3DConstraintSection.h
@@ -29,12 +29,14 @@ public:
 	virtual void SetConstraintId(const FGuid& InConstraintId, const FMovieSceneSequenceID& SequenceID);
 
 	/** Gets the constraint binding for this Constraint section */
+	UFUNCTION(BlueprintCallable, Category = "Movie Scene Section")
 	const FMovieSceneObjectBindingID& GetConstraintBindingID() const
 	{
 		return ConstraintBindingID;
 	}
 
 	/** Sets the constraint binding for this Constraint section */
+	UFUNCTION(BlueprintPure, Category = "Movie Scene Section")
 	void SetConstraintBindingID(const FMovieSceneObjectBindingID& InConstraintBindingID)
 	{
 		ConstraintBindingID = InConstraintBindingID;
diff --git a/Engine/Source/Runtime/MovieSceneTracks/Public/Sections/MovieScene3DPathSection.h b/Engine/Source/Runtime/MovieSceneTracks/Public/Sections/MovieScene3DPathSection.h
--- a/Engine/Source/Runtime/MovieSceneTracks/Public/Sections/MovieScene3DPathSection.h
+++ b/Engine/Source/Runtime/MovieSceneTracks/Public/Sections/MovieScene3DPathSection.h
@@ -61,29 +61,27 @@ public:
 
 	const FMovieSceneFloatChannel& GetTimingChannel() const { return TimingCurve; }
 
-private:
-
 	/** Timing Curve */
 	UPROPERTY()
 	FMovieSceneFloatChannel TimingCurve;
 
 	/** Front Axis */
-	UPROPERTY(EditAnywhere, Category="Path")
+	UPROPERTY(BlueprintReadWrite, EditAnywhere, Category="Path")
 	MovieScene3DPathSection_Axis FrontAxisEnum;
 
 	/** Up Axis */
-	UPROPERTY(EditAnywhere, Category="Path")
+	UPROPERTY(BlueprintReadWrite, EditAnywhere, Category="Path")
 	MovieScene3DPathSection_Axis UpAxisEnum;
 
 	/** Follow Curve */
-	UPROPERTY(EditAnywhere, Category="Path")
+	UPROPERTY(BlueprintReadWrite, EditAnywhere, Category="Path")
 	uint32 bFollow:1;
 
 	/** Reverse Timing */
-	UPROPERTY(EditAnywhere, Category="Path")
+	UPROPERTY(BlueprintReadWrite, EditAnywhere, Category="Path")
 	uint32 bReverse:1;
 
 	/** Force Upright */
-	UPROPERTY(EditAnywhere, Category="Path")
+	UPROPERTY(BlueprintReadWrite, EditAnywhere, Category="Path")
 	uint32 bForceUpright:1;
 };
diff --git a/Engine/Source/Runtime/MovieSceneTracks/Public/Sections/MovieSceneCameraCutSection.h b/Engine/Source/Runtime/MovieSceneTracks/Public/Sections/MovieSceneCameraCutSection.h
--- a/Engine/Source/Runtime/MovieSceneTracks/Public/Sections/MovieSceneCameraCutSection.h
+++ b/Engine/Source/Runtime/MovieSceneTracks/Public/Sections/MovieSceneCameraCutSection.h
@@ -38,12 +38,14 @@ public:
 	}
 
 	/** Gets the camera binding for this CameraCut section */
+	UFUNCTION(BlueprintCallable, Category = "Movie Scene Section")
 	const FMovieSceneObjectBindingID& GetCameraBindingID() const
 	{
 		return CameraBindingID;
 	}
 
 	/** Sets the camera binding for this CameraCut section */
+	UFUNCTION(BlueprintPure, Category = "Movie Scene Section")
 	void SetCameraBindingID(const FMovieSceneObjectBindingID& InCameraBindingID)
 	{
 		CameraBindingID = InCameraBindingID;
diff --git a/Engine/Source/Runtime/MovieSceneTracks/Public/Sections/MovieSceneLevelVisibilitySection.h b/Engine/Source/Runtime/MovieSceneTracks/Public/Sections/MovieSceneLevelVisibilitySection.h
--- a/Engine/Source/Runtime/MovieSceneTracks/Public/Sections/MovieSceneLevelVisibilitySection.h
+++ b/Engine/Source/Runtime/MovieSceneTracks/Public/Sections/MovieSceneLevelVisibilitySection.h
@@ -34,12 +34,18 @@ class MOVIESCENETRACKS_API UMovieSceneLevelVisibilitySection
 
 public:
 
+	UFUNCTION(BlueprintPure, Category = "Movie Scene Section")
 	ELevelVisibility GetVisibility() const;
+
+	UFUNCTION(BlueprintCallable, Category = "Movie Scene Section")
 	void SetVisibility(ELevelVisibility InVisibility);
 
-	TArray<FName>* GetLevelNames();
+	UFUNCTION(BlueprintPure, Category = "Movie Scene Section")
 	const TArray<FName>& GetLevelNames() const { return LevelNames; }
 
+	UFUNCTION(BlueprintCallable, Category = "Movie Scene Section")
+	void SetLevelNames(const TArray<FName>& InLevelNames) { LevelNames = InLevelNames; }
+
 public:
 
 	//~ UMovieSceneSection interface
