Sequencer: Optimized camera cut thumbnail rendering by removing reliance on level viewport client

  - Thumbnails are now drawn using a scene view and directly, which removes the overhead of the thumbnail storing its own scene view state
  - Removed a secondary pass that performed a blit to up/down-sample the initial scene render

diff --git a/Engine/Plugins/Experimental/ImagePlate/Source/ImagePlate/ImagePlate.Build.cs b/Engine/Plugins/Experimental/ImagePlate/Source/ImagePlate/ImagePlate.Build.cs
--- a/Engine/Plugins/Experimental/ImagePlate/Source/ImagePlate/ImagePlate.Build.cs
+++ b/Engine/Plugins/Experimental/ImagePlate/Source/ImagePlate/ImagePlate.Build.cs
@@ -28,7 +28,8 @@ namespace UnrealBuildTool.Rules
 					"RenderCore",
 					"RHI",
 					"TimeManagement",
-				}
+                    "SlateCore",
+                }
 			);
 
 			PrivateIncludePaths.AddRange(
diff --git a/Engine/Plugins/Experimental/ImagePlate/Source/ImagePlate/Private/ImagePlateFileSequence.cpp b/Engine/Plugins/Experimental/ImagePlate/Source/ImagePlate/Private/ImagePlateFileSequence.cpp
--- a/Engine/Plugins/Experimental/ImagePlate/Source/ImagePlate/Private/ImagePlateFileSequence.cpp
+++ b/Engine/Plugins/Experimental/ImagePlate/Source/ImagePlate/Private/ImagePlateFileSequence.cpp
@@ -16,6 +16,8 @@
 #include "Misc/PackageName.h"
 #include "Engine/Texture2DDynamic.h"
 #include "Engine/TextureRenderTarget2D.h"
+#include "Textures/SlateTextureData.h"
+
 
 DEFINE_LOG_CATEGORY_STATIC(LogImagePlateFileSequence, Log, Warning);
 
@@ -168,6 +170,7 @@ bool FImagePlateSourceFrame::EnsureTextureMetrics(UTexture* DestinationTexture)
 	return true;
 }
 
+
 TFuture<void> FImagePlateSourceFrame::CopyTo(UTexture* DestinationTexture)
 {
 	// Copy the data to the texture via a render command
@@ -242,6 +245,28 @@ TFuture<void> FImagePlateSourceFrame::CopyTo(UTexture* DestinationTexture)
 	return CompletionFuture;
 }
 
+TSharedRef<FSlateTextureData, ESPMode::ThreadSafe> FImagePlateSourceFrame::AsSlateTexture() const
+{
+	uint32 BytesPerPixel = (BitDepth * 4) / 8;
+	TSharedRef<FSlateTextureData, ESPMode::ThreadSafe> TextureData = MakeShared<FSlateTextureData, ESPMode::ThreadSafe>(nullptr, Width, Height, BytesPerPixel);
+
+	const uint8* SourceBuffer = Buffer.Get();
+	uint8* DestinationBuffer = TextureData->GetRawBytesPtr();
+	if (Pitch == Width * BytesPerPixel)
+	{
+		FMemory::Memcpy(DestinationBuffer, SourceBuffer, Width*Height*BytesPerPixel);
+	}
+	else for (uint32 Row = 0; Row < Height; ++Row)
+	{
+		FMemory::Memcpy(DestinationBuffer, SourceBuffer, FMath::Min(Width*BytesPerPixel, Pitch));
+
+		DestinationBuffer += Pitch;
+		SourceBuffer += (Width*BytesPerPixel);
+	}
+
+	return TextureData;
+}
+
 namespace ImagePlateFrameCache
 {
 	struct FPendingFrame;
diff --git a/Engine/Plugins/Experimental/ImagePlate/Source/ImagePlate/Public/ImagePlateFileSequence.h b/Engine/Plugins/Experimental/ImagePlate/Source/ImagePlate/Public/ImagePlateFileSequence.h
--- a/Engine/Plugins/Experimental/ImagePlate/Source/ImagePlate/Public/ImagePlateFileSequence.h
+++ b/Engine/Plugins/Experimental/ImagePlate/Source/ImagePlate/Public/ImagePlateFileSequence.h
@@ -10,7 +10,7 @@
 
 class UTexture;
 
-
+struct FSlateTextureData;
 struct FImagePlateAsyncCache;
 namespace ImagePlateFrameCache { struct FImagePlateSequenceCache; }
 
@@ -69,6 +69,9 @@ struct IMAGEPLATE_API FImagePlateSourceFrame
 	/** Copy the contents of this frame to the specified texture */
 	TFuture<void> CopyTo(UTexture* DestinationTexture);
 
+	/** Copy this source frame into a slate texture data format */
+	TSharedRef<FSlateTextureData, ESPMode::ThreadSafe> AsSlateTexture() const;
+
 private:
 
 	/** Ensure the specified texture metrics match this frame */
diff --git a/Engine/Plugins/Experimental/ImagePlate/Source/ImagePlateEditor/Private/ImagePlateTrackEditor.cpp b/Engine/Plugins/Experimental/ImagePlate/Source/ImagePlateEditor/Private/ImagePlateTrackEditor.cpp
--- a/Engine/Plugins/Experimental/ImagePlate/Source/ImagePlateEditor/Private/ImagePlateTrackEditor.cpp
+++ b/Engine/Plugins/Experimental/ImagePlate/Source/ImagePlateEditor/Private/ImagePlateTrackEditor.cpp
@@ -35,14 +35,13 @@
 #define LOCTEXT_NAMESPACE "ImagePlateTrackEditor"
 
 
-class FImagePlateSection : public FThumbnailSection, public ICustomThumbnailClient, FGCObject
+class FImagePlateSection : public FThumbnailSection, public ICustomThumbnailClient
 {
 public:
 
 	FImagePlateSection(UMovieSceneImagePlateSection& InSection, TSharedPtr<FTrackEditorThumbnailPool> InThumbnailPool, TSharedPtr<ISequencer> InSequencer)
 		: FThumbnailSection(InSequencer, InThumbnailPool, this, InSection)
 	{
-		RenderTexture = nullptr;
 		TimeSpace = ETimeSpace::Local;
 	}
 	
@@ -76,11 +75,6 @@ public:
 	// 		];
 	// }
 
-	virtual void AddReferencedObjects( FReferenceCollector& Collector )
-	{
-		Collector.AddReferencedObject(RenderTexture);
-	}
-
 	virtual float GetSectionHeight() const override
 	{
 		// Make space for the film border
@@ -103,7 +97,7 @@ public:
 			}
 			else
 			{
-////				ThumbnailCache.SetSingleReferenceFrame(TOptional<float>());
+				ThumbnailCache.SetSingleReferenceFrame(TOptional<double>());
 			}
 		}
 
@@ -145,12 +139,9 @@ public:
 		if (ImagePlateSection->FileSequence)
 		{
 			ThumbnailLoader = ImagePlateSection->FileSequence->GetAsyncCache();
-			RenderTexture = NewObject<UTexture2DDynamic>(GetTransientPackage(), FName(), RF_Transient);
-			RenderTexture->Init(256, 256, PF_R8G8B8A8);
 		}
 		else
 		{
-			RenderTexture = nullptr;
 			ThumbnailLoader.Reset();
 		}
 	}
@@ -178,27 +169,18 @@ public:
 
 		if (FrameData.IsValid())
 		{
-			// Wait for the texture to render
-			FrameData.CopyTo(RenderTexture).Get();
+			TrackEditorThumbnail.AssignFrom(FrameData.AsSlateTexture());
 
-			FTexture2DRHIRef Texture2DRHI = RenderTexture && RenderTexture->Resource && RenderTexture->Resource->TextureRHI.IsValid() ? RenderTexture->Resource->TextureRHI->GetTexture2D() : nullptr;
-			if (Texture2DRHI.IsValid())
+			TSharedPtr<ISequencer> Sequencer = SequencerPtr.Pin();
+			if (Sequencer.IsValid())
 			{
-				// Resolve the media player texture to the track editor thumbnail
-				TrackEditorThumbnail.CopyTextureIn(Texture2DRHI);
-
-				TSharedPtr<ISequencer> Sequencer = SequencerPtr.Pin();
-				if (Sequencer.IsValid())
-				{
-					TrackEditorThumbnail.SetupFade(Sequencer->GetSequencerWidget());
-				}
+				TrackEditorThumbnail.SetupFade(Sequencer->GetSequencerWidget());
 			}
 		}
 	}
 
 private:
 
-	UTexture2DDynamic* RenderTexture;
 	TOptional<FImagePlateAsyncCache> ThumbnailLoader;
 	TWeakObjectPtr<UImagePlateFileSequence> FileSequence;
 };
diff --git a/Engine/Source/Editor/MovieSceneTools/Private/Sections/CameraCutSection.cpp b/Engine/Source/Editor/MovieSceneTools/Private/Sections/CameraCutSection.cpp
--- a/Engine/Source/Editor/MovieSceneTools/Private/Sections/CameraCutSection.cpp
+++ b/Engine/Source/Editor/MovieSceneTools/Private/Sections/CameraCutSection.cpp
@@ -15,6 +15,7 @@
 #include "Evaluation/MovieSceneEvaluationTemplateInstance.h"
 #include "EditorStyleSet.h"
 #include "EngineUtils.h"
+#include "Camera/CameraComponent.h"
 
 
 #define LOCTEXT_NAMESPACE "FCameraCutSection"
@@ -121,27 +122,17 @@ void FCameraCutSection::BuildSectionContextMenu(FMenuBuilder& MenuBuilder, const
 /* FThumbnailSection interface
  *****************************************************************************/
 
-const AActor* FCameraCutSection::GetCameraForFrame(FFrameNumber Time) const
+AActor* FCameraCutSection::GetCameraForFrame(FFrameNumber Time) const
 {
 	UMovieSceneCameraCutSection* CameraCutSection = Cast<UMovieSceneCameraCutSection>(Section);
 	TSharedPtr<ISequencer> Sequencer = SequencerPtr.Pin();
 
 	if (CameraCutSection && Sequencer.IsValid())
 	{
-		FMovieSceneSequenceID SequenceID = Sequencer->GetFocusedTemplateID();
-		if (CameraCutSection->GetCameraBindingID().GetSequenceID().IsValid())
+		UCameraComponent* CameraComponent = CameraCutSection->GetFirstCamera(*Sequencer, Sequencer->GetFocusedTemplateID());
+		if (CameraComponent)
 		{
-			// Ensure that this ID is resolvable from the root, based on the current local sequence ID
-			FMovieSceneObjectBindingID RootBindingID = CameraCutSection->GetCameraBindingID().ResolveLocalToRoot(SequenceID, Sequencer->GetEvaluationTemplate().GetHierarchy());
-			SequenceID = RootBindingID.GetSequenceID();
-		}
-
-		for (TWeakObjectPtr<>& Object : Sequencer->FindBoundObjects(CameraCutSection->GetCameraBindingID().GetGuid(), SequenceID))
-		{
-			if (AActor* Actor = Cast<AActor>(Object.Get()))
-			{
-				return Actor;
-			}
+			return CameraComponent->GetOwner();
 		}
 
 		FMovieSceneSpawnable* Spawnable = Sequencer->GetFocusedMovieSceneSequence()->GetMovieScene()->FindSpawnable(CameraCutSection->GetCameraBindingID().GetGuid());
@@ -214,5 +205,18 @@ void FCameraCutSection::HandleSetCameraMenuEntryExecute(AActor* InCamera)
 	}
 }
 
+UCameraComponent* FCameraCutSection::GetViewCamera()
+{
+	UMovieSceneCameraCutSection* CameraCutSection = Cast<UMovieSceneCameraCutSection>(Section);
+	TSharedPtr<ISequencer>       Sequencer        = SequencerPtr.Pin();
+
+	if (CameraCutSection && Sequencer.IsValid())
+	{
+		return CameraCutSection->GetFirstCamera(*Sequencer, Sequencer->GetFocusedTemplateID());
+	}
+
+	return nullptr;
+}
+
 
 #undef LOCTEXT_NAMESPACE
diff --git a/Engine/Source/Editor/MovieSceneTools/Private/Sections/CameraCutSection.h b/Engine/Source/Editor/MovieSceneTools/Private/Sections/CameraCutSection.h
--- a/Engine/Source/Editor/MovieSceneTools/Private/Sections/CameraCutSection.h
+++ b/Engine/Source/Editor/MovieSceneTools/Private/Sections/CameraCutSection.h
@@ -38,11 +38,12 @@ public:
 
 	virtual void SetSingleTime(double GlobalTime) override;
 	virtual FText HandleThumbnailTextBlockText() const override;
+	virtual UCameraComponent* GetViewCamera() override;
 
 private:
 
 	/** Get a representative camera for the given time */
-	const AActor* GetCameraForFrame(FFrameNumber Time) const;
+	AActor* GetCameraForFrame(FFrameNumber Time) const;
 
 	/** Callback for executing a "Set Camera" menu entry in the context menu. */
 	void HandleSetCameraMenuEntryExecute(AActor* InCamera);
diff --git a/Engine/Source/Editor/MovieSceneTools/Private/Sections/CinematicShotSection.cpp b/Engine/Source/Editor/MovieSceneTools/Private/Sections/CinematicShotSection.cpp
--- a/Engine/Source/Editor/MovieSceneTools/Private/Sections/CinematicShotSection.cpp
+++ b/Engine/Source/Editor/MovieSceneTools/Private/Sections/CinematicShotSection.cpp
@@ -15,7 +15,9 @@
 #include "MovieSceneToolHelpers.h"
 #include "MovieSceneTimeHelpers.h"
 #include "Toolkits/AssetEditorManager.h"
-
+#include "Tracks/MovieSceneCameraCutTrack.h"
+#include "Sections/MovieSceneCameraCutSection.h"
+#include "Evaluation/MovieSceneEvaluationTemplateInstance.h"
 
 #define LOCTEXT_NAMESPACE "FCinematicShotSection"
 
@@ -47,7 +49,6 @@ FCinematicShotSection::FCinematicSectionCache::FCinematicSectionCache(UMovieScen
 FCinematicShotSection::FCinematicShotSection(TSharedPtr<ISequencer> InSequencer, TSharedPtr<FTrackEditorThumbnailPool> InThumbnailPool, UMovieSceneSection& InSection, TSharedPtr<FCinematicShotTrackEditor> InCinematicShotTrackEditor)
 	: FViewportThumbnailSection(InSequencer, InThumbnailPool, InSection)
 	, SectionObject(*CastChecked<UMovieSceneCinematicShotSection>(&InSection))
-	, Sequencer(InSequencer)
 	, CinematicShotTrackEditor(InCinematicShotTrackEditor)
 	, InitialStartOffsetDuringResize(0)
 	, InitialStartTimeDuringResize(0)
@@ -82,6 +83,110 @@ void FCinematicShotSection::SetSingleTime(double GlobalTime)
 	SectionObject.SetThumbnailReferenceOffset(GlobalTime - ReferenceOffsetSeconds);
 }
 
+
+UCameraComponent* FindCameraCutComponentRecursive(FFrameNumber GlobalTime, FMovieSceneSequenceID InnerSequenceID, const FMovieSceneSequenceHierarchy& Hierarchy, IMovieScenePlayer& Player)
+{
+	const FMovieSceneSequenceHierarchyNode* Node    = Hierarchy.FindNode(InnerSequenceID);
+	const FMovieSceneSubSequenceData*       SubData = Hierarchy.FindSubData(InnerSequenceID);
+	if (!ensure(SubData && Node))
+	{
+		return nullptr;
+	}
+
+	UMovieSceneSequence* InnerSequence   = SubData->GetSequence();
+	UMovieScene*         InnerMovieScene = InnerSequence ? InnerSequence->GetMovieScene() : nullptr;
+	if (!InnerMovieScene)
+	{
+		return nullptr;
+	}
+
+	FFrameNumber InnerTime = (GlobalTime * SubData->RootToSequenceTransform).FloorToFrame();
+	if (!SubData->PlayRange.Value.Contains(InnerTime))
+	{
+		return nullptr;
+	}
+
+	int32 LowestRow      = TNumericLimits<int32>::Max();
+	int32 HighestOverlap = 0;
+
+	UMovieSceneCameraCutSection* ActiveSection = nullptr;
+
+	if (UMovieSceneCameraCutTrack* CutTrack = Cast<UMovieSceneCameraCutTrack>(InnerMovieScene->GetCameraCutTrack()))
+	{
+		for (UMovieSceneSection* ItSection : CutTrack->GetAllSections())
+		{
+			UMovieSceneCameraCutSection* CutSection = Cast<UMovieSceneCameraCutSection>(ItSection);
+			if (CutSection && CutSection->GetRange().Contains(InnerTime))
+			{
+				bool bSectionWins = 
+					( CutSection->GetRowIndex() < LowestRow ) ||
+					( CutSection->GetRowIndex() == LowestRow && CutSection->GetOverlapPriority() > HighestOverlap );
+
+				if (bSectionWins)
+				{
+					HighestOverlap = CutSection->GetOverlapPriority();
+					LowestRow      = CutSection->GetRowIndex();
+					ActiveSection  = CutSection;
+				}
+			}
+		}
+	}
+	
+	if (ActiveSection)
+	{
+		return ActiveSection->GetFirstCamera(Player, InnerSequenceID);
+	}
+
+	for (FMovieSceneSequenceID Child : Node->Children)
+	{
+		UCameraComponent* CameraComponent = FindCameraCutComponentRecursive(GlobalTime, Child, Hierarchy, Player);
+		if (CameraComponent)
+		{
+			return CameraComponent;
+		}
+	}
+
+	return nullptr;
+}
+
+
+UCameraComponent* FCinematicShotSection::GetViewCamera()
+{
+	TSharedPtr<ISequencer> Sequencer = SequencerPtr.Pin();
+	if (!Sequencer.IsValid())
+	{
+		return nullptr;
+	}
+
+
+	const FMovieSceneSequenceID             ThisSequenceID   = Sequencer->GetFocusedTemplateID();
+	const FMovieSceneSequenceID             TargetSequenceID = SectionObject.GetSequenceID();
+	const FMovieSceneSequenceHierarchy&     Hierarchy        = Sequencer->GetEvaluationTemplate().GetHierarchy();
+	const FMovieSceneSequenceHierarchyNode* ThisSequenceNode = Hierarchy.FindNode(ThisSequenceID);
+	
+	check(ThisSequenceNode);
+	
+	// Find the TargetSequenceID by comparing deterministic sequence IDs for all children of the current node
+	const FMovieSceneSequenceID* InnerSequenceID = Algo::FindByPredicate(ThisSequenceNode->Children,
+		[&Hierarchy, TargetSequenceID](FMovieSceneSequenceID InSequenceID)
+		{
+			const FMovieSceneSubSequenceData* SubData = Hierarchy.FindSubData(InSequenceID);
+			return SubData && SubData->DeterministicSequenceID == TargetSequenceID;
+		}
+	);
+	
+	if (InnerSequenceID)
+	{
+		UCameraComponent* CameraComponent = FindCameraCutComponentRecursive(Sequencer->GetGlobalTime().Time.FrameNumber, *InnerSequenceID, Hierarchy, *Sequencer);
+		if (CameraComponent)
+		{
+			return CameraComponent;
+		}
+	}
+
+	return nullptr;
+}
+
 void FCinematicShotSection::BeginResizeSection()
 {
 	InitialStartOffsetDuringResize = SectionObject.Parameters.StartFrameOffset;
@@ -390,7 +495,7 @@ FReply FCinematicShotSection::OnSectionDoubleClicked(const FGeometry& SectionGeo
 			}
 			else
 			{
-				Sequencer.Pin()->FocusSequenceInstance(SectionObject);
+				SequencerPtr.Pin()->FocusSequenceInstance(SectionObject);
 			}
 		}
 	}
diff --git a/Engine/Source/Editor/MovieSceneTools/Private/Sections/CinematicShotSection.h b/Engine/Source/Editor/MovieSceneTools/Private/Sections/CinematicShotSection.h
--- a/Engine/Source/Editor/MovieSceneTools/Private/Sections/CinematicShotSection.h
+++ b/Engine/Source/Editor/MovieSceneTools/Private/Sections/CinematicShotSection.h
@@ -48,6 +48,7 @@ public:
 	virtual void SetSingleTime(double GlobalTime) override;
 	virtual FText HandleThumbnailTextBlockText() const override;
 	virtual void HandleThumbnailTextBlockTextCommitted(const FText& NewThumbnailName, ETextCommit::Type CommitType) override;
+	virtual UCameraComponent* GetViewCamera() override;
 
 private:
 
@@ -59,9 +60,6 @@ private:
 	/** The section we are visualizing */
 	UMovieSceneCinematicShotSection& SectionObject;
 
-	/** Sequencer interface */
-	TWeakPtr<ISequencer> Sequencer;
-
 	/** The cinematic shot track editor that contains this section */
 	TWeakPtr<FCinematicShotTrackEditor> CinematicShotTrackEditor;
 
diff --git a/Engine/Source/Editor/MovieSceneTools/Private/Sections/ThumbnailSection.cpp b/Engine/Source/Editor/MovieSceneTools/Private/Sections/ThumbnailSection.cpp
--- a/Engine/Source/Editor/MovieSceneTools/Private/Sections/ThumbnailSection.cpp
+++ b/Engine/Source/Editor/MovieSceneTools/Private/Sections/ThumbnailSection.cpp
@@ -229,52 +229,72 @@ int32 FThumbnailSection::OnPaintSection( FSequencerSectionPainter& InPainter ) c
 	for (const TSharedPtr<FTrackEditorThumbnail>& Thumbnail : ThumbnailCache.GetThumbnails())
 	{
 		const float Fade = Thumbnail->bHasFinishedDrawing ? Thumbnail->GetFadeInCurve() : 1.f;
-		FIntPoint ThumbnailSize = Thumbnail->GetSize();
+		if (Fade >= 1.f)
+		{
+			continue;
+		}
+
+		FIntPoint ThumbnailRTSize   = Thumbnail->GetSize();
+		FIntPoint ThumbnailCropSize = Thumbnail->GetDesiredSize();
 		
+		const float ThumbnailScale       = float(ThumbnailCropSize.Y) / ThumbnailRTSize.Y;
+		const float HorizontalCropOffset = (ThumbnailRTSize.X*ThumbnailScale - ThumbnailCropSize.X) * 0.5f;
+
 		// Calculate the paint geometry for this thumbnail
 		TOptional<double> SingleReferenceFrame = ThumbnailCache.GetSingleReferenceFrame();
 
 		// Single thumbnails are always drawn at the start of the section, clamped to the visible range
 		// Thumbnail sequences draw relative to their actual position in the sequence/section
-		const int32 Offset = SingleReferenceFrame.IsSet()
+		const float PositionX = SingleReferenceFrame.IsSet()
 			? FMath::Max(float(VisibleRange.GetLowerBoundValue() - GenerationRange.GetLowerBoundValue()) / TimePerPx, 0.f) + SectionThumbnailPadding
 			: (Thumbnail->GetTimeRange().GetLowerBoundValue() - GenerationRange.GetLowerBoundValue()) / TimePerPx;
 
+		const float PositionY = (SectionGeometry.GetLocalSize().Y - ThumbnailCropSize.Y)*.5f;
+
+
 		FPaintGeometry PaintGeometry = SectionGeometry.ToPaintGeometry(
-			ThumbnailSize,
-			FSlateLayoutTransform(
-				SectionGeometry.Scale,
-				FVector2D( Offset, (SectionGeometry.GetLocalSize().Y - ThumbnailSize.Y)*.5f)
-			)
+			ThumbnailRTSize,
+			FSlateLayoutTransform(ThumbnailScale, FVector2D(PositionX-HorizontalCropOffset, PositionY))
 		);
 
-		if (Fade <= 1.f)
+		if (IVREditorModule::Get().IsVREditorModeActive())
 		{
-			if (IVREditorModule::Get().IsVREditorModeActive())
-			{
-				// In VR editor every widget is in the world and gamma corrected by the scene renderer.  Thumbnails will have already been gamma
-				// corrected and so they need to be reversed
-				DrawEffects |= ESlateDrawEffect::ReverseGamma;
-			}
-			else
-			{
-				DrawEffects |= ESlateDrawEffect::NoGamma;
-			}
-
-			FSlateClippingZone ClippingZone(ThumbnailClipRect);
-			InPainter.DrawElements.PushClip(ClippingZone);
-
-			FSlateDrawElement::MakeViewport(
-				InPainter.DrawElements,
-				LayerId,
-				PaintGeometry,
-				Thumbnail,
-				DrawEffects | AdditionalDrawEffect,
-				FLinearColor(1.f, 1.f, 1.f, 1.f - Fade)
-			);
-
-			InPainter.DrawElements.PopClip();
+			// In VR editor every widget is in the world and gamma corrected by the scene renderer.  Thumbnails will have already been gamma
+			// corrected and so they need to be reversed
+			DrawEffects |= ESlateDrawEffect::ReverseGamma;
 		}
+		else
+		{
+			DrawEffects |= ESlateDrawEffect::NoGamma;
+		}
+
+		if (Thumbnail->bIgnoreAlpha)
+		{
+			DrawEffects |= ESlateDrawEffect::IgnoreTextureAlpha;
+		}
+
+		FGeometry ClipGeometry = SectionGeometry.MakeChild(
+			ThumbnailCropSize,
+			FSlateLayoutTransform(
+				FVector2D( PositionX, PositionY)
+			)
+		);
+
+		FSlateRect ThisThumbnailClipRect = ThumbnailClipRect.IntersectionWith(ClipGeometry.GetLayoutBoundingRect());
+
+		FSlateClippingZone ClippingZone(ThisThumbnailClipRect);
+		InPainter.DrawElements.PushClip(ClippingZone);
+
+		FSlateDrawElement::MakeViewport(
+			InPainter.DrawElements,
+			LayerId,
+			PaintGeometry,
+			Thumbnail,
+			DrawEffects | AdditionalDrawEffect,
+			FLinearColor(1.f, 1.f, 1.f, 1.f - Fade)
+		);
+
+		InPainter.DrawElements.PopClip();
 	}
 
 	return LayerId + 2;
@@ -282,20 +302,19 @@ int32 FThumbnailSection::OnPaintSection( FSequencerSectionPainter& InPainter ) c
 
 TRange<double> FThumbnailSection::GetVisibleRange() const
 {
+	const FFrameRate TickResolution = Section->GetTypedOuter<UMovieScene>()->GetTickResolution();
 	TRange<double> GlobalVisibleRange = SequencerPtr.Pin()->GetViewRange();
+	TRange<double> SectionRange = Section->GetRange() / TickResolution;
+
 	if (TimeSpace == ETimeSpace::Global)
 	{
 		return GlobalVisibleRange;
 	}
 
-	const FFrameRate TickResolution = Section->GetTypedOuter<UMovieScene>()->GetTickResolution();
-
-	const double     StartOffset = Section->HasStartFrame() ? Section->GetInclusiveStartFrame() / TickResolution : 0.0;
-	const double     EndOffset   = Section->HasEndFrame()   ? Section->GetExclusiveEndFrame()   / TickResolution : 0.0;
-
+	TRange<double> Intersection = TRange<double>::Intersection(GlobalVisibleRange, SectionRange);
 	return TRange<double>(
-		GlobalVisibleRange.GetLowerBoundValue() - StartOffset,
-		GlobalVisibleRange.GetUpperBoundValue() - EndOffset
+		Intersection.GetLowerBoundValue() - SectionRange.GetLowerBoundValue(),
+		Intersection.GetUpperBoundValue() - SectionRange.GetLowerBoundValue()
 	);
 }
 
@@ -339,30 +358,21 @@ FViewportThumbnailSection::FViewportThumbnailSection(TSharedPtr<ISequencer> InSe
 }
 
 
-void FViewportThumbnailSection::PreDraw(FTrackEditorThumbnail& Thumbnail, FLevelEditorViewportClient& ViewportClient, FSceneViewport& SceneViewport)
+void FViewportThumbnailSection::PreDraw(FTrackEditorThumbnail& Thumbnail)
 {
 	TSharedPtr<ISequencer> Sequencer = SequencerPtr.Pin();
 	if (Sequencer.IsValid())
 	{
 		Sequencer->EnterSilentMode();
-
-		AActor* Camera = nullptr;
-		FDelegateHandle Handle = Sequencer->OnCameraCut().AddLambda([&](UObject* InObject, bool){
-			Camera = Cast<AActor>(InObject);
-		});
-
 		SavedPlaybackStatus = Sequencer->GetPlaybackStatus();
 		Sequencer->SetPlaybackStatus(EMovieScenePlayerStatus::Jumping);
 		Sequencer->SetLocalTimeDirectly(Thumbnail.GetEvalPosition() * Sequencer->GetLocalTime().Rate );
 		Sequencer->ForceEvaluate();
-
-		ViewportClient.SetActorLock(Camera);
-		Sequencer->OnCameraCut().Remove(Handle);
 	}
 }
 
 
-void FViewportThumbnailSection::PostDraw(FTrackEditorThumbnail& Thumbnail, FLevelEditorViewportClient& ViewportClient, FSceneViewport& SceneViewport)
+void FViewportThumbnailSection::PostDraw(FTrackEditorThumbnail& Thumbnail)
 {
 	TSharedPtr<ISequencer> Sequencer = SequencerPtr.Pin();
 	if (Sequencer.IsValid())
@@ -372,5 +382,4 @@ void FViewportThumbnailSection::PostDraw(FTrackEditorThumbnail& Thumbnail, FLeve
 	}
 }
 
-
 #undef LOCTEXT_NAMESPACE
diff --git a/Engine/Source/Editor/MovieSceneTools/Private/TrackEditorThumbnail/TrackEditorThumbnail.cpp b/Engine/Source/Editor/MovieSceneTools/Private/TrackEditorThumbnail/TrackEditorThumbnail.cpp
--- a/Engine/Source/Editor/MovieSceneTools/Private/TrackEditorThumbnail/TrackEditorThumbnail.cpp
+++ b/Engine/Source/Editor/MovieSceneTools/Private/TrackEditorThumbnail/TrackEditorThumbnail.cpp
@@ -11,13 +11,15 @@
 #include "StaticBoundShaderState.h"
 #include "EditorViewportClient.h"
 #include "LevelEditorViewport.h"
-#include "Slate/SlateTextures.h"
 #include "Slate/SceneViewport.h"
 #include "GlobalShader.h"
 #include "ScreenRendering.h"
 #include "TrackEditorThumbnail/TrackEditorThumbnailPool.h"
 #include "PipelineStateCache.h"
 #include "CommonRenderResources.h"
+#include "LegacyScreenPercentageDriver.h"
+#include "CanvasTypes.h"
+#include "EngineModule.h"
 
 namespace TrackEditorThumbnailConstants
 {
@@ -53,21 +55,23 @@ public:
 /* FTrackEditorThumbnail structors
  *****************************************************************************/
 
-FTrackEditorThumbnail::FTrackEditorThumbnail(const FOnThumbnailDraw& InOnDraw, const FIntPoint& InSize, TRange<double> InTimeRange, double InPosition)
+FTrackEditorThumbnail::FTrackEditorThumbnail(const FOnThumbnailDraw& InOnDraw, const FIntPoint& InDesiredSize, TRange<double> InTimeRange, double InPosition)
 	: OnDraw(InOnDraw)
-	, Size(InSize)
-	, Texture(nullptr)
+	, DesiredSize(InDesiredSize)
+	, ThumbnailTexture(nullptr)
+	, ThumbnailRenderTarget(nullptr)
 	, TimeRange(InTimeRange)
 	, Position(InPosition)
 	, FadeInCurve(0.0f, TrackEditorThumbnailConstants::ThumbnailFadeInDuration)
 {
 	SortOrder = 0;
+	bIgnoreAlpha = false;
 }
 
 
 FTrackEditorThumbnail::~FTrackEditorThumbnail()
 {
-	if (Texture && !bHasFinishedDrawing)
+	if (ThumbnailRenderTarget && !bHasFinishedDrawing)
 	{
 		FlushRenderingCommands();
 	}
@@ -75,142 +79,94 @@ FTrackEditorThumbnail::~FTrackEditorThumbnail()
 }
 
 
-void FTrackEditorThumbnail::DestroyTexture()
+void FTrackEditorThumbnail::AssignFrom(TSharedRef<FSlateTextureData, ESPMode::ThreadSafe> InTextureData)
 {
-	if (Texture)
+	if (!ThumbnailTexture)
 	{
-		FSlateTexture2DRHIRef* InTexture = Texture;
-		ENQUEUE_RENDER_COMMAND(DestroyTexture)(
-			[InTexture](FRHICommandList& RHICmdList)
-			{
-				InTexture->ReleaseResource();
-				delete InTexture;
-			});
-
-		Texture = nullptr;
+		EPixelFormat PixelFormat = InTextureData->GetBytesPerPixel() == 4 ? PF_B8G8R8A8 : PF_FloatRGBA;
+		ThumbnailTexture = new FSlateTexture2DRHIRef(InTextureData->GetWidth(), InTextureData->GetHeight(), PixelFormat, NULL, TexCreate_Dynamic);
 	}
-}
 
+	FSlateTexture2DRHIRef* InThumbnailTexture = ThumbnailTexture;
+	ENQUEUE_RENDER_COMMAND(AssignTexture)(
+		[InThumbnailTexture, InTextureData](FRHICommandList& RHICmdList){
+			InThumbnailTexture->SetTextureData(InTextureData);
+			if (InThumbnailTexture->IsInitialized())
+			{
+				InThumbnailTexture->UpdateRHI();
+			}
+			else
+			{
+				InThumbnailTexture->InitResource();
+			}
+		}
+	);
+}
 
-/* FTrackEditorThumbnail interface
- *****************************************************************************/
 
-void FTrackEditorThumbnail::CopyTextureIn(TSharedPtr<FSceneViewport> SceneViewport)
+void FTrackEditorThumbnail::DestroyTexture()
 {
-	// Note: We explicitly capture the viewport scene here to ensure that the render target lives as long as this render command.
-	// This means we don't have to flush the rendering commands all the time
-	SceneViewportReference = SceneViewport;
-
-	FSlateRenderTargetRHI* RenderTarget = (FSlateRenderTargetRHI*)SceneViewport->GetViewportRenderTargetTexture();
-	if (RenderTarget)
+	if (ThumbnailRenderTarget || ThumbnailTexture)
 	{
-		CopyTextureIn(RenderTarget->GetRHIRef());
+		FSlateTexture2DRHIRef*               InThumbnailTexture      = ThumbnailTexture;
+		FSlateTextureRenderTarget2DResource* InThumbnailRenderTarget = ThumbnailRenderTarget;
+
+		ThumbnailTexture      = nullptr;
+		ThumbnailRenderTarget = nullptr;
+
+		ENQUEUE_RENDER_COMMAND(DestroyTexture)(
+			[InThumbnailRenderTarget, InThumbnailTexture](FRHICommandList& RHICmdList)
+			{
+				if (InThumbnailTexture)
+				{
+					InThumbnailTexture->ReleaseResource();
+					delete InThumbnailTexture;
+				}
+
+				if (InThumbnailRenderTarget)
+				{
+					InThumbnailRenderTarget->ReleaseResource();
+					delete InThumbnailRenderTarget;
+				}
+			}
+		);
 	}
 }
 
-void FTrackEditorThumbnail::CopyTextureIn(FTexture2DRHIRef SourceTexture)
-{
-	// This code is a little manual because CopyToResolveTarget on its own can't resolve a sub rect without offsetting it inside the destination texture
-	// So we render our own rectangle onto a render target of the right size, so we can maintain the correct aspect ratio and fov settings of the camera,
-	// but still fulfil the desired thumbnail size.
-	FSlateTexture2DRHIRef* TargetTexture = Texture;
-
-	static const FName RendererModuleName( "Renderer" );
-	IRendererModule* RendererModule = &FModuleManager::GetModuleChecked<IRendererModule>(RendererModuleName);
-	FThreadSafeBool* bHasFinishedDrawingPtr = &bHasFinishedDrawing;
 
-	auto RenderCommand = [TargetTexture, RendererModule, SourceTexture, bHasFinishedDrawingPtr](FRHICommandListImmediate& RHICmdList){
+void FTrackEditorThumbnail::ResizeRenderTarget(const FIntPoint& InSize)
+{
+	// Delay texture creation until we actually draw the thumbnail
+	if (InSize.X <= 0 || InSize.Y <= 0)
+	{
+		return;
+	}
 
-		const FIntPoint TargetSize(TargetTexture->GetWidth(), TargetTexture->GetHeight());
+	if (ThumbnailTexture && ThumbnailRenderTarget && ThumbnailTexture->GetWidth() == InSize.X && ThumbnailTexture->GetHeight() == InSize.Y)
+	{
+		return;
+	}
 
-		FPooledRenderTargetDesc OutputDesc = FPooledRenderTargetDesc::Create2DDesc(
-			TargetSize,
-			PF_B8G8R8A8,
-			FClearValueBinding::None,
-			TexCreate_None,
-			TexCreate_RenderTargetable,
-			false);
+	DestroyTexture();
 
-		TRefCountPtr<IPooledRenderTarget> ResampleTexturePooledRenderTarget;
-		RendererModule->RenderTargetPoolFindFreeElement(RHICmdList, OutputDesc, ResampleTexturePooledRenderTarget, TEXT("ResampleTexture"));
-		check(ResampleTexturePooledRenderTarget);
+	ThumbnailTexture      = new FSlateTexture2DRHIRef(InSize.X, InSize.Y, PF_B8G8R8A8, NULL, TexCreate_Dynamic);
+	ThumbnailRenderTarget = new FSlateTextureRenderTarget2DResource(FLinearColor::Black, InSize.X, InSize.Y, PF_B8G8R8A8, SF_Point, TA_Wrap, TA_Wrap, 0.0f);
 
-		const FSceneRenderTargetItem& DestRenderTarget = ResampleTexturePooledRenderTarget->GetRenderTargetItem();
+	FSlateTexture2DRHIRef*               InThumbnailTexture      = ThumbnailTexture;
+	FSlateTextureRenderTarget2DResource* InThumbnailRenderTarget = ThumbnailRenderTarget;
 
-		FGraphicsPipelineStateInitializer GraphicsPSOInit;
-		FRHIRenderPassInfo RPInfo(DestRenderTarget.TargetableTexture, ERenderTargetActions::Load_Store);
-		RHICmdList.BeginRenderPass(RPInfo, TEXT("CopyTextureIn"));
+	ENQUEUE_RENDER_COMMAND(AssignRenderTarget)(
+		[InThumbnailRenderTarget, InThumbnailTexture](FRHICommandList& RHICmdList)
 		{
-			RHICmdList.ApplyCachedRenderTargets(GraphicsPSOInit);
-			GraphicsPSOInit.PrimitiveType = PT_TriangleList;
-
-			RHICmdList.SetViewport(0, 0, 0.0f, TargetSize.X, TargetSize.Y, 1.0f);
-			//		RHICmdList.ClearColorTexture(DestRenderTarget.TargetableTexture, FLinearColor(0.0f, 0.0f, 0.0f, 1.0f));
-
-			GraphicsPSOInit.BlendState = TStaticBlendState<>::GetRHI();
-			GraphicsPSOInit.RasterizerState = TStaticRasterizerState<>::GetRHI();
-			GraphicsPSOInit.DepthStencilState = TStaticDepthStencilState<false, CF_Always>::GetRHI();
-
-			const ERHIFeatureLevel::Type FeatureLevel = GMaxRHIFeatureLevel;
-
-			TShaderMap<FGlobalShaderType>* ShaderMap = GetGlobalShaderMap(FeatureLevel);
-			TShaderMapRef<FScreenVS> VertexShader(ShaderMap);
-			TShaderMapRef<FScreenPS> PixelShader(ShaderMap);
-
-			GraphicsPSOInit.BoundShaderState.VertexDeclarationRHI = GFilterVertexDeclaration.VertexDeclarationRHI;
-			GraphicsPSOInit.BoundShaderState.VertexShaderRHI = GETSAFERHISHADER_VERTEX(*VertexShader);
-			GraphicsPSOInit.BoundShaderState.PixelShaderRHI = GETSAFERHISHADER_PIXEL(*PixelShader);
-
-			SetGraphicsPipelineState(RHICmdList, GraphicsPSOInit);
-
-			PixelShader->SetParameters(RHICmdList, TStaticSamplerState<SF_Bilinear>::GetRHI(), SourceTexture);
-
-			const float Scale = FMath::Min(float(SourceTexture->GetSizeX()) / TargetTexture->GetWidth(), float(SourceTexture->GetSizeY()) / TargetTexture->GetHeight());
-			const float Left = (SourceTexture->GetSizeX() - TargetTexture->GetWidth()*Scale) * .5f;
-			const float Top = (SourceTexture->GetSizeY() - TargetTexture->GetHeight()*Scale) * .5f;
-
-			const float U = Left / float(SourceTexture->GetSizeX());
-			const float V = Top / float(SourceTexture->GetSizeY());
-			const float SizeU = float(TargetTexture->GetWidth() * Scale) / float(SourceTexture->GetSizeX());
-			const float SizeV = float(TargetTexture->GetHeight() * Scale) / float(SourceTexture->GetSizeY());
-
-			RendererModule->DrawRectangle(
-				RHICmdList,
-				0, 0,									// Dest X, Y
-				TargetSize.X,							// Dest Width
-				TargetSize.Y,							// Dest Height
-				U, V,									// Source U, V
-				SizeU, SizeV,							// Source USize, VSize
-				TargetSize,								// Target buffer size
-				FIntPoint(1, 1),						// Source texture size
-				*VertexShader,
-				EDRF_Default);
+			InThumbnailTexture->InitResource();
+			InThumbnailRenderTarget->InitResource();
+			InThumbnailTexture->SetRHIRef(InThumbnailRenderTarget->GetTextureRHI(), InThumbnailRenderTarget->GetSizeX(), InThumbnailRenderTarget->GetSizeY());
 		}
-		RHICmdList.EndRenderPass();
-		RHICmdList.CopyToResolveTarget(DestRenderTarget.TargetableTexture, TargetTexture->GetTypedResource(), FResolveParams());
-
-		*bHasFinishedDrawingPtr = true;
-	};
-
-	TFunction<void(FRHICommandListImmediate&)> InRenderCommand = RenderCommand;
-	ENQUEUE_RENDER_COMMAND(ResolveCaptureFrameTexture)(
-		[InRenderCommand](FRHICommandListImmediate& RHICmdList)
-		{
-			InRenderCommand(RHICmdList);
-		});
+	);
 }
 
-
 void FTrackEditorThumbnail::DrawThumbnail()
 {
-	// Delay texture creation until we actually draw the thumbnail
-	if (Size.X > 1 && Size.Y > 1)
-	{
-		DestroyTexture();
-		Texture = new FSlateTexture2DRHIRef(Size.X, Size.Y, PF_B8G8R8A8, nullptr, TexCreate_Dynamic, true);
-		BeginInitResource(Texture);
-	}
-
 	OnDraw.ExecuteIfBound(*this);
 }
 
@@ -239,13 +195,17 @@ float FTrackEditorThumbnail::GetFadeInCurve() const
 
 FIntPoint FTrackEditorThumbnail::GetSize() const
 {
-	return Size;
+	if (ThumbnailTexture)
+	{
+		return FIntPoint(ThumbnailTexture->GetWidth(), ThumbnailTexture->GetHeight());
+	}
+	return FIntPoint(0,0);
 }
 
 
 FSlateShaderResource* FTrackEditorThumbnail::GetViewportRenderTargetTexture() const
 {
-	return Texture;
+	return ThumbnailTexture;
 }
 
 
@@ -262,7 +222,6 @@ FTrackEditorThumbnailCache::FTrackEditorThumbnailCache(const TSharedPtr<FTrackEd
 {
 	check(ViewportThumbnailClient);
 
-	FrameCount = 0;
 	LastComputationTime = 0;
 	bForceRedraw = false;
 	bNeedsNewThumbnails = false;
@@ -276,7 +235,6 @@ FTrackEditorThumbnailCache::FTrackEditorThumbnailCache(const TSharedPtr<FTrackEd
 {
 	check(CustomThumbnailClient);
 
-	FrameCount = 0;
 	LastComputationTime = 0;
 	bForceRedraw = false;
 	bNeedsNewThumbnails = false;
@@ -290,11 +248,6 @@ FTrackEditorThumbnailCache::~FTrackEditorThumbnailCache()
 	{
 		PinnedPool->RemoveThumbnailsNeedingRedraw(Thumbnails);
 	}
-
-	if (InternalViewportClient.IsValid())
-	{
-		InternalViewportClient->Viewport = nullptr;
-	}
 }
 
 
@@ -324,10 +277,9 @@ void FTrackEditorThumbnailCache::Update(const TRange<double>& NewRange, const TR
 }
 
 
-FIntPoint FTrackEditorThumbnailCache::CalculateTextureSize() const
+FIntPoint FTrackEditorThumbnailCache::CalculateTextureSize(const FMinimalViewInfo& ViewInfo) const
 {
-	UCameraComponent* CameraComponent = InternalViewportClient->GetCameraComponentForView();
-	float DesiredRatio = CameraComponent ? CameraComponent->AspectRatio : 1.77f;
+	float DesiredRatio = ViewInfo.AspectRatio;
 
 	if (CurrentCache.DesiredSize.X <= 0 || CurrentCache.DesiredSize.Y <= 0)
 	{
@@ -378,48 +330,99 @@ bool FTrackEditorThumbnailCache::ShouldRegenerateEverything() const
 	return PreviousCache.DesiredSize != CurrentCache.DesiredSize || !FMath::IsNearlyEqual(PreviousScale, CurrentScale, Threshold);
 }
 
-
-void FTrackEditorThumbnailCache::DrawViewportThumbnail(FTrackEditorThumbnail& TrackEditorThumbnail)
+void FTrackEditorThumbnailCache::DrawThumbnail(FTrackEditorThumbnail& TrackEditorThumbnail)
 {
 	if (CustomThumbnailClient)
 	{
 		CustomThumbnailClient->Draw(TrackEditorThumbnail);
 	}
-	else if (InternalViewportScene.IsValid())
+	else if (ViewportThumbnailClient)
 	{
-		check(ViewportThumbnailClient);
+		ViewportThumbnailClient->PreDraw(TrackEditorThumbnail);
 
-		// Ask the client to setup the frame
-		ViewportThumbnailClient->PreDraw(TrackEditorThumbnail, *InternalViewportClient, *InternalViewportScene);
-		
-		// Finalize the view
-		InternalViewportClient->bLockedCameraView = true;
-		InternalViewportClient->UpdateViewForLockedActor();
-		InternalViewportClient->GetWorld()->SendAllEndOfFrameUpdates();
-
-		// Update the viewport RHI if necessary
-		FIntPoint Size = CalculateTextureSize();
-		if (InternalViewportScene->GetSize() != Size)
-		{
-			InternalViewportScene->UpdateViewportRHI(false, Size.X, Size.Y, EWindowMode::Windowed, PF_Unknown);
-		}
+		DrawViewportThumbnail(TrackEditorThumbnail);
 
-		InternalViewportClient->DeltaWorldTime = TrackEditorThumbnail.GetEvalPosition() - InternalViewportClient->CurrentWorldTime;
-		InternalViewportClient->CurrentWorldTime = TrackEditorThumbnail.GetEvalPosition();
+		ViewportThumbnailClient->PostDraw(TrackEditorThumbnail);
+	}
 
-		// Draw the frame. If our total frame count is < 3 we re-render some benign frames to ensure the view is correctly set up (first few frames can be black)
-		do
+	FThreadSafeBool* bHasFinishedDrawingPtr = &TrackEditorThumbnail.bHasFinishedDrawing;
+	ENQUEUE_RENDER_COMMAND(SetFinishedDrawing)(
+		[bHasFinishedDrawingPtr](FRHICommandList& RHICmdList)
 		{
-			InternalViewportScene->Draw(false);
+			*bHasFinishedDrawingPtr = true;
 		}
-		while (++FrameCount < 3);
+	);
+}
+void FTrackEditorThumbnailCache::DrawViewportThumbnail(FTrackEditorThumbnail& TrackEditorThumbnail)
+{
+	check(ViewportThumbnailClient);
+
+	UCameraComponent* PreviewCameraComponent = ViewportThumbnailClient->GetViewCamera();
+	if (!PreviewCameraComponent)
+	{
+		return;
+	}
 
-		// Ask the client to finalize the frame
-		ViewportThumbnailClient->PostDraw(TrackEditorThumbnail, *InternalViewportClient, *InternalViewportScene);
+	FMinimalViewInfo ViewInfo;
+	PreviewCameraComponent->GetCameraView(FApp::GetDeltaTime(), ViewInfo);
 
-		// Copy the render target into our texture
-		TrackEditorThumbnail.CopyTextureIn(InternalViewportScene);
+	FIntPoint RTSize = CalculateTextureSize(ViewInfo);
+	if (RTSize.X <= 0 || RTSize.Y <= 0)
+	{
+		return;
 	}
+
+	TrackEditorThumbnail.bIgnoreAlpha = true;
+	TrackEditorThumbnail.ResizeRenderTarget(RTSize);
+
+	UWorld* World = PreviewCameraComponent->GetWorld();
+
+	FSceneViewFamilyContext ViewFamily( FSceneViewFamily::ConstructionValues( TrackEditorThumbnail.GetRenderTarget(), World->Scene, FEngineShowFlags(ESFIM_Game) )
+		.SetWorldTimes(FApp::GetCurrentTime() - GStartTime, FApp::GetDeltaTime(), FApp::GetCurrentTime() - GStartTime)
+		.SetResolveScene(true));
+
+	// Screen percentage is not supported in thumbnail.
+	ViewFamily.EngineShowFlags.ScreenPercentage = false;
+
+	switch (CurrentCache.Quality)
+	{
+	case EThumbnailQuality::Draft:
+		ViewFamily.EngineShowFlags.DisableAdvancedFeatures();
+		ViewFamily.EngineShowFlags.SetPostProcessing(false);
+		break;
+
+	case EThumbnailQuality::Normal:
+	case EThumbnailQuality::Best:
+		ViewFamily.EngineShowFlags.SetMotionBlur(false);
+		break;
+	}
+
+	FSceneViewInitOptions ViewInitOptions;
+
+	ViewInitOptions.BackgroundColor = FLinearColor::Black;
+	ViewInitOptions.SetViewRectangle(FIntRect(FIntPoint::ZeroValue, RTSize));
+	ViewInitOptions.ViewFamily = &ViewFamily;
+
+	ViewInitOptions.ViewOrigin = ViewInfo.Location;
+	ViewInitOptions.ViewRotationMatrix = FInverseRotationMatrix(ViewInfo.Rotation) * FMatrix(
+		FPlane(0, 0, 1, 0),
+		FPlane(1, 0, 0, 0),
+		FPlane(0, 1, 0, 0),
+		FPlane(0, 0, 0, 1));
+
+	ViewInitOptions.ProjectionMatrix = ViewInfo.CalculateProjectionMatrix();
+
+	FSceneView* NewView = new FSceneView(ViewInitOptions);
+	ViewFamily.Views.Add(NewView);
+
+	const float GlobalResolutionFraction = 1.f;
+	const bool  AllowPostProcessSettingsScreenPercentage = false;
+	ViewFamily.SetScreenPercentageInterface(new FLegacyScreenPercentageDriver(ViewFamily, GlobalResolutionFraction, AllowPostProcessSettingsScreenPercentage));
+
+	FCanvas Canvas(TrackEditorThumbnail.GetRenderTarget(), nullptr, FApp::GetCurrentTime() - GStartTime, FApp::GetDeltaTime(), FApp::GetCurrentTime() - GStartTime, World->Scene->GetFeatureLevel());
+	Canvas.Clear(FLinearColor::Transparent);
+
+	GetRendererModule().BeginRenderingViewFamily(&Canvas, &ViewFamily);
 }
 
 
@@ -440,11 +443,6 @@ void FTrackEditorThumbnailCache::Revalidate(double InCurrentTime)
 		return;
 	}
 
-	if (CurrentCache.Quality != PreviousCache.Quality)
-	{
-		SetupViewportEngineFlags();
-	}
-
 	bNeedsNewThumbnails = true;
 
 	if (ShouldRegenerateEverything())
@@ -499,7 +497,7 @@ void FTrackEditorThumbnailCache::UpdateSingleThumbnail()
 	const double EvalPosition = CurrentCache.SingleReferenceFrame.GetValue();
 
 	TSharedPtr<FTrackEditorThumbnail> NewThumbnail = MakeShareable(new FTrackEditorThumbnail(
-		FOnThumbnailDraw::CreateRaw(this, &FTrackEditorThumbnailCache::DrawViewportThumbnail),
+		FOnThumbnailDraw::CreateRaw(this, &FTrackEditorThumbnailCache::DrawThumbnail),
 		CurrentCache.DesiredSize,
 		TRange<double>(EvalPosition - HalfRange, EvalPosition + HalfRange),
 		EvalPosition
@@ -598,7 +596,7 @@ void FTrackEditorThumbnailCache::GenerateFront(const TRange<double>& Boundary)
 		double EvalPosition = CurrentCache.TimeRange.GetLowerBoundValue() + FMath::Clamp(TotalLerp, 0.0, .99)*CurrentCache.TimeRange.Size<double>();
 
 		TSharedPtr<FTrackEditorThumbnail> NewThumbnail = MakeShareable(new FTrackEditorThumbnail(
-			FOnThumbnailDraw::CreateRaw(this, &FTrackEditorThumbnailCache::DrawViewportThumbnail),
+			FOnThumbnailDraw::CreateRaw(this, &FTrackEditorThumbnailCache::DrawThumbnail),
 			TextureSize,
 			TimeRange,
 			EvalPosition
@@ -637,7 +635,7 @@ void FTrackEditorThumbnailCache::GenerateBack(const TRange<double>& Boundary)
 		double EvalPosition = CurrentCache.TimeRange.GetLowerBoundValue() + FMath::Clamp(TotalLerp, 0.0, .99)*CurrentCache.TimeRange.Size<double>();
 
 		TSharedPtr<FTrackEditorThumbnail> NewThumbnail = MakeShareable(new FTrackEditorThumbnail(
-			FOnThumbnailDraw::CreateRaw(this, &FTrackEditorThumbnailCache::DrawViewportThumbnail),
+			FOnThumbnailDraw::CreateRaw(this, &FTrackEditorThumbnailCache::DrawThumbnail),
 			TextureSize,
 			TimeRange,
 			EvalPosition
@@ -655,58 +653,9 @@ void FTrackEditorThumbnailCache::GenerateBack(const TRange<double>& Boundary)
 
 void FTrackEditorThumbnailCache::Setup()
 {
-	// Set up the necessary viewport gubbins for viewport thumbnails
-	if (ViewportThumbnailClient)
-	{
-		if (!InternalViewportClient.IsValid())
-		{
-			InternalViewportClient = MakeShareable(new FThumbnailViewportClient());
-			InternalViewportClient->ViewportType = LVT_Perspective;
-			InternalViewportClient->bDisableInput = true;
-			InternalViewportClient->bDrawAxes = false;
-			InternalViewportClient->SetAllowCinematicControl(false);
-			InternalViewportClient->SetRealtime(false);
-
-			SetupViewportEngineFlags();
-
-			InternalViewportClient->ViewState.GetReference()->SetSequencerState(true);
-		}
-
-		if (!InternalViewportScene.IsValid())
-		{
-			InternalViewportScene = MakeShareable(new FSceneViewport(InternalViewportClient.Get(), nullptr));
-			InternalViewportClient->Viewport = InternalViewportScene.Get();
-		}
-	}
-	else if (CustomThumbnailClient)
+	if (CustomThumbnailClient)
 	{
 		CustomThumbnailClient->Setup();
 	}
 }
 
-void FTrackEditorThumbnailCache::SetupViewportEngineFlags()
-{
-	if (!ViewportThumbnailClient || !InternalViewportClient.IsValid())
-	{
-		return;
-	}
-
-	InternalViewportClient->EngineShowFlags = FEngineShowFlags(ESFIM_Game);
-
-	// Screen percentage is not supported in thumbnail.
-	InternalViewportClient->EngineShowFlags.ScreenPercentage = false;
-
-	switch (CurrentCache.Quality)
-	{
-	case EThumbnailQuality::Draft:
-		InternalViewportClient->EngineShowFlags.DisableAdvancedFeatures();
-		break;
-
-	case EThumbnailQuality::Normal:
-	case EThumbnailQuality::Best:
-		InternalViewportClient->EngineShowFlags.SetMotionBlur(false);
-		break;
-	}
-
-	InternalViewportClient->Invalidate();
-}
diff --git a/Engine/Source/Editor/MovieSceneTools/Public/Sections/ThumbnailSection.h b/Engine/Source/Editor/MovieSceneTools/Public/Sections/ThumbnailSection.h
--- a/Engine/Source/Editor/MovieSceneTools/Public/Sections/ThumbnailSection.h
+++ b/Engine/Source/Editor/MovieSceneTools/Public/Sections/ThumbnailSection.h
@@ -139,6 +139,6 @@ public:
 	FViewportThumbnailSection(TSharedPtr<ISequencer> InSequencer, TSharedPtr<FTrackEditorThumbnailPool> InThumbnailPool, UMovieSceneSection& InSection);
 
 	//~ IViewportThumbnailClient interface
-	virtual void PreDraw(FTrackEditorThumbnail& Thumbnail, FLevelEditorViewportClient& ViewportClient, FSceneViewport& SceneViewport) override;
-	virtual void PostDraw(FTrackEditorThumbnail& Thumbnail, FLevelEditorViewportClient& ViewportClient, FSceneViewport& SceneViewport) override;
+	virtual void PreDraw(FTrackEditorThumbnail& Thumbnail) override;
+	virtual void PostDraw(FTrackEditorThumbnail& Thumbnail) override;
 };
diff --git a/Engine/Source/Editor/MovieSceneTools/Public/TrackEditorThumbnail/TrackEditorThumbnail.h b/Engine/Source/Editor/MovieSceneTools/Public/TrackEditorThumbnail/TrackEditorThumbnail.h
--- a/Engine/Source/Editor/MovieSceneTools/Public/TrackEditorThumbnail/TrackEditorThumbnail.h
+++ b/Engine/Source/Editor/MovieSceneTools/Public/TrackEditorThumbnail/TrackEditorThumbnail.h
@@ -10,6 +10,9 @@
 #include "MovieSceneToolsUserSettings.h"
 #include "RHI.h"
 #include "Misc/FrameTime.h"
+#include "Slate/SlateTextures.h"
+
+struct FMinimalViewInfo;
 
 class FLevelEditorViewportClient;
 class FSceneViewport;
@@ -18,9 +21,12 @@ class FSlateTexture2DRHIRef;
 class FThumbnailViewportClient;
 class FTrackEditorThumbnail;
 class FTrackEditorThumbnailPool;
+class UCameraComponent;
+
 
 DECLARE_DELEGATE_OneParam(FOnThumbnailDraw, FTrackEditorThumbnail&);
 
+
 /**
  * Track Editor Thumbnail, which keeps a Texture to be displayed by a viewport.
  */
@@ -31,18 +37,38 @@ class MOVIESCENETOOLS_API FTrackEditorThumbnail
 public:
 
 	/** Create and initialize a new instance. */
-	FTrackEditorThumbnail(const FOnThumbnailDraw& InOnDraw, const FIntPoint& InSize, TRange<double> InTimeRange, double InPosition);
+	FTrackEditorThumbnail(const FOnThumbnailDraw& InOnDraw, const FIntPoint& InDesiredSize, TRange<double> InTimeRange, double InPosition);
 
 	/** Virtual destructor. */
 	virtual ~FTrackEditorThumbnail();
 
 public:
 
-	/** Copies the incoming render target to this thumbnails texture. */
-	void CopyTextureIn(TSharedPtr<FSceneViewport> SceneViewport);
-	
-	/** Copies the specified texture to this thumbnail's texture while maintining the correct ratios */
-	void CopyTextureIn(FTexture2DRHIRef SourceTexture);
+	/**
+	 * Assign this thumbnail from a slate texture.
+	 */
+	void AssignFrom(TSharedRef<FSlateTextureData, ESPMode::ThreadSafe> InTextureData);
+
+	/**
+	 * Ensure that this thumbnail has a render target of the specified size
+	 */
+	void ResizeRenderTarget(const FIntPoint& InSize);
+
+	/**
+	 * Access the (potentially null) render target to be used for rendering onto this thumbnail
+	 */
+	FSlateTextureRenderTarget2DResource* GetRenderTarget() const
+	{
+		return ThumbnailRenderTarget;
+	}
+
+	/**
+	 * Get the desired size for this thumbnail on the UI
+	 */
+	FIntPoint GetDesiredSize() const
+	{
+		return DesiredSize;
+	}
 
 	/** Renders the thumbnail to the texture. */
 	void DrawThumbnail();
@@ -81,16 +107,21 @@ public:
 	/** True when this thumbnail has been drawn, false otherwise */
 	FThreadSafeBool bHasFinishedDrawing;
 
+	/** True to ignore alpha on this thumbnail */
+	bool bIgnoreAlpha;
+
 private:
 
 	/** Delegate to use to draw the thumbnail. */
 	FOnThumbnailDraw OnDraw;
 
-	/** The size of the texture */
-	FIntPoint Size;
+	/** The desired size of the thumbnail on the actual UI (Not necessarily the same size as the texture) */
+	FIntPoint DesiredSize;
 
 	/** The Texture RHI that holds the thumbnail. */
-	FSlateTexture2DRHIRef* Texture;
+	FSlateTexture2DRHIRef* ThumbnailTexture;
+	/** The texture render target used for 3D rendering on to the texture. May be null. */
+	FSlateTextureRenderTarget2DResource* ThumbnailRenderTarget;
 
 	/** Where in time this thumbnail is a rendering of. */
 	TRange<double> TimeRange;
@@ -100,16 +131,14 @@ private:
 
 	/** Fade curve to display while the thumbnail is redrawing. */
 	FCurveSequence FadeInCurve;
-
-	/** Strong reference to a scene viewport we're currently copying from */
-	TSharedPtr<FSceneViewport> SceneViewportReference;
 };
 
 /** Client interface for thumbanils that render the current world from a viewport */
 struct IViewportThumbnailClient
 {
-	virtual void PreDraw(FTrackEditorThumbnail& TrackEditorThumbnail, FLevelEditorViewportClient& ViewportClient, FSceneViewport& SceneViewport) { }
-	virtual void PostDraw(FTrackEditorThumbnail& TrackEditorThumbnail, FLevelEditorViewportClient& ViewportClient, FSceneViewport& SceneViewport) { }
+	virtual UCameraComponent* GetViewCamera() { return nullptr; }
+	virtual void PreDraw(FTrackEditorThumbnail& TrackEditorThumbnail) { }
+	virtual void PostDraw(FTrackEditorThumbnail& TrackEditorThumbnail) { }
 };
 
 /** Custom thumbnail drawing client interface */
@@ -177,8 +206,6 @@ public:
 
 	void Revalidate(double InCurrentTime);
 
-	void DrawViewportThumbnail(FTrackEditorThumbnail& TrackEditorThumbnail);
-
 	const TArray<TSharedPtr<FTrackEditorThumbnail>>& GetThumbnails() const
 	{
 		return Thumbnails;
@@ -186,13 +213,17 @@ public:
 
 protected:
 
+	void DrawThumbnail(FTrackEditorThumbnail& TrackEditorThumbnail);
+
+	void DrawViewportThumbnail(FTrackEditorThumbnail& TrackEditorThumbnail);
+
 	void ComputeNewThumbnails();
 
 	void Setup();
 
 	bool ShouldRegenerateEverything() const;
 
-	FIntPoint CalculateTextureSize() const;
+	FIntPoint CalculateTextureSize(const FMinimalViewInfo& ViewInfo) const;
 
 	void UpdateSingleThumbnail();
 
@@ -210,12 +241,6 @@ protected:
 	IViewportThumbnailClient* ViewportThumbnailClient;
 	ICustomThumbnailClient* CustomThumbnailClient;
 
-	/** An internal viewport scene we use to render the thumbnails with. */
-	TSharedPtr<FSceneViewport> InternalViewportScene;
-
-	/** An internal editor viewport client to render the thumbnails with. */
-	TSharedPtr<class FThumbnailViewportClient> InternalViewportClient;
-
 	/** The thumbnail pool that we are sending all of our thumbnails to. */
 	TWeakPtr<FTrackEditorThumbnailPool> ThumbnailPool;
 
@@ -226,9 +251,6 @@ protected:
 	TArray<TSharedPtr<FTrackEditorThumbnail>> Thumbnails;
 	TArray<TSharedPtr<FTrackEditorThumbnail>> ThumbnailsNeedingRedraw;
 
-	/** The number of frames we've rendered */
-	uint32 FrameCount;
-
 	double LastComputationTime;
 	bool bNeedsNewThumbnails;
 
diff --git a/Engine/Source/Runtime/MovieSceneTracks/Private/Sections/MovieSceneCameraCutSection.cpp b/Engine/Source/Runtime/MovieSceneTracks/Private/Sections/MovieSceneCameraCutSection.cpp
--- a/Engine/Source/Runtime/MovieSceneTracks/Private/Sections/MovieSceneCameraCutSection.cpp
+++ b/Engine/Source/Runtime/MovieSceneTracks/Private/Sections/MovieSceneCameraCutSection.cpp
@@ -9,6 +9,9 @@
 #include "Evaluation/MovieSceneEvaluationTrack.h"
 #include "Evaluation/MovieSceneCameraCutTemplate.h"
 #include "Compilation/MovieSceneTemplateInterrogation.h"
+#include "IMovieScenePlayer.h"
+#include "Camera/CameraComponent.h"
+#include "Evaluation/MovieSceneEvaluationTemplateInstance.h"
 
 /* UMovieSceneCameraCutSection interface
  *****************************************************************************/
@@ -75,3 +78,28 @@ void UMovieSceneCameraCutSection::PostLoad()
 		CameraGuid_DEPRECATED.Invalidate();
 	}
 }
+
+
+UCameraComponent* UMovieSceneCameraCutSection::GetFirstCamera(IMovieScenePlayer& Player, FMovieSceneSequenceID SequenceID) const
+{
+	if (CameraBindingID.GetSequenceID().IsValid())
+	{
+		// Ensure that this ID is resolvable from the root, based on the current local sequence ID
+		FMovieSceneObjectBindingID RootBindingID = CameraBindingID.ResolveLocalToRoot(SequenceID, Player.GetEvaluationTemplate().GetHierarchy());
+		SequenceID = RootBindingID.GetSequenceID();
+	}
+
+	for (TWeakObjectPtr<>& WeakObject : Player.FindBoundObjects(CameraBindingID.GetGuid(), SequenceID))
+	{
+		if (UObject* Object = WeakObject .Get())
+		{
+			UCameraComponent* Camera = MovieSceneHelpers::CameraComponentFromRuntimeObject(Object);
+			if (Camera)
+			{
+				return Camera;
+			}
+		}
+	}
+
+	return nullptr;
+}
\ No newline at end of file
diff --git a/Engine/Source/Runtime/MovieSceneTracks/Public/Sections/MovieSceneCameraCutSection.h b/Engine/Source/Runtime/MovieSceneTracks/Public/Sections/MovieSceneCameraCutSection.h
--- a/Engine/Source/Runtime/MovieSceneTracks/Public/Sections/MovieSceneCameraCutSection.h
+++ b/Engine/Source/Runtime/MovieSceneTracks/Public/Sections/MovieSceneCameraCutSection.h
@@ -11,6 +11,9 @@
 #include "MovieSceneObjectBindingID.h"
 #include "MovieSceneCameraCutSection.generated.h"
 
+struct FMovieSceneSequenceID;
+class IMovieScenePlayer;
+class UCameraComponent;
 
 /**
  * Movie CameraCuts are sections on the CameraCuts track, that show what the viewer "sees"
@@ -59,6 +62,16 @@ public:
 	/** ~UObject interface */
 	virtual void PostLoad() override;
 
+	/**
+	 * Resolve a camera component for this cut section from the specified player and sequence ID
+	 *
+	 * @param Player     The sequence player to use to resolve the object binding for this camera
+	 * @param SequenceID The sequence ID for the specific instance that this section exists within
+	 *
+	 * @return A camera component to be used for this cut section, or nullptr if one was not found.
+	 */
+	MOVIESCENETRACKS_API UCameraComponent* GetFirstCamera(IMovieScenePlayer& Player, FMovieSceneSequenceID SequenceID) const;
+
 private:
 
 	/** The camera possessable or spawnable that this movie CameraCut uses */
