Add option to burn in a movie sequence with path-tracer for a given number of sample per pixel.

This CL also add an internal delay per frame, allowing to wait a certain amount of time before capturing a movie frame.

diff --git a/Engine/Source/Editor/MovieSceneCaptureDialog/Private/MovieSceneCaptureDialogModule.cpp b/Engine/Source/Editor/MovieSceneCaptureDialog/Private/MovieSceneCaptureDialogModule.cpp
--- a/Engine/Source/Editor/MovieSceneCaptureDialog/Private/MovieSceneCaptureDialogModule.cpp
+++ b/Engine/Source/Editor/MovieSceneCaptureDialog/Private/MovieSceneCaptureDialogModule.cpp
@@ -529,6 +529,12 @@ void FInEditorCapture::OnPIEViewportStarted()
 					CapturingFromWorld->GetWorldSettings()->DefaultGameMode = CaptureObject->Settings.GameModeOverride;
 				}
 
+				CachedEngineShowFlags = SlatePlayInEditorSession->SlatePlayInEditorWindowViewport->GetClient()->GetEngineShowFlags();
+				if (CachedEngineShowFlags && Settings.bUsePathTracer)
+				{
+					CachedPathTracingMode = CachedEngineShowFlags->PathTracing;
+					CachedEngineShowFlags->SetPathTracing(true);
+				}
 				CaptureObject->Initialize(SlatePlayInEditorSession->SlatePlayInEditorWindowViewport, Context.PIEInstance);
 				OnCaptureStarted();
 			}
@@ -566,6 +572,11 @@ void FInEditorCapture::Shutdown()
 	{
 		CapturingFromWorld->GetWorldSettings()->DefaultGameMode = CachedGameMode;
 	}
+	
+	if (CachedEngineShowFlags)
+	{
+		CachedEngineShowFlags->SetPathTracing(CachedPathTracingMode);
+	}
 
 	FObjectReader(GetMutableDefault<ULevelEditorPlaySettings>(), BackedUpPlaySettings);
 
diff --git a/Engine/Source/Editor/MovieSceneCaptureDialog/Public/MovieSceneCaptureDialogModule.h b/Engine/Source/Editor/MovieSceneCaptureDialog/Public/MovieSceneCaptureDialogModule.h
--- a/Engine/Source/Editor/MovieSceneCaptureDialog/Public/MovieSceneCaptureDialogModule.h
+++ b/Engine/Source/Editor/MovieSceneCaptureDialog/Public/MovieSceneCaptureDialogModule.h
@@ -130,6 +130,8 @@ private:
 	int32 BackedUpUseFixedPoolSize;
 	TArray<uint8> BackedUpPlaySettings;
 
+	bool CachedPathTracingMode = false;
+	FEngineShowFlags* CachedEngineShowFlags = nullptr;
 	TSubclassOf<AGameModeBase> CachedGameMode;
 };
 
diff --git a/Engine/Source/Editor/MovieSceneTools/Private/AutomatedLevelSequenceCapture.cpp b/Engine/Source/Editor/MovieSceneTools/Private/AutomatedLevelSequenceCapture.cpp
--- a/Engine/Source/Editor/MovieSceneTools/Private/AutomatedLevelSequenceCapture.cpp
+++ b/Engine/Source/Editor/MovieSceneTools/Private/AutomatedLevelSequenceCapture.cpp
@@ -88,6 +88,7 @@ UAutomatedLevelSequenceCapture::UAutomatedLevelSequenceCapture(const FObjectInit
 	WarmUpFrameCount = 0;
 	DelayBeforeWarmUp = 0.0f;
 	DelayBeforeShotWarmUp = 0.0f;
+	DelayEveryFrame = 0.0f;
 	bWriteEditDecisionList = true;
 	bWriteFinalCutProXML = true;
 
@@ -158,6 +159,17 @@ void UAutomatedLevelSequenceCapture::Initialize(TSharedPtr<FSceneViewport> InVie
 		{
 			DelayBeforeShotWarmUp = DelayBeforeShotWarmUpOverride;
 		}
+
+		float DelayEveryFrameOverride;
+		if (FParse::Value(FCommandLine::Get(), TEXT("-MovieDelayEveryFrame="), DelayEveryFrameOverride))
+		{
+			DelayEveryFrame = DelayEveryFrameOverride;
+		}
+	}
+
+	if (Settings.bUsePathTracer)
+	{
+		DelayEveryFrame = float(Settings.FrameRate.AsSeconds(Settings.PathTracerSamplePerPixel));
 	}
 
 	ALevelSequenceActor* Actor = LevelSequenceActor.Get();
@@ -542,11 +554,11 @@ void UAutomatedLevelSequenceCapture::OnTick(float DeltaSeconds)
 
 		StartWarmup();
 
-		if (DelayBeforeWarmUp + DelayBeforeShotWarmUp > 0)
+		if (DelayBeforeWarmUp + DelayBeforeShotWarmUp + DelayEveryFrame > 0)
 		{
 			CaptureState = ELevelSequenceCaptureState::DelayBeforeWarmUp;
 
-			Actor->GetWorld()->GetTimerManager().SetTimer(DelayTimer, FTimerDelegate::CreateUObject(this, &UAutomatedLevelSequenceCapture::DelayBeforeWarmupFinished), DelayBeforeWarmUp + DelayBeforeShotWarmUp, false);
+			Actor->GetWorld()->GetTimerManager().SetTimer(DelayTimer, FTimerDelegate::CreateUObject(this, &UAutomatedLevelSequenceCapture::DelayBeforeWarmupFinished), DelayBeforeWarmUp + DelayBeforeShotWarmUp + DelayEveryFrame, false);
 		}
 		else
 		{
@@ -677,10 +689,14 @@ void UAutomatedLevelSequenceCapture::SequenceUpdated(const UMovieSceneSequencePl
 		if (Actor && Actor->SequencePlayer)
 		{
 			// If this is a new shot, set the state to shot warm up and pause on this frame until warmed up			
-			bool bHasMultipleShots = PreviousState.CurrentShotName != PreviousState.MasterName;
-			bool bNewShot = bHasMultipleShots && PreviousState.ShotID != CachedState.ShotID;
-			
-			if (bNewShot && Actor->SequencePlayer->IsPlaying() && DelayBeforeShotWarmUp > 0)
+			const bool bHasMultipleShots = PreviousState.CurrentShotName != PreviousState.MasterName;
+			const bool bNewShot = bHasMultipleShots && PreviousState.ShotID != CachedState.ShotID;
+			const bool bNewFrame = PreviousTime != CurrentTime;
+
+			const bool bDelayingBeforeShotWarmUp = (bNewShot && DelayBeforeShotWarmUp > 0);
+			const bool bDelayingEveryFrame = (bNewFrame && DelayEveryFrame > 0);
+
+			if (Actor->SequencePlayer->IsPlaying() && ( bDelayingBeforeShotWarmUp || bDelayingEveryFrame ))
 			{
 				if (bIsAudioCapturePass)
 				{
@@ -701,7 +717,7 @@ void UAutomatedLevelSequenceCapture::SequenceUpdated(const UMovieSceneSequencePl
 				
 				CaptureState = ELevelSequenceCaptureState::Paused;
 
-				Actor->GetWorld()->GetTimerManager().SetTimer(DelayTimer, FTimerDelegate::CreateUObject(this, &UAutomatedLevelSequenceCapture::PauseFinished), DelayBeforeShotWarmUp, false);
+				Actor->GetWorld()->GetTimerManager().SetTimer(DelayTimer, FTimerDelegate::CreateUObject(this, &UAutomatedLevelSequenceCapture::PauseFinished), DelayBeforeShotWarmUp + DelayEveryFrame, false);
 				CachedPlayRate = Actor->SequencePlayer->GetPlayRate();
 				Actor->SequencePlayer->SetPlayRate(0.f);
 			}
diff --git a/Engine/Source/Editor/MovieSceneTools/Public/AutomatedLevelSequenceCapture.h b/Engine/Source/Editor/MovieSceneTools/Public/AutomatedLevelSequenceCapture.h
--- a/Engine/Source/Editor/MovieSceneTools/Public/AutomatedLevelSequenceCapture.h
+++ b/Engine/Source/Editor/MovieSceneTools/Public/AutomatedLevelSequenceCapture.h
@@ -56,7 +56,7 @@ public:
 	float DelayBeforeWarmUp;
 
 	/** The number of seconds to wait (in real-time) at shot boundaries.  Useful for allowing post processing effects to settle down before capturing the animation. */
-	UPROPERTY(config, EditAnywhere, BlueprintReadWrite, Category=Animation, AdvancedDisplay, meta=(Units=Seconds, ClampMin=0))
+	UPROPERTY(config, EditAnywhere, BlueprintReadWrite, Category = Animation, AdvancedDisplay, meta = (Units = Seconds, ClampMin = 0))
 	float DelayBeforeShotWarmUp;
 
 	UPROPERTY(Transient, EditAnywhere, BlueprintReadWrite, Category=CaptureSettings, AdvancedDisplay, meta=(EditInline))
@@ -167,6 +167,9 @@ private:
 	/** The current shot movie that is rendering */
 	int32 ShotIndex;
 
+	/** The number of seconds to wait (in real-time) at every frame.  Useful for allowing post processing effects to settle down before capturing the animation. */
+	float DelayEveryFrame;
+
 	FLevelSequencePlayerSnapshot CachedState;
 
 	TOptional<float> CachedPlayRate;
diff --git a/Engine/Source/Runtime/Engine/Private/GameViewportClient.cpp b/Engine/Source/Runtime/Engine/Private/GameViewportClient.cpp
--- a/Engine/Source/Runtime/Engine/Private/GameViewportClient.cpp
+++ b/Engine/Source/Runtime/Engine/Private/GameViewportClient.cpp
@@ -1104,6 +1104,14 @@ void UGameViewportClient::Draw(FViewport* InViewport, FCanvas* SceneCanvas)
 
 	UWorld* MyWorld = GetWorld();
 
+	// Force path tracing view mode, and extern code set path tracer show flags
+	const bool bForcePathTracing = InViewport->GetClient()->GetEngineShowFlags()->PathTracing;
+	if (bForcePathTracing)
+	{
+		EngineShowFlags.SetPathTracing(true);
+		ViewModeIndex = VMI_PathTracing;
+	}
+
 	// create the view family for rendering the world scene to the viewport's render target
 	FSceneViewFamilyContext ViewFamily(FSceneViewFamily::ConstructionValues( 	
 		InViewport,
@@ -1299,6 +1307,9 @@ void UGameViewportClient::Draw(FViewport* InViewport, FCanvas* SceneCanvas)
 							}
 						}
 
+					#if RHI_RAYTRACING
+						View->SetupRayTracedRendering();
+					#endif
 					}
 
 					// Add view information for resource streaming. Allow up to 5X boost for small FOV.
diff --git a/Engine/Source/Runtime/MovieSceneCapture/Private/MovieSceneCapture.cpp b/Engine/Source/Runtime/MovieSceneCapture/Private/MovieSceneCapture.cpp
--- a/Engine/Source/Runtime/MovieSceneCapture/Private/MovieSceneCapture.cpp
+++ b/Engine/Source/Runtime/MovieSceneCapture/Private/MovieSceneCapture.cpp
@@ -95,6 +95,8 @@ FMovieSceneCaptureSettings::FMovieSceneCaptureSettings()
 	bAllowTurning = false;
 	bShowPlayer = false;
 	bShowHUD = false;
+	bUsePathTracer = false;
+	PathTracerSamplePerPixel = 16;
 
 #if PLATFORM_MAC
 	MovieExtension = TEXT(".mov");
@@ -292,6 +294,17 @@ void UMovieSceneCapture::Initialize(TSharedPtr<FSceneViewport> InSceneViewport,
 			Settings.bCinematicMode = bOverrideCinematicMode;
 		}
 
+		bool bOverridePathTracer;
+		if (FParse::Bool(FCommandLine::Get(), TEXT("-PathTracer="), bOverridePathTracer))
+		{
+			Settings.bUsePathTracer = bOverridePathTracer;
+		}
+
+		uint16 OverridePathTracerSamplePerPixel;
+		if (FParse::Value(FCommandLine::Get(), TEXT("-PathTracerSamplePerPixel="), OverridePathTracerSamplePerPixel))
+		{
+			Settings.PathTracerSamplePerPixel = OverridePathTracerSamplePerPixel;
+		}
 
 		bool bProtocolOverride = false;
 
@@ -372,6 +385,11 @@ void UMovieSceneCapture::Initialize(TSharedPtr<FSceneViewport> InSceneViewport,
 		}
 	}
 
+	if (!IsRayTracingEnabled())
+	{
+		Settings.bUsePathTracer = false;
+	}
+
 	bFinalizeWhenReady = false;
 	bIsAudioCapturePass = false;
 
@@ -514,8 +532,9 @@ void UMovieSceneCapture::CaptureThisFrame(float DeltaSeconds)
 		}
 		UE_LOG(LogMovieSceneCapture, Verbose, TEXT("Captured frame: %d"), CachedMetrics.Frame);
 		++CachedMetrics.Frame;
-		}
+	}
 }
+
 void UMovieSceneCapture::Tick(float DeltaSeconds)
 {
 	if (ImageCaptureProtocol)
diff --git a/Engine/Source/Runtime/MovieSceneCapture/Public/MovieSceneCaptureSettings.h b/Engine/Source/Runtime/MovieSceneCapture/Public/MovieSceneCaptureSettings.h
--- a/Engine/Source/Runtime/MovieSceneCapture/Public/MovieSceneCaptureSettings.h
+++ b/Engine/Source/Runtime/MovieSceneCapture/Public/MovieSceneCaptureSettings.h
@@ -110,4 +110,12 @@ struct MOVIESCENECAPTURE_API FMovieSceneCaptureSettings
 	/** Whether to show the in-game HUD whilst capturing */
 	UPROPERTY(config, EditAnywhere, BlueprintReadWrite, Category=Cinematic, AdvancedDisplay, meta=(EditCondition="bCinematicMode"))
 	bool bShowHUD;
+
+	/** Whether to use the path tracer (if supported) to render the scene */
+	UPROPERTY(config, EditAnywhere, BlueprintReadWrite, Category = Cinematic, AdvancedDisplay)
+	bool bUsePathTracer;
+
+	/** Number of sampler per pixel to be used when rendering the scene with the path tracer (if supported) */
+	UPROPERTY(config, EditAnywhere, BlueprintReadWrite, Category = Cinematic, AdvancedDisplay, meta = (EditCondition = "bUsePathTracer", ClampMin = 1, UIMin = 1, UIMax = 4096))
+	int32 PathTracerSamplePerPixel;
 };
