This is a significant overhaul to the Curve Editor used by Sequencer which adds a plugin-based architecture and extensibility. New tools and toolbar buttons can be added to all usages of the curve edi

diff --git a/Engine/Source/Editor/CurveEditor/Public/Filters/CurveEditorBakeFilter.h b/Engine/Source/Editor/CurveEditor/Public/Filters/CurveEditorBakeFilter.h
new file mode 100644
--- /dev/null
+++ b/Engine/Source/Editor/CurveEditor/Public/Filters/CurveEditorBakeFilter.h
@@ -0,0 +1,30 @@
+// Copyright 1998-2018 Epic Games, Inc. All Rights Reserved.
+
+#pragma once
+
+#include "Filters/CurveEditorFilterBase.h"
+#include "CurveEditorTypes.h"
+#include "CurveEditorBakeFilter.generated.h"
+
+UCLASS(DisplayName="Bake")
+class CURVEEDITOR_API UCurveEditorBakeFilter : public UCurveEditorFilterBase
+{
+	GENERATED_BODY()
+public:
+	UCurveEditorBakeFilter()
+	{
+		bUseSnapRateForInterval = true;
+		BakeInterval = 0.1f;
+	}
+protected:
+	virtual void ApplyFilter_Impl(TSharedRef<FCurveEditor> InCurveEditor, const TMap<FCurveModelID, FKeyHandleSet>& InKeysToOperateOn, TMap<FCurveModelID, FKeyHandleSet>& OutKeysToSelect) override;
+
+public:
+	/** If true we will use the snap rate of the Curve Editor to determine how far apart keys should be when baking. If false, the interval is used. */
+	UPROPERTY(EditAnywhere, BlueprintReadWrite, Category="Settings")
+	bool bUseSnapRateForInterval;
+
+	/** The interval (in seconds) between baked keys. Only used if bUseSnapRateForInterval is false. */
+	UPROPERTY(EditAnywhere, BlueprintReadWrite, meta = (EditCondition = "!bUseSnapRateForInterval", UIMin = 0, UIMax=1), Category = "Settings")
+	float BakeInterval;
+};
\ No newline at end of file
diff --git a/Engine/Source/Editor/CurveEditor/Public/Filters/SCurveEditorFilterPanel.cpp b/Engine/Source/Editor/CurveEditor/Public/Filters/SCurveEditorFilterPanel.cpp
new file mode 100644
--- /dev/null
+++ b/Engine/Source/Editor/CurveEditor/Public/Filters/SCurveEditorFilterPanel.cpp
@@ -0,0 +1,244 @@
+// Copyright 1998-2018 Epic Games, Inc. All Rights Reserved.
+
+#include "SCurveEditorFilterPanel.h"
+#include "Widgets/Input/SButton.h"
+#include "Widgets/Text/STextBlock.h"
+#include "Widgets/Docking/SDockTab.h"
+#include "Widgets/SWindow.h"
+#include "Widgets/SBoxPanel.h"
+#include "Modules/ModuleManager.h"
+#include "Widgets/Layout/SBox.h"
+#include "Framework/Application/SlateApplication.h"
+#include "Framework/Docking/TabManager.h"
+#include "CurveEditorFilterBase.h"
+#include "PropertyEditorModule.h"
+#include "ClassViewerFilter.h"
+#include "ClassViewerModule.h"
+#include "CurveEditor.h"
+
+TWeakPtr<SWindow> SCurveEditorFilterPanel::ExistingFilterWindow;
+
+/** ClassViewerFilter for Animation Modifier classes */
+class FCurveFilterClassFilter : public IClassViewerFilter
+{
+public:
+	bool IsClassAllowed(const FClassViewerInitializationOptions& InInitOptions, const UClass* InClass, TSharedRef<FClassViewerFilterFuncs> InFilterFuncs) override
+	{
+		return InClass->IsChildOf(UCurveEditorFilterBase::StaticClass()) && InClass != UCurveEditorFilterBase::StaticClass();
+	}
+
+	virtual bool IsUnloadedClassAllowed(const FClassViewerInitializationOptions& InInitOptions, const TSharedRef<const IUnloadedBlueprintData> InClass, TSharedRef<FClassViewerFilterFuncs> InFilterFuncs) override
+	{
+		return InClass->IsChildOf(UCurveEditorFilterBase::StaticClass());
+	}
+};
+
+void SCurveEditorFilterPanel::Construct(const FArguments& InArgs, TSharedRef<FCurveEditor> InCurveEditor)
+{
+	WeakCurveEditor = InCurveEditor;
+
+	FClassViewerInitializationOptions Options;
+	Options.bShowUnloadedBlueprints = true;
+	Options.bShowNoneOption = false;
+	// Options.Mode = EClassViewerMode::ClassBrowsing;
+	Options.bAllowViewOptions = false;
+	Options.NameTypeToDisplay = EClassViewerNameTypeToDisplay::DisplayName;
+
+	TSharedPtr<FCurveFilterClassFilter> ClassFilter = MakeShared<FCurveFilterClassFilter>();
+	Options.ClassFilter = ClassFilter;
+
+	FOnClassPicked OnPicked(FOnClassPicked::CreateSP(this, &SCurveEditorFilterPanel::SetFilterClass));
+
+
+	FPropertyEditorModule& PropertyEditor = FModuleManager::LoadModuleChecked<FPropertyEditorModule>("PropertyEditor");
+	FClassViewerModule& ClassViewerModule = FModuleManager::LoadModuleChecked<FClassViewerModule>("ClassViewer");
+	
+	// Configure the Details View
+	FDetailsViewArgs DetailsViewArgs;
+	DetailsViewArgs.bUpdatesFromSelection = false;
+	DetailsViewArgs.bLockable = false;
+	DetailsViewArgs.ViewIdentifier = "CurveEditorFilterPanel";
+	DetailsViewArgs.DefaultsOnlyVisibility = EEditDefaultsOnlyNodeVisibility::Automatic;
+	DetailsViewArgs.bAllowSearch = false;
+	DetailsViewArgs.bShowPropertyMatrixButton = false;
+	DetailsViewArgs.bShowOptions = false;
+
+	// Generate it and store a reference so we can update it with the right object later.
+	DetailView = PropertyEditor.CreateDetailView(DetailsViewArgs);
+	TSharedRef<SWidget> ClassFilterWidget = ClassViewerModule.CreateClassViewer(Options, OnPicked);
+
+	ChildSlot
+	[
+		SNew(SVerticalBox)
+
+		// Class Picker to choose which class shows up in the Details Panel
+		+ SVerticalBox::Slot()
+		.AutoHeight()
+		[
+			ClassFilterWidget
+		]
+		// Details Panel
+		+ SVerticalBox::Slot()
+		.FillHeight(1.f)
+		[
+			DetailView.ToSharedRef()
+		]
+		
+		// Footer Row
+		+ SVerticalBox::Slot()
+		.AutoHeight()
+		.HAlign(HAlign_Fill)
+		[			
+			SNew(SHorizontalBox)
+			+ SHorizontalBox::Slot()
+			.Padding(4, 0, 0, 0)
+			.FillWidth(1.0f)
+			[
+				// Unfortunately the Class Filter can't have a class selected by default so 
+				// when the panel is opened you don't know what filter has been pre-selected.
+				// We're solving this for now by putting a textbox that shows your current filter
+				// by name.
+				SNew(STextBlock)
+				.Text(this, &SCurveEditorFilterPanel::GetCurrentFilterText)
+				.TextStyle(FEditorStyle::Get(), "LargeText")
+			]
+
+			+ SHorizontalBox::Slot()
+			.AutoWidth()
+			[
+				// Apply Button
+				SNew(SButton)
+				.Text(NSLOCTEXT("CurveEditorFilterPanel", "ApplyFilter", "Apply"))
+				.OnClicked(this, &SCurveEditorFilterPanel::OnApplyClicked)
+				.IsEnabled(this, &SCurveEditorFilterPanel::CanApplyFilter)
+			]
+		]
+	];
+}
+
+void SCurveEditorFilterPanel::SetFilterClass(UClass* InClass)
+{
+	UObject* ClassCDO = InClass->GetDefaultObject<UObject>();
+	DetailView->SetObject(ClassCDO);
+}
+
+FReply SCurveEditorFilterPanel::OnApplyClicked()
+{
+	TSharedPtr<FCurveEditor> CurveEditor = WeakCurveEditor.Pin();
+	if (CurveEditor && DetailView->GetSelectedObjects().Num() > 0)
+	{
+		const TMap<FCurveModelID, FKeyHandleSet>& SelectedKeys = CurveEditor->GetSelection().GetAll();
+		
+		FText TransactionText = FText::Format(NSLOCTEXT("CurveEditorFilterApply", "Filter Curves", "Filtered {0}|plural(one=Curve, other=Curves)"), CurveEditor->GetSelection().Count());
+		FScopedTransaction Transaction(TransactionText);
+
+		UCurveEditorFilterBase* Filter = Cast<UCurveEditorFilterBase>(DetailView->GetSelectedObjects()[0]);
+		TMap<FCurveModelID, FKeyHandleSet> OutKeysToSelect;
+		Filter->ApplyFilter(CurveEditor.ToSharedRef(), SelectedKeys, OutKeysToSelect);
+
+		// Clear their selection and then set it to the keys the filter thinks you should have selected.
+		CurveEditor->GetSelection().Clear();
+
+		for (const TTuple<FCurveModelID, FKeyHandleSet>& OutSet : OutKeysToSelect)
+		{
+			CurveEditor->GetSelection().Add(OutSet.Key, ECurvePointType::Key, OutSet.Value.AsArray());
+		}
+	}
+
+	return FReply::Handled();
+}
+
+bool SCurveEditorFilterPanel::CanApplyFilter() const
+{
+	bool bHasSelection = false;
+	bool bHasFilter = false;
+
+	TSharedPtr<FCurveEditor> CurveEditor = WeakCurveEditor.Pin();
+	if (CurveEditor.IsValid())
+	{
+		bHasSelection = CurveEditor->GetSelection().Count() > 0;
+	}
+
+	bHasFilter = DetailView->GetSelectedObjects().Num() > 0;
+
+	return bHasSelection && bHasFilter;
+}
+
+
+void SCurveEditorFilterPanel::OpenDialog(const TSharedRef<FTabManager>& TabManager, TSharedRef<FCurveEditor> InHostCurveEditor, TSubclassOf<UCurveEditorFilterBase> DefaultFilterClass)
+{
+	TSharedPtr<SWindow> ExistingWindow = ExistingFilterWindow.Pin();
+	if (ExistingWindow.IsValid())
+	{
+		ExistingWindow->BringToFront();
+	}
+	else
+	{
+		ExistingWindow = SNew(SWindow)
+			.Title(NSLOCTEXT("CurveEditorFilterPanel", "WindowTitle", "Curve Editor Filters"))
+			.HasCloseButton(true)
+			.SupportsMaximize(false)
+			.SupportsMinimize(false)
+			.ClientSize(FVector2D(480, 360));
+
+		TSharedPtr<SDockTab> OwnerTab = TabManager->GetOwnerTab();
+		TSharedPtr<SWindow> RootWindow = OwnerTab.IsValid() ? OwnerTab->GetParentWindow() : TSharedPtr<SWindow>();
+		if (RootWindow.IsValid())
+		{
+			FSlateApplication::Get().AddWindowAsNativeChild(ExistingWindow.ToSharedRef(), RootWindow.ToSharedRef());
+		}
+		else
+		{
+			FSlateApplication::Get().AddWindow(ExistingWindow.ToSharedRef());
+		}
+	}
+
+	TSharedPtr<SCurveEditorFilterPanel> FilterPanel;
+	ExistingWindow->SetContent(
+		SAssignNew(FilterPanel, SCurveEditorFilterPanel, InHostCurveEditor)
+	);
+
+	FilterPanel->SetFilterClass(DefaultFilterClass);
+	ExistingFilterWindow = ExistingWindow;
+}
+
+void SCurveEditorFilterPanel::CloseDialog()
+{
+	TSharedPtr<SWindow> ExistingWindow = ExistingFilterWindow.Pin();
+	if (ExistingWindow.IsValid())
+	{
+		ExistingWindow->RequestDestroyWindow();
+		ExistingWindow = nullptr;
+	}
+}
+
+FText SCurveEditorFilterPanel::GetCurrentFilterText() const
+{
+	FText CurrentFilterName = NSLOCTEXT("SCurveEditorFilterPanel", "NoFilterSelectedName", "None");
+	if (DetailView->GetSelectedObjects().Num() > 0)
+	{
+		UCurveEditorFilterBase* CurrentFilter = Cast<UCurveEditorFilterBase>(DetailView->GetSelectedObjects()[0].Get());
+		if (CurrentFilter)
+		{
+			CurrentFilterName = CurrentFilter->GetClass()->GetDisplayNameText();
+		}
+	}
+
+	TSharedPtr<FCurveEditor> CurveEditor = WeakCurveEditor.Pin();
+	if (CurveEditor.IsValid())
+	{
+		int32 SelectedKeysCount = CurveEditor->GetSelection().Count();
+		if (SelectedKeysCount > 0)
+		{
+			const FText FormatString = NSLOCTEXT("SCurveEditorFilterPanel", "CurrentFilterWithKeysFormat", "Current Filter: {0} ({1} Keys)");
+			return FText::Format(FormatString, CurrentFilterName, SelectedKeysCount);
+		}
+		else
+		{
+			const FText FormatString = NSLOCTEXT("SCurveEditorFilterPanel", "CurrentFilterNoKeysFormat", "Current Filter: {0} (No Keys Selected)");
+			return FText::Format(FormatString, CurrentFilterName);
+		}
+	}
+
+	return FText();
+}
diff --git a/Engine/Source/Editor/CurveEditor/Public/Filters/SCurveEditorFilterPanel.h b/Engine/Source/Editor/CurveEditor/Public/Filters/SCurveEditorFilterPanel.h
new file mode 100644
--- /dev/null
+++ b/Engine/Source/Editor/CurveEditor/Public/Filters/SCurveEditorFilterPanel.h
@@ -0,0 +1,46 @@
+// Copyright 1998-2018 Epic Games, Inc. All Rights Reserved.
+
+#pragma once
+
+#include "Widgets/DeclarativeSyntaxSupport.h"
+#include "Widgets/SCompoundWidget.h"
+#include "Templates/SharedPointer.h"
+#include "Templates/SubclassOf.h"
+
+class FCurveEditor;
+class IDetailsView;
+class UCurveEditorFilterBase;
+class SWindow;
+class FTabManager;
+
+class CURVEEDITOR_API SCurveEditorFilterPanel : public SCompoundWidget
+{
+public:
+	SLATE_BEGIN_ARGS(SCurveEditorFilterPanel)
+	{}
+
+	SLATE_END_ARGS()
+
+	void Construct(const FArguments& InArgs, TSharedRef<FCurveEditor> InCurveEditor);
+	void SetFilterClass(UClass* InClass);
+
+public:
+	/** Call this to request opening a window containing this panel. */
+	static void OpenDialog(const TSharedRef<FTabManager>& TabManager, TSharedRef<FCurveEditor> InHostCurveEditor, TSubclassOf<UCurveEditorFilterBase> DefaultFilterClass);
+	
+	/** Closes the dialog if there is one open. */
+	static void CloseDialog();
+protected:
+	FReply OnApplyClicked();
+	bool CanApplyFilter() const;
+	FText GetCurrentFilterText() const;
+private:
+	/** Weak pointer to the curve editor which created this filter. */
+	TWeakPtr<FCurveEditor> WeakCurveEditor;
+
+	/** The Details View in our UI that we update each time they choose a class. */
+	TSharedPtr<IDetailsView> DetailView;
+
+	/** Singleton for the pop-up window. */
+	static TWeakPtr<SWindow> ExistingFilterWindow;
+};
\ No newline at end of file
diff --git a/Engine/Source/Editor/CurveEditor/Public/ICurveEditorBounds.h b/Engine/Source/Editor/CurveEditor/Public/ICurveEditorBounds.h
--- a/Engine/Source/Editor/CurveEditor/Public/ICurveEditorBounds.h
+++ b/Engine/Source/Editor/CurveEditor/Public/ICurveEditorBounds.h
@@ -9,26 +9,14 @@ struct ICurveEditorBounds
 	virtual ~ICurveEditorBounds(){}
 
 	virtual void GetInputBounds(double& OutMin, double& OutMax) const = 0;
-	virtual void GetOutputBounds(double& OutMin, double& OutMax) const = 0;
-
 	virtual void SetInputBounds(double InMin, double InMax) = 0;
-	virtual void SetOutputBounds(double InMin, double InMax) = 0;
 };
 
 
 struct FStaticCurveEditorBounds : ICurveEditorBounds
 {
-	double InputMin, InputMax;
-	double OutputMin, OutputMax;
-
-	FStaticCurveEditorBounds()
-	{
-		InputMin = 0;
-		InputMax = 1;
-
-		OutputMin = 0;
-		OutputMax = 1;
-	}
+	double InputMin = 0.0;
+	double InputMax = 1.0;
 
 	virtual void GetInputBounds(double& OutMin, double& OutMax) const override final
 	{
@@ -36,21 +24,9 @@ struct FStaticCurveEditorBounds : ICurveEditorBounds
 		OutMax = InputMax;
 	}
 
-	virtual void GetOutputBounds(double& OutMin, double& OutMax) const override final
-	{
-		OutMin = OutputMin;
-		OutMax = OutputMax;
-	}
-
 	virtual void SetInputBounds(double InMin, double InMax) override final
 	{
 		InputMin = InMin;
 		InputMax = InMax;
 	}
-
-	virtual void SetOutputBounds(double InMin, double InMax) override final
-	{
-		OutputMin = InMin;
-		OutputMax = InMax;
-	}
 };
\ No newline at end of file
diff --git a/Engine/Source/Editor/CurveEditor/Public/ICurveEditorDragOperation.h b/Engine/Source/Editor/CurveEditor/Public/ICurveEditorDragOperation.h
--- a/Engine/Source/Editor/CurveEditor/Public/ICurveEditorDragOperation.h
+++ b/Engine/Source/Editor/CurveEditor/Public/ICurveEditorDragOperation.h
@@ -18,7 +18,6 @@ class ICurveEditorDragOperation
 public:
 
 	ICurveEditorDragOperation()
-		: MouseLockVector(FVector2D::UnitVector)
 	{}
 
 	virtual ~ICurveEditorDragOperation() {}
@@ -39,9 +38,9 @@ public:
 	void EndDrag(FVector2D InitialPosition, FVector2D CurrentPosition, const FPointerEvent& MouseEvent);
 
 	/**
-	 * Paint this drag operation
+	 * Paint this drag operation onto the specified layer
 	 */
-	int32 Paint(const FGeometry& AllottedGeometry, FSlateWindowElementList& OutDrawElements, int32 LayerId);
+	void Paint(const FGeometry& AllottedGeometry, FSlateWindowElementList& OutDrawElements, int32 PaintOnLayerId);
 
 	/**
 	 * Cancel this drag operation
@@ -63,24 +62,13 @@ protected:
 	{}
 
 	/** Implementation method for derived types to paint this drag */
-	virtual int32 OnPaint(const FGeometry& AllottedGeometry, FSlateWindowElementList& OutDrawElements, int32 LayerId)
-	{
-		return LayerId;
-	}
+	virtual void OnPaint(const FGeometry& AllottedGeometry, FSlateWindowElementList& OutDrawElements, int32 PaintOnLayerId)
+	{}
 
 	/** Implementation method for derived types to cancel a drag */
 	virtual void OnCancelDrag()
 	{}
-
-	/**
-	 * Determine the effective mouse position for a drag vector, potentially locked to an axis based on the current pointer event
-	 */
-	FVector2D GetLockedMousePosition(FVector2D InitialPosition, FVector2D CurrentPosition, const FPointerEvent& MouseEvent);
-
 protected:
-
-	/** The vector we're currently locked to, or (1.f, 1.f) if we're not locked */
-	FVector2D MouseLockVector;
 };
 
 /**
@@ -134,5 +122,7 @@ struct FCurveEditorDelayedDrag : FDelayedDrag
 	 */
 	FCurveEditorDelayedDrag(FVector2D InInitialPosition, FKey InEffectiveKey)
 		: FDelayedDrag(InInitialPosition, InEffectiveKey)
-	{}
+	{
+		SetTriggerScaleFactor(0.1f);
+	}
 };
diff --git a/Engine/Source/Editor/CurveEditor/Public/RichCurveEditorModel.h b/Engine/Source/Editor/CurveEditor/Public/RichCurveEditorModel.h
new file mode 100644
--- /dev/null
+++ b/Engine/Source/Editor/CurveEditor/Public/RichCurveEditorModel.h
@@ -0,0 +1,48 @@
+// Copyright 1998-2018 Epic Games, Inc. All Rights Reserved.
+
+#pragma once
+
+#include "CoreTypes.h"
+#include "UObject/WeakObjectPtr.h"
+#include "UObject/WeakObjectPtrTemplates.h"
+#include "CurveModel.h"
+
+struct FRichCurve;
+
+class CURVEEDITOR_API FRichCurveEditorModel : public FCurveModel
+{
+public:
+	FRichCurveEditorModel(FRichCurve* InRichCurve, UObject* InOwner);
+
+	virtual const void* GetCurve() const override;
+
+	virtual void Modify() override;
+
+	virtual void DrawCurve(const FCurveEditor& CurveEditor, const FCurveEditorScreenSpace& ScreenSpace, TArray<TTuple<double, double>>& InterpolatingPoints) const override;
+	virtual void GetKeys(const FCurveEditor& CurveEditor, double MinTime, double MaxTime, double MinValue, double MaxValue, TArray<FKeyHandle>& OutKeyHandles) const override;
+	virtual void GetKeyDrawInfo(ECurvePointType PointType, const FKeyHandle InKeyHandle, FKeyDrawInfo& OutDrawInfo) const override;
+
+	virtual void GetKeyPositions(TArrayView<const FKeyHandle> InKeys, TArrayView<FKeyPosition> OutKeyPositions) const override;
+	virtual void SetKeyPositions(TArrayView<const FKeyHandle> InKeys, TArrayView<const FKeyPosition> InKeyPositions) override;
+
+	virtual void GetKeyAttributes(TArrayView<const FKeyHandle> InKeys, TArrayView<FKeyAttributes> OutAttributes) const override;
+	virtual void SetKeyAttributes(TArrayView<const FKeyHandle> InKeys, TArrayView<const FKeyAttributes> InAttributes) override;
+
+	virtual void GetCurveAttributes(FCurveAttributes& OutCurveAttributes) const override;
+	virtual void SetCurveAttributes(const FCurveAttributes& InCurveAttributes) override;
+	virtual void GetTimeRange(double& MinTime, double& MaxTime) const override;
+	virtual void GetValueRange(double& MinValue, double& MaxValue) const override;
+	virtual int32 GetNumKeys() const override;
+	virtual void GetNeighboringKeys(const FKeyHandle InKeyHandle, TOptional<FKeyHandle>& OutPreviousKeyHandle, TOptional<FKeyHandle>& OutNextKeyHandle) const override;
+	virtual bool Evaluate(double ProspectiveTime, double& OutValue) const override;
+	virtual void AddKeys(TArrayView<const FKeyPosition> InKeyPositions, TArrayView<const FKeyAttributes> InAttributes, TArrayView<TOptional<FKeyHandle>>* OutKeyHandles) override;
+	virtual void RemoveKeys(TArrayView<const FKeyHandle> InKeys) override;
+
+	virtual void CreateKeyProxies(TArrayView<const FKeyHandle> InKeyHandles, TArrayView<UObject*> OutObjects) override;
+
+private:
+
+	FRichCurve* RichCurve;
+
+	TWeakObjectPtr<> WeakOwner;
+};
\ No newline at end of file
diff --git a/Engine/Source/Editor/CurveEditor/Public/SCurveEditorPanel.h b/Engine/Source/Editor/CurveEditor/Public/SCurveEditorPanel.h
--- a/Engine/Source/Editor/CurveEditor/Public/SCurveEditorPanel.h
+++ b/Engine/Source/Editor/CurveEditor/Public/SCurveEditorPanel.h
@@ -5,15 +5,21 @@
 #include "Widgets/DeclarativeSyntaxSupport.h"
 #include "Widgets/SCompoundWidget.h"
 
-#include "CurveEditorScreenSpace.h"
-#include "ICurveEditorDragOperation.h"
 #include "CurveDrawInfo.h"
 #include "CurveEditor.h"
 #include "CurveEditorTypes.h"
+#include "Textures/SlateIcon.h"
 
 struct FCurveEditorEditObjectContainer;
 
 class IDetailsView;
+class IGraphEditorView;
+class ITimeSliderController;
+class SScrollBox;
+class SCurveEditorView;
+class SCurveEditorViewContainer;
+class FTabManager;
+class UCurveEditorFilterBase;
 
 /**
  * Curve editor widget that reflects the state of an FCurveEditor
@@ -27,6 +33,18 @@ class CURVEEDITOR_API SCurveEditorPanel : public SCompoundWidget
 		/** Color to draw grid lines */
 		SLATE_ATTRIBUTE(FLinearColor, GridLineTint)
 
+		/** Tab Manager which owns this panel. */
+		SLATE_ARGUMENT(TSharedPtr<FTabManager>, TabManager)
+
+		/** Optional Time Slider Controller which allows us to synchronize with an externally controlled Time Slider */
+		SLATE_ARGUMENT(TSharedPtr<ITimeSliderController>, ExternalTimeSliderController)
+
+		/** If specified, causes the time snap adjustment UI controls to be disabled and specifies the tooltip to be used. Can be used to disable time snap controls when externally controlled. */
+		SLATE_ATTRIBUTE(FText, DisabledTimeSnapTooltip)
+
+		/** Widget slot for the tree content */
+		SLATE_NAMED_SLOT(FArguments, TreeContent)
+
 	SLATE_END_ARGS()
 
 	SCurveEditorPanel();
@@ -37,14 +55,6 @@ class CURVEEDITOR_API SCurveEditorPanel : public SCompoundWidget
 	 */
 	void Construct(const FArguments& InArgs, TSharedRef<FCurveEditor> InCurveEditor);
 
-	/**
-	 * Access the draw parameters that this curve editor has cached for this frame
-	 */
-	const TArray<FCurveDrawParams>& GetCachedDrawParams() const
-	{
-		return CachedDrawParams;
-	}
-
 	/**
 	 * Access the combined command list for this curve editor and panel widget
 	 */
@@ -56,51 +66,93 @@ class CURVEEDITOR_API SCurveEditorPanel : public SCompoundWidget
 	/**
 	 * Access the details view used for editing selected keys
 	 */
-	TSharedPtr<IDetailsView> GetKeyDetailsView() const
+	TSharedPtr<class SCurveKeyDetailPanel> GetKeyDetailsView() const
 	{
 		return KeyDetailsView;
 	}
 
-private:
+	void AddView(TSharedRef<SCurveEditorView> ViewToAdd);
 
-	virtual int32 OnPaint(const FPaintArgs& Args, const FGeometry& AllottedGeometry, const FSlateRect& MyCullingRect, FSlateWindowElementList& OutDrawElements, int32 LayerId, const FWidgetStyle& InWidgetStyle, bool bParentEnabled) const override;
-	virtual void Tick(const FGeometry& AllottedGeometry, const double InCurrentTime, const float InDeltaTime) override;
+	void RemoveView(TSharedRef<SCurveEditorView> ViewToRemove);
 
-	/*~ Mouse interaction */
-	virtual FReply OnMouseButtonDown(const FGeometry& MyGeometry, const FPointerEvent& MouseEvent) override;
-	virtual FReply OnMouseButtonUp(const FGeometry& MyGeometry, const FPointerEvent& MouseEvent) override;
-	virtual FReply OnMouseMove(const FGeometry& MyGeometry, const FPointerEvent& MouseEvent) override;
-	virtual FReply OnMouseWheel(const FGeometry& MyGeometry, const FPointerEvent& MouseEvent) override;
-	virtual void OnFocusLost(const FFocusEvent& InFocusEvent) override;
-	virtual void OnMouseEnter(const FGeometry& MyGeometry, const FPointerEvent& MouseEvent) override;
-	virtual void OnMouseLeave(const FPointerEvent& MouseEvent) override;
+	/** This returns an extender which is pre-configured with the standard set of Toolbar Icons. Implementers of SCurveEditorPanel should use this
+	* to generate the icons and then add any additional context-specific icons (such as save buttons in the Asset Editor) to ensure that the Curve
+	* Editor has a consistent set (and order) of icons across all usages.
+	*/
+	TSharedPtr<FExtender> GetToolbarExtender();
 
-	/*~ Keyboard interaction */
-	virtual bool SupportsKeyboardFocus() const override { return true; }
-	virtual FReply OnKeyDown(const FGeometry& MyGeometry, const FKeyEvent& InKeyEvent) override;
+	/** Access the cached geometry of the outer scroll panel that contains this panel's views */
+	const FGeometry& GetScrollPanelGeometry() const;
 
-private:
+	/** Access the cached geometry of container housing all this panel's views */
+	const FGeometry& GetViewContainerGeometry() const;
+
+	/** Get all the views stored in this panel. */
+	TArrayView<const TSharedPtr<SCurveEditorView>> GetViews() const;
+
+	/** Get the grid line tint to be used for views on panel */
+	FLinearColor GetGridLineTint() const { return GridLineTintAttribute.Get(); }
+
+	/** Scroll this panel's view scroll box vertically by the specified amount */
+	void ScrollBy(float Amount);
 
 	/**
-	 * Get the visibility for the value splitter control
+	 * Find all the views that the specified curve is being displayed on
+	 * @note: Returns an in-place iterator to this curve's view mapping. Adding or removing curves from views *will* invalidate this iterator.
+	 *
+	 * @param InCurveID The identifier of the curve to find views for
+	 * @return An iterator to all the views that this cuvrve is displayed within.
 	 */
-	EVisibility GetSplitterVisibility() const;
+	TMultiMap<FCurveModelID, TSharedRef<SCurveEditorView>>::TConstKeyIterator FindViews(FCurveModelID InCurveID)
+	{
+		return CurveViews.CreateConstKeyIterator(InCurveID);
+	}
 
 	/**
-	 * Get the currently hovered curve ID
+	 * Remove the specified curve from all views it is currently displayed on.
 	 */
-	TOptional<FCurveModelID> GetHoveredCurve() const;
+	void RemoveCurveFromViews(FCurveModelID InCurveID);
+
+	/** Get the last set View Mode for this UI. Utility function for the UI. */
+	ECurveEditorViewID GetViewMode() const { return DefaultViewID; }
+
+private:
+	// SWidget Interface
+	virtual void Tick(const FGeometry& AllottedGeometry, const double InCurrentTime, const float InDeltaTime) override;
+	// ~SWidget Interface
+
+	/*~ Keyboard interaction */
+	virtual bool SupportsKeyboardFocus() const override { return true; }
+	virtual FReply OnKeyDown(const FGeometry& MyGeometry, const FKeyEvent& InKeyEvent) override;
+
+	TSharedRef<SWidget> MakeTimeSnapMenu();
+	FText GetTimeSnapMenuTooltip() const;
+	TSharedRef<SWidget> MakeValueSnapMenu();
+	TSharedRef<SWidget> MakeAxisSnapMenu();
+	TSharedRef<SWidget> MakeViewModeMenu();
+
+	bool IsInlineEditPanelEditable() const;
+	EVisibility ShouldInstructionOverlayBeVisible() const;
+
+private:
 
 	/**
-	 * Hit test for a point on the curve editor using a mouse position in slate units
+	 * Get the visibility for the value splitter control
 	 */
-	TOptional<FCurvePointHandle> HitPoint(FVector2D MousePixel) const;
+	EVisibility GetSplitterVisibility() const;
 
 	/*~ Event bindings */
 	void UpdateEditBox();
 	void UpdateCommonCurveInfo();
-	void UpdateToolTip(const FGeometry& InMyGeometry, const FPointerEvent& InMouseEvent);
-	void UpdateCurveProximities(FVector2D MousePixel);
+
+	/** Creates the drop-down list you see when changing Curve View options. */
+	TSharedRef<SWidget> MakeCurveEditorCurveViewOptionsMenu();
+	TSharedRef<SWidget> MakeCurveExtrapolationMenu(const bool bInPostExtrapolation);
+	FSlateIcon GetCurveExtrapolationPreIcon() const;
+	FSlateIcon GetCurveExtrapolationPostIcon() const;
+
+	/** Creates the Curve Editor Filter UI and pre-populates it with the specified class. */
+	void ShowCurveFilterUI(TSubclassOf<UCurveEditorFilterBase> FilterClass);
 
 private:
 
@@ -109,16 +161,6 @@ private:
 	 */
 	void BindCommands();
 
-	/**
-	 * Rebind contextual command mappings that rely on the mouse position
-	 */
-	void RebindContextualActions(FVector2D MousePosition);
-
-	/*~ Command binding callbacks */
-	void OnAddKey(FVector2D MousePixel);
-	void OnAddKeyHere(FVector2D MousePixel);
-	void OnAddKey(FCurveModelID CurveToAdd, FVector2D MousePixel);
-
 	/**
 	 * Assign new attributes to the currently selected keys
 	 */
@@ -170,51 +212,42 @@ private:
 	bool CanToggleWeightedTangents() const;
 
 	/**
-	*  Brings up Dialog to specify Tolerance before doing the key reduction.
-	*/
-	void OnSimplifySelection();
-
-	/**
-	*  Callback for when the curve is simplified.
-	*/
-	void OnSimplifySelectionCommited(const FText& InText, ETextCommit::Type CommitInfo);
-
-	/**
-	*  Create popup menu
-	*/
-	void GenericTextEntryModeless(const FText& DialogText, const FText& DefaultText, FOnTextCommitted OnTextComitted);
+	 * Check whether or not we can set a key interpolation on the current selection. If no keys are selected, you can't set an interpolation!
+	 */
+	bool CanSetKeyInterpolation() const;
 
+	/** Sets the axis snapping to the specified value. Only supports X, Y and None. */
+	void SetAxisSnapping(EAxisList::Type InAxis);
 
-	/**
-	*  Close popup menu
-	*/
-	void CloseEntryPopupMenu();
-public:
+	/** Get a reference to the curve editor this panel represents. */
+	TSharedPtr<FCurveEditor> GetCurveEditor() const { return CurveEditor; }
 
-	/**
-	 * Draw grid lines
-	 */
-	int32 DrawGridLines(const FGeometry& AllottedGeometry, FSlateWindowElementList& OutDrawElements, int32 LayerId, ESlateDrawEffect DrawEffects) const;
+private:
 
 	/**
-	 * Draw curve data
+	 * Create a new view for a model of the specified type, and add the curve to the view
+	 *
+	 * @param CurveModelID The ID of the curve we're creating a view for
+	 * @param ViewTypeID   The ID of the view we'd like to create
+	 * @param bPinned      Whether the view should be pinned or not
+	 * @return A new view or nullptr if one could not be created
 	 */
-	int32 DrawCurves(const FGeometry& AllottedGeometry, FSlateWindowElementList& OutDrawElements, int32 LayerId, const FWidgetStyle& InWidgetStyle, ESlateDrawEffect DrawEffects) const;
+	TSharedPtr<SCurveEditorView> CreateViewOfType(FCurveModelID CurveModelID, ECurveEditorViewID ViewTypeID, bool bPinned);
 
 private:
 
 	/** The curve editor pointer */
 	TSharedPtr<FCurveEditor> CurveEditor;
 
-	/** Array of curve proximities in slate units that's updated on mouse move */
-	TArray<TTuple<FCurveModelID, float>> CurveProximities;
+	/** Map from curve model ID to the views that it is on */
+	TMultiMap<FCurveModelID, TSharedRef<SCurveEditorView>> CurveViews;
+
+	/** Set of externally added views */
+	TSet<TSharedRef<SCurveEditorView>> ExternalViews;
 
 	/** (Optional) the current drag operation */
 	TOptional<FCurveEditorDelayedDrag> DragOperation;
 
-	/** Curve draw parameters that are re-generated on tick */
-	TArray<FCurveDrawParams> CachedDrawParams;
-
 	/** Cached curve attributes that are common to all visible curves */
 	FCurveAttributes CachedCommonCurveAttributes;
 
@@ -227,8 +260,11 @@ private:
 	/** Attribute used for retrieving the desired grid line color */
 	TAttribute<FLinearColor> GridLineTintAttribute;
 
+	/** Attribute used for retrieving the tooltip for when the Time Snap control is disabled. Specifying this causes the Time Snap Adjustment to be disabled. */
+	TAttribute<FText> DisabledTimeSnapTooltipAttribute;
+
 	/** Edit panel */
-	TSharedPtr<IDetailsView> KeyDetailsView;
+	TSharedPtr<SCurveKeyDetailPanel> KeyDetailsView;
 
 	/** Map of edit UI widgets for each curve in the current selection set */
 	TMap<FCurveModelID, TSharedPtr<SWidget>> CurveToEditUI;
@@ -241,28 +277,39 @@ private:
 
 private:
 
-	/*~ Cached tooltip data */
-	bool  IsToolTipEnabled() const;
-	FText GetToolTipCurveName() const;
-	FText GetToolTipTimeText() const;
-	FText GetToolTipValueText() const;
+	/** Sets the View Mode for the UI to the specified mode. This will destroy and re-create all views, but leave additional pinned views unmodified. */
+	void SetViewMode(const ECurveEditorViewID NewViewMode);
 
-	struct FCachedToolTipData
-	{
-		FCachedToolTipData() {}
+	/** Compare if our current view mode matches the specified one. Utility function for the UI. */
+	bool CompareViewMode(const ECurveEditorViewID InViewMode) const;
 
-		FText Text;
-		FText EvaluatedValue;
-		FText EvaluatedTime;
-	};
-	TOptional<FCachedToolTipData> CachedToolTipData;
+	/** Rebuild the Curve Views layout to match the currently specified View Mode. */
+	void RebuildCurveViews();
 
-	/** The tolerance to use when reducing curves */
-	float ReduceTolerance;
+	/** The last set View Mode for this UI. */
+	ECurveEditorViewID DefaultViewID;
 
-	/** Generic Popup Entry */
-	TWeakPtr<IMenu> EntryPopupMenu;
+	TMultiMap<ECurveEditorViewID, TSharedRef<SCurveEditorView>> FreeViewsByType;
+
+	/** The scrool box that all curve views live inside */
+	TSharedPtr<SScrollBox> ScrollBox;
+
+	/** The container for all our vertically laid out curve views. */
+	TSharedPtr<SCurveEditorViewContainer> CurveViewsContainer;
+
+private:
+
+	/** Whether to explicitly refresh the views for this panel */
+	bool bNeedsRefresh;
+
+	/** Serial number cached from FCurveEditor::GetActiveCurvesSerialNumber() on tick */
+	uint32 CachedActiveCurvesSerialNumber;
+
+	/** A copy of the View Geometry used to represent the View portion of the Curve Editor. */
+	FGeometry CachedViewGeometry;
 
 	/** Container of objects that are being used to edit keys on the curve editor */
 	TUniquePtr<FCurveEditorEditObjectContainer> EditObjects;
+
+	TWeakPtr<FTabManager> WeakTabManager;
 };
\ No newline at end of file
diff --git a/Engine/Source/Editor/CurveEditor/Public/SCurveKeyDetailPanel.h b/Engine/Source/Editor/CurveEditor/Public/SCurveKeyDetailPanel.h
new file mode 100644
--- /dev/null
+++ b/Engine/Source/Editor/CurveEditor/Public/SCurveKeyDetailPanel.h
@@ -0,0 +1,32 @@
+// Copyright 1998-2018 Epic Games, Inc. All Rights Reserved.
+
+#pragma once
+
+#include "Widgets/DeclarativeSyntaxSupport.h"
+#include "Widgets/SCompoundWidget.h"
+#include "CurveEditor.h"
+
+// Forward Declare
+class IPropertyRowGenerator;
+
+/**
+ * Inline details panel that lets you edit the Time and Value of a generic FCurveEditor Key
+ */
+class CURVEEDITOR_API SCurveKeyDetailPanel : public SCompoundWidget
+{
+public:
+	SLATE_BEGIN_ARGS(SCurveKeyDetailPanel)
+	{}
+
+	SLATE_END_ARGS()
+
+	void Construct(const FArguments& InArgs, TSharedRef<FCurveEditor> InCurveEditor);
+	TSharedPtr<IPropertyRowGenerator> GetPropertyRowGenerator() const { return PropertyRowGenerator; }
+
+private:
+	void PropertyRowsRefreshed();
+	void ConstructChildLayout(TSharedPtr<SWidget> TimeWidget, TSharedPtr<SWidget> ValueWidget);
+
+private:
+	TSharedPtr<IPropertyRowGenerator> PropertyRowGenerator;
+};
\ No newline at end of file
diff --git a/Engine/Source/Editor/CurveEditor/Public/Tree/ICurveEditorTreeItem.h b/Engine/Source/Editor/CurveEditor/Public/Tree/ICurveEditorTreeItem.h
new file mode 100644
--- /dev/null
+++ b/Engine/Source/Editor/CurveEditor/Public/Tree/ICurveEditorTreeItem.h
@@ -0,0 +1,31 @@
+
+#pragma once
+
+#include "Templates/SharedPointer.h"
+#include "Containers/Array.h"
+#include "Templates/UniquePtr.h"
+
+struct FCurveEditorTreeItemID;
+
+class FName;
+class SWidget;
+class FCurveModel;
+class FCurveEditor;
+
+struct CURVEEDITOR_API ICurveEditorTreeItem
+{
+	virtual ~ICurveEditorTreeItem() {}
+	virtual TSharedPtr<SWidget> GenerateCurveEditorTreeWidget(const FName& InColumnName, TWeakPtr<FCurveEditor> InCurveEditor, FCurveEditorTreeItemID InTreeItemID) = 0;
+
+	virtual void CreateCurveModels(TArray<TUniquePtr<FCurveModel>>& OutCurveModels) = 0;
+
+	struct FColumnNames
+	{
+		FName Label;
+		FName PinHeader;
+
+		CURVEEDITOR_API FColumnNames();
+	};
+
+	static const FColumnNames ColumnNames;
+};
\ No newline at end of file
diff --git a/Engine/Source/Editor/CurveEditor/Public/Tree/SCurveEditorTree.h b/Engine/Source/Editor/CurveEditor/Public/Tree/SCurveEditorTree.h
new file mode 100644
--- /dev/null
+++ b/Engine/Source/Editor/CurveEditor/Public/Tree/SCurveEditorTree.h
@@ -0,0 +1,44 @@
+// Copyright 1998-2019 Epic Games, Inc. All Rights Reserved.
+
+#pragma once
+
+#include "Widgets/DeclarativeSyntaxSupport.h"
+#include "Widgets/SCompoundWidget.h"
+#include "CurveEditorTypes.h"
+
+class FCurveEditor;
+class ITableRow;
+class SHeaderRow;
+class STableViewBase;
+class SCurveEditorTreeView;
+
+class CURVEEDITOR_API SCurveEditorTree : public SCompoundWidget
+{
+public:
+
+	SLATE_BEGIN_ARGS(SCurveEditorTree){}
+	SLATE_END_ARGS()
+
+	void Construct(const FArguments& InArgs, TSharedPtr<FCurveEditor> InCurveEditor);
+
+private:
+
+	virtual FReply OnKeyDown(const FGeometry& MyGeometry, const FKeyEvent& InKeyEvent) override;
+
+private:
+
+	TSharedRef<ITableRow> GenerateRow(FCurveEditorTreeItemID ItemID, const TSharedRef<STableViewBase>& OwnerTable);
+
+	void GetTreeItemChildren(FCurveEditorTreeItemID Parent, TArray<FCurveEditorTreeItemID>& OutChildren);
+
+	void OnTreeSelectionChanged(FCurveEditorTreeItemID, ESelectInfo::Type);
+
+	void SetItemExpansionRecursive(FCurveEditorTreeItemID Model, bool bInExpansionState);
+
+private:
+
+	TSharedPtr<FCurveEditor> CurveEditor;
+
+	TSharedPtr<SHeaderRow> HeaderRow;
+	TSharedPtr<SCurveEditorTreeView> TreeView;
+};
\ No newline at end of file
diff --git a/Engine/Source/Editor/CurveEditor/Public/Tree/SCurveEditorTreePin.h b/Engine/Source/Editor/CurveEditor/Public/Tree/SCurveEditorTreePin.h
new file mode 100644
--- /dev/null
+++ b/Engine/Source/Editor/CurveEditor/Public/Tree/SCurveEditorTreePin.h
@@ -0,0 +1,37 @@
+// Copyright 1998-2019 Epic Games, Inc. All Rights Reserved.
+
+#pragma once
+
+#include "Widgets/DeclarativeSyntaxSupport.h"
+#include "Widgets/SCompoundWidget.h"
+
+#include "CurveEditorTypes.h"
+
+class FCurveEditor;
+
+class CURVEEDITOR_API SCurveEditorTreePin : public SCompoundWidget
+{
+public:
+	SLATE_BEGIN_ARGS(SCurveEditorTreePin){}
+	SLATE_END_ARGS()
+
+	void Construct(const FArguments& InArgs, TWeakPtr<FCurveEditor> InCurveEditor, FCurveEditorTreeItemID InTreeItemID);
+
+private:
+
+	FReply TogglePinned();
+	const FSlateBrush* GetPinBrush() const;
+
+	bool IsPinnedRecursive(FCurveEditorTreeItemID InTreeItem, FCurveEditor* CurveEditor) const;
+
+	void PinRecursive(FCurveEditorTreeItemID InTreeItem, FCurveEditor* CurveEditor) const;
+
+	void UnpinRecursive(FCurveEditorTreeItemID InTreeItem, FCurveEditor* CurveEditor) const;
+
+	EVisibility GetPinVisibility() const;
+
+private:
+
+	TWeakPtr<FCurveEditor> WeakCurveEditor;
+	FCurveEditorTreeItemID TreeItemID;
+};
\ No newline at end of file
diff --git a/Engine/Source/Editor/CurveEditor/Public/Views/SInteractiveCurveEditorView.h b/Engine/Source/Editor/CurveEditor/Public/Views/SInteractiveCurveEditorView.h
new file mode 100644
--- /dev/null
+++ b/Engine/Source/Editor/CurveEditor/Public/Views/SInteractiveCurveEditorView.h
@@ -0,0 +1,177 @@
+// Copyright 1998-2018 Epic Games, Inc. All Rights Reserved.
+
+#pragma once
+
+#include "Widgets/DeclarativeSyntaxSupport.h"
+#include "Widgets/SCompoundWidget.h"
+#include "CurveEditorTypes.h"
+#include "SCurveEditorView.h"
+#include "ICurveEditorDragOperation.h"
+#include "ICurveEditorToolExtension.h"
+#include "CurveDrawInfo.h"
+
+struct FCurveModelID;
+struct FCurveEditorScreenSpace;
+struct FOptionalSize;
+struct FCurveEditorDelayedDrag;
+class IMenu;
+
+namespace CurveViewConstants
+{
+	/** The default offset from the top-right corner of curve views for curve labels to be drawn. */
+	constexpr float CurveLabelOffsetX = 15.f;
+	constexpr float CurveLabelOffsetY = 10.f;
+
+	/**
+	 * Pre-defined layer offsets for specific curve view elements. Fixed values are used to decouple draw order and layering
+	 * Some elements deliberately leave some spare layers as a buffer for slight tweaks to layering within that element
+	 */
+	namespace ELayerOffset
+	{
+		enum
+		{
+			Background     = 0,
+			GridLines      = 1,
+			GridOverlays   = 2,
+			GridLabels     = 3,
+			Curves         = 10,
+			HoveredCurves  = 15,
+			Keys           = 20,
+			SelectedKeys   = 30,
+			Tools          = 35,
+			DragOperations = 40,
+			Labels         = 45,
+			WidgetContent  = 50,
+			Last = Labels
+		};
+	}
+}
+
+/**
+ */
+class CURVEEDITOR_API SInteractiveCurveEditorView : public SCurveEditorView
+{
+public:
+
+	SLATE_BEGIN_ARGS(SInteractiveCurveEditorView)
+		: _BackgroundTint(FLinearColor::White)
+		, _MaximumCapacity(0)
+		, _AutoSize(true)
+	{}
+
+		SLATE_ARGUMENT(FLinearColor, BackgroundTint)
+
+		SLATE_ARGUMENT(int32, MaximumCapacity)
+
+		SLATE_ATTRIBUTE(float, FixedHeight)
+
+		SLATE_ARGUMENT(bool, AutoSize)
+
+	SLATE_END_ARGS()
+
+	void Construct(const FArguments& InArgs, TWeakPtr<FCurveEditor> InCurveEditor);
+
+protected:
+
+	// ~SCurveEditorView Interface
+	virtual void GetPointsWithinWidgetRange(const FSlateRect& WidgetRectangle, TArray<FCurvePointHandle>* OutPoints) const override;
+
+	virtual void GetGridLinesX(TSharedRef<FCurveEditor> CurveEditor, TArray<float>& MajorGridLines, TArray<float>& MinorGridLines, TArray<FText>& MajorGridLabels) const;
+	virtual void GetGridLinesY(TSharedRef<FCurveEditor> CurveEditor, TArray<float>& MajorGridLines, TArray<float>& MinorGridLines, TArray<FText>& MajorGridLabels) const;
+
+	virtual void PaintView(const FPaintArgs& Args, const FGeometry& AllottedGeometry, const FSlateRect& MyCullingRect, FSlateWindowElementList& OutDrawElements, int32 BaseLayerId, const FWidgetStyle& InWidgetStyle, bool bParentEnabled) const;
+
+protected:
+
+	// SWidget Interface
+	virtual int32 OnPaint(const FPaintArgs& Args, const FGeometry& AllottedGeometry, const FSlateRect& MyCullingRect, FSlateWindowElementList& OutDrawElements, int32 BaseLayerId, const FWidgetStyle& InWidgetStyle, bool bParentEnabled) const override;
+	virtual void Tick(const FGeometry& AllottedGeometry, const double InCurrentTime, const float InDeltaTime) override;
+	virtual void OnMouseEnter(const FGeometry& MyGeometry, const FPointerEvent& MouseEvent) override;
+	virtual void OnMouseLeave(const FPointerEvent& MouseEvent) override;
+	virtual FReply OnMouseMove(const FGeometry& MyGeometry, const FPointerEvent& MouseEvent) override;
+	virtual FReply OnMouseWheel(const FGeometry& MyGeometry, const FPointerEvent& MouseEvent) override;
+	virtual FReply OnMouseButtonUp(const FGeometry& MyGeometry, const FPointerEvent& MouseEvent) override;
+	virtual FReply OnMouseButtonDown(const FGeometry& MyGeometry, const FPointerEvent& MouseEvent) override;
+	// ~SWidget Interface
+
+protected:
+
+	void DrawBackground(const FGeometry& AllottedGeometry, FSlateWindowElementList& OutDrawElements, int32 BaseLayerId, ESlateDrawEffect DrawEffects) const;
+	void DrawGridLines(TSharedRef<FCurveEditor> CurveEditor, const FGeometry& AllottedGeometry, FSlateWindowElementList& OutDrawElements, int32 BaseLayerId, ESlateDrawEffect DrawEffects) const;
+	void DrawCurves(TSharedRef<FCurveEditor> CurveEditor, const FGeometry& AllottedGeometry, const FSlateRect& MyCullingRect, FSlateWindowElementList& OutDrawElements, int32 BaseLayerId, const FWidgetStyle& InWidgetStyle, ESlateDrawEffect DrawEffects) const;
+
+	FSlateColor GetCurveCaptionColor() const;
+	FText GetCurveCaption() const;
+
+private:
+
+	/** Gets info about the curves being drawn. Converts actual curves into an abstract series of lines/points/handles/etc. */
+	void GetCurveDrawParams(TArray<FCurveDrawParams>& OutDrawParams) const;
+
+	void CreateContextMenu(const FGeometry& MyGeometry, const FPointerEvent& MouseEvent);
+
+	/** Updates our distance to all of the curves we represent. */
+	void UpdateCurveProximities(FVector2D MousePixel);
+
+	TOptional<FCurveModelID> GetHoveredCurve() const;
+	TOptional<FCurvePointHandle> HitPoint(FVector2D MousePixel) const;
+
+	bool IsToolTipEnabled() const;
+	FText GetToolTipCurveName() const;
+	FText GetToolTipTimeText() const;
+	FText GetToolTipValueText() const;
+
+	/*~ Command binding callbacks */
+	void AddKeyAtScrubTime(TSet<FCurveModelID> ForCurves);
+	void AddKeyAtMousePosition(TSet<FCurveModelID> ForCurves);
+	void AddKeyAtTime(const TSet<FCurveModelID>& ToCurves, double InTime);
+
+	void OnCurveEditorToolChanged(FCurveEditorToolID InToolId);
+
+	/**
+	 * Rebind contextual command mappings that rely on the mouse position
+	 */
+	void RebindContextualActions(FVector2D InMousePosition);
+
+	/** Copy the curves from this view and set them as the Curve Editor's buffered curve support. */
+	void BufferVisibleCurves();
+	/** Copy the curves from this view and set them as the Curve Editor's buffered curve support. */
+	void BufferCurve(const FCurveModelID CurveID);
+	/** Attempt to apply the previously buffered curves to the currently visible curves. */
+	void ApplyBufferCurves(TOptional<FCurveModelID> DestinationCurve);
+	/** Check if it's legal to apply any of the buffered curves to our currently visible curves. */
+	bool CanApplyBufferedCurves(TOptional<FCurveModelID> DestinationCurve) const;
+
+protected:
+
+	/** Background tint for this widget */
+	FLinearColor BackgroundTint;
+
+private:
+
+	/** Curve draw parameters that are re-generated on tick. We generate them once and then they're used in multiple places per frame. */
+	TArray<FCurveDrawParams> CachedDrawParams;
+
+	/** (Optional) the current drag operation */
+	TOptional<FCurveEditorDelayedDrag> DragOperation;
+
+	struct FCachedToolTipData
+	{
+		FCachedToolTipData() {}
+
+		FText Text;
+		FText EvaluatedValue;
+		FText EvaluatedTime;
+	};
+
+	TOptional<FCachedToolTipData> CachedToolTipData;
+
+	/** Array of curve proximities in slate units that's updated on mouse move */
+	TArray<TTuple<FCurveModelID, float>> CurveProximities;
+
+	/** Track if we have a context menu active. Used to supress hover updates as it causes flickers in the CanExecute bindings. */
+	TWeakPtr<IMenu> ActiveContextMenu;
+
+	/** Cached location of the mouse relative to this widget each tick. This is so that command bindings related to the mouse cursor can create them at the right time. */
+	FVector2D CachedMousePosition;
+};
\ No newline at end of file
diff --git a/Engine/Source/Editor/EditorStyle/Private/SlateEditorStyle.cpp b/Engine/Source/Editor/EditorStyle/Private/SlateEditorStyle.cpp
--- a/Engine/Source/Editor/EditorStyle/Private/SlateEditorStyle.cpp
+++ b/Engine/Source/Editor/EditorStyle/Private/SlateEditorStyle.cpp
@@ -70,6 +70,7 @@ FSlateEditorStyle::FStyle::FStyle( const TWeakObjectPtr< UEditorStyleSettings >&
 	, SelectionColor_Subdued_LinearRef( MakeShareable( new FLinearColor( 0.807f, 0.596f, 0.388f ) ) )
 	, SelectionColor_Inactive_LinearRef( MakeShareable( new FLinearColor( 0.25f, 0.25f, 0.25f ) ) )
 	, SelectionColor_Pressed_LinearRef( MakeShareable( new FLinearColor( 0.701f, 0.225f, 0.003f ) ) )
+	, HighlightColor_LinearRef( MakeShareable( new FLinearColor(0.068f, 0.068f, 0.068f) ) ) 
 
 	, LogColor_Background_LinearRef(MakeShareable(new FLinearColor(FColor(0x040404FF))))
 	, LogColor_SelectionBackground_LinearRef(MakeShareable(new FLinearColor(FColor(0x020202FF))))
@@ -86,6 +87,7 @@ FSlateEditorStyle::FStyle::FStyle( const TWeakObjectPtr< UEditorStyleSettings >&
 	, SelectionColor_Subdued( SelectionColor_Subdued_LinearRef )
 	, SelectionColor_Inactive( SelectionColor_Inactive_LinearRef )
 	, SelectionColor_Pressed( SelectionColor_Pressed_LinearRef )
+	, HighlightColor( HighlightColor_LinearRef )
 
 	, LogColor_Background( LogColor_Background_LinearRef )
 	, LogColor_SelectionBackground( LogColor_SelectionBackground_LinearRef )
@@ -233,6 +235,8 @@ void FSlateEditorStyle::FStyle::SetupGeneralStyles()
 		.SetActiveHoveredBrush( IMAGE_BRUSH( "Common/Selection", Icon8x8, SelectionColor ) )
 		.SetInactiveBrush( IMAGE_BRUSH( "Common/Selection", Icon8x8, SelectionColor_Inactive ) )
 		.SetInactiveHoveredBrush( IMAGE_BRUSH( "Common/Selection", Icon8x8, SelectionColor_Inactive ) )
+		.SetActiveHighlightedBrush( IMAGE_BRUSH("Common/Selection", Icon8x8, HighlightColor ) )
+		.SetInactiveHighlightedBrush( IMAGE_BRUSH("Common/Selection", Icon8x8, HighlightColor ) )
 		.SetTextColor( DefaultForeground )
 		.SetSelectedTextColor(InvertedForeground)
 		.SetDropIndicator_Above(BOX_BRUSH("Common/DropZoneIndicator_Above", FMargin(10.0f / 16.0f, 10.0f / 16.0f, 0, 0), SelectionColor))
@@ -1917,8 +1921,8 @@ void FSlateEditorStyle::FStyle::SetupGeneralStyles()
 		Set( "Sequencer.SetKeyChanged.Small", new IMAGE_BRUSH( "Sequencer/Main_Icons/Icon_Sequencer_Key_Part_24x", Icon24x24 ) );
 		Set( "Sequencer.ToggleIsSnapEnabled", new IMAGE_BRUSH( "Sequencer/Main_Icons/Icon_Sequencer_Snap_24x", Icon48x48 ) );
 		Set( "Sequencer.ToggleIsSnapEnabled.Small", new IMAGE_BRUSH( "Sequencer/Main_Icons/Icon_Sequencer_Snap_24x", Icon24x24 ) );
-		Set( "Sequencer.ToggleShowCurveEditor", new IMAGE_BRUSH("Icons/SequencerIcons/icon_Sequencer_CurveEditor_24x", Icon48x48) );
-		Set( "Sequencer.ToggleShowCurveEditor.Small", new IMAGE_BRUSH("Icons/SequencerIcons/icon_Sequencer_CurveEditor_24x", Icon24x24) );
+		Set( "Sequencer.ToggleShowCurveEditor", new IMAGE_BRUSH("GenericCurveEditor/Icons/GenericCurveEditor_48x", Icon48x48) );
+		Set( "Sequencer.ToggleShowCurveEditor.Small", new IMAGE_BRUSH("GenericCurveEditor/Icons/GenericCurveEditor_48x", Icon24x24) );
 		Set( "Sequencer.ToggleAutoScroll", new IMAGE_BRUSH( "Icons/icon_Sequencer_ToggleAutoScroll_40x", Icon48x48 ) );
 		Set( "Sequencer.ToggleAutoScroll.Small", new IMAGE_BRUSH( "Icons/icon_Sequencer_ToggleAutoScroll_16x", Icon16x16 ) );
 		Set( "Sequencer.MoveTool.Small", new IMAGE_BRUSH( "Icons/SequencerIcons/icon_Sequencer_Move_24x", Icon16x16 ) );
@@ -7065,37 +7069,120 @@ void FSlateEditorStyle::FStyle::SetupToolkitStyles()
 		Set("CurveEditor.Gradient.HandleUp", new BOX_BRUSH("Sequencer/ScrubHandleUp", FMargin(6.f / 13.f, 8 / 12.f, 6 / 13.f, 5 / 12.f)));
 	}
 
-	// Curve Editor
+	// New Curve Editor 
+	{
+		// Tab
+		Set("GenericCurveEditor.TabIcon", new IMAGE_BRUSH("GenericCurveEditor/Icons/CurveEditor_32x", Icon16x16, FLinearColor(1.f, 1.f, 1.f, 0.8f)));
+
+		// Zoom / Framing
+		Set("GenericCurveEditor.ZoomToFit", new IMAGE_BRUSH("GenericCurveEditor/Icons/FramingSelected_48x", Icon48x48));
+		Set("GenericCurveEditor.ZoomToFit.Small", new IMAGE_BRUSH("GenericCurveEditor/Icons/FramingSelected_48x", Icon24x24));
+
+		// Time/Value Snapping
+		Set("GenericCurveEditor.ToggleInputSnapping", new IMAGE_BRUSH("GenericCurveEditor/Icons/AdjustKeySnapFrame_48x", Icon48x48));
+		Set("GenericCurveEditor.ToggleInputSnapping.Small", new IMAGE_BRUSH("GenericCurveEditor/Icons/AdjustKeySnapFrame_48x", Icon24x24));
+		Set("GenericCurveEditor.ToggleOutputSnapping", new IMAGE_BRUSH("GenericCurveEditor/Icons/AdjustKeySnapValue_48x", Icon48x48));
+		Set("GenericCurveEditor.ToggleOutputSnapping.Small", new IMAGE_BRUSH("GenericCurveEditor/Icons/AdjustKeySnapValue_48x", Icon24x24));
+
+		// Tangent Types
+		Set("GenericCurveEditor.InterpolationCubicAuto", new IMAGE_BRUSH("GenericCurveEditor/Icons/CubicTangentAuto_48x", Icon48x48));
+		Set("GenericCurveEditor.InterpolationCubicAuto.Small", new IMAGE_BRUSH("GenericCurveEditor/Icons/CubicTangentAuto_48x", Icon24x24));
+		Set("GenericCurveEditor.InterpolationCubicUser", new IMAGE_BRUSH("GenericCurveEditor/Icons/TangentsCubic_48x", Icon48x48));
+		Set("GenericCurveEditor.InterpolationCubicUser.Small", new IMAGE_BRUSH("GenericCurveEditor/Icons/TangentsCubic_48x", Icon24x24));
+		Set("GenericCurveEditor.InterpolationCubicBreak", new IMAGE_BRUSH("GenericCurveEditor/Icons/CubicTangentBroken_48x", Icon48x48));
+		Set("GenericCurveEditor.InterpolationCubicBreak.Small", new IMAGE_BRUSH("GenericCurveEditor/Icons/CubicTangentBroken_48x", Icon24x24));
+		Set("GenericCurveEditor.InterpolationToggleWeighted", new IMAGE_BRUSH("GenericCurveEditor/Icons/CubicTangentWeighted_48x", Icon48x48));
+		Set("GenericCurveEditor.InterpolationToggleWeighted.Small", new IMAGE_BRUSH("GenericCurveEditor/Icons/CubicTangentWeighted_48x", Icon24x24));
+		Set("GenericCurveEditor.InterpolationLinear", new IMAGE_BRUSH("GenericCurveEditor/Icons/TangentsLinear_48x", Icon48x48));
+		Set("GenericCurveEditor.InterpolationLinear.Small", new IMAGE_BRUSH("GenericCurveEditor/Icons/TangentsLinear_48x", Icon24x24));
+		Set("GenericCurveEditor.InterpolationConstant", new IMAGE_BRUSH("GenericCurveEditor/Icons/TangentsConstant_48x", Icon48x48));
+		Set("GenericCurveEditor.InterpolationConstant.Small", new IMAGE_BRUSH("GenericCurveEditor/Icons/TangentsConstant_48x", Icon24x24));
+
+		// Tangent Modifications
+		Set("GenericCurveEditor.FlattenTangents", new IMAGE_BRUSH("GenericCurveEditor/Icons/MiscFlatten_48x", Icon48x48));
+		Set("GenericCurveEditor.FlattenTangents.Small", new IMAGE_BRUSH("GenericCurveEditor/Icons/MiscFlatten_48x", Icon24x24));
+		Set("GenericCurveEditor.StraightenTangents", new IMAGE_BRUSH("GenericCurveEditor/Icons/MiscStraighten_48x", Icon48x48));
+		Set("GenericCurveEditor.StraightenTangents.Small", new IMAGE_BRUSH("GenericCurveEditor/Icons/MiscStraighten_48x", Icon24x24));
+
+		// View Modes
+		Set("GenericCurveEditor.SetViewModeAbsolute", new IMAGE_BRUSH("GenericCurveEditor/Icons/GraphViewAbsolute_48x", Icon48x48));
+		Set("GenericCurveEditor.SetViewModeAbsolute.Small", new IMAGE_BRUSH("GenericCurveEditor/Icons/GraphViewAbsolute_48x", Icon24x24));
+		Set("GenericCurveEditor.SetViewModeStacked", new IMAGE_BRUSH("GenericCurveEditor/Icons/GraphViewStack_48x", Icon48x48));
+		Set("GenericCurveEditor.SetViewModeStacked.Small", new IMAGE_BRUSH("GenericCurveEditor/Icons/GraphViewStack_48x", Icon24x24));
+		Set("GenericCurveEditor.SetViewModeNormalized", new IMAGE_BRUSH("GenericCurveEditor/Icons/GraphViewNormalized_48x", Icon48x48));
+		Set("GenericCurveEditor.SetViewModeNormalized.Small", new IMAGE_BRUSH("GenericCurveEditor/Icons/GraphViewNormalized_48x", Icon24x24));
+
+		// Axis Snapping
+		Set("GenericCurveEditor.SetAxisSnappingNone", new IMAGE_BRUSH("GenericCurveEditor/Icons/AdjustKeyMoveFree_48x", Icon48x48));
+		Set("GenericCurveEditor.SetAxisSnappingNone.Small", new IMAGE_BRUSH("GenericCurveEditor/Icons/AdjustKeyMoveFree_48x", Icon24x24));
+		Set("GenericCurveEditor.SetAxisSnappingHorizontal", new IMAGE_BRUSH("GenericCurveEditor/Icons/AdjustKeyLockValue_48x", Icon48x48));
+		Set("GenericCurveEditor.SetAxisSnappingHorizontal.Small", new IMAGE_BRUSH("GenericCurveEditor/Icons/AdjustKeyLockValue_48x", Icon24x24));
+		Set("GenericCurveEditor.SetAxisSnappingVertical", new IMAGE_BRUSH("GenericCurveEditor/Icons/AdjustKeyLockFrame_48x", Icon48x48));
+		Set("GenericCurveEditor.SetAxisSnappingVertical.Small", new IMAGE_BRUSH("GenericCurveEditor/Icons/AdjustKeyLockFrame_48x", Icon24x24));
+
+		// Deactivate Tool
+		Set("GenericCurveEditor.DeactivateCurrentTool", new IMAGE_BRUSH("GenericCurveEditor/Icons/GeneralPointer_48x", Icon48x48));
+		Set("GenericCurveEditor.DeactivateCurrentTool.Small", new IMAGE_BRUSH("GenericCurveEditor/Icons/GeneralPointer_48x", Icon24x24));
+
+		// Filters
+		Set("GenericCurveEditor.OpenUserImplementableFilterWindow", new IMAGE_BRUSH("GenericCurveEditor/Icons/MiscFilters_48x", Icon48x48));
+		Set("GenericCurveEditor.OpenUserImplementableFilterWindow.Small", new IMAGE_BRUSH("GenericCurveEditor/Icons/MiscFilters_48x", Icon24x24));
+
+		// Key Types
+		Set("GenericCurveEditor.Key", new IMAGE_BRUSH("GenericCurveEditor/Keys/Diamond_Filled", FVector2D(11.0f, 11.0f))); 
+		Set("GenericCurveEditor.ConstantKey", new IMAGE_BRUSH("GenericCurveEditor/Keys/Square_Filled", FVector2D(11.0f, 11.0f))); 
+		Set("GenericCurveEditor.LinearKey", new IMAGE_BRUSH("GenericCurveEditor/Keys/Triangle_Filled", FVector2D(11.0f, 11.0f)));
+		Set("GenericCurveEditor.CubicKey", new IMAGE_BRUSH("GenericCurveEditor/Keys/Diamond_Filled", FVector2D(11.0f, 11.0f)));
+		Set("GenericCurveEditor.TangentHandle", new IMAGE_BRUSH("GenericCurveEditor/Keys/TangentHandle", FVector2D(8.0f, 8.0f)));
+
+		// Pre-Infinity
+		Set("GenericCurveEditor.SetPreInfinityExtrapConstant", new IMAGE_BRUSH("GenericCurveEditor/Icons/PreInfinityConstant_48x", Icon48x48));
+		Set("GenericCurveEditor.SetPreInfinityExtrapConstant.Small", new IMAGE_BRUSH("GenericCurveEditor/Icons/PreInfinityConstant_48x", Icon24x24));
+		Set("GenericCurveEditor.SetPreInfinityExtrapCycle", new IMAGE_BRUSH("GenericCurveEditor/Icons/PreInfinityCycle_48x", Icon48x48));
+		Set("GenericCurveEditor.SetPreInfinityExtrapCycle.Small", new IMAGE_BRUSH("GenericCurveEditor/Icons/PreInfinityCycle_48x", Icon24x24));
+		Set("GenericCurveEditor.SetPreInfinityExtrapCycleWithOffset", new IMAGE_BRUSH("GenericCurveEditor/Icons/PreInfinityOffset_48x", Icon48x48));
+		Set("GenericCurveEditor.SetPreInfinityExtrapCycleWithOffset.Small", new IMAGE_BRUSH("GenericCurveEditor/Icons/PreInfinityOffset_48x", Icon24x24));
+		Set("GenericCurveEditor.SetPreInfinityExtrapLinear", new IMAGE_BRUSH("GenericCurveEditor/Icons/PreInfinityLinear_48x", Icon48x48));
+		Set("GenericCurveEditor.SetPreInfinityExtrapLinear.Small", new IMAGE_BRUSH("GenericCurveEditor/Icons/PreInfinityLinear_48x", Icon24x24));
+		Set("GenericCurveEditor.SetPreInfinityExtrapOscillate", new IMAGE_BRUSH("GenericCurveEditor/Icons/PreInfinityOscillate_48x", Icon48x48));
+		Set("GenericCurveEditor.SetPreInfinityExtrapOscillate.Small", new IMAGE_BRUSH("GenericCurveEditor/Icons/PreInfinityOscillate_48x", Icon24x24));
+		Set("GenericCurveEditor.PreInfinityMixed", new IMAGE_BRUSH("GenericCurveEditor/Icons/PreInfinityMultipleSelected_48x", Icon48x48));
+		Set("GenericCurveEditor.PreInfinityMixed.Small", new IMAGE_BRUSH("GenericCurveEditor/Icons/PreInfinityMultipleSelected_48x", Icon24x24));
+
+
+		// Post-Infinity
+		Set("GenericCurveEditor.SetPostInfinityExtrapConstant", new IMAGE_BRUSH("GenericCurveEditor/Icons/PostInfinityConstant_48x", Icon48x48));
+		Set("GenericCurveEditor.SetPostInfinityExtrapConstant.Small", new IMAGE_BRUSH("GenericCurveEditor/Icons/PostInfinityConstant_48x", Icon24x24));
+		Set("GenericCurveEditor.SetPostInfinityExtrapCycle", new IMAGE_BRUSH("GenericCurveEditor/Icons/PostInfinityCycle_48x", Icon48x48));
+		Set("GenericCurveEditor.SetPostInfinityExtrapCycle.Small", new IMAGE_BRUSH("GenericCurveEditor/Icons/PostInfinityCycle_48x", Icon24x24));
+		Set("GenericCurveEditor.SetPostInfinityExtrapCycleWithOffset", new IMAGE_BRUSH("GenericCurveEditor/Icons/PostInfinityOffset_48x", Icon48x48));
+		Set("GenericCurveEditor.SetPostInfinityExtrapCycleWithOffset.Small", new IMAGE_BRUSH("GenericCurveEditor/Icons/PostInfinityOffset_48x", Icon24x24));
+		Set("GenericCurveEditor.SetPostInfinityExtrapLinear", new IMAGE_BRUSH("GenericCurveEditor/Icons/PostInfinityLinear_48x", Icon48x48));
+		Set("GenericCurveEditor.SetPostInfinityExtrapLinear.Small", new IMAGE_BRUSH("GenericCurveEditor/Icons/PostInfinityLinear_48x", Icon24x24));
+		Set("GenericCurveEditor.SetPostInfinityExtrapOscillate", new IMAGE_BRUSH("GenericCurveEditor/Icons/PostInfinityOscillate_48x", Icon48x48));
+		Set("GenericCurveEditor.SetPostInfinityExtrapOscillate.Small", new IMAGE_BRUSH("GenericCurveEditor/Icons/PostInfinityOscillate_48x", Icon24x24));
+		Set("GenericCurveEditor.PostInfinityMixed", new IMAGE_BRUSH("GenericCurveEditor/Icons/PostInfinityMultipleSelected_48x", Icon48x48));
+		Set("GenericCurveEditor.PostInfinityMixed.Small", new IMAGE_BRUSH("GenericCurveEditor/Icons/PostInfinityMultipleSelected_48x", Icon24x24));
+
+		// Misc
+		Set("GenericCurveEditor.VisibilityOptions", new IMAGE_BRUSH("GenericCurveEditor/Icons/SettingsVisibility_48x", Icon48x48));
+		Set("GenericCurveEditor.VisibilityOptions.Small", new IMAGE_BRUSH("GenericCurveEditor/Icons/SettingsVisibility_48x", Icon24x24));
+
+		Set("GenericCurveEditor.Pin_Active", new IMAGE_BRUSH("Common/PushPin_Down", Icon16x16));
+		Set("GenericCurveEditor.Pin_Inactive", new IMAGE_BRUSH("Common/PushPin_Up", Icon16x16));
+	}
+
+	// Generic Curve Editor Tools
 	{
-		Set("GenericCurveEditor.ZoomToFitHorizontal", new IMAGE_BRUSH("Icons/icon_CurveEditor_Horizontal_40x", Icon40x40));
-		Set("GenericCurveEditor.ZoomToFitHorizontal.Small", new IMAGE_BRUSH("Icons/icon_CurveEditor_Horizontal_40x", Icon20x20));
-		Set("GenericCurveEditor.ZoomToFitVertical", new IMAGE_BRUSH("Icons/icon_CurveEditor_Vertical_40x", Icon40x40));
-		Set("GenericCurveEditor.ZoomToFitVertical.Small", new IMAGE_BRUSH("Icons/icon_CurveEditor_Vertical_40x", Icon20x20));
-		Set("GenericCurveEditor.ZoomToFit", new IMAGE_BRUSH("Icons/icon_CurveEditor_ZoomToFit_40x", Icon40x40));
-		Set("GenericCurveEditor.ZoomToFit.Small", new IMAGE_BRUSH("Icons/icon_CurveEditor_ZoomToFit_40x", Icon20x20));
-
-		Set("GenericCurveEditor.ToggleInputSnapping", new IMAGE_BRUSH("Icons/icon_CurveEditor_ToggleInputSnap_40x", Icon40x40));
-		Set("GenericCurveEditor.ToggleInputSnapping.Small", new IMAGE_BRUSH("Icons/icon_CurveEditor_ToggleInputSnap_40x", Icon20x20));
-		Set("GenericCurveEditor.ToggleOutputSnapping", new IMAGE_BRUSH("Icons/icon_CurveEditor_ToggleOutputSnap_40x", Icon40x40));
-		Set("GenericCurveEditor.ToggleOutputSnapping.Small", new IMAGE_BRUSH("Icons/icon_CurveEditor_ToggleOutputSnap_40x", Icon20x20));
-
-		Set("GenericCurveEditor.InterpolationCubicAuto", new IMAGE_BRUSH("Icons/icon_CurveEditor_Auto_40x", Icon40x40));
-		Set("GenericCurveEditor.InterpolationCubicAuto.Small", new IMAGE_BRUSH("Icons/icon_CurveEditor_Auto_40x", Icon20x20));
-		Set("GenericCurveEditor.InterpolationCubicUser", new IMAGE_BRUSH("Icons/icon_CurveEditor_User_40x", Icon40x40));
-		Set("GenericCurveEditor.InterpolationCubicUser.Small", new IMAGE_BRUSH("Icons/icon_CurveEditor_User_40x", Icon20x20));
-		Set("GenericCurveEditor.InterpolationCubicBreak", new IMAGE_BRUSH("Icons/icon_CurveEditor_Break_40x", Icon40x40));
-		Set("GenericCurveEditor.InterpolationCubicBreak.Small", new IMAGE_BRUSH("Icons/icon_CurveEditor_Break_40x", Icon20x20));
-		Set("GenericCurveEditor.InterpolationToggleWeighted", new IMAGE_BRUSH("Icons/icon_CurveEditor_WeightedTangents_40x", Icon40x40));
-		Set("GenericCurveEditor.InterpolationToggleWeighted.Small", new IMAGE_BRUSH("Icons/icon_CurveEditor_WeightedTangents_40x", Icon20x20));
-		Set("GenericCurveEditor.InterpolationLinear", new IMAGE_BRUSH("Icons/icon_CurveEditor_Linear_40x", Icon40x40));
-		Set("GenericCurveEditor.InterpolationLinear.Small", new IMAGE_BRUSH("Icons/icon_CurveEditor_Linear_40x", Icon20x20));
-		Set("GenericCurveEditor.InterpolationConstant", new IMAGE_BRUSH("Icons/icon_CurveEditor_Constant_40x", Icon40x40));
-		Set("GenericCurveEditor.InterpolationConstant.Small", new IMAGE_BRUSH("Icons/icon_CurveEditor_Constant_40x", Icon20x20));
-
-		Set("GenericCurveEditor.FlattenTangents", new IMAGE_BRUSH("Icons/icon_CurveEditor_Flatten_40x", Icon40x40));
-		Set("GenericCurveEditor.FlattenTangents.Small", new IMAGE_BRUSH("Icons/icon_CurveEditor_Flatten_40x", Icon20x20));
-		Set("GenericCurveEditor.StraightenTangents", new IMAGE_BRUSH("Icons/icon_CurveEditor_Straighten_40x", Icon40x40));
-		Set("GenericCurveEditor.StraightenTangents.Small", new IMAGE_BRUSH("Icons/icon_CurveEditor_Straighten_40x", Icon20x20));
+		Set("CurveEditorTools.SetFocusPlaybackTime", new IMAGE_BRUSH("GenericCurveEditor/Icons/FramingPlayback_48x", Icon48x48));
+		Set("CurveEditorTools.SetFocusPlaybackTime.Small", new IMAGE_BRUSH("GenericCurveEditor/Icons/FramingPlayback_48x", Icon24x24));
+		Set("CurveEditorTools.SetFocusPlaybackRange", new IMAGE_BRUSH("GenericCurveEditor/Icons/FramingTimeRange_48x", Icon48x48));
+		Set("CurveEditorTools.SetFocusPlaybackRange.Small", new IMAGE_BRUSH("GenericCurveEditor/Icons/FramingTimeRange_48x", Icon24x24));
+
+		Set("CurveEditorTools.ActivateTransformTool", new IMAGE_BRUSH("GenericCurveEditor/Icons/ToolsTransform_48x", Icon48x48));
+		Set("CurveEditorTools.ActivateTransformTool.Small", new IMAGE_BRUSH("GenericCurveEditor/Icons/ToolsTransform_48x", Icon24x24));
+		Set("CurveEditorTools.ActivateRetimeTool", new IMAGE_BRUSH("GenericCurveEditor/Icons/ToolsRetime_48x", Icon48x48));
+		Set("CurveEditorTools.ActivateRetimeTool.Small", new IMAGE_BRUSH("GenericCurveEditor/Icons/ToolsRetime_48x", Icon24x24));
 	}
 
 	// PhysicsAssetEditor
diff --git a/Engine/Source/Editor/EditorStyle/Private/SlateEditorStyle.h b/Engine/Source/Editor/EditorStyle/Private/SlateEditorStyle.h
--- a/Engine/Source/Editor/EditorStyle/Private/SlateEditorStyle.h
+++ b/Engine/Source/Editor/EditorStyle/Private/SlateEditorStyle.h
@@ -125,6 +125,7 @@ public:
 		const TSharedRef< FLinearColor > SelectionColor_Subdued_LinearRef;
 		const TSharedRef< FLinearColor > SelectionColor_Inactive_LinearRef;
 		const TSharedRef< FLinearColor > SelectionColor_Pressed_LinearRef;
+		const TSharedRef< FLinearColor > HighlightColor_LinearRef;
 
 		const TSharedRef< FLinearColor > LogColor_Background_LinearRef;
 		const TSharedRef< FLinearColor > LogColor_SelectionBackground_LinearRef;
@@ -141,6 +142,7 @@ public:
 		const FSlateColor SelectionColor_Subdued;
 		const FSlateColor SelectionColor_Inactive;
 		const FSlateColor SelectionColor_Pressed;
+		const FSlateColor HighlightColor;
 
 		const FSlateColor LogColor_Background;
 		const FSlateColor LogColor_SelectionBackground;
diff --git a/Engine/Source/Editor/GraphEditor/Private/SGraphActionMenu.cpp b/Engine/Source/Editor/GraphEditor/Private/SGraphActionMenu.cpp
--- a/Engine/Source/Editor/GraphEditor/Private/SGraphActionMenu.cpp
+++ b/Engine/Source/Editor/GraphEditor/Private/SGraphActionMenu.cpp
@@ -294,7 +294,8 @@ void SGraphActionMenu::Construct( const FArguments& InArgs, bool bIsReadOnly/* =
 		.OnGetChildren(this, &SGraphActionMenu::OnGetChildrenForCategory)
 		.SelectionMode(ESelectionMode::Single)
 		.OnItemScrolledIntoView(this, &SGraphActionMenu::OnItemScrolledIntoView)
-		.OnSetExpansionRecursive(this, &SGraphActionMenu::OnSetExpansionRecursive);
+		.OnSetExpansionRecursive(this, &SGraphActionMenu::OnSetExpansionRecursive)
+		.HighlightParentNodesForSelection(true);
 
 
 	this->ChildSlot
diff --git a/Engine/Source/Editor/Kismet/Private/SSCSEditor.cpp b/Engine/Source/Editor/Kismet/Private/SSCSEditor.cpp
--- a/Engine/Source/Editor/Kismet/Private/SSCSEditor.cpp
+++ b/Engine/Source/Editor/Kismet/Private/SSCSEditor.cpp
@@ -100,7 +100,8 @@ void SSCSEditorDragDropTree::Construct( const FArguments& InArgs )
 			.HeaderRow( InArgs._HeaderRow )
 			.ClearSelectionOnClick( InArgs._ClearSelectionOnClick )
 			.ExternalScrollbar( InArgs._ExternalScrollbar )
-			.OnEnteredBadState( InArgs._OnTableViewBadState );
+			.OnEnteredBadState( InArgs._OnTableViewBadState )
+			.HighlightParentNodesForSelection(true);
 
 	STreeView<FSCSEditorTreeNodePtrType>::Construct( BaseArgs );
 }
diff --git a/Engine/Source/Editor/LevelEditor/Private/SLevelEditor.cpp b/Engine/Source/Editor/LevelEditor/Private/SLevelEditor.cpp
--- a/Engine/Source/Editor/LevelEditor/Private/SLevelEditor.cpp
+++ b/Engine/Source/Editor/LevelEditor/Private/SLevelEditor.cpp
@@ -780,6 +780,17 @@ TSharedRef<SDockTab> SLevelEditor::SpawnLevelEditorTab( const FSpawnTabArgs& Arg
 				];
 		}
 	}
+	else if( TabIdentifier == TEXT("SequencerGraphEditor") )
+	{
+		const FSlateIcon SequencerGraphIcon = FSlateIcon(FEditorStyle::GetStyleSetName(), "GenericCurveEditor.TabIcon");
+		// @todo sequencer: remove when world-centric mode is added
+		return SNew(SDockTab)
+			.Icon(SequencerGraphIcon.GetIcon())
+			.Label(NSLOCTEXT("Sequencer", "SequencerMainGraphEditorTitle", "Sequencer Curves"))
+			[
+				SNullWidget::NullWidget
+			];
+	}
 	else if( TabIdentifier == LevelEditorStatsViewerTab )
 	{
 		FStatsViewerModule& StatsViewerModule = FModuleManager::Get().LoadModuleChecked<FStatsViewerModule>( "StatsViewer" );
@@ -1131,6 +1142,13 @@ TSharedRef<SWidget> SLevelEditor::RestoreContentArea( const TSharedRef<SDockTab>
 				.SetIcon( SequencerIcon );
 		}
 
+		{
+			// @todo remove when world-centric mode is added
+			const FSlateIcon SequencerGraphIcon = FSlateIcon(FEditorStyle::GetStyleSetName(), "GenericCurveEditor.TabIcon");
+			LevelEditorTabManager->RegisterTabSpawner("SequencerGraphEditor", FOnSpawnTab::CreateSP<SLevelEditor, FName, FString>(this, &SLevelEditor::SpawnLevelEditorTab, FName("SequencerGraphEditor"), FString()))
+				.SetMenuType(ETabSpawnerMenuType::Type::Hidden);
+		}
+
 		{
 			const FSlateIcon WorldPropertiesIcon(FEditorStyle::GetStyleSetName(), "LevelEditor.WorldProperties.Tab");
 			LevelEditorTabManager->RegisterTabSpawner( "WorldSettingsTab", FOnSpawnTab::CreateSP<SLevelEditor, FName, FString>(this, &SLevelEditor::SpawnLevelEditorTab, FName("WorldSettingsTab"), FString()) )
diff --git a/Engine/Source/Editor/MovieSceneTools/Private/Channels/BuiltInChannelEditors.cpp b/Engine/Source/Editor/MovieSceneTools/Private/Channels/BuiltInChannelEditors.cpp
--- a/Engine/Source/Editor/MovieSceneTools/Private/Channels/BuiltInChannelEditors.cpp
+++ b/Engine/Source/Editor/MovieSceneTools/Private/Channels/BuiltInChannelEditors.cpp
@@ -17,6 +17,7 @@
 #include "Channels/MovieSceneChannelProxy.h"
 #include "Channels/MovieSceneChannelEditorData.h"
 #include "FloatChannelCurveModel.h"
+#include "EventChannelCurveModel.h"
 #include "PropertyCustomizationHelpers.h"
 
 #define LOCTEXT_NAMESPACE "BuiltInChannelEditors"
@@ -798,4 +799,10 @@ TUniquePtr<FCurveModel> CreateCurveEditorModel(const TMovieSceneChannelHandle<FM
 	return MakeUnique<FFloatChannelCurveModel>(FloatChannel, OwningSection, InSequencer);
 }
 
+TUniquePtr<FCurveModel> CreateCurveEditorModel(const TMovieSceneChannelHandle<FMovieSceneEventChannel>& EventChannel, UMovieSceneSection* OwningSection, TSharedRef<ISequencer> InSequencer)
+{
+	return MakeUnique<FEventChannelCurveModel>(EventChannel, OwningSection, InSequencer);
+}
+
+
 #undef LOCTEXT_NAMESPACE
diff --git a/Engine/Source/Editor/MovieSceneTools/Private/Channels/BuiltInChannelEditors.h b/Engine/Source/Editor/MovieSceneTools/Private/Channels/BuiltInChannelEditors.h
--- a/Engine/Source/Editor/MovieSceneTools/Private/Channels/BuiltInChannelEditors.h
+++ b/Engine/Source/Editor/MovieSceneTools/Private/Channels/BuiltInChannelEditors.h
@@ -65,4 +65,8 @@ void ExtendSectionMenu(FMenuBuilder& OuterMenuBuilder, TArray<TMovieSceneChannel
 void ExtendKeyMenu(FMenuBuilder& OuterMenuBuilder, TArray<TExtendKeyMenuParams<FMovieSceneFloatChannel>>&& Channels, TWeakPtr<ISequencer> InSequencer);
 
 /** Curve editor models */
+inline bool SupportsCurveEditorModels(const TMovieSceneChannelHandle<FMovieSceneFloatChannel>& FloatChannel) { return true; }
+inline bool SupportsCurveEditorModels(const TMovieSceneChannelHandle<FMovieSceneEventChannel>& EventChannel) { return true; }
+
 TUniquePtr<FCurveModel> CreateCurveEditorModel(const TMovieSceneChannelHandle<FMovieSceneFloatChannel>& FloatChannel, UMovieSceneSection* OwningSection, TSharedRef<ISequencer> InSequencer);
+TUniquePtr<FCurveModel> CreateCurveEditorModel(const TMovieSceneChannelHandle<FMovieSceneEventChannel>& EventChannel, UMovieSceneSection* OwningSection, TSharedRef<ISequencer> InSequencer);
diff --git a/Engine/Source/Editor/MovieSceneTools/Private/Channels/EventChannelCurveModel.cpp b/Engine/Source/Editor/MovieSceneTools/Private/Channels/EventChannelCurveModel.cpp
new file mode 100644
--- /dev/null
+++ b/Engine/Source/Editor/MovieSceneTools/Private/Channels/EventChannelCurveModel.cpp
@@ -0,0 +1,487 @@
+// Copyright 1998-2018 Epic Games, Inc. All Rights Reserved.
+
+#include "EventChannelCurveModel.h"
+#include "Math/Vector2D.h"
+#include "HAL/PlatformMath.h"
+#include "Channels/MovieSceneFloatChannel.h"
+#include "Channels/FloatChannelKeyProxy.h"
+#include "MovieSceneSection.h"
+#include "MovieScene.h"
+#include "CurveDrawInfo.h"
+#include "CurveDataAbstraction.h"
+#include "CurveEditor.h"
+#include "CurveEditorScreenSpace.h"
+#include "CurveEditorSnapMetrics.h"
+#include "EditorStyleSet.h"
+#include "BuiltInChannelEditors.h"
+#include "SequencerChannelTraits.h"
+#include "SequencerSectionPainter.h"
+#include "ISequencer.h"
+
+ECurveEditorViewID FEventChannelCurveModel::EventView = ECurveEditorViewID::Invalid;
+
+FEventChannelCurveModel::FEventChannelCurveModel(TMovieSceneChannelHandle<FMovieSceneEventChannel> InChannel, UMovieSceneSection* OwningSection, TWeakPtr<ISequencer> InWeakSequencer)
+{
+	ChannelHandle = InChannel;
+	WeakSection = OwningSection;
+	WeakSequencer = InWeakSequencer;
+	SupportedViews = EventView;
+
+	Color = FSequencerSectionPainter::BlendColor(OwningSection->GetTypedOuter<UMovieSceneTrack>()->GetColorTint());
+}
+
+const void* FEventChannelCurveModel::GetCurve() const
+{
+	return ChannelHandle.Get();
+}
+
+void FEventChannelCurveModel::Modify()
+{
+	if (UMovieSceneSection* Section = WeakSection.Get())
+	{
+		Section->Modify();
+	}
+}
+
+void FEventChannelCurveModel::AddKeys(TArrayView<const FKeyPosition> InKeyPositions, TArrayView<const FKeyAttributes> InKeyAttributes, TArrayView<TOptional<FKeyHandle>>* OutKeyHandles)
+{
+	check(InKeyPositions.Num() == InKeyAttributes.Num() && (!OutKeyHandles || OutKeyHandles->Num() == InKeyPositions.Num()));
+
+	FMovieSceneEventChannel* Channel = ChannelHandle.Get();
+	UMovieSceneSection*      Section = WeakSection.Get();
+	if (Channel && Section)
+	{
+		Section->Modify();
+		TMovieSceneChannelData<FMovieSceneEvent> ChannelData = Channel->GetData();
+		FFrameRate TickResolution = Section->GetTypedOuter<UMovieScene>()->GetTickResolution();
+
+		for (int32 Index = 0; Index < InKeyPositions.Num(); ++Index)
+		{
+			FKeyPosition   Position   = InKeyPositions[Index];
+			FKeyAttributes Attributes = InKeyAttributes[Index];
+
+			FFrameNumber Time = (Position.InputValue * TickResolution).RoundToFrame();
+			Section->ExpandToFrame(Time);
+
+			FMovieSceneEvent Value;
+// 
+// 			if (Attributes.HasInterpMode())    { Value.InterpMode            = Attributes.GetInterpMode();    }
+// 			if (Attributes.HasTangentMode())   { Value.TangentMode           = Attributes.GetTangentMode();   }
+// 			if (Attributes.HasArriveTangent()) { Value.Tangent.ArriveTangent = Attributes.GetArriveTangent(); }
+// 			if (Attributes.HasLeaveTangent())  { Value.Tangent.LeaveTangent  = Attributes.GetLeaveTangent();  }
+
+			int32 KeyIndex = ChannelData.AddKey(Time, Value);
+			if (OutKeyHandles)
+			{
+				(*OutKeyHandles)[Index] = ChannelData.GetHandle(KeyIndex);
+			}
+		}
+
+		// Channel->AutoSetTangents();
+	}
+}
+
+bool FEventChannelCurveModel::Evaluate(double Time, double& OutValue) const
+{
+	// Events don't evaluate into a valid value
+	OutValue = 0.0;
+	return false;
+}
+
+void FEventChannelCurveModel::RemoveKeys(TArrayView<const FKeyHandle> InKeys)
+{
+	FMovieSceneEventChannel* Channel = ChannelHandle.Get();
+	UMovieSceneSection*      Section = WeakSection.Get();
+	if (Channel && Section)
+	{
+		Section->Modify();
+
+		TMovieSceneChannelData<FMovieSceneEvent> ChannelData = Channel->GetData();
+
+		for (FKeyHandle Handle : InKeys)
+		{
+			int32 KeyIndex = ChannelData.GetIndex(Handle);
+			if (KeyIndex != INDEX_NONE)
+			{
+				ChannelData.RemoveKey(KeyIndex);
+			}
+		}
+	}
+}
+
+void FEventChannelCurveModel::DrawCurve(const FCurveEditor& CurveEditor, const FCurveEditorScreenSpace& ScreenSpace, TArray<TTuple<double, double>>& InterpolatingPoints) const
+{
+	// Event Channels don't draw any lines so there's no need to fill out the Interpolating Points array.
+}
+
+void FEventChannelCurveModel::GetKeys(const FCurveEditor& CurveEditor, double MinTime, double MaxTime, double MinValue, double MaxValue, TArray<FKeyHandle>& OutKeyHandles) const
+{
+	FMovieSceneEventChannel* Channel = ChannelHandle.Get();
+	UMovieSceneSection*      Section = WeakSection.Get();
+
+	if (Channel && Section)
+	{
+		FFrameRate TickResolution = Section->GetTypedOuter<UMovieScene>()->GetTickResolution();
+
+		TMovieSceneChannelData<FMovieSceneEvent> ChannelData = Channel->GetData();
+		TArrayView<const FFrameNumber>          Times  = ChannelData.GetTimes();
+
+		const FFrameNumber StartFrame = MinTime <= MIN_int32 ? MIN_int32 : (MinTime * TickResolution).CeilToFrame();
+		const FFrameNumber EndFrame   = MaxTime >= MAX_int32 ? MAX_int32 : (MaxTime * TickResolution).FloorToFrame();
+
+		const int32 StartingIndex = Algo::LowerBound(Times, StartFrame);
+		const int32 EndingIndex   = Algo::UpperBound(Times, EndFrame);
+
+		for (int32 KeyIndex = StartingIndex; KeyIndex < EndingIndex; ++KeyIndex)
+		{
+			// Event Channels don't have Values associated with them so we ignore the Min/Max value and always return the key.
+			OutKeyHandles.Add(ChannelData.GetHandle(KeyIndex));
+		}
+	}
+}
+
+void FEventChannelCurveModel::GetKeyDrawInfo(ECurvePointType PointType, const FKeyHandle InKeyHandle, FKeyDrawInfo& OutDrawInfo) const
+{
+	OutDrawInfo.Brush = FEditorStyle::Get().GetBrush("Sequencer.KeyDiamond");
+	OutDrawInfo.ScreenSize = FVector2D(10, 10);
+}
+
+void FEventChannelCurveModel::GetKeyPositions(TArrayView<const FKeyHandle> InKeys, TArrayView<FKeyPosition> OutKeyPositions) const
+{
+	FMovieSceneEventChannel* Channel = ChannelHandle.Get();
+	UMovieSceneSection*      Section = WeakSection.Get();
+
+	if (Channel && Section)
+	{
+		FFrameRate TickResolution = Section->GetTypedOuter<UMovieScene>()->GetTickResolution();
+
+		TMovieSceneChannelData<FMovieSceneEvent> ChannelData = Channel->GetData();
+		TArrayView<const FFrameNumber>          Times  = ChannelData.GetTimes();
+		TArrayView<const FMovieSceneEvent> Values = ChannelData.GetValues();
+
+		for (int32 Index = 0; Index < InKeys.Num(); ++Index)
+		{
+			int32 KeyIndex = ChannelData.GetIndex(InKeys[Index]);
+			if (KeyIndex != INDEX_NONE)
+			{
+				OutKeyPositions[Index].InputValue  = Times[KeyIndex] / TickResolution;
+				// Events have no values so we just output zero.
+				OutKeyPositions[Index].OutputValue = 0.0;
+			}
+		}
+	}
+}
+
+void FEventChannelCurveModel::SetKeyPositions(TArrayView<const FKeyHandle> InKeys, TArrayView<const FKeyPosition> InKeyPositions)
+{
+	FMovieSceneEventChannel* Channel = ChannelHandle.Get();
+	UMovieSceneSection*      Section = WeakSection.Get();
+
+	if (Channel && Section)
+	{
+		Section->MarkAsChanged();
+
+		FFrameRate TickResolution = Section->GetTypedOuter<UMovieScene>()->GetTickResolution();
+
+		TMovieSceneChannelData<FMovieSceneEvent> ChannelData = Channel->GetData();
+		for (int32 Index = 0; Index < InKeys.Num(); ++Index)
+		{
+			int32 KeyIndex = ChannelData.GetIndex(InKeys[Index]);
+			if (KeyIndex != INDEX_NONE)
+			{
+				FFrameNumber NewTime = (InKeyPositions[Index].InputValue * TickResolution).FloorToFrame();
+
+				KeyIndex = ChannelData.MoveKey(KeyIndex, NewTime);
+				Section->ExpandToFrame(NewTime);
+			}
+		}
+	}
+}
+
+void FEventChannelCurveModel::GetKeyAttributes(TArrayView<const FKeyHandle> InKeys, TArrayView<FKeyAttributes> OutAttributes) const
+{
+// 	FMovieSceneEventChannel* Channel = ChannelHandle.Get();
+// 	UMovieSceneSection*      Section = WeakSection.Get();
+// 	if (Channel && Section)
+// 	{
+// 		TMovieSceneChannelData<FMovieSceneEvent> ChannelData = Channel->GetData();
+// 		TArrayView<const FFrameNumber>    Times  = ChannelData.GetTimes();
+// 		TArrayView<FMovieSceneEvent> Values = ChannelData.GetValues();
+// 
+// 		float TimeInterval = Section->GetTypedOuter<UMovieScene>()->GetTickResolution().AsInterval();
+// 
+// 		for (int32 Index = 0; Index < InKeys.Num(); ++Index)
+// 		{
+// 			const int32 KeyIndex = ChannelData.GetIndex(InKeys[Index]);
+// 			if (KeyIndex != INDEX_NONE)
+// 			{
+// 				const FMovieSceneEvent& KeyValue    = Values[KeyIndex];
+// 				FKeyAttributes&              Attributes  = OutAttributes[Index];
+// 
+// 				Attributes.SetInterpMode(KeyValue.InterpMode);
+// 
+// 				if (KeyValue.InterpMode != RCIM_Constant && KeyValue.InterpMode != RCIM_Linear)
+// 				{
+// 					Attributes.SetTangentMode(KeyValue.TangentMode);
+// 					if (KeyIndex != 0)
+// 					{
+// 						Attributes.SetArriveTangent(KeyValue.Tangent.ArriveTangent / TimeInterval);
+// 					}
+// 
+// 					if (KeyIndex != Times.Num()-1)
+// 					{
+// 						Attributes.SetLeaveTangent(KeyValue.Tangent.LeaveTangent / TimeInterval);
+// 					}
+// 					if (KeyValue.InterpMode == RCIM_Cubic)
+// 					{
+// 						Attributes.SetTangentWeightMode(KeyValue.Tangent.TangentWeightMode);
+// 						if (KeyValue.Tangent.TangentWeightMode != RCTWM_WeightedNone)
+// 						{
+// 							Attributes.SetArriveTangentWeight(KeyValue.Tangent.ArriveTangentWeight);
+// 							Attributes.SetLeaveTangentWeight(KeyValue.Tangent.LeaveTangentWeight);
+// 						}
+// 					}
+// 				}
+// 			}
+// 		}
+// 	}
+}
+
+void FEventChannelCurveModel::SetKeyAttributes(TArrayView<const FKeyHandle> InKeys, TArrayView<const FKeyAttributes> InAttributes)
+{
+// 	FMovieSceneEventChannel* Channel = ChannelHandle.Get();
+// 	UMovieSceneSection*      Section = WeakSection.Get();
+// 	if (Channel && Section)
+// 	{
+// 		bool bAutoSetTangents = false;
+// 		Section->MarkAsChanged();
+// 
+// 		TMovieSceneChannelData<FMovieSceneEvent> ChannelData = Channel->GetData();
+// 		TArrayView<FMovieSceneEvent> Values = ChannelData.GetValues();
+// 
+// 		FFrameRate TickResolution = Section->GetTypedOuter<UMovieScene>()->GetTickResolution();
+// 		float TimeInterval = TickResolution.AsInterval();
+// 
+// 		for (int32 Index = 0; Index < InKeys.Num(); ++Index)
+// 		{
+// 			const int32 KeyIndex = ChannelData.GetIndex(InKeys[Index]);
+// 			if (KeyIndex != INDEX_NONE)
+// 			{
+// 				const FKeyAttributes&  Attributes = InAttributes[Index];
+// 				FMovieSceneEvent& KeyValue   = Values[KeyIndex];
+// 				if (Attributes.HasInterpMode())    { KeyValue.InterpMode  = Attributes.GetInterpMode();  bAutoSetTangents = true; }
+// 				if (Attributes.HasTangentMode())
+// 				{
+// 					KeyValue.TangentMode = Attributes.GetTangentMode();
+// 					if (KeyValue.TangentMode == RCTM_Auto)
+// 					{
+// 						KeyValue.Tangent.TangentWeightMode = RCTWM_WeightedNone;
+// 					}
+// 					bAutoSetTangents = true;
+// 				}
+// 				if (Attributes.HasTangentWeightMode()) 
+// 				{ 
+// 					if (KeyValue.Tangent.TangentWeightMode == RCTWM_WeightedNone) //set tangent weights to default use
+// 					{
+// 						TArrayView<const FFrameNumber> Times = Channel->GetTimes();
+// 						const float OneThird = 1.0f / 3.0f;
+// 
+// 						//calculate a tangent weight based upon tangent and time difference
+// 						//calculate arrive tangent weight
+// 						if (KeyIndex > 0 )
+// 						{
+// 							const float X = TickResolution.AsSeconds(Times[KeyIndex].Value - Times[KeyIndex - 1].Value);
+// 							const float ArriveTangentNormal = KeyValue.Tangent.ArriveTangent / (TimeInterval);
+// 							const float Y = ArriveTangentNormal * X;
+// 							KeyValue.Tangent.ArriveTangentWeight = FMath::Sqrt(X*X + Y * Y) * OneThird;
+// 						}
+// 						//calculate leave weight
+// 						if(KeyIndex < ( Times.Num() - 1))
+// 						{
+// 							const float X = TickResolution.AsSeconds(Times[KeyIndex].Value - Times[KeyIndex + 1].Value);
+// 							const float LeaveTangentNormal = KeyValue.Tangent.LeaveTangent / (TimeInterval);
+// 							const float Y = LeaveTangentNormal * X;
+// 							KeyValue.Tangent.LeaveTangentWeight = FMath::Sqrt(X*X + Y*Y) * OneThird;
+// 						}
+// 					}
+// 					KeyValue.Tangent.TangentWeightMode = Attributes.GetTangentWeightMode();
+// 
+// 					if( KeyValue.Tangent.TangentWeightMode != RCTWM_WeightedNone )
+// 					{
+// 						if (KeyValue.TangentMode != RCTM_User && KeyValue.TangentMode != RCTM_Break)
+// 						{
+// 							KeyValue.TangentMode = RCTM_User;
+// 						}
+// 					}
+// 				}
+// 
+// 				if (Attributes.HasArriveTangent())
+// 				{
+// 					if (KeyValue.TangentMode == RCTM_Auto)
+// 					{
+// 						KeyValue.TangentMode = RCTM_User;
+// 						KeyValue.Tangent.TangentWeightMode = RCTWM_WeightedNone;
+// 					}
+// 
+// 					KeyValue.Tangent.ArriveTangent = Attributes.GetArriveTangent() * TimeInterval;
+// 					if (KeyValue.InterpMode == RCIM_Cubic && KeyValue.TangentMode != RCTM_Break)
+// 					{
+// 						KeyValue.Tangent.LeaveTangent = KeyValue.Tangent.ArriveTangent;
+// 					}
+// 				}
+// 
+// 				if (Attributes.HasLeaveTangent())
+// 				{
+// 					if (KeyValue.TangentMode == RCTM_Auto)
+// 					{
+// 						KeyValue.TangentMode = RCTM_User;
+// 						KeyValue.Tangent.TangentWeightMode = RCTWM_WeightedNone;
+// 					}
+// 
+// 					KeyValue.Tangent.LeaveTangent = Attributes.GetLeaveTangent() * TimeInterval;
+// 					if (KeyValue.InterpMode == RCIM_Cubic && KeyValue.TangentMode != RCTM_Break)
+// 					{
+// 						KeyValue.Tangent.ArriveTangent = KeyValue.Tangent.LeaveTangent;
+// 					}
+// 				}
+// 
+// 				if (Attributes.HasArriveTangentWeight())
+// 				{
+// 					if (KeyValue.TangentMode == RCTM_Auto)
+// 					{
+// 						KeyValue.TangentMode = RCTM_User;
+// 						KeyValue.Tangent.TangentWeightMode = RCTWM_WeightedNone;
+// 					}
+// 
+// 					KeyValue.Tangent.ArriveTangentWeight = Attributes.GetArriveTangentWeight(); 
+// 					if (KeyValue.InterpMode == RCIM_Cubic && KeyValue.TangentMode != RCTM_Break)
+// 					{
+// 						KeyValue.Tangent.LeaveTangentWeight = KeyValue.Tangent.ArriveTangentWeight;
+// 					}
+// 				}
+// 
+// 				if (Attributes.HasLeaveTangentWeight())
+// 				{
+// 				
+// 					if (KeyValue.TangentMode == RCTM_Auto)
+// 					{
+// 						KeyValue.TangentMode = RCTM_User;
+// 						KeyValue.Tangent.TangentWeightMode = RCTWM_WeightedNone;
+// 					}
+// 
+// 					KeyValue.Tangent.LeaveTangentWeight = Attributes.GetLeaveTangentWeight();
+// 					if (KeyValue.InterpMode == RCIM_Cubic && KeyValue.TangentMode != RCTM_Break)
+// 					{
+// 						KeyValue.Tangent.ArriveTangentWeight = KeyValue.Tangent.LeaveTangentWeight;
+// 					}
+// 				}
+// 			}
+// 		}
+// 
+// 		if (bAutoSetTangents)
+// 		{
+// 			Channel->AutoSetTangents();
+// 		}
+// 	}
+}
+
+void FEventChannelCurveModel::GetCurveAttributes(FCurveAttributes& OutCurveAttributes) const
+{
+	// Event Channels have no Pre/Post Extrapolation
+	OutCurveAttributes.SetPreExtrapolation(RCCE_None);
+	OutCurveAttributes.SetPostExtrapolation(RCCE_None);
+}
+
+void FEventChannelCurveModel::SetCurveAttributes(const FCurveAttributes& InCurveAttributes)
+{
+	// Event Channels have no Pre/Post Extrapolation
+}
+
+void FEventChannelCurveModel::CreateKeyProxies(TArrayView<const FKeyHandle> InKeyHandles, TArrayView<UObject*> OutObjects)
+{
+	// for (int32 Index = 0; Index < InKeyHandles.Num(); ++Index)
+	// {
+	// 	UFloatChannelKeyProxy* NewProxy = NewObject<UFloatChannelKeyProxy>(GetTransientPackage(), NAME_None);
+	// 
+	// 	NewProxy->Initialize(InKeyHandles[Index], ChannelHandle, WeakSection);
+	// 	OutObjects[Index] = NewProxy;
+	// }
+}
+
+void FEventChannelCurveModel::GetTimeRange(double& MinTime, double& MaxTime) const
+{
+	FMovieSceneEventChannel* Channel = ChannelHandle.Get();
+	UMovieSceneSection*      Section = WeakSection.Get();
+
+	if (Channel && Section)
+	{
+		TArrayView<const FFrameNumber> Times = Channel->GetData().GetTimes();
+		if (Times.Num() == 0)
+		{
+			MinTime = 0.f;
+			MaxTime = 0.f;
+		}
+		else
+		{
+			FFrameRate TickResolution = Section->GetTypedOuter<UMovieScene>()->GetTickResolution();
+			double ToTime = TickResolution.AsInterval();
+			MinTime = static_cast<double> (Times[0].Value) * ToTime;
+			MaxTime = static_cast<double>(Times[Times.Num() - 1].Value) * ToTime;
+		}
+	}
+}
+
+/*	 Finds min/max for cubic curves:
+Looks for feature points in the signal(determined by change in direction of local tangent), these locations are then re-examined in closer detail recursively
+Similar to function in RichCurve but usees the Channel::Evaluate function, instead of CurveModel::Eval*/
+
+void FEventChannelCurveModel::FeaturePointMethod(double StartTime, double EndTime, double StartValue, double Mu, int Depth, int MaxDepth, double& MaxV, double& MinVal) const
+{
+	if (Depth >= MaxDepth)
+	{
+		return;
+	}
+	double PrevValue = StartValue;
+	double EvalValue;
+	Evaluate(StartTime - Mu, EvalValue);
+	double PrevTangent = StartValue - EvalValue;
+	EndTime += Mu;
+	for (double f = StartTime + Mu; f < EndTime; f += Mu)
+	{
+		double Value;
+		Evaluate(f, Value);
+
+		MaxV = FMath::Max(Value, MaxV);
+		MinVal = FMath::Min(Value, MinVal);
+		double CurTangent = Value - PrevValue;
+		if (FMath::Sign(CurTangent) != FMath::Sign(PrevTangent))
+		{
+			//feature point centered around the previous tangent
+			double FeaturePointTime = f - Mu * 2.0f;
+			double NewVal;
+			Evaluate(FeaturePointTime, NewVal);
+			FeaturePointMethod(FeaturePointTime, f,NewVal, Mu*0.4f, Depth + 1, MaxDepth, MaxV, MinVal);
+		}
+		PrevTangent = CurTangent;
+		PrevValue = Value;
+	}
+}
+
+void FEventChannelCurveModel::GetValueRange(double& MinValue, double& MaxValue) const
+{
+	// Event Tracks have no value, thus their value range is just zero.
+	MinValue = MaxValue = 0.0;
+}
+
+int32 FEventChannelCurveModel::GetNumKeys() const
+{
+	FMovieSceneEventChannel* Channel = ChannelHandle.Get();
+
+	if (Channel)
+	{
+		TArrayView<const FFrameNumber> Times = Channel->GetData().GetTimes();
+		return Times.Num();
+	}
+
+	return 0;
+}
diff --git a/Engine/Source/Editor/MovieSceneTools/Private/Channels/FloatChannelCurveModel.cpp b/Engine/Source/Editor/MovieSceneTools/Private/Channels/FloatChannelCurveModel.cpp
--- a/Engine/Source/Editor/MovieSceneTools/Private/Channels/FloatChannelCurveModel.cpp
+++ b/Engine/Source/Editor/MovieSceneTools/Private/Channels/FloatChannelCurveModel.cpp
@@ -4,6 +4,7 @@
 #include "Math/Vector2D.h"
 #include "HAL/PlatformMath.h"
 #include "Channels/MovieSceneFloatChannel.h"
+#include "Channels/FloatChannelKeyProxy.h"
 #include "MovieSceneSection.h"
 #include "MovieScene.h"
 #include "CurveDrawInfo.h"
@@ -15,7 +16,6 @@
 #include "BuiltInChannelEditors.h"
 #include "SequencerChannelTraits.h"
 #include "ISequencer.h"
-#include "Channels/FloatChannelKeyProxy.h"
 
 FFloatChannelCurveModel::FFloatChannelCurveModel(TMovieSceneChannelHandle<FMovieSceneFloatChannel> InChannel, UMovieSceneSection* OwningSection, TWeakPtr<ISequencer> InWeakSequencer)
 {
@@ -27,6 +27,8 @@ FFloatChannelCurveModel::FFloatChannelCurveModel(TMovieSceneChannelHandle<FMovie
 		FMovieSceneFloatChannel* Channel = ChannelHandle.Get();
 		Channel->SetTickResolution(Section->GetTypedOuter<UMovieScene>()->GetTickResolution());
 	}
+
+	SupportedViews = ECurveEditorViewID::Absolute | ECurveEditorViewID::Normalized | ECurveEditorViewID::Stacked;
 }
 
 const void* FFloatChannelCurveModel::GetCurve() const
@@ -122,15 +124,13 @@ void FFloatChannelCurveModel::RemoveKeys(TArrayView<const FKeyHandle> InKeys)
 	}
 }
 
-void FFloatChannelCurveModel::DrawCurve(const FCurveEditor& CurveEditor, TArray<TTuple<double, double>>& InterpolatingPoints) const
+void FFloatChannelCurveModel::DrawCurve(const FCurveEditor& CurveEditor, const FCurveEditorScreenSpace& ScreenSpace, TArray<TTuple<double, double>>& InterpolatingPoints) const
 {
 	FMovieSceneFloatChannel* Channel = ChannelHandle.Get();
 	UMovieSceneSection*      Section = WeakSection.Get();
 
 	if (Channel && Section)
 	{
-		FCurveEditorScreenSpace ScreenSpace = CurveEditor.GetScreenSpace();
-
 		FFrameRate   TickResolution   = Section->GetTypedOuter<UMovieScene>()->GetTickResolution();
 
 		const double DisplayOffset    = GetInputDisplayOffset();
@@ -172,19 +172,51 @@ void FFloatChannelCurveModel::GetKeys(const FCurveEditor& CurveEditor, double Mi
 	}
 }
 
-void FFloatChannelCurveModel::GetKeyDrawInfo(ECurvePointType PointType, FKeyDrawInfo& OutDrawInfo) const
+void FFloatChannelCurveModel::GetKeyDrawInfo(ECurvePointType PointType, const FKeyHandle InKeyHandle, FKeyDrawInfo& OutDrawInfo) const
 {
-	switch (PointType)
+	if (PointType == ECurvePointType::ArriveTangent || PointType == ECurvePointType::LeaveTangent)
+	{
+		OutDrawInfo.Brush = FEditorStyle::GetBrush("GenericCurveEditor.TangentHandle");
+		OutDrawInfo.ScreenSize = FVector2D(8, 8);
+	}
+	else
 	{
-	case ECurvePointType::ArriveTangent:
-	case ECurvePointType::LeaveTangent:
-		OutDrawInfo.Brush = FEditorStyle::GetBrush("Sequencer.TangentHandle");
-		OutDrawInfo.ScreenSize = FVector2D(7, 7);
-		break;
-	default:
-		OutDrawInfo.Brush = FEditorStyle::GetBrush("CurveEd.CurveKey");
+		// All keys are the same size by default
 		OutDrawInfo.ScreenSize = FVector2D(11, 11);
-		break;
+
+		ERichCurveInterpMode KeyType = RCIM_None;
+
+		// Get the key type from the supplied key handle if it's valid
+		FMovieSceneFloatChannel* Channel = ChannelHandle.Get();
+		if (Channel && InKeyHandle != FKeyHandle::Invalid())
+		{
+			TMovieSceneChannelData<FMovieSceneFloatValue> ChannelData = Channel->GetData();
+			const int32 KeyIndex = ChannelData.GetIndex(InKeyHandle);
+			if (KeyIndex != INDEX_NONE)
+			{
+				KeyType = ChannelData.GetValues()[KeyIndex].InterpMode;
+			}
+		}
+
+		switch (KeyType)
+		{
+		case ERichCurveInterpMode::RCIM_Constant:
+			OutDrawInfo.Brush = FEditorStyle::GetBrush("GenericCurveEditor.ConstantKey");
+			OutDrawInfo.Tint = FLinearColor(0, 0.45f, 0.70f);
+			break;
+		case ERichCurveInterpMode::RCIM_Linear:
+			OutDrawInfo.Brush = FEditorStyle::GetBrush("GenericCurveEditor.LinearKey");
+			OutDrawInfo.Tint = FLinearColor(0, 0.62f, 0.46f);
+			break;
+		case ERichCurveInterpMode::RCIM_Cubic:
+			OutDrawInfo.Brush = FEditorStyle::GetBrush("GenericCurveEditor.CubicKey");
+			OutDrawInfo.Tint = FLinearColor::White;
+			break;
+		default:
+			OutDrawInfo.Brush = FEditorStyle::GetBrush("GenericCurveEditor.Key");
+			OutDrawInfo.Tint = FLinearColor::White;
+			break;
+		}
 	}
 }
 
@@ -230,7 +262,7 @@ void FFloatChannelCurveModel::SetKeyPositions(TArrayView<const FKeyHandle> InKey
 			int32 KeyIndex = ChannelData.GetIndex(InKeys[Index]);
 			if (KeyIndex != INDEX_NONE)
 			{
-				FFrameNumber NewTime = (InKeyPositions[Index].InputValue * TickResolution).FloorToFrame();
+				FFrameNumber NewTime = (InKeyPositions[Index].InputValue * TickResolution).RoundToFrame();
 
 				KeyIndex = ChannelData.MoveKey(KeyIndex, NewTime);
 				ChannelData.GetValues()[KeyIndex].Value = InKeyPositions[Index].OutputValue;
@@ -542,7 +574,8 @@ void FFloatChannelCurveModel::GetValueRange(double& MinValue, double& MaxValue)
 
 		if (Times.Num() == 0)
 		{
-			MinValue = MaxValue = 0.f;
+			// If there are no keys we just use the default value for the channel, defaulting to zero if there is no default.
+			MinValue = MaxValue = Channel->GetDefault().Get(0.f);
 		}
 		else
 		{
@@ -569,4 +602,42 @@ void FFloatChannelCurveModel::GetValueRange(double& MinValue, double& MaxValue)
 			}
 		}
 	}
-}
\ No newline at end of file
+}
+
+int32 FFloatChannelCurveModel::GetNumKeys() const
+{
+	FMovieSceneFloatChannel* Channel = ChannelHandle.Get();
+
+	if (Channel)
+	{
+		TArrayView<const FFrameNumber> Times = Channel->GetData().GetTimes();
+
+		return Times.Num();
+	}
+
+	return 0;
+}
+
+void FFloatChannelCurveModel::GetNeighboringKeys(const FKeyHandle InKeyHandle, TOptional<FKeyHandle>& OutPreviousKeyHandle, TOptional<FKeyHandle>& OutNextKeyHandle) const
+{
+	FMovieSceneFloatChannel* Channel = ChannelHandle.Get();
+
+	if (Channel)
+	{
+		TMovieSceneChannelData<FMovieSceneFloatValue> ChannelData = Channel->GetData();
+
+		const int32 KeyIndex = ChannelData.GetIndex(InKeyHandle);
+		if (KeyIndex != INDEX_NONE)
+		{
+			if (KeyIndex - 1 >= 0)
+			{
+				OutPreviousKeyHandle = ChannelData.GetHandle(KeyIndex - 1);
+			}
+
+			if (KeyIndex + 1 < ChannelData.GetTimes().Num())
+			{
+				OutNextKeyHandle = ChannelData.GetHandle(KeyIndex + 1);
+			}
+		}
+	}
+}
diff --git a/Engine/Source/Editor/MovieSceneTools/Private/Channels/FloatChannelCurveModel.h b/Engine/Source/Editor/MovieSceneTools/Private/Channels/FloatChannelCurveModel.h
--- a/Engine/Source/Editor/MovieSceneTools/Private/Channels/FloatChannelCurveModel.h
+++ b/Engine/Source/Editor/MovieSceneTools/Private/Channels/FloatChannelCurveModel.h
@@ -23,9 +23,9 @@ public:
 
 	virtual void Modify() override;
 
-	virtual void DrawCurve(const FCurveEditor& CurveEditor, TArray<TTuple<double, double>>& InterpolatingPoints) const override;
+	virtual void DrawCurve(const FCurveEditor& CurveEditor, const FCurveEditorScreenSpace& ScreenSpace, TArray<TTuple<double, double>>& InterpolatingPoints) const override;
 	virtual void GetKeys(const FCurveEditor& CurveEditor, double MinTime, double MaxTime, double MinValue, double MaxValue, TArray<FKeyHandle>& OutKeyHandles) const override;
-	virtual void GetKeyDrawInfo(ECurvePointType PointType, FKeyDrawInfo& OutDrawInfo) const override;
+	virtual void GetKeyDrawInfo(ECurvePointType PointType, const FKeyHandle InKeyHandle, FKeyDrawInfo& OutDrawInfo) const override;
 
 	virtual void GetKeyPositions(TArrayView<const FKeyHandle> InKeys, TArrayView<FKeyPosition> OutKeyPositions) const override;
 	virtual void SetKeyPositions(TArrayView<const FKeyHandle> InKeys, TArrayView<const FKeyPosition> InKeyPositions) override;
@@ -37,6 +37,8 @@ public:
 	virtual void SetCurveAttributes(const FCurveAttributes& InCurveAttributes) override;
 	virtual void GetTimeRange(double& MinTime, double& MaxTime) const override;
 	virtual void GetValueRange(double& MinValue, double& MaxValue) const override;
+	virtual int32 GetNumKeys() const override;
+	virtual void GetNeighboringKeys(const FKeyHandle InKeyHandle, TOptional<FKeyHandle>& OutPreviousKeyHandle, TOptional<FKeyHandle>& OutNextKeyHandle) const override;
 	virtual bool Evaluate(double ProspectiveTime, double& OutValue) const override;
 	virtual void AddKeys(TArrayView<const FKeyPosition> InKeyPositions, TArrayView<const FKeyAttributes> InAttributes, TArrayView<TOptional<FKeyHandle>>* OutKeyHandles) override;
 	virtual void RemoveKeys(TArrayView<const FKeyHandle> InKeys) override;
diff --git a/Engine/Source/Editor/MovieSceneTools/Private/Channels/SCurveEditorEventChannelView.h b/Engine/Source/Editor/MovieSceneTools/Private/Channels/SCurveEditorEventChannelView.h
new file mode 100644
--- /dev/null
+++ b/Engine/Source/Editor/MovieSceneTools/Private/Channels/SCurveEditorEventChannelView.h
@@ -0,0 +1,26 @@
+// Copyright 1998-2018 Epic Games, Inc. All Rights Reserved.
+
+#pragma once
+
+#include "Views/SInteractiveCurveEditorView.h"
+
+class MOVIESCENETOOLS_API SCurveEditorEventChannelView : public SInteractiveCurveEditorView
+{
+public:
+
+	void Construct(const FArguments& InArgs, TWeakPtr<FCurveEditor> InCurveEditor);
+
+private:
+
+	//~ SWidget Interface
+	virtual void Tick(const FGeometry& AllottedGeometry, const double InCurrentTime, const float InDeltaTime) override;
+	virtual FVector2D ComputeDesiredSize(float LayoutScaleMultiplier) const override;
+
+	//~ SInteractiveCurveEditorView interface
+	virtual void GetGridLinesY(TSharedRef<FCurveEditor> CurveEditor, TArray<float>& MajorGridLines, TArray<float>& MinorGridLines, TArray<FText>& MajorGridLabels) const override;
+	virtual void PaintView(const FPaintArgs& Args, const FGeometry& AllottedGeometry, const FSlateRect& MyCullingRect, FSlateWindowElementList& OutDrawElements, int32 BaseLayerId, const FWidgetStyle& InWidgetStyle, bool bParentEnabled) const override;
+
+	void DrawLabels(const FGeometry& AllottedGeometry, FSlateWindowElementList& OutDrawElements, int32 BaseLayerId, ESlateDrawEffect DrawEffects) const;
+
+	static float TrackHeight;
+};
\ No newline at end of file
diff --git a/Engine/Source/Editor/MovieSceneTools/Private/FrameNumberDetailsCustomization.cpp b/Engine/Source/Editor/MovieSceneTools/Private/FrameNumberDetailsCustomization.cpp
--- a/Engine/Source/Editor/MovieSceneTools/Private/FrameNumberDetailsCustomization.cpp
+++ b/Engine/Source/Editor/MovieSceneTools/Private/FrameNumberDetailsCustomization.cpp
@@ -68,8 +68,12 @@ void FFrameNumberDetailsCustomization::CustomizeChildren(TSharedRef<IPropertyHan
 FText FFrameNumberDetailsCustomization::OnGetTimeText() const
 {
 	int32 CurrentValue = 0.0;
-	FrameNumberProperty->GetValue(CurrentValue);
+	FPropertyAccess::Result Result = FrameNumberProperty->GetValue(CurrentValue);
 
+	if (Result == FPropertyAccess::MultipleValues)
+	{
+		return LOCTEXT("MultipleValues", "Multiple Values");
+	}
 	return FText::FromString(NumericTypeInterface->ToString(CurrentValue));
 }
 
diff --git a/Engine/Source/Editor/MovieSceneTools/Private/MovieSceneToolsModule.cpp b/Engine/Source/Editor/MovieSceneTools/Private/MovieSceneToolsModule.cpp
--- a/Engine/Source/Editor/MovieSceneTools/Private/MovieSceneToolsModule.cpp
+++ b/Engine/Source/Editor/MovieSceneTools/Private/MovieSceneToolsModule.cpp
@@ -6,6 +6,7 @@
 #include "Modules/ModuleManager.h"
 #include "Curves/RichCurve.h"
 #include "ISequencerModule.h"
+#include "ICurveEditorModule.h"
 #include "MovieSceneToolsProjectSettingsCustomization.h"
 
 #include "TrackEditors/PropertyTrackEditors/BoolPropertyTrackEditor.h"
@@ -57,6 +58,8 @@
 #include "Channels/BuiltInChannelEditors.h"
 #include "Channels/MovieSceneObjectPathChannel.h"
 #include "Channels/MovieSceneEventChannel.h"
+#include "Channels/EventChannelCurveModel.h"
+#include "Channels/SCurveEditorEventChannelView.h"
 #include "Sections/MovieSceneEventSection.h"
 
 
@@ -131,10 +134,24 @@ void FMovieSceneToolsModule::StartupModule()
 	SequencerModule.RegisterChannelInterface<FMovieSceneObjectPathChannel>();
 
 	SequencerModule.RegisterChannelInterface<FMovieSceneEventChannel>();
+
+	ICurveEditorModule& CurveEditorModule = FModuleManager::LoadModuleChecked<ICurveEditorModule>("CurveEditor");
+
+	FEventChannelCurveModel::EventView = CurveEditorModule.RegisterView(FOnCreateCurveEditorView::CreateStatic(
+		[](TWeakPtr<FCurveEditor> WeakCurveEditor) -> TSharedRef<SCurveEditorView>
+		{
+			return SNew(SCurveEditorEventChannelView, WeakCurveEditor);
+		}
+	));
 }
 
 void FMovieSceneToolsModule::ShutdownModule()
 {
+	if (ICurveEditorModule* CurveEditorModule = FModuleManager::GetModulePtr<ICurveEditorModule>("CurveEditor"))
+	{
+		CurveEditorModule->UnregisterView(FEventChannelCurveModel::EventView);
+	}
+
 	if (ISettingsModule* SettingsModule = FModuleManager::GetModulePtr<ISettingsModule>("Settings"))
 	{
 		SettingsModule->UnregisterSettings("Project", "Editor", "Level Sequences");
