Sequencer: Only extend an existing section if autokey is on

diff --git a/Engine/Source/Editor/MovieSceneTools/Public/KeyframeTrackEditor.h b/Engine/Source/Editor/MovieSceneTools/Public/KeyframeTrackEditor.h
--- a/Engine/Source/Editor/MovieSceneTools/Public/KeyframeTrackEditor.h
+++ b/Engine/Source/Editor/MovieSceneTools/Public/KeyframeTrackEditor.h
@@ -432,7 +432,7 @@ private:
 		EAllowEditsMode AllowEditsMode = GetSequencer()->GetAllowEditsMode();
 
 		bool bCreateTrack =
-			(KeyMode == ESequencerKeyMode::AutoKey && (AutoChangeMode == EAutoChangeMode::AutoTrack || AutoChangeMode == EAutoChangeMode::All)) ||
+			(KeyMode == ESequencerKeyMode::AutoKey && (AutoChangeMode != EAutoChangeMode::None)) ||
 			KeyMode == ESequencerKeyMode::ManualKey ||
 			KeyMode == ESequencerKeyMode::ManualKeyForced ||
 			AllowEditsMode == EAllowEditsMode::AllowSequencerEditsOnly;
@@ -457,7 +457,7 @@ private:
 		if ( Track )
 		{
 			float Weight = 1.0f;
-			UMovieSceneSection* SectionToKey = Track->FindOrExtendSection(KeyTime, Weight);
+			UMovieSceneSection* SectionToKey = bCreateTrack ? Track->FindOrExtendSection(KeyTime, Weight) : Track->FindSection(KeyTime);
 
 			// If there's no overlapping section to key, create one only if a track was newly created. Otherwise, skip keying altogether
 			// so that the user is forced to create a section to key on.
