Sequencer: Only extend an existing section if autokey is on.

diff --git a/Engine/Source/Editor/MovieSceneTools/Public/KeyframeTrackEditor.h b/Engine/Source/Editor/MovieSceneTools/Public/KeyframeTrackEditor.h
--- a/Engine/Source/Editor/MovieSceneTools/Public/KeyframeTrackEditor.h
+++ b/Engine/Source/Editor/MovieSceneTools/Public/KeyframeTrackEditor.h
@@ -432,11 +432,14 @@ private:
 		EAllowEditsMode AllowEditsMode = GetSequencer()->GetAllowEditsMode();
 
 		bool bCreateTrack =
-			(KeyMode == ESequencerKeyMode::AutoKey && (AutoChangeMode != EAutoChangeMode::None)) ||
+			(KeyMode == ESequencerKeyMode::AutoKey && (AutoChangeMode == EAutoChangeMode::AutoTrack || AutoChangeMode == EAutoChangeMode::All)) ||
 			KeyMode == ESequencerKeyMode::ManualKey ||
 			KeyMode == ESequencerKeyMode::ManualKeyForced ||
 			AllowEditsMode == EAllowEditsMode::AllowSequencerEditsOnly;
 
+		bool bCreateSection = bCreateTrack || 
+			(KeyMode == ESequencerKeyMode::AutoKey && (AutoChangeMode != EAutoChangeMode::None));
+
 		// Try to find an existing Track, and if one doesn't exist check the key params and create one if requested.
 		FFindOrCreateTrackResult TrackResult = FindOrCreateTrackForObject( ObjectHandle, TrackClass, PropertyName, bCreateTrack );
 		TrackType* Track = CastChecked<TrackType>( TrackResult.Track, ECastCheckedType::NullAllowed );
@@ -457,7 +460,7 @@ private:
 		if ( Track )
 		{
 			float Weight = 1.0f;
-			UMovieSceneSection* SectionToKey = bCreateTrack ? Track->FindOrExtendSection(KeyTime, Weight) : Track->FindSection(KeyTime);
+			UMovieSceneSection* SectionToKey = bCreateSection ? Track->FindOrExtendSection(KeyTime, Weight) : Track->FindSection(KeyTime);
 
 			// If there's no overlapping section to key, create one only if a track was newly created. Otherwise, skip keying altogether
 			// so that the user is forced to create a section to key on.
