Fix typo

diff --git a/Engine/Source/Editor/MovieSceneTools/Private/TrackEditors/TransformTrackEditor.cpp b/Engine/Source/Editor/MovieSceneTools/Private/TrackEditors/TransformTrackEditor.cpp
--- a/Engine/Source/Editor/MovieSceneTools/Private/TrackEditors/TransformTrackEditor.cpp
+++ b/Engine/Source/Editor/MovieSceneTools/Private/TrackEditors/TransformTrackEditor.cpp
@@ -499,7 +499,7 @@ void F3DTransformTrackEditor::BuildObjectBindingTrackMenu(FMenuBuilder& MenuBuil
 	{
 		MenuBuilder.AddMenuEntry(
 			NSLOCTEXT("Sequencer", "AddTransform", "Transform"),
-			NSLOCTEXT("Sequencer", "AddPTransformTooltip", "Adds a transform track."),
+			NSLOCTEXT("Sequencer", "AddTransformTooltip", "Adds a transform track."),
 			FSlateIcon(),
 			FUIAction( 
 				FExecuteAction::CreateSP( this, &F3DTransformTrackEditor::AddTransformKeysForHandle, ObjectBinding, EMovieSceneTransformChannel::All, ESequencerKeyMode::ManualKey )
