Sequencer: Split Event track menu to triggers/repeaters instead of having to always create a trigger section and then a repeater.

diff --git a/Engine/Source/Editor/MovieSceneTools/Private/TrackEditors/EventTrackEditor.cpp b/Engine/Source/Editor/MovieSceneTools/Private/TrackEditors/EventTrackEditor.cpp
--- a/Engine/Source/Editor/MovieSceneTools/Private/TrackEditors/EventTrackEditor.cpp
+++ b/Engine/Source/Editor/MovieSceneTools/Private/TrackEditors/EventTrackEditor.cpp
@@ -65,6 +65,27 @@ FEventTrackEditor::FEventTrackEditor(TSharedRef<ISequencer> InSequencer)
 /* ISequencerTrackEditor interface
  *****************************************************************************/
 
+void FEventTrackEditor::AddEventSubMenu(FMenuBuilder& MenuBuilder, TArray<FGuid> ObjectBindings)
+{
+	MenuBuilder.AddMenuEntry(
+		LOCTEXT("AddNewTriggerSection", "Trigger"),
+		LOCTEXT("AddNewTriggerSectionTooltip", "Adds a new section that can trigger a specific event at a specific time"),
+		FSlateIcon(),
+		FUIAction(
+			FExecuteAction::CreateRaw(this, &FEventTrackEditor::HandleAddEventTrackMenuEntryExecute, ObjectBindings, UMovieSceneEventTriggerSection::StaticClass())
+		)
+	);
+
+	MenuBuilder.AddMenuEntry(
+		LOCTEXT("AddNewRepeaterSection", "Repeater"),
+		LOCTEXT("AddNewRepeaterSectionTooltip", "Adds a new section that triggers an event every time it's evaluated"),
+		FSlateIcon(),
+		FUIAction(
+			FExecuteAction::CreateRaw(this, &FEventTrackEditor::HandleAddEventTrackMenuEntryExecute, ObjectBindings, UMovieSceneEventRepeaterSection::StaticClass())
+		)
+	);
+}
+
 void FEventTrackEditor::BuildAddTrackMenu(FMenuBuilder& MenuBuilder)
 {
 	UMovieSceneSequence*       RootMovieSceneSequence = GetSequencer()->GetRootMovieSceneSequence();
@@ -72,13 +93,12 @@ void FEventTrackEditor::BuildAddTrackMenu(FMenuBuilder& MenuBuilder)
 
 	if (SequenceEditor && SequenceEditor->SupportsEvents(RootMovieSceneSequence))
 	{
-		MenuBuilder.AddMenuEntry(
+		MenuBuilder.AddSubMenu(
 			LOCTEXT("AddEventTrack", "Event Track"),
 			LOCTEXT("AddEventTooltip", "Adds a new event track that can trigger events on the timeline."),
-			FSlateIcon(FEditorStyle::GetStyleSetName(), "Sequencer.Tracks.Event"),
-			FUIAction(
-				FExecuteAction::CreateRaw(this, &FEventTrackEditor::HandleAddEventTrackMenuEntryExecute, TArray<FGuid>())
-			)
+			FNewMenuDelegate::CreateRaw(this, &FEventTrackEditor::AddEventSubMenu, TArray<FGuid>()),
+			false,
+			FSlateIcon(FEditorStyle::GetStyleSetName(), "Sequencer.Tracks.Event")
 		);
 	}
 }
@@ -90,13 +110,10 @@ void FEventTrackEditor::BuildObjectBindingTrackMenu(FMenuBuilder& MenuBuilder, c
 
 	if (SequenceEditor && SequenceEditor->SupportsEvents(RootMovieSceneSequence))
 	{
-		MenuBuilder.AddMenuEntry(
+		MenuBuilder.AddSubMenu(
 			LOCTEXT("AddEventTrack_ObjectBinding", "Event"),
 			LOCTEXT("AddEventTooltip_ObjectBinding", "Adds a new event track that will trigger events on this object binding."),
-			FSlateIcon(),
-			FUIAction( 
-				FExecuteAction::CreateSP( this, &FEventTrackEditor::HandleAddEventTrackMenuEntryExecute, ObjectBindings )
-			)
+			FNewMenuDelegate::CreateRaw(this, &FEventTrackEditor::AddEventSubMenu, ObjectBindings)
 		);
 	}
 }
@@ -229,7 +246,7 @@ const FSlateBrush* FEventTrackEditor::GetIconBrush() const
 /* FEventTrackEditor callbacks
  *****************************************************************************/
 
-void FEventTrackEditor::HandleAddEventTrackMenuEntryExecute(TArray<FGuid> InObjectBindingIDs)
+void FEventTrackEditor::HandleAddEventTrackMenuEntryExecute(TArray<FGuid> InObjectBindingIDs, UClass* SectionType)
 {
 	UMovieScene* FocusedMovieScene = GetFocusedMovieScene();
 
@@ -265,10 +282,8 @@ void FEventTrackEditor::HandleAddEventTrackMenuEntryExecute(TArray<FGuid> InObje
 
 	for (UMovieSceneEventTrack* NewTrack : NewTracks)
 	{
-		UMovieSceneSection* NewSection = NewTrack->CreateNewSection();
-		check(NewSection);
+		CreateNewSection(NewTrack, 0, SectionType);
 
-		NewTrack->AddSection(*NewSection);
 		NewTrack->SetDisplayName(LOCTEXT("TrackName", "Events"));
 
 		if (GetSequencer().IsValid())
@@ -285,8 +300,8 @@ void FEventTrackEditor::CreateNewSection(UMovieSceneTrack* Track, int32 RowIndex
 	TSharedPtr<ISequencer> SequencerPtr = GetSequencer();
 	if (SequencerPtr.IsValid())
 	{
+		UMovieScene* FocusedMovieScene = GetFocusedMovieScene();
 		FQualifiedFrameTime CurrentTime = SequencerPtr->GetLocalTime();
-		TRange<double> VisibleRange = SequencerPtr->GetViewRange();
 
 		FScopedTransaction Transaction(LOCTEXT("CreateNewSectionTransactionText", "Add Section"));
 
@@ -305,8 +320,15 @@ void FEventTrackEditor::CreateNewSection(UMovieSceneTrack* Track, int32 RowIndex
 
 		Track->Modify();
 
-		int32 DurationFrames = ( (VisibleRange.Size<double>() * 0.75) * CurrentTime.Rate ).FloorToFrame().Value;
-		NewSection->SetRange(TRange<FFrameNumber>(CurrentTime.Time.FrameNumber, CurrentTime.Time.FrameNumber + DurationFrames));
+		if (SectionType == UMovieSceneEventTriggerSection::StaticClass())
+		{
+			NewSection->SetRange(TRange<FFrameNumber>::All());
+		}
+		else
+		{
+			NewSection->SetRange(TRange<FFrameNumber>(CurrentTime.Time.FrameNumber, FocusedMovieScene->GetPlaybackRange().GetUpperBoundValue()));
+		}
+
 		NewSection->SetOverlapPriority(OverlapPriority);
 		NewSection->SetRowIndex(RowIndex);
 
diff --git a/Engine/Source/Editor/MovieSceneTools/Private/TrackEditors/EventTrackEditor.h b/Engine/Source/Editor/MovieSceneTools/Private/TrackEditors/EventTrackEditor.h
--- a/Engine/Source/Editor/MovieSceneTools/Private/TrackEditors/EventTrackEditor.h
+++ b/Engine/Source/Editor/MovieSceneTools/Private/TrackEditors/EventTrackEditor.h
@@ -55,8 +55,10 @@ public:
 
 private:
 
+	void AddEventSubMenu(FMenuBuilder& MenuBuilder, TArray<FGuid>);
+
 	/** Callback for executing the "Add Event Track" menu entry. */
-	void HandleAddEventTrackMenuEntryExecute(TArray<FGuid> InObjectBindingIDs);
+	void HandleAddEventTrackMenuEntryExecute(TArray<FGuid> InObjectBindingIDs, UClass* SectionType);
 
 	void CreateNewSection(UMovieSceneTrack* Track, int32 RowIndex, UClass* SectionType);
 };
