Sequencer: Select sections when they are created, don't select them when tracks are created

diff --git a/Engine/Source/Editor/MovieSceneTools/Private/TrackEditors/EventTrackEditor.cpp b/Engine/Source/Editor/MovieSceneTools/Private/TrackEditors/EventTrackEditor.cpp
--- a/Engine/Source/Editor/MovieSceneTools/Private/TrackEditors/EventTrackEditor.cpp
+++ b/Engine/Source/Editor/MovieSceneTools/Private/TrackEditors/EventTrackEditor.cpp
@@ -141,14 +141,14 @@ TSharedPtr<SWidget> FEventTrackEditor::BuildOutlinerEditWidget(const FGuid& Obje
 				LOCTEXT("AddNewTriggerSection", "Trigger"),
 				LOCTEXT("AddNewTriggerSectionTooltip", "Adds a new section that can trigger a specific event at a specific time"),
 				FSlateIcon(),
-				FUIAction(FExecuteAction::CreateSP(this, &FEventTrackEditor::CreateNewSection, TrackPtr, RowIndex + 1, UMovieSceneEventTriggerSection::StaticClass()))
+				FUIAction(FExecuteAction::CreateSP(this, &FEventTrackEditor::CreateNewSection, TrackPtr, RowIndex + 1, UMovieSceneEventTriggerSection::StaticClass(), true))
 			);
 
 			MenuBuilder.AddMenuEntry(
 				LOCTEXT("AddNewRepeaterSection", "Repeater"),
 				LOCTEXT("AddNewRepeaterSectionTooltip", "Adds a new section that triggers an event every time it's evaluated"),
 				FSlateIcon(),
-				FUIAction(FExecuteAction::CreateSP(this, &FEventTrackEditor::CreateNewSection, TrackPtr, RowIndex + 1, UMovieSceneEventRepeaterSection::StaticClass()))
+				FUIAction(FExecuteAction::CreateSP(this, &FEventTrackEditor::CreateNewSection, TrackPtr, RowIndex + 1, UMovieSceneEventRepeaterSection::StaticClass(), true))
 			);
 		}
 		else
@@ -282,7 +282,7 @@ void FEventTrackEditor::HandleAddEventTrackMenuEntryExecute(TArray<FGuid> InObje
 
 	for (UMovieSceneEventTrack* NewTrack : NewTracks)
 	{
-		CreateNewSection(NewTrack, 0, SectionType);
+		CreateNewSection(NewTrack, 0, SectionType, false);
 
 		NewTrack->SetDisplayName(LOCTEXT("TrackName", "Events"));
 
@@ -295,7 +295,7 @@ void FEventTrackEditor::HandleAddEventTrackMenuEntryExecute(TArray<FGuid> InObje
 	GetSequencer()->NotifyMovieSceneDataChanged( EMovieSceneDataChangeType::MovieSceneStructureItemAdded );
 }
 
-void FEventTrackEditor::CreateNewSection(UMovieSceneTrack* Track, int32 RowIndex, UClass* SectionType)
+void FEventTrackEditor::CreateNewSection(UMovieSceneTrack* Track, int32 RowIndex, UClass* SectionType, bool bSelect)
 {
 	TSharedPtr<ISequencer> SequencerPtr = GetSequencer();
 	if (SequencerPtr.IsValid())
@@ -335,10 +335,14 @@ void FEventTrackEditor::CreateNewSection(UMovieSceneTrack* Track, int32 RowIndex
 		Track->AddSection(*NewSection);
 		Track->UpdateEasing();
 
+		if (bSelect)
+		{
+			SequencerPtr->EmptySelection();
+			SequencerPtr->SelectSection(NewSection);
+			SequencerPtr->ThrobSectionSelection();
+		}
+
 		SequencerPtr->NotifyMovieSceneDataChanged(EMovieSceneDataChangeType::MovieSceneStructureItemAdded);
-		SequencerPtr->EmptySelection();
-		SequencerPtr->SelectSection(NewSection);
-		SequencerPtr->ThrobSectionSelection();
 	}
 }
 
diff --git a/Engine/Source/Editor/MovieSceneTools/Private/TrackEditors/EventTrackEditor.h b/Engine/Source/Editor/MovieSceneTools/Private/TrackEditors/EventTrackEditor.h
--- a/Engine/Source/Editor/MovieSceneTools/Private/TrackEditors/EventTrackEditor.h
+++ b/Engine/Source/Editor/MovieSceneTools/Private/TrackEditors/EventTrackEditor.h
@@ -60,5 +60,5 @@ private:
 	/** Callback for executing the "Add Event Track" menu entry. */
 	void HandleAddEventTrackMenuEntryExecute(TArray<FGuid> InObjectBindingIDs, UClass* SectionType);
 
-	void CreateNewSection(UMovieSceneTrack* Track, int32 RowIndex, UClass* SectionType);
+	void CreateNewSection(UMovieSceneTrack* Track, int32 RowIndex, UClass* SectionType, bool bSelect);
 };
