Sequencer: Clamp event text position to 0

diff --git a/Engine/Source/Editor/MovieSceneTools/Private/Sections/EventSection.cpp b/Engine/Source/Editor/MovieSceneTools/Private/Sections/EventSection.cpp
--- a/Engine/Source/Editor/MovieSceneTools/Private/Sections/EventSection.cpp
+++ b/Engine/Source/Editor/MovieSceneTools/Private/Sections/EventSection.cpp
@@ -60,6 +60,10 @@ void FEventSectionBase::PaintEventName(FSequencerSectionPainter& Painter, int32
 	// Flip the text position if getting near the end of the view range
 	bool  bDrawLeft    = (Painter.SectionGeometry.Size.X - PixelPos) < (BoxSize.X + 22.f) - BoxOffsetPx;
 	float BoxPositionX = bDrawLeft ? PixelPos - BoxSize.X - BoxOffsetPx : PixelPos + BoxOffsetPx;
+	if (BoxPositionX < 0.f)
+	{
+		BoxPositionX = 0.f;
+	}
 
 	FVector2D BoxOffset  = FVector2D(BoxPositionX,                    Painter.SectionGeometry.Size.Y*.5f - BoxSize.Y*.5f);
 	FVector2D IconOffset = FVector2D(BoxPadding.Left,                 BoxSize.Y*.5f - IconSize.Y*.5f);
