Sequencer: Fixed key structs not working correctly for object channel keys

  - Until recently object path channels had key values of type TSoftObjectPtr<UObject>, but this was changed to be a struct type that contained both a soft object ptr and a strong reference in order to correctly support asnc loads (by hinting the async loader that each asset referenced by the channel was required before the package load completes).
  - This caused a disparity between the key data type and the type of the generated key struct being used to edit it.
  - A new PostConstructKeyInstance function can now be overloaded to provide additional behavior on instantiation of a new key struct which is used here to assign the UObject* to the key value correctly usding its overloaded assignment operator.

diff --git a/Engine/Source/Editor/MovieSceneTools/Private/Channels/BuiltInChannelEditors.cpp b/Engine/Source/Editor/MovieSceneTools/Private/Channels/BuiltInChannelEditors.cpp
--- a/Engine/Source/Editor/MovieSceneTools/Private/Channels/BuiltInChannelEditors.cpp
+++ b/Engine/Source/Editor/MovieSceneTools/Private/Channels/BuiltInChannelEditors.cpp
@@ -312,6 +312,31 @@ UMovieSceneKeyStructType* InstanceGeneratedStruct(FMovieSceneObjectPathChannel*
 	return NewStruct;
 }
 
+void PostConstructKeyInstance(const TMovieSceneChannelHandle<FMovieSceneObjectPathChannel>& ChannelHandle, FKeyHandle InHandle, FStructOnScope* Struct)
+{	
+	const UMovieSceneKeyStructType* GeneratedStructType = CastChecked<const UMovieSceneKeyStructType>(Struct->GetStruct());
+
+	USoftObjectProperty* EditProperty = CastChecked<USoftObjectProperty>(GeneratedStructType->DestValueProperty);
+	const uint8* PropertyAddress = EditProperty->ContainerPtrToValuePtr<uint8>(Struct->GetStructMemory());
+
+	// It is safe to capture the property and address in this lambda because the lambda is owned by the struct itself, so cannot be invoked if the struct has been destroyed
+	auto CopyInstanceToKeyLambda = [ChannelHandle, InHandle, EditProperty, PropertyAddress](const FPropertyChangedEvent&)
+	{
+		if (FMovieSceneObjectPathChannel* DestinationChannel = ChannelHandle.Get())
+		{
+			const int32 KeyIndex = DestinationChannel->GetData().GetIndex(InHandle);
+			if (KeyIndex != INDEX_NONE)
+			{
+				UObject* ObjectPropertyValue = EditProperty->GetObjectPropertyValue(PropertyAddress);
+				DestinationChannel->GetData().GetValues()[KeyIndex] = ObjectPropertyValue;
+			}
+		}
+	};
+
+	FGeneratedMovieSceneKeyStruct* KeyStruct = reinterpret_cast<FGeneratedMovieSceneKeyStruct*>(Struct->GetStructMemory());
+	KeyStruct->OnPropertyChangedEvent = CopyInstanceToKeyLambda;
+}
+
 void DrawKeys(FMovieSceneFloatChannel* Channel, TArrayView<const FKeyHandle> InKeyHandles, TArrayView<FKeyDrawParams> OutKeyDrawParams)
 {
 	static const FName CircleKeyBrushName("Sequencer.KeyCircle");
diff --git a/Engine/Source/Editor/MovieSceneTools/Private/Channels/BuiltInChannelEditors.h b/Engine/Source/Editor/MovieSceneTools/Private/Channels/BuiltInChannelEditors.h
--- a/Engine/Source/Editor/MovieSceneTools/Private/Channels/BuiltInChannelEditors.h
+++ b/Engine/Source/Editor/MovieSceneTools/Private/Channels/BuiltInChannelEditors.h
@@ -55,6 +55,8 @@ TSharedRef<SWidget> CreateKeyEditor(const TMovieSceneChannelHandle<FMovieSceneOb
 UMovieSceneKeyStructType* InstanceGeneratedStruct(FMovieSceneByteChannel* Channel,       FSequencerKeyStructGenerator* Generator);
 UMovieSceneKeyStructType* InstanceGeneratedStruct(FMovieSceneObjectPathChannel* Channel, FSequencerKeyStructGenerator* Generator);
 
+void PostConstructKeyInstance(const TMovieSceneChannelHandle<FMovieSceneObjectPathChannel>& ChannelHandle, FKeyHandle InHandle, FStructOnScope* Struct);
+
 /** Key drawing overrides */
 void DrawKeys(FMovieSceneFloatChannel*    Channel, TArrayView<const FKeyHandle> InKeyHandles, TArrayView<FKeyDrawParams> OutKeyDrawParams);
 void DrawKeys(FMovieSceneParticleChannel* Channel, TArrayView<const FKeyHandle> InKeyHandles, TArrayView<FKeyDrawParams> OutKeyDrawParams);
diff --git a/Engine/Source/Editor/Sequencer/Public/SequencerKeyStructGenerator.h b/Engine/Source/Editor/Sequencer/Public/SequencerKeyStructGenerator.h
--- a/Engine/Source/Editor/Sequencer/Public/SequencerKeyStructGenerator.h
+++ b/Engine/Source/Editor/Sequencer/Public/SequencerKeyStructGenerator.h
@@ -164,6 +164,15 @@ inline UMovieSceneKeyStructType* InstanceGeneratedStruct(void* Channel, FSequenc
 	return nullptr;
 }
 
+/**
+ * Called to initialize a newly allocated key struct for editing.
+ * Empty by default, but can be overridden to perform any per-instance setup required for specific channel's key structs
+ */
+template<typename ChannelType>
+void PostConstructKeyInstance(const TMovieSceneChannelHandle<ChannelType>& ChannelHandle, FKeyHandle InHandle, FStructOnScope* Struct)
+{
+}
+
 template<typename ChannelType>
 TSharedPtr<FStructOnScope> FSequencerKeyStructGenerator::CreateKeyStructInstance(const TMovieSceneChannelHandle<ChannelType>& ChannelHandle, FKeyHandle InHandle)
 {
@@ -188,7 +197,10 @@ TSharedPtr<FStructOnScope> FSequencerKeyStructGenerator::CreateKeyStructInstance
 					FSequencerKeyStructGenerator::CopyInstanceToKey(ChannelHandle, InHandle, StructPtr);
 				};
 
-				reinterpret_cast<FGeneratedMovieSceneKeyStruct*>(StructInstance->GetStructMemory())->OnPropertyChangedEvent = CopyInstanceToKeyLambda;
+				FGeneratedMovieSceneKeyStruct* KeyStruct = reinterpret_cast<FGeneratedMovieSceneKeyStruct*>(StructInstance->GetStructMemory());
+				KeyStruct->OnPropertyChangedEvent = CopyInstanceToKeyLambda;
+
+				PostConstructKeyInstance(ChannelHandle, InHandle, StructInstance.Get());
 
 				return StructInstance;
 			}
