Sequencer: Disable editor only functions in the events quick binding

diff --git a/Engine/Source/Editor/MovieSceneTools/Private/MovieSceneEventCustomization.cpp b/Engine/Source/Editor/MovieSceneTools/Private/MovieSceneEventCustomization.cpp
--- a/Engine/Source/Editor/MovieSceneTools/Private/MovieSceneEventCustomization.cpp
+++ b/Engine/Source/Editor/MovieSceneTools/Private/MovieSceneEventCustomization.cpp
@@ -344,7 +344,7 @@ void FMovieSceneEventCustomization::PopulateQuickBindSubMenu(FMenuBuilder& MenuB
 		Functions.Reset();
 		for (UFunction* Function : TFieldRange<UFunction>(SuperClass, EFieldIteratorFlags::ExcludeSuper, EFieldIteratorFlags::ExcludeDeprecated))
 		{
-			if (Function->HasAllFunctionFlags(FUNC_BlueprintCallable|FUNC_Public) && !Function->HasMetaData(DeprecatedFunctionName))
+			if (Function->HasAllFunctionFlags(FUNC_BlueprintCallable | FUNC_Public) && !Function->HasMetaData(DeprecatedFunctionName) && !Function->HasAnyFunctionFlags(FUNC_EditorOnly))
 			{
 				Functions.Add(Function);
 			}
