Mark blueprints as needing to be recompiled when modifying event endpoints

diff --git a/Engine/Source/Editor/MovieSceneTools/Private/MovieSceneEventCustomization.cpp b/Engine/Source/Editor/MovieSceneTools/Private/MovieSceneEventCustomization.cpp
--- a/Engine/Source/Editor/MovieSceneTools/Private/MovieSceneEventCustomization.cpp
+++ b/Engine/Source/Editor/MovieSceneTools/Private/MovieSceneEventCustomization.cpp
@@ -469,6 +469,8 @@ void FMovieSceneEventCustomization::SetEventEndpoint(UK2Node_FunctionEntry* NewE
 
 	// Ensure that anything listening for property changed notifications are notified of the new binding
 	PropertyHandle->NotifyFinishedChangingProperties();
+
+	FBlueprintEditorUtils::MarkBlueprintAsModified(Blueprint);
 }
 
 bool FMovieSceneEventCustomization::CompareCurrentEventEndpoint(UK2Node_FunctionEntry* NewEndpoint)
