Sequencer: Fix duplicate rendered frames

diff --git a/Engine/Source/Editor/MovieSceneTools/Private/AutomatedLevelSequenceCapture.cpp b/Engine/Source/Editor/MovieSceneTools/Private/AutomatedLevelSequenceCapture.cpp
--- a/Engine/Source/Editor/MovieSceneTools/Private/AutomatedLevelSequenceCapture.cpp
+++ b/Engine/Source/Editor/MovieSceneTools/Private/AutomatedLevelSequenceCapture.cpp
@@ -728,7 +728,12 @@ void UAutomatedLevelSequenceCapture::SequenceUpdated(const UMovieSceneSequencePl
 				bool bOnLastFrame = (CurrentTime.FrameNumber >= Actor->SequencePlayer->GetStartTime().Time.FrameNumber + Actor->SequencePlayer->GetFrameDuration() - 1);
 				bool bLastShot = NumShots == 0 ? true : ShotIndex == NumShots - 1;
 				
-				CaptureThisFrame( (CurrentTime - PreviousTime) / Settings.FrameRate);
+				// Prevent the same frame from being rendered twice
+				if (CurrentTime.FrameNumber.Value != CachedMetrics.PreviousFrame)
+				{
+					CaptureThisFrame((CurrentTime - PreviousTime) / Settings.FrameRate);
+					CachedMetrics.PreviousFrame = CurrentTime.FrameNumber.Value;
+				}
 
 				// Our callback can be called multiple times for a given frame due to how Level Sequences evaluate.
 				// For example, frame 161 is evaluated and an image is written. This isn't considered the end of the sequence
diff --git a/Engine/Source/Runtime/MovieSceneCapture/Public/MovieSceneCapture.h b/Engine/Source/Runtime/MovieSceneCapture/Public/MovieSceneCapture.h
--- a/Engine/Source/Runtime/MovieSceneCapture/Public/MovieSceneCapture.h
+++ b/Engine/Source/Runtime/MovieSceneCapture/Public/MovieSceneCapture.h
@@ -20,7 +20,7 @@ class FSceneViewport;
 /** Structure used to cache various metrics for our capture */
 struct FCachedMetrics
 {
-	FCachedMetrics() : Width(0), Height(0), Frame(0), ElapsedSeconds(0.f) {}
+	FCachedMetrics() : Width(0), Height(0), Frame(0), ElapsedSeconds(0.f), PreviousFrame(INDEX_NONE) {}
 
 	/** The width/Height of the frame */
 	int32 Width, Height;
@@ -28,6 +28,8 @@ struct FCachedMetrics
 	int32 Frame;
 	/** The number of seconds that have elapsed */
 	float ElapsedSeconds;
+	/** The previous frame number */
+	int32 PreviousFrame;
 };
 
 /** Class responsible for capturing scene data */
