Sequencer: Actor reference tracks are now used for any soft object actor reference as well

diff --git a/Engine/Source/Editor/MovieSceneTools/Private/TrackEditors/PropertyTrackEditors/ActorReferencePropertyTrackEditor.h b/Engine/Source/Editor/MovieSceneTools/Private/TrackEditors/PropertyTrackEditors/ActorReferencePropertyTrackEditor.h
--- a/Engine/Source/Editor/MovieSceneTools/Private/TrackEditors/PropertyTrackEditors/ActorReferencePropertyTrackEditor.h
+++ b/Engine/Source/Editor/MovieSceneTools/Private/TrackEditors/PropertyTrackEditors/ActorReferencePropertyTrackEditor.h
@@ -35,7 +35,10 @@ public:
 	 */
 	static TArray<FAnimatedPropertyKey, TInlineAllocator<1>> GetAnimatedPropertyTypes()
 	{
-		return TArray<FAnimatedPropertyKey, TInlineAllocator<1>>({ FAnimatedPropertyKey::FromObjectType(AActor::StaticClass()) });
+		FAnimatedPropertyKey Key = FAnimatedPropertyKey::FromPropertyType(USoftObjectProperty::StaticClass());
+		Key.ObjectTypeName = AActor::StaticClass()->GetFName();
+
+		return TArray<FAnimatedPropertyKey, TInlineAllocator<1>>({ Key, FAnimatedPropertyKey::FromObjectType(AActor::StaticClass()) });
 	}
 
 	/**
