Sequencer: Wrap quick binding submenu in a max size box

diff --git a/Engine/Source/Editor/MovieSceneTools/Private/MovieSceneEventCustomization.cpp b/Engine/Source/Editor/MovieSceneTools/Private/MovieSceneEventCustomization.cpp
--- a/Engine/Source/Editor/MovieSceneTools/Private/MovieSceneEventCustomization.cpp
+++ b/Engine/Source/Editor/MovieSceneTools/Private/MovieSceneEventCustomization.cpp
@@ -16,6 +16,7 @@
 #include "Widgets/DeclarativeSyntaxSupport.h"
 #include "Widgets/Text/STextBlock.h"
 #include "Widgets/Images/SImage.h"
+#include "Widgets/Layout/SBox.h"
 #include "Widgets/Input/SComboButton.h"
 #include "Widgets/Input/SHyperlink.h"
 #include "Widgets/Input/SEditableTextBox.h"
@@ -331,6 +332,8 @@ TSharedRef<SWidget> FMovieSceneEventCustomization::GetMenuContent()
 
 void FMovieSceneEventCustomization::PopulateQuickBindSubMenu(FMenuBuilder& MenuBuilder, UClass* TemplateClass)
 {
+	FMenuBuilder SubMenuBuilder(true, nullptr, nullptr, true);
+
 	FSlateIcon Icon(FEditorStyle::GetStyleSetName(), "GraphEditor.Function_16x");
 	static const FName DeprecatedFunctionName(TEXT("DeprecatedFunction"));
 
@@ -339,7 +342,7 @@ void FMovieSceneEventCustomization::PopulateQuickBindSubMenu(FMenuBuilder& MenuB
 	UClass* SuperClass = TemplateClass;
 	while (SuperClass)
 	{
-		MenuBuilder.BeginSection(NAME_None, SuperClass->GetDisplayNameText());
+		SubMenuBuilder.BeginSection(NAME_None, SuperClass->GetDisplayNameText());
 
 		Functions.Reset();
 		for (UFunction* Function : TFieldRange<UFunction>(SuperClass, EFieldIteratorFlags::ExcludeSuper, EFieldIteratorFlags::ExcludeDeprecated))
@@ -354,7 +357,7 @@ void FMovieSceneEventCustomization::PopulateQuickBindSubMenu(FMenuBuilder& MenuB
 
 		for (UFunction* Function : Functions)
 		{
-			MenuBuilder.AddMenuEntry(
+			SubMenuBuilder.AddMenuEntry(
 				FText::FromName(Function->GetFName()),
 				FText(),
 				Icon,
@@ -364,10 +367,20 @@ void FMovieSceneEventCustomization::PopulateQuickBindSubMenu(FMenuBuilder& MenuB
 			);
 		}
 
-		MenuBuilder.EndSection();
+		SubMenuBuilder.EndSection();
 
 		SuperClass = SuperClass->GetSuperClass();
 	}
+
+	// Limit the height intentionally so that it doesn't go offscreen
+	MenuBuilder.AddWidget(
+		SNew(SBox)
+		.MaxDesiredHeight(600.0f)
+		[
+			SubMenuBuilder.MakeWidget()
+		]
+		, FText(), true, false
+	);
 }
 
 const FSlateBrush* FMovieSceneEventCustomization::GetEventIcon() const
