Sequencer: Fix crash when adding an event repeater section outside the playback range

diff --git a/Engine/Source/Editor/MovieSceneTools/Private/TrackEditors/EventTrackEditor.cpp b/Engine/Source/Editor/MovieSceneTools/Private/TrackEditors/EventTrackEditor.cpp
--- a/Engine/Source/Editor/MovieSceneTools/Private/TrackEditors/EventTrackEditor.cpp
+++ b/Engine/Source/Editor/MovieSceneTools/Private/TrackEditors/EventTrackEditor.cpp
@@ -326,7 +326,20 @@ void FEventTrackEditor::CreateNewSection(UMovieSceneTrack* Track, int32 RowIndex
 		}
 		else
 		{
-			NewSection->SetRange(TRange<FFrameNumber>(CurrentTime.Time.FrameNumber, FocusedMovieScene->GetPlaybackRange().GetUpperBoundValue()));
+
+			TRange<FFrameNumber> NewSectionRange;
+
+			if (CurrentTime.Time.FrameNumber < FocusedMovieScene->GetPlaybackRange().GetUpperBoundValue())
+			{
+				NewSectionRange = TRange<FFrameNumber>(CurrentTime.Time.FrameNumber, FocusedMovieScene->GetPlaybackRange().GetUpperBoundValue());
+			}
+			else
+			{
+				const float DefaultLengthInSeconds = 5.f;
+				NewSectionRange = TRange<FFrameNumber>(CurrentTime.Time.FrameNumber, CurrentTime.Time.FrameNumber + (DefaultLengthInSeconds * SequencerPtr->GetFocusedTickResolution()).FloorToFrame());
+			}
+
+			NewSection->SetRange(NewSectionRange);
 		}
 
 		NewSection->SetOverlapPriority(OverlapPriority);
