Movie Scene Capture: Add option for the custom frame rate. This allows the custom frame rate to be a config that isn't reset everytime you render a sequence (which sets the frame rate)

diff --git a/Engine/Plugins/MovieScene/SequencerScripting/Source/SequencerScriptingEditor/Private/SequencerTools.cpp b/Engine/Plugins/MovieScene/SequencerScripting/Source/SequencerScriptingEditor/Private/SequencerTools.cpp
--- a/Engine/Plugins/MovieScene/SequencerScripting/Source/SequencerScriptingEditor/Private/SequencerTools.cpp
+++ b/Engine/Plugins/MovieScene/SequencerScripting/Source/SequencerScriptingEditor/Private/SequencerTools.cpp
@@ -53,7 +53,8 @@ bool USequencerToolsFunctionLibrary::RenderMovie(UMovieSceneCapture* InCaptureSe
 			FFrameRate DisplayRate = LevelSequence->GetMovieScene()->GetDisplayRate();
 			FFrameRate TickResolution = LevelSequence->GetMovieScene()->GetTickResolution();
 			
-			LevelSequenceCapture->Settings.FrameRate = DisplayRate;
+			LevelSequenceCapture->Settings.CustomFrameRate = DisplayRate;
+			LevelSequenceCapture->Settings.bUseCustomFrameRate = true;
 			LevelSequenceCapture->Settings.bUseRelativeFrameNumbers = false;
 			TRange<FFrameNumber> Range = LevelSequence->GetMovieScene()->GetPlaybackRange();
 
diff --git a/Engine/Source/Editor/MovieSceneTools/Private/AutomatedLevelSequenceCapture.cpp b/Engine/Source/Editor/MovieSceneTools/Private/AutomatedLevelSequenceCapture.cpp
--- a/Engine/Source/Editor/MovieSceneTools/Private/AutomatedLevelSequenceCapture.cpp
+++ b/Engine/Source/Editor/MovieSceneTools/Private/AutomatedLevelSequenceCapture.cpp
@@ -259,7 +259,7 @@ void UAutomatedLevelSequenceCapture::Initialize(TSharedPtr<FSceneViewport> InVie
 	}
 
 	CaptureState = ELevelSequenceCaptureState::Setup;
-	CaptureStrategy = MakeShareable(new FFixedTimeStepCaptureStrategy(Settings.FrameRate));
+	CaptureStrategy = MakeShareable(new FFixedTimeStepCaptureStrategy(Settings.GetFrameRate()));
 	CaptureStrategy->OnInitialize();
 }
 
@@ -319,7 +319,7 @@ bool UAutomatedLevelSequenceCapture::InitializeShots()
 	CachedPlaybackRange = MovieScene->GetPlaybackRange();
 
 	// Compute handle frames in tick resolution space since that is what the section ranges are defined in
-	FFrameNumber HandleFramesResolutionSpace = ConvertFrameTime(Settings.HandleFrames, Settings.FrameRate, MovieScene->GetTickResolution()).FloorToFrame();
+	FFrameNumber HandleFramesResolutionSpace = ConvertFrameTime(Settings.HandleFrames, Settings.GetFrameRate(), MovieScene->GetTickResolution()).FloorToFrame();
 
 	CinematicShotTrack->SortSections();
 
@@ -434,7 +434,7 @@ bool UAutomatedLevelSequenceCapture::SetupShot(FFrameNumber& StartTime, FFrameNu
 			// We intersect with the CachedPlaybackRange instead of copying the playback range from the shot to handle the case where
 			// the playback range intersected the middle of the shot before we started manipulating ranges. We manually expand the master
 			// Movie Sequence's playback range by the number of handle frames to allow handle frames to work as expected on first/last shot.
-			FFrameNumber HandleFramesResolutionSpace = ConvertFrameTime(Settings.HandleFrames, Settings.FrameRate, MovieScene->GetTickResolution()).FloorToFrame();
+			FFrameNumber HandleFramesResolutionSpace = ConvertFrameTime(Settings.HandleFrames, Settings.GetFrameRate(), MovieScene->GetTickResolution()).FloorToFrame();
 			TRange<FFrameNumber> ExtendedCachedPlaybackRange = MovieScene::ExpandRange(CachedPlaybackRange, HandleFramesResolutionSpace);
 
 			TRange<FFrameNumber> TotalRange = TRange<FFrameNumber>::Intersection(ShotSection->GetRange(), ExtendedCachedPlaybackRange);
@@ -464,8 +464,8 @@ void UAutomatedLevelSequenceCapture::SetupFrameRange()
 				FFrameRate           SourceFrameRate = MovieScene->GetTickResolution();
 				TRange<FFrameNumber> SequenceRange   = MovieScene->GetPlaybackRange();
 
-				FFrameNumber PlaybackStartFrame = ConvertFrameTime(MovieScene::DiscreteInclusiveLower(SequenceRange), SourceFrameRate, Settings.FrameRate).CeilToFrame();
-				FFrameNumber PlaybackEndFrame   = ConvertFrameTime(MovieScene::DiscreteExclusiveUpper(SequenceRange), SourceFrameRate, Settings.FrameRate).CeilToFrame();
+				FFrameNumber PlaybackStartFrame = ConvertFrameTime(MovieScene::DiscreteInclusiveLower(SequenceRange), SourceFrameRate, Settings.GetFrameRate()).CeilToFrame();
+				FFrameNumber PlaybackEndFrame   = ConvertFrameTime(MovieScene::DiscreteExclusiveUpper(SequenceRange), SourceFrameRate, Settings.GetFrameRate()).CeilToFrame();
 
 				if( bUseCustomStartFrame )
 				{
@@ -493,7 +493,7 @@ void UAutomatedLevelSequenceCapture::SetupFrameRange()
 				}
 
 				// Override the movie scene's playback range
-				Actor->SequencePlayer->SetFrameRate(Settings.FrameRate);
+				Actor->SequencePlayer->SetFrameRate(Settings.GetFrameRate());
 				Actor->SequencePlayer->SetFrameRange(PlaybackStartFrame.Value, (PlaybackEndFrame - PlaybackStartFrame).Value);
 				Actor->SequencePlayer->JumpToFrame(PlaybackStartFrame.Value);
 
@@ -574,7 +574,7 @@ void UAutomatedLevelSequenceCapture::OnTick(float DeltaSeconds)
 	}
 	else if( CaptureState == ELevelSequenceCaptureState::ReadyToWarmUp )
 	{
-		Actor->SequencePlayer->SetSnapshotSettings(FLevelSequenceSnapshotSettings(Settings.ZeroPadFrameNumbers, Settings.FrameRate));
+		Actor->SequencePlayer->SetSnapshotSettings(FLevelSequenceSnapshotSettings(Settings.ZeroPadFrameNumbers, Settings.GetFrameRate()));
 		Actor->SequencePlayer->Play();
 		// Start warming up
 		CaptureState = ELevelSequenceCaptureState::WarmingUp;
@@ -600,8 +600,8 @@ void UAutomatedLevelSequenceCapture::OnTick(float DeltaSeconds)
 		{
 			UMovieScene* MovieScene = GetMovieScene(LevelSequenceActor);
 
-			FFrameNumber StartTimePlayRateSpace = ConvertFrameTime(StartTime, MovieScene->GetTickResolution(), Settings.FrameRate).CeilToFrame();
-			FFrameNumber EndTimePlayRateSpace   = ConvertFrameTime(EndTime,   MovieScene->GetTickResolution(), Settings.FrameRate).CeilToFrame();
+			FFrameNumber StartTimePlayRateSpace = ConvertFrameTime(StartTime, MovieScene->GetTickResolution(), Settings.GetFrameRate()).CeilToFrame();
+			FFrameNumber EndTimePlayRateSpace   = ConvertFrameTime(EndTime,   MovieScene->GetTickResolution(), Settings.GetFrameRate()).CeilToFrame();
 
 			Actor->SequencePlayer->SetFrameRange(StartTimePlayRateSpace.Value, (EndTimePlayRateSpace - StartTimePlayRateSpace).Value);
 			Actor->SequencePlayer->JumpToFrame(StartTimePlayRateSpace.Value);
@@ -747,7 +747,7 @@ void UAutomatedLevelSequenceCapture::SequenceUpdated(const UMovieSceneSequencePl
 				// Prevent the same frame from being rendered twice
 				if (CurrentTime.FrameNumber.Value != CachedMetrics.PreviousFrame)
 				{
-					CaptureThisFrame((CurrentTime - PreviousTime) / Settings.FrameRate);
+					CaptureThisFrame((CurrentTime - PreviousTime) / Settings.GetFrameRate());
 					CachedMetrics.PreviousFrame = CurrentTime.FrameNumber.Value;
 				}
 
@@ -928,7 +928,7 @@ void UAutomatedLevelSequenceCapture::ExportEDL()
 	int32 HandleFrames = Settings.HandleFrames;
 	FString MovieExtension = Settings.MovieExtension;
 
-	MovieSceneTranslatorEDL::ExportEDL(MovieScene, Settings.FrameRate, SaveFilename, HandleFrames, MovieExtension);
+	MovieSceneTranslatorEDL::ExportEDL(MovieScene, Settings.GetFrameRate(), SaveFilename, HandleFrames, MovieExtension);
 }
 
 double UAutomatedLevelSequenceCapture::GetEstimatedCaptureDurationSeconds() const
@@ -968,7 +968,7 @@ void UAutomatedLevelSequenceCapture::ExportFCPXML()
 	FString SaveFilename = Settings.OutputDirectory.Path / MovieScene->GetOuter()->GetName() + TEXT(".xml");
 	FString FilenameFormat = Settings.OutputFormat;
 	int32 HandleFrames = Settings.HandleFrames;
-	FFrameRate FrameRate = Settings.FrameRate;
+	FFrameRate FrameRate = Settings.GetFrameRate();
 	uint32 ResX = Settings.Resolution.ResX;
 	uint32 ResY = Settings.Resolution.ResY;
 	FString MovieExtension = Settings.MovieExtension;
diff --git a/Engine/Source/Editor/MovieSceneTools/Private/MovieSceneToolHelpers.cpp b/Engine/Source/Editor/MovieSceneTools/Private/MovieSceneToolHelpers.cpp
--- a/Engine/Source/Editor/MovieSceneTools/Private/MovieSceneToolHelpers.cpp
+++ b/Engine/Source/Editor/MovieSceneTools/Private/MovieSceneToolHelpers.cpp
@@ -752,7 +752,7 @@ bool MovieSceneToolHelpers::MovieSceneTranslatorExport(FMovieSceneExporter* InEx
 	int32 HandleFrames = Settings.HandleFrames;
 	// @todo: generate filename based on filename format, currently outputs {shot}.avi
 	FString FilenameFormat = Settings.OutputFormat;
-	FFrameRate FrameRate = Settings.FrameRate;
+	FFrameRate FrameRate = Settings.GetFrameRate();
 	uint32 ResX = Settings.Resolution.ResX;
 	uint32 ResY = Settings.Resolution.ResY;
 	FString MovieExtension = Settings.MovieExtension;
diff --git a/Engine/Source/Runtime/MovieSceneCapture/Private/LevelCapture.cpp b/Engine/Source/Runtime/MovieSceneCapture/Private/LevelCapture.cpp
--- a/Engine/Source/Runtime/MovieSceneCapture/Private/LevelCapture.cpp
+++ b/Engine/Source/Runtime/MovieSceneCapture/Private/LevelCapture.cpp
@@ -17,7 +17,7 @@ void ULevelCapture::SetPrerequisiteActor(AActor* InPrerequisiteActor)
 
 void ULevelCapture::Initialize(TSharedPtr<FSceneViewport> InViewport, int32 PIEInstance)
 {
-	CaptureStrategy = MakeShareable(new FFixedTimeStepCaptureStrategy(Settings.FrameRate));
+	CaptureStrategy = MakeShareable(new FFixedTimeStepCaptureStrategy(Settings.GetFrameRate()));
 	CaptureStrategy->OnInitialize();
 	Super::Initialize(InViewport, PIEInstance);
 
diff --git a/Engine/Source/Runtime/MovieSceneCapture/Private/MovieSceneCapture.cpp b/Engine/Source/Runtime/MovieSceneCapture/Private/MovieSceneCapture.cpp
--- a/Engine/Source/Runtime/MovieSceneCapture/Private/MovieSceneCapture.cpp
+++ b/Engine/Source/Runtime/MovieSceneCapture/Private/MovieSceneCapture.cpp
@@ -86,6 +86,8 @@ FMovieSceneCaptureSettings::FMovieSceneCaptureSettings()
 	HandleFrames = 0;
 	GameModeOverride = nullptr;
 	OutputFormat = TEXT("{world}");
+	bUseCustomFrameRate = false;
+	CustomFrameRate = FFrameRate(24, 1);
 	FrameRate = FFrameRate(24, 1);
 	ZeroPadFrameNumbers = 4;
 	bEnableTextureStreaming = false;
@@ -378,7 +380,8 @@ void UMovieSceneCapture::Initialize(TSharedPtr<FSceneViewport> InSceneViewport,
 		FString FrameRateOverrideString;
 		if ( FParse::Value( FCommandLine::Get(), TEXT( "-MovieFrameRate=" ), FrameRateOverrideString ) )
 		{
-			if (!TryParseString(Settings.FrameRate, *FrameRateOverrideString))
+			Settings.bUseCustomFrameRate = true;
+			if (!TryParseString(Settings.CustomFrameRate, *FrameRateOverrideString))
 			{
 				UE_LOG(LogMovieSceneCapture, Error, TEXT("Unrecognized capture frame rate: %s."), *FrameRateOverrideString);
 			}
@@ -399,7 +402,7 @@ void UMovieSceneCapture::Initialize(TSharedPtr<FSceneViewport> InSceneViewport,
 	CachedMetrics.Width = InitSettings->DesiredSize.X;
 	CachedMetrics.Height = InitSettings->DesiredSize.Y;
 
-	double FrameRate = Settings.FrameRate.AsDecimal();
+	double FrameRate = Settings.GetFrameRate().AsDecimal();
 
 	FormatMappings.Reserve(10);
 	if (FrameRate == FMath::RoundToDouble(FrameRate))
diff --git a/Engine/Source/Runtime/MovieSceneCapture/Public/MovieSceneCaptureSettings.h b/Engine/Source/Runtime/MovieSceneCapture/Public/MovieSceneCaptureSettings.h
--- a/Engine/Source/Runtime/MovieSceneCapture/Public/MovieSceneCaptureSettings.h
+++ b/Engine/Source/Runtime/MovieSceneCapture/Public/MovieSceneCaptureSettings.h
@@ -75,8 +75,16 @@ struct MOVIESCENECAPTURE_API FMovieSceneCaptureSettings
 	UPROPERTY(config, EditAnywhere, BlueprintReadWrite, Category=General, AdvancedDisplay, meta=(ClampMin=0, UIMin=0))
 	uint8 ZeroPadFrameNumbers;
 
-	/** The frame rate at which to capture */
-	UPROPERTY(config, EditAnywhere, BlueprintReadWrite, Category=CaptureSettings, meta=(ClampMin=1, UIMin=1, ClampMax=200, UIMax=200))
+	/** Specify using the custom frame rate as opposed to the sequence's display rate */
+	UPROPERTY(config, EditAnywhere, BlueprintReadWrite, Category=CaptureSettings)
+	bool bUseCustomFrameRate;
+
+	/** The custom frame rate at which to capture if "Use Custom Frame Rate" is enabled */
+	UPROPERTY(config, EditAnywhere, BlueprintReadWrite, Category=CaptureSettings, meta = (ClampMin = 1, UIMin = 1, ClampMax = 200, UIMax = 200, EditCondition=bUseCustomFrameRate))
+	FFrameRate CustomFrameRate;
+
+	/** The sequence's frame rate at which to capture if "Use Custom Frame Rate" is not enabled */
+	UPROPERTY(transient, VisibleAnywhere, BlueprintReadWrite, Category=CaptureSettings, meta = (ClampMin = 1, UIMin = 1, ClampMax = 200, UIMax = 200))
 	FFrameRate FrameRate;
 
 	/** The resolution at which to capture */
@@ -118,4 +126,7 @@ struct MOVIESCENECAPTURE_API FMovieSceneCaptureSettings
 	/** Number of sampler per pixel to be used when rendering the scene with the path tracer (if supported) */
 	UPROPERTY(config, EditAnywhere, BlueprintReadWrite, Category = Cinematic, AdvancedDisplay, meta = (EditCondition = "bUsePathTracer", ClampMin = 1, UIMin = 1, UIMax = 4096))
 	int32 PathTracerSamplePerPixel;
+
+	/** Get which frame rate to use */
+	FFrameRate GetFrameRate() const { return bUseCustomFrameRate ? CustomFrameRate : FrameRate; }
 };
