Bring back ability to key relative tracks via the hotkey(S). Note it may move the keyed object though since relative tracks dont' key like addtivie and absolute layers.
#trivial

diff --git a/Engine/Source/Editor/MovieSceneTools/Public/KeyframeTrackEditor.h b/Engine/Source/Editor/MovieSceneTools/Public/KeyframeTrackEditor.h
--- a/Engine/Source/Editor/MovieSceneTools/Public/KeyframeTrackEditor.h
+++ b/Engine/Source/Editor/MovieSceneTools/Public/KeyframeTrackEditor.h
@@ -558,7 +558,8 @@ private:
 	{
 		FOptionalMovieSceneBlendType BlendType = Section->GetBlendType();
 		// Sections are only eligible for autokey if they are not blendable (or absolute, relative), and overlap the current time
-		return ( !BlendType.IsValid() || BlendType.Get() == EMovieSceneBlendType::Absolute || BlendType.Get() == EMovieSceneBlendType::Additive) && Section->GetRange().Contains(Time);
+		return ( !BlendType.IsValid() || BlendType.Get() == EMovieSceneBlendType::Absolute || BlendType.Get() == EMovieSceneBlendType::Additive ||
+			BlendType.Get() == EMovieSceneBlendType::Relative) && Section->GetRange().Contains(Time);
 	}
 };
 
