Movie Scene Capture: Expose delay every frame to the user. Value is overridden by path tracer sample pixel.

diff --git a/Engine/Source/Editor/MovieSceneTools/Private/AutomatedLevelSequenceCapture.cpp b/Engine/Source/Editor/MovieSceneTools/Private/AutomatedLevelSequenceCapture.cpp
--- a/Engine/Source/Editor/MovieSceneTools/Private/AutomatedLevelSequenceCapture.cpp
+++ b/Engine/Source/Editor/MovieSceneTools/Private/AutomatedLevelSequenceCapture.cpp
@@ -168,7 +168,14 @@ void UAutomatedLevelSequenceCapture::Initialize(TSharedPtr<FSceneViewport> InVie
 
 	if (Settings.bUsePathTracer)
 	{
-		DelayEveryFrame = float(Settings.FrameRate.AsSeconds(Settings.PathTracerSamplePerPixel));
+		float PathTracerSamplePerPixel = float(Settings.FrameRate.AsSeconds(Settings.PathTracerSamplePerPixel));
+
+		if (DelayEveryFrame != PathTracerSamplePerPixel)
+		{
+			UE_LOG(LogMovieSceneCapture, Log, TEXT("Delay every frame overridden by path tracer sample per pixel: %f"), PathTracerSamplePerPixel);
+
+			DelayEveryFrame = PathTracerSamplePerPixel;
+		}
 	}
 
 	ALevelSequenceActor* Actor = LevelSequenceActor.Get();
diff --git a/Engine/Source/Editor/MovieSceneTools/Public/AutomatedLevelSequenceCapture.h b/Engine/Source/Editor/MovieSceneTools/Public/AutomatedLevelSequenceCapture.h
--- a/Engine/Source/Editor/MovieSceneTools/Public/AutomatedLevelSequenceCapture.h
+++ b/Engine/Source/Editor/MovieSceneTools/Public/AutomatedLevelSequenceCapture.h
@@ -59,6 +59,10 @@ public:
 	UPROPERTY(config, EditAnywhere, BlueprintReadWrite, Category = Animation, AdvancedDisplay, meta = (Units = Seconds, ClampMin = 0))
 	float DelayBeforeShotWarmUp;
 
+	/** The number of seconds to wait (in real-time) at every frame.  Useful for allowing post processing effects to settle down before capturing the animation. */
+	UPROPERTY(config, EditAnywhere, BlueprintReadWrite, Category = Animation, AdvancedDisplay, meta = (Units = Seconds, ClampMin = 0))
+	float DelayEveryFrame;
+
 	UPROPERTY(Transient, EditAnywhere, BlueprintReadWrite, Category=CaptureSettings, AdvancedDisplay, meta=(EditInline))
 	ULevelSequenceBurnInOptions* BurnInOptions;
 
@@ -167,9 +171,6 @@ private:
 	/** The current shot movie that is rendering */
 	int32 ShotIndex;
 
-	/** The number of seconds to wait (in real-time) at every frame.  Useful for allowing post processing effects to settle down before capturing the animation. */
-	float DelayEveryFrame;
-
 	FLevelSequencePlayerSnapshot CachedState;
 
 	FTimerHandle DelayTimer;
