Sequencer: Fix bitmask operations

diff --git a/Engine/Source/Editor/MovieSceneTools/Private/Sections/TransformPropertySection.cpp b/Engine/Source/Editor/MovieSceneTools/Private/Sections/TransformPropertySection.cpp
--- a/Engine/Source/Editor/MovieSceneTools/Private/Sections/TransformPropertySection.cpp
+++ b/Engine/Source/Editor/MovieSceneTools/Private/Sections/TransformPropertySection.cpp
@@ -136,7 +136,7 @@ bool FTransformSection::RequestDeleteCategory(const TArray<FName>& CategoryNameP
 		EMovieSceneTransformChannel Channel = TransformSection->GetMask().GetChannels();
 		EMovieSceneTransformChannel ChannelToRemove = TransformSection->GetMaskByName(CategoryName).GetChannels();
 
-		Channel = Channel ^ ChannelToRemove;
+		Channel &= ~ChannelToRemove;
 
 		TransformSection->SetMask(Channel);
 			
@@ -162,7 +162,7 @@ bool FTransformSection::RequestDeleteKeyArea(const TArray<FName>& KeyAreaNamePat
 		EMovieSceneTransformChannel Channel = TransformSection->GetMask().GetChannels();
 		EMovieSceneTransformChannel ChannelToRemove = TransformSection->GetMaskByName(KeyAreaName).GetChannels();
 
-		Channel = Channel ^ ChannelToRemove;
+		Channel &= ~ChannelToRemove;
 
 		TransformSection->SetMask(Channel);
 					
diff --git a/Engine/Source/Editor/UMGEditor/Private/Animation/Sequencer2DTransformSection.cpp b/Engine/Source/Editor/UMGEditor/Private/Animation/Sequencer2DTransformSection.cpp
--- a/Engine/Source/Editor/UMGEditor/Private/Animation/Sequencer2DTransformSection.cpp
+++ b/Engine/Source/Editor/UMGEditor/Private/Animation/Sequencer2DTransformSection.cpp
@@ -134,7 +134,7 @@ bool F2DTransformSection::RequestDeleteCategory(const TArray<FName>& CategoryNam
 		EMovieScene2DTransformChannel Channel = TransformSection->GetMask().GetChannels();
 		EMovieScene2DTransformChannel ChannelToRemove = TransformSection->GetMaskByName(CategoryName).GetChannels();
 
-		Channel = Channel ^ ChannelToRemove;
+		Channel &= ~ChannelToRemove;
 
 		TransformSection->SetMask(Channel);
 			
@@ -160,7 +160,7 @@ bool F2DTransformSection::RequestDeleteKeyArea(const TArray<FName>& KeyAreaNameP
 		EMovieScene2DTransformChannel Channel = TransformSection->GetMask().GetChannels();
 		EMovieScene2DTransformChannel ChannelToRemove = TransformSection->GetMaskByName(KeyAreaName).GetChannels();
 
-		Channel = Channel ^ ChannelToRemove;
+		Channel &= ~ChannelToRemove;
 
 		TransformSection->SetMask(Channel);
 					
