Sequencer: Added "Select Camera" to the right click context menu of the camera cut section

diff --git a/Engine/Source/Editor/MovieSceneTools/Private/Sections/CameraCutSection.cpp b/Engine/Source/Editor/MovieSceneTools/Private/Sections/CameraCutSection.cpp
--- a/Engine/Source/Editor/MovieSceneTools/Private/Sections/CameraCutSection.cpp
+++ b/Engine/Source/Editor/MovieSceneTools/Private/Sections/CameraCutSection.cpp
@@ -77,7 +77,19 @@ void FCameraCutSection::BuildSectionContextMenu(FMenuBuilder& MenuBuilder, const
 		return;
 	}
 
-	const AActor* CameraActor = GetCameraForFrame(Section->GetInclusiveStartFrame());
+	AActor* CameraActor = GetCameraForFrame(Section->GetInclusiveStartFrame());
+
+	if (CameraActor)
+	{
+		MenuBuilder.AddMenuSeparator();
+
+		MenuBuilder.AddMenuEntry(
+			FText::Format(LOCTEXT("SelectCameraTextFormat", "Select {0}"), FText::FromString(CameraActor->GetActorLabel())),
+			FText::Format(LOCTEXT("SelectCameraTooltipFormat", "Select {0}"), FText::FromString(CameraActor->GetActorLabel())),
+			FSlateIcon(),
+			FUIAction(FExecuteAction::CreateRaw(this, &FCameraCutSection::HandleSelectCameraMenuEntryExecute, CameraActor))
+		);
+	}
 
 	// get list of available cameras
 	TArray<AActor*> AllCameras;
@@ -183,6 +195,11 @@ FText FCameraCutSection::HandleThumbnailTextBlockText() const
 /* FCameraCutSection callbacks
  *****************************************************************************/
 
+void FCameraCutSection::HandleSelectCameraMenuEntryExecute(AActor* InCamera)
+{
+	GEditor->SelectActor(InCamera, true, true);
+}
+
 void FCameraCutSection::HandleSetCameraMenuEntryExecute(AActor* InCamera)
 {
 	auto Sequencer = SequencerPtr.Pin();
diff --git a/Engine/Source/Editor/MovieSceneTools/Private/Sections/CameraCutSection.h b/Engine/Source/Editor/MovieSceneTools/Private/Sections/CameraCutSection.h
--- a/Engine/Source/Editor/MovieSceneTools/Private/Sections/CameraCutSection.h
+++ b/Engine/Source/Editor/MovieSceneTools/Private/Sections/CameraCutSection.h
@@ -45,7 +45,9 @@ private:
 	/** Get a representative camera for the given time */
 	AActor* GetCameraForFrame(FFrameNumber Time) const;
 
+	/** Callback for executing a "Select Camera" menu entry in the context menu. */
+	void HandleSelectCameraMenuEntryExecute(AActor* InCamera);
+
 	/** Callback for executing a "Set Camera" menu entry in the context menu. */
 	void HandleSetCameraMenuEntryExecute(AActor* InCamera);
-
 };
