Sequencer: SubSequence and Shot tracks show source frame indicator

diff --git a/Engine/Source/Editor/MovieSceneTools/Private/CommonMovieSceneTools.cpp b/Engine/Source/Editor/MovieSceneTools/Private/CommonMovieSceneTools.cpp
new file mode 100644
--- /dev/null
+++ b/Engine/Source/Editor/MovieSceneTools/Private/CommonMovieSceneTools.cpp
@@ -0,0 +1,51 @@
+// Copyright 1998-2019 Epic Games, Inc. All Rights Reserved.
+
+#include "CommonMovieSceneTools.h"
+#include "Framework/Application/SlateApplication.h"
+#include "EditorStyleSet.h"
+#include "Fonts/FontMeasure.h"
+#include "SequencerSectionPainter.h"
+
+void DrawFrameNumberHint(FSequencerSectionPainter& InPainter, FFrameTime CurrentTime, int32 FrameNumber)
+{
+	FString FrameString = FString::FromInt(FrameNumber);
+
+	const FSlateFontInfo SmallLayoutFont = FCoreStyle::GetDefaultFontStyle("Bold", 10);
+	const TSharedRef< FSlateFontMeasure > FontMeasureService = FSlateApplication::Get().GetRenderer()->GetFontMeasureService();
+	FVector2D TextSize = FontMeasureService->Measure(FrameString, SmallLayoutFont);
+
+	const float PixelX = InPainter.GetTimeConverter().FrameToPixel(CurrentTime);
+
+	// Flip the text position if getting near the end of the view range
+	static const float TextOffsetPx = 10.f;
+	bool  bDrawLeft = (InPainter.SectionGeometry.Size.X - PixelX) < (TextSize.X + 22.f) - TextOffsetPx;
+	float TextPosition = bDrawLeft ? PixelX - TextSize.X - TextOffsetPx : PixelX + TextOffsetPx;
+	//handle mirrored labels
+	const float MajorTickHeight = 9.0f;
+	FVector2D TextOffset(TextPosition, InPainter.SectionGeometry.Size.Y - (MajorTickHeight + TextSize.Y));
+
+	const FLinearColor DrawColor = FEditorStyle::GetSlateColor("SelectionColor").GetColor(FWidgetStyle());
+	const FVector2D BoxPadding = FVector2D(4.0f, 2.0f);
+	// draw time string
+
+	FSlateDrawElement::MakeBox(
+		InPainter.DrawElements,
+		InPainter.LayerId + 5,
+		InPainter.SectionGeometry.ToPaintGeometry(TextOffset - BoxPadding, TextSize + 2.0f * BoxPadding),
+		FEditorStyle::GetBrush("WhiteBrush"),
+		ESlateDrawEffect::None,
+		FLinearColor::Black.CopyWithNewOpacity(0.5f)
+	);
+
+	ESlateDrawEffect DrawEffects = InPainter.bParentEnabled ? ESlateDrawEffect::None : ESlateDrawEffect::DisabledEffect;
+
+	FSlateDrawElement::MakeText(
+		InPainter.DrawElements,
+		InPainter.LayerId + 6,
+		InPainter.SectionGeometry.ToPaintGeometry(TextOffset, TextSize),
+		FrameString,
+		SmallLayoutFont,
+		DrawEffects,
+		DrawColor
+	);
+}
\ No newline at end of file
diff --git a/Engine/Source/Editor/MovieSceneTools/Private/Sections/CinematicShotSection.cpp b/Engine/Source/Editor/MovieSceneTools/Private/Sections/CinematicShotSection.cpp
--- a/Engine/Source/Editor/MovieSceneTools/Private/Sections/CinematicShotSection.cpp
+++ b/Engine/Source/Editor/MovieSceneTools/Private/Sections/CinematicShotSection.cpp
@@ -18,6 +18,7 @@
 #include "Tracks/MovieSceneCameraCutTrack.h"
 #include "Sections/MovieSceneCameraCutSection.h"
 #include "Evaluation/MovieSceneEvaluationTemplateInstance.h"
+#include "CommonMovieSceneTools.h"
 
 #define LOCTEXT_NAMESPACE "FCinematicShotSection"
 
@@ -393,6 +394,22 @@ int32 FCinematicShotSection::OnPaintSection(FSequencerSectionPainter& InPainter)
 		);
 	}
 
+	TSharedPtr<ISequencer> Sequencer = SequencerPtr.Pin();
+
+	if (InPainter.bIsSelected && Sequencer.IsValid())
+	{
+		FFrameTime CurrentTime = Sequencer->GetLocalTime().Time;
+		if (SectionRange.Contains(CurrentTime.FrameNumber))
+		{
+			UMovieScene* SubSequenceMovieScene = SectionObject.GetSequence()->GetMovieScene();
+			const FFrameRate DisplayRate = SubSequenceMovieScene->GetDisplayRate();
+			const FFrameRate TickResolution = SubSequenceMovieScene->GetTickResolution();
+			const FFrameNumber CurrentFrameNumber = ConvertFrameTime(CurrentTime * SectionObject.OuterToInnerTransform(), TickResolution, DisplayRate).FloorToFrame();
+
+			DrawFrameNumberHint(InPainter, CurrentTime, CurrentFrameNumber.Value);
+		}
+	}
+
 	return InPainter.LayerId;
 }
 
diff --git a/Engine/Source/Editor/MovieSceneTools/Private/TrackEditors/SkeletalAnimationTrackEditor.cpp b/Engine/Source/Editor/MovieSceneTools/Private/TrackEditors/SkeletalAnimationTrackEditor.cpp
--- a/Engine/Source/Editor/MovieSceneTools/Private/TrackEditors/SkeletalAnimationTrackEditor.cpp
+++ b/Engine/Source/Editor/MovieSceneTools/Private/TrackEditors/SkeletalAnimationTrackEditor.cpp
@@ -32,7 +32,6 @@
 #include "EditorStyleSet.h"
 #include "DragAndDrop/AssetDragDropOp.h"
 #include "MovieSceneTimeHelpers.h"
-#include "Fonts/FontMeasure.h"
 #include "SequencerTimeSliderController.h"
 #include "AnimationEditorUtils.h"
 #include "Factories/PoseAssetFactory.h"
@@ -44,6 +43,8 @@
 #include "Engine/SimpleConstructionScript.h"
 #include "Engine/Blueprint.h"
 
+#include "CommonMovieSceneTools.h"
+
 namespace SkeletalAnimationEditorConstants
 {
 	// @todo Sequencer Allow this to be customizable
@@ -232,47 +233,11 @@ int32 FSkeletalAnimationSection::OnPaintSection( FSequencerSectionPainter& Paint
 		FFrameTime CurrentTime = SequencerPtr->GetLocalTime().Time;
 		if (Section.GetRange().Contains(CurrentTime.FrameNumber) && Section.Params.Animation != nullptr)
 		{
-			const float Time = TimeToPixelConverter.FrameToPixel(CurrentTime); 
-
 			// Draw the current time next to the scrub handle
 			const float AnimTime = Section.MapTimeToAnimation(CurrentTime, TickResolution);
 			int32 FrameTime = Section.Params.Animation->GetFrameAtTime(AnimTime);
-			FString FrameString = FString::FromInt(FrameTime);
-
-			const FSlateFontInfo SmallLayoutFont = FCoreStyle::GetDefaultFontStyle("Bold", 10);
-			const TSharedRef< FSlateFontMeasure > FontMeasureService = FSlateApplication::Get().GetRenderer()->GetFontMeasureService();
-			FVector2D TextSize = FontMeasureService->Measure(FrameString, SmallLayoutFont);
-
-			// Flip the text position if getting near the end of the view range
-			static const float TextOffsetPx = 10.f;
-			bool  bDrawLeft = (Painter.SectionGeometry.Size.X - Time) < (TextSize.X + 22.f) - TextOffsetPx;
-			float TextPosition = bDrawLeft ? Time - TextSize.X - TextOffsetPx : Time + TextOffsetPx;
-			//handle mirrored labels
-			const float MajorTickHeight = 9.0f; 
-			FVector2D TextOffset(TextPosition, Painter.SectionGeometry.Size.Y - (MajorTickHeight + TextSize.Y));
-
-			const FLinearColor DrawColor = FEditorStyle::GetSlateColor("SelectionColor").GetColor(FWidgetStyle());
-			const FVector2D BoxPadding = FVector2D(4.0f, 2.0f);
-			// draw time string
-	
-			FSlateDrawElement::MakeBox(
-				Painter.DrawElements,
-				LayerId + 5,
-				Painter.SectionGeometry.ToPaintGeometry(TextOffset - BoxPadding, TextSize + 2.0f * BoxPadding),
-				FEditorStyle::GetBrush("WhiteBrush"),
-				ESlateDrawEffect::None,
-				FLinearColor::Black.CopyWithNewOpacity(0.5f)
-			);
 
-			FSlateDrawElement::MakeText(
-				Painter.DrawElements,
-				LayerId + 6,
-				Painter.SectionGeometry.ToPaintGeometry(TextOffset, TextSize),
-				FrameString,
-				SmallLayoutFont,
-				DrawEffects,
-				DrawColor
-			);
+			DrawFrameNumberHint(Painter, CurrentTime, FrameTime);
 		}
 	}
 	
diff --git a/Engine/Source/Editor/MovieSceneTools/Private/TrackEditors/SubTrackEditor.cpp b/Engine/Source/Editor/MovieSceneTools/Private/TrackEditors/SubTrackEditor.cpp
--- a/Engine/Source/Editor/MovieSceneTools/Private/TrackEditors/SubTrackEditor.cpp
+++ b/Engine/Source/Editor/MovieSceneTools/Private/TrackEditors/SubTrackEditor.cpp
@@ -27,6 +27,8 @@
 #include "Misc/QualifiedFrameTime.h"
 #include "MovieSceneTimeHelpers.h"
 
+#include "CommonMovieSceneTools.h"
+
 namespace SubTrackEditorConstants
 {
 	const float TrackHeight = 50.0f;
@@ -67,7 +69,6 @@ public:
 	{
 		return &SectionObject;
 	}
-
 	virtual FText GetSectionTitle() const override
 	{
 		if(SectionObject.GetSequence())
@@ -113,7 +114,8 @@ public:
 		int32 LayerId = InPainter.PaintSectionBackground();
 
 		ESlateDrawEffect DrawEffects = InPainter.bParentEnabled ? ESlateDrawEffect::None : ESlateDrawEffect::DisabledEffect;
-
+		
+		
 		TRange<FFrameNumber> SectionRange = SectionObject.GetRange();
 		if (SectionRange.GetLowerBound().IsOpen() || SectionRange.GetUpperBound().IsOpen())
 		{
@@ -233,6 +235,23 @@ public:
 				DrawEffects,
 				FColor(200, 200, 200)
 			);
+
+
+			TSharedPtr<ISequencer> SequencerPtr = Sequencer.Pin();
+			if (InPainter.bIsSelected && SequencerPtr.IsValid())
+			{
+				FFrameTime CurrentTime = SequencerPtr->GetLocalTime().Time;
+				if (SectionRange.Contains(CurrentTime.FrameNumber))
+				{
+					UMovieScene* SubSequenceMovieScene = SectionObject.GetSequence()->GetMovieScene();
+					const FFrameRate DisplayRate = SubSequenceMovieScene->GetDisplayRate();
+					const FFrameRate TickResolution = SubSequenceMovieScene->GetTickResolution();
+					const FFrameNumber CurrentFrameNumber = ConvertFrameTime(CurrentTime * SectionObject.OuterToInnerTransform(), TickResolution, DisplayRate).FloorToFrame();
+
+					DrawFrameNumberHint(InPainter, CurrentTime, CurrentFrameNumber.Value);
+				}
+			}
+
 		}
 		else if (UMovieSceneSubSection::GetRecordingSection() == &SectionObject)
 		{
@@ -389,9 +408,7 @@ private:
 	/** Display name of the section */
 	FText DisplayName;
 
-	/** The section we are visualizing */
 	UMovieSceneSubSection& SectionObject;
-
 	/** Sequencer interface */
 	TWeakPtr<ISequencer> Sequencer;
 
diff --git a/Engine/Source/Editor/MovieSceneTools/Public/CommonMovieSceneTools.h b/Engine/Source/Editor/MovieSceneTools/Public/CommonMovieSceneTools.h
--- a/Engine/Source/Editor/MovieSceneTools/Public/CommonMovieSceneTools.h
+++ b/Engine/Source/Editor/MovieSceneTools/Public/CommonMovieSceneTools.h
@@ -6,6 +6,8 @@
 #include "Layout/Geometry.h"
 #include "Misc/FrameRate.h"
 
+class FSequencerSectionPainter;
+
 /**
  * Utility for converting time units to slate pixel units and vice versa
  */
@@ -94,3 +96,13 @@ private:
 	/** The number of pixels in the view range */
 	float PixelsPerSecond;
 };
+
+
+/**
+ * Draw the current frame number next to the scrub handle
+ *
+ * @param InPainter Structure that affords common painting operations
+ * @param CurrentTime Current time of the scrub handle
+ * @param FrameNumber Frame number to draw
+ */
+void DrawFrameNumberHint(FSequencerSectionPainter& InPainter, FFrameTime CurrentTime, int32 FrameNumber);
diff --git a/Engine/Source/Runtime/MovieScene/Private/Sections/MovieSceneSubSection.cpp b/Engine/Source/Runtime/MovieScene/Private/Sections/MovieSceneSubSection.cpp
--- a/Engine/Source/Runtime/MovieScene/Private/Sections/MovieSceneSubSection.cpp
+++ b/Engine/Source/Runtime/MovieScene/Private/Sections/MovieSceneSubSection.cpp
@@ -309,3 +309,9 @@ FMovieSceneSubSequenceData UMovieSceneSubSection::GenerateSubSequenceData(const
 {
 	return FMovieSceneSubSequenceData(*this);
 }
+
+FFrameNumber UMovieSceneSubSection::MapTimeToSectionFrame(FFrameTime InPosition) const
+{
+	FFrameNumber LocalPosition = ((InPosition - Parameters.StartFrameOffset) * Parameters.TimeScale).GetFrame();
+	return LocalPosition;
+}
diff --git a/Engine/Source/Runtime/MovieScene/Public/Sections/MovieSceneSubSection.h b/Engine/Source/Runtime/MovieScene/Public/Sections/MovieSceneSubSection.h
--- a/Engine/Source/Runtime/MovieScene/Public/Sections/MovieSceneSubSection.h
+++ b/Engine/Source/Runtime/MovieScene/Public/Sections/MovieSceneSubSection.h
@@ -132,6 +132,8 @@ public:
 		ActorToRecord = InActorToRecord;
 	}
 
+	FFrameNumber MapTimeToSectionFrame(FFrameTime InPosition) const;
+
 public:
 
 	//~ UMovieSceneSection interface
