Support for FBX Animation import/export via sequencer.

Added function to the SequencerTools model to import and export fbx. Did a little bit of refactoring to make some common functions get called by both.
Note that camera's do import but will only import and create transform tracks, will need to refactor how we create default tracks if we want it to also import the camera component.

diff --git a/Engine/Plugins/MovieScene/SequencerScripting/Source/SequencerScriptingEditor/Private/SequencerTools.cpp b/Engine/Plugins/MovieScene/SequencerScripting/Source/SequencerScriptingEditor/Private/SequencerTools.cpp
--- a/Engine/Plugins/MovieScene/SequencerScripting/Source/SequencerScriptingEditor/Private/SequencerTools.cpp
+++ b/Engine/Plugins/MovieScene/SequencerScripting/Source/SequencerScriptingEditor/Private/SequencerTools.cpp
@@ -8,6 +8,14 @@
 #include "MovieSceneTimeHelpers.h"
 #include "UObject/Stack.h"
 #include "UObject/Package.h"
+#include "LevelSequencePlayer.h"
+#include "FbxExporter.h"
+#include "FbxImporter.h"
+#include "MovieSceneToolsUserSettings.h"
+#include "MovieSceneToolHelpers.h"
+#include "CineCameraActor.h"
+#include "CineCameraComponent.h"
+#include "MovieSceneCommonHelpers.h"
 
 #define LOCTEXT_NAMESPACE "SequencerTools"
 
@@ -91,4 +99,222 @@ void USequencerToolsFunctionLibrary::CancelMovieRender()
 		CurrentCapture->Cancel();
 	}
 }
+
+
+bool USequencerToolsFunctionLibrary::ExportFBX(UWorld* World, ULevelSequence* Sequence, const TArray<FSequencerBindingProxy>& InBindings, UFbxExportOption* OverrideOptions, const FString& InFBXFileName)
+{
+	
+	
+	UnFbx::FFbxExporter* Exporter = UnFbx::FFbxExporter::GetInstance();
+	//Show the fbx export dialog options
+	Exporter->SetExportOptionsOverride(OverrideOptions);
+
+	UMovieScene* MovieScene = Sequence->GetMovieScene();
+	TArray<FGuid> Bindings;
+	for (const FSequencerBindingProxy& Proxy : InBindings)
+	{
+		if (Proxy.Sequence == Sequence)
+		{
+			Bindings.Add(Proxy.BindingID);
+		}
+	}
+
+	INodeNameAdapter NodeNameAdapter;
+	ALevelSequenceActor* OutActor;
+	FMovieSceneSequencePlaybackSettings Settings;
+	ULevelSequencePlayer* Player = ULevelSequencePlayer::CreateLevelSequencePlayer(World, Sequence, Settings, OutActor);
+	Player->Initialize(Sequence, World->PersistentLevel, Settings);
+	Player->State.AssignSequence(MovieSceneSequenceID::Root, *Sequence, *Player);
+	FMovieSceneSequenceIDRef Template = MovieSceneSequenceID::Root;
+	bool bDidExport = MovieSceneToolHelpers::ExportFBX(World, MovieScene, Player, Bindings, NodeNameAdapter, Template, InFBXFileName);
+	Exporter->SetExportOptionsOverride(nullptr);
+	return bDidExport;
+}
+
+FString GetCameraName(FbxCamera* InCamera)
+{
+	FbxNode* CameraNode = InCamera->GetNode();
+	if (CameraNode)
+	{
+		return CameraNode->GetName();
+	}
+
+	return InCamera->GetName();
+}
+
+void GetCameras(FbxNode* Parent, TArray<FbxCamera*>& Cameras)
+{
+	FbxCamera* Camera = Parent->GetCamera();
+	if (Camera)
+	{
+		Cameras.Add(Camera);
+	}
+
+	int32 NodeCount = Parent->GetChildCount();
+	for (int32 NodeIndex = 0; NodeIndex < NodeCount; ++NodeIndex)
+	{
+		FbxNode* Child = Parent->GetChild(NodeIndex);
+		GetCameras(Child, Cameras);
+	}
+}
+
+
+TArray<FGuid> AddActors(UWorld* World, UMovieSceneSequence* InSequence, UMovieScene* InMovieScene, IMovieScenePlayer* Player, const FMovieSceneSequenceIDRef& TemplateID,const TArray<TWeakObjectPtr<AActor> >& InActors)
+{
+	TArray<FGuid> PossessableGuids;
+
+	if (InMovieScene->IsReadOnly())
+	{
+		return PossessableGuids;
+	}
+
+	for (TWeakObjectPtr<AActor> WeakActor : InActors)
+	{
+		if (AActor* Actor = WeakActor.Get())
+		{
+			FGuid ExistingGuid = Player->FindObjectId(*Actor, TemplateID);
+			if (!ExistingGuid.IsValid())
+			{
+				InMovieScene->Modify();
+				const FGuid PossessableGuid = InMovieScene->AddPossessable(Actor->GetActorLabel(), Actor->GetClass());
+				PossessableGuids.Add(PossessableGuid);
+				InSequence->BindPossessableObject(PossessableGuid, *Actor, World);
+
+				//TODO New to figure way to call void FLevelSequenceEditorToolkit::AddDefaultTracksForActor(AActor& Actor, const FGuid Binding)
+
+				if (Actor->IsA<ACameraActor>())
+				{
+					MovieSceneToolHelpers::CameraAdded(InMovieScene, PossessableGuid, 0);
+				}
+			}
+		}
+	}
+	return PossessableGuids;
+}
+
+
+void ImportFBXCamera(UnFbx::FFbxImporter* FbxImporter, UWorld* World, UMovieSceneSequence* Sequence, UMovieScene* InMovieScene, IMovieScenePlayer* Player, const FMovieSceneSequenceIDRef& TemplateID, TMap<FGuid, FString>& InObjectBindingMap, bool bMatchByNameOnly, bool bCreateCameras)
+{
+	if (bCreateCameras)
+	{
+		TArray<FbxCamera*> AllCameras;
+		GetCameras(FbxImporter->Scene->GetRootNode(), AllCameras);
+
+		// Find unmatched cameras
+		TArray<FbxCamera*> UnmatchedCameras;
+		for (auto Camera : AllCameras)
+		{
+			FString NodeName = GetCameraName(Camera);
+
+			bool bMatched = false;
+			for (auto InObjectBinding : InObjectBindingMap)
+			{
+				FString ObjectName = InObjectBinding.Value;
+				if (ObjectName == NodeName)
+				{
+					// Look for a valid bound object, otherwise need to create a new camera and assign this binding to it
+					bool bFoundBoundObject = false;
+					TArrayView<TWeakObjectPtr<>> BoundObjects = Player->FindBoundObjects(InObjectBinding.Key, TemplateID);
+					for (auto BoundObject : BoundObjects)
+					{
+						if (BoundObject.IsValid())
+						{
+							bFoundBoundObject = true;
+							break;
+						}
+					}
+				}
+			}
+
+			if (!bMatched)
+			{
+				UnmatchedCameras.Add(Camera);
+			}
+		}
+
+		// If there are new cameras, clear the object binding map so that we're only assigning values to the newly created cameras
+		if (UnmatchedCameras.Num() != 0)
+		{
+			InObjectBindingMap.Reset();
+			bMatchByNameOnly = true;
+		}
+
+		// Add any unmatched cameras
+		for (auto UnmatchedCamera : UnmatchedCameras)
+		{
+			FString CameraName = GetCameraName(UnmatchedCamera);
+
+			AActor* NewCamera = nullptr;
+			if (UnmatchedCamera->GetApertureMode() == FbxCamera::eFocalLength)
+			{
+				FActorSpawnParameters SpawnParams;
+				NewCamera = World->SpawnActor<ACineCameraActor>(SpawnParams);
+				NewCamera->SetActorLabel(*CameraName);
+			}
+			else
+			{
+				FActorSpawnParameters SpawnParams;
+				NewCamera = World->SpawnActor<ACameraActor>(SpawnParams);
+				NewCamera->SetActorLabel(*CameraName);
+			}
+
+			// Copy camera properties before adding default tracks so that initial camera properties match and can be restored after sequencer finishes
+			MovieSceneToolHelpers::CopyCameraProperties(UnmatchedCamera, NewCamera);
+
+			TArray<TWeakObjectPtr<AActor> > NewCameras;
+			NewCameras.Add(NewCamera);
+			TArray<FGuid> NewCameraGuids = AddActors(World, Sequence,InMovieScene, Player, TemplateID,NewCameras);
+
+			if (NewCameraGuids.Num())
+			{
+				InObjectBindingMap.Add(NewCameraGuids[0]);
+				InObjectBindingMap[NewCameraGuids[0]] = CameraName;
+			}
+		}
+	}
+	//everything created now import it in.
+	MovieSceneToolHelpers::ImportFBXCameraToExisting(FbxImporter,  InMovieScene, Player, TemplateID, InObjectBindingMap, bMatchByNameOnly, true);
+}
+
+bool USequencerToolsFunctionLibrary::ImportFBX(UWorld* World, ULevelSequence* Sequence, const TArray<FSequencerBindingProxy>& InBindings, UMovieSceneUserImportFBXSettings* ImportFBXSettings, const FString&  ImportFilename)
+{
+	UMovieScene* MovieScene = Sequence->GetMovieScene();
+	if (!MovieScene || MovieScene->IsReadOnly())
+	{
+		return false;
+	}
+
+	TMap<FGuid, FString> ObjectBindingMap;
+	for (const FSequencerBindingProxy& Binding : InBindings)
+	{
+		FString Name = MovieScene->GetObjectDisplayName(Binding.BindingID).ToString();
+		ObjectBindingMap.Add(Binding.BindingID, Name);
+	}
+
+	FFBXInOutParameters InOutParams;
+	if (!MovieSceneToolHelpers::ReadyFBXForImport(ImportFilename, ImportFBXSettings, InOutParams))
+	{
+		return false;
+	}
+
+	const bool bMatchByNameOnly = ImportFBXSettings->bMatchByNameOnly;
+	ALevelSequenceActor* OutActor;
+	FMovieSceneSequencePlaybackSettings Settings;
+	ULevelSequencePlayer* Player = ULevelSequencePlayer::CreateLevelSequencePlayer(World, Sequence, Settings, OutActor);
+	Player->Initialize(Sequence, World->GetLevel(0), Settings);
+	Player->State.AssignSequence(MovieSceneSequenceID::Root, *Sequence, *Player);
+
+	UnFbx::FFbxImporter* FbxImporter = UnFbx::FFbxImporter::GetInstance();
+
+	ImportFBXCamera(FbxImporter, World, Sequence, MovieScene, Player, MovieSceneSequenceID::Root, ObjectBindingMap, bMatchByNameOnly,
+		ImportFBXSettings->bCreateCameras);
+
+	bool bValid = MovieSceneToolHelpers::ImportFBXIfReady(World, MovieScene, Player, ObjectBindingMap, ImportFBXSettings, InOutParams);
+
+	return bValid;
+}
+
+
+
+
 #undef LOCTEXT_NAMESPACE // "SequencerTools"
\ No newline at end of file
diff --git a/Engine/Plugins/MovieScene/SequencerScripting/Source/SequencerScriptingEditor/Public/SequencerTools.h b/Engine/Plugins/MovieScene/SequencerScripting/Source/SequencerScriptingEditor/Public/SequencerTools.h
--- a/Engine/Plugins/MovieScene/SequencerScripting/Source/SequencerScriptingEditor/Public/SequencerTools.h
+++ b/Engine/Plugins/MovieScene/SequencerScripting/Source/SequencerScriptingEditor/Public/SequencerTools.h
@@ -6,8 +6,11 @@
 #include "Kismet/BlueprintFunctionLibrary.h"
 #include "Modules/ModuleManager.h"
 #include "MovieSceneCaptureDialogModule.h"
+#include "SequencerBindingProxy.h"
 #include "SequencerTools.generated.h"
 
+class UFbxExportOption;
+
 DECLARE_DYNAMIC_DELEGATE_OneParam(FOnRenderMovieStopped, bool, bSuccess);
 
 /** 
@@ -43,4 +46,30 @@ public:
 	*/
 	UFUNCTION(BlueprintCallable, Category = "Editor Scripting | Sequencer Tools")
 	static void CancelMovieRender();
+
+public:
+	/*
+	 * Export Passed in Bindings to FBX
+	 *
+	 * @InWorld World to export
+	 * @InSequence Sequence to export
+	 * @InBindings Bindings to export
+	 * @InFBXFileName File to create
+	 */
+	UFUNCTION(BlueprintCallable, Category = "Editor Scripting | Sequencer Tools")
+	static bool ExportFBX(UWorld* InWorld, ULevelSequence* InSequence, const TArray<FSequencerBindingProxy>& InBindings, UFbxExportOption* OverrideOptions,const FString& InFBXFileName);
+
+	/*
+	 * Import Passed in Bindings to FBX
+	 *
+	 * @InWorld World to import to
+	 * @InSequence InSequence to import
+	 * @InBindings InBindings to import
+	 * @InImportFBXSettings Settings to control import.
+	 * @InImportFileName Path to fbx file to create
+	 */
+	UFUNCTION(BlueprintCallable, Category = "Editor Scripting | Sequencer Tools")
+	static bool ImportFBX(UWorld* InWorld, ULevelSequence* InSequence, const TArray<FSequencerBindingProxy>& InBindings, UMovieSceneUserImportFBXSettings* InImportFBXSettings, const FString&  InImportFilename);
+
+
 };
\ No newline at end of file
diff --git a/Engine/Plugins/MovieScene/SequencerScripting/Source/SequencerScriptingEditor/SequencerScriptingEditor.Build.cs b/Engine/Plugins/MovieScene/SequencerScripting/Source/SequencerScriptingEditor/SequencerScriptingEditor.Build.cs
--- a/Engine/Plugins/MovieScene/SequencerScripting/Source/SequencerScriptingEditor/SequencerScriptingEditor.Build.cs
+++ b/Engine/Plugins/MovieScene/SequencerScripting/Source/SequencerScriptingEditor/SequencerScriptingEditor.Build.cs
@@ -16,8 +16,10 @@ public class SequencerScriptingEditor : ModuleRules
 		PrivateIncludePaths.AddRange(
 			new string[] {
 				"SequencerScriptingEditor/Private",
+                "SequencerScripting/Private",
+                "../../../../Source/Editor/UnrealEd/Private", // TODO: Fix this, for now it's needed for the fbx exporter
             }
-		);
+        );
 
 		PublicDependencyModuleNames.AddRange(
 			new string[]
@@ -27,8 +29,9 @@ public class SequencerScriptingEditor : ModuleRules
 				"MovieScene",
                 "MovieSceneTools",
                 "MovieSceneTracks",
+                "CinematicCamera",
             }
-		);
+        );
 
 		PrivateDependencyModuleNames.AddRange(
 			new string[]
@@ -42,7 +45,9 @@ public class SequencerScriptingEditor : ModuleRules
 				"MovieSceneCaptureDialog",
                 "MovieSceneCapture",
                 "LevelSequence",
-			}
+				"SequencerScripting",
+                "UnrealEd",
+            }
 		);
 
 		DynamicallyLoadedModuleNames.AddRange(
@@ -51,5 +56,7 @@ public class SequencerScriptingEditor : ModuleRules
 				
 			}
 		);
-	}
+        AddEngineThirdPartyPrivateStaticDependencies(Target, "FBX");
+
+    }
 }
diff --git a/Engine/Source/Editor/MovieSceneTools/Private/MovieSceneToolHelpers.cpp b/Engine/Source/Editor/MovieSceneTools/Private/MovieSceneToolHelpers.cpp
--- a/Engine/Source/Editor/MovieSceneTools/Private/MovieSceneToolHelpers.cpp
+++ b/Engine/Source/Editor/MovieSceneTools/Private/MovieSceneToolHelpers.cpp
@@ -52,6 +52,12 @@
 #include "Channels/MovieSceneChannelProxy.h"
 #include "Evaluation/MovieSceneEvaluationTrack.h"
 #include "Evaluation/MovieSceneEvaluationTemplateInstance.h"
+#include "IMovieScenePlayer.h"
+#include "Tracks/MovieSceneCinematicShotTrack.h"
+#include "Tracks/MovieSceneCameraCutTrack.h"
+#include "Sections/MovieSceneCameraCutSection.h"
+#include "FbxExporter.h"
 
+
 /* MovieSceneToolHelpers
  *****************************************************************************/
@@ -854,12 +860,40 @@
 	}
 }
 
-bool ImportFBXProperty(FString NodeName, FString AnimatedPropertyName, FGuid ObjectBinding, UnFbx::FFbxCurvesAPI& CurveAPI, UMovieScene* InMovieScene, ISequencer& InSequencer)
+static FGuid GetHandleToObject(UObject* InObject, UMovieScene* InMovieScene, IMovieScenePlayer* Player, const FMovieSceneSequenceIDRef& TemplateID)
 {
+	// Attempt to resolve the object through the movie scene instance first, 
+	FGuid PropertyOwnerGuid = FGuid();
+	if (InObject != nullptr && !InMovieScene->IsReadOnly())
+	{
+		FGuid ObjectGuid = Player->FindObjectId(*InObject, TemplateID);
+		if (ObjectGuid.IsValid())
+		{
+			// Check here for spawnable otherwise spawnables get recreated as possessables, which doesn't make sense
+			FMovieSceneSpawnable* Spawnable = InMovieScene->FindSpawnable(ObjectGuid);
+			if (Spawnable)
+			{
+				PropertyOwnerGuid = ObjectGuid;
+			}
+			else
+			{
+				FMovieScenePossessable* Possessable = InMovieScene->FindPossessable(ObjectGuid);
+				if (Possessable)
+				{
+					PropertyOwnerGuid = ObjectGuid;
+				}
+			}
+		}
+	}
+	return PropertyOwnerGuid;
+}
+
+bool ImportFBXProperty(FString NodeName, FString AnimatedPropertyName, FGuid ObjectBinding, UnFbx::FFbxCurvesAPI& CurveAPI, UMovieScene* InMovieScene, IMovieScenePlayer* Player, const FMovieSceneSequenceIDRef& TemplateID)
+{
 	const UMovieSceneToolsProjectSettings* ProjectSettings = GetDefault<UMovieSceneToolsProjectSettings>();
 	const UMovieSceneUserImportFBXSettings* ImportFBXSettings = GetDefault<UMovieSceneUserImportFBXSettings>();
 
-	TArrayView<TWeakObjectPtr<>> BoundObjects = InSequencer.FindBoundObjects(ObjectBinding, InSequencer.GetFocusedTemplateID());
+	TArrayView<TWeakObjectPtr<>> BoundObjects = Player->FindBoundObjects(ObjectBinding, TemplateID);
 
 	for (auto FbxSetting : ProjectSettings->FbxSettings)
 	{
@@ -888,7 +922,12 @@
 				continue;
 			}
 		
-			FGuid PropertyOwnerGuid = InSequencer.GetHandleToObject(PropertyOwner);
+			FGuid PropertyOwnerGuid = GetHandleToObject(PropertyOwner,InMovieScene, Player, TemplateID);
 			if (!PropertyOwnerGuid.IsValid())
 			{
 				continue;
+			}
+
+			if (!PropertyOwnerGuid.IsValid())
+			{
+				continue;
@@ -979,6 +1018,51 @@
 	return false;
 }
 
+void MovieSceneToolHelpers::CameraAdded(UMovieScene* OwnerMovieScene, FGuid CameraGuid, FFrameNumber FrameNumber)
+{
+
+	// If there's a cinematic shot track, no need to set this camera to a shot
+	UMovieSceneTrack* CinematicShotTrack = OwnerMovieScene->FindMasterTrack(UMovieSceneCinematicShotTrack::StaticClass());
+	if (CinematicShotTrack)
+	{
+		return;
+	}
+
+	UMovieSceneTrack* CameraCutTrack = OwnerMovieScene->GetCameraCutTrack();
+
+	// If there's a camera cut track with at least one section, no need to change the section
+	if (CameraCutTrack && CameraCutTrack->GetAllSections().Num() > 0)
+	{
+		return;
+	}
+
+	if (!CameraCutTrack)
+	{
+		CameraCutTrack = OwnerMovieScene->AddCameraCutTrack(UMovieSceneCameraCutTrack::StaticClass());
+	}
+
+	if (CameraCutTrack)
+	{
+		UMovieSceneSection* Section = MovieSceneHelpers::FindSectionAtTime(CameraCutTrack->GetAllSections(), FrameNumber);
+		UMovieSceneCameraCutSection* CameraCutSection = Cast<UMovieSceneCameraCutSection>(Section);
+
+		if (CameraCutSection)
+		{
+			CameraCutSection->Modify();
+			CameraCutSection->SetCameraGuid(CameraGuid);
+		}
+		else
+		{
+			CameraCutTrack->Modify();
+
+			UMovieSceneCameraCutSection* NewSection = Cast<UMovieSceneCameraCutSection>(CameraCutTrack->CreateNewSection());
+			NewSection->SetRange(OwnerMovieScene->GetPlaybackRange());
+			NewSection->SetCameraGuid(CameraGuid);
+			CameraCutTrack->AddSection(*NewSection);
+		}
+	}
+}
+
 void ImportTransformChannel(const FRichCurve& Source, FMovieSceneFloatChannel* Dest, FFrameRate DestFrameRate, bool bNegateTangents)
 {
 	// If there are no keys, don't clear the existing channel
@@ -1099,7 +1183,7 @@
 	return true;
 }
 
-bool ImportFBXNode(FString NodeName, UnFbx::FFbxCurvesAPI& CurveAPI, UMovieScene* InMovieScene, ISequencer& InSequencer, FGuid ObjectBinding)
+bool MovieSceneToolHelpers::ImportFBXNode(FString NodeName, UnFbx::FFbxCurvesAPI& CurveAPI, UMovieScene* InMovieScene, IMovieScenePlayer* Player, const FMovieSceneSequenceIDRef& TemplateID, FGuid ObjectBinding)
 {
 	// Look for animated float properties
 	TArray<FString> AnimatedPropertyNames;
@@ -1107,7 +1191,7 @@
 		
 	for (auto AnimatedPropertyName : AnimatedPropertyNames)
 	{
-		ImportFBXProperty(NodeName, AnimatedPropertyName, ObjectBinding, CurveAPI, InMovieScene, InSequencer);
+		ImportFBXProperty(NodeName, AnimatedPropertyName, ObjectBinding, CurveAPI, InMovieScene, Player, TemplateID);
 	}
 	
 	ImportFBXTransform(NodeName, ObjectBinding, CurveAPI, InMovieScene);
@@ -1176,7 +1260,7 @@
 	return nullptr;
 }
 
-void CopyCameraProperties(FbxCamera* CameraNode, AActor* InCameraActor)
+void MovieSceneToolHelpers::CopyCameraProperties(FbxCamera* CameraNode, AActor* InCameraActor)
 {
 	float FieldOfView;
 	float FocalLength;
@@ -1244,100 +1328,13 @@
 	return InCamera->GetName();
 }
 
-void ImportFBXCamera(UnFbx::FFbxImporter* FbxImporter, UMovieScene* InMovieScene, ISequencer& InSequencer, TMap<FGuid, FString>& InObjectBindingMap, bool bMatchByNameOnly, bool bCreateCameras)
+
+void MovieSceneToolHelpers::ImportFBXCameraToExisting(UnFbx::FFbxImporter* FbxImporter, UMovieScene* InMovieScene, IMovieScenePlayer* Player, const FMovieSceneSequenceIDRef& TemplateID, TMap<FGuid, FString>& InObjectBindingMap, bool bMatchByNameOnly, bool bNotifySlate)
 {
-	if (bCreateCameras)
-	{
-		TArray<FbxCamera*> AllCameras;
-		GetCameras(FbxImporter->Scene->GetRootNode(), AllCameras);
-
-		UWorld* World = GCurrentLevelEditingViewportClient ? GCurrentLevelEditingViewportClient->GetWorld() : nullptr;
-
-		// Find unmatched cameras
-		TArray<FbxCamera*> UnmatchedCameras;
-		for (auto Camera : AllCameras)
-		{
-			FString NodeName = GetCameraName(Camera);
-
-			bool bMatched = false;
-			for (auto InObjectBinding : InObjectBindingMap)
-			{		
-				FString ObjectName = InObjectBinding.Value;
-				if (ObjectName == NodeName)
-				{
-					// Look for a valid bound object, otherwise need to create a new camera and assign this binding to it
-					bool bFoundBoundObject = false;
-					TArrayView<TWeakObjectPtr<>> BoundObjects = InSequencer.FindBoundObjects(InObjectBinding.Key, InSequencer.GetFocusedTemplateID());
-					for (auto BoundObject : BoundObjects)
-					{
-						if (BoundObject.IsValid())
-						{
-							bFoundBoundObject = true;
-							break;
-						}
-					}
-
-					if (!bFoundBoundObject)
-					{
-						FNotificationInfo Info(FText::Format(NSLOCTEXT("MovieSceneTools", "NoBoundObjectsError", "Existing binding has no objects. Creating a new camera and binding for {0}"), FText::FromString(ObjectName)));
-						Info.ExpireDuration = 5.0f;
-						FSlateNotificationManager::Get().AddNotification(Info)->SetCompletionState(SNotificationItem::CS_Fail);
-					}
-				}
-			}
-
-			if (!bMatched)
-			{
-				UnmatchedCameras.Add(Camera);
-			}
-		}
-
-		// If there are new cameras, clear the object binding map so that we're only assigning values to the newly created cameras
-		if (UnmatchedCameras.Num() != 0)
-		{
-			InObjectBindingMap.Reset();
-			bMatchByNameOnly = true;
-		}
-
-		// Add any unmatched cameras
-		for (auto UnmatchedCamera : UnmatchedCameras)
-		{
-			FString CameraName = GetCameraName(UnmatchedCamera);
-
-			AActor* NewCamera = nullptr;
-			if (UnmatchedCamera->GetApertureMode() == FbxCamera::eFocalLength)
-			{
-				FActorSpawnParameters SpawnParams;
-				NewCamera = World->SpawnActor<ACineCameraActor>(SpawnParams);
-				NewCamera->SetActorLabel(*CameraName);
-			}
-			else
-			{
-				FActorSpawnParameters SpawnParams;
-				NewCamera = World->SpawnActor<ACameraActor>(SpawnParams);
-				NewCamera->SetActorLabel(*CameraName);
-			}
-
-			// Copy camera properties before adding default tracks so that initial camera properties match and can be restored after sequencer finishes
-			CopyCameraProperties(UnmatchedCamera, NewCamera);
-
-			TArray<TWeakObjectPtr<AActor> > NewCameras;
-			NewCameras.Add(NewCamera);
-			TArray<FGuid> NewCameraGuids = InSequencer.AddActors(NewCameras);
-
-			if (NewCameraGuids.Num())
-			{
-				InObjectBindingMap.Add(NewCameraGuids[0]);
-				InObjectBindingMap[NewCameraGuids[0]] = CameraName;
-			}
-		}
-	}
-
-
 	for (auto InObjectBinding : InObjectBindingMap)
 	{
-		TArrayView<TWeakObjectPtr<>> BoundObjects = InSequencer.FindBoundObjects(InObjectBinding.Key, InSequencer.GetFocusedTemplateID());
-		
+		TArrayView<TWeakObjectPtr<>> BoundObjects = Player->FindBoundObjects(InObjectBinding.Key,TemplateID);
+
 		FString ObjectName = InObjectBinding.Value;
 		FbxCamera* CameraNode = nullptr;
 		FbxNode* Node = RetrieveObjectFromName(*ObjectName, FbxImporter->Scene->GetRootNode());
@@ -1348,7 +1345,7 @@
 
 		if (!CameraNode)
 		{
-			if (bMatchByNameOnly)
+			if (bMatchByNameOnly && bNotifySlate)
 			{
 				FNotificationInfo Info(FText::Format(NSLOCTEXT("MovieSceneTools", "NoMatchingCameraError", "Failed to find any matching camera for {0}"), FText::FromString(ObjectName)));
 				Info.ExpireDuration = 5.0f;
@@ -1358,7 +1355,7 @@
 			}
 
 			CameraNode = FindCamera(FbxImporter->Scene->GetRootNode());
-			if (CameraNode)
+			if (CameraNode && bNotifySlate)
 			{
 				FString CameraName = GetCameraName(CameraNode);
 				FNotificationInfo Info(FText::Format(NSLOCTEXT("MovieSceneTools", "NoMatchingCameraWarning", "Failed to find any matching camera for {0}. Importing onto first camera from fbx {1}"), FText::FromString(ObjectName), FText::FromString(CameraName)));
@@ -1396,7 +1393,7 @@
 				UCameraComponent* CameraComponent = nullptr;
 				FName TrackName;
 				float TrackValue;
-				
+
 				if (ACineCameraActor* CineCameraActor = Cast<ACineCameraActor>(FoundObject))
 				{
 					CameraComponent = CineCameraActor->GetCineCameraComponent();
@@ -1415,7 +1412,9 @@
 				}
 
 				// Set the default value of the current focal length or field of view section
-				FGuid PropertyOwnerGuid = InSequencer.GetHandleToObject(CameraComponent);
+				//FGuid PropertyOwnerGuid = Player->GetHandleToObject(CameraComponent);
+				FGuid PropertyOwnerGuid = GetHandleToObject(CameraComponent, InMovieScene, Player, TemplateID);
+
 				if (!PropertyOwnerGuid.IsValid())
 				{
 					continue;
@@ -1448,6 +1447,98 @@
 	}
 }
 
+void ImportFBXCamera(UnFbx::FFbxImporter* FbxImporter, UMovieScene* InMovieScene, ISequencer& InSequencer,  TMap<FGuid, FString>& InObjectBindingMap, bool bMatchByNameOnly, bool bCreateCameras)
+{
+	if (bCreateCameras)
+	{
+		TArray<FbxCamera*> AllCameras;
+		GetCameras(FbxImporter->Scene->GetRootNode(), AllCameras);
+
+		UWorld* World = GCurrentLevelEditingViewportClient ? GCurrentLevelEditingViewportClient->GetWorld() : nullptr;
+
+		// Find unmatched cameras
+		TArray<FbxCamera*> UnmatchedCameras;
+		for (auto Camera : AllCameras)
+		{
+			FString NodeName = GetCameraName(Camera);
+
+			bool bMatched = false;
+			for (auto InObjectBinding : InObjectBindingMap)
+			{		
+				FString ObjectName = InObjectBinding.Value;
+				if (ObjectName == NodeName)
+				{
+					// Look for a valid bound object, otherwise need to create a new camera and assign this binding to it
+					bool bFoundBoundObject = false;
+					TArrayView<TWeakObjectPtr<>> BoundObjects = InSequencer.FindBoundObjects(InObjectBinding.Key, InSequencer.GetFocusedTemplateID());
+					for (auto BoundObject : BoundObjects)
+					{
+						if (BoundObject.IsValid())
+						{
+							bFoundBoundObject = true;
+							break;
+						}
+					}
+
+					if (!bFoundBoundObject)
+					{
+						FNotificationInfo Info(FText::Format(NSLOCTEXT("MovieSceneTools", "NoBoundObjectsError", "Existing binding has no objects. Creating a new camera and binding for {0}"), FText::FromString(ObjectName)));
+						Info.ExpireDuration = 5.0f;
+						FSlateNotificationManager::Get().AddNotification(Info)->SetCompletionState(SNotificationItem::CS_Fail);
+					}
+				}
+			}
+
+			if (!bMatched)
+			{
+				UnmatchedCameras.Add(Camera);
+			}
+		}
+
+		// If there are new cameras, clear the object binding map so that we're only assigning values to the newly created cameras
+		if (UnmatchedCameras.Num() != 0)
+		{
+			InObjectBindingMap.Reset();
+			bMatchByNameOnly = true;
+		}
+
+		// Add any unmatched cameras
+		for (auto UnmatchedCamera : UnmatchedCameras)
+		{
+			FString CameraName = GetCameraName(UnmatchedCamera);
+
+			AActor* NewCamera = nullptr;
+			if (UnmatchedCamera->GetApertureMode() == FbxCamera::eFocalLength)
+			{
+				FActorSpawnParameters SpawnParams;
+				NewCamera = World->SpawnActor<ACineCameraActor>(SpawnParams);
+				NewCamera->SetActorLabel(*CameraName);
+			}
+			else
+			{
+				FActorSpawnParameters SpawnParams;
+				NewCamera = World->SpawnActor<ACameraActor>(SpawnParams);
+				NewCamera->SetActorLabel(*CameraName);
+			}
+
+			// Copy camera properties before adding default tracks so that initial camera properties match and can be restored after sequencer finishes
+			MovieSceneToolHelpers::CopyCameraProperties(UnmatchedCamera, NewCamera);
+
+			TArray<TWeakObjectPtr<AActor> > NewCameras;
+			NewCameras.Add(NewCamera);
+			TArray<FGuid> NewCameraGuids = InSequencer.AddActors(NewCameras);
+
+			if (NewCameraGuids.Num())
+			{
+				InObjectBindingMap.Add(NewCameraGuids[0]);
+				InObjectBindingMap[NewCameraGuids[0]] = CameraName;
+			}
+		}
+	}
+	
+	MovieSceneToolHelpers::ImportFBXCameraToExisting(FbxImporter, InMovieScene, &InSequencer, InSequencer.GetFocusedTemplateID(), InObjectBindingMap, bMatchByNameOnly, true);
+}
+
 FGuid FindCameraGuid(FbxCamera* Camera, TMap<FGuid, FString>& InObjectBindingMap)
 {
 	FString CameraName = GetCameraName(Camera);
@@ -1593,123 +1684,33 @@
 
 	FReply OnImportFBXClicked()
 	{
+		
 		UMovieSceneUserImportFBXSettings* ImportFBXSettings = GetMutableDefault<UMovieSceneUserImportFBXSettings>();
-
 		FEditorDirectories::Get().SetLastDirectory( ELastDirectory::FBX, FPaths::GetPath( ImportFilename ) ); // Save path as default for next time.
-	
+
+		if (!MovieScene || MovieScene->IsReadOnly())
+		{
+			return FReply::Unhandled();
+		}
+
+		FFBXInOutParameters InOutParams;
+		if (!MovieSceneToolHelpers::ReadyFBXForImport(ImportFilename, ImportFBXSettings,InOutParams))
+		{
+			return FReply::Unhandled();
+		}
+
+		const FScopedTransaction Transaction(NSLOCTEXT("MovieSceneTools", "ImportFBXTransaction", "Import FBX"));
 		UnFbx::FFbxImporter* FbxImporter = UnFbx::FFbxImporter::GetInstance();
 
-		UnFbx::FBXImportOptions* ImportOptions = FbxImporter->GetImportOptions();
-		bool bConvertSceneBackup = ImportOptions->bConvertScene;
-		bool bConvertSceneUnitBackup = ImportOptions->bConvertSceneUnit;
-		bool bForceFrontXAxisBackup = ImportOptions->bForceFrontXAxis;
-
-		ImportOptions->bConvertScene = true;
-		ImportOptions->bConvertSceneUnit = true;
-		ImportOptions->bForceFrontXAxis = ImportFBXSettings->bForceFrontXAxis;
-
-		const FString FileExtension = FPaths::GetExtension(ImportFilename);
-		if (!FbxImporter->ImportFromFile(*ImportFilename, FileExtension, true))
-		{
-			// Log the error message and fail the import.
-			FbxImporter->ReleaseScene();
-			ImportOptions->bConvertScene = bConvertSceneBackup;
-			ImportOptions->bConvertSceneUnit = bConvertSceneUnitBackup;
-			ImportOptions->bForceFrontXAxis = bForceFrontXAxisBackup;
-			return FReply::Unhandled();
-		}
-		
-		if (MovieScene->IsReadOnly())
-		{
-			return FReply::Unhandled();
-		}
-		
 		const bool bMatchByNameOnly = ImportFBXSettings->bMatchByNameOnly;
-
-		const FScopedTransaction Transaction( NSLOCTEXT( "MovieSceneTools", "ImportFBXTransaction", "Import FBX" ) );
-
 		// Import static cameras first
 		ImportFBXCamera(FbxImporter, MovieScene, *Sequencer, ObjectBindingMap, bMatchByNameOnly, bCreateCameras.IsSet() ? bCreateCameras.GetValue() : ImportFBXSettings->bCreateCameras);
 
-		UnFbx::FFbxCurvesAPI CurveAPI;
-		FbxImporter->PopulateAnimatedCurveData(CurveAPI);
-		TArray<FString> AllNodeNames;
-		CurveAPI.GetAllNodeNameArray(AllNodeNames);
-
-		FString RootNodeName = FbxImporter->Scene->GetRootNode()->GetName();
-
-		// First try matching by name
-		for (int32 NodeIndex = 0; NodeIndex < AllNodeNames.Num(); )
-		{
-			FString NodeName = AllNodeNames[NodeIndex];
-			if (RootNodeName == NodeName)
-			{
-				++NodeIndex;
-				continue;
-			}
-
-			bool bFoundMatch = false;
-			for (auto It = ObjectBindingMap.CreateConstIterator(); It; ++It)
-			{
-				if (FCString::Strcmp(*It.Value().ToUpper(), *NodeName.ToUpper()) == 0)
-				{
-					ImportFBXNode(NodeName, CurveAPI, MovieScene, *Sequencer, It.Key());
-
-					ObjectBindingMap.Remove(It.Key());
-					AllNodeNames.RemoveAt(NodeIndex);
-
-					bFoundMatch = true;
-					break;
-				}
-			}
-
-			if (bFoundMatch)
-			{
-				continue;
-			}
-
-			++NodeIndex;
-		}
-
-		// Otherwise, get the first available node that hasn't been imported onto yet
-		if (!bMatchByNameOnly)
-		{
-			for (int32 NodeIndex = 0; NodeIndex < AllNodeNames.Num(); )
-			{
-				FString NodeName = AllNodeNames[NodeIndex];
-				if (RootNodeName == NodeName)
-				{
-					++NodeIndex;
-					continue;
-				}
-
-				auto It = ObjectBindingMap.CreateConstIterator();
-				if (It)
-				{
-					ImportFBXNode(NodeName, CurveAPI, MovieScene, *Sequencer, It.Key());
-
-					UE_LOG(LogMovieScene, Warning, TEXT("Fbx Import: Failed to find any matching node for (%s). Defaulting to first available (%s)."), *NodeName, *It.Value());
-					ObjectBindingMap.Remove(It.Key());
-					AllNodeNames.RemoveAt(NodeIndex);
-					continue;
-				}
-
-				++NodeIndex;
-			}
-		}
-
-		for (FString NodeName : AllNodeNames)
-		{
-			UE_LOG(LogMovieScene, Warning, TEXT("Fbx Import: Failed to find any matching node for (%s)."), *NodeName);
-		}
-
+		UWorld* World = Cast<UWorld>(Sequencer->GetPlaybackContext());
+		bool bValid = MovieSceneToolHelpers::ImportFBXIfReady(World, MovieScene, Sequencer, ObjectBindingMap, ImportFBXSettings, InOutParams);
+	
 		Sequencer->NotifyMovieSceneDataChanged(EMovieSceneDataChangeType::MovieSceneStructureItemAdded);
 
-		FbxImporter->ReleaseScene();
-		ImportOptions->bConvertScene = bConvertSceneBackup;
-		ImportOptions->bConvertSceneUnit = bConvertSceneUnitBackup;
-		ImportOptions->bForceFrontXAxis = bForceFrontXAxisBackup;
-
 		TSharedPtr<SWindow> Window = FSlateApplication::Get().FindWidgetWindow(AsShared());
 
 		if ( Window.IsValid() )
@@ -1717,7 +1718,7 @@
 			Window->RequestDestroyWindow();
 		}
 
-		return FReply::Handled();
+		return bValid ? FReply::Handled() : FReply::Unhandled();
 	}
 
 	TSharedPtr<IDetailsView> DetailView;
@@ -1728,8 +1729,120 @@
 	TOptional<bool> bCreateCameras;
 };
 
+bool MovieSceneToolHelpers::ReadyFBXForImport(const FString&  ImportFilename, UMovieSceneUserImportFBXSettings* ImportFBXSettings, FFBXInOutParameters& OutParams)
+{
+	UnFbx::FFbxImporter* FbxImporter = UnFbx::FFbxImporter::GetInstance();
 
-bool MovieSceneToolHelpers::ImportFBX(UMovieScene* InMovieScene, ISequencer& InSequencer, const TMap<FGuid, FString>& InObjectBindingMap, TOptional<bool> bCreateCameras)
+	UnFbx::FBXImportOptions* ImportOptions = FbxImporter->GetImportOptions();
+	OutParams.bConvertSceneBackup = ImportOptions->bConvertScene;
+	OutParams.bConvertSceneUnitBackup = ImportOptions->bConvertSceneUnit;
+	OutParams.bForceFrontXAxisBackup = ImportOptions->bForceFrontXAxis;
+
+	ImportOptions->bConvertScene = true;
+	ImportOptions->bConvertSceneUnit = true;
+	ImportOptions->bForceFrontXAxis = ImportFBXSettings->bForceFrontXAxis;
+
+	const FString FileExtension = FPaths::GetExtension(ImportFilename);
+	if (!FbxImporter->ImportFromFile(*ImportFilename, FileExtension, true))
+	{
+		// Log the error message and fail the import.
+		FbxImporter->ReleaseScene();
+		ImportOptions->bConvertScene = OutParams.bConvertSceneBackup;
+		ImportOptions->bConvertSceneUnit = OutParams.bConvertSceneUnitBackup;
+		ImportOptions->bForceFrontXAxis = OutParams.bForceFrontXAxisBackup;
+		return false;
+	}
+	return true;
+}
+
+bool MovieSceneToolHelpers::ImportFBXIfReady(UWorld* World, UMovieScene* MovieScene, IMovieScenePlayer* Player, TMap<FGuid, FString>& ObjectBindingMap, UMovieSceneUserImportFBXSettings* ImportFBXSettings,
+	const FFBXInOutParameters& InParams)
 {
+	UnFbx::FFbxImporter* FbxImporter = UnFbx::FFbxImporter::GetInstance();
+
+	const bool bMatchByNameOnly = ImportFBXSettings->bMatchByNameOnly;
+
+	UnFbx::FFbxCurvesAPI CurveAPI;
+	FbxImporter->PopulateAnimatedCurveData(CurveAPI);
+	TArray<FString> AllNodeNames;
+	CurveAPI.GetAllNodeNameArray(AllNodeNames);
+
+	FString RootNodeName = FbxImporter->Scene->GetRootNode()->GetName();
+
+	// First try matching by name
+	for (int32 NodeIndex = 0; NodeIndex < AllNodeNames.Num(); )
+	{
+		FString NodeName = AllNodeNames[NodeIndex];
+		if (RootNodeName == NodeName)
+		{
+			++NodeIndex;
+			continue;
+		}
+
+		bool bFoundMatch = false;
+		for (auto It = ObjectBindingMap.CreateConstIterator(); It; ++It)
+		{
+			if (FCString::Strcmp(*It.Value().ToUpper(), *NodeName.ToUpper()) == 0)
+			{
+				MovieSceneToolHelpers::ImportFBXNode(NodeName, CurveAPI, MovieScene, Player, MovieSceneSequenceID::Root, It.Key());
+
+				ObjectBindingMap.Remove(It.Key());
+				AllNodeNames.RemoveAt(NodeIndex);
+
+				bFoundMatch = true;
+				break;
+			}
+		}
+
+		if (bFoundMatch)
+		{
+			continue;
+		}
+
+		++NodeIndex;
+	}
+
+	// Otherwise, get the first available node that hasn't been imported onto yet
+	if (!bMatchByNameOnly)
+	{
+		for (int32 NodeIndex = 0; NodeIndex < AllNodeNames.Num(); )
+		{
+			FString NodeName = AllNodeNames[NodeIndex];
+			if (RootNodeName == NodeName)
+			{
+				++NodeIndex;
+				continue;
+			}
+
+			auto It = ObjectBindingMap.CreateConstIterator();
+			if (It)
+			{
+				MovieSceneToolHelpers::ImportFBXNode(NodeName, CurveAPI, MovieScene, Player, MovieSceneSequenceID::Root, It.Key());
+
+				UE_LOG(LogMovieScene, Warning, TEXT("Fbx Import: Failed to find any matching node for (%s). Defaulting to first available (%s)."), *NodeName, *It.Value());
+				ObjectBindingMap.Remove(It.Key());
+				AllNodeNames.RemoveAt(NodeIndex);
+				continue;
+			}
+
+			++NodeIndex;
+		}
+	}
+
+	for (FString NodeName : AllNodeNames)
+	{
+		UE_LOG(LogMovieScene, Warning, TEXT("Fbx Import: Failed to find any matching node for (%s)."), *NodeName);
+	}
+
+	FbxImporter->ReleaseScene();
+	UnFbx::FBXImportOptions* ImportOptions = FbxImporter->GetImportOptions();
+	ImportOptions->bConvertScene = InParams.bConvertSceneBackup;
+	ImportOptions->bConvertSceneUnit = InParams.bConvertSceneUnitBackup;
+	ImportOptions->bForceFrontXAxis = InParams.bForceFrontXAxisBackup;
+	return true;
+}
+
+bool MovieSceneToolHelpers::ImportFBXWithDialog(UMovieScene* InMovieScene, ISequencer& InSequencer, const TMap<FGuid, FString>& InObjectBindingMap, TOptional<bool> bCreateCameras)
+{
 	TArray<FString> OpenFilenames;
 	IDesktopPlatform* DesktopPlatform = FDesktopPlatformModule::Get();
@@ -1941,3 +2054,59 @@
 	}
 	return nullptr;
 }
+
+bool MovieSceneToolHelpers::ExportFBX(UWorld* World, UMovieScene* MovieScene, IMovieScenePlayer* Player, TArray<FGuid>& Bindings, INodeNameAdapter& NodeNameAdapter, FMovieSceneSequenceIDRef& Template, const FString& InFBXFileName)
+{
+
+	UnFbx::FFbxExporter* Exporter = UnFbx::FFbxExporter::GetInstance();
+
+	Exporter->CreateDocument();
+	Exporter->SetTrasformBaking(false);
+	Exporter->SetKeepHierarchy(true);
+
+	// Get list of actors based upon bindings...
+	const bool bSelectedOnly = (Bindings.Num()) != 0;
+	TArray<AActor*> ActorToExport;
+	int32 ActorCount = World->PersistentLevel->Actors.Num();
+	for (int32 ActorIndex = 0; ActorIndex < ActorCount; ++ActorIndex)
+	{
+		AActor* Actor = World->PersistentLevel->Actors[ActorIndex];
+		if (Actor != NULL)
+		{
+			FGuid ExistingGuid = Player->FindObjectId(*Actor, MovieSceneSequenceID::Root);
+			if (ExistingGuid.IsValid() && !bSelectedOnly || Bindings.Contains(ExistingGuid))
+			{
+				ActorToExport.Add(Actor);
+			}
+		}
+	}
+
+	// Export the persistent level and all of it's actors
+	const bool bSaveAnimSeq = false; //force off saving any AnimSequences since this can conflict when we export the level sequence animations.
+	Exporter->ExportLevelMesh(World->PersistentLevel, !bSelectedOnly, ActorToExport, NodeNameAdapter, bSaveAnimSeq);
+
+	// Export streaming levels and actors
+	for (int32 CurLevelIndex = 0; CurLevelIndex < World->GetNumLevels(); ++CurLevelIndex)
+	{
+		ULevel* CurLevel = World->GetLevel(CurLevelIndex);
+		if (CurLevel != NULL && CurLevel != (World->PersistentLevel))
+		{
+			Exporter->ExportLevelMesh(CurLevel, !bSelectedOnly, ActorToExport, NodeNameAdapter, bSaveAnimSeq);
+		}
+	}
+
+	Exporter->ExportLevelSequence(MovieScene, Bindings, Player, Template);
+
+	//Export all master tracks
+
+	for (UMovieSceneTrack* MasterTrack : MovieScene->GetMasterTracks())
+	{
+		TArray<UMovieSceneTrack*> Tracks;
+		Tracks.Add(MasterTrack);
+		Exporter->ExportLevelSequenceTracks(MovieScene, Player, nullptr, nullptr, Tracks);
+	}
+	// Save to disk
+	Exporter->WriteToFile(*InFBXFileName);
+
+	return true;
+}
diff --git a/Engine/Source/Editor/MovieSceneTools/Public/MovieSceneToolHelpers.h b/Engine/Source/Editor/MovieSceneTools/Public/MovieSceneToolHelpers.h
--- a/Engine/Source/Editor/MovieSceneTools/Public/MovieSceneToolHelpers.h
+++ b/Engine/Source/Editor/MovieSceneTools/Public/MovieSceneToolHelpers.h
@@ -15,6 +15,7 @@
 #include "MovieSceneTranslator.h"
 #include "MovieSceneCaptureSettings.h"
 
+
 class ISequencer;
 class UMovieScene;
 class UMovieSceneSection;
@@ -22,10 +23,28 @@ class UInterpTrackMoveAxis;
 struct FMovieSceneObjectBindingID;
 class UMovieSceneTrack;
 struct FMovieSceneEvaluationTrack;
-
+class UMovieSceneUserImportFBXSettings;
 struct FMovieSceneFloatValue;
+class INodeNameAdapter;
 template<typename ChannelType> struct TMovieSceneChannelData;
 
+namespace fbxsdk
+{
+	class FbxCamera;
+}
+namespace UnFbx
+{
+	class FFbxImporter;
+	class FFbxCurvesAPI;
+};
+
+struct FFBXInOutParameters
+{
+	bool bConvertSceneBackup;
+	bool bConvertSceneUnitBackup;
+	bool bForceFrontXAxisBackup;
+};
+
 DECLARE_DELEGATE_TwoParams(FOnEnumSelectionChanged, int32 /*Selection*/, ESelectInfo::Type /*SelectionType*/);
 
 class MOVIESCENETOOLS_API MovieSceneToolHelpers
@@ -190,14 +209,100 @@ public:
 	static void MovieSceneTranslatorLogOutput(FMovieSceneTranslator* InTranslator, TSharedRef<FMovieSceneTranslatorContext> InContext);
 
 	/**
-	* Import FBX
+	* Export FBX
+	*
+	* @param World The world to export from
+	* @param InMovieScene The movie scene to export frome
+	* @param MoviePlayer to use
+	* @param Bindings The sequencer binding map
+	* @param NodeNameAdaptor Adaptor to look up actor names.
+	* @param InFBXFileName the fbx file name.
+	* @param Template Movie scene sequence id.
+	* @return Whether the export was successful
+	*/
+	static bool ExportFBX(UWorld* World, UMovieScene* MovieScene, IMovieScenePlayer* Player, TArray<FGuid>& Bindings, INodeNameAdapter& NodeNameAdapter, FMovieSceneSequenceIDRef& Template,  const FString& InFBXFileName);
+
+	/**
+	* Import FBX with dialog
+	*
+	* @param InMovieScene The movie scene to import the fbx into
+	* @param InObjectBindingNameMap The object binding to name map to map import fbx animation onto
+	* @param bCreateCameras Whether to allow creation of cameras if found in the fbx file.
+	* @return Whether the import was successful
+	*/
+	static bool ImportFBXWithDialog(UMovieScene* InMovieScene, ISequencer& InSequencer, const TMap<FGuid, FString>& InObjectBindingNameMap, TOptional<bool> bCreateCameras);
+
+	/**
+	* Get FBX Ready for Import. This make sure the passed in file make be imported. After calling this call ImportReadiedFbx. It returns out some parameters that we forcably change so we reset them later.
+	*
+	* @param ImportFileName The filename to import into
+	* @param ImportFBXSettings FBX Import Settings to enforce.
+	* @param OutFBXParams Paremter to pass back to ImportReadiedFbx
+	* @return Whether the fbx file was ready and is ready to be import.
+	*/
+	static bool ReadyFBXForImport(const FString&  ImportFilename, UMovieSceneUserImportFBXSettings* ImportFBXSettings, FFBXInOutParameters& OutFBXParams);
+
+	/**
+	* Import into an FBX scene that has been readied already, via the ReadyFBXForImport call. We do this as two pass in case the client want's to do something, like create camera's, before actually
+	* loading the data.
+	*
+	* @param World The world to import the fbx into
+	* @param World The movie scene to import the fbx into
+	* @param ObjectBindingMap Map relating binding id's to track names. 
+	* @param ImportFBXSettings FBX Import Settings to enforce.
+	* @param InFBXParams Paremter from ImportReadiedFbx used to reset some internal fbx settings that we override.
+	* @return Whether the fbx file was ready and is ready to be import.
+	*/
+
+	static bool ImportFBXIfReady(UWorld* World, UMovieScene* MovieScene, IMovieScenePlayer* Player, TMap<FGuid, FString>& ObjectBindingMap, UMovieSceneUserImportFBXSettings* ImportFBXSettings,
+		const FFBXInOutParameters& InFBXParams);
+
+	/**
+	* Import FBX Camera to existing camera's
 	*
+	* @param FbxImporter The Fbx importer
 	* @param InMovieScene The movie scene to import the fbx into
+	* @param Player The player we are getting objects from.
+	* @param TemplateID Id of the sequence template ID.
 	* @param InObjectBindingNameMap The object binding to name map to map import fbx animation onto
 	* @param bCreateCameras Whether to allow creation of cameras if found in the fbx file.
+	* @param bNotifySlate  If an issue show's up, whether or not to notify the UI.
+	* @return Whether the import was successful
+	*/
+
+	static void ImportFBXCameraToExisting(UnFbx::FFbxImporter* FbxImporter, UMovieScene* InMovieScene, IMovieScenePlayer* Player, const FMovieSceneSequenceIDRef& TemplateID, TMap<FGuid, FString>& InObjectBindingMap, bool bMatchByNameOnly, bool bNotifySlate);
+
+	/**
+	* Import FBX Node to existing actor/node
+	*
+	* @param NodeName Name of fbx node/actor
+	* @param CurvesApi Existing FBX curves coming in
+	* @param InMovieScene The movie scene to import the fbx into
+	* @param Player The player we are getting objects from.
+	* @param TemplateID Id of the sequence template ID.
+	* @param ObjectBinding Guid of the object we are importing onto.
 	* @return Whether the import was successful
 	*/
-	static bool ImportFBX(UMovieScene* InMovieScene, ISequencer& InSequencer, const TMap<FGuid, FString>& InObjectBindingNameMap, TOptional<bool> bCreateCameras);
+	static bool ImportFBXNode(FString NodeName, UnFbx::FFbxCurvesAPI& CurveAPI, UMovieScene* InMovieScene, IMovieScenePlayer* Player, const FMovieSceneSequenceIDRef& TemplateID, FGuid ObjectBinding);
+
+
+	/**
+	*  Camera track was added, we usually do extra things, like add a Camera Cut tracks
+	*
+	* @param MovieScene MovieScene to add Camera.
+	* @param CameraGuid CameraGuid  Guid of the camera that was added.
+	* @param FrameNumber FrameNumber it's added at.
+	*/
+	static void CameraAdded(UMovieScene* MovieScene, FGuid CameraGuid, FFrameNumber FrameNumber);
+
+	/**
+	* Import FBX Camera to existing camera's
+	*
+	* @param CameraNode The Fbx camera
+	* @param InCameraActor Ue4 actor
+	*/
+	static void CopyCameraProperties(fbxsdk::FbxCamera* CameraNode, AActor* InCameraActor);
+
 
 	/*
 	 * Rich curve interpolation to matinee interpolation
diff --git a/Engine/Source/Editor/Sequencer/Private/Sequencer.cpp b/Engine/Source/Editor/Sequencer/Private/Sequencer.cpp
--- a/Engine/Source/Editor/Sequencer/Private/Sequencer.cpp
+++ b/Engine/Source/Editor/Sequencer/Private/Sequencer.cpp
@@ -854,7 +854,7 @@ FGuid FSequencer::CreateBinding(UObject& InObject, const FString& InName)
 
 	// Attempt to use the parent as a context if necessary
 	UObject* ParentObject = OwnerSequence->GetParentObject(&InObject);
-	UObject* BindingContext = GetPlaybackContext();
+	UObject* BindingContext = GetPlaybackContext(); //UWorld
 
 	AActor* ParentActorAdded = nullptr;
 	FGuid ParentGuid;
@@ -9131,46 +9131,7 @@ void FSequencer::NewCameraAdded(FGuid CameraGuid, ACameraActor* NewCamera)
 	UMovieSceneSequence* Sequence = GetFocusedMovieSceneSequence();
 	UMovieScene* OwnerMovieScene = Sequence->GetMovieScene();
 
-	// If there's a cinematic shot track, no need to set this camera to a shot
-	UMovieSceneTrack* CinematicShotTrack = OwnerMovieScene->FindMasterTrack(UMovieSceneCinematicShotTrack::StaticClass());
-	if (CinematicShotTrack)
-	{
-		return;
-	}
-
-	UMovieSceneTrack* CameraCutTrack = OwnerMovieScene->GetCameraCutTrack();
-
-	// If there's a camera cut track with at least one section, no need to change the section
-	if (CameraCutTrack && CameraCutTrack->GetAllSections().Num() > 0)
-	{
-		return;
-	}
-
-	if (!CameraCutTrack)
-	{
-		CameraCutTrack = OwnerMovieScene->AddCameraCutTrack(UMovieSceneCameraCutTrack::StaticClass());
-	}
-
-	if (CameraCutTrack)
-	{
-		UMovieSceneSection* Section = MovieSceneHelpers::FindSectionAtTime(CameraCutTrack->GetAllSections(), GetLocalTime().Time.FloorToFrame());
-		UMovieSceneCameraCutSection* CameraCutSection = Cast<UMovieSceneCameraCutSection>(Section);
-
-		if (CameraCutSection)
-		{
-			CameraCutSection->Modify();
-			CameraCutSection->SetCameraGuid(CameraGuid);
-		}
-		else
-		{
-			CameraCutTrack->Modify();
-
-			UMovieSceneCameraCutSection* NewSection = Cast<UMovieSceneCameraCutSection>(CameraCutTrack->CreateNewSection());
-			NewSection->SetRange(GetPlaybackRange());
-			NewSection->SetCameraGuid(CameraGuid);
-			CameraCutTrack->AddSection(*NewSection);
-		}
-	}
+	MovieSceneToolHelpers::CameraAdded(OwnerMovieScene, CameraGuid, GetLocalTime().Time.FloorToFrame());
 }
 
 
@@ -9298,7 +9259,7 @@ void FSequencer::ImportFBX()
 		ObjectBindingNameMap.Add(ObjectBinding, RootObjectBindingNode.Get().GetDisplayName().ToString());
 	}
 
-	MovieSceneToolHelpers::ImportFBX(MovieScene, *this, ObjectBindingNameMap, TOptional<bool>());
+	MovieSceneToolHelpers::ImportFBXWithDialog(MovieScene, *this, ObjectBindingNameMap, TOptional<bool>());
 }
 
 void FSequencer::ImportFBXOntoSelectedNodes()
@@ -9320,7 +9281,7 @@ void FSequencer::ImportFBXOntoSelectedNodes()
 		}
 	}
 
-	MovieSceneToolHelpers::ImportFBX(MovieScene, *this, ObjectBindingNameMap, TOptional<bool>(false));
+	MovieSceneToolHelpers::ImportFBXWithDialog(MovieScene, *this, ObjectBindingNameMap, TOptional<bool>(false));
 }
 
 void FSequencer::ExportFBX()
@@ -9378,6 +9339,9 @@ void FSequencer::ExportFBX()
 		FString ExportFilename = SaveFilenames[0];
 		FEditorDirectories::Get().SetLastDirectory( ELastDirectory::FBX, FPaths::GetPath( ExportFilename ) ); // Save path as default for next time.
 
+		// Make sure external selection is up to date since export could happen on tracks that have been right clicked but not have their underlying bound objects selected yet since that happens on mouse up.
+		SynchronizeExternalSelectionWithSequencerSelection();
+		
 		// Select selected nodes if there are selected nodes
 		TArray<FGuid> Bindings;
 		for (const TSharedRef<FSequencerDisplayNode>& Node : Selection.GetSelectedOutlinerNodes())
@@ -9451,69 +9415,10 @@ void FSequencer::ExportFBXInternal(const FString& ExportFilename, TArray<FGuid>&
 		if (!ExportCancel)
 		{
 			UMovieScene* MovieScene = GetFocusedMovieSceneSequence()->GetMovieScene();
-
-			Exporter->CreateDocument();
-			Exporter->SetTrasformBaking(false);
-			Exporter->SetKeepHierarchy(true);
-
-			const bool bSelectedOnly = (Selection.GetSelectedTracks().Num() + Bindings.Num()) != 0;
-
-			// Make sure external selection is up to date since export could happen on tracks that have been right clicked but not have their underlying bound objects selected yet since that happens on mouse up.
-			if (bSelectedOnly)
-			{
-				SynchronizeExternalSelectionWithSequencerSelection();
-			}
-
-			UnFbx::FFbxExporter::FLevelSequenceNodeNameAdapter NodeNameAdapter(MovieScene, this, GetFocusedTemplateID());
-
-			// Export the persistent level and all of it's actors
 			UWorld* World = Cast<UWorld>(GetPlaybackContext());
-			const bool bSaveAnimSeq = false; //force off saving any AnimSequences since this can conflict when we export the level sequence animations.
-			Exporter->ExportLevelMesh(World->PersistentLevel, bSelectedOnly, NodeNameAdapter, bSaveAnimSeq);
-
-			// Export streaming levels and actors
-			for (int32 CurLevelIndex = 0; CurLevelIndex < World->GetNumLevels(); ++CurLevelIndex)
-			{
-				ULevel* CurLevel = World->GetLevel(CurLevelIndex);
-				if (CurLevel != NULL && CurLevel != (World->PersistentLevel))
-				{
-					Exporter->ExportLevelMesh(CurLevel, bSelectedOnly, NodeNameAdapter, bSaveAnimSeq);
-				}
-			}
-
-			// Export the movie scene data.
-			Exporter->ExportLevelSequence(MovieScene, Bindings, this, GetFocusedTemplateID());
-
-			// Export selected or all master tracks
-			if (Selection.GetSelectedOutlinerNodes().Num())
-			{
-				for (const TSharedRef<FSequencerDisplayNode>& Node : Selection.GetSelectedOutlinerNodes())
-				{
-					if (Node->GetType() == ESequencerNode::Track)
-					{
-						TSharedRef<FSequencerTrackNode> TrackNode = StaticCastSharedRef<FSequencerTrackNode>(Node);
-						UMovieSceneTrack* MasterTrack = TrackNode->GetTrack();
-						if (MovieScene->GetMasterTracks().Contains(MasterTrack))
-						{
-							TArray<UMovieSceneTrack*> Tracks;
-							Tracks.Add(MasterTrack);
-							Exporter->ExportLevelSequenceTracks(MovieScene, this, nullptr, nullptr, Tracks);
-						}
-					}
-				}
-			}
-			else
-			{
-				for (UMovieSceneTrack* MasterTrack : MovieScene->GetMasterTracks())
-				{
-					TArray<UMovieSceneTrack*> Tracks;
-					Tracks.Add(MasterTrack);
-					Exporter->ExportLevelSequenceTracks(MovieScene, this, nullptr, nullptr, Tracks);
-				}
-			}
-
-			// Save to disk
-			Exporter->WriteToFile(*ExportFilename);
+			FMovieSceneSequenceIDRef Template = GetFocusedTemplateID();
+			UnFbx::FFbxExporter::FLevelSequenceNodeNameAdapter NodeNameAdapter(MovieScene, this, Template);
+			MovieSceneToolHelpers::ExportFBX(World,MovieScene, this, Bindings, NodeNameAdapter, Template, ExportFilename);
 		}
 	}
 }
diff --git a/Engine/Source/Editor/UnrealEd/Private/Fbx/FbxMainExport.cpp b/Engine/Source/Editor/UnrealEd/Private/Fbx/FbxMainExport.cpp
--- a/Engine/Source/Editor/UnrealEd/Private/Fbx/FbxMainExport.cpp
+++ b/Engine/Source/Editor/UnrealEd/Private/Fbx/FbxMainExport.cpp
@@ -418,24 +418,25 @@ static void SortActorsHierarchy(TArray<AActor*>& Actors)
 	});
 }
 
+
 /**
  * Exports the basic scene information to the FBX document.
  */
-void FFbxExporter::ExportLevelMesh( ULevel* InLevel, bool bSelectedOnly, INodeNameAdapter& NodeNameAdapter, bool bSaveAnimSeq)
+void FFbxExporter::ExportLevelMesh(ULevel* InLevel, bool bExportLevelGeometry, TArray<AActor*>& ActorToExport, INodeNameAdapter& NodeNameAdapter, bool bSaveAnimSeq)
 {
 	if (InLevel == NULL)
 	{
 		return;
 	}
 
-	if( !bSelectedOnly )
+	if (bExportLevelGeometry)
 	{
 		// Exports the level's scene geometry
 		// the vertex number of Model must be more than 2 (at least a triangle panel)
 		if (InLevel->Model != NULL && InLevel->Model->VertexBuffer.Vertices.Num() > 2 && InLevel->Model->MaterialIndexBuffers.Num() > 0)
 		{
 			// create a FbxNode
-			FbxNode* Node = FbxNode::Create(Scene,"LevelMesh");
+			FbxNode* Node = FbxNode::Create(Scene, "LevelMesh");
 
 			// set the shading mode to view texture
 			Node->SetShadingMode(FbxNode::eTextureShading);
@@ -448,21 +449,9 @@ void FFbxExporter::ExportLevelMesh( ULevel* InLevel, bool bSelectedOnly, INodeNa
 		}
 	}
 
-	TArray<AActor*> ActorToExport;
-	int32 ActorCount = InLevel->Actors.Num();
-	for (int32 ActorIndex = 0; ActorIndex < ActorCount; ++ActorIndex)
-	{
-		AActor* Actor = InLevel->Actors[ActorIndex];
-		if (Actor != NULL && (!bSelectedOnly || (bSelectedOnly && Actor->IsSelected())))
-		{
-			ActorToExport.Add(Actor);
-		}
-	}
-
 	//Sort the hierarchy to make sure parent come first
 	SortActorsHierarchy(ActorToExport);
-
-	ActorCount = ActorToExport.Num();
+	int32 ActorCount = ActorToExport.Num();
 	for (int32 ActorIndex = 0; ActorIndex < ActorCount; ++ActorIndex)
 	{
 		AActor* Actor = ActorToExport[ActorIndex];
@@ -485,11 +474,11 @@ void FFbxExporter::ExportLevelMesh( ULevel* InLevel, bool bSelectedOnly, INodeNa
 		}
 		else if (Actor->IsA(ALight::StaticClass()))
 		{
-			ExportLight((ALight*) Actor, NodeNameAdapter );
+			ExportLight((ALight*)Actor, NodeNameAdapter);
 		}
 		else if (Actor->IsA(AStaticMeshActor::StaticClass()))
 		{
-			ExportStaticMesh( Actor, CastChecked<AStaticMeshActor>(Actor)->GetStaticMeshComponent(), NodeNameAdapter );
+			ExportStaticMesh(Actor, CastChecked<AStaticMeshActor>(Actor)->GetStaticMeshComponent(), NodeNameAdapter);
 		}
 		else if (Actor->IsA(ALandscapeProxy::StaticClass()))
 		{
@@ -498,13 +487,13 @@ void FFbxExporter::ExportLevelMesh( ULevel* InLevel, bool bSelectedOnly, INodeNa
 		else if (Actor->IsA(ABrush::StaticClass()))
 		{
 			// All brushes should be included within the world geometry exported above.
-			ExportBrush((ABrush*) Actor, NULL, 0, NodeNameAdapter );
+			ExportBrush((ABrush*)Actor, NULL, 0, NodeNameAdapter);
 		}
 		else if (Actor->IsA(AEmitter::StaticClass()))
 		{
-			ExportActor( Actor, false, NodeNameAdapter, bSaveAnimSeq ); // Just export the placement of the particle emitter.
+			ExportActor(Actor, false, NodeNameAdapter, bSaveAnimSeq); // Just export the placement of the particle emitter.
 		}
-		else if(Actor->IsA(ACameraActor::StaticClass()))
+		else if (Actor->IsA(ACameraActor::StaticClass()))
 		{
 			ExportCamera(CastChecked<ACameraActor>(Actor), false, NodeNameAdapter);
 		}
@@ -516,6 +505,30 @@ void FFbxExporter::ExportLevelMesh( ULevel* InLevel, bool bSelectedOnly, INodeNa
 	}
 }
 
+/**
+ * Exports the basic scene information to the FBX document.
+ */
+void FFbxExporter::ExportLevelMesh( ULevel* InLevel, bool bSelectedOnly, INodeNameAdapter& NodeNameAdapter, bool bSaveAnimSeq)
+{
+	if (InLevel == NULL)
+	{
+		return;
+	}
+	
+	TArray<AActor*> ActorToExport;
+	int32 ActorCount = InLevel->Actors.Num();
+	for (int32 ActorIndex = 0; ActorIndex < ActorCount; ++ActorIndex)
+	{
+		AActor* Actor = InLevel->Actors[ActorIndex];
+		if (Actor != NULL && (!bSelectedOnly || (bSelectedOnly && Actor->IsSelected())))
+		{
+			ActorToExport.Add(Actor);
+		}
+	}
+
+	ExportLevelMesh(InLevel, !bSelectedOnly, ActorToExport, NodeNameAdapter, bSaveAnimSeq);
+}
+
 void FFbxExporter::FillFbxLightAttribute(FbxLight* Light, FbxNode* FbxParentNode, ULightComponent* BaseLight)
 {
 	Light->Intensity.Set(BaseLight->Intensity);
diff --git a/Engine/Source/Editor/UnrealEd/Private/FbxExporter.h b/Engine/Source/Editor/UnrealEd/Private/FbxExporter.h
--- a/Engine/Source/Editor/UnrealEd/Private/FbxExporter.h
+++ b/Engine/Source/Editor/UnrealEd/Private/FbxExporter.h
@@ -113,6 +113,12 @@ public:
 	 */
 	virtual void ExportLevelMesh( ULevel* InLevel, bool bSelectedOnly, INodeNameAdapter& NodeNameAdapter , bool bSaveAnimSeq = true);
 
+	/**
+	* Exports the basic scene information to the FBX document, using the passed in Actors
+	*/
+	virtual void ExportLevelMesh(ULevel* InLevel, bool bExportLevelGeometry, TArray<AActor*>& ActorToExport, INodeNameAdapter& NodeNameAdapter, bool bSaveAnimSeq = true);
+
+
 	/**
 	 * Exports the given Matinee sequence information into a FBX document.
 	 * 
