Fix CIS

diff --git a/Engine/Source/Editor/MovieSceneTools/Private/MovieSceneToolHelpers.cpp b/Engine/Source/Editor/MovieSceneTools/Private/MovieSceneToolHelpers.cpp
--- a/Engine/Source/Editor/MovieSceneTools/Private/MovieSceneToolHelpers.cpp
+++ b/Engine/Source/Editor/MovieSceneTools/Private/MovieSceneToolHelpers.cpp
@@ -860,7 +860,7 @@ void MovieSceneToolHelpers::MovieSceneTranslatorLogOutput(FMovieSceneTranslator
 	}
 }
 
-static FGuid GetHandleToObject(UObject* InObject, UMovieScene* InMovieScene, IMovieScenePlayer* Player, const FMovieSceneSequenceIDRef& TemplateID)
+static FGuid GetHandleToObject(UObject* InObject, UMovieScene* InMovieScene, IMovieScenePlayer* Player, FMovieSceneSequenceIDRef TemplateID)
 {
 	// Attempt to resolve the object through the movie scene instance first, 
 	FGuid PropertyOwnerGuid = FGuid();
@@ -888,7 +888,7 @@ static FGuid GetHandleToObject(UObject* InObject, UMovieScene* InMovieScene, IMo
 	return PropertyOwnerGuid;
 }
 
-bool ImportFBXProperty(FString NodeName, FString AnimatedPropertyName, FGuid ObjectBinding, UnFbx::FFbxCurvesAPI& CurveAPI, UMovieScene* InMovieScene, IMovieScenePlayer* Player, const FMovieSceneSequenceIDRef& TemplateID)
+bool ImportFBXProperty(FString NodeName, FString AnimatedPropertyName, FGuid ObjectBinding, UnFbx::FFbxCurvesAPI& CurveAPI, UMovieScene* InMovieScene, IMovieScenePlayer* Player, FMovieSceneSequenceIDRef TemplateID)
 {
 	const UMovieSceneToolsProjectSettings* ProjectSettings = GetDefault<UMovieSceneToolsProjectSettings>();
 	const UMovieSceneUserImportFBXSettings* ImportFBXSettings = GetDefault<UMovieSceneUserImportFBXSettings>();
@@ -1183,7 +1183,7 @@ bool ImportFBXTransform(FString NodeName, FGuid ObjectBinding, UnFbx::FFbxCurves
 	return true;
 }
 
-bool MovieSceneToolHelpers::ImportFBXNode(FString NodeName, UnFbx::FFbxCurvesAPI& CurveAPI, UMovieScene* InMovieScene, IMovieScenePlayer* Player, const FMovieSceneSequenceIDRef& TemplateID, FGuid ObjectBinding)
+bool MovieSceneToolHelpers::ImportFBXNode(FString NodeName, UnFbx::FFbxCurvesAPI& CurveAPI, UMovieScene* InMovieScene, IMovieScenePlayer* Player, FMovieSceneSequenceIDRef TemplateID, FGuid ObjectBinding)
 {
 	// Look for animated float properties
 	TArray<FString> AnimatedPropertyNames;
@@ -1329,7 +1329,7 @@ FString GetCameraName(FbxCamera* InCamera)
 }
 
 
-void MovieSceneToolHelpers::ImportFBXCameraToExisting(UnFbx::FFbxImporter* FbxImporter, UMovieScene* InMovieScene, IMovieScenePlayer* Player, const FMovieSceneSequenceIDRef& TemplateID, TMap<FGuid, FString>& InObjectBindingMap, bool bMatchByNameOnly, bool bNotifySlate)
+void MovieSceneToolHelpers::ImportFBXCameraToExisting(UnFbx::FFbxImporter* FbxImporter, UMovieScene* InMovieScene, IMovieScenePlayer* Player, FMovieSceneSequenceIDRef TemplateID, TMap<FGuid, FString>& InObjectBindingMap, bool bMatchByNameOnly, bool bNotifySlate)
 {
 	for (auto InObjectBinding : InObjectBindingMap)
 	{
diff --git a/Engine/Source/Editor/MovieSceneTools/Public/MovieSceneToolHelpers.h b/Engine/Source/Editor/MovieSceneTools/Public/MovieSceneToolHelpers.h
--- a/Engine/Source/Editor/MovieSceneTools/Public/MovieSceneToolHelpers.h
+++ b/Engine/Source/Editor/MovieSceneTools/Public/MovieSceneToolHelpers.h
@@ -270,7 +270,7 @@ public:
 	* @return Whether the import was successful
 	*/
 
-	static void ImportFBXCameraToExisting(UnFbx::FFbxImporter* FbxImporter, UMovieScene* InMovieScene, IMovieScenePlayer* Player, const FMovieSceneSequenceIDRef& TemplateID, TMap<FGuid, FString>& InObjectBindingMap, bool bMatchByNameOnly, bool bNotifySlate);
+	static void ImportFBXCameraToExisting(UnFbx::FFbxImporter* FbxImporter, UMovieScene* InMovieScene, IMovieScenePlayer* Player, FMovieSceneSequenceIDRef TemplateID, TMap<FGuid, FString>& InObjectBindingMap, bool bMatchByNameOnly, bool bNotifySlate);
 
 	/**
 	* Import FBX Node to existing actor/node
@@ -283,7 +283,7 @@ public:
 	* @param ObjectBinding Guid of the object we are importing onto.
 	* @return Whether the import was successful
 	*/
-	static bool ImportFBXNode(FString NodeName, UnFbx::FFbxCurvesAPI& CurveAPI, UMovieScene* InMovieScene, IMovieScenePlayer* Player, const FMovieSceneSequenceIDRef& TemplateID, FGuid ObjectBinding);
+	static bool ImportFBXNode(FString NodeName, UnFbx::FFbxCurvesAPI& CurveAPI, UMovieScene* InMovieScene, IMovieScenePlayer* Player, FMovieSceneSequenceIDRef TemplateID, FGuid ObjectBinding);
 
 
 	/**
