Sequencer: Allow for skeletal mesh components inside child actor components

diff --git a/Engine/Source/Editor/MovieSceneTools/Private/TrackEditors/SkeletalAnimationTrackEditor.cpp b/Engine/Source/Editor/MovieSceneTools/Private/TrackEditors/SkeletalAnimationTrackEditor.cpp
--- a/Engine/Source/Editor/MovieSceneTools/Private/TrackEditors/SkeletalAnimationTrackEditor.cpp
+++ b/Engine/Source/Editor/MovieSceneTools/Private/TrackEditors/SkeletalAnimationTrackEditor.cpp
@@ -94,7 +94,17 @@ USkeleton* AcquireSkeletonFromObjectGuid(const FGuid& Guid, TSharedPtr<ISequence
 {
 	UObject* BoundObject = SequencerPtr.IsValid() ? SequencerPtr->FindSpawnedObjectOrTemplate(Guid) : nullptr;
 
-	if (AActor* Actor = Cast<AActor>(BoundObject))
+	AActor* Actor = Cast<AActor>(BoundObject);
+
+	if (!Actor)
+	{
+		if (UChildActorComponent* ChildActorComponent = Cast<UChildActorComponent>(BoundObject))
+		{
+			Actor = ChildActorComponent->GetChildActor();
+		}
+	}
+
+	if (Actor)
 	{
 		for (UActorComponent* Component : Actor->GetComponents())
 		{
@@ -470,7 +480,7 @@ bool FSkeletalAnimationTrackEditor::HandleAssetAdded(UObject* Asset, const FGuid
 
 void FSkeletalAnimationTrackEditor::BuildObjectBindingTrackMenu(FMenuBuilder& MenuBuilder, const TArray<FGuid>& ObjectBindings, const UClass* ObjectClass)
 {
-	if (ObjectClass->IsChildOf(USkeletalMeshComponent::StaticClass()) || ObjectClass->IsChildOf(AActor::StaticClass()))
+	if (ObjectClass->IsChildOf(USkeletalMeshComponent::StaticClass()) || ObjectClass->IsChildOf(AActor::StaticClass()) || ObjectClass->IsChildOf(UChildActorComponent::StaticClass()))
 	{
 		const TSharedPtr<ISequencer> ParentSequencer = GetSequencer();
 
diff --git a/Engine/Source/Runtime/MovieSceneTracks/Private/Evaluation/MovieSceneSkeletalAnimationTemplate.cpp b/Engine/Source/Runtime/MovieSceneTracks/Private/Evaluation/MovieSceneSkeletalAnimationTemplate.cpp
--- a/Engine/Source/Runtime/MovieSceneTracks/Private/Evaluation/MovieSceneSkeletalAnimationTemplate.cpp
+++ b/Engine/Source/Runtime/MovieSceneTracks/Private/Evaluation/MovieSceneSkeletalAnimationTemplate.cpp
@@ -298,8 +298,19 @@ namespace MovieScene
 			{
 				return SkeletalMeshComponent;
 			}
+
 			// then check to see if we are controlling an actor & if so use its first USkeletalMeshComponent 
-			else if (AActor* Actor = Cast<AActor>(InObject))
+			AActor* Actor = Cast<AActor>(InObject);
+
+			if (!Actor)
+			{
+				if (UChildActorComponent* ChildActorComponent = Cast<UChildActorComponent>(InObject))
+				{
+					Actor = ChildActorComponent->GetChildActor();
+				}
+			}
+
+			if (Actor)
 			{
 				return Actor->FindComponentByClass<USkeletalMeshComponent>();
 			}
