Sequencer: Fix monolithic builds

diff --git a/Engine/Plugins/MovieScene/SequencerScripting/Source/SequencerScriptingEditor/Private/SequencerTools.cpp b/Engine/Plugins/MovieScene/SequencerScripting/Source/SequencerScriptingEditor/Private/SequencerTools.cpp
--- a/Engine/Plugins/MovieScene/SequencerScripting/Source/SequencerScriptingEditor/Private/SequencerTools.cpp
+++ b/Engine/Plugins/MovieScene/SequencerScripting/Source/SequencerScriptingEditor/Private/SequencerTools.cpp
@@ -131,32 +131,6 @@ bool USequencerToolsFunctionLibrary::ExportFBX(UWorld* World, ULevelSequence* Se
 	return bDidExport;
 }
 
-FString GetCameraName(FbxCamera* InCamera)
-{
-	FbxNode* CameraNode = InCamera->GetNode();
-	if (CameraNode)
-	{
-		return CameraNode->GetName();
-	}
-
-	return InCamera->GetName();
-}
-
-void GetCameras(FbxNode* Parent, TArray<FbxCamera*>& Cameras)
-{
-	FbxCamera* Camera = Parent->GetCamera();
-	if (Camera)
-	{
-		Cameras.Add(Camera);
-	}
-
-	int32 NodeCount = Parent->GetChildCount();
-	for (int32 NodeIndex = 0; NodeIndex < NodeCount; ++NodeIndex)
-	{
-		FbxNode* Child = Parent->GetChild(NodeIndex);
-		GetCameras(Child, Cameras);
-	}
-}
 
 
 TArray<FGuid> AddActors(UWorld* World, UMovieSceneSequence* InSequence, UMovieScene* InMovieScene, IMovieScenePlayer* Player, FMovieSceneSequenceIDRef TemplateID,const TArray<TWeakObjectPtr<AActor> >& InActors)
@@ -198,13 +172,13 @@ void ImportFBXCamera(UnFbx::FFbxImporter* FbxImporter, UWorld* World, UMovieScen
 	if (bCreateCameras)
 	{
 		TArray<FbxCamera*> AllCameras;
-		GetCameras(FbxImporter->Scene->GetRootNode(), AllCameras);
+		MovieSceneToolHelpers::GetCameras(FbxImporter->Scene->GetRootNode(), AllCameras);
 
 		// Find unmatched cameras
 		TArray<FbxCamera*> UnmatchedCameras;
 		for (auto Camera : AllCameras)
 		{
-			FString NodeName = GetCameraName(Camera);
+			FString NodeName = MovieSceneToolHelpers::GetCameraName(Camera);
 
 			bool bMatched = false;
 			for (auto InObjectBinding : InObjectBindingMap)
@@ -242,7 +216,7 @@ void ImportFBXCamera(UnFbx::FFbxImporter* FbxImporter, UWorld* World, UMovieScen
 		// Add any unmatched cameras
 		for (auto UnmatchedCamera : UnmatchedCameras)
 		{
-			FString CameraName = GetCameraName(UnmatchedCamera);
+			FString CameraName = MovieSceneToolHelpers::GetCameraName(UnmatchedCamera);
 
 			AActor* NewCamera = nullptr;
 			if (UnmatchedCamera->GetApertureMode() == FbxCamera::eFocalLength)
diff --git a/Engine/Source/Editor/MovieSceneTools/Private/MovieSceneToolHelpers.cpp b/Engine/Source/Editor/MovieSceneTools/Private/MovieSceneToolHelpers.cpp
--- a/Engine/Source/Editor/MovieSceneTools/Private/MovieSceneToolHelpers.cpp
+++ b/Engine/Source/Editor/MovieSceneTools/Private/MovieSceneToolHelpers.cpp
@@ -1199,7 +1199,7 @@ bool MovieSceneToolHelpers::ImportFBXNode(FString NodeName, UnFbx::FFbxCurvesAPI
 	return true;
 }
 
-void GetCameras( FbxNode* Parent, TArray<FbxCamera*>& Cameras )
+void MovieSceneToolHelpers::GetCameras( FbxNode* Parent, TArray<FbxCamera*>& Cameras )
 {
 	FbxCamera* Camera = Parent->GetCamera();
 	if( Camera )
@@ -1317,7 +1317,7 @@ void MovieSceneToolHelpers::CopyCameraProperties(FbxCamera* CameraNode, AActor*
 	CameraComponent->SetFieldOfView(FieldOfView);
 }
 
-FString GetCameraName(FbxCamera* InCamera)
+FString MovieSceneToolHelpers::GetCameraName(FbxCamera* InCamera)
 {
 	FbxNode* CameraNode = InCamera->GetNode();
 	if (CameraNode)
@@ -1452,7 +1452,7 @@ void ImportFBXCamera(UnFbx::FFbxImporter* FbxImporter, UMovieScene* InMovieScene
 	if (bCreateCameras)
 	{
 		TArray<FbxCamera*> AllCameras;
-		GetCameras(FbxImporter->Scene->GetRootNode(), AllCameras);
+		MovieSceneToolHelpers::GetCameras(FbxImporter->Scene->GetRootNode(), AllCameras);
 
 		UWorld* World = GCurrentLevelEditingViewportClient ? GCurrentLevelEditingViewportClient->GetWorld() : nullptr;
 
@@ -1460,7 +1460,7 @@ void ImportFBXCamera(UnFbx::FFbxImporter* FbxImporter, UMovieScene* InMovieScene
 		TArray<FbxCamera*> UnmatchedCameras;
 		for (auto Camera : AllCameras)
 		{
-			FString NodeName = GetCameraName(Camera);
+			FString NodeName = MovieSceneToolHelpers::GetCameraName(Camera);
 
 			bool bMatched = false;
 			for (auto InObjectBinding : InObjectBindingMap)
@@ -1505,7 +1505,7 @@ void ImportFBXCamera(UnFbx::FFbxImporter* FbxImporter, UMovieScene* InMovieScene
 		// Add any unmatched cameras
 		for (auto UnmatchedCamera : UnmatchedCameras)
 		{
-			FString CameraName = GetCameraName(UnmatchedCamera);
+			FString CameraName = MovieSceneToolHelpers::GetCameraName(UnmatchedCamera);
 
 			AActor* NewCamera = nullptr;
 			if (UnmatchedCamera->GetApertureMode() == FbxCamera::eFocalLength)
@@ -1541,7 +1541,7 @@ void ImportFBXCamera(UnFbx::FFbxImporter* FbxImporter, UMovieScene* InMovieScene
 
 FGuid FindCameraGuid(FbxCamera* Camera, TMap<FGuid, FString>& InObjectBindingMap)
 {
-	FString CameraName = GetCameraName(Camera);
+	FString CameraName = MovieSceneToolHelpers::GetCameraName(Camera);
 
 	for (auto& Pair : InObjectBindingMap)
 	{
@@ -1587,7 +1587,7 @@ void ImportCameraCut(UnFbx::FFbxImporter* FbxImporter, UMovieScene* InMovieScene
 
 	// The camera switcher camera index refer to depth-first found order of the camera in the FBX
 	TArray<FbxCamera*> AllCameras;
-	GetCameras(FbxImporter->Scene->GetRootNode(), AllCameras);
+	MovieSceneToolHelpers::GetCameras(FbxImporter->Scene->GetRootNode(), AllCameras);
 
 	UMovieSceneCameraCutTrack* CameraCutTrack = GetCameraCutTrack(InMovieScene);
 	FFrameRate FrameRate = CameraCutTrack->GetTypedOuter<UMovieScene>()->GetTickResolution();
diff --git a/Engine/Source/Editor/MovieSceneTools/Public/MovieSceneToolHelpers.h b/Engine/Source/Editor/MovieSceneTools/Public/MovieSceneToolHelpers.h
--- a/Engine/Source/Editor/MovieSceneTools/Public/MovieSceneToolHelpers.h
+++ b/Engine/Source/Editor/MovieSceneTools/Public/MovieSceneToolHelpers.h
@@ -31,6 +31,7 @@ template<typename ChannelType> struct TMovieSceneChannelData;
 namespace fbxsdk
 {
 	class FbxCamera;
+	class FbxNode;
 }
 namespace UnFbx
 {
@@ -343,6 +344,15 @@ public:
 	*/
 	static FMovieSceneEvaluationTrack* GetEvaluationTrack(ISequencer *Sequencer, const FGuid& TrackSignature);
 
+	/*
+	 * Get the fbx cameras from the requested parent node
+	 */
+	static void GetCameras(fbxsdk::FbxNode* Parent, TArray<fbxsdk::FbxCamera*>& Cameras);
+
+	/*
+	 * Get the fbx camera name
+	 */
+	static FString GetCameraName(fbxsdk::FbxCamera* InCamera);
 };
 
 class FTrackEditorBindingIDPicker : public FMovieSceneObjectBindingIDPicker
