Curve Editor: Multi-key math operations in the time widget fix

diff --git a/Engine/Source/Editor/MovieSceneTools/Private/FrameNumberDetailsCustomization.cpp b/Engine/Source/Editor/MovieSceneTools/Private/FrameNumberDetailsCustomization.cpp
--- a/Engine/Source/Editor/MovieSceneTools/Private/FrameNumberDetailsCustomization.cpp
+++ b/Engine/Source/Editor/MovieSceneTools/Private/FrameNumberDetailsCustomization.cpp
@@ -79,14 +79,20 @@ FText FFrameNumberDetailsCustomization::OnGetTimeText() const
 
 void FFrameNumberDetailsCustomization::OnTimeTextCommitted(const FText& InText, ETextCommit::Type CommitInfo)
 {
-	int32 ExistingValue = 0.0;
-	FrameNumberProperty->GetValue(ExistingValue);
-	TOptional<double> TickResolution = NumericTypeInterface->FromString(InText.ToString(), ExistingValue);
-	if (TickResolution.IsSet())
+	TArray<FString> PerObjectValueStrs;
+	FrameNumberProperty->GetPerObjectValues(PerObjectValueStrs);
+
+	for (FString& ValueStr : PerObjectValueStrs)
 	{
-		double ClampedValue = FMath::Clamp(TickResolution.GetValue(), (double)UIClampMin, (double)UIClampMax);
-		FrameNumberProperty->SetValue((int32)ClampedValue);
+		int32 ExistingValue = FCString::Atoi(*ValueStr);
+		TOptional<double> TickResolution = NumericTypeInterface->FromString(InText.ToString(), ExistingValue);
+		if (TickResolution.IsSet())
+		{
+			double ClampedValue = FMath::Clamp(TickResolution.GetValue(), (double)UIClampMin, (double)UIClampMax);
+			ValueStr = FString::FromInt((int32)ClampedValue);
+		}
 	}
+	FrameNumberProperty->SetPerObjectValues(PerObjectValueStrs);
 }
 
 #undef LOCTEXT_NAMESPACE
