Sequencer: Added analytics events for record new sequence and record actor in sequencer.

diff --git a/Engine/Source/Editor/MovieSceneTools/MovieSceneTools.Build.cs b/Engine/Source/Editor/MovieSceneTools/MovieSceneTools.Build.cs
--- a/Engine/Source/Editor/MovieSceneTools/MovieSceneTools.Build.cs
+++ b/Engine/Source/Editor/MovieSceneTools/MovieSceneTools.Build.cs
@@ -71,8 +71,9 @@ public class MovieSceneTools : ModuleRules
 				"SceneOutliner",
                 "MainFrame",
                 "UnrealEd",
+                "Analytics",
             }
-		);
+        );
 
 		DynamicallyLoadedModuleNames.AddRange(
 			new string[] {
diff --git a/Engine/Source/Editor/MovieSceneTools/Private/TrackEditors/SubTrackEditor.cpp b/Engine/Source/Editor/MovieSceneTools/Private/TrackEditors/SubTrackEditor.cpp
--- a/Engine/Source/Editor/MovieSceneTools/Private/TrackEditors/SubTrackEditor.cpp
+++ b/Engine/Source/Editor/MovieSceneTools/Private/TrackEditors/SubTrackEditor.cpp
@@ -26,6 +26,8 @@
 #include "MovieSceneToolHelpers.h"
 #include "Misc/QualifiedFrameTime.h"
 #include "MovieSceneTimeHelpers.h"
+#include "EngineAnalytics.h"
+#include "Interfaces/IAnalyticsProvider.h"
 
 #include "CommonMovieSceneTools.h"
 
@@ -908,6 +910,12 @@ bool FSubTrackEditor::CanRecordNewSequence() const
 
 void FSubTrackEditor::HandleRecordNewSequence(AActor* InActorToRecord, UMovieSceneTrack* InTrack)
 {
+	// Keep track of how many people actually used record new sequence
+	if (FEngineAnalytics::IsAvailable())
+	{
+		FEngineAnalytics::GetProvider().RecordEvent(TEXT("Editor.Sequencer.RecordNewSequence"));
+	}
+
 	FSlateApplication::Get().DismissAllMenus();
 
 	const FScopedTransaction Transaction(LOCTEXT("AddRecordNewSequence_Transaction", "Add Record New Sequence"));
diff --git a/Engine/Source/Editor/Sequencer/Private/Sequencer.cpp b/Engine/Source/Editor/Sequencer/Private/Sequencer.cpp
--- a/Engine/Source/Editor/Sequencer/Private/Sequencer.cpp
+++ b/Engine/Source/Editor/Sequencer/Private/Sequencer.cpp
@@ -130,6 +130,8 @@
 #include "ISerializedRecorder.h"
 #include "Features/IModularFeatures.h"
 #include "SequencerContextMenus.h"
+#include "EngineAnalytics.h"
+#include "Interfaces/IAnalyticsProvider.h"
 
 #define LOCTEXT_NAMESPACE "Sequencer"
 
@@ -3290,6 +3292,12 @@ void FSequencer::ResetPerMovieSceneData()
 
 void FSequencer::RecordSelectedActors()
 {
+	// Keep track of how many people actually used record new sequence
+	if (FEngineAnalytics::IsAvailable())
+	{
+		FEngineAnalytics::GetProvider().RecordEvent(TEXT("Editor.Sequencer.RecordSelectedActors"));
+	}
+
 	ISequenceRecorder& SequenceRecorder = FModuleManager::LoadModuleChecked<ISequenceRecorder>("SequenceRecorder");
 	if (SequenceRecorder.IsRecording())
 	{
diff --git a/Engine/Source/Editor/Sequencer/Sequencer.Build.cs b/Engine/Source/Editor/Sequencer/Sequencer.Build.cs
--- a/Engine/Source/Editor/Sequencer/Sequencer.Build.cs
+++ b/Engine/Source/Editor/Sequencer/Sequencer.Build.cs
@@ -74,6 +74,7 @@ public class Sequencer : ModuleRules
 				"PropertyEditor",
                 "SceneOutliner",
                 "CurveEditor",
+                "Analytics",
             }
         );
 
