Sequencer: Set GetSequenceBinding to the "Sequence" category so that it doesn't show up at the root menu

diff --git a/Engine/Source/Editor/MovieSceneTools/Private/K2Node_GetSequenceBinding.cpp b/Engine/Source/Editor/MovieSceneTools/Private/K2Node_GetSequenceBinding.cpp
--- a/Engine/Source/Editor/MovieSceneTools/Private/K2Node_GetSequenceBinding.cpp
+++ b/Engine/Source/Editor/MovieSceneTools/Private/K2Node_GetSequenceBinding.cpp
@@ -233,6 +233,11 @@ FText UK2Node_GetSequenceBinding::GetTooltipText() const
 	return LOCTEXT("NodeTooltip", "Access an identifier for any object binding within a sequence");
 }
 
+FText UK2Node_GetSequenceBinding::GetMenuCategory() const
+{
+	return LOCTEXT("NodeCategory", "Sequence");
+}
+
 FSlateIcon UK2Node_GetSequenceBinding::GetIconAndTint(FLinearColor& OutColor) const
 {
 	static FSlateIcon Icon("EditorStyle", "GraphEditor.GetSequenceBinding");
diff --git a/Engine/Source/Editor/MovieSceneTools/Public/K2Node_GetSequenceBinding.h b/Engine/Source/Editor/MovieSceneTools/Public/K2Node_GetSequenceBinding.h
--- a/Engine/Source/Editor/MovieSceneTools/Public/K2Node_GetSequenceBinding.h
+++ b/Engine/Source/Editor/MovieSceneTools/Public/K2Node_GetSequenceBinding.h
@@ -38,6 +38,7 @@ public:
 	virtual void ValidateNodeDuringCompilation(FCompilerResultsLog& MessageLog) const override;
 	virtual FText GetNodeTitle(ENodeTitleType::Type TitleType) const override;
 	virtual FText GetTooltipText() const override;
+	virtual FText GetMenuCategory() const override;
 	virtual FSlateIcon GetIconAndTint(FLinearColor& OutColor) const override;
 	virtual bool ShouldShowNodeProperties() const override { return true; }
 	virtual bool IsNodePure() const override { return true; }
