Sequencer: Prevent adding a sequence of 0 duration as a subsequence

diff --git a/Engine/Source/Editor/MovieSceneTools/Private/TrackEditors/CinematicShotTrackEditor.cpp b/Engine/Source/Editor/MovieSceneTools/Private/TrackEditors/CinematicShotTrackEditor.cpp
--- a/Engine/Source/Editor/MovieSceneTools/Private/TrackEditors/CinematicShotTrackEditor.cpp
+++ b/Engine/Source/Editor/MovieSceneTools/Private/TrackEditors/CinematicShotTrackEditor.cpp
@@ -152,6 +152,14 @@ bool FCinematicShotTrackEditor::HandleAssetAdded(UObject* Asset, const FGuid& Ta
 		return false;
 	}
 
+	if (Sequence->GetMovieScene()->GetPlaybackRange().IsEmpty())
+	{
+		FNotificationInfo Info(FText::Format(LOCTEXT("InvalidSequenceDuration", "Invalid level sequence {0}. The sequence has no duration."), Sequence->GetDisplayName()));
+		Info.bUseLargeFont = false;
+		FSlateNotificationManager::Get().AddNotification(Info);
+		return false;
+	}
+
 	if (CanAddSubSequence(*Sequence))
 	{
 		const FScopedTransaction Transaction(LOCTEXT("AddShot_Transaction", "Add Shot"));
@@ -654,6 +662,14 @@ FKeyPropertyResult FCinematicShotTrackEditor::AddKeyInternal(FFrameNumber KeyTim
 {	
 	FKeyPropertyResult KeyPropertyResult;
 
+	if (InMovieSceneSequence->GetMovieScene()->GetPlaybackRange().IsEmpty())
+	{
+		FNotificationInfo Info(FText::Format(LOCTEXT("InvalidSequenceDuration", "Invalid level sequence {0}. The sequence has no duration."), InMovieSceneSequence->GetDisplayName()));
+		Info.bUseLargeFont = false;
+		FSlateNotificationManager::Get().AddNotification(Info);
+		return KeyPropertyResult;
+	}
+
 	if (CanAddSubSequence(*InMovieSceneSequence))
 	{
 		UMovieSceneCinematicShotTrack* CinematicShotTrack = FindOrCreateCinematicShotTrack();
diff --git a/Engine/Source/Editor/MovieSceneTools/Private/TrackEditors/SubTrackEditor.cpp b/Engine/Source/Editor/MovieSceneTools/Private/TrackEditors/SubTrackEditor.cpp
--- a/Engine/Source/Editor/MovieSceneTools/Private/TrackEditors/SubTrackEditor.cpp
+++ b/Engine/Source/Editor/MovieSceneTools/Private/TrackEditors/SubTrackEditor.cpp
@@ -569,6 +569,14 @@ bool FSubTrackEditor::HandleAssetAdded(UObject* Asset, const FGuid& TargetObject
 		return false;
 	}
 
+	if (Sequence->GetMovieScene()->GetPlaybackRange().IsEmpty())
+	{
+		FNotificationInfo Info(FText::Format(LOCTEXT("InvalidSequenceDuration", "Invalid level sequence {0}. The sequence has no duration."), Sequence->GetDisplayName()));
+		Info.bUseLargeFont = false;
+		FSlateNotificationManager::Get().AddNotification(Info);
+		return false;
+	}
+
 	if (CanAddSubSequence(*Sequence))
 	{
 		const FScopedTransaction Transaction(LOCTEXT("AddSubScene_Transaction", "Add Subscene"));
@@ -836,6 +844,14 @@ FKeyPropertyResult FSubTrackEditor::AddKeyInternal(FFrameNumber KeyTime, UMovieS
 {	
 	FKeyPropertyResult KeyPropertyResult;
 
+	if (InMovieSceneSequence->GetMovieScene()->GetPlaybackRange().IsEmpty())
+	{
+		FNotificationInfo Info(FText::Format(LOCTEXT("InvalidSequenceDuration", "Invalid level sequence {0}. The sequence has no duration."), InMovieSceneSequence->GetDisplayName()));
+		Info.bUseLargeFont = false;
+		FSlateNotificationManager::Get().AddNotification(Info);
+		return KeyPropertyResult;
+	}
+
 	if (CanAddSubSequence(*InMovieSceneSequence))
 	{
 		UMovieSceneSubTrack* SubTrack = Cast<UMovieSceneSubTrack>(InTrack);
