Sequencer: Fix crash when opening animation section properties. Set the allowed class type.

diff --git a/Engine/Source/Editor/MovieSceneTools/Private/TrackEditors/SkeletalAnimationTrackEditor.cpp b/Engine/Source/Editor/MovieSceneTools/Private/TrackEditors/SkeletalAnimationTrackEditor.cpp
--- a/Engine/Source/Editor/MovieSceneTools/Private/TrackEditors/SkeletalAnimationTrackEditor.cpp
+++ b/Engine/Source/Editor/MovieSceneTools/Private/TrackEditors/SkeletalAnimationTrackEditor.cpp
@@ -207,6 +207,7 @@ public:
 
 					TSharedRef<SObjectPropertyEntryBox> ContentWidget = SNew(SObjectPropertyEntryBox)
 						.PropertyHandle(ChildPropertyHandle)
+						.AllowedClass(UAnimSequenceBase::StaticClass())
 						.DisplayThumbnail(true)
 						.ThumbnailPool(PropertyUtilities.IsValid() ? PropertyUtilities->GetThumbnailPool() : nullptr)
 						.OnShouldFilterAsset(FOnShouldFilterAsset::CreateRaw(this, &FMovieSceneSkeletalAnimationParamsDetailCustomization::ShouldFilterAsset));
@@ -251,7 +252,7 @@ public:
 	}
 
 private:
-	const FSequencerSectionPropertyDetailsViewCustomizationParams& Params;
+	FSequencerSectionPropertyDetailsViewCustomizationParams Params;
 	FString SkeletonName;
 };
 
