Sequencer: Fix unitialized variables which was causing doubling or undefined values when keying transforms

diff --git a/Engine/Source/Editor/MovieSceneTools/Private/TrackEditors/PropertyTrackEditors/EulerTransformPropertyTrackEditor.cpp b/Engine/Source/Editor/MovieSceneTools/Private/TrackEditors/PropertyTrackEditors/EulerTransformPropertyTrackEditor.cpp
--- a/Engine/Source/Editor/MovieSceneTools/Private/TrackEditors/PropertyTrackEditors/EulerTransformPropertyTrackEditor.cpp
+++ b/Engine/Source/Editor/MovieSceneTools/Private/TrackEditors/PropertyTrackEditors/EulerTransformPropertyTrackEditor.cpp
@@ -124,8 +124,8 @@ bool FEulerTransformPropertyTrackEditor::ModifyGeneratedKeysByCurrentAndWeight(U
 	FMovieSceneContext Context(FMovieSceneEvaluationRange(KeyTime, GetSequencer()->GetFocusedTickResolution()));
 	EvalTrack.Interrogate(Context, InterrogationData, Object);
 
-	FVector CurrentPos; FRotator CurrentRot;
-	FVector CurrentScale;
+	FVector CurrentPos(0.f); FRotator CurrentRot(0.f);
+	FVector CurrentScale(1.f);
 	for (const FEulerTransform& Transform : InterrogationData.Iterate<FEulerTransform>(FMovieScenePropertySectionTemplate::GetEulerTransformInterrogationKey()))
 	{
 		CurrentPos = Transform.Location;
diff --git a/Engine/Source/Editor/MovieSceneTools/Private/TrackEditors/PropertyTrackEditors/TransformPropertyTrackEditor.cpp b/Engine/Source/Editor/MovieSceneTools/Private/TrackEditors/PropertyTrackEditors/TransformPropertyTrackEditor.cpp
--- a/Engine/Source/Editor/MovieSceneTools/Private/TrackEditors/PropertyTrackEditors/TransformPropertyTrackEditor.cpp
+++ b/Engine/Source/Editor/MovieSceneTools/Private/TrackEditors/PropertyTrackEditors/TransformPropertyTrackEditor.cpp
@@ -77,8 +77,8 @@ bool FTransformPropertyTrackEditor::ModifyGeneratedKeysByCurrentAndWeight(UObjec
 	FMovieSceneContext Context(FMovieSceneEvaluationRange(KeyTime, GetSequencer()->GetFocusedTickResolution()));
 	EvalTrack.Interrogate(Context, InterrogationData, Object);
 
-	FVector CurrentPos; FRotator CurrentRot;
-	FVector CurrentScale;
+	FVector CurrentPos(0.f); FRotator CurrentRot(0.f);
+	FVector CurrentScale(1.f);
 	for (const FTransform& Transform : InterrogationData.Iterate<FTransform>(FMovieScenePropertySectionTemplate::GetTransformInterrogationKey()))
 	{
 		CurrentPos = Transform.GetTranslation();
diff --git a/Engine/Source/Editor/MovieSceneTools/Private/TrackEditors/TransformTrackEditor.cpp b/Engine/Source/Editor/MovieSceneTools/Private/TrackEditors/TransformTrackEditor.cpp
--- a/Engine/Source/Editor/MovieSceneTools/Private/TrackEditors/TransformTrackEditor.cpp
+++ b/Engine/Source/Editor/MovieSceneTools/Private/TrackEditors/TransformTrackEditor.cpp
@@ -862,8 +862,8 @@ bool F3DTransformTrackEditor::ModifyGeneratedKeysByCurrentAndWeight(UObject *Obj
 	FMovieSceneContext Context(FMovieSceneEvaluationRange(KeyTime, GetSequencer()->GetFocusedTickResolution()));
 	EvalTrack.Interrogate(Context, InterrogationData, Object);
 			
-	FVector CurrentPos; FRotator CurrentRot;
-	FVector CurrentScale;
+	FVector CurrentPos(0.f); FRotator CurrentRot(0.f);
+	FVector CurrentScale(1.f);
 	for (const FTransformData& Transform : InterrogationData.Iterate<FTransformData>(UMovieScene3DTransformSection::GetInterrogationKey()))
 	{
 		CurrentPos = Transform.Translation;
diff --git a/Engine/Source/Runtime/MovieSceneTracks/Private/Sections/MovieScene3DTransformSection.cpp b/Engine/Source/Runtime/MovieSceneTracks/Private/Sections/MovieScene3DTransformSection.cpp
--- a/Engine/Source/Runtime/MovieSceneTracks/Private/Sections/MovieScene3DTransformSection.cpp
+++ b/Engine/Source/Runtime/MovieSceneTracks/Private/Sections/MovieScene3DTransformSection.cpp
@@ -213,8 +213,8 @@ struct F3DTransformChannelEditorData
 		FMovieSceneContext Context(FMovieSceneEvaluationRange(KeyTime, TickResolution));
 		EvalTrack.Interrogate(Context, InterrogationData, Object);
 
-		FVector CurrentPos; FRotator CurrentRot;
-		FVector CurrentScale;
+		FVector CurrentPos(0.f); FRotator CurrentRot(0.f);
+		FVector CurrentScale(1.f);
 		for (const FTransformData& Transform : InterrogationData.Iterate<FTransformData>(UMovieScene3DTransformSection::GetInterrogationKey()))
 		{
 			CurrentPos = Transform.Translation;
