Provide unique name for generated sequencer event function so that it does not collide across class hierarchies
#rnx

diff --git a/Engine/Source/Editor/MovieSceneTools/Private/MovieSceneEventUtils.cpp b/Engine/Source/Editor/MovieSceneTools/Private/MovieSceneEventUtils.cpp
--- a/Engine/Source/Editor/MovieSceneTools/Private/MovieSceneEventUtils.cpp
+++ b/Engine/Source/Editor/MovieSceneTools/Private/MovieSceneEventUtils.cpp
@@ -340,8 +340,8 @@ UK2Node_FunctionEntry* FMovieSceneEventUtils::GenerateEntryPoint(UMovieSceneEven
 	}
 
 	// @todo: use a more destriptive name for event entry points?
-	static FString DefaultEventEntryName = "SequenceEvent__ENTRYPOINT";
-	UEdGraph* EntryPointGraph = Compiler->SpawnIntermediateFunctionGraph(DefaultEventEntryName);
+	static FString DefaultEventEntryName = TEXT("SequenceEvent__ENTRYPOINT");
+	UEdGraph* EntryPointGraph = Compiler->SpawnIntermediateFunctionGraph(DefaultEventEntryName + Blueprint->GetName());
 	check(EntryPointGraph->Nodes.Num() == 1);
 
 	const UEdGraphSchema*  Schema        = EntryPointGraph->GetSchema();
