Sequencer: Fix crash on null animation asset

diff --git a/Engine/Source/Editor/MovieSceneTools/Private/TrackEditors/SkeletalAnimationTrackEditor.cpp b/Engine/Source/Editor/MovieSceneTools/Private/TrackEditors/SkeletalAnimationTrackEditor.cpp
--- a/Engine/Source/Editor/MovieSceneTools/Private/TrackEditors/SkeletalAnimationTrackEditor.cpp
+++ b/Engine/Source/Editor/MovieSceneTools/Private/TrackEditors/SkeletalAnimationTrackEditor.cpp
@@ -311,7 +311,7 @@ int32 FSkeletalAnimationSection::OnPaintSection( FSequencerSectionPainter& Paint
 	FFrameRate TickResolution = TimeToPixelConverter.GetTickResolution();
 
 	// Add lines where the animation starts and ends/loops
-	const float AnimPlayRate = FMath::IsNearlyZero(Section.Params.PlayRate) ? 1.0f : Section.Params.PlayRate * Section.Params.Animation->RateScale;
+	const float AnimPlayRate = FMath::IsNearlyZero(Section.Params.PlayRate) || Section.Params.Animation == nullptr ? 1.0f : Section.Params.PlayRate * Section.Params.Animation->RateScale;
 	const float SeqLength = (Section.Params.GetSequenceLength() - TickResolution.AsSeconds(Section.Params.StartFrameOffset + Section.Params.EndFrameOffset)) / AnimPlayRate;
 	const float FirstLoopSeqLength = SeqLength - TickResolution.AsSeconds(Section.Params.FirstLoopStartFrameOffset) / AnimPlayRate;
 
