Sequencer: Add tick marks for audio loop frames

diff --git a/Engine/Source/Editor/MovieSceneTools/Private/TrackEditors/AudioTrackEditor.cpp b/Engine/Source/Editor/MovieSceneTools/Private/TrackEditors/AudioTrackEditor.cpp
--- a/Engine/Source/Editor/MovieSceneTools/Private/TrackEditors/AudioTrackEditor.cpp
+++ b/Engine/Source/Editor/MovieSceneTools/Private/TrackEditors/AudioTrackEditor.cpp
@@ -701,6 +701,54 @@ int32 FAudioSection::OnPaintSection( FSequencerSectionPainter& Painter ) const
 		);
 	}
 
+	const ESlateDrawEffect DrawEffects = Painter.bParentEnabled ? ESlateDrawEffect::None : ESlateDrawEffect::DisabledEffect;
+
+	const FTimeToPixel& TimeToPixelConverter = Painter.GetTimeConverter();
+
+	static const FSlateBrush* GenericDivider = FEditorStyle::GetBrush("Sequencer.GenericDivider");
+
+	if (!Section.HasStartFrame() || !Section.HasEndFrame())
+	{
+		return LayerId;
+	}
+
+	UMovieSceneAudioSection* AudioSection = Cast<UMovieSceneAudioSection>(&Section);
+	if (!AudioSection || !AudioSection->GetSound())
+	{
+		return LayerId;
+	}
+
+	FFrameRate TickResolution = TimeToPixelConverter.GetTickResolution();
+	float AudioDuration = AudioSection->GetSound()->GetDuration();
+
+	// Add lines where the animation starts and ends/loops
+	const float SeqLength = AudioDuration - TickResolution.AsSeconds(AudioSection->GetStartOffset());
+
+	if (!FMath::IsNearlyZero(SeqLength, KINDA_SMALL_NUMBER) && SeqLength > 0)
+	{
+		float MaxOffset = Section.GetRange().Size<FFrameTime>() / TickResolution;
+		float OffsetTime = SeqLength;
+		float StartTime = Section.GetInclusiveStartFrame() / TickResolution;
+
+		while (OffsetTime < MaxOffset)
+		{
+			float OffsetPixel = TimeToPixelConverter.SecondsToPixel(StartTime + OffsetTime) - TimeToPixelConverter.SecondsToPixel(StartTime);
+
+			FSlateDrawElement::MakeBox(
+				Painter.DrawElements,
+				LayerId,
+				Painter.SectionGeometry.MakeChild(
+					FVector2D(2.f, Painter.SectionGeometry.Size.Y - 2.f),
+					FSlateLayoutTransform(FVector2D(OffsetPixel, 1.f))
+				).ToPaintGeometry(),
+				GenericDivider,
+				DrawEffects
+			);
+
+			OffsetTime += SeqLength;
+		}
+	}
+
 	return LayerId;
 }
 
