Sequencer: Add ensure for EvalTrack in GetLocationAtTime

diff --git a/Engine/Source/Editor/MovieSceneTools/Private/TrackEditors/AttachTrackEditor.cpp b/Engine/Source/Editor/MovieSceneTools/Private/TrackEditors/AttachTrackEditor.cpp
--- a/Engine/Source/Editor/MovieSceneTools/Private/TrackEditors/AttachTrackEditor.cpp
+++ b/Engine/Source/Editor/MovieSceneTools/Private/TrackEditors/AttachTrackEditor.cpp
@@ -460,6 +460,8 @@ FORCEINLINE FTransform FloatValuesToTransform(TArrayView<const FMovieSceneFloatV
  */
 FTransform GetLocationAtTime(TSharedPtr<ISequencer> Sequencer, FMovieSceneEvaluationTrack* EvalTrack, FFrameNumber KeyTime, UObject* Object)
 {
+	ensure(EvalTrack);
+
 	FMovieSceneInterrogationData InterrogationData;
 	Sequencer->GetEvaluationTemplate().CopyActuators(InterrogationData.GetAccumulator());
 	FMovieSceneContext Context(FMovieSceneEvaluationRange(KeyTime, Sequencer->GetFocusedTickResolution()));
