Sequencer: Fix adjusting animation start time during resizing so that start time is clamped to beginning of clip.

diff --git a/Engine/Source/Editor/MovieSceneTools/Private/TrackEditors/SkeletalAnimationTrackEditor.cpp b/Engine/Source/Editor/MovieSceneTools/Private/TrackEditors/SkeletalAnimationTrackEditor.cpp
--- a/Engine/Source/Editor/MovieSceneTools/Private/TrackEditors/SkeletalAnimationTrackEditor.cpp
+++ b/Engine/Source/Editor/MovieSceneTools/Private/TrackEditors/SkeletalAnimationTrackEditor.cpp
@@ -375,8 +375,10 @@ void FSkeletalAnimationSection::ResizeSection(ESequencerSectionResizeMode Resize
 
 		if (StartOffset < 0)
 		{
+			FFrameTime FrameTimeOver = FFrameTime::FromDecimal(StartOffset.Value / Section.Params.PlayRate);
+
 			// Ensure start offset is not less than 0 and adjust ResizeTime
-			ResizeTime = ResizeTime - StartOffset;
+			ResizeTime = ResizeTime - FrameTimeOver.GetFrame();
 
 			StartOffset = FFrameNumber(0);
 		}
