Sequencer: Change default attach track preserve type to None

diff --git a/Engine/Source/Editor/MovieSceneTools/Private/TrackEditors/AttachTrackEditor.cpp b/Engine/Source/Editor/MovieSceneTools/Private/TrackEditors/AttachTrackEditor.cpp
--- a/Engine/Source/Editor/MovieSceneTools/Private/TrackEditors/AttachTrackEditor.cpp
+++ b/Engine/Source/Editor/MovieSceneTools/Private/TrackEditors/AttachTrackEditor.cpp
@@ -121,7 +121,7 @@ private:
 
 F3DAttachTrackEditor::F3DAttachTrackEditor( TSharedRef<ISequencer> InSequencer )
 : FActorPickerTrackEditor( InSequencer )
-, PreserveType(ETransformPreserveType::CurrentKey)
+, PreserveType(ETransformPreserveType::None)
 {
 }
 
