Override default FBX import settings with inputted one from python or UI, then reset.

diff --git a/Engine/Source/Editor/MovieSceneTools/Private/MovieSceneToolHelpers.cpp b/Engine/Source/Editor/MovieSceneTools/Private/MovieSceneToolHelpers.cpp
--- a/Engine/Source/Editor/MovieSceneTools/Private/MovieSceneToolHelpers.cpp
+++ b/Engine/Source/Editor/MovieSceneTools/Private/MovieSceneToolHelpers.cpp
@@ -58,6 +58,8 @@
 #include "Sections/MovieSceneCameraCutSection.h"
 #include "Engine/LevelStreaming.h"
 #include "FbxExporter.h"
+#include "Serialization/ObjectWriter.h"
+#include "Serialization/ObjectReader.h"
 
 
 /* MovieSceneToolHelpers
@@ -1763,12 +1765,21 @@ bool MovieSceneToolHelpers::ReadyFBXForImport(const FString&  ImportFilename, UM
 	return true;
 }
 
+
 bool MovieSceneToolHelpers::ImportFBXIfReady(UWorld* World, UMovieScene* MovieScene, IMovieScenePlayer* Player, TMap<FGuid, FString>& ObjectBindingMap, UMovieSceneUserImportFBXSettings* ImportFBXSettings,
 	const FFBXInOutParameters& InParams)
 {
-	UnFbx::FFbxImporter* FbxImporter = UnFbx::FFbxImporter::GetInstance();
+	UMovieSceneUserImportFBXSettings* CurrentImportFBXSettings = GetMutableDefault<UMovieSceneUserImportFBXSettings>();
+	TArray<uint8> OriginalSettings;
+	FObjectWriter(CurrentImportFBXSettings, OriginalSettings);
 
-	const bool bMatchByNameOnly = ImportFBXSettings->bMatchByNameOnly;
+	CurrentImportFBXSettings->bMatchByNameOnly = ImportFBXSettings->bMatchByNameOnly;
+	CurrentImportFBXSettings->bForceFrontXAxis = ImportFBXSettings->bForceFrontXAxis;
+	CurrentImportFBXSettings->bCreateCameras = ImportFBXSettings->bCreateCameras;
+	CurrentImportFBXSettings->bReduceKeys = ImportFBXSettings->bReduceKeys;
+	CurrentImportFBXSettings->ReduceKeysTolerance = ImportFBXSettings->ReduceKeysTolerance;
+
+	UnFbx::FFbxImporter* FbxImporter = UnFbx::FFbxImporter::GetInstance();
 
 	UnFbx::FFbxCurvesAPI CurveAPI;
 	FbxImporter->PopulateAnimatedCurveData(CurveAPI);
@@ -1811,7 +1822,7 @@ bool MovieSceneToolHelpers::ImportFBXIfReady(UWorld* World, UMovieScene* MovieSc
 	}
 
 	// Otherwise, get the first available node that hasn't been imported onto yet
-	if (!bMatchByNameOnly)
+	if (!ImportFBXSettings->bMatchByNameOnly)
 	{
 		for (int32 NodeIndex = 0; NodeIndex < AllNodeNames.Num(); )
 		{
@@ -1842,6 +1853,9 @@ bool MovieSceneToolHelpers::ImportFBXIfReady(UWorld* World, UMovieScene* MovieSc
 		UE_LOG(LogMovieScene, Warning, TEXT("Fbx Import: Failed to find any matching node for (%s)."), *NodeName);
 	}
 
+	// restore
+	FObjectReader(GetMutableDefault<UMovieSceneUserImportFBXSettings>(), OriginalSettings);
+
 	FbxImporter->ReleaseScene();
 	UnFbx::FBXImportOptions* ImportOptions = FbxImporter->GetImportOptions();
 	ImportOptions->bConvertScene = InParams.bConvertSceneBackup;
