Sequencer: Fix event track properties menu. Event track for object bindings should hide the event receivers. Master event track should NOT hide the event receivers

diff --git a/Engine/Source/Editor/MovieSceneTools/Private/TrackEditors/EventTrackEditor.cpp b/Engine/Source/Editor/MovieSceneTools/Private/TrackEditors/EventTrackEditor.cpp
--- a/Engine/Source/Editor/MovieSceneTools/Private/TrackEditors/EventTrackEditor.cpp
+++ b/Engine/Source/Editor/MovieSceneTools/Private/TrackEditors/EventTrackEditor.cpp
@@ -193,13 +193,13 @@ void FEventTrackEditor::BuildTrackContextMenu(FMenuBuilder& MenuBuilder, UMovieS
 
 			if (ObjectBindingID.IsValid())
 			{
-				IDetailCategoryBuilder& Category = DetailBuilder.EditCategory("TrackEvent");
-				Category.AddProperty("EventReceivers").ShouldAutoExpand(true);
+				// Do not show event receivers for tracks that exist on object bindings
+				DetailBuilder.HideProperty("EventReceivers");
 			}
 			else
 			{
-				// Do not show event receivers for tracks that exist on object bindings
-				DetailBuilder.HideCategory("TrackEvent");
+				IDetailCategoryBuilder& Category = DetailBuilder.EditCategory("TrackEvent");
+				Category.AddProperty("EventReceivers").ShouldAutoExpand(true);
 			}
 		}
 	};
