Sequencer: Check for valid parent handle guid

diff --git a/Engine/Source/Editor/MovieSceneTools/Private/TrackEditors/AttachTrackEditor.cpp b/Engine/Source/Editor/MovieSceneTools/Private/TrackEditors/AttachTrackEditor.cpp
--- a/Engine/Source/Editor/MovieSceneTools/Private/TrackEditors/AttachTrackEditor.cpp
+++ b/Engine/Source/Editor/MovieSceneTools/Private/TrackEditors/AttachTrackEditor.cpp
@@ -1122,10 +1122,13 @@ FKeyPropertyResult F3DAttachTrackEditor::AddKeyInternal( FFrameNumber KeyTime, c
 
 	FGuid ParentActorHandle = GetSequencer()->GetHandleToObject(ParentActor, false);
 	TOptional<TArrayView<FMovieSceneFloatChannel*>> ParentChannels;
-	UMovieScene3DTransformTrack* ParentTransformTrack = MovieScene->FindTrack<UMovieScene3DTransformTrack>(ParentActorHandle);
-	if (ParentTransformTrack && ParentTransformTrack->GetAllSections().Num() == 1)
+	if (ParentActorHandle.IsValid())
 	{
-		ParentChannels = ParentTransformTrack->GetAllSections()[0]->GetChannelProxy().GetChannels<FMovieSceneFloatChannel>();
+		UMovieScene3DTransformTrack* ParentTransformTrack = MovieScene->FindTrack<UMovieScene3DTransformTrack>(ParentActorHandle);
+		if (ParentTransformTrack && ParentTransformTrack->GetAllSections().Num() == 1)
+		{
+			ParentChannels = ParentTransformTrack->GetAllSections()[0]->GetChannelProxy().GetChannels<FMovieSceneFloatChannel>();
+		}
 	}
 
 	for (int32 ObjectIndex = 0; ObjectIndex < Objects.Num(); ++ObjectIndex)
