Fix implicit conversion to bool which does not seem to be an error on vs2017

diff --git a/Engine/Source/Editor/MovieSceneTools/Public/ObjectBindingTagCache.h b/Engine/Source/Editor/MovieSceneTools/Public/ObjectBindingTagCache.h
--- a/Engine/Source/Editor/MovieSceneTools/Public/ObjectBindingTagCache.h
+++ b/Engine/Source/Editor/MovieSceneTools/Public/ObjectBindingTagCache.h
@@ -48,7 +48,7 @@ public:
 	 */
 	bool HasTag(const FMovieSceneObjectBindingID& BindingID, const FName& TagName) const
 	{
-		return ExposedNameReverseLUT.FindPair(BindingID, TagName);
+		return ExposedNameReverseLUT.FindPair(BindingID, TagName) != nullptr;
 	}
 
 
