Added include guard to SequencerBindingTree.h

diff --git a/Engine/Source/Editor/MovieSceneTools/Public/SequenceBindingTree.h b/Engine/Source/Editor/MovieSceneTools/Public/SequenceBindingTree.h
--- a/Engine/Source/Editor/MovieSceneTools/Public/SequenceBindingTree.h
+++ b/Engine/Source/Editor/MovieSceneTools/Public/SequenceBindingTree.h
@@ -1,5 +1,6 @@
 // Copyright 1998-2019 Epic Games, Inc. All Rights Reserved.
 
+#pragma once
 
 #include "CoreTypes.h"
 #include "Misc/Guid.h"
