Sequencer: Move material track object binding track menu to MaterialTrackEditor so that it will appear for the actor sequence editor and other sequencer editors

diff --git a/Engine/Plugins/MovieScene/LevelSequenceEditor/Source/LevelSequenceEditor/Private/LevelSequenceEditorToolkit.cpp b/Engine/Plugins/MovieScene/LevelSequenceEditor/Source/LevelSequenceEditor/Private/LevelSequenceEditorToolkit.cpp
--- a/Engine/Plugins/MovieScene/LevelSequenceEditor/Source/LevelSequenceEditor/Private/LevelSequenceEditorToolkit.cpp
+++ b/Engine/Plugins/MovieScene/LevelSequenceEditor/Source/LevelSequenceEditor/Private/LevelSequenceEditorToolkit.cpp
@@ -38,7 +38,6 @@
 #include "Sections/MovieSceneSubSection.h"
 #include "Tracks/MovieSceneSubTrack.h"
 #include "Tracks/MovieSceneCinematicShotTrack.h"
-#include "Tracks/MovieSceneMaterialTrack.h"
 #include "Tracks/MovieScenePropertyTrack.h"
 #include "MovieSceneToolsProjectSettings.h"
 #include "MovieSceneToolHelpers.h"
@@ -599,58 +598,6 @@ void FLevelSequenceEditorToolkit::HandleAddComponentActionExecute(UActorComponen
 }
 
 
-void FLevelSequenceEditorToolkit::HandleAddComponentMaterialActionExecute(UPrimitiveComponent* Component, int32 MaterialIndex)
-{
-	UMovieScene* FocusedMovieScene = Sequencer->GetFocusedMovieSceneSequence()->GetMovieScene();
-	if (FocusedMovieScene->IsReadOnly())
-	{
-		return;
-	}
-
-	const FScopedTransaction Transaction(LOCTEXT("AddComponentMaterialTrack", "Add component material track"));
-
-	FocusedMovieScene->Modify();
-
-	FString ComponentName = Component->GetName();
-
-	TArray<UActorComponent*> ActorComponents;
-	ActorComponents.Add(Component);
-
-	USelection* SelectedActors = GEditor->GetSelectedActors();
-	if (SelectedActors && SelectedActors->Num() > 0)
-	{
-		for (FSelectionIterator Iter(*SelectedActors); Iter; ++Iter)
-		{
-			AActor* Actor = CastChecked<AActor>(*Iter);
-
-			TArray<UActorComponent*> OutActorComponents;
-			Actor->GetComponents(OutActorComponents);
-			for (UActorComponent* ActorComponent : OutActorComponents)
-			{
-				if (ActorComponent->GetName() == ComponentName)
-				{
-					ActorComponents.AddUnique(ActorComponent);
-				}
-			}
-		}
-	}
-
-	for (UActorComponent* ActorComponent : ActorComponents)
-	{
-		FGuid ObjectHandle = Sequencer->GetHandleToObject(ActorComponent);
-		FName IndexName(*FString::FromInt(MaterialIndex));
-		if (FocusedMovieScene->FindTrack(UMovieSceneComponentMaterialTrack::StaticClass(), ObjectHandle, IndexName) == nullptr)
-		{
-			UMovieSceneComponentMaterialTrack* MaterialTrack = FocusedMovieScene->AddTrack<UMovieSceneComponentMaterialTrack>(ObjectHandle);
-			MaterialTrack->Modify();
-			MaterialTrack->SetMaterialIndex(MaterialIndex);
-		}
-	}
-
-	Sequencer->NotifyMovieSceneDataChanged(EMovieSceneDataChangeType::MovieSceneStructureItemAdded);
-}
-
-
 void FLevelSequenceEditorToolkit::HandleActorAddedToSequencer(AActor* Actor, const FGuid Binding)
 {
 	AddDefaultTracksForActor(*Actor, Binding);
@@ -911,25 +858,6 @@ void FLevelSequenceEditorToolkit::HandleTrackMenuExtensionAddTrack(FMenuBuilder&
 	}
 	else
 	{
-		if (UPrimitiveComponent* PrimitiveComponent = Cast<UPrimitiveComponent>(ContextObjects[0]))
-		{
-			int32 NumMaterials = PrimitiveComponent->GetNumMaterials();
-			if (NumMaterials > 0)
-			{
-				AddTrackMenuBuilder.BeginSection("Materials", LOCTEXT("MaterialSection", "Material Parameters"));
-				{
-					for (int32 MaterialIndex = 0; MaterialIndex < NumMaterials; MaterialIndex++)
-					{
-						FUIAction AddComponentMaterialAction(FExecuteAction::CreateSP(this, &FLevelSequenceEditorToolkit::HandleAddComponentMaterialActionExecute, PrimitiveComponent, MaterialIndex));
-						FText AddComponentMaterialLabel = FText::Format(LOCTEXT("ComponentMaterialIndexLabelFormat", "Element {0}"), FText::AsNumber(MaterialIndex));
-						FText AddComponentMaterialToolTip = FText::Format(LOCTEXT("ComponentMaterialIndexToolTipFormat", "Add material element {0}"), FText::AsNumber(MaterialIndex));
-						AddTrackMenuBuilder.AddMenuEntry(AddComponentMaterialLabel, AddComponentMaterialToolTip, FSlateIcon(), AddComponentMaterialAction);
-					}
-				}
-				AddTrackMenuBuilder.EndSection();
-			}
-		}
-
 		if (USkeletalMeshComponent* SkeletalComponent = Cast<USkeletalMeshComponent>(ContextObjects[0]))
 		{
 			UAnimInstance* AnimInstance = SkeletalComponent->GetAnimInstance();
diff --git a/Engine/Plugins/MovieScene/LevelSequenceEditor/Source/LevelSequenceEditor/Private/LevelSequenceEditorToolkit.h b/Engine/Plugins/MovieScene/LevelSequenceEditor/Source/LevelSequenceEditor/Private/LevelSequenceEditorToolkit.h
--- a/Engine/Plugins/MovieScene/LevelSequenceEditor/Source/LevelSequenceEditor/Private/LevelSequenceEditorToolkit.h
+++ b/Engine/Plugins/MovieScene/LevelSequenceEditor/Source/LevelSequenceEditor/Private/LevelSequenceEditorToolkit.h
@@ -121,9 +121,6 @@ private:
 	/** Callback for executing the Add Component action. */
 	void HandleAddComponentActionExecute(UActorComponent* Component);
 
-	/** Callback for executing the add component material track. */
-	void HandleAddComponentMaterialActionExecute(UPrimitiveComponent* Component, int32 MaterialIndex);
-
 	/** Create a new binding for the specified skeletal mesh component's animation instance. */
 	void BindAnimationInstance(USkeletalMeshComponent* SkeletalComponent);
 
diff --git a/Engine/Source/Editor/MovieSceneTools/Private/TrackEditors/MaterialTrackEditor.cpp b/Engine/Source/Editor/MovieSceneTools/Private/TrackEditors/MaterialTrackEditor.cpp
--- a/Engine/Source/Editor/MovieSceneTools/Private/TrackEditors/MaterialTrackEditor.cpp
+++ b/Engine/Source/Editor/MovieSceneTools/Private/TrackEditors/MaterialTrackEditor.cpp
@@ -11,6 +11,8 @@
 #include "SequencerUtilities.h"
 #include "Modules/ModuleManager.h"
 #include "MaterialEditorModule.h"
+#include "Engine/Selection.h"
+#include "ISequencerModule.h"
 
 
 #define LOCTEXT_NAMESPACE "MaterialTrackEditor"
@@ -219,5 +221,94 @@ UMaterialInterface* FComponentMaterialTrackEditor::GetMaterialInterfaceForTrack(
 	return nullptr;
 }
 
+void FComponentMaterialTrackEditor::ExtendObjectBindingTrackMenu(TSharedRef<FExtender> Extender, const TArray<FGuid>& ObjectBindings, const UClass* ObjectClass)
+{
+	if (ObjectClass->IsChildOf(UPrimitiveComponent::StaticClass()))
+	{
+		Extender->AddMenuExtension(SequencerMenuExtensionPoints::AddTrackMenu_PropertiesSection, EExtensionHook::Before, nullptr, FMenuExtensionDelegate::CreateSP(this, &FComponentMaterialTrackEditor::ConstructObjectBindingTrackMenu, ObjectBindings));
+	}
+}
+
+void FComponentMaterialTrackEditor::ConstructObjectBindingTrackMenu(FMenuBuilder& MenuBuilder, TArray<FGuid> ObjectBindings)
+{
+	UObject* Object = GetSequencer()->FindSpawnedObjectOrTemplate(ObjectBindings[0]);
+	if (!Object)
+	{
+		return;
+	}
+
+	if (UPrimitiveComponent* PrimitiveComponent = Cast<UPrimitiveComponent>(Object))
+	{
+		int32 NumMaterials = PrimitiveComponent->GetNumMaterials();
+		if (NumMaterials > 0)
+		{
+			MenuBuilder.BeginSection("Materials", LOCTEXT("MaterialSection", "Material Parameters"));
+			{
+				for (int32 MaterialIndex = 0; MaterialIndex < NumMaterials; MaterialIndex++)
+				{
+					FUIAction AddComponentMaterialAction(FExecuteAction::CreateRaw(this, &FComponentMaterialTrackEditor::HandleAddComponentMaterialActionExecute, PrimitiveComponent, MaterialIndex));
+					FText AddComponentMaterialLabel = FText::Format(LOCTEXT("ComponentMaterialIndexLabelFormat", "Element {0}"), FText::AsNumber(MaterialIndex));
+					FText AddComponentMaterialToolTip = FText::Format(LOCTEXT("ComponentMaterialIndexToolTipFormat", "Add material element {0}"), FText::AsNumber(MaterialIndex));
+					MenuBuilder.AddMenuEntry(AddComponentMaterialLabel, AddComponentMaterialToolTip, FSlateIcon(), AddComponentMaterialAction);
+				}
+			}
+			MenuBuilder.EndSection();
+		}
+	}
+}
+
+void FComponentMaterialTrackEditor::HandleAddComponentMaterialActionExecute(UPrimitiveComponent* Component, int32 MaterialIndex)
+{
+	TSharedPtr<ISequencer> SequencerPtr = GetSequencer();
+	UMovieScene* MovieScene = SequencerPtr->GetFocusedMovieSceneSequence()->GetMovieScene();
+	if (MovieScene->IsReadOnly())
+	{
+		return;
+	}
+
+	const FScopedTransaction Transaction(LOCTEXT("AddComponentMaterialTrack", "Add component material track"));
+
+	MovieScene->Modify();
+
+	FString ComponentName = Component->GetName();
+
+	TArray<UActorComponent*> ActorComponents;
+	ActorComponents.Add(Component);
+
+	USelection* SelectedActors = GEditor->GetSelectedActors();
+	if (SelectedActors && SelectedActors->Num() > 0)
+	{
+		for (FSelectionIterator Iter(*SelectedActors); Iter; ++Iter)
+		{
+			AActor* Actor = CastChecked<AActor>(*Iter);
+
+			TArray<UActorComponent*> OutActorComponents;
+			Actor->GetComponents(OutActorComponents);
+			for (UActorComponent* ActorComponent : OutActorComponents)
+			{
+				if (ActorComponent->GetName() == ComponentName)
+				{
+					ActorComponents.AddUnique(ActorComponent);
+				}
+			}
+		}
+	}
+
+	for (UActorComponent* ActorComponent : ActorComponents)
+	{
+		FGuid ObjectHandle = SequencerPtr->GetHandleToObject(ActorComponent);
+		FName IndexName(*FString::FromInt(MaterialIndex));
+		if (MovieScene->FindTrack(UMovieSceneComponentMaterialTrack::StaticClass(), ObjectHandle, IndexName) == nullptr)
+		{
+			UMovieSceneComponentMaterialTrack* MaterialTrack = MovieScene->AddTrack<UMovieSceneComponentMaterialTrack>(ObjectHandle);
+			MaterialTrack->Modify();
+			MaterialTrack->SetMaterialIndex(MaterialIndex);
+		}
+	}
+
+	SequencerPtr->NotifyMovieSceneDataChanged(EMovieSceneDataChangeType::MovieSceneStructureItemAdded);
+}
+
+
 
 #undef LOCTEXT_NAMESPACE
diff --git a/Engine/Source/Editor/MovieSceneTools/Public/TrackEditors/MaterialTrackEditor.h b/Engine/Source/Editor/MovieSceneTools/Public/TrackEditors/MaterialTrackEditor.h
--- a/Engine/Source/Editor/MovieSceneTools/Public/TrackEditors/MaterialTrackEditor.h
+++ b/Engine/Source/Editor/MovieSceneTools/Public/TrackEditors/MaterialTrackEditor.h
@@ -81,7 +81,7 @@ public:
 public:
 
 	// ISequencerTrackEditor interface
-
+	virtual void ExtendObjectBindingTrackMenu(TSharedRef<FExtender> Extender, const TArray<FGuid>& ObjectBindings, const UClass* ObjectClass) override;
 	virtual bool SupportsType( TSubclassOf<UMovieSceneTrack> Type ) const override;
 	virtual bool GetDefaultExpansionState(UMovieSceneTrack* InTrack) const override;
 
@@ -90,4 +90,11 @@ protected:
 	// FMaterialtrackEditor interface
 
 	virtual UMaterialInterface* GetMaterialInterfaceForTrack( FGuid ObjectBinding, UMovieSceneMaterialTrack* MaterialTrack ) override;
+
+private:
+
+	void ConstructObjectBindingTrackMenu(FMenuBuilder& MenuBuilder, TArray<FGuid> ObjectBindings);
+	/** Callback for executing the add component material track. */
+	void HandleAddComponentMaterialActionExecute(UPrimitiveComponent* Component, int32 MaterialIndex);
+
 };
