Renaming Play Rate track to Time Dilation track.
Updating comments and tooltips to match as well.

diff --git a/Engine/Source/Editor/MovieSceneTools/Private/TrackEditors/SlomoTrackEditor.cpp b/Engine/Source/Editor/MovieSceneTools/Private/TrackEditors/SlomoTrackEditor.cpp
--- a/Engine/Source/Editor/MovieSceneTools/Private/TrackEditors/SlomoTrackEditor.cpp
+++ b/Engine/Source/Editor/MovieSceneTools/Private/TrackEditors/SlomoTrackEditor.cpp
@@ -31,8 +31,8 @@ FSlomoTrackEditor::FSlomoTrackEditor(TSharedRef<ISequencer> InSequencer)
 void FSlomoTrackEditor::BuildAddTrackMenu(FMenuBuilder& MenuBuilder)
 {
 	MenuBuilder.AddMenuEntry(
-		LOCTEXT("AddPlayRateTrack", "Play Rate Track"),
-		LOCTEXT("AddPlayRateTrackTooltip", "Adds a new track that controls the playback rate of the sequence."),
+		LOCTEXT("AddTimeDilationTrack", "Time Dilation Track"),
+		LOCTEXT("AddTimeDilationTrackTooltip", "Adds a new track that controls the world's time dilation."),
 		FSlateIcon(FEditorStyle::GetStyleSetName(), "Sequencer.Tracks.Slomo"),
 		FUIAction(
 			FExecuteAction::CreateRaw(this, &FSlomoTrackEditor::HandleAddSlomoTrackMenuEntryExecute),
@@ -81,7 +81,7 @@ void FSlomoTrackEditor::HandleAddSlomoTrackMenuEntryExecute()
 		return;
 	}
 
-	const FScopedTransaction Transaction(NSLOCTEXT("Sequencer", "AddSlomoTrack_Transaction", "Add Play Rate Track"));
+	const FScopedTransaction Transaction(NSLOCTEXT("Sequencer", "AddSlomoTrack_Transaction", "Add Time Dilation Track"));
 
 	MovieScene->Modify();
 
diff --git a/Engine/Source/Runtime/MovieSceneTracks/Private/Tracks/MovieSceneSlomoTrack.cpp b/Engine/Source/Runtime/MovieSceneTracks/Private/Tracks/MovieSceneSlomoTrack.cpp
--- a/Engine/Source/Runtime/MovieSceneTracks/Private/Tracks/MovieSceneSlomoTrack.cpp
+++ b/Engine/Source/Runtime/MovieSceneTracks/Private/Tracks/MovieSceneSlomoTrack.cpp
@@ -37,7 +37,7 @@ FMovieSceneEvalTemplatePtr UMovieSceneSlomoTrack::CreateTemplateForSection(const
 
 FText UMovieSceneSlomoTrack::GetDefaultDisplayName() const
 {
-	return LOCTEXT("TrackName", "Play Rate");
+	return LOCTEXT("TrackName", "Time Dilation");
 }
 
 #endif
diff --git a/Engine/Source/Runtime/MovieSceneTracks/Public/Tracks/MovieSceneSlomoTrack.h b/Engine/Source/Runtime/MovieSceneTracks/Public/Tracks/MovieSceneSlomoTrack.h
--- a/Engine/Source/Runtime/MovieSceneTracks/Public/Tracks/MovieSceneSlomoTrack.h
+++ b/Engine/Source/Runtime/MovieSceneTracks/Public/Tracks/MovieSceneSlomoTrack.h
@@ -8,7 +8,7 @@
 #include "MovieSceneSlomoTrack.generated.h"
 
 /**
- * Implements a movie scene track that controls a movie scene's playback speed.
+ * Implements a movie scene track that controls a movie scene's world time dilation.
  */
 UCLASS(MinimalAPI)
 class UMovieSceneSlomoTrack
