Sequencer: Import camera cut/switcher from fbx got removed in a bad merge

diff --git a/Engine/Source/Editor/MovieSceneTools/Private/MovieSceneToolHelpers.cpp b/Engine/Source/Editor/MovieSceneTools/Private/MovieSceneToolHelpers.cpp
--- a/Engine/Source/Editor/MovieSceneTools/Private/MovieSceneToolHelpers.cpp
+++ b/Engine/Source/Editor/MovieSceneTools/Private/MovieSceneToolHelpers.cpp
@@ -1575,7 +1575,7 @@ UMovieSceneCameraCutTrack* GetCameraCutTrack(UMovieScene* InMovieScene)
 	return CastChecked<UMovieSceneCameraCutTrack>(CameraCutTrack);
 }
 
-void ImportCameraCut(UnFbx::FFbxImporter* FbxImporter, UMovieScene* InMovieScene, ISequencer& InSequencer, TMap<FGuid, FString>& InObjectBindingMap)
+void ImportCameraCut(UnFbx::FFbxImporter* FbxImporter, UMovieScene* InMovieScene, TMap<FGuid, FString>& InObjectBindingMap)
 {
 	// Find a camera switcher
 	FbxCameraSwitcher* CameraSwitcher = FbxImporter->Scene->GlobalCameraSettings().GetCameraSwitcher();
@@ -1618,7 +1618,6 @@ void ImportCameraCut(UnFbx::FFbxImporter* FbxImporter, UMovieScene* InMovieScene
 				}
 			}
 		}
-		InSequencer.NotifyMovieSceneDataChanged(EMovieSceneDataChangeType::MovieSceneStructureItemAdded);
 	}
 }
 
@@ -1786,6 +1785,9 @@ bool MovieSceneToolHelpers::ImportFBXIfReady(UWorld* World, UMovieScene* MovieSc
 	TArray<FString> AllNodeNames;
 	CurveAPI.GetAllNodeNameArray(AllNodeNames);
 
+	// Import a camera cut track if cams were created, do it after populating curve data ensure only one animation layer, if any
+	ImportCameraCut(FbxImporter, MovieScene, ObjectBindingMap);
+
 	FString RootNodeName = FbxImporter->Scene->GetRootNode()->GetName();
 
 	// First try matching by name
