Sequencer: Ensure that sequence director BPs are compiled after updating the underlying event structure when using the RebindTo functionality

diff --git a/Engine/Source/Editor/MovieSceneTools/Private/MovieSceneEventCustomization.cpp b/Engine/Source/Editor/MovieSceneTools/Private/MovieSceneEventCustomization.cpp
--- a/Engine/Source/Editor/MovieSceneTools/Private/MovieSceneEventCustomization.cpp
+++ b/Engine/Source/Editor/MovieSceneTools/Private/MovieSceneEventCustomization.cpp
@@ -1290,13 +1290,16 @@ void FMovieSceneEventCustomization::SetEventEndpoint(UK2Node* NewEndpoint, UEdGr
 		}
 	}
 
+	// Ensure that anything listening for property changed notifications are notified of the new binding
+	PropertyHandle->NotifyFinishedChangingProperties();
+
+	// Compile the blueprint now that clients have had a chance to update underlying data (we do this after to ensure we are compiling the correct data
 	if (Blueprint)
 	{
 		FKismetEditorUtilities::CompileBlueprint(Blueprint);
 	}
 
-	// Ensure that anything listening for property changed notifications are notified of the new binding
-	PropertyHandle->NotifyFinishedChangingProperties();
+	// Forcibly update the panel now that our endpoint has changed
 	PropertyUtilities->ForceRefresh();
 }
 
