Sequencer Movie Renders now flush the shader compiler and mesh distance field builder on each tick in an attempt to ensure shaders are ready before rendering a movie.

diff --git a/Engine/Source/Editor/MovieSceneTools/Private/AutomatedLevelSequenceCapture.cpp b/Engine/Source/Editor/MovieSceneTools/Private/AutomatedLevelSequenceCapture.cpp
--- a/Engine/Source/Editor/MovieSceneTools/Private/AutomatedLevelSequenceCapture.cpp
+++ b/Engine/Source/Editor/MovieSceneTools/Private/AutomatedLevelSequenceCapture.cpp
@@ -22,6 +22,8 @@
 #include "Protocols/AudioCaptureProtocol.h"
 #include "Evaluation/IMovieSceneMotionVectorSimulation.h"
 #include "Rendering/MotionVectorSimulation.h"
+#include "ShaderCompiler.h"
+#include "DistanceFieldAtlas.h"
 
 const FName UAutomatedLevelSequenceCapture::AutomatedLevelSequenceCaptureUIName = FName(TEXT("AutomatedLevelSequenceCaptureUIInstance"));
 
@@ -651,6 +653,21 @@ void UAutomatedLevelSequenceCapture::OnTick(float DeltaSeconds)
 		GetWorld()->FlushLevelStreaming(EFlushLevelStreamingType::Full);
 	}
 
+
+	if (GShaderCompilingManager && GShaderCompilingManager->GetNumRemainingJobs() > 0)
+	{
+		UE_LOG(LogMovieSceneCapture, Log, TEXT("[%d] Waiting for %d shaders to finish compiling..."), GFrameCounter, GShaderCompilingManager->GetNumRemainingJobs());
+		GShaderCompilingManager->FinishAllCompilation();
+		UE_LOG(LogMovieSceneCapture, Log, TEXT("[%d] Done waiting for shaders to finish."), GFrameCounter);
+	}
+
+	if (GDistanceFieldAsyncQueue && GDistanceFieldAsyncQueue->GetNumOutstandingTasks() > 0)
+	{
+		UE_LOG(LogMovieSceneCapture, Log, TEXT("[%d] Waiting for %d Mesh Distance Fields to finish building..."), GFrameCounter, GDistanceFieldAsyncQueue->GetNumOutstandingTasks());
+		GDistanceFieldAsyncQueue->BlockUntilAllBuildsComplete();
+		UE_LOG(LogMovieSceneCapture, Log, TEXT("[%d] Done waiting for Mesh Distance Fields to build."), GFrameCounter);
+	}
+
 	// Setup the automated capture
 	if (CaptureState == ELevelSequenceCaptureState::Setup)
 	{
