UI: Frame number customization wasn't displaying the property tooltip where the property label was.

diff --git a/Engine/Source/Editor/MovieSceneTools/Private/FrameNumberDetailsCustomization.cpp b/Engine/Source/Editor/MovieSceneTools/Private/FrameNumberDetailsCustomization.cpp
--- a/Engine/Source/Editor/MovieSceneTools/Private/FrameNumberDetailsCustomization.cpp
+++ b/Engine/Source/Editor/MovieSceneTools/Private/FrameNumberDetailsCustomization.cpp
@@ -50,13 +50,14 @@ void FFrameNumberDetailsCustomization::CustomizeChildren(TSharedRef<IPropertyHan
 		[
 			SNew(STextBlock)
 			.Text(PropertyHandle->GetPropertyDisplayName())
-			.ToolTipText(LOCTEXT("TimeLabelTooltip", "Time field which takes timecode, frames and seconds formats."))
+			.ToolTipText(PropertyHandle->GetToolTipText())
 			.Font(CustomizationUtils.GetRegularFont())
 		]
 		.ValueContent()
 		[
 			SNew(SEditableTextBox)
 			.Text(this, &FFrameNumberDetailsCustomization::OnGetTimeText)
+			.ToolTipText(LOCTEXT("TimeLabelTooltip", "Time field which takes timecode, frames and seconds formats."))
 			.OnTextCommitted(this, &FFrameNumberDetailsCustomization::OnTimeTextCommitted)
 			.SelectAllTextWhenFocused(true)
 			.ClearKeyboardFocusOnCommit(false)
