Sequencer: Fix quick bind menu functions for BPs

#rnx

diff --git a/Engine/Source/Editor/MovieSceneTools/Private/MovieSceneEventCustomization.cpp b/Engine/Source/Editor/MovieSceneTools/Private/MovieSceneEventCustomization.cpp
--- a/Engine/Source/Editor/MovieSceneTools/Private/MovieSceneEventCustomization.cpp
+++ b/Engine/Source/Editor/MovieSceneTools/Private/MovieSceneEventCustomization.cpp
@@ -78,7 +78,7 @@ namespace UE4_MovieSceneEventCustomization
 
 			for (UObjectProperty* ObjectProperty : TFieldRange<UObjectProperty>(BoundObjectPinClass))
 			{
-				if (ObjectProperty->HasAnyPropertyFlags(CPF_BlueprintVisible) && ObjectProperty->HasMetaData(FBlueprintMetadata::MD_ExposeFunctionCategories))
+				if (ObjectProperty->HasAnyPropertyFlags(CPF_BlueprintVisible) && (ObjectProperty->HasMetaData(FBlueprintMetadata::MD_ExposeFunctionCategories) || FBlueprintEditorUtils::IsSCSComponentProperty(ObjectProperty)))
 				{
 					CallOnMemberFilter.Context.SelectedObjects.Add(ObjectProperty);
 					FBlueprintActionFilter::AddUnique(CallOnMemberFilter.TargetClasses, ObjectProperty->PropertyClass);
