Sequencer: Fix audio waveform crash when samples <= 0

#rnx

diff --git a/Engine/Source/Editor/MovieSceneTools/Private/TrackEditors/AudioTrackEditor.cpp b/Engine/Source/Editor/MovieSceneTools/Private/TrackEditors/AudioTrackEditor.cpp
--- a/Engine/Source/Editor/MovieSceneTools/Private/TrackEditors/AudioTrackEditor.cpp
+++ b/Engine/Source/Editor/MovieSceneTools/Private/TrackEditors/AudioTrackEditor.cpp
@@ -337,6 +337,11 @@ void FAudioThumbnail::GenerateWaveformPreview(TArray<uint8>& OutData, TRange<flo
 	const int32 TrueDrawOffsetPx = FMath::Max(FMath::RoundToInt((DrawRange.GetLowerBoundValue() - SectionStartTime) / DisplayScale), 0);
 	const int32 LastTrueSample = -2.f*SmoothingAmount + FMath::TruncToInt(TrueRangeSize / DisplayScale);
 
+	if (LastTrueSample <= 0)
+	{
+		return;
+	}
+
 	DrawRange = AudioTrueRange;
 
 	float DrawRangeSize = DrawRange.Size<float>();
