Sequencer: Fix audio looping marks to use the unlooped audio duration

#rnx

diff --git a/Engine/Source/Editor/MovieSceneTools/Private/TrackEditors/AudioTrackEditor.cpp b/Engine/Source/Editor/MovieSceneTools/Private/TrackEditors/AudioTrackEditor.cpp
--- a/Engine/Source/Editor/MovieSceneTools/Private/TrackEditors/AudioTrackEditor.cpp
+++ b/Engine/Source/Editor/MovieSceneTools/Private/TrackEditors/AudioTrackEditor.cpp
@@ -724,7 +724,7 @@ int32 FAudioSection::OnPaintSection( FSequencerSectionPainter& Painter ) const
 	}
 
 	FFrameRate TickResolution = TimeToPixelConverter.GetTickResolution();
-	float AudioDuration = AudioSection->GetSound()->GetDuration();
+	float AudioDuration = DeriveUnloopedDuration(AudioSection);
 
 	// Add lines where the animation starts and ends/loops
 	const float SeqLength = AudioDuration - TickResolution.AsSeconds(AudioSection->GetStartOffset());
