Sequencer: Change generated event entry points to use blueprint extensions to make preloading explicit

diff --git a/Engine/Source/Editor/MovieSceneTools/Private/MatineeImportTools.cpp b/Engine/Source/Editor/MovieSceneTools/Private/MatineeImportTools.cpp
--- a/Engine/Source/Editor/MovieSceneTools/Private/MatineeImportTools.cpp
+++ b/Engine/Source/Editor/MovieSceneTools/Private/MatineeImportTools.cpp
@@ -897,7 +897,7 @@ bool FMatineeImportTools::CopyInterpEventTrack( UInterpTrackEvent* MatineeEventT
 			if (SequenceDirectorBP)
 			{
 				SequenceDirectorBP->Modify();
-				SequenceDirectorBP->GenerateFunctionGraphsEvent.AddUniqueDynamic(Section, &UMovieSceneEventSectionBase::HandleGenerateEntryPoints);
+				FMovieSceneEventUtils::BindEventSectionToBlueprint(Section, SequenceDirectorBP);
 			}
 
 			TRange<FFrameNumber> KeyRange = TRange<FFrameNumber>::Empty();
diff --git a/Engine/Source/Editor/MovieSceneTools/Private/MovieSceneEventBlueprintExtension.cpp b/Engine/Source/Editor/MovieSceneTools/Private/MovieSceneEventBlueprintExtension.cpp
new file mode 100644
--- /dev/null
+++ b/Engine/Source/Editor/MovieSceneTools/Private/MovieSceneEventBlueprintExtension.cpp
@@ -0,0 +1,70 @@
+// Copyright 1998-2019 Epic Games, Inc. All Rights Reserved.
+
+#include "MovieSceneEventBlueprintExtension.h"
+#include "MovieSceneEventUtils.h"
+#include "Sections/MoviesceneEventSectionBase.h"
+#include "K2Node_FunctionEntry.h"
+#include "KismetCompiler.h"
+#include "Engine/Blueprint.h"
+#include "EdGraph/EdGraph.h"
+
+void UMovieSceneEventBlueprintExtension::PostLoad()
+{
+	EventSections.Remove(nullptr);
+	Super::PostLoad();
+}
+
+void UMovieSceneEventBlueprintExtension::HandlePreloadObjectsForCompilation(UBlueprint* OwningBlueprint)
+{
+	for (UMovieSceneEventSectionBase* EventSection : EventSections)
+	{
+		if (EventSection)
+		{
+			UBlueprint::ForceLoad(EventSection);
+		}
+	}
+}
+
+void UMovieSceneEventBlueprintExtension::HandleGenerateFunctionGraphs(FKismetCompilerContext* CompilerContext)
+{
+	for (UMovieSceneEventSectionBase* EventSection : EventSections)
+	{
+		if (!EventSection)
+		{
+			continue;
+		}
+
+		ensureMsgf(!EventSection->HasAnyFlags(RF_NeedLoad), TEXT("Attempting to generate entry point functions before an event section has been loaded"));
+
+		EventSection->AttemptUpgrade();
+
+		for (FMovieSceneEvent& EntryPoint : EventSection->GetAllEntryPoints())
+		{
+			UEdGraphNode* Endpoint = FMovieSceneEventUtils::FindEndpoint(&EntryPoint, EventSection, CompilerContext->Blueprint);
+			if (Endpoint)
+			{
+				UK2Node_FunctionEntry* FunctionEntry = FMovieSceneEventUtils::GenerateEntryPoint(EventSection, &EntryPoint, CompilerContext, Endpoint);
+				if (FunctionEntry)
+				{
+					EntryPoint.CompiledFunctionName = FunctionEntry->GetGraph()->GetFName();
+				}
+				else
+				{
+					EntryPoint.CompiledFunctionName = NAME_None;
+				}
+			}
+		}
+
+
+		auto OnFunctionListGenerated = [WeakEventSection = MakeWeakObjectPtr(EventSection)](FKismetCompilerContext* CompilerContext)
+		{
+			UMovieSceneEventSectionBase* PinnedSection = WeakEventSection.Get();
+			if (ensureMsgf(PinnedSection, TEXT("Event section has been collected during blueprint compilation.")))
+			{
+				PinnedSection->OnPostCompile(CompilerContext->Blueprint);
+			}
+		};
+
+		CompilerContext->OnFunctionListCompiled().AddLambda(OnFunctionListGenerated);
+	}
+}
\ No newline at end of file
diff --git a/Engine/Source/Editor/MovieSceneTools/Private/MovieSceneEventBlueprintExtension.h b/Engine/Source/Editor/MovieSceneTools/Private/MovieSceneEventBlueprintExtension.h
new file mode 100644
--- /dev/null
+++ b/Engine/Source/Editor/MovieSceneTools/Private/MovieSceneEventBlueprintExtension.h
@@ -0,0 +1,35 @@
+// Copyright 1998-2019 Epic Games, Inc. All Rights Reserved.
+
+#pragma once
+
+#include "Containers/Array.h"
+#include "Blueprint/BlueprintExtension.h"
+
+#include "MoviesceneEventBlueprintExtension.generated.h"
+
+class UBlueprint;
+class FKismetCompilerContext;
+class UMovieSceneEventSectionBase;
+
+UCLASS()
+class UMovieSceneEventBlueprintExtension : public UBlueprintExtension
+{
+public:
+
+	GENERATED_BODY()
+
+	void Add(UMovieSceneEventSectionBase* EventSection)
+	{
+		EventSections.AddUnique(EventSection);
+	}
+
+private:
+
+	virtual void PostLoad() override final;
+	virtual void HandlePreloadObjectsForCompilation(UBlueprint* OwningBlueprint) override final;
+	virtual void HandleGenerateFunctionGraphs(FKismetCompilerContext* CompilerContext) override final;
+
+	/** List of event sections that are bound to the blueprint */
+	UPROPERTY()
+	TArray<UMovieSceneEventSectionBase*> EventSections;
+};
\ No newline at end of file
diff --git a/Engine/Source/Editor/MovieSceneTools/Private/MovieSceneEventCustomization.cpp b/Engine/Source/Editor/MovieSceneTools/Private/MovieSceneEventCustomization.cpp
--- a/Engine/Source/Editor/MovieSceneTools/Private/MovieSceneEventCustomization.cpp
+++ b/Engine/Source/Editor/MovieSceneTools/Private/MovieSceneEventCustomization.cpp
@@ -1233,7 +1233,7 @@ void FMovieSceneEventCustomization::SetEventEndpoint(UK2Node* NewEndpoint, UEdGr
 			BaseEventSection->Modify();
 			if (Blueprint)
 			{
-				Blueprint->GenerateFunctionGraphsEvent.AddUniqueDynamic(BaseEventSection, &UMovieSceneEventSectionBase::HandleGenerateEntryPoints);
+				FMovieSceneEventUtils::BindEventSectionToBlueprint(BaseEventSection, Blueprint);
 			}
 		}
 	}
@@ -1395,7 +1395,7 @@ void FMovieSceneEventCustomization::CreateEventEndpoint()
 			for (const TPair<UMovieSceneEventSectionBase*, FSectionData>& SectionPair : TrackPair.Value.Sections)
 			{
 				SectionPair.Key->Modify();
-				SequenceDirectorBP->GenerateFunctionGraphsEvent.AddUniqueDynamic(SectionPair.Key, &UMovieSceneEventSectionBase::HandleGenerateEntryPoints);
+				FMovieSceneEventUtils::BindEventSectionToBlueprint(SectionPair.Key, SequenceDirectorBP);
 
 				for (FMovieSceneEvent* EntryPoint : SectionPair.Value.EntryPoints)
 				{
diff --git a/Engine/Source/Editor/MovieSceneTools/Private/MovieSceneEventUtils.cpp b/Engine/Source/Editor/MovieSceneTools/Private/MovieSceneEventUtils.cpp
--- a/Engine/Source/Editor/MovieSceneTools/Private/MovieSceneEventUtils.cpp
+++ b/Engine/Source/Editor/MovieSceneTools/Private/MovieSceneEventUtils.cpp
@@ -10,6 +10,7 @@
 #include "MovieSceneSequence.h"
 #include "Sections/MovieSceneEventSectionBase.h"
 #include "Tracks/MovieSceneEventTrack.h"
+#include "MovieSceneEventBlueprintExtension.h"
 
 #include "ScopedTransaction.h"
 
@@ -108,7 +109,7 @@ UK2Node_CustomEvent* FMovieSceneEventUtils::BindNewUserFacingEvent(FMovieSceneEv
 	Blueprint->Modify();
 
 	// Ensure the section is bound to the blueprint function generation event
-	Blueprint->GenerateFunctionGraphsEvent.AddUniqueDynamic(EventSection, &UMovieSceneEventSectionBase::HandleGenerateEntryPoints);
+	FMovieSceneEventUtils::BindEventSectionToBlueprint(EventSection, Blueprint);
 
 	// Create the new user-facing event node
 	FMovieSceneEventEndpointParameters Params = FMovieSceneEventEndpointParameters::Generate(Track);
@@ -443,4 +444,23 @@ UK2Node_FunctionEntry* FMovieSceneEventUtils::GenerateEntryPoint(UMovieSceneEven
 	return FunctionEntry;
 }
 
+void FMovieSceneEventUtils::BindEventSectionToBlueprint(UMovieSceneEventSectionBase* EventSection, UBlueprint* DirectorBP)
+{
+	check(EventSection && DirectorBP);
+
+	for (UBlueprintExtension* Extension : DirectorBP->Extensions)
+	{
+		UMovieSceneEventBlueprintExtension* EventExtension = Cast<UMovieSceneEventBlueprintExtension>(Extension);
+		if (EventExtension)
+		{
+			EventExtension->Add(EventSection);
+			return;
+		}
+	}
+
+	UMovieSceneEventBlueprintExtension* EventExtension = NewObject<UMovieSceneEventBlueprintExtension>(DirectorBP);
+	EventExtension->Add(EventSection);
+	DirectorBP->Extensions.Add(EventExtension);
+}
+
 #undef LOCTEXT_NAMESPACE
\ No newline at end of file
diff --git a/Engine/Source/Editor/MovieSceneTools/Private/MovieSceneToolsModule.cpp b/Engine/Source/Editor/MovieSceneTools/Private/MovieSceneToolsModule.cpp
--- a/Engine/Source/Editor/MovieSceneTools/Private/MovieSceneToolsModule.cpp
+++ b/Engine/Source/Editor/MovieSceneTools/Private/MovieSceneToolsModule.cpp
@@ -155,15 +155,14 @@ void FMovieSceneToolsModule::StartupModule()
 		));
 	}
 
-	GenerateEventEntryPointsHandle = UMovieSceneEventSectionBase::GenerateEventEntryPointsEvent.AddStatic(HandleGenerateEventEntryPoints);
 	FixupPayloadParameterNameHandle = UMovieSceneEventSectionBase::FixupPayloadParameterNameEvent.AddStatic(FixupPayloadParameterNameForSection);
-	UpgradeLegacyEventEndpointHandle = UMovieSceneEventSectionBase::UpgradeLegacyEventEndpoint.AddStatic(UpgradeLegacyEventEndpointForSection);
+	UMovieSceneEventSectionBase::UpgradeLegacyEventEndpoint.BindStatic(UpgradeLegacyEventEndpointForSection);
 }
 
 void FMovieSceneToolsModule::ShutdownModule()
 {
-	UMovieSceneEventSectionBase::GenerateEventEntryPointsEvent.Remove(GenerateEventEntryPointsHandle);
 	UMovieSceneEventSectionBase::FixupPayloadParameterNameEvent.Remove(FixupPayloadParameterNameHandle);
+	UMovieSceneEventSectionBase::UpgradeLegacyEventEndpoint = UMovieSceneEventSectionBase::FUpgradeLegacyEventEndpoint();
 
 	if (ICurveEditorModule* CurveEditorModule = FModuleManager::GetModulePtr<ICurveEditorModule>("CurveEditor"))
 	{
@@ -228,10 +227,28 @@ void FMovieSceneToolsModule::ShutdownModule()
 	}
 }
 
-void FMovieSceneToolsModule::UpgradeLegacyEventEndpointForSection(UMovieSceneEventSectionBase* Section, UBlueprint* Blueprint)
+bool FMovieSceneToolsModule::UpgradeLegacyEventEndpointForSection(UMovieSceneEventSectionBase* Section, UBlueprint* Blueprint)
 {
+	// Always bind the event section onto the blueprint to ensure that we get another chance to upgrade when the BP compiles if this try wasn't successful
+	FMovieSceneEventUtils::BindEventSectionToBlueprint(Section, Blueprint);
+
+	// We can't do this upgrade if we any of the function graphs are RF_NeedLoad
+	for (UEdGraph* EdGraph : Blueprint->FunctionGraphs)
+	{
+		if (EdGraph->HasAnyFlags(RF_NeedLoad))
+		{
+			return false;
+		}
+	}
+
+	// All the function graphs have been loaded, which means this is a good time to perform legacy data upgrade
 	for (FMovieSceneEvent& EntryPoint : Section->GetAllEntryPoints())
 	{
+		if (UK2Node_FunctionEntry* LegacyFunctionEntry = Cast<UK2Node_FunctionEntry>(EntryPoint.FunctionEntry_DEPRECATED.Get()))
+		{
+			EntryPoint.GraphGuid = LegacyFunctionEntry->GetGraph()->GraphGuid;
+		}
+
 		UK2Node* Endpoint = FMovieSceneEventUtils::FindEndpoint(&EntryPoint, Section, Blueprint);
 		if (!Endpoint)
 		{
@@ -251,6 +268,8 @@ void FMovieSceneToolsModule::UpgradeLegacyEventEndpointForSection(UMovieSceneEve
 		// Set the compiled function name so that any immediate PostCompile steps find the correct function name
 		EntryPoint.CompiledFunctionName = Endpoint->GetGraph()->GetFName();
 	}
+
+	return true;
 }
 
 void FMovieSceneToolsModule::FixupPayloadParameterNameForSection(UMovieSceneEventSectionBase* Section, UK2Node* InNode, FName OldPinName, FName NewPinName)
@@ -286,46 +305,6 @@ void FMovieSceneToolsModule::FixupPayloadParameterNameForSection(UMovieSceneEven
 	}
 }
 
-void FMovieSceneToolsModule::HandleGenerateEventEntryPoints(UMovieSceneEventSectionBase* EventSection, const FGenerateBlueprintFunctionParams& Params)
-{
-	if (EventSection->HasAnyFlags(RF_NeedLoad))
-	{
-		FLinkerLoad* Linker = EventSection->GetLinker();
-		Linker->Preload(EventSection);
-	}
-
-	EventSection->AttemptUpgrade();
-
-	for (FMovieSceneEvent& EntryPoint : EventSection->GetAllEntryPoints())
-	{
-		UEdGraphNode* Endpoint = FMovieSceneEventUtils::FindEndpoint(&EntryPoint, EventSection, Params.CompilerContext->Blueprint);
-		if (Endpoint)
-		{
-			UK2Node_FunctionEntry* FunctionEntry = FMovieSceneEventUtils::GenerateEntryPoint(EventSection, &EntryPoint, Params.CompilerContext, Endpoint);
-			if (FunctionEntry)
-			{
-				EntryPoint.CompiledFunctionName = FunctionEntry->GetGraph()->GetFName();
-			}
-			else
-			{
-				EntryPoint.CompiledFunctionName = NAME_None;
-			}
-		}
-	}
-
-
-	auto OnFunctionListGenerated = [WeakEventSection = MakeWeakObjectPtr(EventSection)](FKismetCompilerContext* CompilerContext)
-	{
-		UMovieSceneEventSectionBase* PinnedSection = WeakEventSection.Get();
-		if (ensureMsgf(PinnedSection, TEXT("Event section has been collected during blueprint compilation.")))
-		{
-			PinnedSection->OnPostCompile(CompilerContext->Blueprint);
-		}
-	};
-
-	Params.CompilerContext->OnFunctionListCompiled().AddLambda(OnFunctionListGenerated);
-}
-
 void FMovieSceneToolsModule::RegisterClipboardConversions()
 {
 	using namespace MovieSceneClipboard;
diff --git a/Engine/Source/Editor/MovieSceneTools/Public/MovieSceneEventUtils.h b/Engine/Source/Editor/MovieSceneTools/Public/MovieSceneEventUtils.h
--- a/Engine/Source/Editor/MovieSceneTools/Public/MovieSceneEventUtils.h
+++ b/Engine/Source/Editor/MovieSceneTools/Public/MovieSceneEventUtils.h
@@ -44,7 +44,7 @@ struct MOVIESCENETOOLS_API FMovieSceneEventUtils
 {
 	/**
 	 * Create a new event endpoint for the specified sequence. By default, event end points are presented as a UK2Node_CustomEvent in an event graph called "Sequencer Events" with a pin for the bound object.
-	 * @note: Does not bind any event sections to UBlueprint::GenerateFunctionGraphsEvent which is necessary for the event's entry point to be generated during blueprint compilation - that must be done before (or after) this call.
+	 * @note: Does not bind any event sections to the UBlueprintExtension which is necessary for the event's entry point to be generated during blueprint compilation - that must be done before (or after) this call.
 	 *
 	 * @param Blueprint       (Required, non-null) The blueprint within which to create the new custom event
 	 * @param Parameters      The end point parameters that define how to name and construct the event
@@ -54,7 +54,7 @@ struct MOVIESCENETOOLS_API FMovieSceneEventUtils
 
 
 	/**
-	 * As with CreateUserFacingEvent, but also binds the section to UBlueprint::GenerateFunctionGraphsEvent which ensures that it will hook into the compilation process to generate its entry point function graph
+	 * As with CreateUserFacingEvent, but also binds the section to a UBlueprintExtension which ensures that it will hook into the compilation process to generate its entry point function graph
 	 *
 	 * @param EntryPoint      (Required, non-null) The event entry point definition
 	 * @param EventSection    (Required, non-null) The event section that owns the specified entry point
@@ -74,6 +74,13 @@ struct MOVIESCENETOOLS_API FMovieSceneEventUtils
 	 */
 	static UK2Node_FunctionEntry* GenerateEntryPoint(UMovieSceneEventSectionBase* EventSection, FMovieSceneEvent* EntryPoint, FKismetCompilerContext* Compiler, UEdGraphNode* Endpoint);
 
+	/**
+	 * Create the necessary mapping between a director Blueprint and and event section in order for the event section to be compiled correctly.
+	 *
+	 * @param EventSection           (Required, non-null) The event section to bind to the director BP
+	 * @param DirectorBP             (Required, non-null) The owning blueprint that will compile the event section's entry points
+	 */
+	static void BindEventSectionToBlueprint(UMovieSceneEventSectionBase* EventSection, UBlueprint* DirectorBP);
 
 public:
 
diff --git a/Engine/Source/Editor/MovieSceneTools/Public/MovieSceneToolsModule.h b/Engine/Source/Editor/MovieSceneTools/Public/MovieSceneToolsModule.h
--- a/Engine/Source/Editor/MovieSceneTools/Public/MovieSceneToolsModule.h
+++ b/Engine/Source/Editor/MovieSceneTools/Public/MovieSceneToolsModule.h
@@ -14,8 +14,6 @@ class UBlueprint;
 class UMovieSceneSection;
 class UMovieSceneEventSectionBase;
 
-struct FGenerateBlueprintFunctionParams;
-
 class IMovieSceneToolsTakeData
 {
 public:
@@ -52,9 +50,8 @@ private:
 
 	void RegisterClipboardConversions();
 
-	static void HandleGenerateEventEntryPoints(UMovieSceneEventSectionBase* EventSection, const FGenerateBlueprintFunctionParams& Params);
 	static void FixupPayloadParameterNameForSection(UMovieSceneEventSectionBase* Section, UK2Node* InNode, FName OldPinName, FName NewPinName);
-	static void UpgradeLegacyEventEndpointForSection(UMovieSceneEventSectionBase* Section, UBlueprint* Blueprint);
+	static bool UpgradeLegacyEventEndpointForSection(UMovieSceneEventSectionBase* Section, UBlueprint* Blueprint);
 
 private:
 
diff --git a/Engine/Source/Editor/UnrealEd/UnrealEd.Build.cs b/Engine/Source/Editor/UnrealEd/UnrealEd.Build.cs
--- a/Engine/Source/Editor/UnrealEd/UnrealEd.Build.cs
+++ b/Engine/Source/Editor/UnrealEd/UnrealEd.Build.cs
@@ -217,7 +217,6 @@ public class UnrealEd : ModuleRules
 				"DeviceManager",
 				"SettingsEditor",
 				"SessionFrontend",
-				"Sequencer",
 				"StringTableEditor",
 				"FoliageEdit",
 				"ImageWrapper",
diff --git a/Engine/Source/Runtime/MovieSceneTracks/Private/Channels/MovieSceneEvent.cpp b/Engine/Source/Runtime/MovieSceneTracks/Private/Channels/MovieSceneEvent.cpp
--- a/Engine/Source/Runtime/MovieSceneTracks/Private/Channels/MovieSceneEvent.cpp
+++ b/Engine/Source/Runtime/MovieSceneTracks/Private/Channels/MovieSceneEvent.cpp
@@ -1,24 +1,3 @@
 // Copyright 1998-2019 Epic Games, Inc. All Rights Reserved.
 
 #include "Channels/MovieSceneEvent.h"
-
-#if WITH_EDITORONLY_DATA
-	#include "EdGraph/EdGraph.h"
-	#include "K2Node_FunctionEntry.h"
-#endif
-
-#if WITH_EDITORONLY_DATA
-void FMovieSceneEvent::PostSerialize(const FArchive& Ar)
-{
-	UK2Node_FunctionEntry* FunctionEntry = Cast<UK2Node_FunctionEntry>(FunctionEntry_DEPRECATED.Get());
-	if (FunctionEntry)
-	{
-		// If the graph needs loading, load it
-		if (FunctionEntry->GetGraph()->HasAnyFlags(RF_NeedLoad) && FunctionEntry->GetLinker())
-		{
-			FunctionEntry->GetLinker()->Preload(FunctionEntry->GetGraph());
-		}
-		GraphGuid = FunctionEntry->GetGraph()->GraphGuid;
-	}
-}
-#endif
diff --git a/Engine/Source/Runtime/MovieSceneTracks/Private/Sections/MovieSceneEventSectionBase.cpp b/Engine/Source/Runtime/MovieSceneTracks/Private/Sections/MovieSceneEventSectionBase.cpp
--- a/Engine/Source/Runtime/MovieSceneTracks/Private/Sections/MovieSceneEventSectionBase.cpp
+++ b/Engine/Source/Runtime/MovieSceneTracks/Private/Sections/MovieSceneEventSectionBase.cpp
@@ -8,7 +8,6 @@
 #include "Engine/Blueprint.h"
 #include "Kismet2/BlueprintEditorUtils.h"
 
-UMovieSceneEventSectionBase::FGenerateEventEntryPointFunctionsEvent UMovieSceneEventSectionBase::GenerateEventEntryPointsEvent;
 UMovieSceneEventSectionBase::FFixupPayloadParameterNameEvent UMovieSceneEventSectionBase::FixupPayloadParameterNameEvent;
 UMovieSceneEventSectionBase::FUpgradeLegacyEventEndpoint UMovieSceneEventSectionBase::UpgradeLegacyEventEndpoint;
 
@@ -59,24 +58,18 @@ void UMovieSceneEventSectionBase::OnPostCompile(UBlueprint* Blueprint)
 void UMovieSceneEventSectionBase::AttemptUpgrade()
 {
 	UBlueprint* Blueprint = DirectorBlueprint_DEPRECATED.Get();
+	// If we do not have the deprecated blueprint then this has already been upgraded and is this function is not necessary
 	if (!Blueprint)
 	{
 		return;
 	}
 
-	// Bind onto this event to ensure that we get another chance to upgrade when the BP compiles
-	Blueprint->GenerateFunctionGraphsEvent.AddUniqueDynamic(this, &UMovieSceneEventSectionBase::HandleGenerateEntryPoints);
-
-	for (UEdGraph* EdGraph : Blueprint->FunctionGraphs)
+	const bool bUpgradeSuccess = UpgradeLegacyEventEndpoint.IsBound() ? UpgradeLegacyEventEndpoint.Execute(this, Blueprint) : false;
+	if (!bUpgradeSuccess)
 	{
-		if (EdGraph->HasAnyFlags(RF_NeedLoad))
-		{
-			return;
-		}
+		return;
 	}
 
-	UpgradeLegacyEventEndpoint.Broadcast(this, Blueprint);
-
 	// If the BP has already been compiled (eg regenerate on load) we must perform PostCompile fixup immediately since
 	// We will not have had a chance to generate function entries. In this case we just bind directly to the already compiled functions.
 	if (Blueprint->bHasBeenRegenerated)
diff --git a/Engine/Source/Runtime/MovieSceneTracks/Public/Channels/MovieSceneEvent.h b/Engine/Source/Runtime/MovieSceneTracks/Public/Channels/MovieSceneEvent.h
--- a/Engine/Source/Runtime/MovieSceneTracks/Public/Channels/MovieSceneEvent.h
+++ b/Engine/Source/Runtime/MovieSceneTracks/Public/Channels/MovieSceneEvent.h
@@ -53,10 +53,6 @@ struct FMovieSceneEvent
 	UPROPERTY()
 	FMovieSceneEventPtrs Ptrs;
 
-#if WITH_EDITORONLY_DATA
-	void PostSerialize(const FArchive& Ar);
-#endif
-
 public:
 
 #if WITH_EDITORONLY_DATA
@@ -83,8 +79,6 @@ public:
 	UPROPERTY(transient)
 	TWeakObjectPtr<UObject> WeakCachedEndpoint;
 
-private:
-
 	/** Deprecated weak pointer to the function entry to call - no longer serialized but cached on load. */
 	UPROPERTY()
 	TWeakObjectPtr<UObject> FunctionEntry_DEPRECATED;
@@ -92,16 +86,6 @@ private:
 #endif // WITH_EDITORONLY_DATA
 };
 
-#if WITH_EDITORONLY_DATA
-template<>
-struct TStructOpsTypeTraits<FMovieSceneEvent> : public TStructOpsTypeTraitsBase2<FMovieSceneEvent>
-{
-	enum
-	{
-		WithPostSerialize = true,
-	};
-};
-#endif
 
 
 
diff --git a/Engine/Source/Runtime/MovieSceneTracks/Public/Sections/MovieSceneEventSectionBase.h b/Engine/Source/Runtime/MovieSceneTracks/Public/Sections/MovieSceneEventSectionBase.h
--- a/Engine/Source/Runtime/MovieSceneTracks/Public/Sections/MovieSceneEventSectionBase.h
+++ b/Engine/Source/Runtime/MovieSceneTracks/Public/Sections/MovieSceneEventSectionBase.h
@@ -26,20 +26,8 @@ public:
 
 	MOVIESCENETRACKS_API void AttemptUpgrade();
 
-	DECLARE_MULTICAST_DELEGATE_TwoParams(FGenerateEventEntryPointFunctionsEvent, UMovieSceneEventSectionBase*, const FGenerateBlueprintFunctionParams&);
 	DECLARE_MULTICAST_DELEGATE_FourParams(FFixupPayloadParameterNameEvent, UMovieSceneEventSectionBase*, UK2Node*, FName, FName);
-	DECLARE_MULTICAST_DELEGATE_TwoParams(FUpgradeLegacyEventEndpoint, UMovieSceneEventSectionBase*, UBlueprint*);
-
-
-	/**
-	 * Event handler that should be bound to a sequence director's UBlueprint::GenerateFunctionGraphsEvent member in order to generate event entrypoints
-	 * This function is declared as a UFunction to ensure that it gets persistently serialized along with the UBlueprint
-	 */
-	UFUNCTION()
-	void HandleGenerateEntryPoints(const FGenerateBlueprintFunctionParams& Params)
-	{
-		GenerateEventEntryPointsEvent.Broadcast(this, Params);
-	}
+	DECLARE_DELEGATE_RetVal_TwoParams(bool, FUpgradeLegacyEventEndpoint, UMovieSceneEventSectionBase*, UBlueprint*);
 
 	/**
 	 * Handler should be invoked when an event endpoint that is referenced from this section has one of its pins renamed
@@ -56,18 +44,14 @@ public:
 	MOVIESCENETRACKS_API void OnPostCompile(UBlueprint* Blueprint);
 
 
-	/**
-	 * Event that is broadcast when events need to be generated for a function. Implemented in this way so that editor-code can be kept within editor modules.
-	 */
-	MOVIESCENETRACKS_API static FGenerateEventEntryPointFunctionsEvent GenerateEventEntryPointsEvent;
-
 	/**
 	 * Event that is broadcast when event payloads may need fixing up due to a pin rename
 	 */
 	MOVIESCENETRACKS_API static FFixupPayloadParameterNameEvent FixupPayloadParameterNameEvent;
 
 	/**
-	 * Event that is broadcast when a legacy event section is found that needs fixing up against an already compile-on-loaded blueprint
+	 * Delegate that is used to upgrade legacy event sections that need fixing up against a blueprint. Called on serialization and on compilation if necessary until successful upgrade occurs.
+	 * Must return true on success or false on failure.
 	 */
 	MOVIESCENETRACKS_API static FUpgradeLegacyEventEndpoint UpgradeLegacyEventEndpoint;
 
