Corrected case of #includes

diff --git a/Engine/Source/Editor/MovieSceneTools/Private/MovieSceneEventBlueprintExtension.cpp b/Engine/Source/Editor/MovieSceneTools/Private/MovieSceneEventBlueprintExtension.cpp
--- a/Engine/Source/Editor/MovieSceneTools/Private/MovieSceneEventBlueprintExtension.cpp
+++ b/Engine/Source/Editor/MovieSceneTools/Private/MovieSceneEventBlueprintExtension.cpp
@@ -2,7 +2,7 @@
 
 #include "MovieSceneEventBlueprintExtension.h"
 #include "MovieSceneEventUtils.h"
-#include "Sections/MoviesceneEventSectionBase.h"
+#include "Sections/MovieSceneEventSectionBase.h"
 #include "K2Node_FunctionEntry.h"
 #include "KismetCompiler.h"
 #include "Engine/Blueprint.h"
diff --git a/Engine/Source/Editor/MovieSceneTools/Private/MovieSceneEventBlueprintExtension.h b/Engine/Source/Editor/MovieSceneTools/Private/MovieSceneEventBlueprintExtension.h
--- a/Engine/Source/Editor/MovieSceneTools/Private/MovieSceneEventBlueprintExtension.h
+++ b/Engine/Source/Editor/MovieSceneTools/Private/MovieSceneEventBlueprintExtension.h
@@ -5,7 +5,7 @@
 #include "Containers/Array.h"
 #include "Blueprint/BlueprintExtension.h"
 
-#include "MoviesceneEventBlueprintExtension.generated.h"
+#include "MovieSceneEventBlueprintExtension.generated.h"
 
 class UBlueprint;
 class FKismetCompilerContext;
