Sequencer: Fix quick bind menu functions for component BPs

#rnx

diff --git a/Engine/Source/Editor/MovieSceneTools/Private/MovieSceneEventCustomization.cpp b/Engine/Source/Editor/MovieSceneTools/Private/MovieSceneEventCustomization.cpp
--- a/Engine/Source/Editor/MovieSceneTools/Private/MovieSceneEventCustomization.cpp
+++ b/Engine/Source/Editor/MovieSceneTools/Private/MovieSceneEventCustomization.cpp
@@ -85,7 +85,12 @@ namespace UE4_MovieSceneEventCustomization
 				}
 			}
 
-			CallOnMemberFilter.AddRejectionTest(FBlueprintActionFilter::FRejectionTestDelegate::CreateStatic(RejectAnyUnboundActions));
+			FBlueprintActionFilter::AddUnique(CallOnMemberFilter.TargetClasses, BoundObjectPinClass);
+	
+			// This removes duplicate entries (ie. Set Static Mesh and Set Static Mesh (StaticMeshComponent)), 
+			// but also prevents displaying functions on BP components. Comment out for now.
+			//CallOnMemberFilter.AddRejectionTest(FBlueprintActionFilter::FRejectionTestDelegate::CreateStatic(RejectAnyUnboundActions));
+			
 			CallOnMemberFilter.AddRejectionTest(FBlueprintActionFilter::FRejectionTestDelegate::CreateStatic(RejectAnyNonFunctions));
 
 			ContextMenuBuilder.AddMenuSection(CallOnMemberFilter, FText::FromName(BoundObjectPinClass->GetFName()), 0);
