Show Visible/Hidden level names on level visibility track

diff --git a/Engine/Source/Editor/MovieSceneTools/Private/Sections/LevelVisibilitySection.cpp b/Engine/Source/Editor/MovieSceneTools/Private/Sections/LevelVisibilitySection.cpp
--- a/Engine/Source/Editor/MovieSceneTools/Private/Sections/LevelVisibilitySection.cpp
+++ b/Engine/Source/Editor/MovieSceneTools/Private/Sections/LevelVisibilitySection.cpp
@@ -12,6 +12,11 @@
 #include "DragAndDrop/LevelDragDropOp.h"
 #include "ScopedTransaction.h"
 
+namespace LevelVisibilitySection
+{
+	constexpr int32 MaxNumLevelsToShow = 3;
+}
+
 FLevelVisibilitySection::FLevelVisibilitySection( UMovieSceneLevelVisibilitySection& InSectionObject )
 	: SectionObject( InSectionObject )
 {
@@ -60,8 +65,26 @@ FSlateColor FLevelVisibilitySection::GetBackgroundColor() const
 
 FText FLevelVisibilitySection::GetVisibilityText() const
 {
+	TArray<FString> LevelNameStrings;
+	LevelNameStrings.Reserve(LevelVisibilitySection::MaxNumLevelsToShow);
+	
+	int32 Count = 0;
+	for ( ; Count < LevelVisibilitySection::MaxNumLevelsToShow && Count < SectionObject.GetLevelNames().Num(); Count++ )
+	{
+		LevelNameStrings.Add( SectionObject.GetLevelNames()[Count].ToString() );
+	}
+
+	int32 NumRemaining = SectionObject.GetLevelNames().Num() - Count;
+	FString LevelsText = FString::Join( LevelNameStrings, TEXT( ", " ) );
+
+	if (LevelNameStrings.Num() == LevelVisibilitySection::MaxNumLevelsToShow )
+	{
+		LevelsText.Append( FString::Format(TEXT(" (+{0} more)"), { FString::FormatAsNumber(NumRemaining) } ) );
+	}
+
 	FText VisibilityText = SectionObject.GetVisibility() == ELevelVisibility::Visible ? VisibleText : HiddenText;
-	return FText::Format( NSLOCTEXT( "LevelVisibilitySection", "SectionTextFormat", "{0} ({1})" ), VisibilityText, FText::AsNumber( SectionObject.GetLevelNames().Num() ) );
+	
+	return FText::Format( NSLOCTEXT( "LevelVisibilitySection", "SectionTextFormat", "{0}: {1}" ), VisibilityText, FText::FromString( LevelsText ) );
 }
 
 FText FLevelVisibilitySection::GetVisibilityToolTip() const
