Fixed off-by-1 error in showing level visibility track names

diff --git a/Engine/Source/Editor/MovieSceneTools/Private/Sections/LevelVisibilitySection.cpp b/Engine/Source/Editor/MovieSceneTools/Private/Sections/LevelVisibilitySection.cpp
--- a/Engine/Source/Editor/MovieSceneTools/Private/Sections/LevelVisibilitySection.cpp
+++ b/Engine/Source/Editor/MovieSceneTools/Private/Sections/LevelVisibilitySection.cpp
@@ -77,7 +77,7 @@ FText FLevelVisibilitySection::GetVisibilityText() const
 	int32 NumRemaining = SectionObject.GetLevelNames().Num() - Count;
 	FString LevelsText = FString::Join( LevelNameStrings, TEXT( ", " ) );
 
-	if (LevelNameStrings.Num() == LevelVisibilitySection::MaxNumLevelsToShow )
+	if (SectionObject.GetLevelNames().Num() > LevelVisibilitySection::MaxNumLevelsToShow )
 	{
 		LevelsText.Append( FString::Format(TEXT(" (+{0} more)"), { FString::FormatAsNumber(NumRemaining) } ) );
 	}
