Curve Editor: Add new symbols for keys with weighted tangents

diff --git a/Engine/Source/Editor/EditorStyle/Private/SlateEditorStyle.cpp b/Engine/Source/Editor/EditorStyle/Private/SlateEditorStyle.cpp
--- a/Engine/Source/Editor/EditorStyle/Private/SlateEditorStyle.cpp
+++ b/Engine/Source/Editor/EditorStyle/Private/SlateEditorStyle.cpp
@@ -7232,6 +7232,7 @@ void FSlateEditorStyle::FStyle::SetupToolkitStyles()
 		Set("GenericCurveEditor.LinearKey", new IMAGE_BRUSH("GenericCurveEditor/Keys/Triangle_Filled", FVector2D(11.0f, 11.0f)));
 		Set("GenericCurveEditor.CubicKey", new IMAGE_BRUSH("GenericCurveEditor/Keys/Diamond_Filled", FVector2D(11.0f, 11.0f)));
 		Set("GenericCurveEditor.TangentHandle", new IMAGE_BRUSH("GenericCurveEditor/Keys/TangentHandle", FVector2D(8.0f, 8.0f)));
+		Set("GenericCurveEditor.WeightedTangentCubicKey", new IMAGE_BRUSH("GenericCurveEditor/Keys/Trapezoid_Filled", FVector2D(11.0f, 11.0f)));
 
 		// Pre-Infinity
 		Set("GenericCurveEditor.SetPreInfinityExtrapConstant", new IMAGE_BRUSH("GenericCurveEditor/Icons/PreInfinityConstant_48x", Icon48x48));
diff --git a/Engine/Source/Editor/MovieSceneTools/Private/Channels/FloatChannelCurveModel.cpp b/Engine/Source/Editor/MovieSceneTools/Private/Channels/FloatChannelCurveModel.cpp
--- a/Engine/Source/Editor/MovieSceneTools/Private/Channels/FloatChannelCurveModel.cpp
+++ b/Engine/Source/Editor/MovieSceneTools/Private/Channels/FloatChannelCurveModel.cpp
@@ -185,6 +185,7 @@ void FFloatChannelCurveModel::GetKeyDrawInfo(ECurvePointType PointType, const FK
 		OutDrawInfo.ScreenSize = FVector2D(11, 11);
 
 		ERichCurveInterpMode KeyType = RCIM_None;
+		ERichCurveTangentWeightMode KeyTWType = RCTWM_WeightedNone;
 
 		// Get the key type from the supplied key handle if it's valid
 		FMovieSceneFloatChannel* Channel = ChannelHandle.Get();
@@ -195,6 +196,7 @@ void FFloatChannelCurveModel::GetKeyDrawInfo(ECurvePointType PointType, const FK
 			if (KeyIndex != INDEX_NONE)
 			{
 				KeyType = ChannelData.GetValues()[KeyIndex].InterpMode;
+				KeyTWType = ChannelData.GetValues()[KeyIndex].Tangent.TangentWeightMode;
 			}
 		}
 
@@ -209,7 +211,15 @@ void FFloatChannelCurveModel::GetKeyDrawInfo(ECurvePointType PointType, const FK
 			OutDrawInfo.Tint = FLinearColor(0, 0.62f, 0.46f);
 			break;
 		case ERichCurveInterpMode::RCIM_Cubic:
-			OutDrawInfo.Brush = FEditorStyle::GetBrush("GenericCurveEditor.CubicKey");
+			if (KeyTWType == ERichCurveTangentWeightMode::RCTWM_WeightedBoth)
+			{
+				OutDrawInfo.Brush = FEditorStyle::GetBrush("GenericCurveEditor.WeightedTangentCubicKey");
+			}
+			else
+			{
+				OutDrawInfo.Brush = FEditorStyle::GetBrush("GenericCurveEditor.CubicKey");
+			}
+
 			OutDrawInfo.Tint = FLinearColor::White;
 			break;
 		default:
