Sequencer: Reworked GatherTakes and GetTakeNumber so that the asset path can be shown in the tooltip. This disambiguates when there are multiple assets with the same take number (ie. an asset was duplicated and therefore the metadata was duplicated). Also, creating a new take, now increments the metadata take number.

diff --git a/Engine/Plugins/VirtualProduction/Takes/Source/TakesCore/Private/TakeData.h b/Engine/Plugins/VirtualProduction/Takes/Source/TakesCore/Private/TakeData.h
--- a/Engine/Plugins/VirtualProduction/Takes/Source/TakesCore/Private/TakeData.h
+++ b/Engine/Plugins/VirtualProduction/Takes/Source/TakesCore/Private/TakeData.h
@@ -18,7 +18,7 @@ class FTakesCoreTakeData : public IMovieSceneToolsTakeData
 public:
 	virtual ~FTakesCoreTakeData() { }
 
-	virtual bool GatherTakes(const UMovieSceneSection* Section, TArray<uint32>& TakeNumbers, uint32& CurrentTakeNumber)
+	virtual bool GatherTakes(const UMovieSceneSection* Section, TArray<FAssetData>& AssetData, uint32& OutCurrentTakeNumber)
 	{
 		const UMovieSceneSubSection* SubSection = Cast<UMovieSceneSubSection>(Section);
 		if (!SubSection)
@@ -46,50 +46,82 @@ public:
 			int32 ThisTakeNumber = 0;
 			if (TakeNumberTag.IsSet() && LexTryParseString(ThisTakeNumber, *TakeNumberTag.GetValue()))
 			{
-				TakeNumbers.Add(ThisTakeNumber);
+				AssetData.Add(TakeData);
 			}
 		}
 
-		CurrentTakeNumber = TakeMetaData->GetTakeNumber();
+		OutCurrentTakeNumber = TakeMetaData->GetTakeNumber();
 
 		return true;
 	}
 
-	virtual UObject* GetTake(const UMovieSceneSection* Section, uint32 TakeNumber)
+	virtual bool GetTakeNumber(const UMovieSceneSection* Section, FAssetData AssetData, uint32& OutTakeNumber)
 	{
 		const UMovieSceneSubSection* SubSection = Cast<UMovieSceneSubSection>(Section);
 		if (!SubSection)
 		{
-			return nullptr;
+			return false;
 		}
 
 		ULevelSequence* Sequence = Cast<ULevelSequence>(SubSection->GetSequence());
 		if (!Sequence)
 		{
-			return nullptr;
+			return false;
 		}
 
 		UTakeMetaData* TakeMetaData = Sequence->FindMetaData<UTakeMetaData>();
 		if (!TakeMetaData)
 		{
-			return nullptr;
+			return false;
 		}
 
 		TArray<FAssetData> TakeDatas = UTakesCoreBlueprintLibrary::FindTakes(TakeMetaData->GetSlate());
 		for (FAssetData TakeData : TakeDatas)
 		{
-			FAssetDataTagMapSharedView::FFindTagResult TakeNumberTag = TakeData.TagsAndValues.FindTag(UTakeMetaData::AssetRegistryTag_TakeNumber);
-
-			int32 ThisTakeNumber = 0;
-			if (TakeNumberTag.IsSet() && LexTryParseString(ThisTakeNumber, *TakeNumberTag.GetValue()))
+			if (TakeData == AssetData)
 			{
-				if (ThisTakeNumber == TakeNumber)
+				FAssetDataTagMapSharedView::FFindTagResult TakeNumberTag = TakeData.TagsAndValues.FindTag(UTakeMetaData::AssetRegistryTag_TakeNumber);
+
+				int32 ThisTakeNumber = 0;
+				if (TakeNumberTag.IsSet() && LexTryParseString(ThisTakeNumber, *TakeNumberTag.GetValue()))
 				{
-					return TakeData.GetAsset();
+					OutTakeNumber = ThisTakeNumber;
+					return true;
 				}
 			}
 		}
 
-		return nullptr;
+		return false;
+	}
+
+	virtual bool SetTakeNumber(const UMovieSceneSection* Section, uint32 InTakeNumber)
+	{
+		const UMovieSceneSubSection* SubSection = Cast<UMovieSceneSubSection>(Section);
+		if (!SubSection)
+		{
+			return false;
+		}
+
+		ULevelSequence* Sequence = Cast<ULevelSequence>(SubSection->GetSequence());
+		if (!Sequence)
+		{
+			return false;
+		}
+
+		UTakeMetaData* TakeMetaData = Sequence->FindMetaData<UTakeMetaData>();
+		if (!TakeMetaData)
+		{
+			return false;
+		}
+
+		bool bWasLocked = TakeMetaData->IsLocked();
+		TakeMetaData->Unlock();
+		TakeMetaData->SetTakeNumber(InTakeNumber);
+		if (bWasLocked)
+		{
+			TakeMetaData->Lock();
+		}
+
+		return true;
 	}
 };
\ No newline at end of file
diff --git a/Engine/Source/Editor/MovieSceneTools/Private/MovieSceneToolHelpers.cpp b/Engine/Source/Editor/MovieSceneTools/Private/MovieSceneToolHelpers.cpp
--- a/Engine/Source/Editor/MovieSceneTools/Private/MovieSceneToolHelpers.cpp
+++ b/Engine/Source/Editor/MovieSceneTools/Private/MovieSceneToolHelpers.cpp
@@ -363,7 +363,7 @@ FString MovieSceneToolHelpers::GenerateNewShotName(const TArray<UMovieSceneSecti
 	return ComposeShotName(ProjectSettings->ShotPrefix, ProjectSettings->FirstShotNumber, ProjectSettings->FirstTakeNumber);
 }
 
-void MovieSceneToolHelpers::GatherTakes(const UMovieSceneSection* Section, TArray<uint32>& TakeNumbers, uint32& CurrentTakeNumber)
+void MovieSceneToolHelpers::GatherTakes(const UMovieSceneSection* Section, TArray<FAssetData>& AssetData, uint32& OutCurrentTakeNumber)
 {
 	const UMovieSceneSubSection* SubSection = Cast<const UMovieSceneSubSection>(Section);
 	
@@ -372,7 +372,7 @@ void MovieSceneToolHelpers::GatherTakes(const UMovieSceneSection* Section, TArra
 		return;
 	}
 
-	if (FMovieSceneToolsModule::Get().GatherTakes(Section, TakeNumbers, CurrentTakeNumber))
+	if (FMovieSceneToolsModule::Get().GatherTakes(Section, AssetData, OutCurrentTakeNumber))
 	{
 		return;
 	}
@@ -383,7 +383,7 @@ void MovieSceneToolHelpers::GatherTakes(const UMovieSceneSection* Section, TArra
 
 	FString ShotPrefix;
 	uint32 ShotNumber = INDEX_NONE;
-	CurrentTakeNumber = INDEX_NONE;
+	OutCurrentTakeNumber = INDEX_NONE;
 
 	FString SubSectionName = SubSection->GetSequence()->GetName();
 	if (SubSection->IsA<UMovieSceneCinematicShotSection>())
@@ -392,7 +392,7 @@ void MovieSceneToolHelpers::GatherTakes(const UMovieSceneSection* Section, TArra
 		SubSectionName = ShotSection->GetShotDisplayName();
 	}
 
-	if (ParseShotName(SubSectionName, ShotPrefix, ShotNumber, CurrentTakeNumber))
+	if (ParseShotName(SubSectionName, ShotPrefix, ShotNumber, OutCurrentTakeNumber))
 	{
 		// Gather up all level sequence assets
 		FAssetRegistryModule& AssetRegistryModule = FModuleManager::LoadModuleChecked<FAssetRegistryModule>(TEXT("AssetRegistry"));
@@ -413,25 +413,23 @@ void MovieSceneToolHelpers::GatherTakes(const UMovieSceneSection* Section, TArra
 				{
 					if (AssetShotPrefix == ShotPrefix && AssetShotNumber == ShotNumber)
 					{
-						TakeNumbers.Add(AssetTakeNumber);
+						AssetData.Add(AssetObject);
 					}
 				}
 			}
 		}
 	}
-
-	TakeNumbers.Sort();
 }
 
-UObject* MovieSceneToolHelpers::GetTake(const UMovieSceneSection* Section, uint32 TakeNumber)
+bool MovieSceneToolHelpers::GetTakeNumber(const UMovieSceneSection* Section, FAssetData AssetData, uint32& OutTakeNumber)
 {
-	const UMovieSceneSubSection* SubSection = Cast<const UMovieSceneSubSection>(Section);
-
-	if (UObject* TakeObject = FMovieSceneToolsModule::Get().GetTake(Section, TakeNumber))
+	if (FMovieSceneToolsModule::Get().GetTakeNumber(Section, AssetData, OutTakeNumber))
 	{
-		return TakeObject;
+		return true;
 	}
 
+	const UMovieSceneSubSection* SubSection = Cast<const UMovieSceneSubSection>(Section);
+
 	FAssetData ShotData(SubSection->GetSequence()->GetOuter());
 
 	FString ShotPackagePath = ShotData.PackagePath.ToString();
@@ -459,29 +457,38 @@ UObject* MovieSceneToolHelpers::GetTake(const UMovieSceneSection* Section, uint3
 
 		for (auto AssetObject : ObjectList)
 		{
-			FString AssetPackagePath = AssetObject.PackagePath.ToString();
-			int32 AssetLastSlashPos = INDEX_NONE;
-			AssetPackagePath.FindLastChar(TCHAR('/'), AssetLastSlashPos);
-			AssetPackagePath = AssetPackagePath.Left(AssetLastSlashPos);
-
-			if (AssetPackagePath == ShotPackagePath)
+			if (AssetObject == AssetData)
 			{
-				FString AssetShotPrefix;
-				uint32 AssetShotNumber = INDEX_NONE;
-				uint32 AssetTakeNumber = INDEX_NONE;
+				FString AssetPackagePath = AssetObject.PackagePath.ToString();
+				int32 AssetLastSlashPos = INDEX_NONE;
+				AssetPackagePath.FindLastChar(TCHAR('/'), AssetLastSlashPos);
+				AssetPackagePath = AssetPackagePath.Left(AssetLastSlashPos);
 
-				if (ParseShotName(AssetObject.AssetName.ToString(), AssetShotPrefix, AssetShotNumber, AssetTakeNumber))
+				if (AssetPackagePath == ShotPackagePath)
 				{
-					if (AssetShotPrefix == ShotPrefix && AssetShotNumber == ShotNumber && TakeNumber == AssetTakeNumber)
+					FString AssetShotPrefix;
+					uint32 AssetShotNumber = INDEX_NONE;
+					uint32 AssetTakeNumber = INDEX_NONE;
+
+					if (ParseShotName(AssetObject.AssetName.ToString(), AssetShotPrefix, AssetShotNumber, AssetTakeNumber))
 					{
-						return AssetObject.GetAsset();
+						if (AssetShotPrefix == ShotPrefix && AssetShotNumber == ShotNumber)
+						{
+							OutTakeNumber = AssetTakeNumber;
+							return true;
+						}
 					}
 				}
 			}
 		}
 	}
 
-	return nullptr;
+	return false;
+}
+
+bool MovieSceneToolHelpers::SetTakeNumber(const UMovieSceneSection* Section, uint32 InTakeNumber)
+{
+	return FMovieSceneToolsModule::Get().SetTakeNumber(Section, InTakeNumber);
 }
 
 int32 MovieSceneToolHelpers::FindAvailableRowIndex(UMovieSceneTrack* InTrack, UMovieSceneSection* InSection)
diff --git a/Engine/Source/Editor/MovieSceneTools/Private/MovieSceneToolsModule.cpp b/Engine/Source/Editor/MovieSceneTools/Private/MovieSceneToolsModule.cpp
--- a/Engine/Source/Editor/MovieSceneTools/Private/MovieSceneToolsModule.cpp
+++ b/Engine/Source/Editor/MovieSceneTools/Private/MovieSceneToolsModule.cpp
@@ -256,11 +256,11 @@ void FMovieSceneToolsModule::UnregisterTakeData(IMovieSceneToolsTakeData* InTake
 	TakeDatas.Remove(InTakeData);
 }
 
-bool FMovieSceneToolsModule::GatherTakes(const UMovieSceneSection* Section, TArray<uint32>& TakeNumbers, uint32& CurrentTakeNumber)
+bool FMovieSceneToolsModule::GatherTakes(const UMovieSceneSection* Section, TArray<FAssetData>& AssetData, uint32& OutCurrentTakeNumber)
 {
 	for (IMovieSceneToolsTakeData* TakeData : TakeDatas)
 	{
-		if (TakeData->GatherTakes(Section, TakeNumbers, CurrentTakeNumber))
+		if (TakeData->GatherTakes(Section, AssetData, OutCurrentTakeNumber))
 		{
 			return true;
 		}
@@ -270,17 +270,30 @@ bool FMovieSceneToolsModule::GatherTakes(const UMovieSceneSection* Section, TArr
 }
 
 
-UObject* FMovieSceneToolsModule::GetTake(const UMovieSceneSection* Section, uint32 TakeNumber)
+bool FMovieSceneToolsModule::GetTakeNumber(const UMovieSceneSection* Section, FAssetData AssetData, uint32& OutTakeNumber)
 {
 	for (IMovieSceneToolsTakeData* TakeData : TakeDatas)
 	{
-		if (UObject* TakeObject = TakeData->GetTake(Section, TakeNumber))
+		if (TakeData->GetTakeNumber(Section, AssetData, OutTakeNumber))
 		{
-			return TakeObject;
+			return true;
 		}
 	}
 
-	return nullptr;
+	return false;
+}
+
+bool FMovieSceneToolsModule::SetTakeNumber(const UMovieSceneSection* Section, uint32 InTakeNumber)
+{
+	for (IMovieSceneToolsTakeData* TakeData : TakeDatas)
+	{
+		if (TakeData->SetTakeNumber(Section, InTakeNumber))
+		{
+			return true;
+		}
+	}
+
+	return false;
 }
 
 IMPLEMENT_MODULE( FMovieSceneToolsModule, MovieSceneTools );
diff --git a/Engine/Source/Editor/MovieSceneTools/Private/Sections/CinematicShotSection.cpp b/Engine/Source/Editor/MovieSceneTools/Private/Sections/CinematicShotSection.cpp
--- a/Engine/Source/Editor/MovieSceneTools/Private/Sections/CinematicShotSection.cpp
+++ b/Engine/Source/Editor/MovieSceneTools/Private/Sections/CinematicShotSection.cpp
@@ -464,18 +464,37 @@ void FCinematicShotSection::BuildSectionContextMenu(FMenuBuilder& MenuBuilder, c
 
 void FCinematicShotSection::AddTakesMenu(FMenuBuilder& MenuBuilder)
 {
-	TArray<uint32> TakeNumbers;
+	TArray<FAssetData> AssetData;
 	uint32 CurrentTakeNumber = INDEX_NONE;
-	MovieSceneToolHelpers::GatherTakes(&SectionObject, TakeNumbers, CurrentTakeNumber);
+	MovieSceneToolHelpers::GatherTakes(&SectionObject, AssetData, CurrentTakeNumber);
 
-	for (auto TakeNumber : TakeNumbers)
+	AssetData.Sort([this](const FAssetData &A, const FAssetData &B) {
+		uint32 TakeNumberA = INDEX_NONE;
+		uint32 TakeNumberB = INDEX_NONE;
+		if (MovieSceneToolHelpers::GetTakeNumber(&SectionObject, A, TakeNumberA) && MovieSceneToolHelpers::GetTakeNumber(&SectionObject, B, TakeNumberB))
+		{
+			return TakeNumberA < TakeNumberB;
+		}
+		return true;
+	});
+
+	for (auto ThisAssetData : AssetData)
 	{
-		MenuBuilder.AddMenuEntry(
-			FText::Format(LOCTEXT("TakeNumber", "Take {0}"), FText::AsNumber(TakeNumber)),
-			FText::Format(LOCTEXT("TakeNumberTooltip", "Switch to take {0}"), FText::AsNumber(TakeNumber)),
-			TakeNumber == CurrentTakeNumber ? FSlateIcon(FEditorStyle::GetStyleSetName(), "Sequencer.Star") : FSlateIcon(FEditorStyle::GetStyleSetName(), "Sequencer.Empty"),
-			FUIAction(FExecuteAction::CreateSP(CinematicShotTrackEditor.Pin().ToSharedRef(), &FCinematicShotTrackEditor::SwitchTake, TakeNumber))
-		);
+		uint32 TakeNumber = INDEX_NONE;
+		if (MovieSceneToolHelpers::GetTakeNumber(&SectionObject, ThisAssetData, TakeNumber))
+		{
+			UObject* TakeObject = ThisAssetData.GetAsset();
+			
+			if (TakeObject)
+			{
+				MenuBuilder.AddMenuEntry(
+					FText::Format(LOCTEXT("TakeNumber", "Take {0}"), FText::AsNumber(TakeNumber)),
+					FText::Format(LOCTEXT("TakeNumberTooltip", "Switch to {0}"), FText::FromString(TakeObject->GetPathName())),
+					TakeNumber == CurrentTakeNumber ? FSlateIcon(FEditorStyle::GetStyleSetName(), "Sequencer.Star") : FSlateIcon(FEditorStyle::GetStyleSetName(), "Sequencer.Empty"),
+					FUIAction(FExecuteAction::CreateSP(CinematicShotTrackEditor.Pin().ToSharedRef(), &FCinematicShotTrackEditor::SwitchTake, TakeObject))
+				);
+			}
+		}
 	}
 }
 
diff --git a/Engine/Source/Editor/MovieSceneTools/Private/TrackEditors/CinematicShotTrackEditor.cpp b/Engine/Source/Editor/MovieSceneTools/Private/TrackEditors/CinematicShotTrackEditor.cpp
--- a/Engine/Source/Editor/MovieSceneTools/Private/TrackEditors/CinematicShotTrackEditor.cpp
+++ b/Engine/Source/Editor/MovieSceneTools/Private/TrackEditors/CinematicShotTrackEditor.cpp
@@ -462,13 +462,21 @@ void FCinematicShotTrackEditor::NewTake(UMovieSceneCinematicShotSection* Section
 	uint32 TakeNumber = INDEX_NONE;
 	if (MovieSceneToolHelpers::ParseShotName(Section->GetShotDisplayName(), ShotPrefix, ShotNumber, TakeNumber))
 	{
-		TArray<uint32> TakeNumbers;
-		uint32 CurrentTakeNumber;
-		MovieSceneToolHelpers::GatherTakes(Section, TakeNumbers, CurrentTakeNumber);
+		TArray<FAssetData> AssetData;
+		uint32 CurrentTakeNumber = INDEX_NONE;
+		MovieSceneToolHelpers::GatherTakes(Section, AssetData, CurrentTakeNumber);
 		uint32 NewTakeNumber = CurrentTakeNumber;
-		if (TakeNumbers.Num() > 0)
+
+		for (auto ThisAssetData : AssetData)
 		{
-			NewTakeNumber = TakeNumbers[TakeNumbers.Num()-1] + 1;
+			uint32 ThisTakeNumber = INDEX_NONE;
+			if (MovieSceneToolHelpers::GetTakeNumber(Section, ThisAssetData, ThisTakeNumber))
+			{
+				if (ThisTakeNumber >= NewTakeNumber)
+				{
+					NewTakeNumber = ThisTakeNumber + 1;
+				}
+			}
 		}
 
 		FString NewShotName = MovieSceneToolHelpers::ComposeShotName(ShotPrefix, ShotNumber, NewTakeNumber);
@@ -493,6 +501,8 @@ void FCinematicShotTrackEditor::NewTake(UMovieSceneCinematicShotSection* Section
 			NewShot->SetPreRollFrames(NewShotPrerollFrames);
 			NewShot->SetRowIndex(NewRowIndex);
 
+			MovieSceneToolHelpers::SetTakeNumber(NewShot, NewTakeNumber);
+
 			GetSequencer()->NotifyMovieSceneDataChanged( EMovieSceneDataChangeType::MovieSceneStructureItemsChanged );
 			GetSequencer()->EmptySelection();
 			GetSequencer()->SelectSection(NewShot);
@@ -502,7 +512,7 @@ void FCinematicShotTrackEditor::NewTake(UMovieSceneCinematicShotSection* Section
 }
 
 
-void FCinematicShotTrackEditor::SwitchTake(uint32 TakeNumber)
+void FCinematicShotTrackEditor::SwitchTake(UObject* TakeObject)
 {
 	bool bSwitchedTake = false;
 
@@ -520,8 +530,6 @@ void FCinematicShotTrackEditor::SwitchTake(uint32 TakeNumber)
 
 		UMovieSceneSubSection* Section = Cast<UMovieSceneSubSection>(Sections[SectionIndex]);
 
-		UObject* TakeObject = MovieSceneToolHelpers::GetTake(Section, TakeNumber);
-
 		if (TakeObject && TakeObject->IsA(UMovieSceneSequence::StaticClass()))
 		{
 			UMovieSceneSequence* MovieSceneSequence = CastChecked<UMovieSceneSequence>(TakeObject);
diff --git a/Engine/Source/Editor/MovieSceneTools/Private/TrackEditors/CinematicShotTrackEditor.h b/Engine/Source/Editor/MovieSceneTools/Private/TrackEditors/CinematicShotTrackEditor.h
--- a/Engine/Source/Editor/MovieSceneTools/Private/TrackEditors/CinematicShotTrackEditor.h
+++ b/Engine/Source/Editor/MovieSceneTools/Private/TrackEditors/CinematicShotTrackEditor.h
@@ -106,9 +106,9 @@ public:
 	/*
 	* Switch take for the selected sections
 	*
-	* @param TakeNumber The take number to switch to.
+	* @param TakeObject The take object to switch to.
 	*/
-	void SwitchTake(uint32 TakeNumber);
+	void SwitchTake(UObject* TakeObject);
 
 private:
 
diff --git a/Engine/Source/Editor/MovieSceneTools/Private/TrackEditors/SubTrackEditor.cpp b/Engine/Source/Editor/MovieSceneTools/Private/TrackEditors/SubTrackEditor.cpp
--- a/Engine/Source/Editor/MovieSceneTools/Private/TrackEditors/SubTrackEditor.cpp
+++ b/Engine/Source/Editor/MovieSceneTools/Private/TrackEditors/SubTrackEditor.cpp
@@ -460,18 +460,37 @@ private:
 
 	void AddTakesMenu(FMenuBuilder& MenuBuilder)
 	{
-		TArray<uint32> TakeNumbers;
+		TArray<FAssetData> AssetData;
 		uint32 CurrentTakeNumber = INDEX_NONE;
-		MovieSceneToolHelpers::GatherTakes(&SectionObject, TakeNumbers, CurrentTakeNumber);
+		MovieSceneToolHelpers::GatherTakes(&SectionObject, AssetData, CurrentTakeNumber);
 
-		for (auto TakeNumber : TakeNumbers)
+		AssetData.Sort([this](const FAssetData &A, const FAssetData &B) {
+			uint32 TakeNumberA = INDEX_NONE;
+			uint32 TakeNumberB = INDEX_NONE;
+			if (MovieSceneToolHelpers::GetTakeNumber(&SectionObject, A, TakeNumberA) && MovieSceneToolHelpers::GetTakeNumber(&SectionObject, B, TakeNumberB))
+			{
+				return TakeNumberA < TakeNumberB;
+			}
+			return true;
+		});
+
+		for (auto ThisAssetData : AssetData)
 		{
-			MenuBuilder.AddMenuEntry(
-				FText::Format(LOCTEXT("TakeNumber", "Take {0}"), FText::AsNumber(TakeNumber)),
-				FText::Format(LOCTEXT("TakeNumberTooltip", "Switch to take {0}"), FText::AsNumber(TakeNumber)),
-				TakeNumber == CurrentTakeNumber ? FSlateIcon(FEditorStyle::GetStyleSetName(), "Sequencer.Star") : FSlateIcon(FEditorStyle::GetStyleSetName(), "Sequencer.Empty"),
-				FUIAction(FExecuteAction::CreateSP(SubTrackEditor.Pin().ToSharedRef(), &FSubTrackEditor::SwitchTake, TakeNumber))
-			);
+			uint32 TakeNumber = INDEX_NONE;
+			if (MovieSceneToolHelpers::GetTakeNumber(&SectionObject, ThisAssetData, TakeNumber))
+			{
+				UObject* TakeObject = ThisAssetData.GetAsset();
+
+				if (TakeObject)
+				{
+					MenuBuilder.AddMenuEntry(
+						FText::Format(LOCTEXT("TakeNumber", "Take {0}"), FText::AsNumber(TakeNumber)),
+						FText::Format(LOCTEXT("TakeNumberTooltip", "Switch to {0}"), FText::FromString(TakeObject->GetPathName())),
+						TakeNumber == CurrentTakeNumber ? FSlateIcon(FEditorStyle::GetStyleSetName(), "Sequencer.Star") : FSlateIcon(FEditorStyle::GetStyleSetName(), "Sequencer.Empty"),
+						FUIAction(FExecuteAction::CreateSP(SubTrackEditor.Pin().ToSharedRef(), &FSubTrackEditor::SwitchTake, TakeObject))
+					);
+				}
+			}
 		}
 	}
 
@@ -933,7 +952,7 @@ FKeyPropertyResult FSubTrackEditor::HandleRecordNewSequenceInternal(FFrameNumber
 	return KeyPropertyResult;
 }
 
-void FSubTrackEditor::SwitchTake(uint32 TakeNumber)
+void FSubTrackEditor::SwitchTake(UObject* TakeObject)
 {
 	bool bSwitchedTake = false;
 
@@ -951,8 +970,6 @@ void FSubTrackEditor::SwitchTake(uint32 TakeNumber)
 
 		UMovieSceneSubSection* Section = Cast<UMovieSceneSubSection>(Sections[SectionIndex]);
 
-		UObject* TakeObject = MovieSceneToolHelpers::GetTake(Section, TakeNumber);
-
 		if (TakeObject && TakeObject->IsA(UMovieSceneSequence::StaticClass()))
 		{
 			UMovieSceneSequence* MovieSceneSequence = CastChecked<UMovieSceneSequence>(TakeObject);
diff --git a/Engine/Source/Editor/MovieSceneTools/Public/MovieSceneToolHelpers.h b/Engine/Source/Editor/MovieSceneTools/Public/MovieSceneToolHelpers.h
--- a/Engine/Source/Editor/MovieSceneTools/Public/MovieSceneToolHelpers.h
+++ b/Engine/Source/Editor/MovieSceneTools/Public/MovieSceneToolHelpers.h
@@ -112,20 +112,30 @@ public:
 	 * Gather takes - level sequence assets that have the same shot prefix and shot number in the same asset path (directory)
 	 * 
 	 * @param Section The section to gather takes from
-	 * @param TakeNumbers The gathered take numbers
-	 * @param CurrentTakeNumber The current take number of the section
+	 * @param AssetData The gathered asset take data
+	 * @param OutCurrentTakeNumber The current take number of the section
 	 */
-	static void GatherTakes(const UMovieSceneSection* Section, TArray<uint32>& TakeNumbers, uint32& CurrentTakeNumber);
+	static void GatherTakes(const UMovieSceneSection* Section, TArray<FAssetData>& AssetData, uint32& OutCurrentTakeNumber);
 
 
 	/**
-	 * Get the asset associated with the take number
+	 * Get the take number for the given asset
 	 *
-	 * @param Section The section to gather the take from
-	 * @param TakeNumber The take number to get
-	 * @return The asset
+	 * @param Section The section to gather the take number from
+	 * @param AssetData The take asset to search for
+	 * @param OutTakeNumber The take number for the given asset
+	 * @return Whether the take number was found
 	 */
-	static UObject* GetTake(const UMovieSceneSection* Section, uint32 TakeNumber);
+	static bool GetTakeNumber(const UMovieSceneSection* Section, FAssetData AssetData, uint32& OutTakeNumber);
+
+	/**
+	 * Set the take number for the given asset
+	 *
+	 * @param Section The section to set the take number on
+	 * @param InTakeNumber The take number for the given asset
+	 * @return Whether the take number could be set
+	 */
+	static bool SetTakeNumber(const UMovieSceneSection* Section, uint32 InTakeNumber);
 
 	/**
 	 * Get the next available row index for the section so that it doesn't overlap any other sections in time.
diff --git a/Engine/Source/Editor/MovieSceneTools/Public/MovieSceneToolsModule.h b/Engine/Source/Editor/MovieSceneTools/Public/MovieSceneToolsModule.h
--- a/Engine/Source/Editor/MovieSceneTools/Public/MovieSceneToolsModule.h
+++ b/Engine/Source/Editor/MovieSceneTools/Public/MovieSceneToolsModule.h
@@ -12,8 +12,9 @@ class UMovieSceneSection;
 class IMovieSceneToolsTakeData
 {
 public:
-	virtual bool GatherTakes(const UMovieSceneSection* Section, TArray<uint32>& TakeNumbers, uint32& CurrentTakeNumber) = 0;
-	virtual UObject* GetTake(const UMovieSceneSection* Section, uint32 TakeNumber) = 0;
+	virtual bool GatherTakes(const UMovieSceneSection* Section, TArray<FAssetData>& AssetData, uint32& OutCurrentTakeNumber) = 0;
+	virtual bool GetTakeNumber(const UMovieSceneSection* Section, FAssetData AssetData, uint32& OutTakeNumber) = 0;
+	virtual bool SetTakeNumber(const UMovieSceneSection*, uint32 InTakeNumber) = 0;
 };
 
 /**
@@ -36,9 +37,10 @@ public:
 	void RegisterTakeData(IMovieSceneToolsTakeData*);
 	void UnregisterTakeData(IMovieSceneToolsTakeData*);
 
-	bool GatherTakes(const UMovieSceneSection* Section, TArray<uint32>& TakeNumbers, uint32& CurrentTakeNumber);
-	UObject* GetTake(const UMovieSceneSection* Section, uint32 TakeNumber);
-		
+	bool GatherTakes(const UMovieSceneSection* Section, TArray<FAssetData>& AssetData, uint32& OutCurrentTakeNumber);
+	bool GetTakeNumber(const UMovieSceneSection* Section, FAssetData AssetData, uint32& OutTakeNumber);
+	bool SetTakeNumber(const UMovieSceneSection* Section, uint32 InTakeNumber);
+
 private:
 
 	void RegisterClipboardConversions();
diff --git a/Engine/Source/Editor/MovieSceneTools/Public/TrackEditors/SubTrackEditor.h b/Engine/Source/Editor/MovieSceneTools/Public/TrackEditors/SubTrackEditor.h
--- a/Engine/Source/Editor/MovieSceneTools/Public/TrackEditors/SubTrackEditor.h
+++ b/Engine/Source/Editor/MovieSceneTools/Public/TrackEditors/SubTrackEditor.h
@@ -57,7 +57,7 @@ public:
 	virtual FReply OnDrop(const FDragDropEvent& DragDropEvent, UMovieSceneTrack* Track, int32 RowIndex, const FGuid& TargetObjectGuid) override;
 	
 	/** Switch the selected sub section's take */
-	void SwitchTake(uint32 TakeNumber);
+	void SwitchTake(UObject* TakeObject);
 
 protected:
 
