Sequencer: Fix CIS

diff --git a/Engine/Source/Editor/MovieSceneTools/Public/MovieSceneToolsModule.h b/Engine/Source/Editor/MovieSceneTools/Public/MovieSceneToolsModule.h
--- a/Engine/Source/Editor/MovieSceneTools/Public/MovieSceneToolsModule.h
+++ b/Engine/Source/Editor/MovieSceneTools/Public/MovieSceneToolsModule.h
@@ -8,6 +8,7 @@
 #include "IMovieSceneTools.h"
 
 class UMovieSceneSection;
+struct FAssetData;
 
 class IMovieSceneToolsTakeData
 {
