Movie Rendering audio pass now ignores per-shot warmup and overall warmup times. This resolves an issue where audio would only be partially captured when using these settings, and the renders should be faster overall. The warmup times are used for visual effects to settle and those are irrelevant during the audio pass.

diff --git a/Engine/Source/Editor/MovieSceneTools/Private/AutomatedLevelSequenceCapture.cpp b/Engine/Source/Editor/MovieSceneTools/Private/AutomatedLevelSequenceCapture.cpp
--- a/Engine/Source/Editor/MovieSceneTools/Private/AutomatedLevelSequenceCapture.cpp
+++ b/Engine/Source/Editor/MovieSceneTools/Private/AutomatedLevelSequenceCapture.cpp
@@ -365,6 +365,12 @@ void UAutomatedLevelSequenceCapture::Initialize(TSharedPtr<FSceneViewport> InVie
 	CaptureState = ELevelSequenceCaptureState::Setup;
 	CaptureStrategy = MakeShareable(new FFixedTimeStepCaptureStrategy(Settings.GetFrameRate()));
 	CaptureStrategy->OnInitialize();
+
+	// Cache these off. We override them for audio passes and that overrides the settings members directly
+	// so we need to be able to restore them at the end
+	CachedWarmUpFrameCount = WarmUpFrameCount;
+	CachedDelayBeforeWarmUp = DelayBeforeWarmUp;
+	CachedDelayBeforeShotWarmUp = DelayBeforeShotWarmUp;
 }
 
 bool UAutomatedLevelSequenceCapture::InitializeShots()
@@ -561,6 +567,21 @@ void UAutomatedLevelSequenceCapture::SetupFrameRange()
 				 	PlaybackEndFrame = CustomEndFrame;
 				}
 
+				// This is a fun hack... Due to the fragility of this code (which this makes more fragile admittedly...) the original audio implementation
+				// just ran the entire process twice, starting all the way back at the Setup loop so that everything would be re-initialized like it was
+				// for the video capture. Unfortunately, it looks like when we switch from a Fixed Timestep clock to the Platform clock (needed for audio
+				// as audio is realtime only) this allows the sequence player to get out of sync with the warmup frame counter. Audio recording doens't start
+				// until the warmup time has passed, but because the sequence is playing at realtime it can start playing the part of the sequence you wanted
+				// the audio for before the recorder ever kicks in!
+				// To minimize changes, we're just going to override any warmup/delay times to zero here in the event that this is an audio pass, so that
+				// they should be more in sync.
+				if (bIsAudioCapturePass)
+				{
+					WarmUpFrameCount = 0;
+					DelayBeforeWarmUp = 0.f;
+					DelayBeforeShotWarmUp = 0.f; 
+				}
+
 				RemainingWarmUpFrames = FMath::Max( WarmUpFrameCount, 0 );
 				if( RemainingWarmUpFrames > 0 )
 				{
@@ -573,6 +594,8 @@ void UAutomatedLevelSequenceCapture::SetupFrameRange()
 					Actor->SequencePlayer->MotionVectorSimulation->PreserveSimulatedMotion(true);
 				}
 
+
+
 				// Override the movie scene's playback range
 				Actor->SequencePlayer->SetFrameRate(Settings.GetFrameRate());
 				Actor->SequencePlayer->SetFrameRange(PlaybackStartFrame.Value, (PlaybackEndFrame - PlaybackStartFrame).Value);
@@ -722,6 +745,11 @@ void UAutomatedLevelSequenceCapture::OnTick(float DeltaSeconds)
 				// If they don't want to render audio, or they have rendered an audio pass, we finish and finalize the data.
 				Actor->SequencePlayer->OnSequenceUpdated().Remove( OnPlayerUpdatedBinding );
 				FinalizeWhenReady();
+
+				// Restore our cached variables since these are the actual one represented in the in-engine UI
+				WarmUpFrameCount = CachedWarmUpFrameCount;
+				DelayBeforeWarmUp = CachedDelayBeforeWarmUp;
+				DelayBeforeShotWarmUp = CachedDelayBeforeShotWarmUp;
 			}
 			else
 			{
diff --git a/Engine/Source/Editor/MovieSceneTools/Public/AutomatedLevelSequenceCapture.h b/Engine/Source/Editor/MovieSceneTools/Public/AutomatedLevelSequenceCapture.h
--- a/Engine/Source/Editor/MovieSceneTools/Public/AutomatedLevelSequenceCapture.h
+++ b/Engine/Source/Editor/MovieSceneTools/Public/AutomatedLevelSequenceCapture.h
@@ -204,5 +204,10 @@ private:
 	TOptional<bool> bCachedUseCustomEndFrame;
 
 	TSharedPtr<FMovieSceneTimeController_FrameStep> TimeController;
+
+	// We cache these off on initialization so we can restore to them after running an audio pass.
+	int32 CachedWarmUpFrameCount;
+	int32 CachedDelayBeforeWarmUp;
+	int32 CachedDelayBeforeShotWarmUp;
 };
 
