Sequencer: Add menu item to attach audio to an existing binding

diff --git a/Engine/Source/Editor/MovieSceneTools/Private/TrackEditors/AudioTrackEditor.cpp b/Engine/Source/Editor/MovieSceneTools/Private/TrackEditors/AudioTrackEditor.cpp
--- a/Engine/Source/Editor/MovieSceneTools/Private/TrackEditors/AudioTrackEditor.cpp
+++ b/Engine/Source/Editor/MovieSceneTools/Private/TrackEditors/AudioTrackEditor.cpp
@@ -833,6 +833,17 @@ void FAudioTrackEditor::BuildAddTrackMenu(FMenuBuilder& MenuBuilder)
 	);
 }
 
+void FAudioTrackEditor::BuildObjectBindingTrackMenu(FMenuBuilder& MenuBuilder, const TArray<FGuid>& ObjectBindings, const UClass* ObjectClass)
+{
+	if (ObjectClass != nullptr && ObjectClass->IsChildOf(AActor::StaticClass()))
+	{
+		MenuBuilder.AddSubMenu(
+			LOCTEXT("AddAttachedAudioTrack", "Audio"),
+			LOCTEXT("AddAttachedAudioTooltip", "Adds an audio track attached to the object."),
+			FNewMenuDelegate::CreateSP(this, &FAudioTrackEditor::HandleAddAttachedAudioTrackMenuEntryExecute, ObjectBindings));
+	}
+}
+
 
 bool FAudioTrackEditor::SupportsType( TSubclassOf<UMovieSceneTrack> Type ) const
 {
@@ -983,7 +994,7 @@ TSharedPtr<SWidget> FAudioTrackEditor::BuildOutlinerEditWidget(const FGuid& Obje
 	.AutoWidth()
 	.VAlign(VAlign_Center)
 	[
-		FSequencerUtilities::MakeAddButton(LOCTEXT("AudioText", "Audio"), FOnGetContent::CreateSP(this, &FAudioTrackEditor::BuildAudioSubMenu, Track), Params.NodeIsHovered, GetSequencer())
+		FSequencerUtilities::MakeAddButton(LOCTEXT("AudioText", "Audio"), FOnGetContent::CreateSP(this, &FAudioTrackEditor::BuildAudioSubMenu, FOnAssetSelected::CreateRaw(this, &FAudioTrackEditor::OnAudioAssetSelected, Track), FOnAssetEnterPressed::CreateRaw(this, &FAudioTrackEditor::OnAudioAssetEnterPressed, Track)), Params.NodeIsHovered, GetSequencer())
 	];
 }
 
@@ -1119,7 +1130,13 @@ void FAudioTrackEditor::HandleAddAudioTrackMenuEntryExecute()
 	}
 }
 
-TSharedRef<SWidget> FAudioTrackEditor::BuildAudioSubMenu(UMovieSceneTrack* Track)
+void FAudioTrackEditor::HandleAddAttachedAudioTrackMenuEntryExecute(FMenuBuilder& MenuBuilder, TArray<FGuid> ObjectBindings)
+{
+	MenuBuilder.AddWidget(BuildAudioSubMenu(FOnAssetSelected::CreateRaw(this, &FAudioTrackEditor::OnAttachedAudioAssetSelected, ObjectBindings), FOnAssetEnterPressed::CreateRaw(this, &FAudioTrackEditor::OnAttachedAudioEnterPressed, ObjectBindings)), FText::GetEmpty(), true);
+}
+
+
+TSharedRef<SWidget> FAudioTrackEditor::BuildAudioSubMenu(FOnAssetSelected OnAssetSelected, FOnAssetEnterPressed OnAssetEnterPressed)
 {
 	FAssetRegistryModule& AssetRegistryModule = FModuleManager::LoadModuleChecked<FAssetRegistryModule>(TEXT("AssetRegistry"));
 	TArray<FName> ClassNames;
@@ -1131,8 +1148,8 @@ TSharedRef<SWidget> FAudioTrackEditor::BuildAudioSubMenu(UMovieSceneTrack* Track
 
 	FAssetPickerConfig AssetPickerConfig;
 	{
-		AssetPickerConfig.OnAssetSelected = FOnAssetSelected::CreateRaw( this, &FAudioTrackEditor::OnAudioAssetSelected, Track);
-		AssetPickerConfig.OnAssetEnterPressed = FOnAssetEnterPressed::CreateRaw( this, &FAudioTrackEditor::OnAudioAssetEnterPressed, Track);
+		AssetPickerConfig.OnAssetSelected = OnAssetSelected;
+		AssetPickerConfig.OnAssetEnterPressed = OnAssetEnterPressed;
 		AssetPickerConfig.bAllowNullSelection = false;
 		AssetPickerConfig.InitialAssetViewType = EAssetViewType::List;
 		for (auto ClassName : DerivedClassNames)
@@ -1192,5 +1209,32 @@ void FAudioTrackEditor::OnAudioAssetEnterPressed(const TArray<FAssetData>& Asset
 	}
 }
 
+void FAudioTrackEditor::OnAttachedAudioAssetSelected(const FAssetData& AssetData, TArray<FGuid> ObjectBindings)
+{
+	FSlateApplication::Get().DismissAllMenus();
+
+	UObject* SelectedObject = AssetData.GetAsset();
+
+	if (SelectedObject)
+	{
+		const FScopedTransaction Transaction(NSLOCTEXT("Sequencer", "AddAudio_Transaction", "Add Audio"));
+
+		for (FGuid ObjectBinding : ObjectBindings)
+		{
+			HandleAssetAdded(SelectedObject, ObjectBinding);
+		}
+	}
+}
+
+void FAudioTrackEditor::OnAttachedAudioEnterPressed(const TArray<FAssetData>& AssetData, TArray<FGuid> ObjectBindings)
+{
+	if (AssetData.Num() > 0)
+	{
+		OnAttachedAudioAssetSelected(AssetData[0].GetAsset(), ObjectBindings);
+	}
+}
+
+
+
 
 #undef LOCTEXT_NAMESPACE
diff --git a/Engine/Source/Editor/MovieSceneTools/Public/TrackEditors/AudioTrackEditor.h b/Engine/Source/Editor/MovieSceneTools/Public/TrackEditors/AudioTrackEditor.h
--- a/Engine/Source/Editor/MovieSceneTools/Public/TrackEditors/AudioTrackEditor.h
+++ b/Engine/Source/Editor/MovieSceneTools/Public/TrackEditors/AudioTrackEditor.h
@@ -11,6 +11,7 @@
 #include "ISequencer.h"
 #include "ISequencerTrackEditor.h"
 #include "MovieSceneTrackEditor.h"
+#include "IContentBrowserSingleton.h"
 
 struct FAssetData;
 class FAudioThumbnail;
@@ -49,6 +50,7 @@ public:
 	// ISequencerTrackEditor interface
 
 	virtual void BuildAddTrackMenu(FMenuBuilder& MenuBuilder) override;
+	virtual void BuildObjectBindingTrackMenu(FMenuBuilder& MenuBuilder, const TArray<FGuid>& ObjectBindings, const UClass* ObjectClass) override;
 	virtual TSharedPtr<SWidget> BuildOutlinerEditWidget(const FGuid& ObjectBinding, UMovieSceneTrack* Track, const FBuildEditWidgetParams& Params) override;
 	virtual bool HandleAssetAdded(UObject* Asset, const FGuid& TargetObjectGuid) override;
 	virtual TSharedRef<ISequencerSection> MakeSectionInterface(UMovieSceneSection& SectionObject, UMovieSceneTrack& Track, FGuid ObjectBinding) override;
@@ -74,14 +76,23 @@ private:
 	/** Callback for executing the "Add Audio Track" menu entry. */
 	void HandleAddAudioTrackMenuEntryExecute();
 
+	/** Callback for executing the "Add Audio Track" menu entry on an actor */
+	void HandleAddAttachedAudioTrackMenuEntryExecute(FMenuBuilder& MenuBuilder, TArray<FGuid> ObjectBindings);
+
 	/** Audio sub menu */
-	TSharedRef<SWidget> BuildAudioSubMenu(UMovieSceneTrack* Track);
+	TSharedRef<SWidget> BuildAudioSubMenu(FOnAssetSelected OnAssetSelected, FOnAssetEnterPressed OnAssetEnterPressed);
 
 	/** Audio asset selected */
 	void OnAudioAssetSelected(const FAssetData& AssetData, UMovieSceneTrack* Track);
 
 	/** Audio asset enter pressed */
 	void OnAudioAssetEnterPressed(const TArray<FAssetData>& AssetData, UMovieSceneTrack* Track);
+
+	/** Attached audio asset selected */
+	void OnAttachedAudioAssetSelected(const FAssetData& AssetData, TArray<FGuid> ObjectBindings);
+
+	/** Attached audio asset enter pressed */
+	void OnAttachedAudioEnterPressed(const TArray<FAssetData>& AssetData, TArray<FGuid> ObjectBindings);
 };
 
 
