Sequencer: Dragging audio asset on to object's audio subtrack should add an attached track

diff --git a/Engine/Source/Editor/MovieSceneTools/Private/TrackEditors/AudioTrackEditor.cpp b/Engine/Source/Editor/MovieSceneTools/Private/TrackEditors/AudioTrackEditor.cpp
--- a/Engine/Source/Editor/MovieSceneTools/Private/TrackEditors/AudioTrackEditor.cpp
+++ b/Engine/Source/Editor/MovieSceneTools/Private/TrackEditors/AudioTrackEditor.cpp
@@ -969,7 +969,20 @@ FReply FAudioTrackEditor::OnDrop(const FDragDropEvent& DragDropEvent, UMovieScen
 
 		if (Sound)
 		{
-			AnimatablePropertyChanged( FOnKeyProperty::CreateRaw(this, &FAudioTrackEditor::AddNewMasterSound, Sound, RowIndex));
+			if (TargetObjectGuid.IsValid())
+			{
+				TArray<TWeakObjectPtr<>> OutObjects;
+				for (TWeakObjectPtr<> Object : GetSequencer()->FindObjectsInCurrentSequence(TargetObjectGuid))
+				{
+					OutObjects.Add(Object);
+				}
+
+				AnimatablePropertyChanged(FOnKeyProperty::CreateRaw(this, &FAudioTrackEditor::AddNewAttachedSound, Sound, OutObjects));
+			}
+			else
+			{
+				AnimatablePropertyChanged(FOnKeyProperty::CreateRaw(this, &FAudioTrackEditor::AddNewMasterSound, Sound, RowIndex));
+			}
 
 			bAnyDropped = true;
 		}
