Sequencer: Fix the "Properties" context-menu of an animation section so that the asset picker correctly filters out animations that aren't compatible with the parent bound actor's skeleton.

diff --git a/Engine/Source/Editor/MovieSceneTools/Private/TrackEditors/SkeletalAnimationTrackEditor.cpp b/Engine/Source/Editor/MovieSceneTools/Private/TrackEditors/SkeletalAnimationTrackEditor.cpp
--- a/Engine/Source/Editor/MovieSceneTools/Private/TrackEditors/SkeletalAnimationTrackEditor.cpp
+++ b/Engine/Source/Editor/MovieSceneTools/Private/TrackEditors/SkeletalAnimationTrackEditor.cpp
@@ -47,6 +47,15 @@
 #include "Subsystems/AssetEditorSubsystem.h"
 #include "Editor.h"
 
+#include "IDetailCustomization.h"
+#include "DetailLayoutBuilder.h"
+#include "PropertyEditorModule.h"
+#include "IDetailChildrenBuilder.h"
+#include "DetailWidgetRow.h"
+#include "PropertyCustomizationHelpers.h"
+#include "IPropertyUtilities.h"
+
+
 namespace SkeletalAnimationEditorConstants
 {
 	// @todo Sequencer Allow this to be customizable
@@ -157,6 +166,96 @@ USkeleton* AcquireSkeletonFromObjectGuid(const FGuid& Guid, TSharedPtr<ISequence
 }
 
 
+class FMovieSceneSkeletalAnimationParamsDetailCustomization : public IPropertyTypeCustomization
+{
+public:
+	FMovieSceneSkeletalAnimationParamsDetailCustomization(const FSequencerSectionPropertyDetailsViewCustomizationParams& InParams)
+		: Params(InParams)
+	{
+	}
+
+	// IDetailCustomization interface
+	virtual void CustomizeHeader(TSharedRef<IPropertyHandle> PropertyHandle, FDetailWidgetRow& HeaderRow, IPropertyTypeCustomizationUtils& CustomizationUtils) override
+	{
+	}
+
+	virtual void CustomizeChildren(TSharedRef<IPropertyHandle> PropertyHandle, IDetailChildrenBuilder& ChildBuilder, IPropertyTypeCustomizationUtils& CustomizationUtils) override
+	{
+		const FName AnimationPropertyName = GET_MEMBER_NAME_CHECKED(FMovieSceneSkeletalAnimationParams, Animation);
+
+		uint32 NumChildren;
+		PropertyHandle->GetNumChildren(NumChildren);
+		for (uint32 i = 0; i < NumChildren; ++i)
+		{
+			TSharedPtr<IPropertyHandle> ChildPropertyHandle = PropertyHandle->GetChildHandle(i);
+			IDetailPropertyRow& ChildPropertyRow = ChildBuilder.AddProperty(ChildPropertyHandle.ToSharedRef());
+
+			// Let most properties be whatever they want to be... we just want to customize the `Animation` property
+			// by making it look like a normal asset reference property, but with some custom filtering.
+			if (ChildPropertyHandle->GetProperty()->GetFName() == AnimationPropertyName)
+			{
+				FDetailWidgetRow& Row = ChildPropertyRow.CustomWidget();
+
+				if (Params.ParentObjectBindingGuid.IsValid())
+				{
+					// Store the compatible skeleton's name, and create a property widget with a filter that will check
+					// for animations that match that skeleton.
+					USkeleton* Skeleton = AcquireSkeletonFromObjectGuid(Params.ParentObjectBindingGuid, Params.Sequencer);
+					SkeletonName = FAssetData(Skeleton).GetExportTextName();
+
+					TSharedPtr<IPropertyUtilities> PropertyUtilities = CustomizationUtils.GetPropertyUtilities();
+
+					TSharedRef<SObjectPropertyEntryBox> ContentWidget = SNew(SObjectPropertyEntryBox)
+						.PropertyHandle(ChildPropertyHandle)
+						.DisplayThumbnail(true)
+						.ThumbnailPool(PropertyUtilities.IsValid() ? PropertyUtilities->GetThumbnailPool() : nullptr)
+						.OnShouldFilterAsset(FOnShouldFilterAsset::CreateRaw(this, &FMovieSceneSkeletalAnimationParamsDetailCustomization::ShouldFilterAsset));
+
+					Row.NameContent()[ChildPropertyHandle->CreatePropertyNameWidget()];
+					Row.ValueContent()[ContentWidget];
+
+					float MinDesiredWidth, MaxDesiredWidth;
+					ContentWidget->GetDesiredWidth(MinDesiredWidth, MaxDesiredWidth);
+					Row.ValueContent().MinWidth = MinDesiredWidth;
+					Row.ValueContent().MaxWidth = MaxDesiredWidth;
+
+					// The content widget already contains a "reset to default" button, so we don't want the details view row
+					// to make another one. We add this metadata on the property handle instance to suppress it.
+					ChildPropertyHandle->SetInstanceMetaData(TEXT("NoResetToDefault"), TEXT("true"));
+				}
+			}
+		}
+	}
+
+	bool ShouldFilterAsset(const FAssetData& AssetData)
+	{
+		// Since the `SObjectPropertyEntryBox` doesn't support passing some `Filter` properties for the asset picker, 
+		// we just combine the tag value filtering we want (i.e. checking the skeleton compatibility) along with the
+		// other filtering we already get from the track editor's filter callback.
+		FSkeletalAnimationTrackEditor& TrackEditor = static_cast<FSkeletalAnimationTrackEditor&>(Params.TrackEditor);
+		if (TrackEditor.ShouldFilterAsset(AssetData))
+		{
+			return true;
+		}
+
+		if (!SkeletonName.IsEmpty())
+		{
+			const FString& SkeletonTag = AssetData.GetTagValueRef<FString>(TEXT("Skeleton"));
+			if (SkeletonTag != SkeletonName)
+			{
+				return true;
+			}
+		}
+
+		return false;
+	}
+
+private:
+	const FSequencerSectionPropertyDetailsViewCustomizationParams& Params;
+	FString SkeletonName;
+};
+
+
 FSkeletalAnimationSection::FSkeletalAnimationSection( UMovieSceneSection& InSection, TWeakPtr<ISequencer> InSequencer)
 	: Section(*CastChecked<UMovieSceneSkeletalAnimationSection>(&InSection))
 	, Sequencer(InSequencer)
@@ -392,6 +491,13 @@ void FSkeletalAnimationSection::BuildSectionContextMenu(FMenuBuilder& MenuBuilde
 	MenuBuilder.EndSection();
 }
 
+void FSkeletalAnimationSection::CustomizePropertiesDetailsView(TSharedRef<IDetailsView> DetailsView, const FSequencerSectionPropertyDetailsViewCustomizationParams& InParams) const
+{
+	DetailsView->RegisterInstancedCustomPropertyTypeLayout(
+		TEXT("MovieSceneSkeletalAnimationParams"),
+		FOnGetPropertyTypeCustomizationInstance::CreateLambda([=]() { return MakeShared<FMovieSceneSkeletalAnimationParamsDetailCustomization>(InParams); }));
+}
+
 
 FSkeletalAnimationTrackEditor::FSkeletalAnimationTrackEditor( TSharedRef<ISequencer> InSequencer )
 	: FMovieSceneTrackEditor( InSequencer ) 
diff --git a/Engine/Source/Editor/MovieSceneTools/Private/TrackEditors/SkeletalAnimationTrackEditor.h b/Engine/Source/Editor/MovieSceneTools/Private/TrackEditors/SkeletalAnimationTrackEditor.h
--- a/Engine/Source/Editor/MovieSceneTools/Private/TrackEditors/SkeletalAnimationTrackEditor.h
+++ b/Engine/Source/Editor/MovieSceneTools/Private/TrackEditors/SkeletalAnimationTrackEditor.h
@@ -73,6 +73,8 @@ private:
 
 	/** Delegate for AnimatablePropertyChanged in AddKey */
 	FKeyPropertyResult AddKeyInternal(FFrameNumber KeyTime, UObject* Object, class UAnimSequenceBase* AnimSequence, UMovieSceneTrack* Track, int32 RowIndex);
+
+	friend class FMovieSceneSkeletalAnimationParamsDetailCustomization;
 };
 
 
@@ -103,6 +105,7 @@ public:
 	virtual void BeginSlipSection() override;
 	virtual void SlipSection(FFrameNumber SlipTime) override;
 	virtual void BuildSectionContextMenu(FMenuBuilder& MenuBuilder, const FGuid& InObjectBinding) override;
+	virtual void CustomizePropertiesDetailsView(TSharedRef<IDetailsView> DetailsView, const FSequencerSectionPropertyDetailsViewCustomizationParams& InParams) const override;
 
 private:
 
diff --git a/Engine/Source/Editor/Sequencer/Private/SequencerContextMenus.cpp b/Engine/Source/Editor/Sequencer/Private/SequencerContextMenus.cpp
--- a/Engine/Source/Editor/Sequencer/Private/SequencerContextMenus.cpp
+++ b/Engine/Source/Editor/Sequencer/Private/SequencerContextMenus.cpp
@@ -5,6 +5,7 @@
 #include "EditorStyleSet.h"
 #include "DisplayNodes/SequencerSectionKeyAreaNode.h"
 #include "DisplayNodes/SequencerTrackNode.h"
+#include "DisplayNodes/SequencerObjectBindingNode.h"
 #include "SequencerCommonHelpers.h"
 #include "SSequencer.h"
 #include "SectionLayout.h"
@@ -536,7 +537,7 @@ void FSectionContextMenu::AddPropertiesMenu(FMenuBuilder& MenuBuilder)
 
 	TArray<TWeakObjectPtr<UObject>> Sections;
 	{
-		for (auto Section : Sequencer->GetSelection().GetSelectedSections())
+		for (TWeakObjectPtr<UMovieSceneSection> Section : Sequencer->GetSelection().GetSelectedSections())
 		{
 			if (Section.IsValid())
 			{
@@ -554,6 +555,27 @@ void FSectionContextMenu::AddPropertiesMenu(FMenuBuilder& MenuBuilder)
 	DetailsView->RegisterInstancedCustomPropertyLayout(UMovieSceneSection::StaticClass(), FOnGetDetailCustomizationInstance::CreateLambda([=]() {
 		return MakeShared<FMovieSceneSectionDetailsCustomization>(Sequencer->GetNumericTypeInterface(), CurrentScene); }));
 
+	// Let section interfaces further customize the properties details view.
+	TSharedRef<FSequencerNodeTree> SequencerNodeTree = Sequencer->GetNodeTree();
+	for (TWeakObjectPtr<UObject> Section : Sections)
+	{
+		if (Section.IsValid())
+		{
+			TOptional<FSectionHandle> SectionHandle = SequencerNodeTree->GetSectionHandle(Cast<UMovieSceneSection>(Section));
+			if (SectionHandle)
+			{
+				TSharedRef<ISequencerSection> SectionInterface = SectionHandle->GetSectionInterface();
+				FSequencerSectionPropertyDetailsViewCustomizationParams CustomizationDetails(
+					SectionInterface, Sequencer, SectionHandle->GetTrackNode()->GetTrackEditor());
+				TSharedPtr<FSequencerObjectBindingNode> ParentObjectBindingNode = SectionHandle->GetTrackNode()->FindParentObjectBindingNode();
+				if (ParentObjectBindingNode.IsValid())
+				{
+					CustomizationDetails.ParentObjectBindingGuid = ParentObjectBindingNode->GetObjectBinding();
+				}
+				SectionInterface->CustomizePropertiesDetailsView(DetailsView, CustomizationDetails);
+			}
+		}
+	}
 
 	Sequencer->OnInitializeDetailsPanel().Broadcast(DetailsView, Sequencer);
 	DetailsView->SetObjects(Sections);
diff --git a/Engine/Source/Editor/Sequencer/Public/ISequencerSection.h b/Engine/Source/Editor/Sequencer/Public/ISequencerSection.h
--- a/Engine/Source/Editor/Sequencer/Public/ISequencerSection.h
+++ b/Engine/Source/Editor/Sequencer/Public/ISequencerSection.h
@@ -13,6 +13,10 @@
 
 class FMenuBuilder;
 class FSequencerSectionPainter;
+class IDetailsView;
+class ISequencer;
+class ISequencerSection;
+class ISequencerTrackEditor;
 struct FSlateBrush;
 
 /** Enumerates which edge is being resized */
@@ -40,6 +44,23 @@ namespace SequencerSectionConstants
 	const FName SelectionInactiveColorName("SelectionColorInactive");
 }
 
+/**
+ * Parameters for the callback used when a section wants to customize how its properties details context menu looks like.
+ */
+struct FSequencerSectionPropertyDetailsViewCustomizationParams
+{
+	FSequencerSectionPropertyDetailsViewCustomizationParams(TSharedRef<ISequencerSection> InSectionInterface, TSharedRef<ISequencer> InSequencer, ISequencerTrackEditor& InTrackEditor)
+		: SectionInterface(InSectionInterface)
+		, Sequencer(InSequencer)
+		, TrackEditor(InTrackEditor)
+	{}
+
+	FGuid ParentObjectBindingGuid;
+	TSharedRef<ISequencerSection> SectionInterface;
+	TSharedRef<ISequencer> Sequencer;
+	ISequencerTrackEditor& TrackEditor;
+};
+
 /**
  * Interface that should be implemented for the UI portion of a section
  */
@@ -179,6 +200,14 @@ public:
 	 */
 	virtual void BeginSlipSection() {}
 	virtual void SlipSection(FFrameNumber SlipTime) {}
+
+	/**
+	 * Called when the properties context menu is being built, so this section can customize how the menu's details view looks like.
+	 *
+	 * @param DetailsView The details view widget
+	 * @param InParams Information about the current operation
+	 */
+	virtual void CustomizePropertiesDetailsView(TSharedRef<IDetailsView> DetailsView, const FSequencerSectionPropertyDetailsViewCustomizationParams& InParams) const {}
 };
 
 class SEQUENCER_VTABLE FSequencerSection : public ISequencerSection
