Sequencer: Add options to attach tracks for specifying how to handle the transition to relative coordinates

diff --git a/Engine/Source/Editor/MovieSceneTools/Private/TrackEditors/AttachTrackEditor.cpp b/Engine/Source/Editor/MovieSceneTools/Private/TrackEditors/AttachTrackEditor.cpp
--- a/Engine/Source/Editor/MovieSceneTools/Private/TrackEditors/AttachTrackEditor.cpp
+++ b/Engine/Source/Editor/MovieSceneTools/Private/TrackEditors/AttachTrackEditor.cpp
@@ -14,7 +14,10 @@
 #include "MovieSceneToolHelpers.h"
 #include "Channels/MovieSceneChannelProxy.h"
 #include "Evaluation/MovieSceneEvaluationTemplateInstance.h"
-
+#include "Algo/Transform.h"
+#include "Algo/Copy.h"
+#include "Containers/Union.h"
+#include "Widgets/Input/SCheckBox.h"
 
 #define LOCTEXT_NAMESPACE "F3DAttachTrackEditor"
 
@@ -84,9 +87,27 @@ public:
 		TArray<FGuid> ObjectBindings;
 		ObjectBindings.Add(ObjectBinding);
 
+		MenuBuilder.BeginSection(NAME_None, LOCTEXT("AttachOptions", "Attach Section Options"));
+
 		MenuBuilder.AddSubMenu(
 			LOCTEXT("SetAttach", "Attach"), LOCTEXT("SetAttachTooltip", "Set attach"),
 			FNewMenuDelegate::CreateRaw(AttachTrackEditor, &FActorPickerTrackEditor::ShowActorSubMenu, ObjectBindings, &Section));
+
+		MenuBuilder.AddMenuEntry(
+			LOCTEXT("TrimRightPreserve", "Trim Right and Preserve"),
+			LOCTEXT("TrimRightPreserveToolTip", "Trims the right side of this attach at the current time and preserves the last key's world coordinates"),
+			FSlateIcon(),
+			FUIAction(FExecuteAction::CreateRaw(AttachTrackEditor, &F3DAttachTrackEditor::TrimAndPreserve, ObjectBinding, &Section, false))
+		);
+
+		MenuBuilder.AddMenuEntry(
+			LOCTEXT("TrimLeftPreserve", "Trim Left and Preserve"),
+			LOCTEXT("TrimLeftPreserveToolTip", "Trims the left side of this attach at the current time and preserves the first key's world coordinates"),
+			FSlateIcon(),
+			FUIAction(FExecuteAction::CreateRaw(AttachTrackEditor, &F3DAttachTrackEditor::TrimAndPreserve, ObjectBinding, &Section, true))
+		);
+
+		MenuBuilder.EndSection();
 	}
 
 private:
@@ -98,24 +119,21 @@ private:
 	F3DAttachTrackEditor* AttachTrackEditor;
 };
 
-
 F3DAttachTrackEditor::F3DAttachTrackEditor( TSharedRef<ISequencer> InSequencer )
-	: FActorPickerTrackEditor( InSequencer ) 
+: FActorPickerTrackEditor( InSequencer )
+, PreserveType(ETransformPreserveType::CurrentKey)
 {
 }
 
-
 F3DAttachTrackEditor::~F3DAttachTrackEditor()
 {
 }
 
-
 TSharedRef<ISequencerTrackEditor> F3DAttachTrackEditor::CreateTrackEditor( TSharedRef<ISequencer> InSequencer )
 {
 	return MakeShareable( new F3DAttachTrackEditor( InSequencer ) );
 }
 
-
 bool F3DAttachTrackEditor::SupportsType( TSubclassOf<UMovieSceneTrack> Type ) const
 {
 	// We support animatable transforms
@@ -139,10 +157,80 @@ void F3DAttachTrackEditor::BuildObjectBindingTrackMenu(FMenuBuilder& MenuBuilder
 
 		MenuBuilder.AddSubMenu(
 			LOCTEXT("AddAttach", "Attach"), LOCTEXT("AddAttachTooltip", "Adds an attach track."),
-			FNewMenuDelegate::CreateRaw(this, &FActorPickerTrackEditor::ShowActorSubMenu, ObjectBindings, DummySection));
+			FNewMenuDelegate::CreateRaw(this, &F3DAttachTrackEditor::ShowPickerSubMenu, ObjectBindings, DummySection));
 	}
 }
 
+void F3DAttachTrackEditor::ShowPickerSubMenu(FMenuBuilder& MenuBuilder, TArray<FGuid> ObjectBindings, UMovieSceneSection* Section)
+{
+	ShowActorSubMenu(MenuBuilder, ObjectBindings, Section);
+
+	FText PreserveText = LOCTEXT("ExistingBinding", "Existing Binding");
+
+	MenuBuilder.BeginSection(NAME_None, LOCTEXT("AttachOptions", "Attach Options"));
+
+	MenuBuilder.AddMenuEntry(
+		LOCTEXT("TogglePreserveCurrentTransform", "Preserve Current"),
+		LOCTEXT("TogglePreserveCurrentTransformTooltip", "Preserve this object's transform in world space for first frame of attach"),
+		FSlateIcon(),
+		FUIAction(
+			FExecuteAction::CreateLambda([this]() { 
+				PreserveType = ETransformPreserveType::CurrentKey; //(PreserveType != ETransformPreserveType::CurrentKey) ? ETransformPreserveType::CurrentKey : ETransformPreserveType::None;
+			}), 
+			FCanExecuteAction::CreateLambda([]() { return true; }),
+			FIsActionChecked::CreateLambda([this]() { return PreserveType == ETransformPreserveType::CurrentKey; })
+		),
+		NAME_None,
+		EUserInterfaceActionType::RadioButton
+	);
+
+	MenuBuilder.AddMenuEntry(
+		LOCTEXT("TogglePreserveAllTransform", "Preserve All"),
+		LOCTEXT("TogglePreserveAllTransformTooltip", "Preserve this object's transform in world space for every child and parent key in attach range"),
+		FSlateIcon(),
+		FUIAction(
+			FExecuteAction::CreateLambda([this]() { 
+				PreserveType = ETransformPreserveType::AllKeys; //(PreserveType != ETransformPreserveType::AllKeys) ? ETransformPreserveType::AllKeys : ETransformPreserveType::None; 
+			}),
+			FCanExecuteAction::CreateLambda([this]() { return true; }),
+			FIsActionChecked::CreateLambda([this]() { return PreserveType == ETransformPreserveType::AllKeys; })
+		),
+		NAME_None,
+		EUserInterfaceActionType::RadioButton
+	);
+
+	MenuBuilder.AddMenuEntry(
+		LOCTEXT("TogglePreserveBake", "Preserve with Bake"),
+		LOCTEXT("TogglePreserveBakeTooltip", "Object's relative transform will be calculated every frame to preserve original world space transform"),
+		FSlateIcon(),
+		FUIAction(
+			FExecuteAction::CreateLambda([this]() { 
+				PreserveType = ETransformPreserveType::Bake; //(PreserveType != ETransformPreserveType::Bake) ? ETransformPreserveType::Bake : ETransformPreserveType::None; 
+			}),
+			FCanExecuteAction::CreateLambda([this]() { return true; }),
+			FIsActionChecked::CreateLambda([this]() { return PreserveType == ETransformPreserveType::Bake; })
+		),
+		NAME_None,
+		EUserInterfaceActionType::RadioButton
+	);
+
+	MenuBuilder.AddMenuEntry(
+		LOCTEXT("TogglePreserveNone", "None"),
+		LOCTEXT("TogglePreserveNoneTooltip", "Object's transform will not be compensated"),
+		FSlateIcon(),
+		FUIAction(
+			FExecuteAction::CreateLambda([this]() {
+				PreserveType = ETransformPreserveType::None; //(PreserveType != ETransformPreserveType::Bake) ? ETransformPreserveType::Bake : ETransformPreserveType::None; 
+			}),
+			FCanExecuteAction::CreateLambda([this]() { return true; }),
+			FIsActionChecked::CreateLambda([this]() { return PreserveType == ETransformPreserveType::None; })
+		),
+		NAME_None,
+		EUserInterfaceActionType::RadioButton
+		);
+
+	MenuBuilder.EndSection();
+}
 
 bool F3DAttachTrackEditor::IsActorPickable(const AActor* const ParentActor, FGuid ObjectBinding, UMovieSceneSection* InSection)
 {
@@ -229,207 +317,986 @@ void F3DAttachTrackEditor::ActorSocketPicked(const FName SocketName, USceneCompo
 	}
 }
 
-FKeyPropertyResult F3DAttachTrackEditor::AddKeyInternal( FFrameNumber KeyTime, const TArray<TWeakObjectPtr<UObject>> Objects, const FName SocketName, const FName ComponentName, FActorPickerID ActorPickerID)
+void F3DAttachTrackEditor::FindOrCreateTransformTrack(const TRange<FFrameNumber>& InAttachRange, UMovieScene* InMovieScene, const FGuid& InObjectHandle, UMovieScene3DTransformTrack*& OutTransformTrack, UMovieScene3DTransformSection*& OutTransformSection, FMovieSceneEvaluationTrack*& OutEvalTrack)
 {
-	FKeyPropertyResult KeyPropertyResult;
+	OutTransformTrack = nullptr;
+	OutTransformSection = nullptr;
+	OutEvalTrack = nullptr;
 
-	FMovieSceneObjectBindingID ConstraintBindingID;
+	FName TransformPropertyName("Transform");
 
-	if (ActorPickerID.ExistingBindingID.IsValid())
+	// Create a transform track if it doesn't exist
+	UMovieScene3DTransformTrack* TransformTrack = Cast<UMovieScene3DTransformTrack>(InMovieScene->FindTrack<UMovieScene3DTransformTrack>(InObjectHandle));
+	if (!TransformTrack)
 	{
-		ConstraintBindingID = ActorPickerID.ExistingBindingID;
+		InMovieScene->Modify();
+		FFindOrCreateTrackResult TransformTrackResult = FindOrCreateTrackForObject(InObjectHandle, UMovieScene3DTransformTrack::StaticClass());
+		TransformTrack = Cast<UMovieScene3DTransformTrack>(TransformTrackResult.Track);
+
+		if (TransformTrack)
+		{
+			OutEvalTrack = MovieSceneToolHelpers::GetEvaluationTrack(GetSequencer().Get(), TransformTrack->GetSignature());
+			TransformTrack->SetPropertyNameAndPath(TransformPropertyName, TransformPropertyName.ToString());
+		}
 	}
-	else if (ActorPickerID.ActorPicked.IsValid())
+	else
 	{
-		FFindOrCreateHandleResult HandleResult = FindOrCreateHandleToObject(ActorPickerID.ActorPicked.Get());
-		FGuid ParentActorId = HandleResult.Handle;
-		KeyPropertyResult.bHandleCreated |= HandleResult.bWasCreated;
-		ConstraintBindingID = FMovieSceneObjectBindingID(ParentActorId, MovieSceneSequenceID::Root, EMovieSceneObjectBindingSpace::Local);
+		OutEvalTrack = MovieSceneToolHelpers::GetEvaluationTrack(GetSequencer().Get(), TransformTrack->GetSignature());
 	}
 
-	if (!ConstraintBindingID.IsValid())
+	if (!TransformTrack)
 	{
-		return KeyPropertyResult;
+		return;
 	}
 
-	UMovieScene* MovieScene = GetSequencer()->GetFocusedMovieSceneSequence()->GetMovieScene();
+	// Create a transform section if it doesn't exist
+	UMovieScene3DTransformSection* TransformSection = nullptr;
+	if (TransformTrack->IsEmpty())
+	{
+		TransformTrack->Modify();
+		TransformSection = Cast<UMovieScene3DTransformSection>(TransformTrack->CreateNewSection());
+		if (TransformSection)
+		{
+			TransformSection->SetRange(TRange<FFrameNumber>::All());
+
+			TransformTrack->AddSection(*TransformSection);
+		}
+	}
+	// Reuse the transform section if it overlaps and check if there are no keys
+	else if (TransformTrack->GetAllSections().Num() == 1)
+	{
+		TRange<FFrameNumber> TransformRange = TransformTrack->GetAllSections()[0]->GetRange();
+		if (TRange<FFrameNumber>::Intersection(InAttachRange, TransformRange).IsEmpty())
+		{
+			return;
+		}
+
+		TransformSection = Cast<UMovieScene3DTransformSection>(TransformTrack->GetAllSections()[0]);
+	}
+
+	OutTransformTrack = TransformTrack;
+	OutTransformSection = TransformSection;
+}
+
+/**
+ * Helper method to safely return an array allocated to store the proper the number of float channels if not already allocated
+ */
+TArray<FMovieSceneFloatValue>& ResizeAndAddKey(const FFrameNumber& InKey, int32 InNum, TMap<FFrameNumber, TArray<FMovieSceneFloatValue>>& OutTransformMap, TSet<FFrameNumber>* OutTimesAdded)
+{
+	TArray<FMovieSceneFloatValue>& Transform = OutTransformMap.FindOrAdd(InKey);
+	if (Transform.Num() == 0)
+	{
+		Transform.SetNum(InNum);
+		if (OutTimesAdded)
+		{
+			OutTimesAdded->Add(InKey);
+		}
+	}
+	return Transform;
+}
 
-	FMovieSceneSequenceID SequenceID = GetSequencer()->GetFocusedTemplateID();
-	if (ConstraintBindingID.GetSequenceID().IsValid())
+/**
+ * Helper method which adds keys from a list of float channels to a map mapping the time to a full transform
+ */
+void AddKeysFromChannels(TArrayView<FMovieSceneFloatChannel*> InChannels, const TRange<FFrameNumber>& InAttachRange, TMap<FFrameNumber, TArray<FMovieSceneFloatValue>>& OutTransformMap, TSet<FFrameNumber>& OutTimesAdded)
+{
+	const int32 NumChannels = 9;
+	for (int32 ChannelIndex = 0; ChannelIndex < NumChannels; ChannelIndex++)
+	{
+		TArray<FFrameNumber> TimesInRange;
+		InChannels[ChannelIndex]->GetKeys(InAttachRange, &TimesInRange, nullptr);
+		if (TimesInRange.Num() == 0)
+		{
+			continue;
+		}
+
+		const int32 BeginRangeIndex = InChannels[ChannelIndex]->GetTimes().FindLastByPredicate(
+			[FirstKey = TimesInRange[0]](const FFrameNumber& FrameNum) { return FrameNum.Value == FirstKey.Value; });
+		if (BeginRangeIndex == INDEX_NONE)
+		{
+			continue;
+		}
+
+		const int32 NumValsInRange = TimesInRange.Num();
+		TArrayView<const FMovieSceneFloatValue> ValuesInRange = InChannels[ChannelIndex]->GetValues().Slice(BeginRangeIndex, NumValsInRange);
+		for (int32 KeyIndex = 0; KeyIndex < ValuesInRange.Num(); KeyIndex++)
+		{
+			TArray<FMovieSceneFloatValue>& Transform = ResizeAndAddKey(TimesInRange[KeyIndex], InChannels.Num(), OutTransformMap, &OutTimesAdded);
+			Transform[ChannelIndex] = ValuesInRange[KeyIndex];
+		}
+	}
+}
+
+/**
+ * Helper method which updates the values in each channel in a list of movie scene float values given a 
+ * transform, preserving the interpolation style and other attributes
+ */
+void UpdateFloatValueTransform(const FTransform& InTransform, TArrayView<FMovieSceneFloatValue> OutFloatValueTransform)
+{
+	OutFloatValueTransform[0].Value = InTransform.GetTranslation().X;
+	OutFloatValueTransform[1].Value = InTransform.GetTranslation().Y;
+	OutFloatValueTransform[2].Value = InTransform.GetTranslation().Z;
+
+	OutFloatValueTransform[3].Value = InTransform.GetRotation().Euler().X;
+	OutFloatValueTransform[4].Value = InTransform.GetRotation().Euler().Y;
+	OutFloatValueTransform[5].Value = InTransform.GetRotation().Euler().Z;
+
+	OutFloatValueTransform[6].Value = InTransform.GetScale3D().X;
+	OutFloatValueTransform[7].Value = InTransform.GetScale3D().Y;
+	OutFloatValueTransform[8].Value = InTransform.GetScale3D().Z;
+}
+
+/**
+ * Helper method which converts a list of float values to a transform
+ */
+FORCEINLINE FTransform FloatValuesToTransform(TArrayView<const FMovieSceneFloatValue> InFloatValues)
+{
+	return FTransform(FRotator::MakeFromEuler(FVector(InFloatValues[3].Value, InFloatValues[4].Value, InFloatValues[5].Value)), 
+		FVector(InFloatValues[0].Value, InFloatValues[1].Value, InFloatValues[2].Value), FVector(InFloatValues[6].Value, InFloatValues[7].Value, InFloatValues[8].Value));
+}
+
+/**
+ * Evaluates the transform of an object at a certain point in time
+ */
+FTransform GetLocationAtTime(TSharedPtr<ISequencer> Sequencer, FMovieSceneEvaluationTrack* EvalTrack, FFrameNumber KeyTime, UObject* Object)
+{
+	FMovieSceneInterrogationData InterrogationData;
+	Sequencer->GetEvaluationTemplate().CopyActuators(InterrogationData.GetAccumulator());
+	FMovieSceneContext Context(FMovieSceneEvaluationRange(KeyTime, Sequencer->GetFocusedTickResolution()));
+	EvalTrack->Interrogate(Context, InterrogationData, Object);
+
+	for (const FTransformData& Transform : InterrogationData.Iterate<FTransformData>(UMovieScene3DTransformSection::GetInterrogationKey()))
+	{
+		return FTransform(Transform.Rotation, Transform.Translation, Transform.Scale);
+	}
+
+	return FTransform::Identity;
+}
+
+AActor* GetConstraintActor(TSharedPtr<ISequencer> InSequencer, const FMovieSceneObjectBindingID& InConstraintBindingID)
+{
+	FMovieSceneSequenceID SequenceID = InSequencer->GetFocusedTemplateID();
+	if (InConstraintBindingID.GetSequenceID().IsValid())
 	{
 		// Ensure that this ID is resolvable from the root, based on the current local sequence ID
-		FMovieSceneObjectBindingID RootBindingID = ConstraintBindingID.ResolveLocalToRoot(SequenceID, GetSequencer()->GetEvaluationTemplate().GetHierarchy());
+		FMovieSceneObjectBindingID RootBindingID = InConstraintBindingID.ResolveLocalToRoot(SequenceID, InSequencer->GetEvaluationTemplate().GetHierarchy());
 		SequenceID = RootBindingID.GetSequenceID();
 	}
 
-	FTransform ParentTransform;
-	TArrayView<TWeakObjectPtr<UObject>> RuntimeObjects = GetSequencer()->FindBoundObjects(ConstraintBindingID.GetGuid(), SequenceID);
+	TArrayView<TWeakObjectPtr<UObject>> RuntimeObjects = InSequencer->FindBoundObjects(InConstraintBindingID.GetGuid(), SequenceID);
 
 	if (RuntimeObjects.Num() == 1 && RuntimeObjects[0].IsValid())
 	{
-		AActor* ParentActor = Cast<AActor>(RuntimeObjects[0].Get());
-		if (ParentActor)
+		return Cast<AActor>(RuntimeObjects[0].Get());
+	}
+
+	return nullptr;
+}
+
+struct ITransformEvaluator
+{
+	virtual FTransform operator()(const FFrameNumber& InTime) const { return FTransform::Identity; };
+	virtual ~ITransformEvaluator() {}
+};
+
+/**
+ * Helper functor for evaluating the local transform for an object.
+ * It can be animated by sequencer but does not have to be
+ */
+struct FLocalTransformEvaluator : ITransformEvaluator
+{
+	FLocalTransformEvaluator() = default;
+
+	/**
+	 * Creates an evaluator for an object. Uses the evaluation track if it exists, otherwise uses the actor's transform
+	 */
+	FLocalTransformEvaluator(TWeakPtr<ISequencer> InWeakSequencer, UObject* InObject)
+		: WeakSequencer(InWeakSequencer)
+	{
+		TSharedPtr<ISequencer> Sequencer = InWeakSequencer.Pin();
+		if (!Sequencer)
 		{
-			ParentTransform = ParentActor->GetActorTransform();
-			
-			if (ParentActor->GetRootComponent()->DoesSocketExist(SocketName))
+			return;
+		}
+
+		UMovieScene* MovieScene = Sequencer->GetFocusedMovieSceneSequence()->GetMovieScene();
+		AActor* Actor = Cast<AActor>(InObject);
+
+		FTransform ActorTransform = Actor->GetActorTransform();
+		TransformEval.SetSubtype<FTransform>(ActorTransform);
+
+		FGuid ActorHandle = Sequencer->GetHandleToObject(Actor, false);
+		if (ActorHandle.IsValid())
+		{
+			UMovieScene3DTransformTrack* ActorTransformTrack = Cast<UMovieScene3DTransformTrack>(MovieScene->FindTrack<UMovieScene3DTransformTrack>(ActorHandle));
+			if (ActorTransformTrack)
 			{
-				ParentTransform = ParentActor->GetRootComponent()->GetSocketTransform(SocketName);
+				FMovieSceneEvaluationTrack* EvalTrack = MovieSceneToolHelpers::GetEvaluationTrack(Sequencer.Get(), ActorTransformTrack->GetSignature());
+				check(EvalTrack)
+				TransformEval.SetSubtype<TTuple<FMovieSceneEvaluationTrack*, UObject*>>(TTuple<FMovieSceneEvaluationTrack*, UObject*>(EvalTrack, Actor));
 			}
 		}
 	}
-		
-	for (int32 ObjectIndex = 0; ObjectIndex < Objects.Num(); ++ObjectIndex)
+
+	/**
+	 * Creates an evaluator for an object with an already existing evaluation track
+	 */
+	FLocalTransformEvaluator(TWeakPtr<ISequencer> InWeakSequencer, UObject* InObject, FMovieSceneEvaluationTrack* InEvalTrack)
+		: WeakSequencer(InWeakSequencer)
 	{
-		UObject* Object = Objects[ObjectIndex].Get();
+		TransformEval.SetSubtype<TTuple<FMovieSceneEvaluationTrack*, UObject*>>(TTuple<FMovieSceneEvaluationTrack*, UObject*>(InEvalTrack, InObject));
+	}
 
-		FFindOrCreateHandleResult HandleResult = FindOrCreateHandleToObject( Object );
-		FGuid ObjectHandle = HandleResult.Handle;
-		KeyPropertyResult.bHandleCreated |= HandleResult.bWasCreated;
-		if (ObjectHandle.IsValid())
+	/**
+	 * Evaluates the transform for this object at the given time
+	 */
+	FTransform operator()(const FFrameNumber& InTime) const override
+	{
+		TSharedPtr<ISequencer> Sequencer = WeakSequencer.Pin();
+		if (Sequencer)
 		{
-			FFindOrCreateTrackResult TrackResult = FindOrCreateTrackForObject( ObjectHandle, UMovieScene3DAttachTrack::StaticClass());
-			UMovieSceneTrack* Track = TrackResult.Track;
-			KeyPropertyResult.bTrackCreated |= TrackResult.bWasCreated;
-
-			if (ensure(Track))
+			const bool bEvalParentTransform = TransformEval.GetCurrentSubtypeIndex() == 1;
+			if (bEvalParentTransform)
+			{
+				FMovieSceneEvaluationTrack* EvalTrack = TransformEval.GetSubtype<TTuple<FMovieSceneEvaluationTrack*, UObject*>>().Get<0>();
+				UObject* Object = TransformEval.GetSubtype<TTuple<FMovieSceneEvaluationTrack*, UObject*>>().Get<1>();
+				return GetLocationAtTime(WeakSequencer.Pin(), EvalTrack, InTime, Object);
+			}
+			else
 			{
-				FFrameRate TickResolution = Track->GetTypedOuter<UMovieScene>()->GetTickResolution();
+				return TransformEval.GetSubtype<FTransform>();
+			}
+		}
+
+		return FTransform::Identity;
+	}
+
+private:
+	TUnion<FTransform, TTuple<FMovieSceneEvaluationTrack*, UObject*>> TransformEval;
+	TWeakPtr<ISequencer> WeakSequencer;
+};
+
+/**
+ * Helper functor for finding the world transform of actors
+ * World transform evaluator gets the world transform of an object during an animation by accumulating the transforms of its parents.
+ * The parents can be animated by sequencer but do not have to be
+ */
+struct FWorldTransformEvaluator : ITransformEvaluator
+{
+	FWorldTransformEvaluator() = default;
+
+	/**
+	 * Creates a new evaluator for a given object
+	 * @param InSocketName is the socket to evaluate for if this is a skeletal mesh
+	 */
+	FWorldTransformEvaluator(TWeakPtr<ISequencer> InWeakSequencer, UObject* InObject, const FName InSocketName = NAME_None)
+		: WeakSequencer(InWeakSequencer)
+	{
+		TSharedPtr<ISequencer> Sequencer = InWeakSequencer.Pin();
+		if (!Sequencer)
+		{
+			return;
+		}
+
+		UMovieScene* MovieScene = Sequencer->GetFocusedMovieSceneSequence()->GetMovieScene();
 
-				// Clamp to next attach section's start time or the end of the current movie scene range
-				FFrameNumber AttachEndTime = MovieScene->GetPlaybackRange().GetUpperBoundValue();
+		FName SocketName = InSocketName;
+		AActor* Actor = Cast<AActor>(InObject);
+		if (!Actor)
+		{
+			return;
+		}
+
+		// Loop through all parents to get an accumulated array of evaluators
+		do
+		{
+			TUnion<FTransform, TTuple<FMovieSceneEvaluationTrack*, UObject*>> ActorEval;
+			// If we find a socket, get the world transform of the socket and break out immediately
+			if (Actor->GetRootComponent()->DoesSocketExist(SocketName))
+			{
+				const FTransform SocketWorldSpace = Actor->GetRootComponent()->GetSocketTransform(SocketName);
+				ActorEval.SetSubtype<FTransform>(SocketWorldSpace);
+				TransformEvals.Add(ActorEval);
+				return;
+			}
+			
+			FTransform ActorTransform = Actor->GetActorTransform();
+			ActorEval.SetSubtype<FTransform>(ActorTransform);
 
-				for (UMovieSceneSection* Section : Track->GetAllSections())
+			FGuid ActorHandle = Sequencer->GetHandleToObject(Actor, false);
+			if (ActorHandle.IsValid())
+			{
+				UMovieScene3DTransformTrack* ActorTransformTrack = MovieScene->FindTrack<UMovieScene3DTransformTrack>(ActorHandle);
+				if (ActorTransformTrack)
 				{
-					FFrameNumber StartTime = Section->HasStartFrame() ? Section->GetInclusiveStartFrame() : 0;
-					if (KeyTime < StartTime)
-					{
-						if (AttachEndTime > StartTime)
-						{
-							AttachEndTime = StartTime;
-						}
-					}
+					FMovieSceneEvaluationTrack* EvalTrack = MovieSceneToolHelpers::GetEvaluationTrack(Sequencer.Get(), ActorTransformTrack->GetSignature());
+					check(EvalTrack)
+					ActorEval.SetSubtype<TTuple<FMovieSceneEvaluationTrack*, UObject*>>(TTuple<FMovieSceneEvaluationTrack*, UObject*>(EvalTrack, Actor));
 				}
+			}
 
-				AActor* Child = Cast<AActor>(Object);
-				FTransform RelativeTransform = Child->GetRootComponent()->GetComponentTransform().GetRelativeTransform(ParentTransform);
+			TransformEvals.Add(ActorEval);
 
-				int32 Duration = FMath::Max(0, (AttachEndTime - KeyTime).Value);
-				Cast<UMovieScene3DAttachTrack>(Track)->AddConstraint( KeyTime, Duration, SocketName, ComponentName, ConstraintBindingID);
-				KeyPropertyResult.bTrackModified = true;
+			Actor = Actor->GetAttachParentActor();
+			if (Actor)
+			{
+				SocketName = Actor->GetAttachParentSocketName();
+			}
+		} 
+		while (Actor);
+	}
 
-				// Compensate
-				FName TransformPropertyName("Transform");
-				TRange<FFrameNumber> AttachRange(KeyTime, AttachEndTime);
+	/**
+	 * Copies the array of all individual actor evaluators to create a new evaluator
+	 */
+	FWorldTransformEvaluator(TWeakPtr<ISequencer> InWeakSequencer, TArrayView<const TUnion<FTransform, TTuple<FMovieSceneEvaluationTrack*, UObject*>>> InTransformEvals)
+		: WeakSequencer(InWeakSequencer)
+	{
+		Algo::Copy(InTransformEvals, TransformEvals);
+	}
 
-				MovieScene->Modify();
+	/**
+	 * Adds an evaluation track for the child of the first transform evaluator
+	 */
+	void PrependTransformEval(UObject* InObject, FMovieSceneEvaluationTrack* InEvalTrack)
+	{
+		TUnion<FTransform, TTuple<FMovieSceneEvaluationTrack*, UObject*>> ActorEval;
+		ActorEval.SetSubtype<TTuple<FMovieSceneEvaluationTrack*, UObject*>>(TTuple<FMovieSceneEvaluationTrack*, UObject*>(InEvalTrack, InObject));
+		TransformEvals.Insert(ActorEval, 0);
+	}
 
-				// Create a transform track if it doesn't exist
-				UMovieScene3DTransformTrack* TransformTrack = Cast<UMovieScene3DTransformTrack>(MovieScene->FindTrack<UMovieScene3DTransformTrack>(ObjectHandle));
-				if (!TransformTrack)
+	/**
+	 *  Adds a transform for the child of the first transform evaluator
+	 */
+	void PrependTransformEval(const FTransform& InTransform)
+	{
+		TUnion<FTransform, TTuple<FMovieSceneEvaluationTrack*, UObject*>> ActorEval;
+		ActorEval.SetSubtype<FTransform>(InTransform);
+		TransformEvals.Insert(ActorEval, 0);
+	}
+
+	/**
+	 * Evaluates the world transform for this object at a certain time
+	 */
+	FTransform operator()(const FFrameNumber& InTime) const override
+	{
+		FTransform Accumulated = FTransform::Identity;
+		TSharedPtr<ISequencer> Sequencer = WeakSequencer.Pin();
+		if (Sequencer)
+		{
+			for (TUnion<FTransform, TTuple<FMovieSceneEvaluationTrack*, UObject*>> TransformEval : TransformEvals)
+			{
+				FTransform ActorTransform;
+				const bool bEvalParentTransform = TransformEval.GetCurrentSubtypeIndex() == 1;
+				if (bEvalParentTransform)
 				{
-					FFindOrCreateTrackResult TransformTrackResult = FindOrCreateTrackForObject(ObjectHandle, UMovieScene3DTransformTrack::StaticClass());
-					TransformTrack = Cast<UMovieScene3DTransformTrack>(TransformTrackResult.Track);
-					TransformTrack->SetPropertyNameAndPath(TransformPropertyName, TransformPropertyName.ToString());
+					FMovieSceneEvaluationTrack* EvalTrack = TransformEval.GetSubtype<TTuple<FMovieSceneEvaluationTrack*, UObject*>>().Get<0>();
+					UObject* Object = TransformEval.GetSubtype<TTuple<FMovieSceneEvaluationTrack*, UObject*>>().Get<1>();
+					ActorTransform = GetLocationAtTime(Sequencer, EvalTrack, InTime, Object);
 				}
-
-				if (!TransformTrack)
+				else
 				{
-					continue;
+					ActorTransform = TransformEval.GetSubtype<FTransform>();
 				}
 
-				// Create a transform section if it doesn't exist
-				UMovieScene3DTransformSection* TransformSection = nullptr;
-				if (TransformTrack->IsEmpty())
-				{
-					TransformTrack->Modify();
-					TransformSection = Cast<UMovieScene3DTransformSection>(TransformTrack->CreateNewSection());
-					if (TransformSection)
-					{
-						TransformSection->SetRange(TRange<FFrameNumber>::All());
+				Accumulated *= ActorTransform;
+			}
+		}
 
-						TransformTrack->AddSection(*TransformSection);
-					}
-				}
+		return Accumulated;
+	}
 
-				// Reuse the transform section if it overlaps and there's no keys
-				else if (TransformTrack->GetAllSections().Num() == 1)
-				{
-					TRange<FFrameNumber> TransformRange = TransformTrack->GetAllSections()[0]->GetRange();
-					if (TRange<FFrameNumber>::Intersection(AttachRange, TransformRange).IsEmpty())
-					{
-						continue;
-					}
+	/**
+	 * Gets the individual actor evaluators for each parent
+	 */
+	TArrayView <const TUnion<FTransform, TTuple<FMovieSceneEvaluationTrack*, UObject*>>> GetTransformEvalsView() const 
+	{
+		return TransformEvals;
+	}
 
-					bool bEmptyKeys = true;
-					TArrayView<FMovieSceneFloatChannel*> Channels = Cast<UMovieScene3DTransformSection>(TransformTrack->GetAllSections()[0])->GetChannelProxy().GetChannels<FMovieSceneFloatChannel>();
-					for (auto Channel : Channels)
-					{
-						if (Channel->GetTimes().Num() != 0)
-						{
-							bEmptyKeys = false;
-							break;
-						}
-					}
+private:
+	TArray<TUnion<FTransform, TTuple<FMovieSceneEvaluationTrack*, UObject*>>> TransformEvals;
+	TWeakPtr<ISequencer> WeakSequencer;
+};
 
-					if (bEmptyKeys)
-					{
-						TransformSection = Cast<UMovieScene3DTransformSection>(TransformTrack->GetAllSections()[0]);
-					}
-				}
+/**
+ * Helper functor to revert transforms that are in the relative space of a constraint
+ */
+struct FAttachRevertModifier
+{
+	/**
+	 * Constructor finds the constraint for the given attach section and finds the evaluation track/transform for it.
+	 * @param bInFullRevert If true: Does a full revert with a simple compensation for the first frame,
+	 *                               modifying the object's movements to how they were before the attach.
+	 *                      If false: Parent's movement is kept and transforms are simply converted to world space.
+	 */
+	FAttachRevertModifier(TSharedPtr<ISequencer> InSequencer, const TRange<FFrameNumber>& InRevertRange, UMovieScene3DAttachSection* InAttachSection, const FName InSocketName, bool bInFullRevert)
+		: bFullRevert(bInFullRevert)
+		, RevertRange(InRevertRange)
+	{
+		FMovieSceneObjectBindingID ConstraintID = InAttachSection->GetConstraintBindingID();
+		AActor* ConstraintActor = GetConstraintActor(InSequencer, ConstraintID);
 
-				// Create a new additive transform section
-				if (!TransformSection)
-				{
-					TransformTrack->Modify();
-					TransformSection = Cast<UMovieScene3DTransformSection>(TransformTrack->CreateNewSection());
-					if (TransformSection)
-					{
-						for (auto Section : TransformTrack->GetAllSections())
-						{
-							Section->SetRowIndex(Section->GetRowIndex() + 1);
-						}
+		TransformEvaluator = FWorldTransformEvaluator(InSequencer, ConstraintActor, InSocketName);
 
-						TransformSection->SetRange(AttachRange);
-						TransformSection->SetBlendType(EMovieSceneBlendType::Additive);
-						TransformSection->SetRowIndex(0);
-						TransformSection->SetMask(FMovieSceneTransformMask(EMovieSceneTransformChannel::Rotation | EMovieSceneTransformChannel::Translation));
-						TransformTrack->AddSection(*TransformSection);
+		BeginConstraintTransform = TransformEvaluator(InRevertRange.GetLowerBoundValue());
 
-						RelativeTransform = ParentTransform.Inverse();
-					}
+	}
+
+	/**
+	 * Creates a new revert modifier with a given evaluator for a parent transform to undo compensation
+	 */
+	FAttachRevertModifier(TSharedPtr<ISequencer> InSequencer, const TRange<FFrameNumber>& InRevertRange, const FWorldTransformEvaluator& InTransformEvaluator, bool bInFullRevert)
+		: bFullRevert(bInFullRevert)
+		, TransformEvaluator(InTransformEvaluator)
+		, BeginConstraintTransform(InTransformEvaluator(InRevertRange.GetLowerBoundValue()))
+		, RevertRange(InRevertRange)
+	{}
+
+	/** Reverts a transform in relative space to world space */
+	FTransform operator()(const FTransform& InTransform, const FFrameNumber& InTime)
+	{
+		FTransform OutTransform = InTransform;
+
+		FTransform ConstraintTransform = TransformEvaluator(InTime);
+
+		// If in revert range, revert the transform to world coordinates first
+		if (RevertRange.Contains(InTime))
+		{
+			OutTransform = OutTransform * ConstraintTransform;
+		}
+
+		if (bFullRevert)
+		{
+			const FTransform ConstraintChange = BeginConstraintTransform.GetRelativeTransform(ConstraintTransform);
+			OutTransform = OutTransform * ConstraintChange;
+		}
+
+		return OutTransform;
+	}
+
+private:
+	bool bFullRevert;
+
+	FWorldTransformEvaluator TransformEvaluator;
+	FTransform BeginConstraintTransform;
+	TRange<FFrameNumber> RevertRange;
+};
+
+/**
+ * Updates an array of float channels with the keys in a given transform map mapping times to float values
+ */
+void UpdateChannelTransforms(const TRange<FFrameNumber>& InAttachRange, TMap<FFrameNumber, TArray<FMovieSceneFloatValue>>& InTransformMap, TArrayView<FMovieSceneFloatChannel*>& InChannels, int32 InNumChannels, bool bInBakedData)
+{
+	// Remove all handles in range so we can add the new ones
+	for (FMovieSceneFloatChannel* Channel : InChannels)
+	{
+		TArray<FKeyHandle> KeysToRemove;
+		Channel->GetKeys(InAttachRange, nullptr, &KeysToRemove);
+		Channel->DeleteKeys(KeysToRemove);
+	}
+
+	// Find max extent of all channels
+	TRange<FFrameNumber> TotalRange = TRange<FFrameNumber>(TNumericLimits<FFrameNumber>::Lowest(), TNumericLimits<FFrameNumber>::Max());
+	TArray<TRange<FFrameNumber>> ExcludedRanges = TRange<FFrameNumber>::Difference(TotalRange, InAttachRange);
+
+	InTransformMap.KeySort([](const FFrameNumber& LHS, const FFrameNumber& RHS) { return LHS.Value < RHS.Value; });
+	TArray<FFrameNumber> NewKeyFrames;
+	InTransformMap.GetKeys(NewKeyFrames);
+	TArray<FMovieSceneFloatValue> NewKeyValues;
+
+	// Update keys in channels
+	for (int32 ChannelIndex = 0; ChannelIndex < InNumChannels; ChannelIndex++)
+	{
+		NewKeyValues.Reset();
+		Algo::Transform(InTransformMap, NewKeyValues, [ChannelIndex](const auto& Pair) { return Pair.Value[ChannelIndex]; });
+
+		// All the keys in this channel must be sorted, as adding a set of keys in the curve model before all the others will cause problems.
+		// In order to do this, we assume all 3 sets of keys (before attach range, in attach range, and after attach range) are already sorted
+		// and simply remove and re-add all of the keys from first to last
+		TArray<FKeyHandle> LowerKeyHandles, UpperKeyHandles;
+		TArray<FFrameNumber> LowerKeyTimes, UpperKeyTimes;
+		TArray<FMovieSceneFloatValue> PrevKeyValues;
+		Algo::Copy(InChannels[ChannelIndex]->GetValues(), PrevKeyValues);
+
+		// Get the keys contained in before attach and after attach ranges
+		InChannels[ChannelIndex]->GetKeys(ExcludedRanges[0], &LowerKeyTimes, &LowerKeyHandles);
+		ExcludedRanges.Top().SetLowerBound(TRangeBound<FFrameNumber>::Exclusive(ExcludedRanges.Top().GetLowerBoundValue()));
+		InChannels[ChannelIndex]->GetKeys(ExcludedRanges.Top(), &UpperKeyTimes, &UpperKeyHandles);
+
+		// Add all keys before attach range if they exist
+		int32 ValueIndex = 0;
+		if (ExcludedRanges.Num() > 0 && ExcludedRanges[0].GetUpperBoundValue() <= InAttachRange.GetLowerBoundValue() && LowerKeyTimes.Num() > 0)
+		{
+			InChannels[ChannelIndex]->DeleteKeys(LowerKeyHandles);
+			TArray<FMovieSceneFloatValue> ValuesToAdd;
+			Algo::Copy(TArrayView<const FMovieSceneFloatValue>(PrevKeyValues).Slice(ValueIndex, LowerKeyTimes.Num()), ValuesToAdd);
+			InChannels[ChannelIndex]->AddKeys(LowerKeyTimes, ValuesToAdd);
+			ValueIndex += LowerKeyTimes.Num();
+		}
+
+		// Add all keys in the attach range if they exist
+		InChannels[ChannelIndex]->AddKeys(NewKeyFrames, NewKeyValues);
+
+		// Add all keys after attach range if they exist
+		if (ExcludedRanges.Num() > 0 && ExcludedRanges.Top().GetLowerBoundValue() >= InAttachRange.GetUpperBoundValue() && ValueIndex < PrevKeyValues.Num() && UpperKeyTimes.Num() > 0)
+		{
+			InChannels[ChannelIndex]->DeleteKeys(UpperKeyHandles);
+			TArray<FMovieSceneFloatValue> ValuesToAdd;
+			Algo::Copy(TArrayView<const FMovieSceneFloatValue>(PrevKeyValues).Slice(ValueIndex, UpperKeyTimes.Num()), ValuesToAdd);
+			InChannels[ChannelIndex]->AddKeys(UpperKeyTimes, ValuesToAdd);
+		}
+
+		// If the data is baked, then we also optimize the curves at this point, but do not set tangents since baked keys use linear interpolation
+		if (bInBakedData)
+		{
+			FKeyDataOptimizationParams OptimizationParams;
+			OptimizationParams.bAutoSetInterpolation = false;
+			OptimizationParams.Range = InAttachRange;
+			InChannels[ChannelIndex]->Optimize(OptimizationParams);
+		}
+		else
+		{
+			InChannels[ChannelIndex]->AutoSetTangents();
+		}
+	}
+}
+
+void F3DAttachTrackEditor::TrimAndPreserve(FGuid InObjectBinding, UMovieSceneSection* InSection, bool bInTrimLeft)
+{
+	// Find the transform track associated with the selected object
+	UMovieScene3DTransformTrack* TransformTrack = GetMovieSceneSequence()->GetMovieScene()->FindTrack<UMovieScene3DTransformTrack>(InObjectBinding);
+	if (TransformTrack->GetAllSections().Num() != 1)
+	{
+		return;
+	}
+
+	FMovieSceneEvaluationTrack* EvalTrack = MovieSceneToolHelpers::GetEvaluationTrack(GetSequencer().Get(), TransformTrack->GetSignature());
+
+	TArrayView<TWeakObjectPtr<>> BoundObjects = GetSequencer()->FindBoundObjects(InObjectBinding, GetSequencer()->GetFocusedTemplateID());
+
+	FQualifiedFrameTime QualifiedNewDetachTime = GetSequencer()->GetLocalTime();
+	TRange<FFrameNumber> BeforeTrimRange = InSection->GetRange();
+	if (InSection && EvalTrack && BoundObjects.Num() == 1 && BoundObjects[0].IsValid())
+	{
+		const FScopedTransaction Transaction(LOCTEXT("TrimAttach", "Trim Attach"));
+
+		UObject* Object = BoundObjects[0].Get();
+
+		// Trim the section and find the range of the cut
+		InSection->TrimSection(QualifiedNewDetachTime, bInTrimLeft);
+		TArray<TRange<FFrameNumber>> ExcludedRanges = TRange<FFrameNumber>::Difference(BeforeTrimRange, InSection->GetRange());
+		if (ExcludedRanges.Num() == 0)
+		{
+			return;
+		}
+
+		TRange<FFrameNumber> ExcludedRange = bInTrimLeft ? ExcludedRanges[0] : ExcludedRanges.Top();
+
+		UMovieScene3DAttachSection* AttachSection = Cast<UMovieScene3DAttachSection>(InSection);
+		check(AttachSection);
+
+		// Create a revert modifier with the range and section as parameters
+		FAttachRevertModifier RevertModifier(GetSequencer(), ExcludedRange, AttachSection, AttachSection->AttachSocketName, AttachSection->bFullRevertOnDetach);
+
+		// Find the transform section associated with the track, so far we only support modifying transform tracks with one section
+		UMovieScene3DTransformSection* TransformSection = Cast<UMovieScene3DTransformSection>(TransformTrack->GetAllSections()[0]);
+		if (!TransformSection->TryModify())
+		{
+			return;
+		}
+
+		TArrayView<FMovieSceneFloatChannel*> Channels = TransformSection->GetChannelProxy().GetChannels<FMovieSceneFloatChannel>();
+
+		FLocalTransformEvaluator LocalTransformEval(GetSequencer(), Object, EvalTrack);
+
+		if (AttachSection->ReAttachOnDetach)
+		{
+			FWorldTransformEvaluator ReAttachParentEvaluator(GetSequencer(), AttachSection->ReAttachOnDetach);
+
+			CompensateChildTrack(ExcludedRange, Channels, TOptional<TArrayView<FMovieSceneFloatChannel*>>(), ReAttachParentEvaluator, LocalTransformEval, ETransformPreserveType::CurrentKey, RevertModifier);
+		}
+		else
+		{
+			TSet<FFrameNumber> KeyTimesToCompensate;
+			TMap<FFrameNumber, TArray<FMovieSceneFloatValue>> TransformMap;
+
+			// Add all keys already existing in the range to the transform map
+			AddKeysFromChannels(Channels, ExcludedRange, TransformMap, KeyTimesToCompensate);
+			TArray<FFrameNumber> EdgeKeys;
+
+			// Add the edge keys before and after the cut
+			FFrameNumber RevertEdgeTime;
+			FFrameNumber PreserveEdgeTime;
+			if (bInTrimLeft)
+			{
+				PreserveEdgeTime = ExcludedRange.GetUpperBoundValue();
+				RevertEdgeTime = PreserveEdgeTime.Value - 1;
+				EdgeKeys = { PreserveEdgeTime, RevertEdgeTime };
+				ResizeAndAddKey(PreserveEdgeTime, Channels.Num(), TransformMap, nullptr);
+				ResizeAndAddKey(RevertEdgeTime, Channels.Num(), TransformMap, &KeyTimesToCompensate);
+			}
+			else
+			{
+				RevertEdgeTime = ExcludedRange.GetLowerBoundValue();
+				PreserveEdgeTime = RevertEdgeTime.Value - 1;
+				EdgeKeys = { RevertEdgeTime, PreserveEdgeTime };
+				ResizeAndAddKey(RevertEdgeTime, Channels.Num(), TransformMap, &KeyTimesToCompensate);
+				ResizeAndAddKey(PreserveEdgeTime, Channels.Num(), TransformMap, nullptr);
+			}
+
+			// Evaluate the transform at all times with keys
+			for (auto Itr = TransformMap.CreateIterator(); Itr; ++Itr)
+			{
+				UpdateFloatValueTransform(LocalTransformEval(Itr->Key), Itr->Value);
+			}
+
+			// Modify each transform
+			for (const FFrameNumber& CompTime : KeyTimesToCompensate)
+			{
+				const FTransform RevertedTransform = RevertModifier(FloatValuesToTransform(TransformMap[CompTime]), CompTime);
+				UpdateFloatValueTransform(RevertedTransform, TransformMap[CompTime]);
+			}
+
+			// Manually set edge keys to have linear interpolation
+			for (const FFrameNumber& EdgeKey : EdgeKeys)
+			{
+				for (FMovieSceneFloatValue& Key : TransformMap[EdgeKey])
+				{
+					Key.InterpMode = ERichCurveInterpMode::RCIM_Linear;
 				}
+			}
+
+			// Update the channels with the transform map
+			UpdateChannelTransforms(ExcludedRange, TransformMap, Channels, 9, false);
+		}
+
+		// Remove previous boundary keys
+		for (FMovieSceneFloatChannel* Channel : Channels)
+		{
+			TArray<FKeyHandle> KeyAtTime;
+
+			bInTrimLeft ? 
+			Channel->GetKeys(TRange<FFrameNumber>::Inclusive(ExcludedRange.GetLowerBoundValue() - 1, ExcludedRange.GetLowerBoundValue() - 1), nullptr, &KeyAtTime) : 
+			Channel->GetKeys(TRange<FFrameNumber>::Inclusive(ExcludedRange.GetUpperBoundValue() - 1, ExcludedRange.GetUpperBoundValue() - 1), nullptr, &KeyAtTime);
 
-				if (!TransformSection)
+			Channel->DeleteKeys(KeyAtTime);
+			Channel->AutoSetTangents();
+		}
+	}
+}
+
+template<typename ModifierFuncType>
+void F3DAttachTrackEditor::CompensateChildTrack(const TRange<FFrameNumber>& InAttachRange, TArrayView<FMovieSceneFloatChannel*> Channels, TOptional<TArrayView<FMovieSceneFloatChannel*>> ParentChannels,
+	const ITransformEvaluator& InParentTransformEval, const ITransformEvaluator& InChildTransformEval,
+	ETransformPreserveType InPreserveType, ModifierFuncType InModifyTransform)
+{
+	const FFrameNumber& KeyTime = InAttachRange.GetLowerBoundValue();
+	const FFrameNumber& AttachEndTime = InAttachRange.GetUpperBoundValue();
+	const int32 NumChannels = 9;
+
+	TSet<FFrameNumber> KeyTimesToCompensate;
+	TMap<FFrameNumber, TArray<FMovieSceneFloatValue>> TransformMap;
+
+	// Add all times with keys to the map
+	if (PreserveType == ETransformPreserveType::Bake)
+	{
+		FFrameRate TickResolution = GetSequencer()->GetFocusedTickResolution();
+		FFrameRate DisplayRate = GetSequencer()->GetFocusedDisplayRate();
+		for (FFrameNumber FrameItr = InAttachRange.GetLowerBoundValue(); FrameItr < InAttachRange.GetUpperBoundValue(); 
+			FrameItr += FMath::RoundToInt(TickResolution.AsDecimal() / DisplayRate.AsDecimal()))
+		{
+			ResizeAndAddKey(FrameItr, Channels.Num(), TransformMap, &KeyTimesToCompensate);
+			for (FMovieSceneFloatValue& FloatVal : TransformMap[FrameItr])
+			{
+				FloatVal.InterpMode = ERichCurveInterpMode::RCIM_Linear;
+			}
+		}
+	}
+	else
+	{
+		AddKeysFromChannels(Channels, InAttachRange, TransformMap, KeyTimesToCompensate);
+	}
+
+	const bool bRangeEmpty = TransformMap.Num() == 0 || (TransformMap.Num() == 1 && TransformMap.CreateConstIterator()->Key.Value == KeyTime.Value);
+
+	// Add keys at before and after attach times
+	FFrameNumber BeforeAttachTime = KeyTime.Value - 1;
+	FFrameNumber BeforeDetachTime = AttachEndTime.Value - 1;
+	ResizeAndAddKey(BeforeAttachTime, Channels.Num(), TransformMap, nullptr);
+	ResizeAndAddKey(KeyTime, Channels.Num(), TransformMap, &KeyTimesToCompensate);
+	ResizeAndAddKey(BeforeDetachTime, Channels.Num(), TransformMap, &KeyTimesToCompensate);
+	ResizeAndAddKey(AttachEndTime, Channels.Num(), TransformMap, nullptr);
+
+	if (PreserveType == ETransformPreserveType::AllKeys && ParentChannels.IsSet())
+	{
+		AddKeysFromChannels(ParentChannels.GetValue(), InAttachRange, TransformMap, KeyTimesToCompensate);
+	}
+
+	KeyTimesToCompensate.Remove(AttachEndTime);
+	KeyTimesToCompensate.Remove(BeforeAttachTime);
+	TArray<FFrameNumber> EdgeKeys = { BeforeAttachTime, KeyTime, AttachEndTime, BeforeDetachTime };
+
+	// Evaluate the transform at all times with keys
+	for (auto Itr = TransformMap.CreateIterator(); Itr; ++Itr)
+	{
+		const FTransform TempTransform = InChildTransformEval(Itr->Key);
+		UpdateFloatValueTransform(TempTransform, Itr->Value);
+	}
+
+	if (InPreserveType == ETransformPreserveType::AllKeys || InPreserveType == ETransformPreserveType::Bake)
+	{
+		// If the parent has a transform track, evaluate it's transform at each of the key times found above and calculate the diffs with its child
+		for (const FFrameNumber& CompTime : KeyTimesToCompensate)
+		{
+			const FTransform ParentTransformAtTime = InParentTransformEval(CompTime);
+			const FTransform NewTransform = InModifyTransform(FloatValuesToTransform(TransformMap[CompTime]), CompTime);
+			const FTransform RelativeTransform = NewTransform.GetRelativeTransform(ParentTransformAtTime);
+			UpdateFloatValueTransform(RelativeTransform, TransformMap[CompTime]);
+		}
+	}
+	else if (InPreserveType == ETransformPreserveType::CurrentKey)
+	{
+		// Find the relative transform on the first frame of the attach
+		const FTransform BeginChildTransform = InModifyTransform(FloatValuesToTransform(TransformMap[KeyTime]), KeyTime);
+		const FTransform BeginParentTransform = InParentTransformEval(KeyTime);
+
+		const FTransform BeginRelativeTransform = BeginChildTransform.GetRelativeTransform(BeginParentTransform);
+
+		// offset each transform by initial relative transform calculated before
+		for (const FFrameNumber& CompTime : KeyTimesToCompensate)
+		{
+			const FTransform ChildTransformAtTime = InModifyTransform(FloatValuesToTransform(TransformMap[CompTime]), CompTime);
+			const FTransform StartToCurrentTransform = ChildTransformAtTime.GetRelativeTransform(BeginChildTransform);
+
+			UpdateFloatValueTransform(BeginRelativeTransform * StartToCurrentTransform, TransformMap[CompTime]);
+		}
+
+		const FTransform EndParentTransform = InParentTransformEval(AttachEndTime);
+		UpdateFloatValueTransform(EndParentTransform * FloatValuesToTransform(TransformMap[BeforeDetachTime]), TransformMap[AttachEndTime]);
+	}
+
+	// Manually set edge keys to have linear interpolation
+	for (const FFrameNumber& EdgeKey : EdgeKeys)
+	{
+		for (FMovieSceneFloatValue& Key : TransformMap[EdgeKey])
+		{
+			Key.InterpMode = ERichCurveInterpMode::RCIM_Linear;
+		}
+	}
+
+	UpdateChannelTransforms(InAttachRange, TransformMap, Channels, NumChannels, PreserveType == ETransformPreserveType::Bake);
+}
+
+FKeyPropertyResult F3DAttachTrackEditor::AddKeyInternal( FFrameNumber KeyTime, const TArray<TWeakObjectPtr<UObject>> Objects, const FName SocketName, const FName ComponentName, FActorPickerID ActorPickerID)
+{
+	FKeyPropertyResult KeyPropertyResult;
+
+	FMovieSceneObjectBindingID ConstraintBindingID;
+
+	if (ActorPickerID.ExistingBindingID.IsValid())
+	{
+		ConstraintBindingID = ActorPickerID.ExistingBindingID;
+	}
+	else if (ActorPickerID.ActorPicked.IsValid())
+	{
+		FFindOrCreateHandleResult HandleResult = FindOrCreateHandleToObject(ActorPickerID.ActorPicked.Get());
+		FGuid ParentActorId = HandleResult.Handle;
+		KeyPropertyResult.bHandleCreated |= HandleResult.bWasCreated;
+		ConstraintBindingID = FMovieSceneObjectBindingID(ParentActorId, MovieSceneSequenceID::Root, EMovieSceneObjectBindingSpace::Local);
+	}
+
+	if (!ConstraintBindingID.IsValid())
+	{
+		return KeyPropertyResult;
+	}
+
+	UMovieScene* MovieScene = GetSequencer()->GetFocusedMovieSceneSequence()->GetMovieScene();
+
+	AActor* ParentActor = GetConstraintActor(GetSequencer(), ConstraintBindingID);	
+	check(ParentActor)
+
+	FWorldTransformEvaluator ParentTransformEval(GetSequencer(), ParentActor, SocketName);
+
+	FGuid ParentActorHandle = GetSequencer()->GetHandleToObject(ParentActor, false);
+	TOptional<TArrayView<FMovieSceneFloatChannel*>> ParentChannels;
+	UMovieScene3DTransformTrack* ParentTransformTrack = MovieScene->FindTrack<UMovieScene3DTransformTrack>(ParentActorHandle);
+	if (ParentTransformTrack->GetAllSections().Num() == 1)
+	{
+		ParentChannels = ParentTransformTrack->GetAllSections()[0]->GetChannelProxy().GetChannels<FMovieSceneFloatChannel>();
+	}
+
+	for (int32 ObjectIndex = 0; ObjectIndex < Objects.Num(); ++ObjectIndex)
+	{
+		UObject* Object = Objects[ObjectIndex].Get();
+
+		// Disallow attaching an object to itself
+		if (Object == ParentActor)
+		{
+			continue;
+		}
+
+		// Get handle to object
+		FFindOrCreateHandleResult HandleResult = FindOrCreateHandleToObject( Object );
+		FGuid ObjectHandle = HandleResult.Handle;
+		KeyPropertyResult.bHandleCreated |= HandleResult.bWasCreated;
+		if (!ObjectHandle.IsValid())
+		{
+			continue;
+		}
+
+		// Get attach Track for Object
+		FFindOrCreateTrackResult TrackResult = FindOrCreateTrackForObject(ObjectHandle, UMovieScene3DAttachTrack::StaticClass());
+		UMovieSceneTrack* Track = TrackResult.Track;
+		KeyPropertyResult.bTrackCreated |= TrackResult.bWasCreated;
+		if (!ensure(Track))
+		{
+			continue;
+		}
+
+		// Clamp to next attach section's start time or the end of the current movie scene range
+		FFrameNumber AttachEndTime = MovieScene->GetPlaybackRange().GetUpperBoundValue();
+		for (UMovieSceneSection* Section : Track->GetAllSections())
+		{
+			FFrameNumber StartTime = Section->HasStartFrame() ? Section->GetInclusiveStartFrame() : 0;
+			if (KeyTime < StartTime)
+			{
+				if (AttachEndTime > StartTime)
 				{
-					continue;
+					AttachEndTime = StartTime;
 				}
+			}
+		}
+
+		int32 Duration = FMath::Max(0, (AttachEndTime - KeyTime).Value);
+
+		// Just add the constraint section if no preservation should be done
+		if (PreserveType == ETransformPreserveType::None)
+		{
+			Track->Modify();
+			KeyPropertyResult.bTrackModified = true;
+			Cast<UMovieScene3DAttachTrack>(Track)->AddConstraint(KeyTime, Duration, SocketName, ComponentName, ConstraintBindingID);
+			return KeyPropertyResult;
+		}
+
+		// Create a blank world transform evaluator, add parent evaluator if there is a parent
+		FWorldTransformEvaluator WorldChildTransformEval(GetSequencer(), nullptr);
+		AActor* Actor = Cast<AActor>(Object);
+		if (AActor* PrevParentActor = Actor->GetAttachParentActor())
+		{
+			WorldChildTransformEval = FWorldTransformEvaluator(GetSequencer(), PrevParentActor);
+		}
+
+		// Create transform track for object
+		TRange<FFrameNumber> AttachRange(KeyTime, AttachEndTime);
+		UMovieScene3DTransformTrack* TransformTrack = nullptr;
+		UMovieScene3DTransformSection* TransformSection = nullptr;
+		FMovieSceneEvaluationTrack* EvalTrack = nullptr;
+		FindOrCreateTransformTrack(AttachRange, MovieScene, ObjectHandle, TransformTrack, TransformSection, EvalTrack);
+
+		if (EvalTrack)
+		{
+			WorldChildTransformEval.PrependTransformEval(Object, EvalTrack);
+		}
+		else
+		{
+			WorldChildTransformEval.PrependTransformEval(Actor->GetTransform());
+		}
+
+		if (!TransformSection || !TransformTrack)
+		{
+			continue;
+		}
+
+		if (!TransformSection->TryModify())
+		{
+			continue;
+		}
 
-				if (!TransformSection->TryModify())
+		// Get transform track channels
+		TArrayView<FMovieSceneFloatChannel*> Channels = TransformSection->GetChannelProxy().GetChannels<FMovieSceneFloatChannel>();
+
+		// find intersecting section
+		TOptional<UMovieSceneSection*> IntersectingSection;
+		if (Track->GetAllSections().Num() > 0)
+		{
+			for (UMovieSceneSection* OtherSection : Track->GetAllSections())
+			{
+				if (OtherSection->GetRange().Contains(KeyTime))
 				{
-					continue;
+					IntersectingSection = OtherSection;
+					break;
 				}
+			}
+		}
 
-				TArrayView<FMovieSceneFloatChannel*> Channels = TransformSection->GetChannelProxy().GetChannels<FMovieSceneFloatChannel>();
+		FFrameRate TickResolution = Track->GetTypedOuter<UMovieScene>()->GetTickResolution();
+		
+		Track->Modify();
+		KeyPropertyResult.bTrackModified = true;
+		KeyPropertyResult.bKeyCreated = true;
 
-				FVector DefaultLocation = RelativeTransform.GetLocation();
-				FVector DefaultRotation = RelativeTransform.GetRotation().Euler();
-				FVector DefaultScale3D = RelativeTransform.GetScale3D();
+		TOptional<FAttachRevertModifier> RevertModifier;
+		AActor* ReAttachOnDetach = nullptr;
 
-				Channels[0]->SetDefault(DefaultLocation.X);
-				Channels[1]->SetDefault(DefaultLocation.Y);
-				Channels[2]->SetDefault(DefaultLocation.Z);
+		// If there are existing channels, revert the transform from the previous parent's transform before setting the new relative transform
+		// Currently don't handle objects with both other attach sections and are already attached to other objects because its hard to think about
+		if (IntersectingSection)
+		{
+			// Calculate range to revert
+			TRange<FFrameNumber> RevertRange = TRange<FFrameNumber>(KeyTime, FMath::Min(AttachEndTime, IntersectingSection.GetValue()->GetExclusiveEndFrame()));
 
-				Channels[3]->SetDefault(DefaultRotation.X);
-				Channels[4]->SetDefault(DefaultRotation.Y);
-				Channels[5]->SetDefault(DefaultRotation.Z);
+			// Modify the intersecting section, find the parent actor for it and the respective transform
+			IntersectingSection.GetValue()->TryModify();
+			IntersectingSection.GetValue()->SetEndFrame(KeyTime-1);
 
-				Channels[6]->SetDefault(DefaultScale3D.X);
-				Channels[7]->SetDefault(DefaultScale3D.Y);
-				Channels[8]->SetDefault(DefaultScale3D.Z);
+			UMovieScene3DAttachSection* IntersectingAttachSection = Cast<UMovieScene3DAttachSection>(IntersectingSection.GetValue());
+			if (!IntersectingAttachSection)
+			{
+				continue;
 			}
+
+			RevertModifier = FAttachRevertModifier(GetSequencer(), RevertRange, IntersectingAttachSection, SocketName, PreserveType == ETransformPreserveType::CurrentKey);
 		}
-	}
+		// Existing parent that's not an attach track
+		else if (WorldChildTransformEval.GetTransformEvalsView().Num() > 1)
+		{
+			// Calculate range to revert
+			TRange<FFrameNumber> RevertRange = AttachRange;
+
+			// Get the evaluator for the previous parent track
+			const int32 NumChildEvals = WorldChildTransformEval.GetTransformEvalsView().Num();
+			auto PrevParentTransformEvals = WorldChildTransformEval.GetTransformEvalsView().Slice(1, NumChildEvals - 1);
+			FWorldTransformEvaluator PrevParentEvaluator(GetSequencer(), PrevParentTransformEvals);
+
+			RevertModifier = FAttachRevertModifier(GetSequencer(), RevertRange, PrevParentEvaluator, PreserveType == ETransformPreserveType::CurrentKey);
+
+			ReAttachOnDetach = Actor->GetAttachParentActor();
+		}
+
+		if (RevertModifier.IsSet())
+		{
+			FLocalTransformEvaluator LocalChildTransformEval(GetSequencer(), Object, EvalTrack);
+
+			// Add the new attach section to the track
+			Cast<UMovieScene3DAttachTrack>(Track)->AddConstraint(KeyTime, Duration, SocketName, ComponentName, ConstraintBindingID);
+
+			// Compensate
+			CompensateChildTrack(AttachRange, Channels, ParentChannels, ParentTransformEval, LocalChildTransformEval, PreserveType, RevertModifier.GetValue());
+		}
+		else
+		{
+
+			// Add the new attach section to the track
+			Cast<UMovieScene3DAttachTrack>(Track)->AddConstraint(KeyTime, Duration, SocketName, ComponentName, ConstraintBindingID);
+
+			// Compensate
+			CompensateChildTrack(AttachRange, Channels, ParentChannels, ParentTransformEval, WorldChildTransformEval, PreserveType, [&](const FTransform& InTransform, const FFrameNumber& InTime) { return InTransform; });
+		}
+
+		Cast<UMovieScene3DAttachSection>(Cast<UMovieScene3DAttachTrack>(Track)->GetAllSections().Top())->bFullRevertOnDetach = (PreserveType == ETransformPreserveType::CurrentKey);
+		Cast<UMovieScene3DAttachSection>(Cast<UMovieScene3DAttachTrack>(Track)->GetAllSections().Top())->ReAttachOnDetach = ReAttachOnDetach;
+	} // for
 
 	return KeyPropertyResult;
 }
diff --git a/Engine/Source/Editor/MovieSceneTools/Private/TrackEditors/AttachTrackEditor.h b/Engine/Source/Editor/MovieSceneTools/Private/TrackEditors/AttachTrackEditor.h
--- a/Engine/Source/Editor/MovieSceneTools/Private/TrackEditors/AttachTrackEditor.h
+++ b/Engine/Source/Editor/MovieSceneTools/Private/TrackEditors/AttachTrackEditor.h
@@ -10,10 +10,28 @@
 #include "ISequencer.h"
 #include "ISequencerTrackEditor.h"
 #include "TrackEditors/ActorPickerTrackEditor.h"
+#include "Containers/Union.h"
+#include "Containers/Map.h"
+#include "Templates/Function.h"
+
+struct FMovieSceneFloatChannel;
 
 class AActor;
 class FMenuBuilder;
 class USceneComponent;
+class UMovieScene3DTransformTrack;
+class UMovieScene3DTransformSection;
+class UMovieScene3DAttachSection;
+
+struct ITransformEvaluator;
+
+enum class ETransformPreserveType
+{
+	CurrentKey,
+	AllKeys,
+	Bake,
+	None
+};
 
 /**
  * Tools for attaching an object to another object
@@ -53,8 +71,24 @@ public:
 	virtual bool IsActorPickable( const AActor* const ParentActor, FGuid ObjectBinding, UMovieSceneSection* InSection ) override;
 	virtual void ActorSocketPicked(const FName SocketName, USceneComponent* Component, FActorPickerID ActorPickerID, TArray<FGuid> ObjectBindings, UMovieSceneSection* Section) override;
 
+	/** Trims an attach track on the left/right and preserves the world space transform immediately before/after attach/detach */
+	void TrimAndPreserve(const FGuid InObjectBinding, UMovieSceneSection* InSection, bool bInTrimLeft);
+
 private:
 
+	void ShowPickerSubMenu(FMenuBuilder& MenuBuilder, TArray<FGuid> ObjectBindings, UMovieSceneSection* Section);
+
+	/** Helper for AddKeyInternal to get transform tracks for child */
+	void FindOrCreateTransformTrack(const TRange<FFrameNumber>& InAttachRange, UMovieScene* InMovieScene, const FGuid& InObjectHandle, UMovieScene3DTransformTrack*& OutTransformTrack, UMovieScene3DTransformSection*& OutTransformSection, FMovieSceneEvaluationTrack*& OutEvalTrack);
+
+	/** Helper for AddKeyInternal to offset child track's keys */
+	template<typename ModifierFuncType>
+	void CompensateChildTrack(const TRange<FFrameNumber>& InAttachRange, TArrayView<FMovieSceneFloatChannel*> Channels, TOptional<TArrayView<FMovieSceneFloatChannel*>> ParentChannels,
+		const ITransformEvaluator& InParentTransformEval, const ITransformEvaluator& InChildTransformEval, ETransformPreserveType InPreserveType, ModifierFuncType InModifyTransform);
+
 	/** Delegate for AnimatablePropertyChanged in AddKey */
 	FKeyPropertyResult AddKeyInternal(FFrameNumber KeyTime, const TArray<TWeakObjectPtr<UObject>> Objects, const FName SocketName, const FName ComponentName, FActorPickerID ActorPickerID);
+
+private:
+	ETransformPreserveType PreserveType;
 };
diff --git a/Engine/Source/Runtime/MovieSceneTracks/Public/Sections/MovieScene3DAttachSection.h b/Engine/Source/Runtime/MovieSceneTracks/Public/Sections/MovieScene3DAttachSection.h
--- a/Engine/Source/Runtime/MovieSceneTracks/Public/Sections/MovieScene3DAttachSection.h
+++ b/Engine/Source/Runtime/MovieSceneTracks/Public/Sections/MovieScene3DAttachSection.h
@@ -37,6 +37,17 @@ public:
 	UPROPERTY(BlueprintReadWrite, EditAnywhere, Category="Attach")
 	FName AttachComponentName;
 
+
+#if WITH_EDITORONLY_DATA
+
+	UPROPERTY(BlueprintReadWrite, EditAnywhere, Category="Attach")
+	bool bFullRevertOnDetach;
+
+	UPROPERTY(BlueprintReadWrite, EditAnywhere, Category="Attach")
+	AActor* ReAttachOnDetach;
+
+#endif
+
 	UPROPERTY(BlueprintReadWrite, EditAnywhere, Category="Attach")
 	EAttachmentRule AttachmentLocationRule;
 
