Sequencer: fix crash on attach when parent does not have transform track

diff --git a/Engine/Source/Editor/MovieSceneTools/Private/TrackEditors/AttachTrackEditor.cpp b/Engine/Source/Editor/MovieSceneTools/Private/TrackEditors/AttachTrackEditor.cpp
--- a/Engine/Source/Editor/MovieSceneTools/Private/TrackEditors/AttachTrackEditor.cpp
+++ b/Engine/Source/Editor/MovieSceneTools/Private/TrackEditors/AttachTrackEditor.cpp
@@ -1121,7 +1121,7 @@ FKeyPropertyResult F3DAttachTrackEditor::AddKeyInternal( FFrameNumber KeyTime, c
 	FGuid ParentActorHandle = GetSequencer()->GetHandleToObject(ParentActor, false);
 	TOptional<TArrayView<FMovieSceneFloatChannel*>> ParentChannels;
 	UMovieScene3DTransformTrack* ParentTransformTrack = MovieScene->FindTrack<UMovieScene3DTransformTrack>(ParentActorHandle);
-	if (ParentTransformTrack->GetAllSections().Num() == 1)
+	if (ParentTransformTrack && ParentTransformTrack->GetAllSections().Num() == 1)
 	{
 		ParentChannels = ParentTransformTrack->GetAllSections()[0]->GetChannelProxy().GetChannels<FMovieSceneFloatChannel>();
 	}
@@ -1176,7 +1176,7 @@ FKeyPropertyResult F3DAttachTrackEditor::AddKeyInternal( FFrameNumber KeyTime, c
 			Track->Modify();
 			KeyPropertyResult.bTrackModified = true;
 			Cast<UMovieScene3DAttachTrack>(Track)->AddConstraint(KeyTime, Duration, SocketName, ComponentName, ConstraintBindingID);
-			return KeyPropertyResult;
+			continue;
 		}
 
 		// Create a blank world transform evaluator, add parent evaluator if there is a parent
