Sequencer: Fix bug when object is attached twice

diff --git a/Engine/Source/Editor/MovieSceneTools/Private/TrackEditors/AttachTrackEditor.cpp b/Engine/Source/Editor/MovieSceneTools/Private/TrackEditors/AttachTrackEditor.cpp
--- a/Engine/Source/Editor/MovieSceneTools/Private/TrackEditors/AttachTrackEditor.cpp
+++ b/Engine/Source/Editor/MovieSceneTools/Private/TrackEditors/AttachTrackEditor.cpp
@@ -884,7 +884,7 @@ void F3DAttachTrackEditor::TrimAndPreserve(FGuid InObjectBinding, UMovieSceneSec
 		UObject* Object = BoundObjects[0].Get();
 
 		// Trim the section and find the range of the cut
-		InSection->TrimSection(QualifiedNewDetachTime, bInTrimLeft, true);
+		InSection->TrimSection(QualifiedNewDetachTime, bInTrimLeft, false);
 		TArray<TRange<FFrameNumber>> ExcludedRanges = TRange<FFrameNumber>::Difference(BeforeTrimRange, InSection->GetRange());
 		if (ExcludedRanges.Num() == 0)
 		{
@@ -1246,9 +1246,20 @@ FKeyPropertyResult F3DAttachTrackEditor::AddKeyInternal( FFrameNumber KeyTime, c
 			// Calculate range to revert
 			TRange<FFrameNumber> RevertRange = TRange<FFrameNumber>(KeyTime, FMath::Min(AttachEndTime, IntersectingSection.GetValue()->GetExclusiveEndFrame()));
 
-			// Modify the intersecting section, find the parent actor for it and the respective transform
-			IntersectingSection.GetValue()->TryModify();
-			IntersectingSection.GetValue()->SetEndFrame(KeyTime-1);
+			// If the intersecting section starts at the same time as the new section, remove it
+			if (IntersectingSection.GetValue()->GetInclusiveStartFrame() == KeyTime)
+			{
+				Track->RemoveSection(*IntersectingSection.GetValue());
+			}
+			// Otherwise trim the end frame of the intersecting section
+			else
+			{
+				if (!IntersectingSection.GetValue()->TryModify())
+				{
+					continue;
+				}
+				IntersectingSection.GetValue()->SetEndFrame(KeyTime - 1);
+			}
 
 			UMovieScene3DAttachSection* IntersectingAttachSection = Cast<UMovieScene3DAttachSection>(IntersectingSection.GetValue());
 			if (!IntersectingAttachSection)
