Account for frame rate when exporting tangents for camera anims

diff --git a/Engine/Source/Editor/MovieSceneTools/Private/MovieSceneToolHelpers.cpp b/Engine/Source/Editor/MovieSceneTools/Private/MovieSceneToolHelpers.cpp
--- a/Engine/Source/Editor/MovieSceneTools/Private/MovieSceneToolHelpers.cpp
+++ b/Engine/Source/Editor/MovieSceneTools/Private/MovieSceneToolHelpers.cpp
@@ -1946,8 +1946,8 @@ void MovieSceneToolHelpers::CopyKeyDataToMoveAxis(const TMovieSceneChannelData<F
 		MoveAxis->LookupTrack.AddPoint(Time, LookupName);
 
 		FInterpCurvePoint<float>& Point = MoveAxis->FloatTrack.Points[PointIndex];
-		Point.ArriveTangent = Value.Tangent.ArriveTangent;
-		Point.LeaveTangent = Value.Tangent.LeaveTangent;
+		Point.ArriveTangent = Value.Tangent.ArriveTangent * InFrameRate.AsDecimal();
+		Point.LeaveTangent = Value.Tangent.LeaveTangent * InFrameRate.AsDecimal();
 		Point.InterpMode = RichCurveInterpolationToMatineeInterpolation(Value.InterpMode, Value.TangentMode);
 	}
 }
