Sequencer: Newly created section for import root motion should be additive with scale disabled so that the imported motion moves the object from its current position.

diff --git a/Engine/Source/Editor/MovieSceneTools/Private/TrackEditors/TransformTrackEditor.cpp b/Engine/Source/Editor/MovieSceneTools/Private/TrackEditors/TransformTrackEditor.cpp
--- a/Engine/Source/Editor/MovieSceneTools/Private/TrackEditors/TransformTrackEditor.cpp
+++ b/Engine/Source/Editor/MovieSceneTools/Private/TrackEditors/TransformTrackEditor.cpp
@@ -957,6 +957,10 @@ void F3DTransformTrackEditor::ImportAnimSequenceTransforms(const FAssetData& Ass
 		TransformTrack->Modify();
 
 		UMovieScene3DTransformSection* Section = Cast<UMovieScene3DTransformSection>(TransformTrack->CreateNewSection());
+
+		Section->SetBlendType(EMovieSceneBlendType::Additive);
+		Section->SetMask(EMovieSceneTransformChannel::Translation | EMovieSceneTransformChannel::Rotation);
+		
 		FFrameRate TickResolution = Section->GetTypedOuter<UMovieScene>()->GetTickResolution();
 
 		TArrayView<FMovieSceneFloatChannel*> FloatChannels = Section->GetChannelProxy().GetChannels<FMovieSceneFloatChannel>();
