Sequencer: Make audio loop when extended

diff --git a/Engine/Source/Editor/MovieSceneTools/Private/TrackEditors/AudioTrackEditor.cpp b/Engine/Source/Editor/MovieSceneTools/Private/TrackEditors/AudioTrackEditor.cpp
--- a/Engine/Source/Editor/MovieSceneTools/Private/TrackEditors/AudioTrackEditor.cpp
+++ b/Engine/Source/Editor/MovieSceneTools/Private/TrackEditors/AudioTrackEditor.cpp
@@ -329,11 +329,16 @@ void FAudioThumbnail::GenerateWaveformPreview(TArray<uint8>& OutData, TRange<flo
 	double     SectionStartTime     = AudioSection->GetInclusiveStartFrame() / FrameRate;
 
 	// @todo Sequencer This fixes looping drawing by pretending we are only dealing with a SoundWave
-	TRange<double> AudioTrueRange = TRange<double>(
+	TRange<float> AudioTrueRange = TRange<float>(
 		SectionStartTime - FrameRate.AsSeconds(AudioSection->GetStartOffset()),
 		SectionStartTime - FrameRate.AsSeconds(AudioSection->GetStartOffset()) + DeriveUnloopedDuration(AudioSection) * (1.0f / PitchMultiplierValue));
 
 	float TrueRangeSize = AudioTrueRange.Size<float>();
+	const int32 TrueDrawOffsetPx = FMath::Max(FMath::RoundToInt((DrawRange.GetLowerBoundValue() - SectionStartTime) / DisplayScale), 0);
+	const int32 LastTrueSample = -2.f*SmoothingAmount + FMath::TruncToInt(TrueRangeSize / DisplayScale);
+
+	DrawRange = AudioTrueRange;
+
 	float DrawRangeSize = DrawRange.Size<float>();
 
 	const int32 MaxAmplitude = NumChannels == 1 ? GetSize().Y : GetSize().Y / 2;
@@ -345,14 +350,14 @@ void FAudioThumbnail::GenerateWaveformPreview(TArray<uint8>& OutData, TRange<flo
 	float RangeLookupFraction = (SmoothingAmount*DisplayScale) / TrueRangeSize;
 	int32 LookupRange = FMath::Clamp(FMath::TruncToInt(RangeLookupFraction * LookupSize), 1, LookupSize);
 
-	int32 SampleLockOffset = DrawOffsetPx % SmoothingAmount;
+	int32 SampleLockOffset = TrueDrawOffsetPx % SmoothingAmount;
 
 	int32 FirstSample = -2.f*SmoothingAmount - SampleLockOffset;
-	int32 LastSample = GetSize().X + 2*SmoothingAmount;
+	int32 LastSample = LastTrueSample + 2*SmoothingAmount;
 
 	{
 		// @todo: when SampleCount <= 0, we have fewer samples than pixels, and should start to interpolate the spline by that distance, rather than a hard coded pixel density
-		int32 NumSamplesInRange = FMath::TruncToInt(LookupSize * (DrawRangeSize /GetSize().X) / TrueRangeSize);
+		int32 NumSamplesInRange = FMath::TruncToInt(LookupSize * (DrawRangeSize /LastTrueSample) / TrueRangeSize);
 		int32 SampleCount = NumSamplesInRange / NumChannels;
 	}
 
@@ -360,11 +365,11 @@ void FAudioThumbnail::GenerateWaveformPreview(TArray<uint8>& OutData, TRange<flo
 	// Sample the audio one pixel to the left and right
 	for (int32 X = FirstSample; X < LastSample; ++X)
 	{
-		float LookupTime = ((float)(X - 0.5f) / (float)GetSize().X) * DrawRangeSize + DrawRange.GetLowerBoundValue();
+		float LookupTime = ((float)(X - 0.5f) / (float)LastTrueSample) * DrawRangeSize + DrawRange.GetLowerBoundValue();
 		float LookupFraction = (LookupTime - AudioTrueRange.GetLowerBoundValue()) / TrueRangeSize;
 		int32 LookupIndex = FMath::TruncToInt(LookupFraction * LookupSize);
 		
-		float NextLookupTime = ((float)(X + 0.5f) / (float)GetSize().X) * DrawRangeSize + DrawRange.GetLowerBoundValue();
+		float NextLookupTime = ((float)(X + 0.5f) / (float)LastTrueSample) * DrawRangeSize + DrawRange.GetLowerBoundValue();
 		float NextLookupFraction = (NextLookupTime - AudioTrueRange.GetLowerBoundValue()) / TrueRangeSize;
 		int32 NextLookupIndex = FMath::TruncToInt(NextLookupFraction  * LookupSize);
 		
@@ -376,10 +381,13 @@ void FAudioThumbnail::GenerateWaveformPreview(TArray<uint8>& OutData, TRange<flo
 
 	// Now draw the spline
 	const int32 Height = GetSize().Y;
-	const int32 Width = GetSize().X;
+	const int32 Width = LastTrueSample;
 
 	FLinearColor BoundaryColor = BoundaryColorHSV.HSVToLinearRGB();
 
+	uint32 Size = LastTrueSample * GetSize().Y * GPixelFormats[PF_B8G8R8A8].BlockBytes;
+	TArray<uint8> TempData;
+	TempData.SetNum(Size);
 	for (int32 ChannelIndex = 0; ChannelIndex < SoundWave->NumChannels; ++ChannelIndex)
 	{
 		int32 SplineIndex = 0;
@@ -414,7 +422,7 @@ void FAudioThumbnail::GenerateWaveformPreview(TArray<uint8>& OutData, TRange<flo
 
 			for (int32 PixelIndex = 0; PixelIndex < MaxAmplitude; ++PixelIndex)
 			{
-				uint8* Pixel = LookupPixel(OutData, X, PixelIndex, Width, Height, ChannelIndex, NumChannels);
+				uint8* Pixel = LookupPixel(TempData, X, PixelIndex, Width, Height, ChannelIndex, NumChannels);
 
 				const float PixelCenter = PixelIndex + 0.5f;
 
@@ -445,6 +453,30 @@ void FAudioThumbnail::GenerateWaveformPreview(TArray<uint8>& OutData, TRange<flo
 			}
 		}
 	}
+
+	// Then loop the texture with the start offset
+	int32 TotalDrawOffsetPx = TrueDrawOffsetPx + FMath::RoundToInt(FrameRate.AsSeconds(AudioSection->GetStartOffset()) / DisplayScale);
+	for (int32 ChannelIndex = 0; ChannelIndex < SoundWave->NumChannels; ++ChannelIndex)
+	{
+		int32 SplineIndex = 0;
+
+		int32 Iters = 0;
+		for (int32 X = 0; X < GetSize().X; ++X)
+		{
+			int32 XRot = ((X + TotalDrawOffsetPx) % LastTrueSample);
+
+			for (int32 PixelIndex = 0; PixelIndex < MaxAmplitude; ++PixelIndex)
+			{
+				uint8* PixelToCopy = LookupPixel(TempData, XRot, PixelIndex, Width, Height, ChannelIndex, NumChannels);
+				uint8* Pixel = LookupPixel(OutData, X, PixelIndex, GetSize().X, Height, ChannelIndex, NumChannels);
+
+				*Pixel++ = *PixelToCopy++;
+				*Pixel++ = *PixelToCopy++;
+				*Pixel++ = *PixelToCopy++;
+				*Pixel++ = *PixelToCopy++;
+			}
+		}
+	}
 }
 
 void FAudioThumbnail::GenerateSpline(int32 NumChannels, int32 SamplePositionOffset)
diff --git a/Engine/Source/Runtime/MovieSceneTracks/Private/Evaluation/MovieSceneAudioTemplate.cpp b/Engine/Source/Runtime/MovieSceneTracks/Private/Evaluation/MovieSceneAudioTemplate.cpp
--- a/Engine/Source/Runtime/MovieSceneTracks/Private/Evaluation/MovieSceneAudioTemplate.cpp
+++ b/Engine/Source/Runtime/MovieSceneTracks/Private/Evaluation/MovieSceneAudioTemplate.cpp
@@ -362,9 +362,11 @@ struct FAudioSectionExecutionToken : IMovieSceneExecutionToken
 
 			float SectionStartTimeSeconds = (AudioSection->HasStartFrame() ? AudioSection->GetInclusiveStartFrame() : 0) / AudioSection->GetTypedOuter<UMovieScene>()->GetTickResolution();
 
-			const float AudioTime = (Context.GetTime() / Context.GetFrameRate()) - SectionStartTimeSeconds + (float)Context.GetFrameRate().AsSeconds(AudioSection->GetStartOffset());
-			if (AudioTime >= 0.f && AudioComponent.Sound && AudioTime < AudioComponent.Sound->GetDuration())
+			float AudioTime = (Context.GetTime() / Context.GetFrameRate()) - SectionStartTimeSeconds + (float)Context.GetFrameRate().AsSeconds(AudioSection->GetStartOffset());
+			if (AudioTime >= 0.f && AudioComponent.Sound)
 			{
+				float Duration = AudioComponent.Sound->GetDuration();
+				AudioTime = FMath::Fmod(AudioTime, Duration/1000.f);
 				AudioComponent.Play(AudioTime);
 			}
 
