Sequencer: Remove unused variable in AudioTrackEditor.cpp

diff --git a/Engine/Source/Editor/MovieSceneTools/Private/TrackEditors/AudioTrackEditor.cpp b/Engine/Source/Editor/MovieSceneTools/Private/TrackEditors/AudioTrackEditor.cpp
--- a/Engine/Source/Editor/MovieSceneTools/Private/TrackEditors/AudioTrackEditor.cpp
+++ b/Engine/Source/Editor/MovieSceneTools/Private/TrackEditors/AudioTrackEditor.cpp
@@ -459,8 +459,6 @@ void FAudioThumbnail::GenerateWaveformPreview(TArray<uint8>& OutData, TRange<flo
 	for (int32 ChannelIndex = 0; ChannelIndex < SoundWave->NumChannels; ++ChannelIndex)
 	{
 		int32 SplineIndex = 0;
-
-		int32 Iters = 0;
 		for (int32 X = 0; X < GetSize().X; ++X)
 		{
 			int32 XRot = ((X + TotalDrawOffsetPx) % LastTrueSample);
