Fix for doubled values.
Was little involved since had to share data from the property templates in MovieScene with those defined in MovieSceneTracks, so moved FTransformData over to MovieScene and created a new Interrogation Key to be used by all. Note used FTransformData and not FTransform so we could keep Rotator info and not Quat which may be lossy on conversions back and forth.

diff --git a/Engine/Source/Editor/MovieSceneTools/Private/TrackEditors/PropertyTrackEditors/EulerTransformPropertyTrackEditor.cpp b/Engine/Source/Editor/MovieSceneTools/Private/TrackEditors/PropertyTrackEditors/EulerTransformPropertyTrackEditor.cpp
--- a/Engine/Source/Editor/MovieSceneTools/Private/TrackEditors/PropertyTrackEditors/EulerTransformPropertyTrackEditor.cpp
+++ b/Engine/Source/Editor/MovieSceneTools/Private/TrackEditors/PropertyTrackEditors/EulerTransformPropertyTrackEditor.cpp
@@ -126,23 +126,24 @@ bool FEulerTransformPropertyTrackEditor::ModifyGeneratedKeysByCurrentAndWeight(U
 
 	FVector CurrentPos; FRotator CurrentRot;
 	FVector CurrentScale;
-	for (const FEulerTransform& Transform : InterrogationData.Iterate<FEulerTransform>(FMovieScenePropertySectionTemplate::GetEulerTransformInterrogationKey()))
+	for (const FTransformData& Transform : InterrogationData.Iterate<FTransformData>(FMovieSceneInterrogationKey::GetTransformInterrogationKey()))
 	{
-		CurrentPos = Transform.Location;
+		CurrentPos = Transform.Translation;
 		CurrentRot = Transform.Rotation;
 		CurrentScale = Transform.Scale;
-		break;
+
+		FMovieSceneChannelProxy& Proxy = SectionToKey->GetChannelProxy();
+		GeneratedTotalKeys[0]->ModifyByCurrentAndWeight(Proxy, KeyTime, (void *)&CurrentPos.X, Weight);
+		GeneratedTotalKeys[1]->ModifyByCurrentAndWeight(Proxy, KeyTime, (void *)&CurrentPos.Y, Weight);
+		GeneratedTotalKeys[2]->ModifyByCurrentAndWeight(Proxy, KeyTime, (void *)&CurrentPos.Z, Weight);
+		GeneratedTotalKeys[3]->ModifyByCurrentAndWeight(Proxy, KeyTime, (void *)&CurrentRot.Roll, Weight);
+		GeneratedTotalKeys[4]->ModifyByCurrentAndWeight(Proxy, KeyTime, (void *)&CurrentRot.Pitch, Weight);
+		GeneratedTotalKeys[5]->ModifyByCurrentAndWeight(Proxy, KeyTime, (void *)&CurrentRot.Yaw, Weight);
+		GeneratedTotalKeys[6]->ModifyByCurrentAndWeight(Proxy, KeyTime, (void *)&CurrentScale.X, Weight);
+		GeneratedTotalKeys[7]->ModifyByCurrentAndWeight(Proxy, KeyTime, (void *)&CurrentScale.Y, Weight);
+		GeneratedTotalKeys[8]->ModifyByCurrentAndWeight(Proxy, KeyTime, (void *)&CurrentScale.Z, Weight);
+		return true;
 	}
-	FMovieSceneChannelProxy& Proxy = SectionToKey->GetChannelProxy();
-	GeneratedTotalKeys[0]->ModifyByCurrentAndWeight(Proxy, KeyTime, (void *)&CurrentPos.X, Weight);
-	GeneratedTotalKeys[1]->ModifyByCurrentAndWeight(Proxy, KeyTime, (void *)&CurrentPos.Y, Weight);
-	GeneratedTotalKeys[2]->ModifyByCurrentAndWeight(Proxy, KeyTime, (void *)&CurrentPos.Z, Weight);
-	GeneratedTotalKeys[3]->ModifyByCurrentAndWeight(Proxy, KeyTime, (void *)&CurrentRot.Roll, Weight);
-	GeneratedTotalKeys[4]->ModifyByCurrentAndWeight(Proxy, KeyTime, (void *)&CurrentRot.Pitch, Weight);
-	GeneratedTotalKeys[5]->ModifyByCurrentAndWeight(Proxy, KeyTime, (void *)&CurrentRot.Yaw, Weight);
-	GeneratedTotalKeys[6]->ModifyByCurrentAndWeight(Proxy, KeyTime, (void *)&CurrentScale.X, Weight);
-	GeneratedTotalKeys[7]->ModifyByCurrentAndWeight(Proxy, KeyTime, (void *)&CurrentScale.Y, Weight);
-	GeneratedTotalKeys[8]->ModifyByCurrentAndWeight(Proxy, KeyTime, (void *)&CurrentScale.Z, Weight);
-	return true;
+	return false;
 
 }
\ No newline at end of file
diff --git a/Engine/Source/Editor/MovieSceneTools/Private/TrackEditors/PropertyTrackEditors/TransformPropertyTrackEditor.cpp b/Engine/Source/Editor/MovieSceneTools/Private/TrackEditors/PropertyTrackEditors/TransformPropertyTrackEditor.cpp
--- a/Engine/Source/Editor/MovieSceneTools/Private/TrackEditors/PropertyTrackEditors/TransformPropertyTrackEditor.cpp
+++ b/Engine/Source/Editor/MovieSceneTools/Private/TrackEditors/PropertyTrackEditors/TransformPropertyTrackEditor.cpp
@@ -79,23 +79,24 @@ bool FTransformPropertyTrackEditor::ModifyGeneratedKeysByCurrentAndWeight(UObjec
 
 	FVector CurrentPos; FRotator CurrentRot;
 	FVector CurrentScale;
-	for (const FTransform& Transform : InterrogationData.Iterate<FTransform>(FMovieScenePropertySectionTemplate::GetTransformInterrogationKey()))
+	for (const FTransformData& Transform : InterrogationData.Iterate<FTransformData>(FMovieSceneInterrogationKey::GetTransformInterrogationKey()))
 	{
-		CurrentPos = Transform.GetTranslation();
-		CurrentRot = Transform.Rotator();
-		CurrentScale = Transform.GetScale3D();
-		break;
+		CurrentPos = Transform.Translation;
+		CurrentRot = Transform.Rotation;
+		CurrentScale = Transform.Scale;
+
+		FMovieSceneChannelProxy& Proxy = SectionToKey->GetChannelProxy();
+		GeneratedTotalKeys[0]->ModifyByCurrentAndWeight(Proxy, KeyTime, (void *)&CurrentPos.X, Weight);
+		GeneratedTotalKeys[1]->ModifyByCurrentAndWeight(Proxy, KeyTime, (void *)&CurrentPos.Y, Weight);
+		GeneratedTotalKeys[2]->ModifyByCurrentAndWeight(Proxy, KeyTime, (void *)&CurrentPos.Z, Weight);
+		GeneratedTotalKeys[3]->ModifyByCurrentAndWeight(Proxy, KeyTime, (void *)&CurrentRot.Roll, Weight);
+		GeneratedTotalKeys[4]->ModifyByCurrentAndWeight(Proxy, KeyTime, (void *)&CurrentRot.Pitch, Weight);
+		GeneratedTotalKeys[5]->ModifyByCurrentAndWeight(Proxy, KeyTime, (void *)&CurrentRot.Yaw, Weight);
+		GeneratedTotalKeys[6]->ModifyByCurrentAndWeight(Proxy, KeyTime, (void *)&CurrentScale.X, Weight);
+		GeneratedTotalKeys[7]->ModifyByCurrentAndWeight(Proxy, KeyTime, (void *)&CurrentScale.Y, Weight);
+		GeneratedTotalKeys[8]->ModifyByCurrentAndWeight(Proxy, KeyTime, (void *)&CurrentScale.Z, Weight);
+		return true;
 	}
-	FMovieSceneChannelProxy& Proxy = SectionToKey->GetChannelProxy();
-	GeneratedTotalKeys[0]->ModifyByCurrentAndWeight(Proxy, KeyTime, (void *)&CurrentPos.X, Weight);
-	GeneratedTotalKeys[1]->ModifyByCurrentAndWeight(Proxy, KeyTime, (void *)&CurrentPos.Y, Weight);
-	GeneratedTotalKeys[2]->ModifyByCurrentAndWeight(Proxy, KeyTime, (void *)&CurrentPos.Z, Weight);
-	GeneratedTotalKeys[3]->ModifyByCurrentAndWeight(Proxy, KeyTime, (void *)&CurrentRot.Roll, Weight);
-	GeneratedTotalKeys[4]->ModifyByCurrentAndWeight(Proxy, KeyTime, (void *)&CurrentRot.Pitch, Weight);
-	GeneratedTotalKeys[5]->ModifyByCurrentAndWeight(Proxy, KeyTime, (void *)&CurrentRot.Yaw, Weight);
-	GeneratedTotalKeys[6]->ModifyByCurrentAndWeight(Proxy, KeyTime, (void *)&CurrentScale.X, Weight);
-	GeneratedTotalKeys[7]->ModifyByCurrentAndWeight(Proxy, KeyTime, (void *)&CurrentScale.Y, Weight);
-	GeneratedTotalKeys[8]->ModifyByCurrentAndWeight(Proxy, KeyTime, (void *)&CurrentScale.Z, Weight);
-	return true;
+	return false;
 
 }
\ No newline at end of file
diff --git a/Engine/Source/Editor/MovieSceneTools/Private/TrackEditors/TransformTrackEditor.cpp b/Engine/Source/Editor/MovieSceneTools/Private/TrackEditors/TransformTrackEditor.cpp
--- a/Engine/Source/Editor/MovieSceneTools/Private/TrackEditors/TransformTrackEditor.cpp
+++ b/Engine/Source/Editor/MovieSceneTools/Private/TrackEditors/TransformTrackEditor.cpp
@@ -869,20 +869,20 @@ bool F3DTransformTrackEditor::ModifyGeneratedKeysByCurrentAndWeight(UObject *Obj
 		CurrentPos = Transform.Translation;
 		CurrentRot = Transform.Rotation;
 		CurrentScale = Transform.Scale;
-		break;
-	}
-	FMovieSceneChannelProxy& Proxy = SectionToKey->GetChannelProxy();
-	GeneratedTotalKeys[0]->ModifyByCurrentAndWeight(Proxy, KeyTime, (void *)&CurrentPos.X, Weight);
-	GeneratedTotalKeys[1]->ModifyByCurrentAndWeight(Proxy, KeyTime, (void *)&CurrentPos.Y, Weight);
-	GeneratedTotalKeys[2]->ModifyByCurrentAndWeight(Proxy, KeyTime, (void *)&CurrentPos.Z, Weight);
-	GeneratedTotalKeys[3]->ModifyByCurrentAndWeight(Proxy, KeyTime, (void *)&CurrentRot.Roll, Weight);
-	GeneratedTotalKeys[4]->ModifyByCurrentAndWeight(Proxy, KeyTime, (void *)&CurrentRot.Pitch, Weight);
-	GeneratedTotalKeys[5]->ModifyByCurrentAndWeight(Proxy, KeyTime, (void *)&CurrentRot.Yaw, Weight);
-	GeneratedTotalKeys[6]->ModifyByCurrentAndWeight(Proxy, KeyTime, (void *)&CurrentScale.X, Weight);
-	GeneratedTotalKeys[7]->ModifyByCurrentAndWeight(Proxy, KeyTime, (void *)&CurrentScale.Y, Weight);
-	GeneratedTotalKeys[8]->ModifyByCurrentAndWeight(Proxy, KeyTime, (void *)&CurrentScale.Z, Weight);
-	return true;
 
+		FMovieSceneChannelProxy& Proxy = SectionToKey->GetChannelProxy();
+		GeneratedTotalKeys[0]->ModifyByCurrentAndWeight(Proxy, KeyTime, (void *)&CurrentPos.X, Weight);
+		GeneratedTotalKeys[1]->ModifyByCurrentAndWeight(Proxy, KeyTime, (void *)&CurrentPos.Y, Weight);
+		GeneratedTotalKeys[2]->ModifyByCurrentAndWeight(Proxy, KeyTime, (void *)&CurrentPos.Z, Weight);
+		GeneratedTotalKeys[3]->ModifyByCurrentAndWeight(Proxy, KeyTime, (void *)&CurrentRot.Roll, Weight);
+		GeneratedTotalKeys[4]->ModifyByCurrentAndWeight(Proxy, KeyTime, (void *)&CurrentRot.Pitch, Weight);
+		GeneratedTotalKeys[5]->ModifyByCurrentAndWeight(Proxy, KeyTime, (void *)&CurrentRot.Yaw, Weight);
+		GeneratedTotalKeys[6]->ModifyByCurrentAndWeight(Proxy, KeyTime, (void *)&CurrentScale.X, Weight);
+		GeneratedTotalKeys[7]->ModifyByCurrentAndWeight(Proxy, KeyTime, (void *)&CurrentScale.Y, Weight);
+		GeneratedTotalKeys[8]->ModifyByCurrentAndWeight(Proxy, KeyTime, (void *)&CurrentScale.Z, Weight);
+		return true;
+	}
+	return false;
 }
 
 void AddUnwoundKey(FMovieSceneFloatChannel& Channel, FFrameNumber Time, float Value)
diff --git a/Engine/Source/Runtime/MovieScene/Private/Evaluation/MovieScenePropertyTemplate.cpp b/Engine/Source/Runtime/MovieScene/Private/Evaluation/MovieScenePropertyTemplate.cpp
--- a/Engine/Source/Runtime/MovieScene/Private/Evaluation/MovieScenePropertyTemplate.cpp
+++ b/Engine/Source/Runtime/MovieScene/Private/Evaluation/MovieScenePropertyTemplate.cpp
@@ -40,16 +40,7 @@ const FMovieSceneInterrogationKey FMovieScenePropertySectionTemplate::GetInt32In
 	const static FMovieSceneAnimTypeID TypeID = FMovieSceneAnimTypeID::Unique();
 	return TypeID;
 }
-const FMovieSceneInterrogationKey FMovieScenePropertySectionTemplate::GetTransformInterrogationKey()
-{
-	const static FMovieSceneAnimTypeID TypeID = FMovieSceneAnimTypeID::Unique();
-	return TypeID;
-}
-const FMovieSceneInterrogationKey FMovieScenePropertySectionTemplate::GetEulerTransformInterrogationKey()
-{
-	const static FMovieSceneAnimTypeID TypeID = FMovieSceneAnimTypeID::Unique();
-	return TypeID;
-}
+
 const FMovieSceneInterrogationKey FMovieScenePropertySectionTemplate::GetVector4InterrogationKey()
 {
 	const static FMovieSceneAnimTypeID TypeID = FMovieSceneAnimTypeID::Unique();
diff --git a/Engine/Source/Runtime/MovieScene/Private/MovieSceneTemplateInterrogation.cpp b/Engine/Source/Runtime/MovieScene/Private/MovieSceneTemplateInterrogation.cpp
--- a/Engine/Source/Runtime/MovieScene/Private/MovieSceneTemplateInterrogation.cpp
+++ b/Engine/Source/Runtime/MovieScene/Private/MovieSceneTemplateInterrogation.cpp
@@ -20,3 +20,9 @@ FMovieSceneBlendingAccumulator& FMovieSceneInterrogationData::GetAccumulator()
 
 	return *Accumulator;
 }
+
+FMovieSceneInterrogationKey FMovieSceneInterrogationKey::GetTransformInterrogationKey()
+{
+	static FMovieSceneAnimTypeID TypeID = FMovieSceneAnimTypeID::Unique();
+	return TypeID;
+}
diff --git a/Engine/Source/Runtime/MovieScene/Public/Compilation/MovieSceneTemplateInterrogation.h b/Engine/Source/Runtime/MovieScene/Public/Compilation/MovieSceneTemplateInterrogation.h
--- a/Engine/Source/Runtime/MovieScene/Public/Compilation/MovieSceneTemplateInterrogation.h
+++ b/Engine/Source/Runtime/MovieScene/Public/Compilation/MovieSceneTemplateInterrogation.h
@@ -29,6 +29,12 @@ struct FMovieSceneInterrogationKey
 	{
 		return A.AnimTypeID != B.AnimTypeID;
 	}
+
+	/**
+	* Access the interrogation key for transform data - any interrogation data stored with this key is guaranteed to be of type 'FTransform'
+	* to be used by various editors and sections that reference transforms
+	*/
+	MOVIESCENE_API static FMovieSceneInterrogationKey GetTransformInterrogationKey();
 };
 
 /** Tokens that are stored in FMovieSceneInterrogationData */
@@ -161,3 +167,5 @@ private:
 	TOptional<FMovieSceneInterrogationKey> PredicateKey;
 	int32 Index;
 };
+
+
diff --git a/Engine/Source/Runtime/MovieScene/Public/Evaluation/MovieScenePropertyTemplate.h b/Engine/Source/Runtime/MovieScene/Public/Evaluation/MovieScenePropertyTemplate.h
--- a/Engine/Source/Runtime/MovieScene/Public/Evaluation/MovieScenePropertyTemplate.h
+++ b/Engine/Source/Runtime/MovieScene/Public/Evaluation/MovieScenePropertyTemplate.h
@@ -19,6 +19,7 @@
 #include "Evaluation/Blending/MovieSceneBlendingActuator.h"
 #include "Evaluation/Blending/MovieSceneMultiChannelBlending.h"
 #include "Compilation/MovieSceneTemplateInterrogation.h"
+#include "TransformData.h"
 #include "MovieScenePropertyTemplate.generated.h"
 
 class UMovieScenePropertyTrack;
@@ -236,8 +237,6 @@ public:
 	//use thse keys for setting and iterating the correct types.
 	const static FMovieSceneInterrogationKey GetFloatInterrogationKey();
 	const static FMovieSceneInterrogationKey GetInt32InterrogationKey();
-	const static FMovieSceneInterrogationKey GetTransformInterrogationKey();
-	const static FMovieSceneInterrogationKey GetEulerTransformInterrogationKey();
 	const static FMovieSceneInterrogationKey GetVector4InterrogationKey();
 	const static FMovieSceneInterrogationKey GetVectorInterrogationKey();
 	const static FMovieSceneInterrogationKey GetVector2DInterrogationKey();
@@ -297,15 +296,21 @@ inline void TPropertyActuator<FVector2D>::Actuate(FMovieSceneInterrogationData&
 template<>
 inline void TPropertyActuator<FEulerTransform>::Actuate(FMovieSceneInterrogationData& InterrogationData, typename TCallTraits<FEulerTransform>::ParamType InValue, const TBlendableTokenStack<FEulerTransform>& OriginalStack, const FMovieSceneContext& Context) const
 {
-	FEulerTransform Value = InValue;
-	InterrogationData.Add(Value, FMovieScenePropertySectionTemplate::GetEulerTransformInterrogationKey());
+	FTransformData Value;
+	Value.Translation = InValue.Location;
+	Value.Rotation = InValue.Rotation;
+	Value.Scale = InValue.Scale;
+	InterrogationData.Add(Value, FMovieSceneInterrogationKey::GetTransformInterrogationKey());
 };
 
 template<>
 inline void TPropertyActuator<FTransform>::Actuate(FMovieSceneInterrogationData& InterrogationData, typename TCallTraits<FTransform>::ParamType InValue, const TBlendableTokenStack<FTransform>& OriginalStack, const FMovieSceneContext& Context) const
 {
-	FTransform Value = InValue;
-	InterrogationData.Add(Value, FMovieScenePropertySectionTemplate::GetTransformInterrogationKey());
+	FTransformData Value;
+	Value.Translation = InValue.GetTranslation();
+	Value.Rotation = InValue.GetRotation().Rotator();
+	Value.Scale = InValue.GetScale3D();
+	InterrogationData.Add(Value, FMovieSceneInterrogationKey::GetTransformInterrogationKey());
 };
 
 template<>
diff --git a/Engine/Source/Runtime/MovieScene/Public/TransformData.h b/Engine/Source/Runtime/MovieScene/Public/TransformData.h
new file mode 100644
--- /dev/null
+++ b/Engine/Source/Runtime/MovieScene/Public/TransformData.h
@@ -0,0 +1,36 @@
+// Copyright 1998-2019 Epic Games, Inc. All Rights Reserved.
+
+#pragma once
+
+#include "CoreMinimal.h"
+/**
+* Stores information about a transform for the purpose of adding keys to a transform section
+*/
+struct MOVIESCENE_API FTransformData
+{
+	/** Translation component */
+	FVector Translation;
+	/** Rotation component */
+	FRotator Rotation;
+	/** Scale component */
+	FVector Scale;
+
+	FTransformData()
+		: Translation(ForceInitToZero)
+		, Rotation(ForceInitToZero)
+		, Scale(ForceInitToZero)
+	{}
+
+	/**
+	* Constructor.  Builds the data from a scene component
+	* Uses relative transform only
+	*
+	* @param InComponent	The component to build from
+	*/
+	FTransformData(const USceneComponent* InComponent)
+		: Translation(InComponent->RelativeLocation)
+		, Rotation(InComponent->RelativeRotation)
+		, Scale(InComponent->RelativeScale3D)
+	{}
+
+};
diff --git a/Engine/Source/Runtime/MovieSceneTracks/Private/Sections/MovieScene3DTransformSection.cpp b/Engine/Source/Runtime/MovieSceneTracks/Private/Sections/MovieScene3DTransformSection.cpp
--- a/Engine/Source/Runtime/MovieSceneTracks/Private/Sections/MovieScene3DTransformSection.cpp
+++ b/Engine/Source/Runtime/MovieSceneTracks/Private/Sections/MovieScene3DTransformSection.cpp
@@ -728,6 +728,5 @@ void UMovieScene3DTransformSection::SetBlendType(EMovieSceneBlendType InBlendTyp
 
 FMovieSceneInterrogationKey UMovieScene3DTransformSection::GetInterrogationKey()
 {
-	static FMovieSceneAnimTypeID TypeID = FMovieSceneAnimTypeID::Unique();
-	return TypeID;
+	return FMovieSceneInterrogationKey::GetTransformInterrogationKey();
 }
diff --git a/Engine/Source/Runtime/MovieSceneTracks/Private/Tracks/MovieSceneEulerTransformTrack.cpp b/Engine/Source/Runtime/MovieSceneTracks/Private/Tracks/MovieSceneEulerTransformTrack.cpp
--- a/Engine/Source/Runtime/MovieSceneTracks/Private/Tracks/MovieSceneEulerTransformTrack.cpp
+++ b/Engine/Source/Runtime/MovieSceneTracks/Private/Tracks/MovieSceneEulerTransformTrack.cpp
@@ -35,8 +35,3 @@ FMovieSceneEvalTemplatePtr UMovieSceneEulerTransformTrack::CreateTemplateForSect
 	return FMovieSceneEulerTransformPropertySectionTemplate(*CastChecked<UMovieScene3DTransformSection>(&InSection), *this);
 }
 
-FMovieSceneInterrogationKey UMovieSceneEulerTransformTrack::GetInterrogationKey()
-{
-	static FMovieSceneAnimTypeID TypeID = FMovieSceneAnimTypeID::Unique();
-	return TypeID;
-}
diff --git a/Engine/Source/Runtime/MovieSceneTracks/Public/Sections/MovieScene3DTransformSection.h b/Engine/Source/Runtime/MovieSceneTracks/Public/Sections/MovieScene3DTransformSection.h
--- a/Engine/Source/Runtime/MovieSceneTracks/Public/Sections/MovieScene3DTransformSection.h
+++ b/Engine/Source/Runtime/MovieSceneTracks/Public/Sections/MovieScene3DTransformSection.h
@@ -9,6 +9,7 @@
 #include "MovieSceneSection.h"
 #include "MovieSceneKeyStruct.h"
 #include "Channels/MovieSceneFloatChannel.h"
+#include "TransformData.h"
 #include "MovieScene3DTransformSection.generated.h"
 
 
@@ -24,36 +25,6 @@ enum class EShow3DTrajectory : uint8
 #endif
 
 
-/**
-* Stores information about a transform for the purpose of adding keys to a transform section
-*/
-struct FTransformData
-{
-	/** Translation component */
-	FVector Translation;
-	/** Rotation component */
-	FRotator Rotation;
-	/** Scale component */
-	FVector Scale;
-
-	FTransformData()
-		: Translation( ForceInitToZero )
-		, Rotation( ForceInitToZero )
-		, Scale( ForceInitToZero )
-	{}
-
-	/**
-	* Constructor.  Builds the data from a scene component
-	* Uses relative transform only
-	*
-	* @param InComponent	The component to build from
-	*/
-	FTransformData( const USceneComponent* InComponent )
-		: Translation( InComponent->RelativeLocation )
-		, Rotation( InComponent->RelativeRotation )
-		, Scale( InComponent->RelativeScale3D )
-	{}
-};
 
 /**
  * Proxy structure for translation keys in 3D transform sections.
diff --git a/Engine/Source/Runtime/MovieSceneTracks/Public/Tracks/MovieSceneEulerTransformTrack.h b/Engine/Source/Runtime/MovieSceneTracks/Public/Tracks/MovieSceneEulerTransformTrack.h
--- a/Engine/Source/Runtime/MovieSceneTracks/Public/Tracks/MovieSceneEulerTransformTrack.h
+++ b/Engine/Source/Runtime/MovieSceneTracks/Public/Tracks/MovieSceneEulerTransformTrack.h
@@ -7,8 +7,6 @@
 #include "Tracks/MovieScenePropertyTrack.h"
 #include "MovieSceneEulerTransformTrack.generated.h"
 
-struct FMovieSceneInterrogationKey;
-
 /**
  * Handles manipulation of 3D euler transform properties in a movie scene
  */
@@ -26,9 +24,5 @@ public:
 	virtual UMovieSceneSection* CreateNewSection() override;
 	virtual FMovieSceneEvalTemplatePtr CreateTemplateForSection(const UMovieSceneSection& InSection) const override;
 
-	/**
-	 * Access the interrogation key for transform data - any interrogation data stored with this key is guaranteed to be of type 'FEulerTransform'
-	 */
-	MOVIESCENETRACKS_API static FMovieSceneInterrogationKey GetInterrogationKey();
 };
 
