Sequencer: Remove redundant level visibility text in the section (ie. Visible, Hidden) so that it doesn't clutter up when there's a lot of levels and sections

diff --git a/Engine/Source/Editor/MovieSceneTools/Private/Sections/LevelVisibilitySection.cpp b/Engine/Source/Editor/MovieSceneTools/Private/Sections/LevelVisibilitySection.cpp
--- a/Engine/Source/Editor/MovieSceneTools/Private/Sections/LevelVisibilitySection.cpp
+++ b/Engine/Source/Editor/MovieSceneTools/Private/Sections/LevelVisibilitySection.cpp
@@ -82,9 +82,14 @@ FText FLevelVisibilitySection::GetVisibilityText() const
 		LevelsText.Append( FString::Format(TEXT(" (+{0} more)"), { FString::FormatAsNumber(NumRemaining) } ) );
 	}
 
-	FText VisibilityText = SectionObject.GetVisibility() == ELevelVisibility::Visible ? VisibleText : HiddenText;
-	
-	return FText::Format( NSLOCTEXT( "LevelVisibilitySection", "SectionTextFormat", "{0}: {1}" ), VisibilityText, FText::FromString( LevelsText ) );
+	if (LevelsText.IsEmpty())
+	{
+		FText VisibilityText = SectionObject.GetVisibility() == ELevelVisibility::Visible ? VisibleText : HiddenText;
+
+		return VisibilityText;
+	}
+
+	return FText::Format( NSLOCTEXT( "LevelVisibilitySection", "SectionTextFormat", "{0}" ), FText::FromString( LevelsText ) );
 }
 
 FText FLevelVisibilitySection::GetVisibilityToolTip() const
