Sequencer: Fix attach getting existing binding ID

diff --git a/Engine/Source/Editor/MovieSceneTools/Private/TrackEditors/ActorPickerTrackEditor.cpp b/Engine/Source/Editor/MovieSceneTools/Private/TrackEditors/ActorPickerTrackEditor.cpp
--- a/Engine/Source/Editor/MovieSceneTools/Private/TrackEditors/ActorPickerTrackEditor.cpp
+++ b/Engine/Source/Editor/MovieSceneTools/Private/TrackEditors/ActorPickerTrackEditor.cpp
@@ -264,7 +264,16 @@ void FActorPickerTrackEditor::ExistingBindingPicked(FMovieSceneObjectBindingID E
 {
 	TSharedPtr<ISequencer> SequencerPtr = GetSequencer();
 
-	TArrayView<TWeakObjectPtr<UObject>> RuntimeObjects = SequencerPtr->FindBoundObjects(ExistingBindingID.GetGuid(), ExistingBindingID.GetSequenceID());
+	FMovieSceneSequenceID SequenceID = SequencerPtr->GetFocusedTemplateID();
+
+	if (ExistingBindingID.IsValid())
+	{
+		// Ensure that this ID is resolvable from the root, based on the current local sequence ID
+		FMovieSceneObjectBindingID RootBindingID = ExistingBindingID.ResolveLocalToRoot(SequenceID, SequencerPtr->GetEvaluationTemplate().GetHierarchy());
+		SequenceID = RootBindingID.GetSequenceID();
+	}
+
+	TArrayView<TWeakObjectPtr<UObject>> RuntimeObjects = SequencerPtr->FindBoundObjects(ExistingBindingID.GetGuid(), SequenceID);
 	for (auto RuntimeObject : RuntimeObjects)
 	{
 		if (RuntimeObject.IsValid())
