Sequencer: Fix GLEOs in attach track

diff --git a/Engine/Source/Editor/MovieSceneTools/Private/TrackEditors/AttachTrackEditor.cpp b/Engine/Source/Editor/MovieSceneTools/Private/TrackEditors/AttachTrackEditor.cpp
--- a/Engine/Source/Editor/MovieSceneTools/Private/TrackEditors/AttachTrackEditor.cpp
+++ b/Engine/Source/Editor/MovieSceneTools/Private/TrackEditors/AttachTrackEditor.cpp
@@ -910,9 +910,9 @@ void F3DAttachTrackEditor::TrimAndPreserve(FGuid InObjectBinding, UMovieSceneSec
 
 		FLocalTransformEvaluator LocalTransformEval(GetSequencer(), Object, EvalTrack);
 
-		if (AttachSection->ReAttachOnDetach)
+		if (AttachSection->ReAttachOnDetach.IsValid())
 		{
-			FWorldTransformEvaluator ReAttachParentEvaluator(GetSequencer(), AttachSection->ReAttachOnDetach);
+			FWorldTransformEvaluator ReAttachParentEvaluator(GetSequencer(), AttachSection->ReAttachOnDetach.Get());
 
 			CompensateChildTrack(ExcludedRange, Channels, TOptional<TArrayView<FMovieSceneFloatChannel*>>(), ReAttachParentEvaluator, LocalTransformEval, ETransformPreserveType::CurrentKey, RevertModifier);
 		}
diff --git a/Engine/Source/Runtime/MovieSceneTracks/Public/Sections/MovieScene3DAttachSection.h b/Engine/Source/Runtime/MovieSceneTracks/Public/Sections/MovieScene3DAttachSection.h
--- a/Engine/Source/Runtime/MovieSceneTracks/Public/Sections/MovieScene3DAttachSection.h
+++ b/Engine/Source/Runtime/MovieSceneTracks/Public/Sections/MovieScene3DAttachSection.h
@@ -44,7 +44,7 @@ public:
 	bool bFullRevertOnDetach;
 
 	UPROPERTY(BlueprintReadWrite, EditAnywhere, Category="Attach")
-	AActor* ReAttachOnDetach;
+	TSoftObjectPtr<AActor> ReAttachOnDetach;
 
 #endif
 
