Sequencer: Event usability overhaul

  - It is now possible to define payload parameters for events again
  - Creating a new event endpoint will leverage custom event nodes rather than separate functions. This is now possible due to the added flexibility that generated function graphs gives us.
  - Users can now also specify which pin to pass the bound object through (if any) - previous rules apply (only object and interface pins supported)
  - Reworked quick-binding menu to use the blueprint graph action menu

diff --git a/Engine/Plugins/MovieScene/ActorSequence/Source/ActorSequenceEditor/Private/MovieSceneSequenceEditor_ActorSequence.h b/Engine/Plugins/MovieScene/ActorSequence/Source/ActorSequenceEditor/Private/MovieSceneSequenceEditor_ActorSequence.h
--- a/Engine/Plugins/MovieScene/ActorSequence/Source/ActorSequenceEditor/Private/MovieSceneSequenceEditor_ActorSequence.h
+++ b/Engine/Plugins/MovieScene/ActorSequence/Source/ActorSequenceEditor/Private/MovieSceneSequenceEditor_ActorSequence.h
@@ -44,22 +44,4 @@ struct FMovieSceneSequenceEditor_ActorSequence : FMovieSceneSequenceEditor
 		bool bDontCreateNewBlueprint = false;
 		return Level ? Level->GetLevelScriptBlueprint(bDontCreateNewBlueprint) : nullptr;
 	}
-
-	virtual void SetupDefaultPinForEndpoint(UMovieSceneEventTrack* EventTrack, UK2Node_FunctionEntry* Endpoint) const override
-	{
-		// By default for master tracks we create a pin (that receives the level blueprint for master event tracks, or the event receivers)
-		UClass* PinClass = nullptr;
-
-		// When event receivers are not set, we set the pin type to the type of the object binding
-		if (EventTrack->EventReceivers.Num() == 0)
-		{
-			PinClass = FindTrackObjectBindingClass(EventTrack);
-		}
-
-		FEdGraphPinType PinType;
-		PinType.PinCategory = UEdGraphSchema_K2::PC_Object;
-		PinType.PinSubCategoryObject = PinClass ? PinClass : UObject::StaticClass();
-
-		Endpoint->CreateUserDefinedPin(TargetPinName, PinType, EGPD_Output, true);
-	}
 };
\ No newline at end of file
diff --git a/Engine/Plugins/MovieScene/LevelSequenceEditor/Source/LevelSequenceEditor/Private/Misc/MovieSceneSequenceEditor_LevelSequence.h b/Engine/Plugins/MovieScene/LevelSequenceEditor/Source/LevelSequenceEditor/Private/Misc/MovieSceneSequenceEditor_LevelSequence.h
--- a/Engine/Plugins/MovieScene/LevelSequenceEditor/Source/LevelSequenceEditor/Private/Misc/MovieSceneSequenceEditor_LevelSequence.h
+++ b/Engine/Plugins/MovieScene/LevelSequenceEditor/Source/LevelSequenceEditor/Private/Misc/MovieSceneSequenceEditor_LevelSequence.h
@@ -38,22 +38,4 @@ struct FMovieSceneSequenceEditor_LevelSequence : FMovieSceneSequenceEditor
 		LevelSequence->SetDirectorBlueprint(Blueprint);
 		return Blueprint;
 	}
-
-	virtual void SetupDefaultPinForEndpoint(UMovieSceneEventTrack* EventTrack, UK2Node_FunctionEntry* Endpoint) const override
-	{
-		// By default for master tracks we create a pin (that receives the level blueprint for master event tracks, or the event receivers)
-		UClass* PinClass = nullptr;
-
-		// When event receivers are not set, we set the pin type to the type of the object binding
-		if (EventTrack->EventReceivers.Num() == 0)
-		{
-			PinClass = FindTrackObjectBindingClass(EventTrack);
-		}
-
-		FEdGraphPinType PinType;
-		PinType.PinCategory = UEdGraphSchema_K2::PC_Object;
-		PinType.PinSubCategoryObject = PinClass ? PinClass : UObject::StaticClass();
-
-		Endpoint->CreateUserDefinedPin(TargetPinName, PinType, EGPD_Output, true);
-	}
 };
\ No newline at end of file
diff --git a/Engine/Source/Editor/MovieSceneTools/MovieSceneTools.Build.cs b/Engine/Source/Editor/MovieSceneTools/MovieSceneTools.Build.cs
--- a/Engine/Source/Editor/MovieSceneTools/MovieSceneTools.Build.cs
+++ b/Engine/Source/Editor/MovieSceneTools/MovieSceneTools.Build.cs
@@ -43,6 +43,7 @@ public class MovieSceneTools : ModuleRules
 				"MovieScene",
 				"MovieSceneTracks",
 				"BlueprintGraph",
+				"Kismet",
 				"KismetCompiler",
                 "GraphEditor",
                 "ContentBrowser",
diff --git a/Engine/Source/Editor/MovieSceneTools/Private/Channels/BuiltInChannelEditors.cpp b/Engine/Source/Editor/MovieSceneTools/Private/Channels/BuiltInChannelEditors.cpp
--- a/Engine/Source/Editor/MovieSceneTools/Private/Channels/BuiltInChannelEditors.cpp
+++ b/Engine/Source/Editor/MovieSceneTools/Private/Channels/BuiltInChannelEditors.cpp
@@ -1,6 +1,10 @@
 // Copyright 1998-2019 Epic Games, Inc. All Rights Reserved.
 
 #include "BuiltInChannelEditors.h"
+#include "MovieSceneSequence.h"
+#include "MovieSceneSequenceEditor.h"
+#include "MovieSceneEventUtils.h"
+#include "Sections/MovieSceneEventSectionBase.h"
 #include "ISequencerChannelInterface.h"
 #include "Widgets/SNullWidget.h"
 #include "ISequencer.h"
@@ -339,7 +343,7 @@ void PostConstructKeyInstance(const TMovieSceneChannelHandle<FMovieSceneObjectPa
 	KeyStruct->OnPropertyChangedEvent = CopyInstanceToKeyLambda;
 }
 
-void DrawKeys(FMovieSceneFloatChannel* Channel, TArrayView<const FKeyHandle> InKeyHandles, TArrayView<FKeyDrawParams> OutKeyDrawParams)
+void DrawKeys(FMovieSceneFloatChannel* Channel, TArrayView<const FKeyHandle> InKeyHandles, const UMovieSceneSection* InOwner, TArrayView<FKeyDrawParams> OutKeyDrawParams)
 {
 	static const FName CircleKeyBrushName("Sequencer.KeyCircle");
 	static const FName DiamondKeyBrushName("Sequencer.KeyDiamond");
@@ -403,7 +407,7 @@ void DrawKeys(FMovieSceneFloatChannel* Channel, TArrayView<const FKeyHandle> InK
 	}
 }
 
-void DrawKeys(FMovieSceneParticleChannel* Channel, TArrayView<const FKeyHandle> InKeyHandles, TArrayView<FKeyDrawParams> OutKeyDrawParams)
+void DrawKeys(FMovieSceneParticleChannel* Channel, TArrayView<const FKeyHandle> InKeyHandles, const UMovieSceneSection* InOwner, TArrayView<FKeyDrawParams> OutKeyDrawParams)
 {
 	static const FName KeyLeftBrushName("Sequencer.KeyLeft");
 	static const FName KeyRightBrushName("Sequencer.KeyRight");
@@ -442,8 +446,10 @@ void DrawKeys(FMovieSceneParticleChannel* Channel, TArrayView<const FKeyHandle>
 	}
 }
 
-void DrawKeys(FMovieSceneEventChannel* Channel, TArrayView<const FKeyHandle> InKeyHandles, TArrayView<FKeyDrawParams> OutKeyDrawParams)
+void DrawKeys(FMovieSceneEventChannel* Channel, TArrayView<const FKeyHandle> InKeyHandles, const UMovieSceneSection* InOwner, TArrayView<FKeyDrawParams> OutKeyDrawParams)
 {
+	UMovieSceneEventSectionBase* EventSection = CastChecked<UMovieSceneEventSectionBase>(const_cast<UMovieSceneSection*>(InOwner));
+
 	FKeyDrawParams ValidEventParams, InvalidEventParams;
 
 	ValidEventParams.BorderBrush   = ValidEventParams.FillBrush   = FEditorStyle::Get().GetBrush("Sequencer.KeyDiamond");
@@ -453,12 +459,17 @@ void DrawKeys(FMovieSceneEventChannel* Channel, TArrayView<const FKeyHandle> InK
 	InvalidEventParams.FillTint    = FLinearColor(1.f,1.f,1.f,.2f);
 
 	TMovieSceneChannelData<FMovieSceneEvent> ChannelData = Channel->GetData();
-	TArrayView<const FMovieSceneEvent>       Events = ChannelData.GetValues();
+	TArrayView<FMovieSceneEvent>             Events = ChannelData.GetValues();
+
+	UMovieSceneSequence*       Sequence           = InOwner->GetTypedOuter<UMovieSceneSequence>();
+	FMovieSceneSequenceEditor* SequenceEditor     = Sequence ? FMovieSceneSequenceEditor::Find(Sequence) : nullptr;
+	UBlueprint*                SequenceDirectorBP = SequenceEditor ? SequenceEditor->FindDirectorBlueprint(Sequence) : nullptr;
 
 	for (int32 Index = 0; Index < InKeyHandles.Num(); ++Index)
 	{
 		int32 KeyIndex = ChannelData.GetIndex(InKeyHandles[Index]);
-		if (KeyIndex != INDEX_NONE && Events[KeyIndex].IsBoundToBlueprint())
+
+		if (KeyIndex != INDEX_NONE && SequenceDirectorBP && FMovieSceneEventUtils::FindEndpoint(&Events[KeyIndex], EventSection, SequenceDirectorBP))
 		{
 			OutKeyDrawParams[Index] = ValidEventParams;
 		}
diff --git a/Engine/Source/Editor/MovieSceneTools/Private/Channels/BuiltInChannelEditors.h b/Engine/Source/Editor/MovieSceneTools/Private/Channels/BuiltInChannelEditors.h
--- a/Engine/Source/Editor/MovieSceneTools/Private/Channels/BuiltInChannelEditors.h
+++ b/Engine/Source/Editor/MovieSceneTools/Private/Channels/BuiltInChannelEditors.h
@@ -58,9 +58,9 @@ UMovieSceneKeyStructType* InstanceGeneratedStruct(FMovieSceneObjectPathChannel*
 void PostConstructKeyInstance(const TMovieSceneChannelHandle<FMovieSceneObjectPathChannel>& ChannelHandle, FKeyHandle InHandle, FStructOnScope* Struct);
 
 /** Key drawing overrides */
-void DrawKeys(FMovieSceneFloatChannel*    Channel, TArrayView<const FKeyHandle> InKeyHandles, TArrayView<FKeyDrawParams> OutKeyDrawParams);
-void DrawKeys(FMovieSceneParticleChannel* Channel, TArrayView<const FKeyHandle> InKeyHandles, TArrayView<FKeyDrawParams> OutKeyDrawParams);
-void DrawKeys(FMovieSceneEventChannel*    Channel, TArrayView<const FKeyHandle> InKeyHandles, TArrayView<FKeyDrawParams> OutKeyDrawParams);
+void DrawKeys(FMovieSceneFloatChannel*    Channel, TArrayView<const FKeyHandle> InKeyHandles, const UMovieSceneSection* InOwner, TArrayView<FKeyDrawParams> OutKeyDrawParams);
+void DrawKeys(FMovieSceneParticleChannel* Channel, TArrayView<const FKeyHandle> InKeyHandles, const UMovieSceneSection* InOwner, TArrayView<FKeyDrawParams> OutKeyDrawParams);
+void DrawKeys(FMovieSceneEventChannel*    Channel, TArrayView<const FKeyHandle> InKeyHandles, const UMovieSceneSection* InOwner, TArrayView<FKeyDrawParams> OutKeyDrawParams);
 
 /** Context menu overrides */
 void ExtendSectionMenu(FMenuBuilder& OuterMenuBuilder, TArray<TMovieSceneChannelHandle<FMovieSceneFloatChannel>>&& Channels, TArrayView<UMovieSceneSection* const> Sections, TWeakPtr<ISequencer> InSequencer);
diff --git a/Engine/Source/Editor/MovieSceneTools/Private/Channels/EventChannelCurveModel.cpp b/Engine/Source/Editor/MovieSceneTools/Private/Channels/EventChannelCurveModel.cpp
--- a/Engine/Source/Editor/MovieSceneTools/Private/Channels/EventChannelCurveModel.cpp
+++ b/Engine/Source/Editor/MovieSceneTools/Private/Channels/EventChannelCurveModel.cpp
@@ -64,20 +64,12 @@ void FEventChannelCurveModel::AddKeys(TArrayView<const FKeyPosition> InKeyPositi
 			Section->ExpandToFrame(Time);
 
 			FMovieSceneEvent Value;
-// 
-// 			if (Attributes.HasInterpMode())    { Value.InterpMode            = Attributes.GetInterpMode();    }
-// 			if (Attributes.HasTangentMode())   { Value.TangentMode           = Attributes.GetTangentMode();   }
-// 			if (Attributes.HasArriveTangent()) { Value.Tangent.ArriveTangent = Attributes.GetArriveTangent(); }
-// 			if (Attributes.HasLeaveTangent())  { Value.Tangent.LeaveTangent  = Attributes.GetLeaveTangent();  }
-
 			int32 KeyIndex = ChannelData.AddKey(Time, Value);
 			if (OutKeyHandles)
 			{
 				(*OutKeyHandles)[Index] = ChannelData.GetHandle(KeyIndex);
 			}
 		}
-
-		// Channel->AutoSetTangents();
 	}
 }
 
@@ -200,189 +192,10 @@ void FEventChannelCurveModel::SetKeyPositions(TArrayView<const FKeyHandle> InKey
 
 void FEventChannelCurveModel::GetKeyAttributes(TArrayView<const FKeyHandle> InKeys, TArrayView<FKeyAttributes> OutAttributes) const
 {
-// 	FMovieSceneEventChannel* Channel = ChannelHandle.Get();
-// 	UMovieSceneSection*      Section = WeakSection.Get();
-// 	if (Channel && Section)
-// 	{
-// 		TMovieSceneChannelData<FMovieSceneEvent> ChannelData = Channel->GetData();
-// 		TArrayView<const FFrameNumber>    Times  = ChannelData.GetTimes();
-// 		TArrayView<FMovieSceneEvent> Values = ChannelData.GetValues();
-// 
-// 		float TimeInterval = Section->GetTypedOuter<UMovieScene>()->GetTickResolution().AsInterval();
-// 
-// 		for (int32 Index = 0; Index < InKeys.Num(); ++Index)
-// 		{
-// 			const int32 KeyIndex = ChannelData.GetIndex(InKeys[Index]);
-// 			if (KeyIndex != INDEX_NONE)
-// 			{
-// 				const FMovieSceneEvent& KeyValue    = Values[KeyIndex];
-// 				FKeyAttributes&              Attributes  = OutAttributes[Index];
-// 
-// 				Attributes.SetInterpMode(KeyValue.InterpMode);
-// 
-// 				if (KeyValue.InterpMode != RCIM_Constant && KeyValue.InterpMode != RCIM_Linear)
-// 				{
-// 					Attributes.SetTangentMode(KeyValue.TangentMode);
-// 					if (KeyIndex != 0)
-// 					{
-// 						Attributes.SetArriveTangent(KeyValue.Tangent.ArriveTangent / TimeInterval);
-// 					}
-// 
-// 					if (KeyIndex != Times.Num()-1)
-// 					{
-// 						Attributes.SetLeaveTangent(KeyValue.Tangent.LeaveTangent / TimeInterval);
-// 					}
-// 					if (KeyValue.InterpMode == RCIM_Cubic)
-// 					{
-// 						Attributes.SetTangentWeightMode(KeyValue.Tangent.TangentWeightMode);
-// 						if (KeyValue.Tangent.TangentWeightMode != RCTWM_WeightedNone)
-// 						{
-// 							Attributes.SetArriveTangentWeight(KeyValue.Tangent.ArriveTangentWeight);
-// 							Attributes.SetLeaveTangentWeight(KeyValue.Tangent.LeaveTangentWeight);
-// 						}
-// 					}
-// 				}
-// 			}
-// 		}
-// 	}
 }
 
 void FEventChannelCurveModel::SetKeyAttributes(TArrayView<const FKeyHandle> InKeys, TArrayView<const FKeyAttributes> InAttributes)
 {
-// 	FMovieSceneEventChannel* Channel = ChannelHandle.Get();
-// 	UMovieSceneSection*      Section = WeakSection.Get();
-// 	if (Channel && Section)
-// 	{
-// 		bool bAutoSetTangents = false;
-// 		Section->MarkAsChanged();
-// 
-// 		TMovieSceneChannelData<FMovieSceneEvent> ChannelData = Channel->GetData();
-// 		TArrayView<FMovieSceneEvent> Values = ChannelData.GetValues();
-// 
-// 		FFrameRate TickResolution = Section->GetTypedOuter<UMovieScene>()->GetTickResolution();
-// 		float TimeInterval = TickResolution.AsInterval();
-// 
-// 		for (int32 Index = 0; Index < InKeys.Num(); ++Index)
-// 		{
-// 			const int32 KeyIndex = ChannelData.GetIndex(InKeys[Index]);
-// 			if (KeyIndex != INDEX_NONE)
-// 			{
-// 				const FKeyAttributes&  Attributes = InAttributes[Index];
-// 				FMovieSceneEvent& KeyValue   = Values[KeyIndex];
-// 				if (Attributes.HasInterpMode())    { KeyValue.InterpMode  = Attributes.GetInterpMode();  bAutoSetTangents = true; }
-// 				if (Attributes.HasTangentMode())
-// 				{
-// 					KeyValue.TangentMode = Attributes.GetTangentMode();
-// 					if (KeyValue.TangentMode == RCTM_Auto)
-// 					{
-// 						KeyValue.Tangent.TangentWeightMode = RCTWM_WeightedNone;
-// 					}
-// 					bAutoSetTangents = true;
-// 				}
-// 				if (Attributes.HasTangentWeightMode()) 
-// 				{ 
-// 					if (KeyValue.Tangent.TangentWeightMode == RCTWM_WeightedNone) //set tangent weights to default use
-// 					{
-// 						TArrayView<const FFrameNumber> Times = Channel->GetTimes();
-// 						const float OneThird = 1.0f / 3.0f;
-// 
-// 						//calculate a tangent weight based upon tangent and time difference
-// 						//calculate arrive tangent weight
-// 						if (KeyIndex > 0 )
-// 						{
-// 							const float X = TickResolution.AsSeconds(Times[KeyIndex].Value - Times[KeyIndex - 1].Value);
-// 							const float ArriveTangentNormal = KeyValue.Tangent.ArriveTangent / (TimeInterval);
-// 							const float Y = ArriveTangentNormal * X;
-// 							KeyValue.Tangent.ArriveTangentWeight = FMath::Sqrt(X*X + Y * Y) * OneThird;
-// 						}
-// 						//calculate leave weight
-// 						if(KeyIndex < ( Times.Num() - 1))
-// 						{
-// 							const float X = TickResolution.AsSeconds(Times[KeyIndex].Value - Times[KeyIndex + 1].Value);
-// 							const float LeaveTangentNormal = KeyValue.Tangent.LeaveTangent / (TimeInterval);
-// 							const float Y = LeaveTangentNormal * X;
-// 							KeyValue.Tangent.LeaveTangentWeight = FMath::Sqrt(X*X + Y*Y) * OneThird;
-// 						}
-// 					}
-// 					KeyValue.Tangent.TangentWeightMode = Attributes.GetTangentWeightMode();
-// 
-// 					if( KeyValue.Tangent.TangentWeightMode != RCTWM_WeightedNone )
-// 					{
-// 						if (KeyValue.TangentMode != RCTM_User && KeyValue.TangentMode != RCTM_Break)
-// 						{
-// 							KeyValue.TangentMode = RCTM_User;
-// 						}
-// 					}
-// 				}
-// 
-// 				if (Attributes.HasArriveTangent())
-// 				{
-// 					if (KeyValue.TangentMode == RCTM_Auto)
-// 					{
-// 						KeyValue.TangentMode = RCTM_User;
-// 						KeyValue.Tangent.TangentWeightMode = RCTWM_WeightedNone;
-// 					}
-// 
-// 					KeyValue.Tangent.ArriveTangent = Attributes.GetArriveTangent() * TimeInterval;
-// 					if (KeyValue.InterpMode == RCIM_Cubic && KeyValue.TangentMode != RCTM_Break)
-// 					{
-// 						KeyValue.Tangent.LeaveTangent = KeyValue.Tangent.ArriveTangent;
-// 					}
-// 				}
-// 
-// 				if (Attributes.HasLeaveTangent())
-// 				{
-// 					if (KeyValue.TangentMode == RCTM_Auto)
-// 					{
-// 						KeyValue.TangentMode = RCTM_User;
-// 						KeyValue.Tangent.TangentWeightMode = RCTWM_WeightedNone;
-// 					}
-// 
-// 					KeyValue.Tangent.LeaveTangent = Attributes.GetLeaveTangent() * TimeInterval;
-// 					if (KeyValue.InterpMode == RCIM_Cubic && KeyValue.TangentMode != RCTM_Break)
-// 					{
-// 						KeyValue.Tangent.ArriveTangent = KeyValue.Tangent.LeaveTangent;
-// 					}
-// 				}
-// 
-// 				if (Attributes.HasArriveTangentWeight())
-// 				{
-// 					if (KeyValue.TangentMode == RCTM_Auto)
-// 					{
-// 						KeyValue.TangentMode = RCTM_User;
-// 						KeyValue.Tangent.TangentWeightMode = RCTWM_WeightedNone;
-// 					}
-// 
-// 					KeyValue.Tangent.ArriveTangentWeight = Attributes.GetArriveTangentWeight(); 
-// 					if (KeyValue.InterpMode == RCIM_Cubic && KeyValue.TangentMode != RCTM_Break)
-// 					{
-// 						KeyValue.Tangent.LeaveTangentWeight = KeyValue.Tangent.ArriveTangentWeight;
-// 					}
-// 				}
-// 
-// 				if (Attributes.HasLeaveTangentWeight())
-// 				{
-// 				
-// 					if (KeyValue.TangentMode == RCTM_Auto)
-// 					{
-// 						KeyValue.TangentMode = RCTM_User;
-// 						KeyValue.Tangent.TangentWeightMode = RCTWM_WeightedNone;
-// 					}
-// 
-// 					KeyValue.Tangent.LeaveTangentWeight = Attributes.GetLeaveTangentWeight();
-// 					if (KeyValue.InterpMode == RCIM_Cubic && KeyValue.TangentMode != RCTM_Break)
-// 					{
-// 						KeyValue.Tangent.ArriveTangentWeight = KeyValue.Tangent.LeaveTangentWeight;
-// 					}
-// 				}
-// 			}
-// 		}
-// 
-// 		if (bAutoSetTangents)
-// 		{
-// 			Channel->AutoSetTangents();
-// 		}
-// 	}
 }
 
 void FEventChannelCurveModel::GetCurveAttributes(FCurveAttributes& OutCurveAttributes) const
@@ -394,18 +207,10 @@ void FEventChannelCurveModel::GetCurveAttributes(FCurveAttributes& OutCurveAttri
 
 void FEventChannelCurveModel::SetCurveAttributes(const FCurveAttributes& InCurveAttributes)
 {
-	// Event Channels have no Pre/Post Extrapolation
 }
 
 void FEventChannelCurveModel::CreateKeyProxies(TArrayView<const FKeyHandle> InKeyHandles, TArrayView<UObject*> OutObjects)
 {
-	// for (int32 Index = 0; Index < InKeyHandles.Num(); ++Index)
-	// {
-	// 	UFloatChannelKeyProxy* NewProxy = NewObject<UFloatChannelKeyProxy>(GetTransientPackage(), NAME_None);
-	// 
-	// 	NewProxy->Initialize(InKeyHandles[Index], ChannelHandle, WeakSection);
-	// 	OutObjects[Index] = NewProxy;
-	// }
 }
 
 void FEventChannelCurveModel::GetTimeRange(double& MinTime, double& MaxTime) const
@@ -431,42 +236,6 @@ void FEventChannelCurveModel::GetTimeRange(double& MinTime, double& MaxTime) con
 	}
 }
 
-/*	 Finds min/max for cubic curves:
-Looks for feature points in the signal(determined by change in direction of local tangent), these locations are then re-examined in closer detail recursively
-Similar to function in RichCurve but usees the Channel::Evaluate function, instead of CurveModel::Eval*/
-
-void FEventChannelCurveModel::FeaturePointMethod(double StartTime, double EndTime, double StartValue, double Mu, int Depth, int MaxDepth, double& MaxV, double& MinVal) const
-{
-	if (Depth >= MaxDepth)
-	{
-		return;
-	}
-	double PrevValue = StartValue;
-	double EvalValue;
-	Evaluate(StartTime - Mu, EvalValue);
-	double PrevTangent = StartValue - EvalValue;
-	EndTime += Mu;
-	for (double f = StartTime + Mu; f < EndTime; f += Mu)
-	{
-		double Value;
-		Evaluate(f, Value);
-
-		MaxV = FMath::Max(Value, MaxV);
-		MinVal = FMath::Min(Value, MinVal);
-		double CurTangent = Value - PrevValue;
-		if (FMath::Sign(CurTangent) != FMath::Sign(PrevTangent))
-		{
-			//feature point centered around the previous tangent
-			double FeaturePointTime = f - Mu * 2.0f;
-			double NewVal;
-			Evaluate(FeaturePointTime, NewVal);
-			FeaturePointMethod(FeaturePointTime, f,NewVal, Mu*0.4f, Depth + 1, MaxDepth, MaxV, MinVal);
-		}
-		PrevTangent = CurTangent;
-		PrevValue = Value;
-	}
-}
-
 void FEventChannelCurveModel::GetValueRange(double& MinValue, double& MaxValue) const
 {
 	// Event Tracks have no value, thus their value range is just zero.
@@ -479,8 +248,7 @@ int32 FEventChannelCurveModel::GetNumKeys() const
 
 	if (Channel)
 	{
-		TArrayView<const FFrameNumber> Times = Channel->GetData().GetTimes();
-		return Times.Num();
+		return Channel->GetData().GetTimes().Num();
 	}
 
 	return 0;
diff --git a/Engine/Source/Editor/MovieSceneTools/Private/Channels/EventChannelCurveModel.h b/Engine/Source/Editor/MovieSceneTools/Private/Channels/EventChannelCurveModel.h
--- a/Engine/Source/Editor/MovieSceneTools/Private/Channels/EventChannelCurveModel.h
+++ b/Engine/Source/Editor/MovieSceneTools/Private/Channels/EventChannelCurveModel.h
@@ -45,8 +45,6 @@ public:
 
 	virtual void CreateKeyProxies(TArrayView<const FKeyHandle> InKeyHandles, TArrayView<UObject*> OutObjects) override;
 
-private:
-	void FeaturePointMethod(double StartTime, double EndTime, double StartValue, double Mu, int Depth, int MaxDepth, double& MaxV, double& MinVal) const;
 private:
 
 	TMovieSceneChannelHandle<FMovieSceneEventChannel> ChannelHandle;
diff --git a/Engine/Source/Editor/MovieSceneTools/Private/MatineeImportTools.cpp b/Engine/Source/Editor/MovieSceneTools/Private/MatineeImportTools.cpp
--- a/Engine/Source/Editor/MovieSceneTools/Private/MatineeImportTools.cpp
+++ b/Engine/Source/Editor/MovieSceneTools/Private/MatineeImportTools.cpp
@@ -10,8 +10,10 @@
 #include "MovieScene.h"
 #include "MovieSceneSequenceEditor.h"
 #include "Channels/MovieSceneEvent.h"
-#include "K2Node_FunctionEntry.h"
+#include "K2Node_CustomEvent.h"
+#include "MovieSceneSequenceEditor.h"
 #include "Kismet2/KismetEditorUtilities.h"
+#include "MovieSceneEventUtils.h"
 #include "Animation/SkeletalMeshActor.h"
 #include "EngineUtils.h"
 
@@ -882,57 +884,52 @@ bool FMatineeImportTools::CopyInterpEventTrack( UInterpTrackEvent* MatineeEventT
 		if ( Section == nullptr )
 		{
 			Section = Cast<UMovieSceneEventTriggerSection>( EventTrack->CreateNewSection() );
+			Section->SetRange(TRange<FFrameNumber>::All());
 			EventTrack->AddSection(*Section);
 			bSectionCreated = true;
-			
 		}
+
+		UMovieSceneSequence* Sequence = Section->GetTypedOuter<UMovieSceneSequence>();
 		if (Section->TryModify())
 		{
+			FMovieSceneSequenceEditor* SequenceEditor = FMovieSceneSequenceEditor::Find(Sequence);
+			UBlueprint* SequenceDirectorBP = SequenceEditor  ? SequenceEditor->GetOrCreateDirectorBlueprint(Sequence) : nullptr;
+			if (SequenceDirectorBP)
+			{
+				SequenceDirectorBP->Modify();
+				SequenceDirectorBP->GenerateFunctionGraphsEvent.AddUniqueDynamic(Section, &UMovieSceneEventSectionBase::HandleGenerateEntryPoints);
+			}
+
 			TRange<FFrameNumber> KeyRange = TRange<FFrameNumber>::Empty();
 
 			FMovieSceneEventChannel* EventChannel = Section->GetChannelProxy().GetChannel<FMovieSceneEventChannel>(0);
 			check(EventChannel);
 			TMovieSceneChannelData<FMovieSceneEvent> ChannelData = EventChannel->GetData();
-			TMap<FName, FMovieSceneEvent> EventMap;
+			TMap<FName, UK2Node_CustomEvent*> EventMap;
 
-			for (FEventTrackKey EventTrackKey : MatineeEventTrack->EventTrack)
+			for (const FEventTrackKey& EventTrackKey : MatineeEventTrack->EventTrack)
 			{
 				FFrameNumber KeyTime = (EventTrackKey.Time * FrameRate).RoundToFrame();
-				
+
 				FMovieSceneEvent Event;
 
-				if (FMovieSceneEvent* FoundEvent = EventMap.Find(EventTrackKey.EventName))
+				UK2Node_CustomEvent* ExistingEndpoint = EventMap.FindRef(EventTrackKey.EventName);
+				if (ExistingEndpoint)
 				{
-					Event = *FoundEvent;
+					FMovieSceneEventUtils::SetEndpoint(&Event, Section, ExistingEndpoint, nullptr);
 				}
-				else
+				else if (SequenceDirectorBP)
 				{
-					FMovieSceneSequenceEditor* SequenceEditor = FMovieSceneSequenceEditor::Find(Section->GetTypedOuter<UMovieSceneSequence>());
-					if (ensure(SequenceEditor))
-					{
-						// Create and bind event endpoint for this event
-						UK2Node_FunctionEntry* NewEndpoint = SequenceEditor->CreateEventEndpoint(Section->GetTypedOuter<UMovieSceneSequence>(), EventTrackKey.EventName.ToString());
-						if (NewEndpoint)
-						{
-							SequenceEditor->InitializeEndpointForTrack(Section->GetTypedOuter<UMovieSceneEventTrack>(), NewEndpoint);
-							FMovieSceneSequenceEditor::BindEventToEndpoint(Section, &Event, NewEndpoint);
-						}
-					}
+					FMovieSceneEventEndpointParameters Params;
+					Params.SanitizedEventName = TEXT("MatineeEvent");
 
-					if (UK2Node_FunctionEntry* FunctionEntry = Event.GetFunctionEntry())
-					{
-						FKismetEditorUtilities::BringKismetToFocusAttentionOnObject(FunctionEntry, false);
-					}
-					EventMap.Add(EventTrackKey.EventName, Event);
-				}
+					UK2Node_CustomEvent* NewEndpoint = FMovieSceneEventUtils::CreateUserFacingEvent(SequenceDirectorBP, Params);
+					EventMap.Add(EventTrackKey.EventName, NewEndpoint);
 
-				ChannelData.UpdateOrAddKey(KeyTime, Event);
-				KeyRange = TRange<FFrameNumber>::Hull(KeyRange, TRange<FFrameNumber>(KeyTime));
-			}
+					FMovieSceneEventUtils::SetEndpoint(&Event, Section, NewEndpoint, nullptr);
+				}
 
-			if (!KeyRange.IsEmpty())
-			{
-				Section->SetRange( KeyRange );
+				ChannelData.AddKey(KeyTime, Event);
 			}
 		}
 	}	
diff --git a/Engine/Source/Editor/MovieSceneTools/Private/MovieSceneEventCustomization.cpp b/Engine/Source/Editor/MovieSceneTools/Private/MovieSceneEventCustomization.cpp
--- a/Engine/Source/Editor/MovieSceneTools/Private/MovieSceneEventCustomization.cpp
+++ b/Engine/Source/Editor/MovieSceneTools/Private/MovieSceneEventCustomization.cpp
@@ -2,41 +2,160 @@
 
 #include "MovieSceneEventCustomization.h"
 #include "Modules/ModuleManager.h"
-#include "Algo/Find.h"
+#include "MovieSceneEventUtils.h"
 
 #include "UObject/UnrealType.h"
 #include "Channels/MovieSceneEvent.h"
 #include "Tracks/MovieSceneEventTrack.h"
 #include "Sections/MovieSceneEventSectionBase.h"
+#include "Algo/Find.h"
 #include "MovieSceneSequence.h"
 #include "ISequencerModule.h"
 #include "ScopedTransaction.h"
 
+#include "Framework/Application/SlateApplication.h"
 #include "Framework/MultiBox/MultiBoxBuilder.h"
 #include "Widgets/DeclarativeSyntaxSupport.h"
 #include "Widgets/Text/STextBlock.h"
 #include "Widgets/Images/SImage.h"
 #include "Widgets/Layout/SBox.h"
+#include "Widgets/Input/SCheckBox.h"
 #include "Widgets/Input/SComboButton.h"
 #include "Widgets/Input/SHyperlink.h"
 #include "Widgets/Input/SEditableTextBox.h"
+#include "Widgets/Views/SExpanderArrow.h"
 
 #include "PropertyHandle.h"
+#include "IPropertyUtilities.h"
 #include "PropertyCustomizationHelpers.h"
 #include "IDetailChildrenBuilder.h"
+#include "DetailLayoutBuilder.h"
+#include "DetailCategoryBuilder.h"
 
 #include "K2Node_Variable.h"
-#include "K2Node_FunctionEntry.h"
+#include "K2Node_CustomEvent.h"
 #include "K2Node_CallFunction.h"
+#include "K2Node_FunctionEntry.h"
 #include "Kismet2/KismetEditorUtilities.h"
 #include "Kismet2/BlueprintEditorUtils.h"
 #include "Kismet2/Kismet2NameValidators.h"
+#include "SGraphActionMenu.h"
+#include "BlueprintActionMenuUtils.h"
+#include "BlueprintActionMenuBuilder.h"
+#include "BlueprintActionMenuItem.h"
+#include "BlueprintFunctionNodeSpawner.h"
 
 #include "EditorStyleSet.h"
 #include "EditorFontGlyphs.h"
 
 #define LOCTEXT_NAMESPACE "MovieSceneEventCustomization"
 
+
+namespace UE4_MovieSceneEventCustomization
+{
+	void CollectQuickBindActions(FGraphActionListBuilderBase& OutAllActions, UBlueprint* Blueprint, UClass* BoundObjectPinClass)
+	{
+		// Build up the context object
+		auto RejectAnyNonFunctions = [](const FBlueprintActionFilter& Filter, FBlueprintActionInfo& BlueprintAction)
+		{
+			const UFunction* Function = BlueprintAction.GetAssociatedFunction();
+			return Function == nullptr || Function->HasAnyFunctionFlags(FUNC_BlueprintPure);
+		};
+		auto RejectAnyUnboundActions = [](const FBlueprintActionFilter& Filter, FBlueprintActionInfo& BlueprintAction)
+		{
+			return BlueprintAction.GetBindings().Num() <= 0;
+		};
+
+
+		FBlueprintActionMenuBuilder ContextMenuBuilder(nullptr);
+
+		if (BoundObjectPinClass)
+		{
+			// Add actions that are relevant to the bound object from the pin class
+			FBlueprintActionFilter CallOnMemberFilter(FBlueprintActionFilter::BPFILTER_RejectGlobalFields | FBlueprintActionFilter::BPFILTER_RejectPermittedSubClasses);
+			CallOnMemberFilter.PermittedNodeTypes.Add(UK2Node_CallFunction::StaticClass());
+			CallOnMemberFilter.Context.Blueprints.Add(Blueprint);
+
+			for (UObjectProperty* ObjectProperty : TFieldRange<UObjectProperty>(BoundObjectPinClass))
+			{
+				if (ObjectProperty->HasAnyPropertyFlags(CPF_BlueprintVisible) && ObjectProperty->HasMetaData(FBlueprintMetadata::MD_ExposeFunctionCategories))
+				{
+					CallOnMemberFilter.Context.SelectedObjects.Add(ObjectProperty);
+					FBlueprintActionFilter::AddUnique(CallOnMemberFilter.TargetClasses, ObjectProperty->PropertyClass);
+				}
+			}
+
+			CallOnMemberFilter.AddRejectionTest(FBlueprintActionFilter::FRejectionTestDelegate::CreateStatic(RejectAnyUnboundActions));
+			CallOnMemberFilter.AddRejectionTest(FBlueprintActionFilter::FRejectionTestDelegate::CreateStatic(RejectAnyNonFunctions));
+
+			ContextMenuBuilder.AddMenuSection(CallOnMemberFilter, FText::FromName(BoundObjectPinClass->GetFName()), 0);
+		}
+
+		{
+			// Add all actions that are relevant to the sequence director BP itself
+			FBlueprintActionFilter MenuFilter(FBlueprintActionFilter::BPFILTER_RejectGlobalFields | FBlueprintActionFilter::BPFILTER_RejectPermittedSubClasses);
+
+			MenuFilter.PermittedNodeTypes.Add(UK2Node_CallFunction::StaticClass());
+
+			MenuFilter.Context.Blueprints.Add(Blueprint);
+			MenuFilter.Context.Graphs.Append(Blueprint->UbergraphPages);
+			MenuFilter.Context.Graphs.Append(Blueprint->FunctionGraphs);
+
+			if (Blueprint->SkeletonGeneratedClass)
+			{
+				FBlueprintActionFilter::AddUnique(MenuFilter.TargetClasses, Blueprint->SkeletonGeneratedClass);
+			}
+			MenuFilter.AddRejectionTest(FBlueprintActionFilter::FRejectionTestDelegate::CreateStatic(RejectAnyNonFunctions));
+
+
+			ContextMenuBuilder.AddMenuSection(MenuFilter, LOCTEXT("SequenceDirectorMenu", "This Sequence"), 0);
+		}
+
+		ContextMenuBuilder.RebuildActionList();
+
+		OutAllActions.Append(ContextMenuBuilder);
+	}
+
+	void CollectAllRebindActions(FGraphActionListBuilderBase& OutAllActions, UBlueprint* Blueprint)
+	{
+		// Build up the context object
+		auto RejectAnyForeignFunctions = [Blueprint](const FBlueprintActionFilter& Filter, FBlueprintActionInfo& BlueprintAction)
+		{
+			const UBlueprintFunctionNodeSpawner* FunctionNodeSpawner = Cast<UBlueprintFunctionNodeSpawner>(BlueprintAction.NodeSpawner);
+			const UFunction* FunctionToCall = FunctionNodeSpawner ? FunctionNodeSpawner->GetFunction() : nullptr;
+
+			if (!FunctionToCall || FunctionToCall->HasAnyFunctionFlags(FUNC_BlueprintPure))
+			{
+				return true;
+			}
+
+			UClass* OuterClass = CastChecked<UClass>(FunctionToCall->GetOuter());
+			return OuterClass->ClassGeneratedBy != Blueprint;
+		};
+
+		FBlueprintActionFilter MenuFilter(FBlueprintActionFilter::BPFILTER_RejectGlobalFields | FBlueprintActionFilter::BPFILTER_RejectPermittedSubClasses);
+		MenuFilter.PermittedNodeTypes.Add(UK2Node_CallFunction::StaticClass());
+
+		MenuFilter.Context.Blueprints.Add(Blueprint);
+
+		MenuFilter.AddRejectionTest(FBlueprintActionFilter::FRejectionTestDelegate::CreateLambda(RejectAnyForeignFunctions));
+
+		if (Blueprint->SkeletonGeneratedClass)
+		{
+			FBlueprintActionFilter::AddUnique(MenuFilter.TargetClasses, Blueprint->SkeletonGeneratedClass);
+		}
+
+		FBlueprintActionMenuBuilder ContextMenuBuilder(nullptr);
+
+		ContextMenuBuilder.AddMenuSection(MenuFilter, LOCTEXT("SequenceDirectorMenu", "This Sequence"), 0);
+		ContextMenuBuilder.RebuildActionList();
+		OutAllActions.Append(ContextMenuBuilder);
+	}
+}
+
+
+
+
 TSharedRef<IPropertyTypeCustomization> FMovieSceneEventCustomization::MakeInstance()
 {
 	return MakeShared<FMovieSceneEventCustomization>();
@@ -68,6 +187,7 @@ void FMovieSceneEventCustomization::CustomizeHeader(TSharedRef<IPropertyHandle>
 void FMovieSceneEventCustomization::CustomizeChildren(TSharedRef<IPropertyHandle> InPropertyHandle, IDetailChildrenBuilder& ChildBuilder, IPropertyTypeCustomizationUtils& CustomizationUtils)
 {
 	PropertyHandle = InPropertyHandle;
+	PropertyUtilities = CustomizationUtils.GetPropertyUtilities();
 
 	TArray<void*> RawData;
 	PropertyHandle->AccessRawData(RawData);
@@ -122,70 +242,334 @@ void FMovieSceneEventCustomization::CustomizeChildren(TSharedRef<IPropertyHandle
 		[
 			PropertyCustomizationHelpers::MakeBrowseButton(FSimpleDelegate::CreateSP(this, &FMovieSceneEventCustomization::NavigateToDefinition), LOCTEXT("NavigateToDefinition_Tip", "Navigate to this event's definition"))
 		]
+	];
 
-		+ SHorizontalBox::Slot()
-		.VAlign(VAlign_Center)
-		.AutoWidth()
+	const bool bAnyBoundEvents = GetAllValidEndpoints().Num() != 0;
+	if (bAnyBoundEvents)
+	{
+		FText CallInEditorText    = LOCTEXT("CallInEditor_Label",   "Call In Editor");
+		FText CallInEditorTooltip = LOCTEXT("CallInEditor_Tooltip", "When checked, this event will be triggered in the Editor outside of PIE.\n\nBEWARE: ANY CHANGES MADE AS A RESULT OF THIS EVENT BEING CALLED MAY END UP BEING SAVED IN THE CURRENT LEVEL OR ASSET.");
+
+		ChildBuilder.AddCustomRow(CallInEditorText)
+		.NameContent()
 		[
 			SNew(STextBlock)
-			.Visibility(this, &FMovieSceneEventCustomization::GetErrorVisibility)
-			.ToolTipText(this, &FMovieSceneEventCustomization::GetErrorTooltip)
-			.TextStyle(FEditorStyle::Get(), "Log.Warning")
-			.Font(FEditorStyle::Get().GetFontStyle("FontAwesome.9"))
-			.Text(FEditorFontGlyphs::Exclamation_Triangle)
+			.Text(CallInEditorText)
+			.ToolTipText(CallInEditorTooltip)
+			.Font(IDetailLayoutBuilder::GetDetailFont())
 		]
-	];
-}
+		.ValueContent()
+		[
+			SNew(SCheckBox)
+			.ToolTipText(CallInEditorTooltip)
+			.IsChecked(this, &FMovieSceneEventCustomization::GetCallInEditorCheckState)
+			.OnCheckStateChanged(this, &FMovieSceneEventCustomization::OnSetCallInEditorCheckState)
+		];
+	}
 
-UClass* FMovieSceneEventCustomization::FindObjectBindingClass()
-{
-	UMovieSceneSequence* CommonSequence = GetCommonSequence();
-	UMovieScene*         MovieScene     = CommonSequence ? CommonSequence->GetMovieScene() : nullptr;
 
-	if (!MovieScene)
+	UK2Node* CommonEndpoint = GetCommonEndpoint();
+	if (!CommonEndpoint)
 	{
-		return nullptr;
+		return;
 	}
 
-	TArray<UObject*> EditObjects;
-	GetEditObjects(EditObjects);
+	TArray<UEdGraphPin*> ValidBoundObjectPins;
+	for (UEdGraphPin* Pin : CommonEndpoint->Pins)
+	{
+		if (Pin->Direction == EGPD_Output && (Pin->PinType.PinCategory == UEdGraphSchema_K2::PC_Object || Pin->PinType.PinCategory == UEdGraphSchema_K2::PC_Interface))
+		{
+			ValidBoundObjectPins.Add(Pin);
+		}
+	}
+
+
+	if (ValidBoundObjectPins.Num() > 0)
+	{
+		FText BoundObjectPinText    = LOCTEXT("BoundObjectPin_Label", "Pass Bound Object To");
+		FText BoundObjectPinTooltip = LOCTEXT("BoundObjectPin_Tooltip", "Specifies a pin to pass the bound object(s) through when the event is triggered. Interface and object pins are both supported.");
+
+		ChildBuilder.AddCustomRow(BoundObjectPinText)
+		.NameContent()
+		[
+			SNew(STextBlock)
+			.Text(BoundObjectPinText)
+			.ToolTipText(BoundObjectPinTooltip)
+			.Font(IDetailLayoutBuilder::GetDetailFont())
+		]
+		.ValueContent()
+		[
+			SNew(SComboButton)
+			.ButtonStyle(FEditorStyle::Get(), "HoverHintOnly")
+			.ForegroundColor(FSlateColor::UseForeground())
+			.ToolTipText(BoundObjectPinTooltip)
+			.OnGetMenuContent(this, &FMovieSceneEventCustomization::GetBoundObjectPinMenuContent)
+			.CollapseMenuOnParentFocus(true)
+			.ContentPadding(FMargin(4.f, 0.f))
+			.ButtonContent()
+			[
+				SNew(SBox)
+				.HeightOverride(21.f)
+				[
+					SNew(SHorizontalBox)
+
+					+ SHorizontalBox::Slot()
+					.Padding(FMargin(0.f, 0.f, 4.f, 0.f))
+					.VAlign(VAlign_Center)
+					.AutoWidth()
+					[
+						SNew(SImage)
+						.Image(this, &FMovieSceneEventCustomization::GetBoundObjectPinIcon)
+					]
+
+					+ SHorizontalBox::Slot()
+					.Padding(FMargin(0.f, 0.f, 4.f, 0.f))
+					.VAlign(VAlign_Center)
+					[
+						SNew(STextBlock)
+						.Font(CustomizationUtils.GetRegularFont())
+						.Text(this, &FMovieSceneEventCustomization::GetBoundObjectPinText)
+					]
+				]
+			]
+		];
+	}
+
+
+	UFunction* CommonFunction = nullptr;
+	UBlueprint* Blueprint = CommonEndpoint->GetBlueprint();
+	if (Blueprint)
+	{
+		if (UK2Node_Event* Event = Cast<UK2Node_Event>(CommonEndpoint))
+		{
+			CommonFunction = Blueprint->SkeletonGeneratedClass ? Blueprint->SkeletonGeneratedClass->FindFunctionByName(Event->GetFunctionName()) : nullptr;
+		}
+		else if (UK2Node_FunctionEntry* EndpointEntry = Cast<UK2Node_FunctionEntry>(CommonEndpoint))
+		{
+			CommonFunction = EndpointEntry->FindSignatureFunction();
+		}
+		else
+		{
+			// @todo: Error not supported
+		}
 
-	UClass* BindingClass = nullptr;
+		Blueprint->OnCompiled().AddSP(this, &FMovieSceneEventCustomization::OnBlueprintCompiled);
+	}
 
-	for (UObject* EditObject : EditObjects)
+	if (CommonFunction)
 	{
-		UMovieSceneTrack* Track = EditObject->GetTypedOuter<UMovieSceneTrack>();
-		if (Track)
+		IDetailCategoryBuilder& DetailCategoryBuilder = ChildBuilder.GetParentCategory().GetParentLayout().EditCategory("Payload", LOCTEXT("PayloadLabel", "Payload"), ECategoryPriority::Uncommon);
+
+		bool bPayloadUpToDate = true;
+
+		for (int32 Index = 0; Index < RawData.Num(); ++Index)
 		{
-			for (const FMovieSceneBinding& Binding : MovieScene->GetBindings())
+			TSharedPtr<FStructOnScope> StructData = MakeShared<FStructOnScope>(CommonFunction);
+
+			const FMovieSceneEvent* EntryPoint = static_cast<FMovieSceneEvent*>(RawData[Index]);
+
+			TArray<FName, TInlineAllocator<8>> AllValidNames;
+
+			for (UProperty* Field : TFieldRange<UProperty>(CommonFunction))
 			{
-				if (Algo::Find(Binding.GetTracks(), Track))
+				if (Field->HasAnyPropertyFlags(CPF_OutParm | CPF_ReturnParm | CPF_OutParm | CPF_ReferenceParm) || Field->GetFName() == EntryPoint->BoundObjectPinName)
 				{
-					UClass* ThisTrackClass = nullptr;
-					if (FMovieScenePossessable* Possessable = MovieScene->FindPossessable(Binding.GetObjectGuid()))
-					{
-						ThisTrackClass = const_cast<UClass*>(Possessable->GetPossessedObjectClass());
-					}
-					else if(FMovieSceneSpawnable* Spawnable = MovieScene->FindSpawnable(Binding.GetObjectGuid()))
-					{
-						ThisTrackClass = Spawnable->GetObjectTemplate()->GetClass();
-					}
+					continue;
+				}
+
+				const FMovieSceneEventPayloadVariable* PayloadVariable = EntryPoint->PayloadVariables.Find(Field->GetFName());
+				if (PayloadVariable)
+				{
+					AllValidNames.Add(Field->GetFName());
 
-					if (!BindingClass)
+					// We have an override for this variable
+					const bool bImportSuccess = FBlueprintEditorUtils::PropertyValueFromString(Field, PayloadVariable->Value, StructData->GetStructMemory());
+					if (!bImportSuccess)
 					{
-						BindingClass = ThisTrackClass;
+						// @todo: error
 					}
-					else if (BindingClass != ThisTrackClass)
+				}
+				else if (UEdGraphPin* Pin = CommonEndpoint->FindPin(Field->GetFName(), EGPD_Output))
+				{
+					if (!Pin->DefaultValue.IsEmpty())
 					{
-						return nullptr;
+						const bool bImportSuccess = FBlueprintEditorUtils::PropertyValueFromString(Field, Pin->DefaultValue, StructData->GetStructMemory());
+						if (!bImportSuccess)
+						{
+							// @todo: error
+						}
 					}
-					break;
 				}
+
+				IDetailPropertyRow* ExternalRow = DetailCategoryBuilder.AddExternalStructureProperty(StructData.ToSharedRef(), Field->GetFName(), EPropertyLocation::Default, FAddPropertyParams().ForceShowProperty());
+
+				TSharedPtr<IPropertyHandle> LocalVariableProperty = ExternalRow->GetPropertyHandle();
+				FSimpleDelegate Delegate = FSimpleDelegate::CreateSP(this, &FMovieSceneEventCustomization::OnPayloadVariableChanged, StructData.ToSharedRef(), LocalVariableProperty);
+
+				LocalVariableProperty->SetOnPropertyValueChanged(Delegate);
+				LocalVariableProperty->SetOnChildPropertyValueChanged(Delegate);
 			}
+
+			bPayloadUpToDate &= (AllValidNames.Num() == EntryPoint->PayloadVariables.Num());
+		}
+
+		if (!bPayloadUpToDate)
+		{
+			DetailCategoryBuilder.AddCustomRow(FText())
+			.WholeRowContent()
+			[
+				SNew(SHorizontalBox)
+
+				+ SHorizontalBox::Slot()
+				.Padding(FMargin(5.f, 0.f, 5.f, 0.f))
+				.AutoWidth()
+				[
+					SNew(STextBlock)
+					.TextStyle(FEditorStyle::Get(), "Log.Warning")
+					.Font(FEditorStyle::GetFontStyle("FontAwesome.10"))
+					.Text(FEditorFontGlyphs::Exclamation_Triangle)
+				]
+
+				+ SHorizontalBox::Slot()
+				.Padding(FMargin(0.f, 0.f, 5.f, 0.f))
+				[
+					SNew(STextBlock)
+					.TextStyle(FEditorStyle::Get(), "Log.Warning")
+					.Text(LOCTEXT("PayloadOutOfDateError", "Payload variables may be out-of-date. Please compile the blueprint."))
+				]
+			];
+		}
+	}
+}
+
+void FMovieSceneEventCustomization::OnPayloadVariableChanged(TSharedRef<FStructOnScope> InStructData, TSharedPtr<IPropertyHandle> LocalVariableProperty)
+{
+	// This function should only ever be bound if all the entry points call the same function
+	UProperty* Property = LocalVariableProperty->GetProperty();
+	if (!Property)
+	{
+		return;
+	}
+
+	FScopedTransaction Transaction(LOCTEXT("SetPayloadValue", "Set payload value"));
+
+	bool bChangedAnything = false;
+
+	UBlueprint* Blueprint = nullptr;
+
+	FString NewValueString;
+
+	const bool bSuccessfulTextExport = FBlueprintEditorUtils::PropertyValueToString(Property, InStructData->GetStructMemory(), NewValueString, nullptr);
+	if (!bSuccessfulTextExport)
+	{
+		// @todo: error
+		return;
+	}
+
+
+	TArray<UObject*> EditObjects;
+	GetEditObjects(EditObjects);
+
+	TArray<void*> RawData;
+	PropertyHandle->AccessRawData(RawData);
+
+	if (!ensure(RawData.Num() == EditObjects.Num()))
+	{
+		return;
+	}
+
+	for (int32 Index = 0; Index < RawData.Num(); ++Index)
+	{
+		UMovieSceneEventSectionBase* EventSection = Cast<UMovieSceneEventSectionBase>(EditObjects[Index]);
+		FMovieSceneEvent*  EntryPoint   = static_cast<FMovieSceneEvent*>(RawData[Index]);
+		if (!EventSection || !EntryPoint)
+		{
+			continue;
+		}
+
+		EventSection->Modify();
+
+		FMovieSceneEventPayloadVariable* PayloadVariable = EntryPoint->PayloadVariables.Find(Property->GetFName());
+		if (!PayloadVariable)
+		{
+			PayloadVariable = &EntryPoint->PayloadVariables.Add(Property->GetFName());
+		}
+
+		PayloadVariable->Value = NewValueString;
+		bChangedAnything = true;
+	}
+
+	if (bChangedAnything)
+	{
+		UBlueprint* BP = GetCommonEndpoint()->GetBlueprint();
+		FBlueprintEditorUtils::MarkBlueprintAsModified(BP);
+	}
+	else
+	{
+		Transaction.Cancel();
+	}
+}
+
+ECheckBoxState FMovieSceneEventCustomization::GetCallInEditorCheckState() const
+{
+	ECheckBoxState CheckState = ECheckBoxState::Undetermined;
+	for (UK2Node* Endpoint : GetAllValidEndpoints())
+	{
+		bool bCallInEditor = false;
+		if (UK2Node_CustomEvent* CustomEvent = Cast<UK2Node_CustomEvent>(Endpoint))
+		{
+			bCallInEditor = CustomEvent->bCallInEditor;
+		}
+		else if (UK2Node_FunctionEntry* FunctionEntry = Cast<UK2Node_FunctionEntry>(Endpoint))
+		{
+			bCallInEditor = FunctionEntry->MetaData.bCallInEditor;
+		}
+
+		if (CheckState == ECheckBoxState::Undetermined)
+		{
+			CheckState = bCallInEditor ? ECheckBoxState::Checked : ECheckBoxState::Unchecked;
+		}
+		else if (bCallInEditor != (CheckState == ECheckBoxState::Checked) )
+		{
+			return ECheckBoxState::Undetermined;
 		}
 	}
+	return CheckState;
+}
+
+void FMovieSceneEventCustomization::OnSetCallInEditorCheckState(ECheckBoxState NewCheckedState)
+{
+	FScopedTransaction Transaction(LOCTEXT("SetCallInEditor", "Set Call in Editor"));
+
+	const bool bCallInEditor = (NewCheckedState == ECheckBoxState::Checked);
+
+	TSet<UBlueprint*> Blueprints;
+
+	for (UK2Node* Endpoint : GetAllValidEndpoints())
+	{
+		if (UK2Node_CustomEvent* CustomEvent = Cast<UK2Node_CustomEvent>(Endpoint))
+		{
+			CustomEvent->Modify();
+			CustomEvent->bCallInEditor = bCallInEditor;
+			Blueprints.Add(CustomEvent->GetBlueprint());
+		}
+		else if (UK2Node_FunctionEntry* FunctionEntry = Cast<UK2Node_FunctionEntry>(Endpoint))
+		{
+			FunctionEntry->Modify();
+			FunctionEntry->MetaData.bCallInEditor = bCallInEditor;
+			Blueprints.Add(FunctionEntry->GetBlueprint());
+		}
+	}
+
+	for (UBlueprint* Blueprint : Blueprints)
+	{
+		FBlueprintEditorUtils::MarkBlueprintAsModified(Blueprint);
+	}
+}
 
-	return BindingClass;
+void FMovieSceneEventCustomization::OnBlueprintCompiled(UBlueprint*)
+{
+	PropertyUtilities->ForceRefresh();
 }
 
 UMovieSceneSequence* FMovieSceneEventCustomization::GetCommonSequence() const
@@ -228,26 +612,83 @@ UMovieSceneEventTrack* FMovieSceneEventCustomization::GetCommonTrack() const
 	return CommonEventTrack;
 }
 
-UK2Node_FunctionEntry* FMovieSceneEventCustomization::GetCommonEndpoint() const
+void FMovieSceneEventCustomization::IterateEndpoints(TFunctionRef<bool(UK2Node*)> Callback) const
 {
+	TArray<UObject*> EditObjects;
+	GetEditObjects(EditObjects);
+
 	TArray<void*> RawData;
 	PropertyHandle->AccessRawData(RawData);
 
-	UK2Node_FunctionEntry* CommonEndpoint = nullptr;
+	if (!ensure(RawData.Num() == EditObjects.Num()))
+	{
+		return;
+	}
 
-	for (void* Ptr : RawData)
+	for (int32 Index = 0; Index < RawData.Num(); ++Index)
 	{
-		if (Ptr)
+		UMovieSceneEventSectionBase* EventSection = Cast<UMovieSceneEventSectionBase>(EditObjects[Index]);
+		FMovieSceneEvent*  EntryPoint   = static_cast<FMovieSceneEvent*>(RawData[Index]);
+		if (!EventSection || !EntryPoint)
 		{
-			UK2Node_FunctionEntry* ThisEndpoint = static_cast<FMovieSceneEvent*>(Ptr)->GetFunctionEntry();
-			if (CommonEndpoint && CommonEndpoint != ThisEndpoint)
+			continue;
+		}
+
+		UMovieSceneSequence*       Sequence       = EventSection->GetTypedOuter<UMovieSceneSequence>();
+		FMovieSceneSequenceEditor* SequenceEditor = FMovieSceneSequenceEditor::Find(Sequence);
+		if (!SequenceEditor)
+		{
+			continue;
+		}
+
+		UBlueprint* SequenceDirectorBP = SequenceEditor->FindDirectorBlueprint(Sequence);
+		if (SequenceDirectorBP)
+		{
+			UK2Node* Endpoint = FMovieSceneEventUtils::FindEndpoint(EntryPoint, EventSection, SequenceDirectorBP);
+			if (!Callback(Endpoint))
 			{
-				return nullptr;
+				return;
 			}
-
-			CommonEndpoint = ThisEndpoint;
 		}
 	}
+}
+
+TArray<UK2Node*> FMovieSceneEventCustomization::GetAllValidEndpoints() const
+{
+	TArray<UK2Node*> Endpoints;
+
+	IterateEndpoints(
+		[&Endpoints](UK2Node* Endpoint)
+		{
+			if (Endpoint)
+			{
+				Endpoints.Add(Endpoint);
+			}
+			return true;
+		}
+	);
+
+	return Endpoints;
+}
+
+UK2Node* FMovieSceneEventCustomization::GetCommonEndpoint() const
+{
+	UK2Node* CommonEndpoint = nullptr;
+
+	IterateEndpoints(
+		[&CommonEndpoint](UK2Node* Endpoint)
+		{
+			if (CommonEndpoint && Endpoint != CommonEndpoint)
+			{
+				CommonEndpoint = nullptr;
+				return false;
+			}
+
+			CommonEndpoint = Endpoint;
+			return true;
+		}
+	);
+
 	return CommonEndpoint;
 }
 
@@ -256,12 +697,9 @@ TSharedRef<SWidget> FMovieSceneEventCustomization::GetMenuContent()
 	FMenuBuilder MenuBuilder(true, nullptr, nullptr, true);
 
 	UMovieSceneSequence*       Sequence       = GetCommonSequence();
-	UK2Node_FunctionEntry*     CommonEndpoint = GetCommonEndpoint();
 	FMovieSceneSequenceEditor* SequenceEditor = FMovieSceneSequenceEditor::Find(Sequence);
 
-	UBlueprint* DirectorBP = SequenceEditor ? SequenceEditor->GetDirectorBlueprint(Sequence) : nullptr;
-
-	CachedCommonEndpoint = CommonEndpoint;
+	UBlueprint* DirectorBP = SequenceEditor ? SequenceEditor->FindDirectorBlueprint(Sequence) : nullptr;
 
 	MenuBuilder.AddMenuEntry(
 		LOCTEXT("CreateEventEndpoint_Text",    "Create New Endpoint"),
@@ -272,69 +710,66 @@ TSharedRef<SWidget> FMovieSceneEventCustomization::GetMenuContent()
 		)
 	);
 
-	UClass* TemplateClass = FindObjectBindingClass();
-	if (TemplateClass)
+	const bool bAnyBoundEvents = GetAllValidEndpoints().Num() != 0;
+	if (!bAnyBoundEvents && Sequence)
 	{
 		MenuBuilder.AddSubMenu(
-			LOCTEXT("CreateQuickBinding_Text",    "Create Quick Binding"),
+			LOCTEXT("CreateQuickBinding_Text",    "Quick Bind"),
 			LOCTEXT("CreateQuickBinding_Tooltip", "Shows a list of functions on this object binding that can be bound directly to this event."),
-			FNewMenuDelegate::CreateSP(this, &FMovieSceneEventCustomization::PopulateQuickBindSubMenu, TemplateClass),
+			FNewMenuDelegate::CreateSP(this, &FMovieSceneEventCustomization::PopulateQuickBindSubMenu, Sequence),
+			false /* bInOpenSubMenuOnClick */,
+			FSlateIcon(FEditorStyle::GetStyleSetName(), "Sequencer.CreateQuickBinding"),
+			false /* bInShouldWindowAfterMenuSelection */
+		);
+	}
+	else
+	{
+		MenuBuilder.AddSubMenu(
+			LOCTEXT("Rebind_Text",         "Rebind To"),
+			LOCTEXT("Rebind_Text_Tooltip", "Rebinds this event to a different function call or event node."),
+			FNewMenuDelegate::CreateSP(this, &FMovieSceneEventCustomization::PopulateRebindSubMenu, Sequence),
 			false /* bInOpenSubMenuOnClick */,
 			FSlateIcon(FEditorStyle::GetStyleSetName(), "Sequencer.CreateQuickBinding"),
 			false /* bInShouldCloseWindowAfterMenuSelection */
 		);
+
+		MenuBuilder.AddMenuEntry(
+			LOCTEXT("ClearEventEndpoint_Text",    "Clear"),
+			LOCTEXT("ClearEventEndpoint_Tooltip", "Unbinds this event from its current binding."),
+			FSlateIcon(FEditorStyle::GetStyleSetName(), "Sequencer.ClearEventBinding"),
+			FUIAction(
+				FExecuteAction::CreateSP(this, &FMovieSceneEventCustomization::ClearEventEndpoint)
+			)
+		);
 	}
 
-	MenuBuilder.AddMenuEntry(
-		LOCTEXT("ClearEventEndpoint_Text",    "Clear"),
-		LOCTEXT("ClearEventEndpoint_Tooltip", "Unbinds this event from its current binding."),
-		FSlateIcon(FEditorStyle::GetStyleSetName(), "Sequencer.ClearEventBinding"),
-		FUIAction(
-			FExecuteAction::CreateSP(this, &FMovieSceneEventCustomization::ClearEventEndpoint)
-		)
-	);
+	return MenuBuilder.MakeWidget();
+}
 
-	if (DirectorBP)
+TSharedRef<SWidget> FMovieSceneEventCustomization::GetBoundObjectPinMenuContent()
+{
+	UK2Node* CommonEndpoint = GetCommonEndpoint();
+	if (CommonEndpoint)
 	{
-		FSlateIcon Icon(FEditorStyle::GetStyleSetName(), "GraphEditor.Function_16x");
-
-		MenuBuilder.BeginSection(NAME_None, LOCTEXT("ExistingEndpoints", "Existing"));
-
-		TArray<UK2Node_FunctionEntry*> EntryNodes;
+		bool bAnyValidPins        = false;
+		bool bBoundObjectPinIsSet = false;
 
-		TArray<UEdGraph*> AllFunctionGraphs;
-		UBlueprint* BlueprintClass = DirectorBP;
-		do
+		FMenuBuilder MenuBuilder(true, nullptr, nullptr, true);
+		for (UEdGraphPin* Pin : CommonEndpoint->Pins)
 		{
-			// Add the functions that belong to this Blueprint Class
-			AllFunctionGraphs.Append(BlueprintClass->FunctionGraphs);
-
-			// And then see if we have a parent Blueprint Class to get functions from too.
-			if (BlueprintClass->ParentClass)
-			{
-				BlueprintClass = Cast<UBlueprint>(BlueprintClass->ParentClass->ClassGeneratedBy);
-			}
-			else
+			if (Pin->Direction == EGPD_Output && (Pin->PinType.PinCategory == UEdGraphSchema_K2::PC_Object || Pin->PinType.PinCategory == UEdGraphSchema_K2::PC_Interface))
 			{
-				BlueprintClass = nullptr;
-			}
-		} while (BlueprintClass != nullptr);
+				bAnyValidPins = true;
 
-		for (UEdGraph* FunctionGraph : AllFunctionGraphs) 
-		{
-			EntryNodes.Reset();
-			FunctionGraph->GetNodesOfClass<UK2Node_FunctionEntry>(EntryNodes);
-
-			if (EntryNodes.Num() == 1 && FMovieSceneEvent::IsValidFunction(EntryNodes[0]))
-			{
+				FText PinText = FText::FromName(Pin->PinName);
 				MenuBuilder.AddMenuEntry(
-					EntryNodes[0]->GetNodeTitle(ENodeTitleType::MenuTitle),
-					EntryNodes[0]->GetTooltipText(),
-					Icon,
+					PinText,
+					FText::Format(LOCTEXT("SetBoundObjectPin_Tooltip", "When calling this event with a bound object, pass the object through pin {0}."), PinText),
+					FSlateIcon(),
 					FUIAction(
-						FExecuteAction::CreateSP(this, &FMovieSceneEventCustomization::SetEventEndpoint, EntryNodes[0]),
-						FCanExecuteAction::CreateLambda([]{ return true; }),
-						FIsActionChecked::CreateSP(this, &FMovieSceneEventCustomization::CompareCurrentEventEndpoint, EntryNodes[0])
+						FExecuteAction::CreateSP(this, &FMovieSceneEventCustomization::SetBoundObjectPinName, Pin->PinName),
+						FCanExecuteAction(),
+						FIsActionChecked::CreateSP(this, &FMovieSceneEventCustomization::CompareBoundObjectPinName, Pin->PinName)
 					),
 					NAME_None,
 					EUserInterfaceActionType::RadioButton
@@ -342,74 +777,390 @@ TSharedRef<SWidget> FMovieSceneEventCustomization::GetMenuContent()
 			}
 		}
 
-		MenuBuilder.EndSection();
+		MenuBuilder.AddMenuEntry(
+			LOCTEXT("ClearBoundObjectPin_Label", "Clear"),
+			LOCTEXT("ClearBoundObjectPin_Tooltip", "Clears the reference to the bound object pin, meaning any bound objects will not be passed to the event"),
+			FSlateIcon(),
+			FUIAction(
+				FExecuteAction::CreateSP(this, &FMovieSceneEventCustomization::SetBoundObjectPinName, FName())
+			)
+		);
+
+		if (bAnyValidPins || bBoundObjectPinIsSet)
+		{
+			return MenuBuilder.MakeWidget();
+		}
+		else
+		{
+			return SNew(STextBlock).Text(LOCTEXT("BoundObjectPinError_NoPins", "No compatible pins were found. Only Object and Interface pins can be bound to objects."));
+		}
 	}
 
-	return MenuBuilder.MakeWidget();
+	return SNew(STextBlock).Text(LOCTEXT("BoundObjectPinError_MultipleEvents", "Cannot choose a bound object pin with multiple events selected"));
+}
+
+bool FMovieSceneEventCustomization::CompareBoundObjectPinName(FName InPinName) const
+{
+	bool bMatch = false;
+
+	auto Enumerator = [&bMatch, InPinName](void* Ptr, const int32, const int32)
+	{
+		bMatch = (static_cast<FMovieSceneEvent*>(Ptr)->BoundObjectPinName == InPinName);
+		return bMatch;
+	};
+
+	PropertyHandle->EnumerateRawData(Enumerator);
+
+	return bMatch;
 }
 
-void FMovieSceneEventCustomization::PopulateQuickBindSubMenu(FMenuBuilder& MenuBuilder, UClass* TemplateClass)
+void FMovieSceneEventCustomization::SetBoundObjectPinName(FName InNewBoundObjectPinName)
 {
-	FMenuBuilder SubMenuBuilder(true, nullptr, nullptr, true);
+	FScopedTransaction Transaction(LOCTEXT("SetBoundObjectPinTransaction", "Set Bound Object Pin"));
+
+	TArray<UObject*> EditObjects;
+	GetEditObjects(EditObjects);
+
+	for (UObject* Object : EditObjects)
+	{
+		Object->Modify();
+	}
+
+	auto Enumerator = [InNewBoundObjectPinName](void* Ptr, const int32, const int32)
+	{
+		static_cast<FMovieSceneEvent*>(Ptr)->BoundObjectPinName = InNewBoundObjectPinName;
+		return true;
+	};
+	PropertyHandle->EnumerateRawData(Enumerator);
+
+	// Ensure that anything listening for property changed notifications are notified of the new binding
+	PropertyHandle->NotifyFinishedChangingProperties();
+	PropertyUtilities->ForceRefresh();
+}
+
+const FSlateBrush* FMovieSceneEventCustomization::GetBoundObjectPinIcon() const
+{
+	if (CompareBoundObjectPinName(NAME_None))
+	{
+		return nullptr;
+	}
+
+	return FEditorStyle::GetBrush("Graph.Pin.Disconnected_VarA");
+}
+
+FText FMovieSceneEventCustomization::GetBoundObjectPinText() const
+{
+	TArray<void*> RawData;
+	PropertyHandle->AccessRawData(RawData);
+
+	FName CommonPinName;
+
+	for (void* Ptr : RawData)
+	{
+		FName PinName = static_cast<FMovieSceneEvent*>(Ptr)->BoundObjectPinName;
+		if (CommonPinName == NAME_None)
+		{
+			CommonPinName = PinName;
+		}
+		else if (CommonPinName != PinName)
+		{
+			return LOCTEXT("MultiplePinValues", "Multiple Values");
+		}
+	}
+
+	return FText::FromName(CommonPinName);
+}
+
+void FMovieSceneEventCustomization::PopulateQuickBindSubMenu(FMenuBuilder& MenuBuilder, UMovieSceneSequence* Sequence)
+{
+	FMovieSceneSequenceEditor* SequenceEditor = FMovieSceneSequenceEditor::Find(Sequence);
+	if (!SequenceEditor)
+	{
+		return;
+	}
+
+	UBlueprint* Blueprint = SequenceEditor->GetOrCreateDirectorBlueprint(Sequence);
+	if (!Blueprint)
+	{
+		return;
+	}
+
+	UMovieScene* MovieScene = Sequence->GetMovieScene();
+
+	TOptional<FGuid> CommonObjectBindingID;
+	{
+		TArray<UObject*> EditObjects;
+		GetEditObjects(EditObjects);
 
-	FSlateIcon Icon(FEditorStyle::GetStyleSetName(), "GraphEditor.Function_16x");
-	static const FName DeprecatedFunctionName(TEXT("DeprecatedFunction"));
+		for (UObject* Outer : EditObjects)
+		{
+			FGuid ThisBindingID;
+			MovieScene->FindTrackBinding(*Outer->GetTypedOuter<UMovieSceneTrack>(), ThisBindingID);
+
+			if (CommonObjectBindingID.IsSet() && CommonObjectBindingID != ThisBindingID)
+			{
+				CommonObjectBindingID.Reset();
+				break;
+			}
+
+			CommonObjectBindingID = ThisBindingID;
+		}
+	}
+
+
+	FMovieSceneEventEndpointParameters EventParams = FMovieSceneEventEndpointParameters::Generate(MovieScene, CommonObjectBindingID.Get(FGuid()));
+
+	TSharedRef<SGraphActionMenu> ActionMenu = SNew(SGraphActionMenu)
+		.OnCreateCustomRowExpander_Static([](const FCustomExpanderData& Data) -> TSharedRef<SExpanderArrow> { return SNew(SExpanderArrow, Data.TableRow); })
+		.OnCollectAllActions_Static(UE4_MovieSceneEventCustomization::CollectQuickBindActions, Blueprint, EventParams.BoundObjectPinClass)
+		.OnActionSelected(this, &FMovieSceneEventCustomization::HandleQuickBindActionSelected, Blueprint, EventParams);
 
-	TArray<UFunction*> Functions;
+	ActionMenu->RegisterActiveTimer(0.f, FWidgetActiveTimerDelegate::CreateLambda(
+		[FilterTextBox = ActionMenu->GetFilterTextBox()](double, float)
+		{
+			FSlateApplication::Get().SetKeyboardFocus(FilterTextBox);
+			return EActiveTimerReturnType::Stop;
+		}
+	));
+
+	MenuBuilder.AddWidget(
+		SNew(SBox)
+		.WidthOverride(300.f)
+		.MaxDesiredHeight(500.f)
+		[
+			ActionMenu
+		],
+		FText()
+	);
+}
 
-	UClass* SuperClass = TemplateClass;
-	while (SuperClass)
+void FMovieSceneEventCustomization::HandleQuickBindActionSelected(const TArray<TSharedPtr<FEdGraphSchemaAction>>& SelectedAction, ESelectInfo::Type InSelectionType, UBlueprint* Blueprint, FMovieSceneEventEndpointParameters Params)
+{
+	if (InSelectionType != ESelectInfo::OnMouseClick && InSelectionType != ESelectInfo::OnKeyPress)
 	{
-		SubMenuBuilder.BeginSection(NAME_None, SuperClass->GetDisplayNameText());
+		return;
+	}
 
-		Functions.Reset();
-		for (UFunction* Function : TFieldRange<UFunction>(SuperClass, EFieldIteratorFlags::ExcludeSuper, EFieldIteratorFlags::ExcludeDeprecated))
+	for (TSharedPtr<FEdGraphSchemaAction> Action : SelectedAction)
+	{
+		if (Action->GetTypeId() == FBlueprintActionMenuItem::StaticGetTypeId())
 		{
-			if (Function->HasAllFunctionFlags(FUNC_BlueprintCallable | FUNC_Public) && !Function->HasMetaData(DeprecatedFunctionName) && !Function->HasAnyFunctionFlags(FUNC_EditorOnly))
+			const UBlueprintFunctionNodeSpawner* FunctionNodeSpawner = Cast<UBlueprintFunctionNodeSpawner>(static_cast<FBlueprintActionMenuItem*>(Action.Get())->GetRawAction());
+			const UFunction* FunctionToCall = FunctionNodeSpawner ? FunctionNodeSpawner->GetFunction() : nullptr;
+
+			if (FunctionToCall)
 			{
-				Functions.Add(Function);
+				UClass* OuterClass = CastChecked<UClass>(FunctionToCall->GetOuter());
+				if (OuterClass->ClassGeneratedBy == Blueprint && Blueprint->SkeletonGeneratedClass)
+				{
+					// Attempt to locate a custom event or a function graph of this name on the blueprint
+					for (UEdGraph* Graph : Blueprint->UbergraphPages)
+					{
+						for (UEdGraphNode* Node : Graph->Nodes)
+						{
+							UK2Node_CustomEvent* CustomEvent = Cast<UK2Node_CustomEvent>(Node);
+							if (CustomEvent && Blueprint->SkeletonGeneratedClass->FindFunctionByName(CustomEvent->GetFunctionName()) == FunctionToCall)
+							{
+								// Use this custom event
+								UEdGraphPin* BoundObjectPin = FMovieSceneEventUtils::FindBoundObjectPin(CustomEvent, Params.BoundObjectPinClass);
+								this->SetEventEndpoint(CustomEvent, BoundObjectPin, CustomEvent, EAutoCreatePayload::Variables);
+								return;
+							}
+						}
+					}
+
+					for (UEdGraph* Graph : Blueprint->FunctionGraphs)
+					{
+						if (Blueprint->SkeletonGeneratedClass->FindFunctionByName(Graph->GetFName()) != FunctionToCall)
+						{
+							continue;
+						}
+
+						// Use this function graph for the event endpoint
+						for (UEdGraphNode* Node : Graph->Nodes)
+						{
+							UK2Node_FunctionEntry* FunctionEntry = Cast<UK2Node_FunctionEntry>(Node);
+							if (FunctionEntry)
+							{
+								UEdGraphPin* BoundObjectPin = FMovieSceneEventUtils::FindBoundObjectPin(FunctionEntry, Params.BoundObjectPinClass);
+								this->SetEventEndpoint(FunctionEntry, BoundObjectPin, FunctionEntry, EAutoCreatePayload::Variables);
+								return;
+							}
+						}
+					}
+				}
 			}
 		}
-	
-		Algo::SortBy(Functions, &UFunction::GetFName, FNameLexicalLess());
 
-		for (UFunction* Function : Functions)
+		if (Action)
 		{
-			SubMenuBuilder.AddMenuEntry(
-				FText::FromName(Function->GetFName()),
-				FText(),
-				Icon,
-				FUIAction(
-					FExecuteAction::CreateSP(this, &FMovieSceneEventCustomization::CreateEventEndpointFromFunction, Function, TemplateClass)
-				)
-			);
+			UK2Node_CustomEvent* NewEventEndpoint = FMovieSceneEventUtils::CreateUserFacingEvent(Blueprint, Params);
+
+			UEdGraphPin* ThenPin = NewEventEndpoint->FindPin(UEdGraphSchema_K2::PN_Then, EGPD_Output);
+			UEdGraphPin* BoundObjectPin = FMovieSceneEventUtils::FindBoundObjectPin(NewEventEndpoint, Params.BoundObjectPinClass);
+
+			FVector2D NodePosition(NewEventEndpoint->NodePosX + 400.f, NewEventEndpoint->NodePosY);
+			UEdGraphNode* NewNode = Action->PerformAction(NewEventEndpoint->GetGraph(), BoundObjectPin ? BoundObjectPin : ThenPin, NodePosition);
+
+			if (NewNode)
+			{
+				UEdGraphPin* ExecPin = NewNode->FindPin(UEdGraphSchema_K2::PN_Execute, EGPD_Input);
+				if (ensure(ThenPin && ExecPin))
+				{
+					ThenPin->MakeLinkTo(ExecPin);
+				}
+			}
+
+			this->SetEventEndpoint(NewEventEndpoint, BoundObjectPin, Cast<UK2Node>(NewNode), EAutoCreatePayload::Pins | EAutoCreatePayload::Variables);
 		}
+	}
+}
+
+void FMovieSceneEventCustomization::PopulateRebindSubMenu(FMenuBuilder& MenuBuilder, UMovieSceneSequence* Sequence)
+{
+	FMovieSceneSequenceEditor* SequenceEditor = FMovieSceneSequenceEditor::Find(Sequence);
+	if (!SequenceEditor)
+	{
+		return;
+	}
 
-		SubMenuBuilder.EndSection();
+	UBlueprint* Blueprint = SequenceEditor->GetOrCreateDirectorBlueprint(Sequence);
+	if (!Blueprint)
+	{
+		return;
+	}
+
+	UMovieScene* MovieScene = Sequence->GetMovieScene();
+
+	TOptional<FGuid> CommonObjectBindingID;
+	{
+		TArray<UObject*> EditObjects;
+		GetEditObjects(EditObjects);
+
+		for (UObject* Outer : EditObjects)
+		{
+			FGuid ThisBindingID;
+			MovieScene->FindTrackBinding(*Outer->GetTypedOuter<UMovieSceneTrack>(), ThisBindingID);
+
+			if (CommonObjectBindingID.IsSet() && CommonObjectBindingID != ThisBindingID)
+			{
+				CommonObjectBindingID.Reset();
+				break;
+			}
 
-		SuperClass = SuperClass->GetSuperClass();
+			CommonObjectBindingID = ThisBindingID;
+		}
 	}
 
-	// Limit the height intentionally so that it doesn't go offscreen
+	FMovieSceneEventEndpointParameters EventParams = FMovieSceneEventEndpointParameters::Generate(MovieScene, CommonObjectBindingID.Get(FGuid()));
+
+	TSharedRef<SGraphActionMenu> ActionMenu = SNew(SGraphActionMenu)
+		.OnCreateCustomRowExpander_Static([](const FCustomExpanderData& Data) -> TSharedRef<SExpanderArrow> { return SNew(SExpanderArrow, Data.TableRow); })
+		.OnCollectAllActions_Static(UE4_MovieSceneEventCustomization::CollectAllRebindActions, Blueprint)
+		.OnActionSelected(this, &FMovieSceneEventCustomization::HandleRebindActionSelected, Blueprint, EventParams.BoundObjectPinClass);
+
+	ActionMenu->RegisterActiveTimer(0.f, FWidgetActiveTimerDelegate::CreateLambda(
+		[FilterTextBox = ActionMenu->GetFilterTextBox()](double, float)
+		{
+			FSlateApplication::Get().SetKeyboardFocus(FilterTextBox);
+			return EActiveTimerReturnType::Stop;
+		}
+	));
+
 	MenuBuilder.AddWidget(
 		SNew(SBox)
-		.MaxDesiredHeight(600.0f)
+		.WidthOverride(300.f)
+		.MaxDesiredHeight(500.f)
 		[
-			SubMenuBuilder.MakeWidget()
-		]
-		, FText(), true, false
+			ActionMenu
+		],
+		FText()
 	);
 }
 
+void FMovieSceneEventCustomization::HandleRebindActionSelected(const TArray<TSharedPtr<FEdGraphSchemaAction>>& SelectedAction, ESelectInfo::Type InSelectionType, UBlueprint* Blueprint, UClass* BoundObjectPinClass)
+{
+	if (InSelectionType != ESelectInfo::OnMouseClick && InSelectionType != ESelectInfo::OnKeyPress)
+	{
+		return;
+	}
+
+	if (!Blueprint->SkeletonGeneratedClass)
+	{
+		return;
+	}
+
+	for (TSharedPtr<FEdGraphSchemaAction> Action : SelectedAction)
+	{
+		if (Action->GetTypeId() != FBlueprintActionMenuItem::StaticGetTypeId())
+		{
+			continue;
+		}
+
+		const UBlueprintFunctionNodeSpawner* FunctionNodeSpawner = Cast<UBlueprintFunctionNodeSpawner>(static_cast<FBlueprintActionMenuItem*>(Action.Get())->GetRawAction());
+		const UFunction* FunctionToCall = FunctionNodeSpawner ? FunctionNodeSpawner->GetFunction() : nullptr;
+		if (!FunctionToCall)
+		{
+			continue;
+		}
+
+		UClass* OuterClass = CastChecked<UClass>(FunctionToCall->GetOuter());
+		if (OuterClass->ClassGeneratedBy != Blueprint)
+		{
+			continue;
+		}
+
+		// Attempt to locate a custom event or a function graph of this name on the blueprint
+		for (UEdGraph* Graph : Blueprint->UbergraphPages)
+		{
+			for (UEdGraphNode* Node : Graph->Nodes)
+			{
+				UK2Node_CustomEvent* CustomEvent = Cast<UK2Node_CustomEvent>(Node);
+				if (CustomEvent && Blueprint->SkeletonGeneratedClass->FindFunctionByName(CustomEvent->GetFunctionName()) == FunctionToCall)
+				{
+					// Use this custom event
+					UEdGraphPin* BoundObjectPin = FMovieSceneEventUtils::FindBoundObjectPin(CustomEvent, BoundObjectPinClass);
+					this->SetEventEndpoint(CustomEvent, BoundObjectPin, CustomEvent, EAutoCreatePayload::Variables);
+					return;
+				}
+			}
+		}
+
+		for (UEdGraph* Graph : Blueprint->FunctionGraphs)
+		{
+			if (Blueprint->SkeletonGeneratedClass->FindFunctionByName(Graph->GetFName()) != FunctionToCall)
+			{
+				continue;
+			}
+
+			// Use this function graph for the event endpoint
+			for (UEdGraphNode* Node : Graph->Nodes)
+			{
+				UK2Node_FunctionEntry* FunctionEntry = Cast<UK2Node_FunctionEntry>(Node);
+				if (FunctionEntry)
+				{
+					UEdGraphPin* BoundObjectPin = FMovieSceneEventUtils::FindBoundObjectPin(FunctionEntry, BoundObjectPinClass);
+					this->SetEventEndpoint(FunctionEntry, BoundObjectPin, FunctionEntry, EAutoCreatePayload::Variables);
+					return;
+				}
+			}
+		}
+	}
+
+	ensureMsgf(false, TEXT("Unknown blueprint action type encountered for rebinding"));
+}
+
 const FSlateBrush* FMovieSceneEventCustomization::GetEventIcon() const
 {
-	UK2Node_FunctionEntry* CommonEndpoint = GetCommonEndpoint();
+	UK2Node* CommonEndpoint = GetCommonEndpoint();
 
-	if (FMovieSceneEvent::IsValidFunction(CommonEndpoint))
+	if (CommonEndpoint)
 	{
-		return FEditorStyle::GetBrush("GraphEditor.Function_16x");
+		FLinearColor Color;
+		FSlateIcon EndpointIcon = CommonEndpoint->GetIconAndTint(Color);
+		return EndpointIcon.GetIcon();
 	}
-	else if (!CommonEndpoint)
+	else
 	{
 		TArray<void*> RawData;
 		PropertyHandle->AccessRawData(RawData);
@@ -424,13 +1175,13 @@ const FSlateBrush* FMovieSceneEventCustomization::GetEventIcon() const
 
 FText FMovieSceneEventCustomization::GetEventName() const
 {
-	UK2Node_FunctionEntry* CommonEndpoint = GetCommonEndpoint();
+	UEdGraphNode* CommonEndpoint = GetCommonEndpoint();
 
-	if (FMovieSceneEvent::IsValidFunction(CommonEndpoint))
+	if (CommonEndpoint)
 	{
 		return CommonEndpoint->GetNodeTitle(ENodeTitleType::MenuTitle);
 	}
-	else if (!CommonEndpoint)
+	else
 	{
 		TArray<void*> RawData;
 		PropertyHandle->AccessRawData(RawData);
@@ -443,209 +1194,232 @@ FText FMovieSceneEventCustomization::GetEventName() const
 	return LOCTEXT("UnboundText", "Unbound");
 }
 
-EVisibility FMovieSceneEventCustomization::GetErrorVisibility() const
-{
-	UK2Node_FunctionEntry* CommonEndpoint = GetCommonEndpoint();
-	if (CommonEndpoint)
-	{
-		return FMovieSceneEvent::IsValidFunction(CommonEndpoint) ? EVisibility::Collapsed : EVisibility::Visible;
-	}
-	return EVisibility::Collapsed;
-}
-
-FText FMovieSceneEventCustomization::GetErrorTooltip() const
-{
-	UK2Node_FunctionEntry* CommonEndpoint = GetCommonEndpoint();
-	if (CommonEndpoint)
-	{
-		return FText::Format(LOCTEXT("ErrorToolTipFormat", "The currently assigned function '{0}' is not a valid event endpoint. Event endpoints must have a single object or interface parameter that is not passed by-reference."), CommonEndpoint->GetNodeTitle(ENodeTitleType::MenuTitle));
-	}
-	return FText::GetEmpty();
-}
-
-void FMovieSceneEventCustomization::SetEventEndpoint(UK2Node_FunctionEntry* NewEndpoint)
+void FMovieSceneEventCustomization::SetEventEndpoint(UK2Node* NewEndpoint, UEdGraphPin* BoundObjectPin, UK2Node* PayloadTemplate, EAutoCreatePayload AutoCreatePayload)
 {
 	FScopedTransaction Transaction(LOCTEXT("SetEventEndpoint", "Set Event Endpoint"));
 
-	UBlueprint* Blueprint = NewEndpoint->GetBlueprint();
-
-	CachedCommonEndpoint = NewEndpoint;
-
 	// Modify and assign the blueprint for outer sections
 	TArray<UObject*> EditObjects;
 	GetEditObjects(EditObjects);
 
+	TArray<UBlueprint*> AllBlueprints;
+
+	// If we're assigning a new valid endpoint, it must reside within the same blueprint as everything we're assigning it to.
+	// Anything else must be implemented as a call function node connected to a custom event node
+	UBlueprint* Blueprint = NewEndpoint ? NewEndpoint->GetBlueprint() : nullptr;
+	if (Blueprint)
+	{
+		for (UObject* Outer : EditObjects)
+		{
+			UMovieSceneSequence*       Sequence           = Outer->GetTypedOuter<UMovieSceneSequence>();
+			FMovieSceneSequenceEditor* SequenceEditor     = FMovieSceneSequenceEditor::Find(Sequence);
+			UBlueprint*                SequenceDirectorBP = SequenceEditor ? SequenceEditor->GetOrCreateDirectorBlueprint(Sequence) : nullptr;
+
+			if (!ensureAlwaysMsgf(SequenceDirectorBP == Blueprint, TEXT("Attempting to assign an event endpoint to an event with a different sequence director Blueprint.")))
+			{
+				Transaction.Cancel();
+				return;
+			}
+		}
+
+		Blueprint->Modify();
+	}
+
 	for (UObject* Outer : EditObjects)
 	{
 		UMovieSceneEventSectionBase* BaseEventSection = Cast<UMovieSceneEventSectionBase>(Outer);
 		if (BaseEventSection)
 		{
 			BaseEventSection->Modify();
-			BaseEventSection->SetDirectorBlueprint(Blueprint);
+			if (Blueprint)
+			{
+				Blueprint->GenerateFunctionGraphsEvent.AddUniqueDynamic(BaseEventSection, &UMovieSceneEventSectionBase::HandleGenerateEntryPoints);
+			}
 		}
 	}
 
-	// Assign the endpoints to all events
-	TArray<void*> RawData;
-	PropertyHandle->AccessRawData(RawData);
-
-	for (void* Ptr : RawData)
+	TArray<FName> PayloadNames;
+	if (PayloadTemplate && EnumHasAnyFlags(AutoCreatePayload, EAutoCreatePayload::Variables | EAutoCreatePayload::Pins))
 	{
-		FMovieSceneEvent* Event = static_cast<FMovieSceneEvent*>(Ptr);
-		if (Event)
+		UK2Node_EditablePinBase* EditableNode = Cast<UK2Node_EditablePinBase>(NewEndpoint);
+
+		for (UEdGraphPin* PayloadPin : PayloadTemplate->Pins)
 		{
-			Event->SetFunctionEntry(NewEndpoint);
+			if (PayloadPin != BoundObjectPin && PayloadPin->Direction == EGPD_Input && PayloadPin->PinType.PinCategory != UEdGraphSchema_K2::PC_Exec && PayloadPin->LinkedTo.Num() == 0)
+			{
+				// Make a payload variable for this pin
+				if (EnumHasAnyFlags(AutoCreatePayload, EAutoCreatePayload::Variables))
+				{
+					PayloadNames.Add(PayloadPin->PinName);
+				}
+
+				if (EditableNode && EnumHasAnyFlags(AutoCreatePayload, EAutoCreatePayload::Pins))
+				{
+					UEdGraphPin* NewPin = EditableNode->CreateUserDefinedPin(PayloadPin->PinName, PayloadPin->PinType, EGPD_Output);
+					if (PayloadTemplate != NewEndpoint && NewPin)
+					{
+						NewPin->MakeLinkTo(PayloadPin);
+					}
+				}
+			}
 		}
 	}
 
-	// Ensure that anything listening for property changed notifications are notified of the new binding
-	PropertyHandle->NotifyFinishedChangingProperties();
-
-	FBlueprintEditorUtils::MarkBlueprintAsModified(Blueprint);
-}
+	// Assign the endpoints to all events
+	TArray<void*> RawData;
+	PropertyHandle->AccessRawData(RawData);
 
-bool FMovieSceneEventCustomization::CompareCurrentEventEndpoint(UK2Node_FunctionEntry* NewEndpoint)
-{
-	return CachedCommonEndpoint == NewEndpoint;
-}
+	check(EditObjects.Num() == RawData.Num());
 
-void FMovieSceneEventCustomization::CreateEventEndpoint()
-{
-	UMovieSceneSequence*       Sequence       = GetCommonSequence();
-	FMovieSceneSequenceEditor* SequenceEditor = FMovieSceneSequenceEditor::Find(Sequence);
-	if (!SequenceEditor)
+	for (int32 Index = 0; Index < RawData.Num(); ++Index)
 	{
-		return;
-	}
+		UMovieSceneEventSectionBase* EventSection = Cast<UMovieSceneEventSectionBase>(EditObjects[Index]);
+		FMovieSceneEvent*  EntryPoint   = static_cast<FMovieSceneEvent*>(RawData[Index]);
 
-	FScopedTransaction Transaction(LOCTEXT("CreateEventEndpoint", "Create Event Endpoint"));
+		if (EntryPoint)
+		{
+			FMovieSceneEventUtils::SetEndpoint(EntryPoint, EventSection, NewEndpoint, BoundObjectPin);
 
-	// Create a single event binding and point all events in this property handle to it
-	UK2Node_FunctionEntry* NewFunctionEntry = SequenceEditor->CreateEventEndpoint(Sequence);
-	if (!ensure(NewFunctionEntry))
-	{
-		return;
+			for (FName PayloadVar : PayloadNames)
+			{
+				if (!EntryPoint->PayloadVariables.Contains(PayloadVar))
+				{
+					EntryPoint->PayloadVariables.Add(PayloadVar);
+				}
+			}
+		}
 	}
 
-	if (UMovieSceneEventTrack* CommonTrack = GetCommonTrack())
+	if (Blueprint)
 	{
-		SequenceEditor->InitializeEndpointForTrack(CommonTrack, NewFunctionEntry);
+		//FBlueprintEditorUtils::MarkBlueprintAsModified(Blueprint);
+		//FBlueprintEditorUtils::MarkBlueprintAsModified(Blueprint);
+		FKismetEditorUtilities::CompileBlueprint(Blueprint);
 	}
 
-	SetEventEndpoint(NewFunctionEntry);
-	FKismetEditorUtilities::BringKismetToFocusAttentionOnObject(NewFunctionEntry, false);
+	// Ensure that anything listening for property changed notifications are notified of the new binding
+	PropertyHandle->NotifyFinishedChangingProperties();
+	PropertyUtilities->ForceRefresh();
 }
 
-void FMovieSceneEventCustomization::CreateEventEndpointFromFunction(UFunction* QuickBindFunction, UClass* PinClassType)
+void FMovieSceneEventCustomization::CreateEventEndpoint()
 {
-	UMovieSceneSequence*       Sequence       = GetCommonSequence();
-	FMovieSceneSequenceEditor* SequenceEditor = FMovieSceneSequenceEditor::Find(Sequence);
-
-	if (!SequenceEditor)
+	// @todo: These TLess<const T*> parameters should not be necessary but TSortedMap has a poor default that uses T rather than ConstPointerType
+	struct FSectionData
 	{
-		return;
-	}
-
-	FScopedTransaction Transaction(LOCTEXT("CreateEventEndpoint", "Create Event Endpoint"));
+		TArray<FMovieSceneEvent*, TInlineAllocator<1>>  EntryPoints;
+	};
+	struct FTrackData
+	{
+		TSortedMap<UMovieSceneEventSectionBase*, FSectionData, FDefaultAllocator, TLess<const UMovieSceneEventSectionBase*>>  Sections;
+	};
+	struct FSequenceData
+	{
+		TSortedMap<UMovieSceneEventTrack*, FTrackData, FDefaultAllocator, TLess<const UMovieSceneEventTrack*>> Tracks;
+	};
 
-	FString DesiredNewEventName = FString(TEXT("Call ")) + QuickBindFunction->GetName();
+	TSortedMap<UMovieSceneSequence*, FSequenceData, FDefaultAllocator, TLess<const UMovieSceneSequence*>> PerSequenceData;
 
-	// Create a single event binding and point all events in this property handle to it
-	UK2Node_FunctionEntry* NewFunctionEntry = SequenceEditor->CreateEventEndpoint(Sequence, DesiredNewEventName);
-	if (!ensure(NewFunctionEntry))
+	// Populate all the sequences represented by this customization
 	{
-		return;
-	}
+		TArray<void*> RawData;
+		PropertyHandle->AccessRawData(RawData);
 
-	if (ensure(NewFunctionEntry->UserDefinedPins.Num() == 0))
-	{
-		FEdGraphPinType PinType;
-		PinType.PinCategory = PinClassType->IsChildOf(UInterface::StaticClass()) ? UEdGraphSchema_K2::PC_Interface : UEdGraphSchema_K2::PC_Object;
-		PinType.PinSubCategoryObject = PinClassType;
+		TArray<UObject*> EditObjects;
+		GetEditObjects(EditObjects);
 
-		NewFunctionEntry->CreateUserDefinedPin(FMovieSceneSequenceEditor::TargetPinName, PinType, EGPD_Output, true);
-		NewFunctionEntry->ReconstructNode();
-	}
+		check(RawData.Num() == EditObjects.Num());
 
-	NewFunctionEntry->bCommentBubblePinned = false;
-	NewFunctionEntry->bCommentBubbleVisible = false;
+		for (int32 Index = 0; Index < RawData.Num(); ++Index)
+		{
+			UMovieSceneEventSectionBase* EventSection = Cast<UMovieSceneEventSectionBase>(EditObjects[Index]);
+			FMovieSceneEvent*  EntryPoint   = static_cast<FMovieSceneEvent*>(RawData[Index]);
 
-	UEdGraph* Graph = NewFunctionEntry->GetGraph();
+			if (EventSection)
+			{
+				FSequenceData& SequenceData = PerSequenceData.FindOrAdd(EventSection->GetTypedOuter<UMovieSceneSequence>());
+				FTrackData&    TrackData    = SequenceData.Tracks.FindOrAdd(EventSection->GetTypedOuter<UMovieSceneEventTrack>());
+				FSectionData&  SectionData  = TrackData.Sections.FindOrAdd(EventSection);
 
-	// Make a call function template
-	UK2Node_CallFunction* CallFuncNode = NewObject<UK2Node_CallFunction>(Graph);
-	CallFuncNode->NodePosX = NewFunctionEntry->NodePosX + NewFunctionEntry->NodeWidth + 200;
-	CallFuncNode->NodePosY = NewFunctionEntry->NodePosY;
-	CallFuncNode->CreateNewGuid();
-	CallFuncNode->SetFromFunction(QuickBindFunction);
-	CallFuncNode->PostPlacedNewNode();
-	CallFuncNode->ReconstructNode();
+				SectionData.EntryPoints.Add(EntryPoint);
+			}
+		}
+	}
 
-	Graph->AddNode(CallFuncNode, false, false);
+	FScopedTransaction Transaction(LOCTEXT("CreateEventEndpoint", "Create Event Endpoint"));
 
+	for (const TPair<UMovieSceneSequence*, FSequenceData>& SequencePair : PerSequenceData)
 	{
-		// Connect the exec pins together
-		UEdGraphPin* ThenPin = NewFunctionEntry->FindPin(UEdGraphSchema_K2::PN_Then);
-		UEdGraphPin* ExecPin = CallFuncNode->GetExecPin();
+		FMovieSceneSequenceEditor* SequenceEditor = FMovieSceneSequenceEditor::Find(SequencePair.Key);
+		if (!SequenceEditor)
+		{
+			continue;
+		}
 
-		if (ThenPin && ExecPin)
+		UBlueprint* SequenceDirectorBP = SequenceEditor->GetOrCreateDirectorBlueprint(SequencePair.Key);
+		if (!SequenceDirectorBP)
 		{
-			const UEdGraphSchema* Schema = ThenPin->GetSchema();
-			Schema->TryCreateConnection(ThenPin, ExecPin);
+			continue;
 		}
-	}
 
-	// Connect the object target pin to the self (input) pin on the call function node
-	if (UEdGraphPin* SelfPin = CallFuncNode->FindPin(UEdGraphSchema_K2::PSC_Self))
-	{
-		for (UEdGraphPin* Pin : NewFunctionEntry->Pins)
+		UMovieScene* MovieScene = SequencePair.Key->GetMovieScene();
+
+		TOptional<FGuid> CommonObjectBindingID;
+		for (const TPair<UMovieSceneEventTrack*, FTrackData>& TrackPair : SequencePair.Value.Tracks)
 		{
-			if (Pin->PinType.PinCategory == UEdGraphSchema_K2::PC_Interface  || Pin->PinType.PinCategory == UEdGraphSchema_K2::PC_Object)
+			FGuid ThisBindingID;
+			MovieScene->FindTrackBinding(*TrackPair.Key, ThisBindingID);
+
+			if (CommonObjectBindingID.IsSet() && CommonObjectBindingID != ThisBindingID)
 			{
-				const UEdGraphSchema* Schema = Pin->GetSchema();
-				Schema->TryCreateConnection(Pin, SelfPin);
+				CommonObjectBindingID.Reset();
+				break;
 			}
-		}
-	}
 
-	SetEventEndpoint(NewFunctionEntry);
-	FKismetEditorUtilities::BringKismetToFocusAttentionOnObject(NewFunctionEntry, false);
-}
+			CommonObjectBindingID = ThisBindingID;
+		}
 
-void FMovieSceneEventCustomization::ClearEventEndpoint()
-{
-	FScopedTransaction Transaction(LOCTEXT("ClearEventEndpoint", "Clear Event Endpoint"));
+		FMovieSceneEventEndpointParameters Parameters = FMovieSceneEventEndpointParameters::Generate(MovieScene, CommonObjectBindingID.Get(FGuid()));
 
-	TArray<UObject*> EditObjects;
-	GetEditObjects(EditObjects);
+		SequenceDirectorBP->Modify();
 
-	for (UObject* Outer : EditObjects)
-	{
-		if (Outer)
+		UK2Node_CustomEvent* NewEventEndpoint = FMovieSceneEventUtils::CreateUserFacingEvent(SequenceDirectorBP, Parameters);
+		if (!NewEventEndpoint)
 		{
-			Outer->Modify();
+			continue;
 		}
-	}
-
-	TArray<void*> RawData;
-	PropertyHandle->AccessRawData(RawData);
 
-	for (void* Ptr : RawData)
-	{
-		FMovieSceneEvent* Event = static_cast<FMovieSceneEvent*>(Ptr);
-		if (Event)
+		UEdGraphPin* BoundObjectPin = FMovieSceneEventUtils::FindBoundObjectPin(NewEventEndpoint, Parameters.BoundObjectPinClass);
+		for (const TPair<UMovieSceneEventTrack*, FTrackData>& TrackPair : SequencePair.Value.Tracks)
 		{
-			Event->SetFunctionEntry(nullptr);
+			for (const TPair<UMovieSceneEventSectionBase*, FSectionData>& SectionPair : TrackPair.Value.Sections)
+			{
+				SectionPair.Key->Modify();
+				SequenceDirectorBP->GenerateFunctionGraphsEvent.AddUniqueDynamic(SectionPair.Key, &UMovieSceneEventSectionBase::HandleGenerateEntryPoints);
+
+				for (FMovieSceneEvent* EntryPoint : SectionPair.Value.EntryPoints)
+				{
+					FMovieSceneEventUtils::SetEndpoint(EntryPoint, SectionPair.Key, NewEventEndpoint, BoundObjectPin);
+				}
+			}
 		}
+
+		FBlueprintEditorUtils::MarkBlueprintAsModified(SequenceDirectorBP);
 	}
 
 	// Ensure that anything listening for property changed notifications are notified of the new binding
 	PropertyHandle->NotifyFinishedChangingProperties();
+	PropertyUtilities->ForceRefresh();
+}
+
+void FMovieSceneEventCustomization::ClearEventEndpoint()
+{
+	SetEventEndpoint(nullptr, nullptr, nullptr, EAutoCreatePayload::None);
 }
 
 void FMovieSceneEventCustomization::NavigateToDefinition()
 {
-	UK2Node_FunctionEntry* CommonEndpoint = GetCommonEndpoint();
+	UEdGraphNode* CommonEndpoint = GetCommonEndpoint();
 	if (CommonEndpoint)
 	{
 		FKismetEditorUtilities::BringKismetToFocusAttentionOnObject(CommonEndpoint, false);
diff --git a/Engine/Source/Editor/MovieSceneTools/Private/MovieSceneEventUtils.cpp b/Engine/Source/Editor/MovieSceneTools/Private/MovieSceneEventUtils.cpp
new file mode 100644
--- /dev/null
+++ b/Engine/Source/Editor/MovieSceneTools/Private/MovieSceneEventUtils.cpp
@@ -0,0 +1,446 @@
+// Copyright 1998-2019 Epic Games, Inc. All Rights Reserved.
+
+#include "MovieSceneEventUtils.h"
+
+#include "Algo/Find.h"
+
+#include "UObject/StrongObjectPtr.h"
+#include "UObject/Package.h"
+
+#include "MovieSceneSequence.h"
+#include "Sections/MovieSceneEventSectionBase.h"
+#include "Tracks/MovieSceneEventTrack.h"
+
+#include "ScopedTransaction.h"
+
+#include "Engine/Blueprint.h"
+
+#include "KismetCompiler.h"
+#include "K2Node_CustomEvent.h"
+#include "K2Node_VariableGet.h"
+#include "K2Node_CallFunction.h"
+#include "K2Node_FunctionEntry.h"
+#include "K2Node_EditablePinBase.h"
+#include "Kismet2/BlueprintEditorUtils.h"
+#include "Kismet2/KismetEditorUtilities.h"
+#include "Kismet2/Kismet2NameValidators.h"
+
+#define LOCTEXT_NAMESPACE "MovieSceneEventUtils"
+
+FMovieSceneEventEndpointParameters FMovieSceneEventEndpointParameters::Generate(UMovieSceneEventTrack* Track)
+{
+	check(Track);
+
+	UMovieScene* MovieScene = Track->GetTypedOuter<UMovieScene>();
+
+	FGuid ObjectBindingID;
+	if (MovieScene->FindTrackBinding(*Track, ObjectBindingID))
+	{
+		return Generate(MovieScene, ObjectBindingID);
+	}
+
+	FMovieSceneEventEndpointParameters Params;
+
+	if (Track->EventReceivers.Num() != 0)
+	{
+		// By default for master tracks we create a pin (that receives the level blueprint for master event tracks, or the event receivers)
+		Params.BoundObjectPinClass = UObject::StaticClass();
+
+		Params.SanitizedObjectName = TEXT("Event Receiver");
+		Params.SanitizedEventName  = TEXT("SequenceEvent");
+	}
+	else
+	{
+		Params.SanitizedObjectName = TEXT("None");
+		Params.SanitizedEventName  = TEXT("SequenceEvent");
+	}
+
+	return Params;
+}
+
+FMovieSceneEventEndpointParameters FMovieSceneEventEndpointParameters::Generate(UMovieScene* MovieScene, const FGuid& ObjectBindingID)
+{
+	check(MovieScene);
+
+	FMovieSceneEventEndpointParameters Params;
+
+	if (ObjectBindingID.IsValid())
+	{
+		Params.SanitizedObjectName = MovieScene->GetObjectDisplayName(ObjectBindingID).ToString();
+
+		for (TCHAR& Char : Params.SanitizedObjectName)
+		{
+			if (FCString::Strchr(INVALID_NAME_CHARACTERS, Char) != nullptr)
+			{
+				Char = '_';
+			}
+		}
+
+		Params.SanitizedEventName = Params.SanitizedObjectName + TEXT("_Event");
+
+		if (FMovieScenePossessable* Possessable = MovieScene->FindPossessable(ObjectBindingID))
+		{
+			Params.BoundObjectPinClass = const_cast<UClass*>(Possessable->GetPossessedObjectClass());
+		}
+		else if(FMovieSceneSpawnable* Spawnable = MovieScene->FindSpawnable(ObjectBindingID))
+		{
+			Params.BoundObjectPinClass = Spawnable->GetObjectTemplate()->GetClass();
+		}
+	}
+	else
+	{
+		Params.BoundObjectPinClass = nullptr;
+		Params.SanitizedObjectName = TEXT("None");
+		Params.SanitizedEventName  = TEXT("SequenceEvent");
+	}
+
+	return Params;
+}
+
+UK2Node_CustomEvent* FMovieSceneEventUtils::BindNewUserFacingEvent(FMovieSceneEvent* EntryPoint, UMovieSceneEventSectionBase* EventSection, UBlueprint* Blueprint)
+{
+	check(EntryPoint && EventSection && Blueprint);
+
+	UMovieSceneEventTrack* Track = EventSection->GetTypedOuter<UMovieSceneEventTrack>();
+
+	// Modify necessary objects
+	EventSection->Modify();
+	Blueprint->Modify();
+
+	// Ensure the section is bound to the blueprint function generation event
+	Blueprint->GenerateFunctionGraphsEvent.AddUniqueDynamic(EventSection, &UMovieSceneEventSectionBase::HandleGenerateEntryPoints);
+
+	// Create the new user-facing event node
+	FMovieSceneEventEndpointParameters Params = FMovieSceneEventEndpointParameters::Generate(Track);
+	UK2Node_CustomEvent* NewEventNode = CreateUserFacingEvent(Blueprint, Params);
+
+	if (NewEventNode)
+	{
+		// Bind the node to the event entry point
+		UEdGraphPin* BoundObjectPin = FMovieSceneEventUtils::FindBoundObjectPin(NewEventNode, Params.BoundObjectPinClass);
+		FMovieSceneEventUtils::SetEndpoint(EntryPoint, EventSection, NewEventNode, BoundObjectPin);
+	}
+
+	return NewEventNode;
+}
+
+UK2Node_CustomEvent* FMovieSceneEventUtils::CreateUserFacingEvent(UBlueprint* Blueprint, const FMovieSceneEventEndpointParameters& Parameters)
+{
+	static const TCHAR* const EventGraphName = TEXT("Sequencer Events");
+	UEdGraph* SequenceEventGraph = FindObject<UEdGraph>(Blueprint, EventGraphName);
+	if (!SequenceEventGraph)
+	{
+		SequenceEventGraph = FBlueprintEditorUtils::CreateNewGraph(Blueprint, EventGraphName, UEdGraph::StaticClass(), UEdGraphSchema_K2::StaticClass());
+
+		SequenceEventGraph->GraphGuid = FGuid::NewGuid();
+		Blueprint->UbergraphPages.Add(SequenceEventGraph);
+	}
+
+	// Create a custom event node to replace the original event node imported from text
+	UK2Node_CustomEvent* CustomEventNode = NewObject<UK2Node_CustomEvent>(SequenceEventGraph);
+
+	check(Parameters.SanitizedEventName.Len() != 0);
+	CustomEventNode->CustomFunctionName = FBlueprintEditorUtils::FindUniqueKismetName(Blueprint, Parameters.SanitizedEventName);
+
+	// Ensure that it is editable
+	CustomEventNode->bIsEditable = true;
+
+	CustomEventNode->CreateNewGuid();
+	CustomEventNode->PostPlacedNewNode();
+	CustomEventNode->AllocateDefaultPins();
+
+	const FVector2D NewPosition = SequenceEventGraph->GetGoodPlaceForNewNode();
+	CustomEventNode->NodePosX = NewPosition.X;
+	CustomEventNode->NodePosY = NewPosition.Y;
+
+	if (Parameters.BoundObjectPinClass != nullptr)
+	{
+		FEdGraphPinType PinType;
+		PinType.PinCategory = UEdGraphSchema_K2::PC_Object;
+		PinType.PinSubCategoryObject = Parameters.BoundObjectPinClass;
+
+		CustomEventNode->CreateUserDefinedPin(*Parameters.SanitizedObjectName, PinType, EGPD_Output, true);
+	}
+
+	SequenceEventGraph->AddNode(CustomEventNode, false, false);
+
+	FBlueprintEditorUtils::ValidateBlueprintChildVariables(Blueprint, CustomEventNode->CustomFunctionName);
+	FBlueprintEditorUtils::MarkBlueprintAsStructurallyModified(Blueprint);
+	return CustomEventNode;
+}
+
+UK2Node* FMovieSceneEventUtils::FindEndpoint(FMovieSceneEvent* EntryPoint, UMovieSceneEventSectionBase* EventSection, UBlueprint* OwnerBlueprint)
+{
+	check(OwnerBlueprint);
+	check(EntryPoint);
+
+	if (EntryPoint->WeakCachedEndpoint.IsStale())
+	{
+		return nullptr;
+	}
+
+	UK2Node* CachedEndpoint = CastChecked<UK2Node>(EntryPoint->WeakCachedEndpoint.Get(), ECastCheckedType::NullAllowed);
+	if (CachedEndpoint)
+	{
+		if (FBlueprintEditorUtils::FindBlueprintForNode(CachedEndpoint) == OwnerBlueprint)
+		{
+			return CachedEndpoint;
+		}
+
+		CachedEndpoint->OnUserDefinedPinRenamed().RemoveAll(EventSection);
+		EntryPoint->WeakCachedEndpoint = nullptr;
+		return nullptr;
+	}
+
+	// The cached entry point is either null or stale
+	EntryPoint->WeakCachedEndpoint = nullptr;
+	if (!EntryPoint->GraphGuid.IsValid())
+	{
+		return nullptr;
+	}
+
+	if (EntryPoint->NodeGuid.IsValid())
+	{
+		for (UEdGraph* Graph : OwnerBlueprint->UbergraphPages)
+		{
+			if (Graph->GraphGuid == EntryPoint->GraphGuid)
+			{
+				for (UEdGraphNode* Node : Graph->Nodes)
+				{
+					if (Node->NodeGuid == EntryPoint->NodeGuid)
+					{
+						UK2Node_CustomEvent* CustomEvent = Cast<UK2Node_CustomEvent>(Node);
+						if (ensureMsgf(CustomEvent, TEXT("Encountered an event entry point node that is bound to something other than a custom event")))
+						{
+							CustomEvent->OnUserDefinedPinRenamed().AddUObject(EventSection, &UMovieSceneEventSectionBase::OnUserDefinedPinRenamed);
+							EntryPoint->WeakCachedEndpoint = CustomEvent;
+							return CustomEvent;
+						}
+					}
+				}
+			}
+		}
+	}
+	// If the node guid is invalid, this must be a function graph on the BP
+	else for (UEdGraph* Graph : OwnerBlueprint->FunctionGraphs)
+	{
+		if (Graph->GraphGuid == EntryPoint->GraphGuid)
+		{
+			for (UEdGraphNode* Node : Graph->Nodes)
+			{
+				if (UK2Node_FunctionEntry* FunctionEntry = Cast<UK2Node_FunctionEntry>(Node))
+				{
+					FunctionEntry->OnUserDefinedPinRenamed().AddUObject(EventSection, &UMovieSceneEventSectionBase::OnUserDefinedPinRenamed);
+					EntryPoint->WeakCachedEndpoint = FunctionEntry;
+					return FunctionEntry;
+				}
+			}
+		}
+	}
+
+	return nullptr;
+}
+
+UEdGraphPin* FMovieSceneEventUtils::FindBoundObjectPin(UK2Node* InEndpoint, UClass* BoundObjectPinClass)
+{
+	if (!BoundObjectPinClass)
+	{
+		return nullptr;
+	}
+
+	for (UEdGraphPin* Pin : InEndpoint->Pins)
+	{
+		if (Pin->Direction == EGPD_Output && Pin->PinType.PinCategory == UEdGraphSchema_K2::PC_Object && Pin->PinType.PinSubCategoryObject == BoundObjectPinClass)
+		{
+			return Pin;
+		}
+	}
+
+	return nullptr;
+}
+
+void FMovieSceneEventUtils::SetEndpoint(FMovieSceneEvent* EntryPoint, UMovieSceneEventSectionBase* EventSection, UK2Node* InNewEndpoint, UEdGraphPin* BoundObjectPin)
+{
+	check(EntryPoint);
+
+	UK2Node* ExistingEndpoint = CastChecked<UK2Node>(EntryPoint->WeakCachedEndpoint.Get(), ECastCheckedType::NullAllowed);
+	if (ExistingEndpoint)
+	{
+		ExistingEndpoint->OnUserDefinedPinRenamed().RemoveAll(EventSection);
+	}
+
+	if (InNewEndpoint)
+	{
+		const bool bIsFunction    = InNewEndpoint->IsA<UK2Node_FunctionEntry>();
+		const bool bIsCustomEvent = InNewEndpoint->IsA<UK2Node_CustomEvent>();
+
+		checkf(bIsFunction || bIsCustomEvent, TEXT("Only functions and custom events are supported as event endpoints"));
+
+		EntryPoint->GraphGuid = InNewEndpoint->GetGraph()->GraphGuid;
+
+		if (bIsCustomEvent)
+		{
+			EntryPoint->NodeGuid = InNewEndpoint->NodeGuid;
+		}
+		else
+		{
+			EntryPoint->NodeGuid = FGuid();
+		}
+
+		if (BoundObjectPin)
+		{
+			EntryPoint->BoundObjectPinName = BoundObjectPin->GetFName();
+		}
+		else
+		{
+			EntryPoint->BoundObjectPinName = NAME_None;
+		}
+
+		InNewEndpoint->OnUserDefinedPinRenamed().AddUObject(EventSection, &UMovieSceneEventSectionBase::OnUserDefinedPinRenamed);
+		EntryPoint->WeakCachedEndpoint = InNewEndpoint;
+	}
+	else
+	{
+		EntryPoint->NodeGuid = FGuid();
+		EntryPoint->GraphGuid = FGuid();
+		EntryPoint->BoundObjectPinName = NAME_None;
+
+		EntryPoint->WeakCachedEndpoint = nullptr;
+	}
+}
+
+
+UK2Node_FunctionEntry* FMovieSceneEventUtils::GenerateEntryPoint(UMovieSceneEventSectionBase* EventSection, FMovieSceneEvent* EntrypointDefinition, FKismetCompilerContext* Compiler, UEdGraphNode* Endpoint)
+{
+	check(EventSection && Compiler && Endpoint);
+
+	UBlueprint* Blueprint = Compiler->Blueprint;
+
+	UFunction* EndpointFunction = nullptr;
+
+	// Find the function that we need to call on the skeleton class
+	{
+		if (UK2Node_Event* Event = Cast<UK2Node_Event>(Endpoint))
+		{
+			EndpointFunction = Blueprint->SkeletonGeneratedClass ? Blueprint->SkeletonGeneratedClass->FindFunctionByName(Event->GetFunctionName()) : nullptr;
+		}
+		else if (UK2Node_FunctionEntry* EndpointEntry = Cast<UK2Node_FunctionEntry>(Endpoint))
+		{
+			EndpointFunction = EndpointEntry->FindSignatureFunction();
+		}
+		else
+		{
+			Compiler->MessageLog.Error(*LOCTEXT("InvalidEndpoint_Error", "Sequencer event is bound to an invalid endpoint node @@").ToString(), Endpoint);
+		}
+	}
+
+	if (EndpointFunction == nullptr)
+	{
+		return nullptr;
+	}
+
+	// @todo: use a more destriptive name for event entry points?
+	static FString DefaultEventEntryName = "SequenceEvent__ENTRYPOINT";
+	UEdGraph* EntryPointGraph = Compiler->SpawnIntermediateFunctionGraph(DefaultEventEntryName);
+	check(EntryPointGraph->Nodes.Num() == 1);
+
+	const UEdGraphSchema*  Schema        = EntryPointGraph->GetSchema();
+	UK2Node_FunctionEntry* FunctionEntry = CastChecked<UK2Node_FunctionEntry>(EntryPointGraph->Nodes[0]);
+
+	// -------------------------------------------------------------------------------------
+	// Locate and initialize the function entry node
+	{
+		int32 ExtraFunctionFlags = ( FUNC_BlueprintEvent | FUNC_Public );
+		FunctionEntry->AddExtraFlags(ExtraFunctionFlags);
+		FunctionEntry->bIsEditable = false;
+		FunctionEntry->MetaData.Category = LOCTEXT("DefaultCategory", "Sequencer Event Endpoints");
+		FunctionEntry->MetaData.bCallInEditor = EndpointFunction->GetBoolMetaData(FBlueprintMetadata::MD_CallInEditor);
+	}
+
+	// -------------------------------------------------------------------------------------
+	// Create a function call node to invoke the endpoint function itself
+	UK2Node_CallFunction* CallFunctionNode = NewObject<UK2Node_CallFunction>(EntryPointGraph);
+	{
+		CallFunctionNode->CreateNewGuid();
+		CallFunctionNode->PostPlacedNewNode();
+
+		CallFunctionNode->FunctionReference.SetSelfMember(EndpointFunction->GetFName());
+
+		CallFunctionNode->ReconstructNode();
+
+		CallFunctionNode->NodePosX = FunctionEntry->NodePosX + 400.f;
+		CallFunctionNode->NodePosY = FunctionEntry->NodePosY - 16.f;
+
+		EntryPointGraph->AddNode(CallFunctionNode, false, false);
+	}
+
+	// -------------------------------------------------------------------------------------
+	// Create a pin for the bound object if possible
+	if (EntrypointDefinition->BoundObjectPinName != NAME_None)
+	{
+		UEdGraphPin* BoundObjectPin = CallFunctionNode->FindPin(EntrypointDefinition->BoundObjectPinName, EGPD_Input);
+		if (BoundObjectPin)
+		{
+			UEdGraphPin* BoundObjectPinInput = FunctionEntry->CreateUserDefinedPin(BoundObjectPin->PinName, BoundObjectPin->PinType, EGPD_Output, true);
+			if (BoundObjectPinInput)
+			{
+				BoundObjectPinInput->MakeLinkTo(BoundObjectPin);
+			}
+		}
+	}
+
+	// -------------------------------------------------------------------------------------
+	// Wire up the function entry 'then' pin to the call function 'execute' pin
+	{
+		// Wire up the exec nodes
+		UEdGraphPin* ThenPin = FunctionEntry->FindPin(UEdGraphSchema_K2::PN_Then, EGPD_Output);
+		UEdGraphPin* ExecPin = CallFunctionNode->FindPin(UEdGraphSchema_K2::PN_Execute, EGPD_Input);
+
+		if (ensure(ThenPin && ExecPin))
+		{
+			ThenPin->MakeLinkTo(ExecPin);
+		}
+	}
+
+	// -------------------------------------------------------------------------------------
+	// Set pin defaults for each of the relevant pins on the call function node according to the payload
+	{
+		TArray<FName, TInlineAllocator<8>> ValidPinNames;
+
+		// Gather all input pins for the function call
+		for (UEdGraphPin* Pin : CallFunctionNode->Pins)
+		{
+			// Only consider input pins that are data, and not already connected
+			if (Pin->Direction != EGPD_Input || Pin->PinName == UEdGraphSchema_K2::PN_Execute || Pin->PinName == UEdGraphSchema_K2::PN_Self || Pin->LinkedTo.Num() != 0)
+			{
+				continue;
+			}
+
+			const FMovieSceneEventPayloadVariable* PayloadVariable = EntrypointDefinition->PayloadVariables.Find(Pin->PinName);
+			if (PayloadVariable)
+			{
+				bool bMarkAsModified = false;
+				Schema->TrySetDefaultValue(*Pin, PayloadVariable->Value, bMarkAsModified);
+
+				ValidPinNames.Add(Pin->PinName);
+			}
+		}
+
+		// Remove any invalid payload variable names
+		if (ValidPinNames.Num() != EntrypointDefinition->PayloadVariables.Num())
+		{
+			for (auto It = EntrypointDefinition->PayloadVariables.CreateIterator(); It; ++It)
+			{
+				if (!ValidPinNames.Contains(It->Key))
+				{
+					Compiler->MessageLog.Note(*FText::Format(LOCTEXT("PayloadParameterRemoved", "Stale Sequencer event payload parameter %s has been removed."), FText::FromName(It->Key)).ToString());
+					It.RemoveCurrent();
+				}
+			}
+		}
+	}
+
+	return FunctionEntry;
+}
+
+#undef LOCTEXT_NAMESPACE
\ No newline at end of file
diff --git a/Engine/Source/Editor/MovieSceneTools/Private/MovieSceneToolsModule.cpp b/Engine/Source/Editor/MovieSceneTools/Private/MovieSceneToolsModule.cpp
--- a/Engine/Source/Editor/MovieSceneTools/Private/MovieSceneToolsModule.cpp
+++ b/Engine/Source/Editor/MovieSceneTools/Private/MovieSceneToolsModule.cpp
@@ -8,6 +8,12 @@
 #include "ISequencerModule.h"
 #include "ICurveEditorModule.h"
 #include "MovieSceneToolsProjectSettingsCustomization.h"
+#include "Engine/Blueprint.h"
+#include "EdGraph/EdGraph.h"
+#include "K2Node_CustomEvent.h"
+#include "K2Node_FunctionEntry.h"
+#include "KismetCompiler.h"
+#include "Sections/MovieSceneEventSectionBase.h"
 
 #include "TrackEditors/PropertyTrackEditors/BoolPropertyTrackEditor.h"
 #include "TrackEditors/PropertyTrackEditors/BytePropertyTrackEditor.h"
@@ -62,6 +68,8 @@
 #include "Channels/SCurveEditorEventChannelView.h"
 #include "Sections/MovieSceneEventSection.h"
 
+#include "MovieSceneEventUtils.h"
+
 
 #define LOCTEXT_NAMESPACE "FMovieSceneToolsModule"
 
@@ -143,10 +151,17 @@ void FMovieSceneToolsModule::StartupModule()
 			return SNew(SCurveEditorEventChannelView, WeakCurveEditor);
 		}
 	));
+
+	GenerateEventEntryPointsHandle = UMovieSceneEventSectionBase::GenerateEventEntryPointsEvent.AddStatic(HandleGenerateEventEntryPoints);
+	FixupPayloadParameterNameHandle = UMovieSceneEventSectionBase::FixupPayloadParameterNameEvent.AddStatic(FixupPayloadParameterNameForSection);
+	UpgradeLegacyEventEndpointHandle = UMovieSceneEventSectionBase::UpgradeLegacyEventEndpoint.AddStatic(UpgradeLegacyEventEndpointForSection);
 }
 
 void FMovieSceneToolsModule::ShutdownModule()
 {
+	UMovieSceneEventSectionBase::GenerateEventEntryPointsEvent.Remove(GenerateEventEntryPointsHandle);
+	UMovieSceneEventSectionBase::FixupPayloadParameterNameEvent.Remove(FixupPayloadParameterNameHandle);
+
 	if (ICurveEditorModule* CurveEditorModule = FModuleManager::GetModulePtr<ICurveEditorModule>("CurveEditor"))
 	{
 		CurveEditorModule->UnregisterView(FEventChannelCurveModel::EventView);
@@ -210,6 +225,101 @@ void FMovieSceneToolsModule::ShutdownModule()
 	}
 }
 
+void FMovieSceneToolsModule::UpgradeLegacyEventEndpointForSection(UMovieSceneEventSectionBase* Section, UBlueprint* Blueprint)
+{
+	for (FMovieSceneEvent& EntryPoint : Section->GetAllEntryPoints())
+	{
+		UK2Node* Endpoint = FMovieSceneEventUtils::FindEndpoint(&EntryPoint, Section, Blueprint);
+		if (!Endpoint)
+		{
+			continue;
+		}
+
+		// Discover its bound object pin name from the node
+		for (UEdGraphPin* Pin : Endpoint->Pins)
+		{
+			if (Pin->Direction == EGPD_Output && (Pin->PinType.PinCategory == UEdGraphSchema_K2::PC_Object || Pin->PinType.PinCategory == UEdGraphSchema_K2::PC_Interface) )
+			{
+				EntryPoint.BoundObjectPinName = Pin->PinName;
+				break;
+			}
+		}
+
+		// Set the compiled function name so that any immediate PostCompile steps find the correct function name
+		EntryPoint.CompiledFunctionName = Endpoint->GetGraph()->GetFName();
+	}
+}
+
+void FMovieSceneToolsModule::FixupPayloadParameterNameForSection(UMovieSceneEventSectionBase* Section, UK2Node* InNode, FName OldPinName, FName NewPinName)
+{
+	check(Section && InNode);
+
+	UEdGraph* Graph = InNode->GetGraph();
+
+	const bool bCheckNodeGuid = InNode->IsA<UK2Node_CustomEvent>();
+
+	for (FMovieSceneEvent& EntryPoint : Section->GetAllEntryPoints())
+	{
+		if (EntryPoint.GraphGuid != Graph->GraphGuid)
+		{
+			continue;
+		}
+
+		if (bCheckNodeGuid && EntryPoint.NodeGuid != InNode->NodeGuid)
+		{
+			continue;
+		}
+
+		if (EntryPoint.BoundObjectPinName == OldPinName)
+		{
+			EntryPoint.BoundObjectPinName = NewPinName;
+		}
+
+		if (FMovieSceneEventPayloadVariable* Variable = EntryPoint.PayloadVariables.Find(OldPinName))
+		{
+			EntryPoint.PayloadVariables.Add(NewPinName, MoveTemp(*Variable));
+			EntryPoint.PayloadVariables.Remove(OldPinName);
+		}
+	}
+}
+
+void FMovieSceneToolsModule::HandleGenerateEventEntryPoints(UMovieSceneEventSectionBase* EventSection, const FGenerateBlueprintFunctionParams& Params)
+{
+	if (EventSection->HasAnyFlags(RF_NeedLoad))
+	{
+		FLinkerLoad* Linker = EventSection->GetLinker();
+		Linker->Preload(EventSection);
+	}
+	for (FMovieSceneEvent& EntryPoint : EventSection->GetAllEntryPoints())
+	{
+		UEdGraphNode* Endpoint = FMovieSceneEventUtils::FindEndpoint(&EntryPoint, EventSection, Params.CompilerContext->Blueprint);
+		if (Endpoint)
+		{
+			UK2Node_FunctionEntry* FunctionEntry = FMovieSceneEventUtils::GenerateEntryPoint(EventSection, &EntryPoint, Params.CompilerContext, Endpoint);
+			if (FunctionEntry)
+			{
+				EntryPoint.CompiledFunctionName = FunctionEntry->GetGraph()->GetFName();
+			}
+			else
+			{
+				EntryPoint.CompiledFunctionName = NAME_None;
+			}
+		}
+	}
+
+
+	auto OnFunctionListGenerated = [WeakEventSection = MakeWeakObjectPtr(EventSection)](FKismetCompilerContext* CompilerContext)
+	{
+		UMovieSceneEventSectionBase* PinnedSection = WeakEventSection.Get();
+		if (ensureMsgf(PinnedSection, TEXT("Event section has been collected during blueprint compilation.")))
+		{
+			PinnedSection->OnPostCompile(CompilerContext->Blueprint);
+		}
+	};
+
+	Params.CompilerContext->OnFunctionListCompiled().AddLambda(OnFunctionListGenerated);
+}
+
 void FMovieSceneToolsModule::RegisterClipboardConversions()
 {
 	using namespace MovieSceneClipboard;
diff --git a/Engine/Source/Editor/MovieSceneTools/Private/SKeyEditInterface.cpp b/Engine/Source/Editor/MovieSceneTools/Private/SKeyEditInterface.cpp
--- a/Engine/Source/Editor/MovieSceneTools/Private/SKeyEditInterface.cpp
+++ b/Engine/Source/Editor/MovieSceneTools/Private/SKeyEditInterface.cpp
@@ -90,6 +90,12 @@ void SKeyEditInterface::Initialize()
 	];
 }
 
+FVector2D SKeyEditInterface::ComputeDesiredSize(float LayoutScaleMultiplier) const
+{
+	FVector2D ChildSize = SCompoundWidget::ComputeDesiredSize(LayoutScaleMultiplier);
+	return FVector2D::Max(ChildSize, FVector2D(350.f, 0.f));
+}
+
 TSharedRef<IPropertyTypeCustomization> SKeyEditInterface::CreateBindingIDCustomization()
 {
 	TSharedPtr<ISequencer> SequencerPtr = WeakSequencer.Pin();
diff --git a/Engine/Source/Editor/MovieSceneTools/Private/Sections/EventSection.cpp b/Engine/Source/Editor/MovieSceneTools/Private/Sections/EventSection.cpp
--- a/Engine/Source/Editor/MovieSceneTools/Private/Sections/EventSection.cpp
+++ b/Engine/Source/Editor/MovieSceneTools/Private/Sections/EventSection.cpp
@@ -1,6 +1,8 @@
 // Copyright 1998-2019 Epic Games, Inc. All Rights Reserved.
 
 #include "Sections/EventSection.h"
+#include "MovieSceneEventUtils.h"
+#include "K2Node_CustomEvent.h"
 #include "Rendering/DrawElements.h"
 #include "SequencerSectionPainter.h"
 #include "ISectionLayoutBuilder.h"
@@ -144,22 +146,31 @@ int32 FEventTriggerSection::OnPaintSection(FSequencerSectionPainter& Painter) co
 		return LayerId;
 	}
 
+	UMovieSceneSequence* Sequence = EventTriggerSection->GetTypedOuter<UMovieSceneSequence>();
+	FMovieSceneSequenceEditor* SequenceEditor = FMovieSceneSequenceEditor::Find(Sequence);
+	UBlueprint* SequenceDirectorBP = SequenceEditor ? SequenceEditor->FindDirectorBlueprint(Sequence) : nullptr;
+
+	// If we do not have a sequence director BP we can't possibly be bound to anything
+	if (!SequenceDirectorBP)
+	{
+		return LayerId;
+	}
+
 	const FTimeToPixel& TimeToPixelConverter = Painter.GetTimeConverter();
 
-	TMovieSceneChannelData<FMovieSceneEvent> EventData = EventTriggerSection->EventChannel.GetData();
-	TArrayView<const FFrameNumber>     Times  = EventData.GetTimes();
-	TArrayView<const FMovieSceneEvent> Events = EventData.GetValues();
+	TArrayView<const FFrameNumber> Times  = EventTriggerSection->EventChannel.GetData().GetTimes();
+	TArrayView<FMovieSceneEvent>   Events = EventTriggerSection->EventChannel.GetData().GetValues();
 
 	TRange<FFrameNumber> EventSectionRange = EventTriggerSection->GetRange();
 	for (int32 KeyIndex = 0; KeyIndex < Times.Num(); ++KeyIndex)
 	{
 		FFrameNumber EventTime = Times[KeyIndex];
-		UK2Node_FunctionEntry* FunctionEntry = Events[KeyIndex].GetFunctionEntry();
-
 		if (EventSectionRange.Contains(EventTime))
 		{
-			FString EventString = FunctionEntry ? FunctionEntry->GetNodeTitle(ENodeTitleType::MenuTitle).ToString() : FString();
-			bool bIsEventValid = FMovieSceneEvent::IsValidFunction(FunctionEntry);
+			UK2Node* EndpointNode = FMovieSceneEventUtils::FindEndpoint(&Events[KeyIndex], EventTriggerSection, SequenceDirectorBP);
+
+			FString EventString = EndpointNode ? EndpointNode->GetNodeTitle(ENodeTitleType::MenuTitle).ToString() : FString();
+			bool bIsEventValid = true;
 
 			const float PixelPos = TimeToPixelConverter.FrameToPixel(EventTime);
 			PaintEventName(Painter, LayerId, EventString, PixelPos, bIsEventValid);
@@ -172,36 +183,45 @@ int32 FEventTriggerSection::OnPaintSection(FSequencerSectionPainter& Painter) co
 FReply FEventTriggerSection::OnKeyDoubleClicked(FKeyHandle KeyHandle)
 {
 	UMovieSceneEventTriggerSection* EventTriggerSection = Cast<UMovieSceneEventTriggerSection>( WeakSection.Get() );
-	if (EventTriggerSection)
+	if (!EventTriggerSection)
 	{
-		TMovieSceneChannelData<FMovieSceneEvent> ChannelData = EventTriggerSection->EventChannel.GetData();
-		int32 EventIndex = ChannelData.GetIndex(KeyHandle);
-		if (EventIndex != INDEX_NONE)
-		{
-			TArrayView<FMovieSceneEvent> Events = ChannelData.GetValues();
-			FMovieSceneEvent& Event = Events[EventIndex];
+		return FReply::Handled();
+	}
 
-			if (!Event.IsBoundToBlueprint())
-			{
-				FMovieSceneSequenceEditor* SequenceEditor = FMovieSceneSequenceEditor::Find(EventTriggerSection->GetTypedOuter<UMovieSceneSequence>());
-				if (ensure(SequenceEditor))
-				{
-					FScopedTransaction Transaction(LOCTEXT("BindTriggerEvent", "Create Event Endpoint"));
-
-					UK2Node_FunctionEntry* NewEndpoint = SequenceEditor->CreateEventEndpoint(EventTriggerSection->GetTypedOuter<UMovieSceneSequence>());
-					if (NewEndpoint)
-					{
-						SequenceEditor->InitializeEndpointForTrack(EventTriggerSection->GetTypedOuter<UMovieSceneEventTrack>(), NewEndpoint);
-						FMovieSceneSequenceEditor::BindEventToEndpoint(EventTriggerSection, &Event, NewEndpoint);
-					}
-				}
-			}
+	UMovieSceneSequence* Sequence = EventTriggerSection->GetTypedOuter<UMovieSceneSequence>();
+	check(Sequence);
 
-			if (UK2Node_FunctionEntry* FunctionEntry = Event.GetFunctionEntry())
-			{
-				FKismetEditorUtilities::BringKismetToFocusAttentionOnObject(FunctionEntry, false);
-			}
-		}
+	FMovieSceneSequenceEditor* SequenceEditor = FMovieSceneSequenceEditor::Find(Sequence);
+	if (!SequenceEditor)
+	{
+		return FReply::Handled();
+	}
+
+	UBlueprint* SequenceDirectorBP = SequenceEditor->GetOrCreateDirectorBlueprint(Sequence);
+	if (!SequenceDirectorBP)
+	{
+		return FReply::Handled();
+	}
+
+	TMovieSceneChannelData<FMovieSceneEvent> ChannelData = EventTriggerSection->EventChannel.GetData();
+	const int32 EventIndex = ChannelData.GetIndex(KeyHandle);
+	if (EventIndex == INDEX_NONE)
+	{
+		return FReply::Handled();
+	}
+
+	FMovieSceneEvent* EventEntryPoint = &ChannelData.GetValues()[EventIndex];
+	UK2Node* Endpoint = FMovieSceneEventUtils::FindEndpoint(EventEntryPoint, EventTriggerSection, SequenceDirectorBP);
+
+	if (!Endpoint)
+	{
+		FScopedTransaction Transaction(LOCTEXT("CreateEventEndpoint", "Create Event Endpoint"));
+		Endpoint = FMovieSceneEventUtils::BindNewUserFacingEvent(EventEntryPoint, EventTriggerSection, SequenceDirectorBP);
+	}
+
+	if (Endpoint)
+	{
+		FKismetEditorUtilities::BringKismetToFocusAttentionOnObject(Endpoint, false);
 	}
 
 	return FReply::Handled();
@@ -212,45 +232,72 @@ int32 FEventRepeaterSection::OnPaintSection(FSequencerSectionPainter& Painter) c
 	int32 LayerId = Painter.PaintSectionBackground();
 
 	UMovieSceneEventRepeaterSection* EventRepeaterSection = Cast<UMovieSceneEventRepeaterSection>(WeakSection.Get());
-	if (EventRepeaterSection)
+	if (!EventRepeaterSection)
 	{
-		UK2Node_FunctionEntry* FunctionEntry = EventRepeaterSection->Event.GetFunctionEntry();
+		return LayerId;
+	}
 
-		float TextOffsetX = EventRepeaterSection->GetRange().GetLowerBound().IsClosed() ? FMath::Max(0.f, Painter.GetTimeConverter().FrameToPixel(EventRepeaterSection->GetRange().GetLowerBoundValue())) : 0.f;
+	UMovieSceneSequence* Sequence = EventRepeaterSection->GetTypedOuter<UMovieSceneSequence>();
+	FMovieSceneSequenceEditor* SequenceEditor = FMovieSceneSequenceEditor::Find(Sequence);
+	UBlueprint* SequenceDirectorBP = SequenceEditor ? SequenceEditor->FindDirectorBlueprint(Sequence) : nullptr;
 
-		FString EventString = FunctionEntry ? FunctionEntry->GetNodeTitle(ENodeTitleType::MenuTitle).ToString() : FString();
-		bool bIsEventValid = FMovieSceneEvent::IsValidFunction(FunctionEntry);
-		PaintEventName(Painter, LayerId, EventString, TextOffsetX, bIsEventValid);
+	// If we do not have a sequence director BP we can't possibly be bound to anything
+	if (!SequenceDirectorBP)
+	{
+		return LayerId;
 	}
 
+	UK2Node* EndpointNode = FMovieSceneEventUtils::FindEndpoint(&EventRepeaterSection->Event, EventRepeaterSection, SequenceDirectorBP);
+
+	float TextOffsetX = EventRepeaterSection->GetRange().GetLowerBound().IsClosed() ? FMath::Max(0.f, Painter.GetTimeConverter().FrameToPixel(EventRepeaterSection->GetRange().GetLowerBoundValue())) : 0.f;
+
+	FString EventString = EndpointNode ? EndpointNode->GetNodeTitle(ENodeTitleType::MenuTitle).ToString() : FString();
+	bool bIsEventValid = true;
+	PaintEventName(Painter, LayerId, EventString, TextOffsetX, bIsEventValid);
+
 	return LayerId + 1;
 }
 
 FReply FEventRepeaterSection::OnSectionDoubleClicked(const FGeometry& SectionGeometry, const FPointerEvent& MouseEvent)
 {
-	UMovieSceneEventRepeaterSection* EventRepeaterSection = Cast<UMovieSceneEventRepeaterSection>( WeakSection.Get() );
-	if (EventRepeaterSection)
+	TSharedPtr<ISequencer> SequencerPtr = Sequencer.Pin();
+	if (!SequencerPtr)
 	{
-		if (!EventRepeaterSection->Event.IsBoundToBlueprint())
-		{
-			FMovieSceneSequenceEditor* SequenceEditor = FMovieSceneSequenceEditor::Find(EventRepeaterSection->GetTypedOuter<UMovieSceneSequence>());
-			if (ensure(SequenceEditor))
-			{
-				FScopedTransaction Transaction(LOCTEXT("BindRepeaterEvent", "Create Event Endpoint"));
-
-				UK2Node_FunctionEntry* NewEndpoint = SequenceEditor->CreateEventEndpoint(EventRepeaterSection->GetTypedOuter<UMovieSceneSequence>());
-				if (NewEndpoint)
-				{
-					SequenceEditor->InitializeEndpointForTrack(EventRepeaterSection->GetTypedOuter<UMovieSceneEventTrack>(), NewEndpoint);
-					FMovieSceneSequenceEditor::BindEventToEndpoint(EventRepeaterSection, &EventRepeaterSection->Event, NewEndpoint);
-				}
-			}
-		}
+		return FReply::Handled();
+	}
 
-		if (UK2Node* FunctionEntry = EventRepeaterSection->Event.GetFunctionEntry())
-		{
-			FKismetEditorUtilities::BringKismetToFocusAttentionOnObject(FunctionEntry, false);
-		}
+	UMovieSceneEventRepeaterSection* EventRepeaterSection = Cast<UMovieSceneEventRepeaterSection>(WeakSection.Get());
+	if (!EventRepeaterSection)
+	{
+		return FReply::Handled();
+	}
+
+	UMovieSceneSequence* Sequence = EventRepeaterSection->GetTypedOuter<UMovieSceneSequence>();
+	check(Sequence);
+
+	FMovieSceneSequenceEditor* SequenceEditor = FMovieSceneSequenceEditor::Find(Sequence);
+	if (!SequenceEditor)
+	{
+		return FReply::Handled();
+	}
+
+	UBlueprint* SequenceDirectorBP = SequenceEditor->GetOrCreateDirectorBlueprint(Sequence);
+	if (!SequenceDirectorBP)
+	{
+		return FReply::Handled();
+	}
+
+	UK2Node* Endpoint = FMovieSceneEventUtils::FindEndpoint(&EventRepeaterSection->Event, EventRepeaterSection, SequenceDirectorBP);
+
+	if (!Endpoint)
+	{
+		FScopedTransaction Transaction(LOCTEXT("BindRepeaterEvent", "Create Event Endpoint"));
+		Endpoint = FMovieSceneEventUtils::BindNewUserFacingEvent(&EventRepeaterSection->Event, EventRepeaterSection, SequenceDirectorBP);
+	}
+
+	if (Endpoint)
+	{
+		FKismetEditorUtilities::BringKismetToFocusAttentionOnObject(Endpoint, false);
 	}
 
 	return FReply::Handled();
diff --git a/Engine/Source/Editor/MovieSceneTools/Private/TrackEditors/EventTrackEditor.cpp b/Engine/Source/Editor/MovieSceneTools/Private/TrackEditors/EventTrackEditor.cpp
--- a/Engine/Source/Editor/MovieSceneTools/Private/TrackEditors/EventTrackEditor.cpp
+++ b/Engine/Source/Editor/MovieSceneTools/Private/TrackEditors/EventTrackEditor.cpp
@@ -173,21 +173,38 @@ void FEventTrackEditor::BuildTrackContextMenu(FMenuBuilder& MenuBuilder, UMovieS
 	UMovieSceneEventTrack* EventTrack = CastChecked<UMovieSceneEventTrack>(Track);
 	UProperty* EventPositionProperty = FindField<UProperty>(Track->GetClass(), GET_MEMBER_NAME_STRING_CHECKED(UMovieSceneEventTrack, EventPosition));
 
+	FGuid ObjectBinding;
+	EventTrack->GetTypedOuter<UMovieScene>()->FindTrackBinding(*EventTrack, ObjectBinding);
+
 	/** Specific details customization for the event track */
 	class FEventTrackCustomization : public IDetailCustomization
 	{
 	public:
+		FGuid ObjectBindingID;
+
+		FEventTrackCustomization(const FGuid& InObjectBindingID)
+			: ObjectBindingID(InObjectBindingID)
+		{}
+
 		virtual void CustomizeDetails(IDetailLayoutBuilder& DetailBuilder) override
 		{
 			DetailBuilder.HideCategory("Track");
 			DetailBuilder.HideCategory("General");
 
-			IDetailCategoryBuilder& Category = DetailBuilder.EditCategory("TrackEvent");
-			Category.AddProperty("EventReceivers").ShouldAutoExpand(true);
+			if (ObjectBindingID.IsValid())
+			{
+				IDetailCategoryBuilder& Category = DetailBuilder.EditCategory("TrackEvent");
+				Category.AddProperty("EventReceivers").ShouldAutoExpand(true);
+			}
+			else
+			{
+				// Do not show event receivers for tracks that exist on object bindings
+				DetailBuilder.HideCategory("TrackEvent");
+			}
 		}
 	};
 
-	auto PopulateSubMenu = [this, EventTrack](FMenuBuilder& SubMenuBuilder)
+	auto PopulateSubMenu = [this, EventTrack, ObjectBinding](FMenuBuilder& SubMenuBuilder)
 	{
 		FPropertyEditorModule& PropertyEditor = FModuleManager::Get().LoadModuleChecked<FPropertyEditorModule>("PropertyEditor");
 
@@ -200,7 +217,7 @@ void FEventTrackEditor::BuildTrackContextMenu(FMenuBuilder& MenuBuilder, UMovieS
 		TSharedRef<IDetailsView> DetailsView = PropertyEditor.CreateDetailView(DetailsViewArgs);
 		
 		// Register the custom type layout for the class
-		FOnGetDetailCustomizationInstance CreateInstance = FOnGetDetailCustomizationInstance::CreateLambda(&MakeShared<FEventTrackCustomization>);
+		FOnGetDetailCustomizationInstance CreateInstance = FOnGetDetailCustomizationInstance::CreateLambda([ObjectBinding]{ return MakeShared<FEventTrackCustomization>(ObjectBinding); });
 		DetailsView->RegisterInstancedCustomPropertyLayout(UMovieSceneEventTrack::StaticClass(), CreateInstance);
 
 		GetSequencer()->OnInitializeDetailsPanel().Broadcast(DetailsView, GetSequencer().ToSharedRef());
diff --git a/Engine/Source/Editor/MovieSceneTools/Public/MovieSceneEventCustomization.h b/Engine/Source/Editor/MovieSceneTools/Public/MovieSceneEventCustomization.h
--- a/Engine/Source/Editor/MovieSceneTools/Public/MovieSceneEventCustomization.h
+++ b/Engine/Source/Editor/MovieSceneTools/Public/MovieSceneEventCustomization.h
@@ -4,11 +4,18 @@
 
 #include "CoreMinimal.h"
 #include "IPropertyTypeCustomization.h"
-#include "IStructureDetailsView.h"
+#include "Types/SlateEnums.h"
 
 struct FAssetData;
+struct FEdGraphSchemaAction;
+struct FMovieSceneEventEndpointParameters;
 
+class UK2Node;
+class UBlueprint;
+class UEdGraphPin;
+class UEdGraphNode;
 class FMenuBuilder;
+class FStructOnScope;
 class IPropertyHandle;
 class FDetailWidgetRow;
 class UMovieSceneSection;
@@ -17,6 +24,18 @@ class UMovieSceneEventTrack;
 class UK2Node_FunctionEntry;
 class IDetailChildrenBuilder;
 
+enum class ECheckBoxState : uint8;
+
+template <typename KeyType, typename ValueType, typename ArrayAllocator, typename SortPredicate> class TSortedMap;
+
+enum class EAutoCreatePayload : uint8
+{
+	None      = 0,
+	Pins      = 1 << 0,
+	Variables = 1 << 1,
+};
+ENUM_CLASS_FLAGS(EAutoCreatePayload)
+
 /**
  * Customization for FMovieSceneEvent structs.
  * Will deduce the event's section either from the outer objects on the details customization, or use the one provided on construction (for instanced property type customizations)
@@ -45,21 +64,10 @@ private:
 	void CreateEventEndpoint();
 
 
-	/**
-	 * Creates a single new endpoint for the event(s) represented by this property handle using the specified function as a quick binding.
-	 */
-	void CreateEventEndpointFromFunction(UFunction* QuickBindFunction, UClass* PinClassType);
-
-
 	/**
 	 * Assigns the specified function entry to the event(s) represented by this property handle.
 	 */
-	void SetEventEndpoint(UK2Node_FunctionEntry* NewEndpoint);
-
-	/**
-	 * Compare the currently assigned endpoint with the specified entry. Used as the check-state for the radio buttons on the menu.
-	 */
-	bool CompareCurrentEventEndpoint(UK2Node_FunctionEntry* NewEndpoint);
+	void SetEventEndpoint(UK2Node* NewEndpoint, UEdGraphPin* BoundObjectPin, UK2Node* PayloadTemplate, EAutoCreatePayload AutoCreatePayload);
 
 
 	/**
@@ -75,9 +83,19 @@ private:
 
 
 	/**
-	 * Generate the content of the creation shortcut sub-menu dropdown
+	 * Generate the content of the quick bind sub-menu dropdown (shown if the event is not already bound)
 	 */
-	void PopulateQuickBindSubMenu(FMenuBuilder& MenuBuilder, UClass* TemplateClass);
+	void PopulateQuickBindSubMenu(FMenuBuilder& MenuBuilder, UMovieSceneSequence* Sequence);
+
+	/**
+	 * Generate the content of the rebind sub-menu dropdown (shown if the event is already bound)
+	 */
+	void PopulateRebindSubMenu(FMenuBuilder& MenuBuilder, UMovieSceneSequence* Sequence);
+
+	/**
+	 * Called when a payload variable property value is changed on the details panel
+	 */
+	void OnPayloadVariableChanged(TSharedRef<FStructOnScope> InStructData, TSharedPtr<IPropertyHandle> LocalVariableProperty);
 
 private:
 
@@ -96,8 +114,11 @@ private:
 	/**
 	 * Get the endpoint that is common to all the events represented by this property handle, or nullptr if they are not all the same.
 	 */
-	UK2Node_FunctionEntry* GetCommonEndpoint() const;
+	UK2Node* GetCommonEndpoint() const;
 
+	TArray<UK2Node*> GetAllValidEndpoints() const;
+
+	void IterateEndpoints(TFunctionRef<bool(UK2Node*)> Callback) const;
 
 	/**
 	 * Get all the objects that the events reside within.
@@ -111,12 +132,19 @@ private:
 	/** Get the icon of the event to display on the main combo button */
 	const FSlateBrush* GetEventIcon() const;
 
-	/** Get the visibility of the error icon */
-	EVisibility GetErrorVisibility() const;
-	/** Get the tooltip text for the error icon */
-	FText GetErrorTooltip() const;
+	void HandleRebindActionSelected(const TArray<TSharedPtr<FEdGraphSchemaAction>>& SelectedAction, ESelectInfo::Type InSelectionType, UBlueprint* Blueprint, UClass* BoundObjectPinClass);
+	void HandleQuickBindActionSelected(const TArray<TSharedPtr<FEdGraphSchemaAction>>& SelectedAction, ESelectInfo::Type InSelectionType, UBlueprint* Blueprint, FMovieSceneEventEndpointParameters Params);
+
+	void OnBlueprintCompiled(UBlueprint*);
+
+	ECheckBoxState GetCallInEditorCheckState() const;
+	void OnSetCallInEditorCheckState(ECheckBoxState NewCheckBoxState);
 
-	UClass* FindObjectBindingClass();
+	TSharedRef<SWidget> GetBoundObjectPinMenuContent();
+	bool CompareBoundObjectPinName(FName InPinName) const;
+	const FSlateBrush* GetBoundObjectPinIcon() const;
+	FText GetBoundObjectPinText() const;
+	void SetBoundObjectPinName(FName InNewBoundObjectPinName);
 
 private:
 
@@ -128,4 +156,9 @@ private:
 
 	/** The property handle we're reflecting */
 	TSharedPtr<IPropertyHandle> PropertyHandle;
+
+	TSharedPtr<IPropertyUtilities> PropertyUtilities;
+
+	template<typename KeyType, typename ValueType>
+	using TInlineMap = TSortedMap<KeyType, ValueType, TInlineAllocator<1>>;
 };
diff --git a/Engine/Source/Editor/MovieSceneTools/Public/MovieSceneEventUtils.h b/Engine/Source/Editor/MovieSceneTools/Public/MovieSceneEventUtils.h
new file mode 100644
--- /dev/null
+++ b/Engine/Source/Editor/MovieSceneTools/Public/MovieSceneEventUtils.h
@@ -0,0 +1,111 @@
+// Copyright 1998-2019 Epic Games, Inc. All Rights Reserved.
+
+#pragma once
+
+#include "CoreTypes.h"
+#include "Containers/UnrealString.h"
+
+struct FGuid;
+struct FMovieSceneEvent;
+
+class UClass;
+class UK2Node;
+class UBlueprint;
+class UEdGraphPin;
+class UEdGraphNode;
+class UK2Node_CustomEvent;
+class UK2Node_CallFunction;
+class UK2Node_FunctionEntry;
+class FKismetCompilerContext;
+
+class UMovieScene;
+class UMovieSceneTrack;
+class UMovieSceneSequence;
+class UMovieSceneEventTrack;
+class UMovieSceneEventSectionBase;
+
+
+
+struct MOVIESCENETOOLS_API FMovieSceneEventEndpointParameters
+{
+	FString SanitizedObjectName;
+	FString SanitizedEventName;
+	UClass* BoundObjectPinClass = nullptr;
+
+	static FMovieSceneEventEndpointParameters Generate(UMovieSceneEventTrack* Track);
+	static FMovieSceneEventEndpointParameters Generate(UMovieScene* MovieScene, const FGuid& ObjectBindingID);
+};
+
+
+/**
+ * Static utility library for dealing with movie-scene events at edit/cook-time
+ */
+struct MOVIESCENETOOLS_API FMovieSceneEventUtils
+{
+	/**
+	 * Create a new event endpoint for the specified sequence. By default, event end points are presented as a UK2Node_CustomEvent in an event graph called "Sequencer Events" with a pin for the bound object.
+	 * @note: Does not bind any event sections to UBlueprint::GenerateFunctionGraphsEvent which is necessary for the event's entry point to be generated during blueprint compilation - that must be done before (or after) this call.
+	 *
+	 * @param Blueprint       (Required, non-null) The blueprint within which to create the new custom event
+	 * @param Parameters      The end point parameters that define how to name and construct the event
+	 * @return A valid pointer to the blueprint node for the event endpoint
+	 */
+	static UK2Node_CustomEvent* CreateUserFacingEvent(UBlueprint* Blueprint, const FMovieSceneEventEndpointParameters& Parameters);
+
+
+	/**
+	 * As with CreateUserFacingEvent, but also binds the section to UBlueprint::GenerateFunctionGraphsEvent which ensures that it will hook into the compilation process to generate its entry point function graph
+	 *
+	 * @param EntryPoint      (Required, non-null) The event entry point definition
+	 * @param EventSection    (Required, non-null) The event section that owns the specified entry point
+	 * @param Blueprint       (Required, non-null) The blueprint within which to create the new custom event
+	 * @return A valid pointer to the blueprint node for the event endpoint
+	 */
+	static UK2Node_CustomEvent* BindNewUserFacingEvent(FMovieSceneEvent* EntryPoint, UMovieSceneEventSectionBase* EventSection, UBlueprint* Blueprint);
+
+
+	/**
+	 * Main generation function for creating the necessary function graph for calling an event endpoint with the parameters specified in the entry point payload
+	 *
+	 * @param EventSection           (Required, non-null) The event section that owns the event
+	 * @param EntryPoint             (Required, non-null) The event entry point definition to generate a function graph for
+	 * @param Compiler               (Required, non-null) The kismet compiler context for generation
+	 * @param Endpoint               (Required, non-null) The endpoint node that the event is bound to
+	 */
+	static UK2Node_FunctionEntry* GenerateEntryPoint(UMovieSceneEventSectionBase* EventSection, FMovieSceneEvent* EntryPoint, FKismetCompilerContext* Compiler, UEdGraphNode* Endpoint);
+
+
+public:
+
+	/**
+	 * Attempt to locate the blueprint node that relates to an event's end-point.
+	 * 
+	 * @param EntryPoint             (Required, non-null) The event entry point definition to look up
+	 * @param EventSection           (Required, non-null) The event section that owns the event
+	 * @param OwnerBlueprint         (Required, non-null) The blueprint to search within
+	 * @return The K2Node that relates to the specified event's end point, or nullptr if one was not found
+	 */
+	static UK2Node* FindEndpoint(FMovieSceneEvent* EntryPoint, UMovieSceneEventSectionBase* EventSection, UBlueprint* OwnerBlueprint);
+
+
+	/**
+	 * Attempt to locate an output pin on the specified endpoint node that matches the specified pin class
+	 * 
+	 * @param InEndpoint             (Required, non-null) The endpoint node to search on
+	 * @param BoundObjectPinClass    (Required, non-null) The class of the object pin to locate
+	 * @return An output pin that is of type PC_Object with the specified class as its PinSubCategoryObject, or nullptr if one was not found
+	 */
+	static UEdGraphPin* FindBoundObjectPin(UK2Node* InEndpoint, UClass* BoundObjectPinClass);
+
+
+	/**
+	 * Set the specified EntryPoint to be bound to a new K2 node
+	 * 
+	 * @param EntryPoint             (Required, non-null) The event entry point definition to assign to
+	 * @param EventSection           (Required, non-null) The event section that owns the event entry point definition
+	 * @param InNewEndpoint          (Optional) The new endpoint node to call from the event
+	 * @param BoundObjectPin         (Optional) The pin that should receive the bound object when invoking the event
+	 */
+	static void SetEndpoint(FMovieSceneEvent* EntryPoint, UMovieSceneEventSectionBase* EventSection, UK2Node* InNewEndpoint, UEdGraphPin* BoundObjectPin);
+};
+
diff --git a/Engine/Source/Editor/MovieSceneTools/Public/MovieSceneToolsModule.h b/Engine/Source/Editor/MovieSceneTools/Public/MovieSceneToolsModule.h
--- a/Engine/Source/Editor/MovieSceneTools/Public/MovieSceneToolsModule.h
+++ b/Engine/Source/Editor/MovieSceneTools/Public/MovieSceneToolsModule.h
@@ -7,9 +7,15 @@
 #include "Modules/ModuleManager.h"
 #include "IMovieSceneTools.h"
 
-class UMovieSceneSection;
 struct FAssetData;
 
+class UK2Node;
+class UBlueprint;
+class UMovieSceneSection;
+class UMovieSceneEventSectionBase;
+
+struct FGenerateBlueprintFunctionParams;
+
 class IMovieSceneToolsTakeData
 {
 public:
@@ -46,6 +52,10 @@ private:
 
 	void RegisterClipboardConversions();
 
+	static void HandleGenerateEventEntryPoints(UMovieSceneEventSectionBase* EventSection, const FGenerateBlueprintFunctionParams& Params);
+	static void FixupPayloadParameterNameForSection(UMovieSceneEventSectionBase* Section, UK2Node* InNode, FName OldPinName, FName NewPinName);
+	static void UpgradeLegacyEventEndpointForSection(UMovieSceneEventSectionBase* Section, UBlueprint* Blueprint);
+
 private:
 
 	/** Registered delegate handles */
@@ -83,5 +93,9 @@ private:
 	FDelegateHandle MPCTrackCreateEditorHandle;
 	FDelegateHandle PrimitiveMaterialCreateEditorHandle;
 
+	FDelegateHandle GenerateEventEntryPointsHandle;
+	FDelegateHandle FixupPayloadParameterNameHandle;
+	FDelegateHandle UpgradeLegacyEventEndpointHandle;
+
 	TArray<IMovieSceneToolsTakeData*> TakeDatas;
 };
\ No newline at end of file
diff --git a/Engine/Source/Editor/MovieSceneTools/Public/SKeyEditInterface.h b/Engine/Source/Editor/MovieSceneTools/Public/SKeyEditInterface.h
--- a/Engine/Source/Editor/MovieSceneTools/Public/SKeyEditInterface.h
+++ b/Engine/Source/Editor/MovieSceneTools/Public/SKeyEditInterface.h
@@ -56,6 +56,8 @@ private:
 
 private:
 
+	virtual FVector2D ComputeDesiredSize(float LayoutScaleMultiplier) const override;
+
 	virtual void NotifyPreChange(UProperty* PropertyAboutToChange) override;
 	virtual void PostUndo(bool bSuccess) override;
 	virtual void PostRedo(bool bSuccess) override;
diff --git a/Engine/Source/Editor/Sequencer/Private/IKeyArea.cpp b/Engine/Source/Editor/Sequencer/Private/IKeyArea.cpp
--- a/Engine/Source/Editor/Sequencer/Private/IKeyArea.cpp
+++ b/Engine/Source/Editor/Sequencer/Private/IKeyArea.cpp
@@ -139,10 +139,11 @@ void IKeyArea::DrawKeys(TArrayView<const FKeyHandle> InKeyHandles, TArrayView<FK
 
 	ISequencerChannelInterface* EditorInterface = FindChannelEditorInterface();
 	FMovieSceneChannel* Channel = ChannelHandle.Get();
+	UMovieSceneSection* OwningSection = GetOwningSection();
 
-	if (EditorInterface && Channel)
+	if (EditorInterface && Channel && OwningSection)
 	{
-		return EditorInterface->DrawKeys_Raw(Channel, InKeyHandles, OutKeyDrawParams);
+		return EditorInterface->DrawKeys_Raw(Channel, InKeyHandles, OwningSection, OutKeyDrawParams);
 	}
 }
 
diff --git a/Engine/Source/Editor/Sequencer/Private/MovieSceneSequenceEditor.cpp b/Engine/Source/Editor/Sequencer/Private/MovieSceneSequenceEditor.cpp
--- a/Engine/Source/Editor/Sequencer/Private/MovieSceneSequenceEditor.cpp
+++ b/Engine/Source/Editor/Sequencer/Private/MovieSceneSequenceEditor.cpp
@@ -2,45 +2,9 @@
 
 #include "MovieSceneSequenceEditor.h"
 #include "Modules/ModuleManager.h"
-#include "MovieSceneSection.h"
 #include "MovieSceneSequence.h"
-#include "Channels/MovieSceneEvent.h"
-#include "Tracks/MovieSceneEventTrack.h"
-#include "Sections/MovieSceneEventSectionBase.h"
 #include "ISequencerModule.h"
-#include "ScopedTransaction.h"
-
-#include "Engine/Blueprint.h"
-#include "EdGraph/EdGraph.h"
-#include "K2Node_FunctionEntry.h"
-#include "Kismet2/KismetEditorUtilities.h"
-#include "Kismet2/BlueprintEditorUtils.h"
-
-#define LOCTEXT_NAMESPACE "MovieSceneSequenceEditor"
-
-namespace
-{
-	FText GetDefaultCommentText()
-	{
-		FText DefaultComment_01 = LOCTEXT("DefaultComment_01", "Sequencer Events can be bound to either of the following function signatures. Return values are not supported.");
-		FText DefaultComment_02 = LOCTEXT("DefaultComment_02", "\t1. A function with no parameters");
-		FText DefaultComment_03 = LOCTEXT("DefaultComment_03", "\t\tCompatible with master tracks or object bindings");
-		FText DefaultComment_04 = LOCTEXT("DefaultComment_04", "\t2. A function with a single object or interface parameter");
-		FText DefaultComment_05 = LOCTEXT("DefaultComment_05", "\t\tCompatible with master tracks or object bindings");
-		FText DefaultComment_06 = LOCTEXT("DefaultComment_06", "\t\tWill only trigger if the source object is of the same type as the parameter (or interface)");
-		FText DefaultComment_07 = LOCTEXT("DefaultComment_07", "\t\tWill be triggered with objects in the following order:");
-		FText DefaultComment_08 = LOCTEXT("DefaultComment_08", "\t\t\ta) Objects bound to the track's object binding, or:");
-		FText DefaultComment_09 = LOCTEXT("DefaultComment_09", "\t\t\tb) Objects specified on the track's event receivers array, or:");
-		FText DefaultComment_10 = LOCTEXT("DefaultComment_10", "\t\t\tc) Objects provided by the playback context (level blueprints, widgets etc)");
-		FText DefaultComment_11 = LOCTEXT("DefaultComment_11", "Tip: Trigger events on level blueprints using an interface pin, and adding the interface to the desired level blueprint");
-
-		return FText::Format(LOCTEXT("DefaultComment_Format", "{0}\n{1}\n{2}\n{3}\n{4}\n{5}\n{6}\n{7}\n{8}\n{9}\n\n{10}"),
-			DefaultComment_01, DefaultComment_02, DefaultComment_03, DefaultComment_04, DefaultComment_05,
-			DefaultComment_06, DefaultComment_07, DefaultComment_08, DefaultComment_09, DefaultComment_10, DefaultComment_11);
-	}
-}
-
-FName FMovieSceneSequenceEditor::TargetPinName("Target");
+#include "BlueprintActionDatabase.h"
 
 FMovieSceneSequenceEditor* FMovieSceneSequenceEditor::Find(UMovieSceneSequence* InSequence)
 {
@@ -53,108 +17,33 @@ bool FMovieSceneSequenceEditor::SupportsEvents(UMovieSceneSequence* InSequence)
 	return CanCreateEvents(InSequence);
 }
 
-UBlueprint* FMovieSceneSequenceEditor::GetDirectorBlueprint(UMovieSceneSequence* Sequence) const
+UBlueprint* FMovieSceneSequenceEditor::FindDirectorBlueprint(UMovieSceneSequence* Sequence) const
 {
 	return GetBlueprintForSequence(Sequence);
 }
 
-UBlueprint* FMovieSceneSequenceEditor::AccessDirectorBlueprint(UMovieSceneSequence* Sequence) const
+UBlueprint* FMovieSceneSequenceEditor::GetOrCreateDirectorBlueprint(UMovieSceneSequence* Sequence) const
 {
 	UBlueprint* Blueprint = GetBlueprintForSequence(Sequence);
 	if (!Blueprint)
 	{
 		Blueprint = CreateBlueprintForSequence(Sequence);
-	}
-	return Blueprint;
-}
-
-UK2Node_FunctionEntry* FMovieSceneSequenceEditor::CreateEventEndpoint(UMovieSceneSequence* Sequence, const FString& DesiredName) const
-{
-	FScopedTransaction Transaction(LOCTEXT("CreateEventEndpoint", "Create Event Endpoint"));
-
-	UBlueprint* Blueprint = AccessDirectorBlueprint(Sequence);
-	if (Blueprint)
-	{
-		static FString DefaultEventName = "SequenceEvent";
-		FName UniqueGraphName = FBlueprintEditorUtils::FindUniqueKismetName(Blueprint, DesiredName.Len() == 0 ? DefaultEventName : DesiredName);
-
-		Blueprint->Modify();
-
-		UEdGraph* Graph = FBlueprintEditorUtils::CreateNewGraph(Blueprint, UniqueGraphName, UEdGraph::StaticClass(), UEdGraphSchema_K2::StaticClass());
-
-		const bool bIsUserCreated = false;
-		FBlueprintEditorUtils::AddFunctionGraph<UClass>(Blueprint, Graph, bIsUserCreated, nullptr);
-
-		TArray<UK2Node_FunctionEntry*> EntryNodes;
-		Graph->GetNodesOfClass<UK2Node_FunctionEntry>(EntryNodes);
-
-		if (ensure(EntryNodes.Num() == 1 && EntryNodes[0]))
-		{
-			int32 ExtraFunctionFlags = ( FUNC_BlueprintCallable | FUNC_BlueprintEvent | FUNC_Public );
-			EntryNodes[0]->AddExtraFlags(ExtraFunctionFlags);
-
-			EntryNodes[0]->bIsEditable = true;
-			EntryNodes[0]->MetaData.Category = LOCTEXT("DefaultCategory", "Sequencer Event Endpoints");
-
-			EntryNodes[0]->NodeComment = GetDefaultCommentText().ToString();
-			EntryNodes[0]->bCommentBubblePinned = true;
-			EntryNodes[0]->bCommentBubbleVisible = true;
-
-			return EntryNodes[0];
-		}
-	}
-
-	return nullptr;
-}
 
-void FMovieSceneSequenceEditor::InitializeEndpointForTrack(UMovieSceneEventTrack* EventTrack, UK2Node_FunctionEntry* Endpoint) const
-{
-	UMovieScene* MovieScene = EventTrack->GetTypedOuter<UMovieScene>();
-
-	// Only add an object pin for the event if this is an event track on an object binding
-	FGuid Unused;
-	if (MovieScene->FindTrackBinding(*EventTrack, Unused))
-	{
-		SetupDefaultPinForEndpoint(EventTrack, Endpoint);
-	}
-}
-
-UClass* FMovieSceneSequenceEditor::FindTrackObjectBindingClass(UMovieSceneTrack* Track)
-{
-	UMovieScene* MovieScene = Track->GetTypedOuter<UMovieScene>();
-
-	// We have to find the object binding that owns this track, which is potentially slow, but there is no back reference here.
-	for (const FMovieSceneBinding& Binding : MovieScene->GetBindings())
-	{
-		if (!Algo::Find(Binding.GetTracks(), Track))
+		// This asset now has a blueprint where before it did not: refresh asset actions for this asset
+		UObject* AssetObject = Sequence;
+		while (AssetObject)
 		{
-			continue;
+			if (AssetObject->IsAsset())
+			{
+				break;
+			}
+			AssetObject = AssetObject->GetOuter();
 		}
 
-		if (FMovieScenePossessable* Possessable = MovieScene->FindPossessable(Binding.GetObjectGuid()))
+		if (AssetObject)
 		{
-			return const_cast<UClass*>(Possessable->GetPossessedObjectClass());
+			FBlueprintActionDatabase::Get().RefreshAssetActions(AssetObject);
 		}
-		else if(FMovieSceneSpawnable* Spawnable = MovieScene->FindSpawnable(Binding.GetObjectGuid()))
-		{
-			return Spawnable->GetObjectTemplate()->GetClass();
-		}
-	}
-
-	return nullptr;
-}
-
-void FMovieSceneSequenceEditor::BindEventToEndpoint(UMovieSceneEventSectionBase* EventSection, FMovieSceneEvent* Event, UK2Node_FunctionEntry* Endpoint)
-{
-	check(EventSection && Event);
-
-	EventSection->Modify();
-	Event->SetFunctionEntry(Endpoint);
-
-	if (UBlueprint* Blueprint = Endpoint ? Endpoint->GetBlueprint() : nullptr)
-	{
-		EventSection->SetDirectorBlueprint(Blueprint);
 	}
+	return Blueprint;
 }
-
-#undef LOCTEXT_NAMESPACE	// MovieSceneSequenceEditor
\ No newline at end of file
diff --git a/Engine/Source/Editor/Sequencer/Private/Sequencer.cpp b/Engine/Source/Editor/Sequencer/Private/Sequencer.cpp
--- a/Engine/Source/Editor/Sequencer/Private/Sequencer.cpp
+++ b/Engine/Source/Editor/Sequencer/Private/Sequencer.cpp
@@ -18,6 +18,7 @@
 #include "Engine/Engine.h"
 #include "Settings/LevelEditorViewportSettings.h"
 #include "Editor.h"
+#include "BlueprintActionDatabase.h"
 #include "Channels/MovieSceneChannelProxy.h"
 #include "MovieScenePossessable.h"
 #include "MovieScene.h"
@@ -469,6 +470,16 @@ void FSequencer::InitSequencer(const FSequencerInitParams& InitParams, const TSh
 
 	BindCommands();
 
+	// Ensure that the director BP is registered with the action database
+	if (FMovieSceneSequenceEditor* SequenceEditor = FMovieSceneSequenceEditor::Find(InitParams.RootSequence))
+	{
+		UBlueprint* Blueprint = SequenceEditor->FindDirectorBlueprint(InitParams.RootSequence);
+		if (Blueprint)
+		{
+			FBlueprintActionDatabase::Get().RefreshAssetActions(Blueprint);
+		}
+	}
+
 	for (auto TrackEditor : TrackEditors)
 	{
 		TrackEditor->OnInitialize();
@@ -788,6 +799,16 @@ void FSequencer::ResetToNewRootSequence(UMovieSceneSequence& NewSequence)
 	RootSequence = &NewSequence;
 	RestorePreAnimatedState();
 
+	// Ensure that the director BP is registered with the action database
+	if (FMovieSceneSequenceEditor* SequenceEditor = FMovieSceneSequenceEditor::Find(&NewSequence))
+	{
+		UBlueprint* Blueprint = SequenceEditor->FindDirectorBlueprint(&NewSequence);
+		if (Blueprint)
+		{
+			FBlueprintActionDatabase::Get().RefreshAssetActions(Blueprint);
+		}
+	}
+
 	RootTemplateInstance.Finish(*this);
 
 	TemplateStore->Reset();
@@ -842,9 +863,21 @@ void FSequencer::FocusSequenceInstance(UMovieSceneSubSection& InSubSection)
 	ResetPerMovieSceneData();
 	SequencerWidget->UpdateBreadcrumbs();
 
+	UMovieSceneSequence* FocusedSequence = GetFocusedMovieSceneSequence();
+
 	if (!State.FindSequence(SequenceID))
 	{
-		State.AssignSequence(SequenceID, *GetFocusedMovieSceneSequence(), *this);
+		State.AssignSequence(SequenceID, *FocusedSequence, *this);
+	}
+
+	// Ensure that the director BP is registered with the action database
+	if (FMovieSceneSequenceEditor* SequenceEditor = FMovieSceneSequenceEditor::Find(FocusedSequence))
+	{
+		UBlueprint* Blueprint = SequenceEditor->FindDirectorBlueprint(FocusedSequence);
+		if (Blueprint)
+		{
+			FBlueprintActionDatabase::Get().RefreshAssetActions(Blueprint);
+		}
 	}
 
 	OnActivateSequenceEvent.Broadcast(ActiveTemplateIDs.Top());
@@ -10880,7 +10913,7 @@ void FSequencer::RecompileDirtyDirectors()
 	for (UMovieSceneSequence* Sequence : AllSequences)
 	{
 		FMovieSceneSequenceEditor* SequenceEditor = SequencerModule.FindSequenceEditor(Sequence->GetClass());
-		UBlueprint*                DirectorBP     = SequenceEditor ? SequenceEditor->GetDirectorBlueprint(Sequence) : nullptr;
+		UBlueprint*                DirectorBP     = SequenceEditor ? SequenceEditor->FindDirectorBlueprint(Sequence) : nullptr;
 
 		if (DirectorBP && (DirectorBP->Status == BS_Unknown || DirectorBP->Status == BS_Dirty))
 		{
diff --git a/Engine/Source/Editor/Sequencer/Private/SequencerChannelTraits.cpp b/Engine/Source/Editor/Sequencer/Private/SequencerChannelTraits.cpp
--- a/Engine/Source/Editor/Sequencer/Private/SequencerChannelTraits.cpp
+++ b/Engine/Source/Editor/Sequencer/Private/SequencerChannelTraits.cpp
@@ -8,7 +8,7 @@ namespace Sequencer
 {
 
 
-void DrawKeys(FMovieSceneChannel* Channel, TArrayView<const FKeyHandle> InHandles, TArrayView<FKeyDrawParams> OutKeyDrawParams)
+void DrawKeys(FMovieSceneChannel* Channel, TArrayView<const FKeyHandle> InHandles, const UMovieSceneSection* InOwner, TArrayView<FKeyDrawParams> OutKeyDrawParams)
 {
 	// By default just render diamonds for keys
 	FKeyDrawParams DefaultParams;
diff --git a/Engine/Source/Editor/Sequencer/Private/SequencerModule.cpp b/Engine/Source/Editor/Sequencer/Private/SequencerModule.cpp
--- a/Engine/Source/Editor/Sequencer/Private/SequencerModule.cpp
+++ b/Engine/Source/Editor/Sequencer/Private/SequencerModule.cpp
@@ -210,6 +210,8 @@ public:
 	{
 		if (GIsEditor)
 		{
+			// EditorStyle must be initialized by now
+			FModuleManager::Get().LoadModule("EditorStyle");
 			FSequencerCommands::Register();
 
 			FEditorModeRegistry::Get().RegisterMode<FSequencerEdMode>(
diff --git a/Engine/Source/Editor/Sequencer/Public/ISequencerChannelInterface.h b/Engine/Source/Editor/Sequencer/Public/ISequencerChannelInterface.h
--- a/Engine/Source/Editor/Sequencer/Public/ISequencerChannelInterface.h
+++ b/Engine/Source/Editor/Sequencer/Public/ISequencerChannelInterface.h
@@ -137,9 +137,10 @@ struct ISequencerChannelInterface
 	 *
 	 * @param Channel               The channel to query
 	 * @param InKeyHandles          Array of handles to duplicate
+	 * @param InOwner               The section that owns the channel
 	 * @param OutKeyDrawParams      Pre-sized array to receive key draw parameters. Invalid key handles will not be assigned to this array. Must match size of InKeyHandles.
 	 */
-	virtual void DrawKeys_Raw(FMovieSceneChannel* Channel, TArrayView<const FKeyHandle> InKeyHandles, TArrayView<FKeyDrawParams> OutKeyDrawParams) const = 0;
+	virtual void DrawKeys_Raw(FMovieSceneChannel* Channel, TArrayView<const FKeyHandle> InKeyHandles, const UMovieSceneSection* InOwner, TArrayView<FKeyDrawParams> OutKeyDrawParams) const = 0;
 
 	/**
 	 * Whether this channel supports curve models
diff --git a/Engine/Source/Editor/Sequencer/Public/MovieSceneSequenceEditor.h b/Engine/Source/Editor/Sequencer/Public/MovieSceneSequenceEditor.h
--- a/Engine/Source/Editor/Sequencer/Public/MovieSceneSequenceEditor.h
+++ b/Engine/Source/Editor/Sequencer/Public/MovieSceneSequenceEditor.h
@@ -5,26 +5,14 @@
 #include "CoreTypes.h"
 #include "UObject/NameTypes.h"
 
-struct FMovieSceneEvent;
-
 class UBlueprint;
-class UMovieSceneTrack;
 class UMovieSceneSequence;
-class UMovieSceneEventTrack;
-class UK2Node_FunctionEntry;
-class UMovieSceneEventSectionBase;
 
 struct SEQUENCER_API FMovieSceneSequenceEditor
 {
 	virtual ~FMovieSceneSequenceEditor(){}
 
 
-	/**
-	 * The name of the target pin for event function entry nodes
-	 */
-	static FName TargetPinName;
-
-
 	/**
 	 * Attempt to find a sequence editor for the specified sequence
 	 *
@@ -34,14 +22,6 @@ struct SEQUENCER_API FMovieSceneSequenceEditor
 	static FMovieSceneSequenceEditor* Find(UMovieSceneSequence* InSequence);
 
 
-	/**
-	 * Bind an event to an endpoint. Ensures the necessary pointers are set up to enable automatic re-compilation on evaluation.
-	 *
-	 * @param EventSection    The section that the event belongs to
-	 * @param Event           The event to bind to the endpoint
-	 * @param Endpoint        The blueprint node to bind to
-	 */
-	static void BindEventToEndpoint(UMovieSceneEventSectionBase* EventSection, FMovieSceneEvent* Event, UK2Node_FunctionEntry* Endpoint);
 
 	/**
 	 * Check whether the specified sequence supports events
@@ -57,37 +37,7 @@ struct SEQUENCER_API FMovieSceneSequenceEditor
 	 * @param Sequence        The sequence to access the director blueprint for
 	 * @return The sequence's director blueprint or nullptr if it does not have one (or cannot)
 	 */
-	UBlueprint* GetDirectorBlueprint(UMovieSceneSequence* Sequence) const;
-
-
-	/**
-	 * Create a new event endpoint for the specified sequence
-	 *
-	 * @param Sequence        The sequence to create an endpoint in
-	 * @param DesiredName     (Optional) A user-specified name to use for the function
-	 * @return The blueprint node for the event endpoint, or nullptr if one does not exist
-	 */
-	UK2Node_FunctionEntry* CreateEventEndpoint(UMovieSceneSequence* Sequence, const FString& DesiredName = FString()) const;
-
-
-	/**
-	 * Initializes the specified endpoint node by creating an appropriate input pin for the track's object binding, if necessary
-	 *
-	 * @param EventTrack      The event track that owns the event
-	 * @param Endpoint        The event endpoint to initialize
-	 */
-	void InitializeEndpointForTrack(UMovieSceneEventTrack* EventTrack, UK2Node_FunctionEntry* Endpoint) const;
-
-protected:
-
-
-	/**
-	 * Find the class of the object binding the specified track is on, or nullptr
-	 */
-	static UClass* FindTrackObjectBindingClass(UMovieSceneTrack* Track);
-
-
-private:
+	UBlueprint* FindDirectorBlueprint(UMovieSceneSequence* Sequence) const;
 
 	/**
 	 * Access the director blueprint for the specified sequence
@@ -95,13 +45,12 @@ private:
 	 * @param Sequence        The sequence to access the director blueprint for
 	 * @return The sequence's director blueprint or nullptr if it does not have one (or cannot)
 	 */
-	UBlueprint* AccessDirectorBlueprint(UMovieSceneSequence* Sequence) const;
+	UBlueprint* GetOrCreateDirectorBlueprint(UMovieSceneSequence* Sequence) const;
 
 private:
 
 	virtual bool CanCreateEvents(UMovieSceneSequence* InSequence) const { return false; }
 	virtual UBlueprint* GetBlueprintForSequence(UMovieSceneSequence* InSequence) const { return nullptr; }
 	virtual UBlueprint* CreateBlueprintForSequence(UMovieSceneSequence* InSequence) const { return nullptr; }
-	virtual void SetupDefaultPinForEndpoint(UMovieSceneEventTrack* EventTrack, UK2Node_FunctionEntry* Endpoint) const {}
 };
 
diff --git a/Engine/Source/Editor/Sequencer/Public/SequencerChannelInterface.h b/Engine/Source/Editor/Sequencer/Public/SequencerChannelInterface.h
--- a/Engine/Source/Editor/Sequencer/Public/SequencerChannelInterface.h
+++ b/Engine/Source/Editor/Sequencer/Public/SequencerChannelInterface.h
@@ -129,14 +129,15 @@ struct TSequencerChannelInterfaceCommon : ISequencerChannelInterface
 	 *
 	 * @param Channel               The channel to query
 	 * @param InKeyHandles          Array of handles to duplicate
+	 * @param InOwner               The section that owns the channel
 	 * @param OutKeyDrawParams      Pre-sized array to receive key draw parameters. Invalid key handles will not be assigned to this array. Must match size of InKeyHandles.
 	 */
-	virtual void DrawKeys_Raw(FMovieSceneChannel* InChannel, TArrayView<const FKeyHandle> InKeyHandles, TArrayView<FKeyDrawParams> OutKeyDrawParams) const override
+	virtual void DrawKeys_Raw(FMovieSceneChannel* InChannel, TArrayView<const FKeyHandle> InKeyHandles, const UMovieSceneSection* InOwner, TArrayView<FKeyDrawParams> OutKeyDrawParams) const override
 	{
 		check(InKeyHandles.Num() == OutKeyDrawParams.Num());
 
 		using namespace Sequencer;
-		DrawKeys(static_cast<ChannelType*>(InChannel), InKeyHandles, OutKeyDrawParams);
+		DrawKeys(static_cast<ChannelType*>(InChannel), InKeyHandles, InOwner, OutKeyDrawParams);
 	}
 
 	/**
diff --git a/Engine/Source/Editor/Sequencer/Public/SequencerChannelTraits.h b/Engine/Source/Editor/Sequencer/Public/SequencerChannelTraits.h
--- a/Engine/Source/Editor/Sequencer/Public/SequencerChannelTraits.h
+++ b/Engine/Source/Editor/Sequencer/Public/SequencerChannelTraits.h
@@ -246,9 +246,10 @@ namespace Sequencer
 	 *
 	 * @param InChannel          The channel to duplicate keys in
 	 * @param InHandles          Array of key handles that should be deleted
+	 * @param InOwner            The owning movie scene section for this channel
 	 * @param OutKeyDrawParams   Array to receive key draw information. Must be exactly the size of InHandles.
 	 */
-	SEQUENCER_API void DrawKeys(FMovieSceneChannel* Channel, TArrayView<const FKeyHandle> InHandles, TArrayView<FKeyDrawParams> OutKeyDrawParams);
+	SEQUENCER_API void DrawKeys(FMovieSceneChannel* Channel, TArrayView<const FKeyHandle> InHandles, const UMovieSceneSection* InOwner, TArrayView<FKeyDrawParams> OutKeyDrawParams);
 
 
 	/**
diff --git a/Engine/Source/Editor/UMGEditor/Private/Animation/MovieSceneSequenceEditor_WidgetAnimation.h b/Engine/Source/Editor/UMGEditor/Private/Animation/MovieSceneSequenceEditor_WidgetAnimation.h
--- a/Engine/Source/Editor/UMGEditor/Private/Animation/MovieSceneSequenceEditor_WidgetAnimation.h
+++ b/Engine/Source/Editor/UMGEditor/Private/Animation/MovieSceneSequenceEditor_WidgetAnimation.h
@@ -18,30 +18,4 @@ struct FMovieSceneSequenceEditor_WidgetAnimation : FMovieSceneSequenceEditor
 	{
 		return InSequence->GetTypedOuter<UBlueprint>();
 	}
-
-	virtual void SetupDefaultPinForEndpoint(UMovieSceneEventTrack* EventTrack, UK2Node_FunctionEntry* Endpoint) const override
-	{
-		UClass* PinClass = nullptr;
-
-		// When event receivers are set, we just add a generic object pin as it's impossible to know what type the event receiver is here
-		if (EventTrack->EventReceivers.Num() > 0)
-		{
-			PinClass = UObject::StaticClass();
-		}
-		// Otherwise we use the object binding's class, or nullptr
-		else
-		{
-			PinClass = FindTrackObjectBindingClass(EventTrack);
-		}
-
-		// Do not create a pin for nullptr since we're probably a master track that is just triggering on self
-		if (PinClass)
-		{
-			FEdGraphPinType PinType;
-			PinType.PinCategory = UEdGraphSchema_K2::PC_Object;
-			PinType.PinSubCategoryObject = PinClass;
-
-			Endpoint->CreateUserDefinedPin(TargetPinName, PinType, EGPD_Output, true);
-		}
-	}
 };
\ No newline at end of file
diff --git a/Engine/Source/Runtime/Engine/Private/Blueprint.cpp b/Engine/Source/Runtime/Engine/Private/Blueprint.cpp
--- a/Engine/Source/Runtime/Engine/Private/Blueprint.cpp
+++ b/Engine/Source/Runtime/Engine/Private/Blueprint.cpp
@@ -2010,8 +2010,11 @@ UEdGraph* UBlueprint::GetLastEditedUberGraph() const
 #if WITH_EDITORONLY_DATA
 void UBlueprint::LoadModulesRequiredForCompilation()
 {
-	static const FName ModuleName(TEXT("KismetCompiler"));
-	FModuleManager::Get().LoadModule(ModuleName);
+	static const FName KismetCompilerModuleName("KismetCompiler");
+	static const FName MovieSceneToolsModuleName("MovieSceneTools");
+
+	FModuleManager::Get().LoadModule(KismetCompilerModuleName);
+	FModuleManager::Get().LoadModule(MovieSceneToolsModuleName);
 }
 #endif //WITH_EDITORONLY_DATA
 
diff --git a/Engine/Source/Runtime/MovieSceneTracks/MovieSceneTracks.Build.cs b/Engine/Source/Runtime/MovieSceneTracks/MovieSceneTracks.Build.cs
--- a/Engine/Source/Runtime/MovieSceneTracks/MovieSceneTracks.Build.cs
+++ b/Engine/Source/Runtime/MovieSceneTracks/MovieSceneTracks.Build.cs
@@ -38,6 +38,11 @@ public class MovieSceneTracks : ModuleRules
 					"BlueprintGraph"
 				}
 			);
-		}
+            PrivateDependencyModuleNames.AddRange(
+                new string[] {
+                    "UnrealEd"
+                }
+            );
+        }
 	}
 }
diff --git a/Engine/Source/Runtime/MovieSceneTracks/Private/Channels/MovieSceneEvent.cpp b/Engine/Source/Runtime/MovieSceneTracks/Private/Channels/MovieSceneEvent.cpp
--- a/Engine/Source/Runtime/MovieSceneTracks/Private/Channels/MovieSceneEvent.cpp
+++ b/Engine/Source/Runtime/MovieSceneTracks/Private/Channels/MovieSceneEvent.cpp
@@ -3,192 +3,24 @@
 #include "Channels/MovieSceneEvent.h"
 
 #if WITH_EDITORONLY_DATA
-
-#include "EdGraph/EdGraph.h"
-#include "K2Node_FunctionEntry.h"
-#include "EdGraphSchema_K2.h"
-#include "UObject/PropertyPortFlags.h"
-
-UK2Node_FunctionEntry* FMovieSceneEvent::GetFunctionEntry() const
-{
-	if (SoftBlueprintPath.IsNull())
-	{
-		// The function entry used to be serialized but is now only stored transiently. We use this pointer for the current lifecycle until the asset is saved, when we do the data upgrade.
-		UK2Node_FunctionEntry* FunctionEntryPtr = CastChecked<UK2Node_FunctionEntry>(FunctionEntry_DEPRECATED.Get(), ECastCheckedType::NullAllowed);
-		if (FunctionEntryPtr)
-		{
-			return FunctionEntryPtr;
-		}
-	}
-
-	UK2Node_FunctionEntry* Cached = CastChecked<UK2Node_FunctionEntry>(CachedFunctionEntry.Get(), ECastCheckedType::NullAllowed);
-	if (Cached)
-	{
-		return Cached;
-	}
-
-	if (GraphGuid.IsValid())
-	{
-		if (UBlueprint* Blueprint = SoftBlueprintPath.LoadSynchronous())
-		{
-			for (UEdGraph* Graph : Blueprint->FunctionGraphs)
-			{
-				if (Graph->GraphGuid == GraphGuid)
-				{
-					for (UEdGraphNode* Node : Graph->Nodes)
-					{
-						if (UK2Node_FunctionEntry* FunctionEntry = Cast<UK2Node_FunctionEntry>(Node))
-						{
-							CachedFunctionEntry = FunctionEntry;
-							return FunctionEntry;
-						}
-					}
-				}
-			}
-		}
-	}
-
-	return nullptr;
-}
-
-bool FMovieSceneEvent::IsReadyToGetFunctionEntry()
-{
-	// We cannot get the function entry until all graphs in the blueprint are loaded.
-	if (UBlueprint* Blueprint = SoftBlueprintPath.LoadSynchronous())
-	{
-		for (UEdGraph* Graph : Blueprint->FunctionGraphs)
-		{
-			if (Graph->HasAnyFlags(RF_NeedLoad))
-			{
-				return false;
-			}
-		}
-	}
-	return true;
-}
-
-void FMovieSceneEvent::SetFunctionEntry(UK2Node_FunctionEntry* InFunctionEntry)
-{
-	UEdGraph* EdGraph = InFunctionEntry ? InFunctionEntry->GetGraph() : nullptr;
-	if (EdGraph)
-	{
-		CachedFunctionEntry = InFunctionEntry;
-		SoftBlueprintPath   = EdGraph->GetTypedOuter<UBlueprint>();
-		GraphGuid           = EdGraph->GraphGuid;
-	}
-	else
-	{
-		CachedFunctionEntry = nullptr;
-		SoftBlueprintPath   = nullptr;
-		GraphGuid           = FGuid();
-	}
-	CacheFunctionName();
-}
-
-bool FMovieSceneEvent::IsBoundToBlueprint() const
-{
-	return IsValidFunction(GetFunctionEntry());
-}
-
-bool FMovieSceneEvent::IsValidFunction(UK2Node_FunctionEntry* Function)
-{
-	// Must have a single non-reference parameter on the function
-	if (!Function || Function->IsPendingKill() || !Function->GetGraph() || Function->GetGraph()->IsPendingKill())
-	{
-		return false;
-	}
-	else if (Function->UserDefinedPins.Num() == 0)
-	{
-		return true;
-	}
-	else if (Function->UserDefinedPins.Num() != 1 || Function->UserDefinedPins[0]->PinType.bIsReference)
-	{
-		return false;
-	}
-
-	// Pin must be an object or interface property
-	FName PinCategory = Function->UserDefinedPins[0]->PinType.PinCategory;
-	return PinCategory == UEdGraphSchema_K2::PC_Object || PinCategory == UEdGraphSchema_K2::PC_Interface;
-}
-
-void FMovieSceneEvent::CacheFunctionName()
-{
-	UK2Node_FunctionEntry* Node  = GetFunctionEntry();
-	UEdGraph*              Graph = Node ? Node->GetGraph() : nullptr;
-
-	FunctionName = NAME_None;
-
-	if (Graph && IsValidFunction(Node))
-	{
-		FunctionName = Graph->GetFName();
-	}
-}
-
-#endif // WITH_EDITORONLY_DATA
-
+	#include "EdGraph/EdGraph.h"
+	#include "K2Node_FunctionEntry.h"
+#endif
 
 void FMovieSceneEvent::PostSerialize(const FArchive& Ar)
 {
 #if WITH_EDITORONLY_DATA
-	if (Ar.IsLoading() && !Ar.HasAnyPortFlags(PPF_Duplicate | PPF_DuplicateForPIE))
-	{
-		// Re-cache the function name when loading in-editor in case of renamed function graphs and the like 
-		// We can only do it safely if the blueprint and Graph are done loading.
-		// Without this check the call will reset the function name to null and subsequent call to it will mark the object as Changed even if it do not. 
-		// causing Non deterministic cooking.
-		if (IsReadyToGetFunctionEntry())
-		{
-			CacheFunctionName();
-		}
-	}
-#endif
-}
 
-bool FMovieSceneEvent::Serialize(FArchive& Ar)
-{
-#if WITH_EDITORONLY_DATA
-	if (Ar.IsSaving())
+	UK2Node_FunctionEntry* FunctionEntry = Cast<UK2Node_FunctionEntry>(FunctionEntry_DEPRECATED.Get());
+	if (FunctionEntry)
 	{
-		// ---------------------------------------------------------------------------------------
-		// Data upgrade for content that was saved with FunctionEntry_DEPRECATED instead of SoftFunctionGraph
-		// We do this on save because there is no reliable way to ensure that FunctionGraph is fully loaded here
-		// since the track may live inside or outside of a blueprint. When not fully loaded, GraphGuid is not correct.
-		if (SoftBlueprintPath.IsNull())
+		// If the graph needs loading, load it
+		if (FunctionEntry->GetGraph()->HasAnyFlags(RF_NeedLoad) && FunctionEntry->GetLinker())
 		{
-			// The function entry used to be serialized but is now only stored transiently. If this is set without the soft function graph being set, copy the graph reference over
-			UK2Node_FunctionEntry* FunctionEntryPtr = CastChecked<UK2Node_FunctionEntry>(FunctionEntry_DEPRECATED.Get(), ECastCheckedType::NullAllowed);
-			UEdGraph*              FunctionGraph    = FunctionEntryPtr ? FunctionEntryPtr->GetGraph() : nullptr;
-
-			if (FunctionGraph)
-			{
-				CachedFunctionEntry = FunctionEntryPtr;
-				SoftBlueprintPath   = FunctionGraph->GetTypedOuter<UBlueprint>();
-				GraphGuid           = FunctionGraph->GraphGuid;
-			}
+			FunctionEntry->GetLinker()->Preload(FunctionEntry->GetGraph());
 		}
-	}
-#endif
-
-	// Return false to ensure that the struct receives default serialization
-	return false;
-}
-
-bool FMovieSceneEvent::IsValidFunction(UFunction* Function)
-{
-	// Must have a single non-reference parameter on the function
-	if (!Function)
-	{
-		return false;
-	}
-	else if (Function->NumParms == 0)
-	{
-		return true;
-	}
-	else if (Function->NumParms != 1 || !Function->PropertyLink || (Function->PropertyLink->GetPropertyFlags() & CPF_ReferenceParm))
-	{
-		return false;
+		GraphGuid = FunctionEntry->GetGraph()->GraphGuid;
 	}
 
-	// Parameter must be an object or interface property
-	return Function->PropertyLink->IsA<UObjectProperty>() || Function->PropertyLink->IsA<UInterfaceProperty>();
+#endif
 }
diff --git a/Engine/Source/Runtime/MovieSceneTracks/Private/Evaluation/MovieSceneEventTemplate.cpp b/Engine/Source/Runtime/MovieSceneTracks/Private/Evaluation/MovieSceneEventTemplate.cpp
--- a/Engine/Source/Runtime/MovieSceneTracks/Private/Evaluation/MovieSceneEventTemplate.cpp
+++ b/Engine/Source/Runtime/MovieSceneTracks/Private/Evaluation/MovieSceneEventTemplate.cpp
@@ -168,7 +168,7 @@ struct FEventTrackExecutionToken
 struct FEventTriggerExecutionToken
 	: IMovieSceneExecutionToken
 {
-	FEventTriggerExecutionToken(TArray<FName> InEvents, const TArray<FMovieSceneObjectBindingID>& InEventReceivers)
+	FEventTriggerExecutionToken(TArray<FMovieSceneEventPtrs> InEvents, const TArray<FMovieSceneObjectBindingID>& InEventReceivers)
 		: Events(MoveTemp(InEvents)), EventReceivers(InEventReceivers)
 	{}
 
@@ -182,7 +182,7 @@ struct FEventTriggerExecutionToken
 		{
 #if !NO_LOGGING
 			UE_LOG(LogMovieScene, Warning, TEXT("Failed to trigger the following events because no director instance was available: %s."), *GenerateEventListString());
-#endif			
+#endif
 			return;
 		}
 
@@ -234,129 +234,135 @@ struct FEventTriggerExecutionToken
 		bool bIsGameWorld = World && World->IsGameWorld();
 #endif // WITH_EDITOR
 
-		for (FName EventName : Events)
+		for (const FMovieSceneEventPtrs& Event : Events)
 		{
-			if (EventName == NAME_None)
-			{
-				continue;
-			}
 
-			UFunction* Function = DirectorInstance->FindFunction(EventName);
-			// Event must have only a single object parameter, and the director instance must be an implementation of the function's class
-			if (!Function)
-			{
-				FMessageLog("PIE").Warning()
-					->AddToken(FUObjectToken::Create(Player.GetEvaluationTemplate().GetSequence(Operand.SequenceID)))
-					->AddToken(FTextToken::Create(FText::Format(LOCTEXT("LevelBP_MissingEvent_Error", "Failed to trigger event '{0}' because the function does not exist on the director instance."), FText::FromName(EventName))));
-				continue;
-			}
 #if WITH_EDITOR
-			else if (!bIsGameWorld && !Function->HasMetaData(NAME_CallInEditor))
+			if (!bIsGameWorld && !Event.Function->HasMetaData(NAME_CallInEditor))
 			{
-				UE_LOG(LogMovieScene, Verbose, TEXT("Refusing to trigger event '%s' in editor world when 'Call in Editor' is false."), *EventName.ToString());
+				UE_LOG(LogMovieScene, Verbose, TEXT("Refusing to trigger event '%s' in editor world when 'Call in Editor' is false."), *Event.Function->GetName());
 				continue;
 			}
 #endif // WITH_EDITOR
 
-			if (Function->NumParms == 0)
-			{
-				UE_LOG(LogMovieScene, VeryVerbose, TEXT("Triggering event '%s'."), *EventName.ToString());
-				DirectorInstance->ProcessEvent(Function, nullptr);
-			}
-			else if (Function->NumParms == 1 && Function->PropertyLink && (Function->PropertyLink->GetPropertyFlags() & CPF_ReferenceParm) == 0)
-			{
-				const int32 NumLevelScripts = Algo::Accumulate(EventContexts, 0, 
-					[](int32 Count, UObject* Obj)
-					{
-						return Obj && Obj->IsA<ALevelScriptActor>() ? Count + 1 : Count;
-					}
-				);
 
-				// Never pass through level script actors to event endpoints on non-interface pins.
-				if (NumLevelScripts > 0 && NumLevelScripts == EventContexts.Num() && !Function->PropertyLink->IsA<UInterfaceProperty>())
-				{
-					FMessageLog("PIE").Warning()
-						->AddToken(FUObjectToken::Create(Function))
-						->AddToken(FUObjectToken::Create(Player.GetEvaluationTemplate().GetSequence(Operand.SequenceID)))
-						->AddToken(FTextToken::Create(LOCTEXT("LevelBP_ObjectPin_Error", "Failed to trigger event: only Interface pins are supported for master tracks within Level Sequences. Please remove the pin, or change it to an interface that is implemented on the desired level blueprint.")));
-					continue;
-				}
+			UE_LOG(LogMovieScene, VeryVerbose, TEXT("Triggering event '%s'."), *Event.Function->GetName());
 
-				for (UObject* EventContextObject : EventContexts)
-				{
-					TriggerEvent(DirectorInstance, Function, EventContextObject, Player, Operand.SequenceID);
-				}
+			if (Event.Function->NumParms == 0)
+			{
+				DirectorInstance->ProcessEvent(Event.Function, nullptr);
 			}
 			else
 			{
-				FMessageLog("PIE").Warning()
-					->AddToken(FUObjectToken::Create(Function))
-					->AddToken(FUObjectToken::Create(Player.GetEvaluationTemplate().GetSequence(Operand.SequenceID)))
-					->AddToken(FTextToken::Create(LOCTEXT("LevelBP_InvalidEvent_Error", "Failed to trigger event because its signature is not compatible. Function signatures must have either 0 or 1 (non-ref) parameters.")));
+				TriggerEventWithParameters(DirectorInstance, Event, EventContexts, Player, Operand.SequenceID);
 			}
 		}
 	}
 
-	void TriggerEvent(UObject* DirectorInstance, UFunction* Function, UObject* ObjectParamValue, IMovieScenePlayer& Player, FMovieSceneSequenceID SequenceID)
+	void TriggerEventWithParameters(UObject* DirectorInstance, const FMovieSceneEventPtrs& Event, TArrayView<UObject* const> EventContexts, IMovieScenePlayer& Player, FMovieSceneSequenceID SequenceID)
 	{
-		// We know by now that the parameter type is compatible because it wouldn't be added to the array if not
-		if (UObjectProperty* ObjectParameter = Cast<UObjectProperty>(Function->PropertyLink))
+		// Parse all function parameters.
+		uint8* Parameters = (uint8*)FMemory_Alloca(Event.Function->ParmsSize + Event.Function->MinAlignment);
+		Parameters = Align(Parameters, Event.Function->MinAlignment);
+
+		Event.Function->InitializeStruct(Parameters);
+
+		for (UObject* BoundObject : EventContexts)
 		{
-			if (!ObjectParameter->PropertyClass || ObjectParamValue->IsA(ObjectParameter->PropertyClass))
+			int32 NumParmsPatched = 0;
+
+			// Attempt to bind the object to the function parameters
+			if (!PatchBoundObject(Parameters, BoundObject, Event.BoundObjectProperty, Player, SequenceID))
 			{
-				DirectorInstance->ProcessEvent(Function, &ObjectParamValue);
-				return;
+				continue;
 			}
+			else
+			{
+				++NumParmsPatched;
+			}
+
+			ensureAlwaysMsgf(Event.Function->NumParms == NumParmsPatched, TEXT("Failed to patch the correct number of parameters for function call. Some parameters may be incorrect."));
 
+			// Call the function
+			DirectorInstance->ProcessEvent(Event.Function, Parameters);
+		}
 
-			UE_LOG(LogMovieScene, VeryVerbose, TEXT("Failed to trigger event '%s' with object '%s' because it is not the correct type. Function expects a '%s' but target object is a '%s'."),
-				*Function->GetName(),
-				*ObjectParamValue->GetName(),
-				*ObjectParameter->PropertyClass->GetName(),
-				*ObjectParamValue->GetClass()->GetName()
-				);
+		Event.Function->DestroyStruct(Parameters);
+	}
 
-			return;
+	bool PatchBoundObject(uint8* Parameters, UObject* BoundObject, UProperty* BoundObjectProperty, IMovieScenePlayer& Player, FMovieSceneSequenceID SequenceID)
+	{
+		if (!BoundObjectProperty)
+		{
+			return true;
 		}
 
-		if (UInterfaceProperty* InterfaceParameter = Cast<UInterfaceProperty>(Function->PropertyLink))
+		if (UInterfaceProperty* InterfaceParameter = Cast<UInterfaceProperty>(BoundObjectProperty))
 		{
-			if (ObjectParamValue->GetClass()->ImplementsInterface(InterfaceParameter->InterfaceClass))
+			if (BoundObject->GetClass()->ImplementsInterface(InterfaceParameter->InterfaceClass))
 			{
-				DirectorInstance->ProcessEvent(Function, &ObjectParamValue);
-				return;
+				FScriptInterface Interface;
+				Interface.SetObject(BoundObject);
+				Interface.SetInterface(BoundObject->GetInterfaceAddress(InterfaceParameter->InterfaceClass));
+				InterfaceParameter->SetPropertyValue_InContainer(Parameters, BoundObject);
+				return true;
 			}
 
+			FMessageLog("PIE").Warning()
+				->AddToken(FUObjectToken::Create(BoundObjectProperty->GetOuter()))
+				->AddToken(FUObjectToken::Create(Player.GetEvaluationTemplate().GetSequence(SequenceID)))
+				->AddToken(FTextToken::Create(FText::Format(LOCTEXT("LevelBP_InterfaceNotImplemented_Error", "Failed to trigger event because it does not implement the necessary interface. Function expects a '{0}'."), FText::FromName(InterfaceParameter->InterfaceClass->GetFName()))));
+			return false;
+		}
+
+		if (UObjectProperty* ObjectParameter = Cast<UObjectProperty>(BoundObjectProperty))
+		{
+			if (BoundObject->IsA<ALevelScriptActor>())
+			{
+				FMessageLog("PIE").Warning()
+					->AddToken(FUObjectToken::Create(BoundObjectProperty->GetOuter()))
+					->AddToken(FUObjectToken::Create(Player.GetEvaluationTemplate().GetSequence(SequenceID)))
+					->AddToken(FTextToken::Create(LOCTEXT("LevelBP_LevelScriptActor_Error", "Failed to trigger event: only Interface pins are supported for master tracks within Level Sequences. Please remove the pin, or change it to an interface that is implemented on the desired level blueprint.")));
+
+				return false;
+			}
+			else if (!BoundObject->IsA(ObjectParameter->PropertyClass))
+			{
+				FMessageLog("PIE").Warning()
+					->AddToken(FUObjectToken::Create(Player.GetEvaluationTemplate().GetSequence(SequenceID)))
+					->AddToken(FUObjectToken::Create(BoundObjectProperty->GetOuter()))
+					->AddToken(FUObjectToken::Create(BoundObject))
+					->AddToken(FTextToken::Create(FText::Format(LOCTEXT("LevelBP_InvalidCast_Error", "Failed to trigger event: Cast to %s failed."), FText::FromName(ObjectParameter->PropertyClass->GetFName()))));
 
-			UE_LOG(LogMovieScene, VeryVerbose, TEXT("Failed to trigger event '%s' with object '%s' because it does not implement the necessary interface. Function expects a '%s'."),
-				*Function->GetName(),
-				*ObjectParamValue->GetName(),
-				*InterfaceParameter->InterfaceClass->GetName()
-				);
+				return false;
+			}
 
-			return;
+			ObjectParameter->SetObjectPropertyValue_InContainer(Parameters, BoundObject);
+			return true;
 		}
 
 		FMessageLog("PIE").Warning()
-			->AddToken(FUObjectToken::Create(Function))
 			->AddToken(FUObjectToken::Create(Player.GetEvaluationTemplate().GetSequence(SequenceID)))
-			->AddToken(FTextToken::Create(FText::Format(LOCTEXT("LevelBP_InvalidObjectEvent_Error", "Failed to trigger event because its signature is not compatible. Function expects a '{0}' parameter, but only object and interface parameters are supported."), FText::FromName(Function->PropertyLink->GetClass()->GetFName()))));
+			->AddToken(FUObjectToken::Create(BoundObjectProperty->GetOuter()))
+			->AddToken(FUObjectToken::Create(BoundObject))
+			->AddToken(FTextToken::Create(FText::Format(LOCTEXT("LevelBP_UnsupportedProperty_Error", "Failed to trigger event: Unsupported property type for bound object: {0}."), FText::FromName(BoundObjectProperty->GetClass()->GetFName()))));
+		return false;
 	}
 
 #if !NO_LOGGING
 	FString GenerateEventListString() const
 	{
-		return Algo::Accumulate(Events, FString(), [](FString&& InString, FName Event){
+		return Algo::Accumulate(Events, FString(), [](FString&& InString, const FMovieSceneEventPtrs& EventPtrs){
 			if (InString.Len() > 0)
 			{
 				InString += TEXT(", ");
 			}
-			return InString + Event.ToString();
+			return InString + EventPtrs.Function->GetName();
 		});
 	}
 #endif
 
-	TArray<FName> Events;
+	TArray<FMovieSceneEventPtrs> Events;
 	TArray<FMovieSceneObjectBindingID, TInlineAllocator<2>> EventReceivers;
 };
 
@@ -437,17 +443,17 @@ void FMovieSceneEventSectionTemplate::EvaluateSwept(const FMovieSceneEvaluationO
 FMovieSceneEventTriggerTemplate::FMovieSceneEventTriggerTemplate(const UMovieSceneEventTriggerSection& Section, const UMovieSceneEventTrack& Track)
 	: FMovieSceneEventTemplateBase(Track)
 {
-	TMovieSceneChannelData<const FMovieSceneEvent> EventData = Section.EventChannel.GetData();
-	TArrayView<const FFrameNumber>     Times  = EventData.GetTimes();
-	TArrayView<const FMovieSceneEvent> Events = EventData.GetValues();
+	TMovieSceneChannelData<const FMovieSceneEvent> EventData   = Section.EventChannel.GetData();
+	TArrayView<const FFrameNumber>                 Times       = EventData.GetTimes();
+	TArrayView<const FMovieSceneEvent>             EntryPoints = EventData.GetValues();
 
 	EventTimes.Reserve(Times.Num());
-	EventFunctions.Reserve(Times.Num());
+	Events.Reserve(Times.Num());
 
 	for (int32 Index = 0; Index < Times.Num(); ++Index)
 	{
 		EventTimes.Add(Times[Index]);
-		EventFunctions.Add(Events[Index].FunctionName);
+		Events.Add(EntryPoints[Index].Ptrs);
 	}
 }
 
@@ -469,7 +475,7 @@ void FMovieSceneEventTriggerTemplate::EvaluateSwept(const FMovieSceneEvaluationO
 		return;
 	}
 
-	TArray<FName> Events;
+	TArray<FMovieSceneEventPtrs> EventsToTrigger;
 
 	const int32 First = bBackwards ? EventTimes.Num() - 1 : 0;
 	const int32 Last = bBackwards ? 0 : EventTimes.Num() - 1;
@@ -483,9 +489,9 @@ void FMovieSceneEventTriggerTemplate::EvaluateSwept(const FMovieSceneEvaluationO
 		for (int32 KeyIndex = EventTimes.Num() - 1; KeyIndex >= 0; --KeyIndex)
 		{
 			FFrameNumber Time = EventTimes[KeyIndex];
-			if (SweptRange.Contains(Time))
+			if (Events[KeyIndex].Function && SweptRange.Contains(Time))
 			{
-				Events.Add(EventFunctions[KeyIndex]);
+				EventsToTrigger.Add(Events[KeyIndex]);
 			}
 		}
 	}
@@ -493,16 +499,16 @@ void FMovieSceneEventTriggerTemplate::EvaluateSwept(const FMovieSceneEvaluationO
 	else for (int32 KeyIndex = 0; KeyIndex < EventTimes.Num(); ++KeyIndex)
 	{
 		FFrameNumber Time = EventTimes[KeyIndex];
-		if (SweptRange.Contains(Time))
+		if (Events[KeyIndex].Function && SweptRange.Contains(Time))
 		{
-			Events.Add(EventFunctions[KeyIndex]);
+			EventsToTrigger.Add(Events[KeyIndex]);
 		}
 	}
 
 
-	if (Events.Num())
+	if (EventsToTrigger.Num())
 	{
-		ExecutionTokens.Add(FEventTriggerExecutionToken(MoveTemp(Events), EventReceivers));
+		ExecutionTokens.Add(FEventTriggerExecutionToken(MoveTemp(EventsToTrigger), EventReceivers));
 	}
 }
 
@@ -510,7 +516,7 @@ void FMovieSceneEventTriggerTemplate::EvaluateSwept(const FMovieSceneEvaluationO
 
 FMovieSceneEventRepeaterTemplate::FMovieSceneEventRepeaterTemplate(const UMovieSceneEventRepeaterSection& Section, const UMovieSceneEventTrack& Track)
 	: FMovieSceneEventTemplateBase(Track)
-	, EventToTrigger(Section.Event.FunctionName)
+	, EventToTrigger(Section.Event.Ptrs)
 {
 }
 
@@ -522,7 +528,7 @@ void FMovieSceneEventRepeaterTemplate::EvaluateSwept(const FMovieSceneEvaluation
 	// Don't allow events to fire when playback is in a stopped state. This can occur when stopping 
 	// playback and returning the current position to the start of playback. It's not desireable to have 
 	// all the events from the last playback position to the start of playback be fired.
-	if (!SweptRange.Contains(CurrentFrame) || Context.GetStatus() == EMovieScenePlayerStatus::Stopped || Context.IsSilent())
+	if (!EventToTrigger.Function || !SweptRange.Contains(CurrentFrame) || Context.GetStatus() == EMovieScenePlayerStatus::Stopped || Context.IsSilent())
 	{
 		return;
 	}
@@ -530,7 +536,7 @@ void FMovieSceneEventRepeaterTemplate::EvaluateSwept(const FMovieSceneEvaluation
 	
 	if ((!bBackwards && bFireEventsWhenForwards) || (bBackwards && bFireEventsWhenBackwards))
 	{
-		TArray<FName> Events = { EventToTrigger };
+		TArray<FMovieSceneEventPtrs> Events = { EventToTrigger };
 		ExecutionTokens.Add(FEventTriggerExecutionToken(MoveTemp(Events), EventReceivers));
 	}
 }
diff --git a/Engine/Source/Runtime/MovieSceneTracks/Private/Evaluation/MovieSceneEventTemplate.h b/Engine/Source/Runtime/MovieSceneTracks/Private/Evaluation/MovieSceneEventTemplate.h
--- a/Engine/Source/Runtime/MovieSceneTracks/Private/Evaluation/MovieSceneEventTemplate.h
+++ b/Engine/Source/Runtime/MovieSceneTracks/Private/Evaluation/MovieSceneEventTemplate.h
@@ -67,7 +67,7 @@ struct FMovieSceneEventTriggerTemplate : public FMovieSceneEventTemplateBase
 	TArray<FFrameNumber> EventTimes;
 
 	UPROPERTY()
-	TArray<FName> EventFunctions;
+	TArray<FMovieSceneEventPtrs> Events;
 
 private:
 
@@ -84,7 +84,7 @@ struct FMovieSceneEventRepeaterTemplate : public FMovieSceneEventTemplateBase
 	FMovieSceneEventRepeaterTemplate(const UMovieSceneEventRepeaterSection& Section, const UMovieSceneEventTrack& Track);
 
 	UPROPERTY()
-	FName EventToTrigger;
+	FMovieSceneEventPtrs EventToTrigger;
 
 private:
 
diff --git a/Engine/Source/Runtime/MovieSceneTracks/Private/Sections/MovieSceneEventRepeaterSection.cpp b/Engine/Source/Runtime/MovieSceneTracks/Private/Sections/MovieSceneEventRepeaterSection.cpp
--- a/Engine/Source/Runtime/MovieSceneTracks/Private/Sections/MovieSceneEventRepeaterSection.cpp
+++ b/Engine/Source/Runtime/MovieSceneTracks/Private/Sections/MovieSceneEventRepeaterSection.cpp
@@ -1,20 +1,3 @@
 // Copyright 1998-2019 Epic Games, Inc. All Rights Reserved.
 
 #include "Sections/MovieSceneEventRepeaterSection.h"
-
-#if WITH_EDITORONLY_DATA
-
-void UMovieSceneEventRepeaterSection::OnBlueprintRecompiled(UBlueprint*)
-{
-	FName OldFunctionName = Event.FunctionName;
-
-	Event.CacheFunctionName();
-
-	if (Event.FunctionName != OldFunctionName)
-	{
-		MarkAsChanged();
-		MarkPackageDirty();
-	}
-}
-
-#endif		// WITH_EDITORONLY_DATA
\ No newline at end of file
diff --git a/Engine/Source/Runtime/MovieSceneTracks/Private/Sections/MovieSceneEventSectionBase.cpp b/Engine/Source/Runtime/MovieSceneTracks/Private/Sections/MovieSceneEventSectionBase.cpp
--- a/Engine/Source/Runtime/MovieSceneTracks/Private/Sections/MovieSceneEventSectionBase.cpp
+++ b/Engine/Source/Runtime/MovieSceneTracks/Private/Sections/MovieSceneEventSectionBase.cpp
@@ -1,40 +1,83 @@
 // Copyright 1998-2019 Epic Games, Inc. All Rights Reserved.
 
 #include "Sections/MovieSceneEventSectionBase.h"
+#include "Modules/ModuleManager.h"
+
+#if WITH_EDITOR
+
 #include "Engine/Blueprint.h"
+#include "Kismet2/BlueprintEditorUtils.h"
 
-#if WITH_EDITORONLY_DATA
+UMovieSceneEventSectionBase::FGenerateEventEntryPointFunctionsEvent UMovieSceneEventSectionBase::GenerateEventEntryPointsEvent;
+UMovieSceneEventSectionBase::FFixupPayloadParameterNameEvent UMovieSceneEventSectionBase::FixupPayloadParameterNameEvent;
+UMovieSceneEventSectionBase::FUpgradeLegacyEventEndpoint UMovieSceneEventSectionBase::UpgradeLegacyEventEndpoint;
 
-void UMovieSceneEventSectionBase::Serialize(FArchive& Ar)
+
+void UMovieSceneEventSectionBase::OnPostCompile(UBlueprint* Blueprint)
 {
-	Super::Serialize(Ar);
+	if (Blueprint->GeneratedClass)
+	{
+		for (FMovieSceneEvent& EntryPoint : GetAllEntryPoints())
+		{
+			if (EntryPoint.CompiledFunctionName != NAME_None)
+			{
+				// @todo: Validate that the function is good
+				EntryPoint.Ptrs.Function = Blueprint->GeneratedClass->FindFunctionByName(EntryPoint.CompiledFunctionName);
 
+				if (EntryPoint.Ptrs.Function && EntryPoint.BoundObjectPinName != NAME_None)
+				{
+					EntryPoint.Ptrs.BoundObjectProperty = EntryPoint.Ptrs.Function->FindPropertyByName(EntryPoint.BoundObjectPinName);
+					check(!EntryPoint.Ptrs.BoundObjectProperty || EntryPoint.Ptrs.BoundObjectProperty->GetOuter() == EntryPoint.Ptrs.Function);
+					if (Cast<UObjectProperty>(EntryPoint.Ptrs.BoundObjectProperty) || Cast<UInterfaceProperty>(EntryPoint.Ptrs.BoundObjectProperty))
+					{
+					}
+				}
+				else
+				{
+					EntryPoint.Ptrs.BoundObjectProperty = nullptr;
+				}
+			}
+			else
+			{
+				EntryPoint.Ptrs.Function = nullptr;
+				EntryPoint.Ptrs.BoundObjectProperty = nullptr;
+			}
 
-	if (Ar.IsLoading())
-	{
-		if (UBlueprint* Blueprint = DirectorBlueprint.Get())
+			EntryPoint.CompiledFunctionName = NAME_None;
+		}
+
+		if (!Blueprint->bIsRegeneratingOnLoad)
 		{
-			OnBlueprintCompiledHandle = Blueprint->OnCompiled().AddUObject(this, &UMovieSceneEventSectionBase::OnBlueprintRecompiled);
+			MarkAsChanged();
+			MarkPackageDirty();
 		}
 	}
+
+	Blueprint->OnCompiled().RemoveAll(this);
 }
 
-void UMovieSceneEventSectionBase::SetDirectorBlueprint(UBlueprint* InBlueprint)
+#endif
+
+void UMovieSceneEventSectionBase::PostLoad()
 {
-	if (DirectorBlueprint != InBlueprint)
+	Super::PostLoad();
+
+#if WITH_EDITORONLY_DATA
+	if (UBlueprint* Blueprint = DirectorBlueprint_DEPRECATED.Get())
 	{
-		UBlueprint* ExistingBlueprint = DirectorBlueprint.Get();
-		if (ExistingBlueprint)
-		{
-			ExistingBlueprint->OnCompiled().Remove(OnBlueprintCompiledHandle);
-		}
+		UpgradeLegacyEventEndpoint.Broadcast(this, Blueprint);
 
-		DirectorBlueprint = InBlueprint;
-		if (InBlueprint)
+		Blueprint->GenerateFunctionGraphsEvent.AddUniqueDynamic(this, &UMovieSceneEventSectionBase::HandleGenerateEntryPoints);
+
+		// If the BP has already been compiled (eg regenerate on load) we must perform PostCompile fixup immediately since
+		// We will not have had a chance to generate function entries. In this case we just bind directly to the already compiled functions.
+		if (Blueprint->bHasBeenRegenerated)
 		{
-			OnBlueprintCompiledHandle = InBlueprint->OnCompiled().AddUObject(this, &UMovieSceneEventSectionBase::OnBlueprintRecompiled);
+			OnPostCompile(Blueprint);
 		}
-	}
-}
 
-#endif		// WITH_EDITORONLY_DATA
\ No newline at end of file
+		// We're done with data upgrade now
+		DirectorBlueprint_DEPRECATED = nullptr;
+	}
+#endif
+}
\ No newline at end of file
diff --git a/Engine/Source/Runtime/MovieSceneTracks/Private/Sections/MovieSceneEventTriggerSection.cpp b/Engine/Source/Runtime/MovieSceneTracks/Private/Sections/MovieSceneEventTriggerSection.cpp
--- a/Engine/Source/Runtime/MovieSceneTracks/Private/Sections/MovieSceneEventTriggerSection.cpp
+++ b/Engine/Source/Runtime/MovieSceneTracks/Private/Sections/MovieSceneEventTriggerSection.cpp
@@ -21,31 +21,3 @@ UMovieSceneEventTriggerSection::UMovieSceneEventTriggerSection(const FObjectInit
 
 #endif
 }
-
-
-#if WITH_EDITORONLY_DATA
-
-void UMovieSceneEventTriggerSection::OnBlueprintRecompiled(UBlueprint* InBlueprint)
-{
-	bool bHasChanged = false;
-
-	for (FMovieSceneEvent& Event : EventChannel.GetData().GetValues())
-	{
-		FName OldFunctionName = Event.FunctionName;
-
-		Event.CacheFunctionName();
-
-		if (Event.FunctionName != OldFunctionName)
-		{
-			bHasChanged = true;
-		}
-	}
-
-	if (bHasChanged)
-	{
-		MarkAsChanged();
-		MarkPackageDirty();
-	}
-}
-
-#endif		// WITH_EDITORONLY_DATA
\ No newline at end of file
diff --git a/Engine/Source/Runtime/MovieSceneTracks/Public/Channels/MovieSceneEvent.h b/Engine/Source/Runtime/MovieSceneTracks/Public/Channels/MovieSceneEvent.h
--- a/Engine/Source/Runtime/MovieSceneTracks/Public/Channels/MovieSceneEvent.h
+++ b/Engine/Source/Runtime/MovieSceneTracks/Public/Channels/MovieSceneEvent.h
@@ -3,123 +3,85 @@
 #pragma once
 
 #include "CoreTypes.h"
+#include "Engine/Blueprint.h"
 #include "MovieSceneEvent.generated.h"
 
+class UK2Node;
 class UEdGraph;
 class UBlueprint;
+class UEdGraphNode;
 class UK2Node_FunctionEntry;
 class UMovieSceneEventSectionBase;
 
+/** Value definition for any type-agnostic variable (exported as text) */
+USTRUCT()
+struct FMovieSceneEventPayloadVariable
+{
+	GENERATED_BODY()
 
-/**
- * Struct type that is bound to a blueprint function entry node, and resolved to a cached UFunction when the blueprint/track is recompiled.
- * Events can be bound to either of the following function signatures:
- *
- *   1. A function with no parameters (and no return value).
- *     - Compatible with master tracks or object bindings
- *     - No context passed through to event
-
- *   2. A function with a single object or interface parameter (and no return value).
- *     - Compatible with master tracks or object bindings
- *     - Will be triggered with objects in the following order:
- *         - Objects bound to the track's object binding, or:
- *         - Objects specified on the track's event receivers array, or:
- *         - Objects provided by the playback context.
- *     - Will only trigger if the specified object is of the same type as the parameter (or interface)
- */
+	UPROPERTY()
+	FString Value;
+};
+
+/** Compiled reflection pointers for the event function and parameters */
 USTRUCT(BlueprintType)
-struct MOVIESCENETRACKS_API FMovieSceneEvent
+struct FMovieSceneEventPtrs
 {
 	GENERATED_BODY()
 
+	FMovieSceneEventPtrs()
+		: Function(nullptr)
+		, BoundObjectProperty(nullptr)
+	{}
 
-	/**
-	 * Called after this event has been serialized in order to cache the function pointer if necessary
-	 */
-	void PostSerialize(const FArchive& Ar);
-
-	/**
-	 * Called to perform custom serialization logic for this struct.
-	 */
-	bool Serialize(FArchive& Ar);
+	UPROPERTY()
+	UFunction* Function;
 
-	/**
-	 * Check whether the specified function is valid for a movie scene event
-	 * Functions must have either no parameters, or a single, pass-by-value object/interface parameter, with no return parameter.
-	 */
-	static bool IsValidFunction(UFunction* Function);
+	UPROPERTY()
+	UProperty* BoundObjectProperty;
+};
 
+USTRUCT(BlueprintType)
+struct FMovieSceneEvent
+{
+	GENERATED_BODY()
 
 	/**
 	 * The function that should be called to invoke this event.
 	 * Functions must have either no parameters, or a single, pass-by-value object/interface parameter, with no return parameter.
 	 */
-	UPROPERTY(EditDefaultsOnly, Category="Sequencer")
-	FName FunctionName;
+	UPROPERTY()
+	FMovieSceneEventPtrs Ptrs;
 
-#if WITH_EDITORONLY_DATA
+	void PostSerialize(const FArchive& Ar);
 
 public:
 
+#if WITH_EDITORONLY_DATA
 
-	/**
-	 * Cache the function name to call from the blueprint function entry node. Will only cache the function if it has a valid signature.
-	 */
-	void CacheFunctionName();
-
-
-	/**
-	 * Check whether this event is bound to a valid blueprint entry node
-	 *
-	 * @return true if this event is bound to a function entry node with a valid signature, false otherwise.
-	 */
-	bool IsBoundToBlueprint() const;
-
-
-	/**
-	 * Helper function to determine whether the specified function entry is valid for this event
-	 *
-	 * @param Node         The node to test
-	 * @return true if the function entry node is compatible with a moviescene event, false otherwise
-	 */
-	static bool IsValidFunction(UK2Node_FunctionEntry* Node);
-
-
-	/**
-	 * Retrieve the function entry node this event is bound to
-	 *
-	 * @note Events may be bound to invalid function entries if they have been changed since they were assigned.
-	 * @see SetFunctionEntry, IsValidFunction
-	 * @return The function entry node if still available, nullptr if it has been destroyed, or was never assigned.
-	 */
-	UK2Node_FunctionEntry* GetFunctionEntry() const;
-
-
-	/**
-	 * Set the function entry that this event should trigger
-	 *
-	 * @param Entry        The graph node to bind to
-	 */
-	void SetFunctionEntry(UK2Node_FunctionEntry* Entry);
-
-private:
+	/** Array of payload variables to be added to the generated function */
+	UPROPERTY()
+	TMap<FName, FMovieSceneEventPayloadVariable> PayloadVariables;
 
-	/**
-	* Validates that all Graphs in the blueprint are loaded.
-	*/
-	bool IsReadyToGetFunctionEntry();
+	UPROPERTY(transient)
+	FName CompiledFunctionName;
 
-	/** Serialized soft pointer to the blueprint that contains the function graph endpoint for this event. Stored as a soft path so that renames of the blueprint don't break this event binding. */
 	UPROPERTY()
-	TSoftObjectPtr<UBlueprint> SoftBlueprintPath;
+	FName BoundObjectPinName;
 
-	/** The UEdGraph::GraphGuid property that relates to the function entry to call. */
+	/** The UEdGraph::GraphGuid property that relates the graph within which our endpoint lives. */
 	UPROPERTY()
 	FGuid GraphGuid;
 
+	/** When valid, relates to the The UEdGraphNode::NodeGuid for a custom event node that defines our event endpoint. When invalid, we must be bound to a UBlueprint::FunctionGraphs graph. */
+	UPROPERTY()
+	FGuid NodeGuid;
+
 	/** Non-serialized weak pointer to the function entry within the blueprint graph for this event. Stored as an editor-only UObject so UHT can parse it when building for non-editor. */
 	UPROPERTY(transient)
-	mutable TWeakObjectPtr<UObject> CachedFunctionEntry;
+	TWeakObjectPtr<UObject> WeakCachedEndpoint;
+
+private:
 
 	/** Deprecated weak pointer to the function entry to call - no longer serialized but cached on load. */
 	UPROPERTY()
@@ -129,8 +91,14 @@ private:
 };
 
 template<>
-struct TStructOpsTypeTraits<FMovieSceneEvent> : TStructOpsTypeTraitsBase2<FMovieSceneEvent>
+struct TStructOpsTypeTraits<FMovieSceneEvent> : public TStructOpsTypeTraitsBase2<FMovieSceneEvent>
 {
-	enum { WithSerializer = true, WithPostSerialize = true };
+	enum
+	{
+		WithPostSerialize = true,
+	};
 };
 
+
+
+
diff --git a/Engine/Source/Runtime/MovieSceneTracks/Public/Sections/MovieSceneEventRepeaterSection.h b/Engine/Source/Runtime/MovieSceneTracks/Public/Sections/MovieSceneEventRepeaterSection.h
--- a/Engine/Source/Runtime/MovieSceneTracks/Public/Sections/MovieSceneEventRepeaterSection.h
+++ b/Engine/Source/Runtime/MovieSceneTracks/Public/Sections/MovieSceneEventRepeaterSection.h
@@ -18,7 +18,7 @@ public:
 	GENERATED_BODY()
 
 #if WITH_EDITORONLY_DATA
-	virtual void OnBlueprintRecompiled(UBlueprint*) override;
+	virtual TArrayView<FMovieSceneEvent> GetAllEntryPoints() override { return MakeArrayView(&Event, 1); }
 #endif
 
 	/** The event that should be triggered each time this section is evaluated */
diff --git a/Engine/Source/Runtime/MovieSceneTracks/Public/Sections/MovieSceneEventSectionBase.h b/Engine/Source/Runtime/MovieSceneTracks/Public/Sections/MovieSceneEventSectionBase.h
--- a/Engine/Source/Runtime/MovieSceneTracks/Public/Sections/MovieSceneEventSectionBase.h
+++ b/Engine/Source/Runtime/MovieSceneTracks/Public/Sections/MovieSceneEventSectionBase.h
@@ -4,9 +4,10 @@
 
 #include "CoreTypes.h"
 #include "MovieSceneSection.h"
+#include "Channels/MovieSceneEvent.h"
+#include "Engine/Blueprint.h"
 #include "MovieSceneEventSectionBase.generated.h"
 
-
 /**
  * Base class for all event sections. Manages dirtying the section and track on recompilation of the director blueprint.
  */
@@ -17,26 +18,64 @@ class UMovieSceneEventSectionBase
 public:
 	GENERATED_BODY()
 
-#if WITH_EDITORONLY_DATA
+	virtual TArrayView<FMovieSceneEvent> GetAllEntryPoints() { return TArrayView<FMovieSceneEvent>(); }
+
+	virtual void PostLoad() override;
+
+#if WITH_EDITOR
+
+	DECLARE_MULTICAST_DELEGATE_TwoParams(FGenerateEventEntryPointFunctionsEvent, UMovieSceneEventSectionBase*, const FGenerateBlueprintFunctionParams&);
+	DECLARE_MULTICAST_DELEGATE_FourParams(FFixupPayloadParameterNameEvent, UMovieSceneEventSectionBase*, UK2Node*, FName, FName);
+	DECLARE_MULTICAST_DELEGATE_TwoParams(FUpgradeLegacyEventEndpoint, UMovieSceneEventSectionBase*, UBlueprint*);
+
 
 	/**
-	 * Assign the director blueprint pointer for this event section.
+	 * Event handler that should be bound to a sequence director's UBlueprint::GenerateFunctionGraphsEvent member in order to generate event entrypoints
+	 * This function is declared as a UFunction to ensure that it gets persistently serialized along with the UBlueprint
 	 */
-	MOVIESCENETRACKS_API void SetDirectorBlueprint(UBlueprint* InBlueprint);
+	UFUNCTION()
+	void HandleGenerateEntryPoints(const FGenerateBlueprintFunctionParams& Params)
+	{
+		GenerateEventEntryPointsEvent.Broadcast(this, Params);
+	}
 
-protected:
+	/**
+	 * Handler should be invoked when an event endpoint that is referenced from this section has one of its pins renamed
+	 */
+	void OnUserDefinedPinRenamed(UK2Node* InNode, FName OldPinName, FName NewPinName)
+	{
+		FixupPayloadParameterNameEvent.Broadcast(this, InNode, OldPinName, NewPinName);
+	}
 
-	virtual void Serialize(FArchive& Ar) override;
-	virtual void OnBlueprintRecompiled(UBlueprint*) {}
 
-private:
+	/**
+	 * Post compilation handler that is invoked once generated function graphs have been compiled. Fixes up UFunction pointers for each event.
+	 */
+	MOVIESCENETRACKS_API void OnPostCompile(UBlueprint* Blueprint);
 
-	/** A handle to the blueprint recompiled delegate binding for this section's director blueprint */
-	UPROPERTY()
-	TWeakObjectPtr<UBlueprint> DirectorBlueprint;
 
-	/** A handle to the blueprint recompiled delegate binding for this section's director blueprint */
-	FDelegateHandle OnBlueprintCompiledHandle;
+	/**
+	 * Event that is broadcast when events need to be generated for a function. Implemented in this way so that editor-code can be kept within editor modules.
+	 */
+	MOVIESCENETRACKS_API static FGenerateEventEntryPointFunctionsEvent GenerateEventEntryPointsEvent;
+
+	/**
+	 * Event that is broadcast when event payloads may need fixing up due to a pin rename
+	 */
+	MOVIESCENETRACKS_API static FFixupPayloadParameterNameEvent FixupPayloadParameterNameEvent;
+
+	/**
+	 * Event that is broadcast when a legacy event section is found that needs fixing up against an already compile-on-loaded blueprint
+	 */
+	MOVIESCENETRACKS_API static FUpgradeLegacyEventEndpoint UpgradeLegacyEventEndpoint;
+
+#endif
+
+private:
 
-#endif // WITH_EDITORONLY_DATA
+#if WITH_EDITORONLY_DATA
+	/** Legacy pointer to the sequence director BP */
+	UPROPERTY()
+	TWeakObjectPtr<UBlueprint> DirectorBlueprint_DEPRECATED;
+#endif
 };
\ No newline at end of file
diff --git a/Engine/Source/Runtime/MovieSceneTracks/Public/Sections/MovieSceneEventTriggerSection.h b/Engine/Source/Runtime/MovieSceneTracks/Public/Sections/MovieSceneEventTriggerSection.h
--- a/Engine/Source/Runtime/MovieSceneTracks/Public/Sections/MovieSceneEventTriggerSection.h
+++ b/Engine/Source/Runtime/MovieSceneTracks/Public/Sections/MovieSceneEventTriggerSection.h
@@ -22,7 +22,7 @@ public:
 
 #if WITH_EDITORONLY_DATA
 
-	virtual void OnBlueprintRecompiled(UBlueprint*) override;
+	virtual TArrayView<FMovieSceneEvent> GetAllEntryPoints() override { return EventChannel.GetData().GetValues(); }
 
 #endif
 
