Sequencer: Address static analysis warning

diff --git a/Engine/Source/Editor/MovieSceneTools/Private/MovieSceneEventCustomization.cpp b/Engine/Source/Editor/MovieSceneTools/Private/MovieSceneEventCustomization.cpp
--- a/Engine/Source/Editor/MovieSceneTools/Private/MovieSceneEventCustomization.cpp
+++ b/Engine/Source/Editor/MovieSceneTools/Private/MovieSceneEventCustomization.cpp
@@ -372,7 +372,7 @@ void FMovieSceneEventCustomization::CustomizeChildren(TSharedRef<IPropertyHandle
 
 			for (UProperty* Field : TFieldRange<UProperty>(CommonFunction))
 			{
-				if (Field->HasAnyPropertyFlags(CPF_OutParm | CPF_ReturnParm | CPF_OutParm | CPF_ReferenceParm) || Field->GetFName() == EntryPoint->BoundObjectPinName)
+				if (Field->HasAnyPropertyFlags(CPF_OutParm | CPF_ReturnParm | CPF_ReferenceParm) || Field->GetFName() == EntryPoint->BoundObjectPinName)
 				{
 					continue;
 				}
