Removed explicit sort predicates for TSortedMap now that the default is correct for T* key types

diff --git a/Engine/Source/Editor/MovieSceneTools/Private/MovieSceneEventCustomization.cpp b/Engine/Source/Editor/MovieSceneTools/Private/MovieSceneEventCustomization.cpp
--- a/Engine/Source/Editor/MovieSceneTools/Private/MovieSceneEventCustomization.cpp
+++ b/Engine/Source/Editor/MovieSceneTools/Private/MovieSceneEventCustomization.cpp
@@ -1292,8 +1292,6 @@ void FMovieSceneEventCustomization::SetEventEndpoint(UK2Node* NewEndpoint, UEdGr
 
 	if (Blueprint)
 	{
-		//FBlueprintEditorUtils::MarkBlueprintAsModified(Blueprint);
-		//FBlueprintEditorUtils::MarkBlueprintAsModified(Blueprint);
 		FKismetEditorUtilities::CompileBlueprint(Blueprint);
 	}
 
@@ -1304,21 +1302,20 @@ void FMovieSceneEventCustomization::SetEventEndpoint(UK2Node* NewEndpoint, UEdGr
 
 void FMovieSceneEventCustomization::CreateEventEndpoint()
 {
-	// @todo: These TLess<const T*> parameters should not be necessary but TSortedMap has a poor default that uses T rather than ConstPointerType
 	struct FSectionData
 	{
 		TArray<FMovieSceneEvent*, TInlineAllocator<1>>  EntryPoints;
 	};
 	struct FTrackData
 	{
-		TSortedMap<UMovieSceneEventSectionBase*, FSectionData, FDefaultAllocator, TLess<const UMovieSceneEventSectionBase*>>  Sections;
+		TSortedMap<UMovieSceneEventSectionBase*, FSectionData>  Sections;
 	};
 	struct FSequenceData
 	{
-		TSortedMap<UMovieSceneEventTrack*, FTrackData, FDefaultAllocator, TLess<const UMovieSceneEventTrack*>> Tracks;
+		TSortedMap<UMovieSceneEventTrack*, FTrackData> Tracks;
 	};
 
-	TSortedMap<UMovieSceneSequence*, FSequenceData, FDefaultAllocator, TLess<const UMovieSceneSequence*>> PerSequenceData;
+	TSortedMap<UMovieSceneSequence*, FSequenceData> PerSequenceData;
 
 	// Populate all the sequences represented by this customization
 	{
