Sequencer: Fixed module load issues when running editor -game

diff --git a/Engine/Source/Editor/CurveEditor/Private/CurveEditorModule.cpp b/Engine/Source/Editor/CurveEditor/Private/CurveEditorModule.cpp
--- a/Engine/Source/Editor/CurveEditor/Private/CurveEditorModule.cpp
+++ b/Engine/Source/Editor/CurveEditor/Private/CurveEditorModule.cpp
@@ -12,7 +12,11 @@ class FCurveEditorModule : public ICurveEditorModule
 public:
 	virtual void StartupModule() override
 	{
-		FCurveEditorCommands::Register();
+		if (GIsEditor)
+		{
+			FModuleManager::Get().LoadModule("EditorStyle");
+			FCurveEditorCommands::Register();
+		}
 	}
 
 	virtual void ShutdownModule() override
diff --git a/Engine/Source/Editor/MovieSceneTools/Private/MovieSceneToolsModule.cpp b/Engine/Source/Editor/MovieSceneTools/Private/MovieSceneToolsModule.cpp
--- a/Engine/Source/Editor/MovieSceneTools/Private/MovieSceneToolsModule.cpp
+++ b/Engine/Source/Editor/MovieSceneTools/Private/MovieSceneToolsModule.cpp
@@ -75,82 +75,85 @@
 
 void FMovieSceneToolsModule::StartupModule()
 {
-	if (ISettingsModule* SettingsModule = FModuleManager::GetModulePtr<ISettingsModule>("Settings"))
+	if (GIsEditor)
 	{
-		SettingsModule->RegisterSettings("Project", "Editor", "Level Sequences",
-			LOCTEXT("RuntimeSettingsName", "Level Sequences"),
-			LOCTEXT("RuntimeSettingsDescription", "Configure project settings relating to Level Sequences"),
-			GetMutableDefault<UMovieSceneToolsProjectSettings>()
-		);
-	}
-
-	ISequencerModule& SequencerModule = FModuleManager::Get().LoadModuleChecked<ISequencerModule>( "Sequencer" );
-
-	// register property track editors
-	BoolPropertyTrackCreateEditorHandle = SequencerModule.RegisterPropertyTrackEditor<FBoolPropertyTrackEditor>();
-	BytePropertyTrackCreateEditorHandle = SequencerModule.RegisterPropertyTrackEditor<FBytePropertyTrackEditor>();
-	ColorPropertyTrackCreateEditorHandle = SequencerModule.RegisterPropertyTrackEditor<FColorPropertyTrackEditor>();
-	FloatPropertyTrackCreateEditorHandle = SequencerModule.RegisterPropertyTrackEditor<FFloatPropertyTrackEditor>();
-	IntegerPropertyTrackCreateEditorHandle = SequencerModule.RegisterPropertyTrackEditor<FIntegerPropertyTrackEditor>();
-	VectorPropertyTrackCreateEditorHandle = SequencerModule.RegisterPropertyTrackEditor<FVectorPropertyTrackEditor>();
-	TransformPropertyTrackCreateEditorHandle = SequencerModule.RegisterPropertyTrackEditor<FTransformPropertyTrackEditor>();
-	EulerTransformPropertyTrackCreateEditorHandle = SequencerModule.RegisterPropertyTrackEditor<FEulerTransformPropertyTrackEditor>();
-	VisibilityPropertyTrackCreateEditorHandle = SequencerModule.RegisterPropertyTrackEditor<FVisibilityPropertyTrackEditor>();
-	ActorReferencePropertyTrackCreateEditorHandle = SequencerModule.RegisterPropertyTrackEditor<FActorReferencePropertyTrackEditor>();
-	StringPropertyTrackCreateEditorHandle = SequencerModule.RegisterPropertyTrackEditor<FStringPropertyTrackEditor>();
-	ObjectTrackCreateEditorHandle = SequencerModule.RegisterPropertyTrackEditor<FObjectPropertyTrackEditor>();
-
-	// register specialty track editors
-	AnimationTrackCreateEditorHandle = SequencerModule.RegisterTrackEditor( FOnCreateTrackEditor::CreateStatic( &FSkeletalAnimationTrackEditor::CreateTrackEditor ) );
-	AttachTrackCreateEditorHandle = SequencerModule.RegisterTrackEditor( FOnCreateTrackEditor::CreateStatic( &F3DAttachTrackEditor::CreateTrackEditor ) );
-	AudioTrackCreateEditorHandle = SequencerModule.RegisterTrackEditor( FOnCreateTrackEditor::CreateStatic( &FAudioTrackEditor::CreateTrackEditor ) );
-	EventTrackCreateEditorHandle = SequencerModule.RegisterTrackEditor( FOnCreateTrackEditor::CreateStatic( &FEventTrackEditor::CreateTrackEditor ) );
-	ParticleTrackCreateEditorHandle = SequencerModule.RegisterTrackEditor( FOnCreateTrackEditor::CreateStatic( &FParticleTrackEditor::CreateTrackEditor ) );
-	ParticleParameterTrackCreateEditorHandle = SequencerModule.RegisterTrackEditor( FOnCreateTrackEditor::CreateStatic( &FParticleParameterTrackEditor::CreateTrackEditor ) );
-	PathTrackCreateEditorHandle = SequencerModule.RegisterTrackEditor( FOnCreateTrackEditor::CreateStatic( &F3DPathTrackEditor::CreateTrackEditor ) );
-	CameraCutTrackCreateEditorHandle = SequencerModule.RegisterTrackEditor( FOnCreateTrackEditor::CreateStatic( &FCameraCutTrackEditor::CreateTrackEditor ) );
-	CinematicShotTrackCreateEditorHandle = SequencerModule.RegisterTrackEditor( FOnCreateTrackEditor::CreateStatic( &FCinematicShotTrackEditor::CreateTrackEditor ) );
-	SlomoTrackCreateEditorHandle = SequencerModule.RegisterTrackEditor( FOnCreateTrackEditor::CreateStatic( &FSlomoTrackEditor::CreateTrackEditor ) );
-	SubTrackCreateEditorHandle = SequencerModule.RegisterTrackEditor( FOnCreateTrackEditor::CreateStatic( &FSubTrackEditor::CreateTrackEditor ) );
-	TransformTrackCreateEditorHandle = SequencerModule.RegisterTrackEditor( FOnCreateTrackEditor::CreateStatic( &F3DTransformTrackEditor::CreateTrackEditor ) );
-	ComponentMaterialTrackCreateEditorHandle = SequencerModule.RegisterTrackEditor( FOnCreateTrackEditor::CreateStatic( &FComponentMaterialTrackEditor::CreateTrackEditor ) );
-	FadeTrackCreateEditorHandle = SequencerModule.RegisterTrackEditor( FOnCreateTrackEditor::CreateStatic( &FFadeTrackEditor::CreateTrackEditor ) );
-	SpawnTrackCreateEditorHandle = SequencerModule.RegisterTrackEditor( FOnCreateTrackEditor::CreateStatic( &FSpawnTrackEditor::CreateTrackEditor ) );
-	LevelVisibilityTrackCreateEditorHandle = SequencerModule.RegisterTrackEditor( FOnCreateTrackEditor::CreateStatic( &FLevelVisibilityTrackEditor::CreateTrackEditor ) );
-	CameraAnimTrackCreateEditorHandle = SequencerModule.RegisterTrackEditor(FOnCreateTrackEditor::CreateStatic(&FCameraAnimTrackEditor::CreateTrackEditor));
-	CameraShakeTrackCreateEditorHandle = SequencerModule.RegisterTrackEditor(FOnCreateTrackEditor::CreateStatic(&FCameraShakeTrackEditor::CreateTrackEditor));
-	MPCTrackCreateEditorHandle = SequencerModule.RegisterTrackEditor(FOnCreateTrackEditor::CreateStatic(&FMaterialParameterCollectionTrackEditor::CreateTrackEditor));
-	PrimitiveMaterialCreateEditorHandle = SequencerModule.RegisterTrackEditor(FOnCreateTrackEditor::CreateStatic(&FPrimitiveMaterialTrackEditor::CreateTrackEditor));
-
-	RegisterClipboardConversions();
-
-	// register details customization
-	FPropertyEditorModule& PropertyModule = FModuleManager::LoadModuleChecked<FPropertyEditorModule>("PropertyEditor");
-	PropertyModule.RegisterCustomClassLayout("MovieSceneToolsProjectSettings", FOnGetDetailCustomizationInstance::CreateStatic(&FMovieSceneToolsProjectSettingsCustomization::MakeInstance));
-	PropertyModule.RegisterCustomClassLayout("MovieSceneBuiltInEasingFunction", FOnGetDetailCustomizationInstance::CreateLambda(&MakeShared<FMovieSceneBuiltInEasingFunctionCustomization>));
-	PropertyModule.RegisterCustomPropertyTypeLayout("MovieSceneObjectBindingID", FOnGetPropertyTypeCustomizationInstance::CreateLambda(&MakeShared<FMovieSceneObjectBindingIDCustomization>));
-	PropertyModule.RegisterCustomPropertyTypeLayout("MovieSceneEvent", FOnGetPropertyTypeCustomizationInstance::CreateStatic(&FMovieSceneEventCustomization::MakeInstance));
-
-	SequencerModule.RegisterChannelInterface<FMovieSceneBoolChannel>();
-	SequencerModule.RegisterChannelInterface<FMovieSceneByteChannel>();
-	SequencerModule.RegisterChannelInterface<FMovieSceneIntegerChannel>();
-	SequencerModule.RegisterChannelInterface<FMovieSceneFloatChannel>();
-	SequencerModule.RegisterChannelInterface<FMovieSceneStringChannel>();
-	SequencerModule.RegisterChannelInterface<FMovieSceneParticleChannel>();
-	SequencerModule.RegisterChannelInterface<FMovieSceneActorReferenceData>();
-	SequencerModule.RegisterChannelInterface<FMovieSceneEventSectionData>();
-	SequencerModule.RegisterChannelInterface<FMovieSceneObjectPathChannel>();
-
-	SequencerModule.RegisterChannelInterface<FMovieSceneEventChannel>();
-
-	ICurveEditorModule& CurveEditorModule = FModuleManager::LoadModuleChecked<ICurveEditorModule>("CurveEditor");
-
-	FEventChannelCurveModel::EventView = CurveEditorModule.RegisterView(FOnCreateCurveEditorView::CreateStatic(
-		[](TWeakPtr<FCurveEditor> WeakCurveEditor) -> TSharedRef<SCurveEditorView>
+		if (ISettingsModule* SettingsModule = FModuleManager::GetModulePtr<ISettingsModule>("Settings"))
 		{
-			return SNew(SCurveEditorEventChannelView, WeakCurveEditor);
+			SettingsModule->RegisterSettings("Project", "Editor", "Level Sequences",
+				LOCTEXT("RuntimeSettingsName", "Level Sequences"),
+				LOCTEXT("RuntimeSettingsDescription", "Configure project settings relating to Level Sequences"),
+				GetMutableDefault<UMovieSceneToolsProjectSettings>()
+			);
 		}
-	));
+
+		ISequencerModule& SequencerModule = FModuleManager::Get().LoadModuleChecked<ISequencerModule>( "Sequencer" );
+
+		// register property track editors
+		BoolPropertyTrackCreateEditorHandle = SequencerModule.RegisterPropertyTrackEditor<FBoolPropertyTrackEditor>();
+		BytePropertyTrackCreateEditorHandle = SequencerModule.RegisterPropertyTrackEditor<FBytePropertyTrackEditor>();
+		ColorPropertyTrackCreateEditorHandle = SequencerModule.RegisterPropertyTrackEditor<FColorPropertyTrackEditor>();
+		FloatPropertyTrackCreateEditorHandle = SequencerModule.RegisterPropertyTrackEditor<FFloatPropertyTrackEditor>();
+		IntegerPropertyTrackCreateEditorHandle = SequencerModule.RegisterPropertyTrackEditor<FIntegerPropertyTrackEditor>();
+		VectorPropertyTrackCreateEditorHandle = SequencerModule.RegisterPropertyTrackEditor<FVectorPropertyTrackEditor>();
+		TransformPropertyTrackCreateEditorHandle = SequencerModule.RegisterPropertyTrackEditor<FTransformPropertyTrackEditor>();
+		EulerTransformPropertyTrackCreateEditorHandle = SequencerModule.RegisterPropertyTrackEditor<FEulerTransformPropertyTrackEditor>();
+		VisibilityPropertyTrackCreateEditorHandle = SequencerModule.RegisterPropertyTrackEditor<FVisibilityPropertyTrackEditor>();
+		ActorReferencePropertyTrackCreateEditorHandle = SequencerModule.RegisterPropertyTrackEditor<FActorReferencePropertyTrackEditor>();
+		StringPropertyTrackCreateEditorHandle = SequencerModule.RegisterPropertyTrackEditor<FStringPropertyTrackEditor>();
+		ObjectTrackCreateEditorHandle = SequencerModule.RegisterPropertyTrackEditor<FObjectPropertyTrackEditor>();
+
+		// register specialty track editors
+		AnimationTrackCreateEditorHandle = SequencerModule.RegisterTrackEditor( FOnCreateTrackEditor::CreateStatic( &FSkeletalAnimationTrackEditor::CreateTrackEditor ) );
+		AttachTrackCreateEditorHandle = SequencerModule.RegisterTrackEditor( FOnCreateTrackEditor::CreateStatic( &F3DAttachTrackEditor::CreateTrackEditor ) );
+		AudioTrackCreateEditorHandle = SequencerModule.RegisterTrackEditor( FOnCreateTrackEditor::CreateStatic( &FAudioTrackEditor::CreateTrackEditor ) );
+		EventTrackCreateEditorHandle = SequencerModule.RegisterTrackEditor( FOnCreateTrackEditor::CreateStatic( &FEventTrackEditor::CreateTrackEditor ) );
+		ParticleTrackCreateEditorHandle = SequencerModule.RegisterTrackEditor( FOnCreateTrackEditor::CreateStatic( &FParticleTrackEditor::CreateTrackEditor ) );
+		ParticleParameterTrackCreateEditorHandle = SequencerModule.RegisterTrackEditor( FOnCreateTrackEditor::CreateStatic( &FParticleParameterTrackEditor::CreateTrackEditor ) );
+		PathTrackCreateEditorHandle = SequencerModule.RegisterTrackEditor( FOnCreateTrackEditor::CreateStatic( &F3DPathTrackEditor::CreateTrackEditor ) );
+		CameraCutTrackCreateEditorHandle = SequencerModule.RegisterTrackEditor( FOnCreateTrackEditor::CreateStatic( &FCameraCutTrackEditor::CreateTrackEditor ) );
+		CinematicShotTrackCreateEditorHandle = SequencerModule.RegisterTrackEditor( FOnCreateTrackEditor::CreateStatic( &FCinematicShotTrackEditor::CreateTrackEditor ) );
+		SlomoTrackCreateEditorHandle = SequencerModule.RegisterTrackEditor( FOnCreateTrackEditor::CreateStatic( &FSlomoTrackEditor::CreateTrackEditor ) );
+		SubTrackCreateEditorHandle = SequencerModule.RegisterTrackEditor( FOnCreateTrackEditor::CreateStatic( &FSubTrackEditor::CreateTrackEditor ) );
+		TransformTrackCreateEditorHandle = SequencerModule.RegisterTrackEditor( FOnCreateTrackEditor::CreateStatic( &F3DTransformTrackEditor::CreateTrackEditor ) );
+		ComponentMaterialTrackCreateEditorHandle = SequencerModule.RegisterTrackEditor( FOnCreateTrackEditor::CreateStatic( &FComponentMaterialTrackEditor::CreateTrackEditor ) );
+		FadeTrackCreateEditorHandle = SequencerModule.RegisterTrackEditor( FOnCreateTrackEditor::CreateStatic( &FFadeTrackEditor::CreateTrackEditor ) );
+		SpawnTrackCreateEditorHandle = SequencerModule.RegisterTrackEditor( FOnCreateTrackEditor::CreateStatic( &FSpawnTrackEditor::CreateTrackEditor ) );
+		LevelVisibilityTrackCreateEditorHandle = SequencerModule.RegisterTrackEditor( FOnCreateTrackEditor::CreateStatic( &FLevelVisibilityTrackEditor::CreateTrackEditor ) );
+		CameraAnimTrackCreateEditorHandle = SequencerModule.RegisterTrackEditor(FOnCreateTrackEditor::CreateStatic(&FCameraAnimTrackEditor::CreateTrackEditor));
+		CameraShakeTrackCreateEditorHandle = SequencerModule.RegisterTrackEditor(FOnCreateTrackEditor::CreateStatic(&FCameraShakeTrackEditor::CreateTrackEditor));
+		MPCTrackCreateEditorHandle = SequencerModule.RegisterTrackEditor(FOnCreateTrackEditor::CreateStatic(&FMaterialParameterCollectionTrackEditor::CreateTrackEditor));
+		PrimitiveMaterialCreateEditorHandle = SequencerModule.RegisterTrackEditor(FOnCreateTrackEditor::CreateStatic(&FPrimitiveMaterialTrackEditor::CreateTrackEditor));
+
+		RegisterClipboardConversions();
+
+		// register details customization
+		FPropertyEditorModule& PropertyModule = FModuleManager::LoadModuleChecked<FPropertyEditorModule>("PropertyEditor");
+		PropertyModule.RegisterCustomClassLayout("MovieSceneToolsProjectSettings", FOnGetDetailCustomizationInstance::CreateStatic(&FMovieSceneToolsProjectSettingsCustomization::MakeInstance));
+		PropertyModule.RegisterCustomClassLayout("MovieSceneBuiltInEasingFunction", FOnGetDetailCustomizationInstance::CreateLambda(&MakeShared<FMovieSceneBuiltInEasingFunctionCustomization>));
+		PropertyModule.RegisterCustomPropertyTypeLayout("MovieSceneObjectBindingID", FOnGetPropertyTypeCustomizationInstance::CreateLambda(&MakeShared<FMovieSceneObjectBindingIDCustomization>));
+		PropertyModule.RegisterCustomPropertyTypeLayout("MovieSceneEvent", FOnGetPropertyTypeCustomizationInstance::CreateStatic(&FMovieSceneEventCustomization::MakeInstance));
+
+		SequencerModule.RegisterChannelInterface<FMovieSceneBoolChannel>();
+		SequencerModule.RegisterChannelInterface<FMovieSceneByteChannel>();
+		SequencerModule.RegisterChannelInterface<FMovieSceneIntegerChannel>();
+		SequencerModule.RegisterChannelInterface<FMovieSceneFloatChannel>();
+		SequencerModule.RegisterChannelInterface<FMovieSceneStringChannel>();
+		SequencerModule.RegisterChannelInterface<FMovieSceneParticleChannel>();
+		SequencerModule.RegisterChannelInterface<FMovieSceneActorReferenceData>();
+		SequencerModule.RegisterChannelInterface<FMovieSceneEventSectionData>();
+		SequencerModule.RegisterChannelInterface<FMovieSceneObjectPathChannel>();
+
+		SequencerModule.RegisterChannelInterface<FMovieSceneEventChannel>();
+
+		ICurveEditorModule& CurveEditorModule = FModuleManager::LoadModuleChecked<ICurveEditorModule>("CurveEditor");
+
+		FEventChannelCurveModel::EventView = CurveEditorModule.RegisterView(FOnCreateCurveEditorView::CreateStatic(
+			[](TWeakPtr<FCurveEditor> WeakCurveEditor) -> TSharedRef<SCurveEditorView>
+			{
+				return SNew(SCurveEditorEventChannelView, WeakCurveEditor);
+			}
+		));
+	}
 
 	GenerateEventEntryPointsHandle = UMovieSceneEventSectionBase::GenerateEventEntryPointsEvent.AddStatic(HandleGenerateEventEntryPoints);
 	FixupPayloadParameterNameHandle = UMovieSceneEventSectionBase::FixupPayloadParameterNameEvent.AddStatic(FixupPayloadParameterNameForSection);
