Sequencer: Improve sub-sequence track and cinematic shot track to show looping sections.

diff --git a/Engine/Source/Editor/MovieSceneTools/Private/Sections/CinematicShotSection.cpp b/Engine/Source/Editor/MovieSceneTools/Private/Sections/CinematicShotSection.cpp
--- a/Engine/Source/Editor/MovieSceneTools/Private/Sections/CinematicShotSection.cpp
+++ b/Engine/Source/Editor/MovieSceneTools/Private/Sections/CinematicShotSection.cpp
@@ -303,16 +303,13 @@ int32 FCinematicShotSection::OnPaintSection(FSequencerSectionPainter& InPainter)
 		InPainter.bParentEnabled ? ESlateDrawEffect::None : ESlateDrawEffect::DisabledEffect
 	);
 
-	TRange<FFrameNumber> SectionRange = SectionObject.GetRange();
+	const TRange<FFrameNumber> SectionRange = SectionObject.GetRange();
 	if (SectionRange.GetLowerBound().IsOpen() || SectionRange.GetUpperBound().IsOpen())
 	{
 		return InPainter.LayerId;
 	}
 
-	const FFrameNumber SectionStartFrame = SectionObject.GetInclusiveStartFrame();
-	const FFrameNumber SectionEndFrame   = SectionObject.GetExclusiveEndFrame();
-	const int32        SectionSize       = MovieScene::DiscreteSize(SectionRange);
-
+	const int32 SectionSize = MovieScene::DiscreteSize(SectionRange);
 	if (SectionSize <= 0)
 	{
 		return InPainter.LayerId;
@@ -326,12 +323,47 @@ int32 FCinematicShotSection::OnPaintSection(FSequencerSectionPainter& InPainter)
 		return InPainter.LayerId;
 	}
 
+	if (SectionObject.Parameters.bCanLoop)
+	{
+		DoPaintLoopingSection(*InnerSequence, InPainter);
+	}
+	else
+	{
+		DoPaintNonLoopingSection(*InnerSequence, InPainter);
+	}
+
+	TSharedPtr<ISequencer> Sequencer = SequencerPtr.Pin();
+
+	if (InPainter.bIsSelected && Sequencer.IsValid())
+	{
+		FFrameTime CurrentTime = Sequencer->GetLocalTime().Time;
+		if (SectionRange.Contains(CurrentTime.FrameNumber))
+		{
+			UMovieScene* SubSequenceMovieScene = InnerSequence->GetMovieScene();
+			const FFrameRate DisplayRate = SubSequenceMovieScene->GetDisplayRate();
+			const FFrameRate TickResolution = SubSequenceMovieScene->GetTickResolution();
+			const FFrameNumber CurrentFrameNumber = ConvertFrameTime(CurrentTime * SectionObject.OuterToInnerTransform(), TickResolution, DisplayRate).FloorToFrame();
+
+			DrawFrameNumberHint(InPainter, CurrentTime, CurrentFrameNumber.Value);
+		}
+	}
+
+	return InPainter.LayerId;
+}
+
+void FCinematicShotSection::DoPaintNonLoopingSection(UMovieSceneSequence& InnerSequence, FSequencerSectionPainter& InPainter) const
+{
+	const TRange<FFrameNumber> SectionRange = SectionObject.GetRange();
+	const FFrameNumber SectionStartFrame = SectionObject.GetInclusiveStartFrame();
+	const FFrameNumber SectionEndFrame = SectionObject.GetExclusiveEndFrame();
+	const int32 SectionSize = MovieScene::DiscreteSize(SectionRange);
+
 	const float PixelsPerFrame = InPainter.SectionGeometry.Size.X / float(SectionSize);
 
-	UMovieScene*         MovieScene    = InnerSequence->GetMovieScene();
+	UMovieScene* MovieScene = InnerSequence.GetMovieScene();
 	TRange<FFrameNumber> PlaybackRange = MovieScene->GetPlaybackRange();
 
-	FMovieSceneSequenceTransform InnerToOuterTransform = SectionObject.OuterToInnerTransform().Inverse();
+	FMovieSceneSequenceTransform InnerToOuterTransform = SectionObject.OuterToInnerTransform().InverseLinearOnly();
 
 	const FFrameNumber PlaybackStart = (MovieScene::DiscreteInclusiveLower(PlaybackRange) * InnerToOuterTransform).FloorToFrame();
 	if (SectionRange.Contains(PlaybackStart))
@@ -395,24 +427,64 @@ int32 FCinematicShotSection::OnPaintSection(FSequencerSectionPainter& InPainter)
 			FColor(128, 32, 32)	// 0, 75, 50 (HSV)
 		);
 	}
+}
 
-	TSharedPtr<ISequencer> Sequencer = SequencerPtr.Pin();
+void FCinematicShotSection::DoPaintLoopingSection(UMovieSceneSequence& InnerSequence, FSequencerSectionPainter& InPainter) const
+{
+	const FTimeToPixel& TimeToPixelConverter = InPainter.GetTimeConverter();
 
-	if (InPainter.bIsSelected && Sequencer.IsValid())
+	const FFrameNumber SectionStartFrame = SectionObject.GetInclusiveStartFrame();
+	const FFrameNumber SectionEndFrame = SectionObject.GetExclusiveEndFrame();
+	const float InvTimeScale = FMath::IsNearlyZero(SectionObject.Parameters.TimeScale) ? 1.0f : 1.0f / SectionObject.Parameters.TimeScale;
+
+	UMovieScene* MovieScene = InnerSequence.GetMovieScene();
+	TRange<FFrameNumber> InnerPlaybackRange = UMovieSceneSubSection::GetValidatedInnerPlaybackRange(SectionObject.Parameters, *MovieScene);
+
+	const FFrameNumber ShotLength = InnerPlaybackRange.Size<FFrameNumber>() * InvTimeScale;
+	const FFrameNumber FirstLoopShotLength = ShotLength - SectionObject.Parameters.FirstLoopStartFrameOffset * InvTimeScale;
+
+	FFrameNumber CurOffsetFrame = FMath::Max(FirstLoopShotLength, FFrameNumber(0));
+	const ESlateDrawEffect DrawEffects = InPainter.bParentEnabled ? ESlateDrawEffect::None : ESlateDrawEffect::DisabledEffect;
+
+	// Draw separators where the shot is looping. To be consistent with the non-looping case, we draw a red and green
+	// separator back to back.
+	uint32 MaxLoopBoundaries = 100;
+	while (CurOffsetFrame < SectionEndFrame)
 	{
-		FFrameTime CurrentTime = Sequencer->GetLocalTime().Time;
-		if (SectionRange.Contains(CurrentTime.FrameNumber))
-		{
-			UMovieScene* SubSequenceMovieScene = SectionObject.GetSequence()->GetMovieScene();
-			const FFrameRate DisplayRate = SubSequenceMovieScene->GetDisplayRate();
-			const FFrameRate TickResolution = SubSequenceMovieScene->GetTickResolution();
-			const FFrameNumber CurrentFrameNumber = ConvertFrameTime(CurrentTime * SectionObject.OuterToInnerTransform(), TickResolution, DisplayRate).FloorToFrame();
+		const float CurOffsetPixel = TimeToPixelConverter.FrameToPixel(SectionStartFrame + CurOffsetFrame);
 
-			DrawFrameNumberHint(InPainter, CurrentTime, CurrentFrameNumber.Value);
+		FSlateDrawElement::MakeBox(
+			InPainter.DrawElements,
+			InPainter.LayerId++,
+			InPainter.SectionGeometry.MakeChild(
+				FVector2D(1.f, InPainter.SectionGeometry.Size.Y - 2.f),
+				FSlateLayoutTransform(FVector2D(CurOffsetPixel, 1.f))
+			).ToPaintGeometry(),
+			FEditorStyle::GetBrush("WhiteBrush"),
+			DrawEffects,
+			FColor(32, 128, 32)	// 120, 75, 50 (HSV)
+		);
+		if (CurOffsetFrame > 0) {
+			FSlateDrawElement::MakeBox(
+				InPainter.DrawElements,
+				InPainter.LayerId++,
+				InPainter.SectionGeometry.MakeChild(
+					FVector2D(1.f, InPainter.SectionGeometry.Size.Y - 2.f),
+					FSlateLayoutTransform(FVector2D(CurOffsetPixel - 1.f, 1.f))
+				).ToPaintGeometry(),
+				FEditorStyle::GetBrush("WhiteBrush"),
+				DrawEffects,
+				FColor(128, 32, 32)	// 0, 75, 50 (HSV)
+			);
 		}
-	}
 
-	return InPainter.LayerId;
+		CurOffsetFrame += ShotLength;
+
+		if ((--MaxLoopBoundaries) == 0)
+		{
+			break;
+		}
+	}
 }
 
 void FCinematicShotSection::BuildSectionContextMenu(FMenuBuilder& MenuBuilder, const FGuid& ObjectBinding)
diff --git a/Engine/Source/Editor/MovieSceneTools/Private/Sections/CinematicShotSection.h b/Engine/Source/Editor/MovieSceneTools/Private/Sections/CinematicShotSection.h
--- a/Engine/Source/Editor/MovieSceneTools/Private/Sections/CinematicShotSection.h
+++ b/Engine/Source/Editor/MovieSceneTools/Private/Sections/CinematicShotSection.h
@@ -55,6 +55,9 @@ private:
 	/** Add shot takes menu */
 	void AddTakesMenu(FMenuBuilder& MenuBuilder);
 
+	void DoPaintNonLoopingSection(UMovieSceneSequence& InnerSequence, FSequencerSectionPainter& InPainter) const;
+	void DoPaintLoopingSection(UMovieSceneSequence& InnerSequence, FSequencerSectionPainter& InPainter) const;
+
 private:
 
 	/** The section we are visualizing */
diff --git a/Engine/Source/Editor/MovieSceneTools/Private/TrackEditors/SubTrackEditor.cpp b/Engine/Source/Editor/MovieSceneTools/Private/TrackEditors/SubTrackEditor.cpp
--- a/Engine/Source/Editor/MovieSceneTools/Private/TrackEditors/SubTrackEditor.cpp
+++ b/Engine/Source/Editor/MovieSceneTools/Private/TrackEditors/SubTrackEditor.cpp
@@ -115,100 +115,35 @@
 	{
 		int32 LayerId = InPainter.PaintSectionBackground();
 
-		ESlateDrawEffect DrawEffects = InPainter.bParentEnabled ? ESlateDrawEffect::None : ESlateDrawEffect::DisabledEffect;
-		
-		
-		TRange<FFrameNumber> SectionRange = SectionObject.GetRange();
+		const TRange<FFrameNumber> SectionRange = SectionObject.GetRange();
 		if (SectionRange.GetLowerBound().IsOpen() || SectionRange.GetUpperBound().IsOpen())
 		{
 			return InPainter.LayerId;
 		}
 
-		const FFrameNumber SectionStartFrame = SectionObject.GetInclusiveStartFrame();
-		const FFrameNumber SectionEndFrame   = SectionObject.GetExclusiveEndFrame();
-		const int32        SectionSize       = MovieScene::DiscreteSize(SectionRange);
-
+		const int32 SectionSize = MovieScene::DiscreteSize(SectionRange);
 		if (SectionSize <= 0)
 		{
 			return InPainter.LayerId;
 		}
 
-		const float        PixelsPerFrame    = InPainter.SectionGeometry.Size.X / float(SectionSize);
-
 		UMovieSceneSequence* InnerSequence = SectionObject.GetSequence();
+		const ESlateDrawEffect DrawEffects = InPainter.bParentEnabled ? ESlateDrawEffect::None : ESlateDrawEffect::DisabledEffect;
+		
 		if (InnerSequence)
 		{
-			UMovieScene*         MovieScene    = InnerSequence->GetMovieScene();
-			TRange<FFrameNumber> PlaybackRange = MovieScene->GetPlaybackRange();
-
-			FMovieSceneSequenceTransform InnerToOuterTransform = SectionObject.OuterToInnerTransform().Inverse();
-
-			const FFrameNumber PlaybackStart = (MovieScene::DiscreteInclusiveLower(PlaybackRange) * InnerToOuterTransform).FloorToFrame();
-			if (SectionRange.Contains(PlaybackStart))
+			if (SectionObject.Parameters.bCanLoop)
 			{
-				const int32 StartOffset = (PlaybackStart-SectionStartFrame).Value;
-				// add dark tint for left out-of-bounds
-				FSlateDrawElement::MakeBox(
-					InPainter.DrawElements,
-					InPainter.LayerId++,
-					InPainter.SectionGeometry.ToPaintGeometry(
-						FVector2D(0.0f, 0.f),
-						FVector2D(StartOffset * PixelsPerFrame, InPainter.SectionGeometry.Size.Y)
-					),
-					FEditorStyle::GetBrush("WhiteBrush"),
-					DrawEffects,
-					FLinearColor::Black.CopyWithNewOpacity(0.5f)
-				);
-
-				// add green line for playback start
-				FSlateDrawElement::MakeBox(
-					InPainter.DrawElements,
-					InPainter.LayerId++,
-					InPainter.SectionGeometry.ToPaintGeometry(
-						FVector2D(StartOffset * PixelsPerFrame, 0.f),
-						FVector2D(1.0f, InPainter.SectionGeometry.Size.Y)
-					),
-					FEditorStyle::GetBrush("WhiteBrush"),
-					DrawEffects,
-					FColor(32, 128, 32)	// 120, 75, 50 (HSV)
-				);
+				DoPaintLoopingSection(*InnerSequence, InPainter, DrawEffects);
+			}
+			else
+			{
+				DoPaintNonLoopingSection(*InnerSequence, InPainter, DrawEffects);
 			}
 
-			const FFrameNumber PlaybackEnd = (MovieScene::DiscreteExclusiveUpper(PlaybackRange) * InnerToOuterTransform).FloorToFrame();
-			if (SectionRange.Contains(PlaybackEnd))
-			{
-				// add dark tint for right out-of-bounds
-				const int32 EndOffset = (PlaybackEnd-SectionStartFrame).Value;
-				FSlateDrawElement::MakeBox(
-					InPainter.DrawElements,
-					InPainter.LayerId++,
-					InPainter.SectionGeometry.ToPaintGeometry(
-						FVector2D(EndOffset * PixelsPerFrame, 0.f),
-						FVector2D((SectionSize - EndOffset) * PixelsPerFrame, InPainter.SectionGeometry.Size.Y)
-					),
-					FEditorStyle::GetBrush("WhiteBrush"),
-					DrawEffects,
-					FLinearColor::Black.CopyWithNewOpacity(0.5f)
-				);
-
-
-				// add red line for playback end
-				FSlateDrawElement::MakeBox(
-					InPainter.DrawElements,
-					InPainter.LayerId++,
-					InPainter.SectionGeometry.ToPaintGeometry(
-						FVector2D(EndOffset * PixelsPerFrame, 0.f),
-						FVector2D(1.0f, InPainter.SectionGeometry.Size.Y)
-					),
-					FEditorStyle::GetBrush("WhiteBrush"),
-					DrawEffects,
-					FColor(128, 32, 32)	// 0, 75, 50 (HSV)
-				);
-			}
-
 			FMargin ContentPadding = GetContentPadding();
+			UMovieScene* MovieScene = InnerSequence->GetMovieScene();
+			const int32 NumTracks = MovieScene->GetPossessableCount() + MovieScene->GetSpawnableCount() + MovieScene->GetMasterTracks().Num();
 
-			int32 NumTracks = MovieScene->GetPossessableCount() + MovieScene->GetSpawnableCount() + MovieScene->GetMasterTracks().Num();
-
 			FVector2D TopLeft = InPainter.SectionGeometry.AbsoluteToLocal(InPainter.SectionClippingRect.GetTopLeft()) + FVector2D(1.f, -1.f);
 
@@ -238,7 +173,6 @@
 				FColor(200, 200, 200)
 			);
 
-
 			TSharedPtr<ISequencer> SequencerPtr = Sequencer.Pin();
 			if (InPainter.bIsSelected && SequencerPtr.IsValid())
 			{
@@ -296,6 +230,143 @@
 		return LayerId;
 	}
 
+	void DoPaintNonLoopingSection(UMovieSceneSequence& InnerSequence, FSequencerSectionPainter& InPainter, ESlateDrawEffect DrawEffects) const
+	{
+		const FFrameNumber SectionStartFrame = SectionObject.GetInclusiveStartFrame();
+
+		const TRange<FFrameNumber> SectionRange = SectionObject.GetRange();
+		const int32 SectionSize = MovieScene::DiscreteSize(SectionRange);
+		const float PixelsPerFrame = InPainter.SectionGeometry.Size.X / float(SectionSize);
+
+		UMovieScene* MovieScene = InnerSequence.GetMovieScene();
+		TRange<FFrameNumber> PlaybackRange = MovieScene->GetPlaybackRange();
+
+		// We're in the non-looping case so we know we have a purely linear transform.
+		const FMovieSceneSequenceTransform InnerToOuterTransform = SectionObject.OuterToInnerTransform().InverseLinearOnly();
+		const FFrameNumber PlaybackStart = (MovieScene::DiscreteInclusiveLower(PlaybackRange) * InnerToOuterTransform).FloorToFrame();
+		if (SectionRange.Contains(PlaybackStart))
+		{
+			const int32 StartOffset = (PlaybackStart - SectionStartFrame).Value;
+			// add dark tint for left out-of-bounds
+			FSlateDrawElement::MakeBox(
+					InPainter.DrawElements,
+					InPainter.LayerId++,
+					InPainter.SectionGeometry.ToPaintGeometry(
+						FVector2D(0.0f, 0.f),
+						FVector2D(StartOffset * PixelsPerFrame, InPainter.SectionGeometry.Size.Y)
+						),
+					FEditorStyle::GetBrush("WhiteBrush"),
+					DrawEffects,
+					FLinearColor::Black.CopyWithNewOpacity(0.5f)
+					);
+
+			// add green line for playback start
+			FSlateDrawElement::MakeBox(
+					InPainter.DrawElements,
+					InPainter.LayerId++,
+					InPainter.SectionGeometry.ToPaintGeometry(
+						FVector2D(StartOffset * PixelsPerFrame, 0.f),
+						FVector2D(1.0f, InPainter.SectionGeometry.Size.Y)
+						),
+					FEditorStyle::GetBrush("WhiteBrush"),
+					DrawEffects,
+					FColor(32, 128, 32)	// 120, 75, 50 (HSV)
+					);
+		}
+
+		const FFrameNumber PlaybackEnd = (MovieScene::DiscreteExclusiveUpper(PlaybackRange) * InnerToOuterTransform) .FloorToFrame();
+		if (SectionRange.Contains(PlaybackEnd))
+		{
+			// add dark tint for right out-of-bounds
+			const int32 EndOffset = (PlaybackEnd - SectionStartFrame).Value;
+			FSlateDrawElement::MakeBox(
+					InPainter.DrawElements,
+					InPainter.LayerId++,
+					InPainter.SectionGeometry.ToPaintGeometry(
+						FVector2D(EndOffset * PixelsPerFrame, 0.f),
+						FVector2D((SectionSize - EndOffset) * PixelsPerFrame, InPainter.SectionGeometry.Size.Y)
+						),
+					FEditorStyle::GetBrush("WhiteBrush"),
+					DrawEffects,
+					FLinearColor::Black.CopyWithNewOpacity(0.5f)
+					);
+
+
+			// add red line for playback end
+			FSlateDrawElement::MakeBox(
+					InPainter.DrawElements,
+					InPainter.LayerId++,
+					InPainter.SectionGeometry.ToPaintGeometry(
+						FVector2D(EndOffset * PixelsPerFrame, 0.f),
+						FVector2D(1.0f, InPainter.SectionGeometry.Size.Y)
+						),
+					FEditorStyle::GetBrush("WhiteBrush"),
+					DrawEffects,
+					FColor(128, 32, 32)	// 0, 75, 50 (HSV)
+					);
+		}
+	}
+	
+	void DoPaintLoopingSection(UMovieSceneSequence& InnerSequence, FSequencerSectionPainter& InPainter, ESlateDrawEffect DrawEffects) const
+	{
+		const FFrameNumber SectionStartFrame = SectionObject.GetInclusiveStartFrame();
+		const FFrameNumber SectionEndFrame   = SectionObject.GetExclusiveEndFrame();
+
+		const TRange<FFrameNumber> SectionRange = SectionObject.GetRange();
+		const int32 SectionSize = MovieScene::DiscreteSize(SectionRange);
+		const float PixelsPerFrame = InPainter.SectionGeometry.Size.X / float(SectionSize);
+
+		const float InvTimeScale = FMath::IsNearlyZero(SectionObject.Parameters.TimeScale) ? 1.0f : 1.0f / SectionObject.Parameters.TimeScale;
+
+		UMovieScene* MovieScene = InnerSequence.GetMovieScene();
+		TRange<FFrameNumber> InnerPlaybackRange = UMovieSceneSubSection::GetValidatedInnerPlaybackRange(SectionObject.Parameters, *MovieScene);
+
+		const FFrameNumber SubSeqLength = InnerPlaybackRange.Size<FFrameNumber>() * InvTimeScale;
+		const FFrameNumber FirstLoopSubSeqLength = SubSeqLength - SectionObject.Parameters.FirstLoopStartFrameOffset * InvTimeScale;
+
+		FFrameNumber CurOffsetFrame = FMath::Max(FirstLoopSubSeqLength, FFrameNumber(0));
+
+		// Draw separators where the sub-sequence is looping. To be consistent with the non-looping case, we draw a red and green
+		// separator back to back.
+		uint32 MaxLoopBoundaries = 100;
+		while (CurOffsetFrame < SectionEndFrame)
+		{
+			const int32 CurOffset = CurOffsetFrame.Value;
+			FSlateDrawElement::MakeBox(
+				InPainter.DrawElements,
+				InPainter.LayerId++,
+				InPainter.SectionGeometry.ToPaintGeometry(
+					FVector2D(CurOffset * PixelsPerFrame, 0.f),
+					FVector2D(1.0f, InPainter.SectionGeometry.Size.Y)
+				),
+				FEditorStyle::GetBrush("WhiteBrush"),
+				DrawEffects,
+				FColor(32, 128, 32)	// 120, 75, 50 (HSV)
+			);
+			if (CurOffset > 0)
+			{
+				FSlateDrawElement::MakeBox(
+					InPainter.DrawElements,
+					InPainter.LayerId++,
+					InPainter.SectionGeometry.ToPaintGeometry(
+						FVector2D(CurOffset * PixelsPerFrame - 1.f, 0.f),
+						FVector2D(1.0f, InPainter.SectionGeometry.Size.Y)
+					),
+					FEditorStyle::GetBrush("WhiteBrush"),
+					DrawEffects,
+					FColor(128, 32, 32)	// 0, 75, 50 (HSV)
+				);
+			}
+
+			CurOffsetFrame += SubSeqLength;
+
+			if ((--MaxLoopBoundaries) == 0)
+			{
+				break;
+			}
+		}
+	}
+
 	virtual FReply OnSectionDoubleClicked(const FGeometry& SectionGeometry, const FPointerEvent& MouseEvent) override
 	{
 		if( MouseEvent.GetEffectingButton() == EKeys::LeftMouseButton )
