Moved SEnumCombobox to the EditorWidgets module and removed the two separate implementations (Niagara and Sequencer)

diff --git a/Engine/Plugins/FX/Niagara/Source/NiagaraEditor/NiagaraEditor.Build.cs b/Engine/Plugins/FX/Niagara/Source/NiagaraEditor/NiagaraEditor.Build.cs
--- a/Engine/Plugins/FX/Niagara/Source/NiagaraEditor/NiagaraEditor.Build.cs
+++ b/Engine/Plugins/FX/Niagara/Source/NiagaraEditor/NiagaraEditor.Build.cs
@@ -48,6 +48,7 @@ public class NiagaraEditor : ModuleRules
 				"Projects",
                 "MainFrame",
 				"ToolMenus",
+				"EditorWidgets",
 			}
         );
 
diff --git a/Engine/Plugins/FX/Niagara/Source/NiagaraEditor/Private/TypeEditorUtilities/NiagaraEnumTypeEditorUtilities.cpp b/Engine/Plugins/FX/Niagara/Source/NiagaraEditor/Private/TypeEditorUtilities/NiagaraEnumTypeEditorUtilities.cpp
--- a/Engine/Plugins/FX/Niagara/Source/NiagaraEditor/Private/TypeEditorUtilities/NiagaraEnumTypeEditorUtilities.cpp
+++ b/Engine/Plugins/FX/Niagara/Source/NiagaraEditor/Private/TypeEditorUtilities/NiagaraEnumTypeEditorUtilities.cpp
@@ -5,104 +5,9 @@
 #include "NiagaraTypes.h"
 #include "NiagaraEditorStyle.h"
 #include "NiagaraEditorCommon.h"
+#include "SEnumCombobox.h"
 
 #include "Widgets/Input/SSpinBox.h"
-#include "Widgets/Input/SComboBox.h"
-
-class SEnumCombobox : public SComboBox<TSharedPtr<int32>>
-{
-public:
-	DECLARE_DELEGATE_TwoParams(FOnValueChanged, int32, ESelectInfo::Type);
-
-public:
-	SLATE_BEGIN_ARGS(SEnumCombobox) {}
-		SLATE_ATTRIBUTE(int32, Value)
-		SLATE_EVENT(FOnValueChanged, OnValueChanged)
-	SLATE_END_ARGS()
-
-	void Construct(const FArguments& InArgs, const UEnum* InEnum)
-	{
-		Enum = InEnum;
-		Value = InArgs._Value;
-		check(Value.IsBound());
-		OnValueChanged = InArgs._OnValueChanged;
-
-		bUpdatingSelectionInternally = false;
-
-		for (int32 i = 0; i < Enum->NumEnums() - 1; i++)
-		{
-			if (Enum->HasMetaData(TEXT("Hidden"), i) == false)
-			{
-				VisibleEnumNameIndices.Add(MakeShareable(new int32(i)));
-			}
-		}
-
-		SComboBox::Construct(SComboBox<TSharedPtr<int32>>::FArguments()
-			//.ButtonStyle(FEditorStyle::Get(), "FlatButton.Light")
-			.OptionsSource(&VisibleEnumNameIndices)
-			.OnGenerateWidget(this, &SEnumCombobox::OnGenerateWidget)
-			.OnSelectionChanged(this, &SEnumCombobox::OnComboSelectionChanged)
-			.OnComboBoxOpening(this, &SEnumCombobox::OnComboMenuOpening)
-			.ContentPadding(FMargin(2, 0))
-			[
-				SNew(STextBlock)
-				//.Font(FEditorStyle::GetFontStyle("Sequencer.AnimationOutliner.RegularFont"))
-				.Text(this, &SEnumCombobox::GetValueText)
-			]);
-	}
-
-private:
-	FText GetValueText() const
-	{
-		int32 ValueNameIndex = Enum->GetIndexByValue(Value.Get());
-		return Enum->GetDisplayNameTextByIndex(ValueNameIndex);
-	}
-
-	TSharedRef<SWidget> OnGenerateWidget(TSharedPtr<int32> InItem)
-	{
-		return SNew(STextBlock)
-			.Text(Enum->GetDisplayNameTextByIndex(*InItem));
-	}
-
-	void OnComboSelectionChanged(TSharedPtr<int32> InSelectedItem, ESelectInfo::Type SelectInfo)
-	{
-		if (bUpdatingSelectionInternally == false)
-		{
-			OnValueChanged.ExecuteIfBound(*InSelectedItem, SelectInfo);
-		}
-	}
-
-	void OnComboMenuOpening()
-	{
-		int32 CurrentNameIndex = Enum->GetIndexByValue(Value.Get());
-		TSharedPtr<int32> FoundNameIndexItem;
-		for (int32 i = 0; i < VisibleEnumNameIndices.Num(); i++)
-		{
-			if (*VisibleEnumNameIndices[i] == CurrentNameIndex)
-			{
-				FoundNameIndexItem = VisibleEnumNameIndices[i];
-				break;
-			}
-		}
-		if (FoundNameIndexItem.IsValid())
-		{
-			bUpdatingSelectionInternally = true;
-			SetSelectedItem(FoundNameIndexItem);
-			bUpdatingSelectionInternally = false;
-		}
-	}
-
-private:
-	const UEnum* Enum;
-
-	TAttribute<int32> Value;
-
-	TArray<TSharedPtr<int32>> VisibleEnumNameIndices;
-
-	bool bUpdatingSelectionInternally;
-
-	FOnValueChanged OnValueChanged;
-};
 
 class SNiagaraEnumParameterEditor : public SNiagaraParameterEditor
 {
@@ -119,8 +24,8 @@ public:
 		ChildSlot
 		[
 			SNew(SEnumCombobox, Enum)
-			.Value(this, &SNiagaraEnumParameterEditor::GetValue)
-			.OnValueChanged(this, &SNiagaraEnumParameterEditor::ValueChanged)
+			.CurrentValue(this, &SNiagaraEnumParameterEditor::GetValue) 
+			.OnEnumSelectionChanged(SEnumCombobox::FOnEnumSelectionChanged::CreateSP(this, &SNiagaraEnumParameterEditor::ValueChanged))
 		];
 	}
 
diff --git a/Engine/Source/Editor/EditorWidgets/Public/SEnumCombobox.h b/Engine/Source/Editor/EditorWidgets/Public/SEnumCombobox.h
new file mode 100644
--- /dev/null
+++ b/Engine/Source/Editor/EditorWidgets/Public/SEnumCombobox.h
@@ -0,0 +1,111 @@
+// Copyright 1998-2019 Epic Games, Inc. All Rights Reserved.
+
+#pragma once
+
+#include "CoreMinimal.h"
+#include "Widgets/Input/SComboBox.h"
+#include "Types/SlateEnums.h"
+#include "EditorStyleSet.h"
+
+/*
+* Editor only widget for creating a combo-box which allows the user to pick the values from a dropdown combo-box. 
+*/
+class SEnumCombobox : public SComboBox<TSharedPtr<int32>>
+{
+public:
+	DECLARE_DELEGATE_TwoParams(FOnEnumSelectionChanged, int32 /*Selection*/, ESelectInfo::Type /*SelectionType*/);
+public:
+	SLATE_BEGIN_ARGS(SEnumCombobox) {}
+
+		SLATE_ATTRIBUTE(int32, CurrentValue)
+		SLATE_ARGUMENT(FOnEnumSelectionChanged, OnEnumSelectionChanged)
+
+	SLATE_END_ARGS()
+
+	void Construct(const FArguments& InArgs, const UEnum* InEnum)
+	{
+		Enum = InEnum;
+		CurrentValue = InArgs._CurrentValue;
+		check(CurrentValue.IsBound());
+		OnEnumSelectionChangedDelegate = InArgs._OnEnumSelectionChanged;
+
+		bUpdatingSelectionInternally = false;
+
+		for (int32 i = 0; i < Enum->NumEnums() - 1; i++)
+		{
+			if (Enum->HasMetaData(TEXT("Hidden"), i) == false)
+			{
+				VisibleEnumNameIndices.Add(MakeShareable(new int32(i)));
+			}
+		}
+
+		SComboBox::Construct(SComboBox<TSharedPtr<int32>>::FArguments()
+			.ButtonStyle(FEditorStyle::Get(), "FlatButton.Light")
+			.OptionsSource(&VisibleEnumNameIndices)
+			.OnGenerateWidget_Lambda([this](TSharedPtr<int32> InItem)
+		{
+			return SNew(STextBlock)
+				.Text(Enum->GetDisplayNameTextByIndex(*InItem));
+		})
+			.OnSelectionChanged(this, &SEnumCombobox::OnComboSelectionChanged)
+			.OnComboBoxOpening(this, &SEnumCombobox::OnComboMenuOpening)
+			.ContentPadding(FMargin(2, 0))
+			[
+				SNew(STextBlock)
+				.Font(FEditorStyle::GetFontStyle("Sequencer.AnimationOutliner.RegularFont"))
+			.Text(this, &SEnumCombobox::GetCurrentValue)
+			]);
+	}
+
+private:
+	FText GetCurrentValue() const
+	{
+		int32 CurrentNameIndex = Enum->GetIndexByValue(CurrentValue.Get());
+		return Enum->GetDisplayNameTextByIndex(CurrentNameIndex);
+	}
+
+	TSharedRef<SWidget> OnGenerateWidget(TSharedPtr<int32> InItem)
+	{
+		return SNew(STextBlock)
+			.Text(Enum->GetDisplayNameTextByIndex(*InItem));
+	}
+
+	void OnComboSelectionChanged(TSharedPtr<int32> InSelectedItem, ESelectInfo::Type SelectInfo)
+	{
+		if (bUpdatingSelectionInternally == false)
+		{
+			OnEnumSelectionChangedDelegate.ExecuteIfBound(*InSelectedItem, SelectInfo);
+		}
+	}
+
+	void OnComboMenuOpening()
+	{
+		int32 CurrentNameIndex = Enum->GetIndexByValue(CurrentValue.Get());
+		TSharedPtr<int32> FoundNameIndexItem;
+		for (int32 i = 0; i < VisibleEnumNameIndices.Num(); i++)
+		{
+			if (*VisibleEnumNameIndices[i] == CurrentNameIndex)
+			{
+				FoundNameIndexItem = VisibleEnumNameIndices[i];
+				break;
+			}
+		}
+		if (FoundNameIndexItem.IsValid())
+		{
+			bUpdatingSelectionInternally = true;
+			SetSelectedItem(FoundNameIndexItem);
+			bUpdatingSelectionInternally = false;
+		}
+	}
+
+private:
+	const UEnum* Enum;
+
+	TAttribute<int32> CurrentValue;
+
+	TArray<TSharedPtr<int32>> VisibleEnumNameIndices;
+
+	bool bUpdatingSelectionInternally;
+
+	FOnEnumSelectionChanged OnEnumSelectionChangedDelegate;
+};
\ No newline at end of file
diff --git a/Engine/Source/Editor/MovieSceneTools/MovieSceneTools.Build.cs b/Engine/Source/Editor/MovieSceneTools/MovieSceneTools.Build.cs
--- a/Engine/Source/Editor/MovieSceneTools/MovieSceneTools.Build.cs
+++ b/Engine/Source/Editor/MovieSceneTools/MovieSceneTools.Build.cs
@@ -26,9 +26,10 @@ public class MovieSceneTools : ModuleRules
                 "InputCore",
                 "MovieSceneCapture",
 				"UnrealEd",
-				"Sequencer"
-			}
-		);
+				"Sequencer",
+                "EditorWidgets",
+            }
+        );
 
 		PrivateDependencyModuleNames.AddRange(
 			new string[] {
@@ -50,7 +51,6 @@ public class MovieSceneTools : ModuleRules
 				"Slate",
 				"SlateCore",
                 "EditorStyle",
-				"EditorWidgets",
 				"PropertyEditor",
                 "MaterialEditor",
 				"RenderCore",
diff --git a/Engine/Source/Editor/MovieSceneTools/Private/CurveKeyEditors/SEnumKeyEditor.cpp b/Engine/Source/Editor/MovieSceneTools/Private/CurveKeyEditors/SEnumKeyEditor.cpp
--- a/Engine/Source/Editor/MovieSceneTools/Private/CurveKeyEditors/SEnumKeyEditor.cpp
+++ b/Engine/Source/Editor/MovieSceneTools/Private/CurveKeyEditors/SEnumKeyEditor.cpp
@@ -16,7 +16,7 @@ void SEnumCurveKeyEditor::Construct(const FArguments& InArgs, const TSequencerKe
 		MovieSceneToolHelpers::MakeEnumComboBox(
 			Enum,
 			TAttribute<int32>::Create(TAttribute<int32>::FGetter::CreateSP(this, &SEnumCurveKeyEditor::OnGetCurrentValueAsInt)),
-			FOnEnumSelectionChanged::CreateSP(this, &SEnumCurveKeyEditor::OnChangeKey)
+			SEnumCombobox::FOnEnumSelectionChanged::CreateSP(this, &SEnumCurveKeyEditor::OnChangeKey)
 		)
 	];
 }
diff --git a/Engine/Source/Editor/MovieSceneTools/Private/MovieSceneToolHelpers.cpp b/Engine/Source/Editor/MovieSceneTools/Private/MovieSceneToolHelpers.cpp
--- a/Engine/Source/Editor/MovieSceneTools/Private/MovieSceneToolHelpers.cpp
+++ b/Engine/Source/Editor/MovieSceneTools/Private/MovieSceneToolHelpers.cpp
@@ -61,7 +61,6 @@
 #include "Serialization/ObjectWriter.h"
 #include "Serialization/ObjectReader.h"
 
-
 /* MovieSceneToolHelpers
  *****************************************************************************/
 
@@ -521,105 +520,7 @@ int32 MovieSceneToolHelpers::FindAvailableRowIndex(UMovieSceneTrack* InTrack, UM
 	return InTrack->GetMaxRowIndex() + 1;
 }
 
-class SEnumCombobox : public SComboBox<TSharedPtr<int32>>
-{
-public:
-	SLATE_BEGIN_ARGS(SEnumCombobox) {}
-
-	SLATE_ATTRIBUTE(int32, CurrentValue)
-	SLATE_ARGUMENT(FOnEnumSelectionChanged, OnEnumSelectionChanged)
-
-	SLATE_END_ARGS()
-
-	void Construct(const FArguments& InArgs, const UEnum* InEnum)
-	{
-		Enum = InEnum;
-		CurrentValue = InArgs._CurrentValue;
-		check(CurrentValue.IsBound());
-		OnEnumSelectionChangedDelegate = InArgs._OnEnumSelectionChanged;
-
-		bUpdatingSelectionInternally = false;
-
-		for (int32 i = 0; i < Enum->NumEnums() - 1; i++)
-		{
-			if (Enum->HasMetaData( TEXT("Hidden"), i ) == false)
-			{
-				VisibleEnumNameIndices.Add(MakeShareable(new int32(i)));
-			}
-		}
-
-		SComboBox::Construct(SComboBox<TSharedPtr<int32>>::FArguments()
-			.ButtonStyle(FEditorStyle::Get(), "FlatButton.Light")
-			.OptionsSource(&VisibleEnumNameIndices)
-			.OnGenerateWidget_Lambda([this](TSharedPtr<int32> InItem)
-			{
-				return SNew(STextBlock)
-					.Text(Enum->GetDisplayNameTextByIndex(*InItem));
-			})
-			.OnSelectionChanged(this, &SEnumCombobox::OnComboSelectionChanged)
-			.OnComboBoxOpening(this, &SEnumCombobox::OnComboMenuOpening)
-			.ContentPadding(FMargin(2, 0))
-			[
-				SNew(STextBlock)
-				.Font(FEditorStyle::GetFontStyle("Sequencer.AnimationOutliner.RegularFont"))
-				.Text(this, &SEnumCombobox::GetCurrentValue)
-			]);
-	}
-
-private:
-	FText GetCurrentValue() const
-	{
-		int32 CurrentNameIndex = Enum->GetIndexByValue(CurrentValue.Get());
-		return Enum->GetDisplayNameTextByIndex(CurrentNameIndex);
-	}
-
-	TSharedRef<SWidget> OnGenerateWidget(TSharedPtr<int32> InItem)
-	{
-		return SNew(STextBlock)
-			.Text(Enum->GetDisplayNameTextByIndex(*InItem));
-	}
-
-	void OnComboSelectionChanged(TSharedPtr<int32> InSelectedItem, ESelectInfo::Type SelectInfo)
-	{
-		if (bUpdatingSelectionInternally == false)
-		{
-			OnEnumSelectionChangedDelegate.ExecuteIfBound(*InSelectedItem, SelectInfo);
-		}
-	}
-
-	void OnComboMenuOpening()
-	{
-		int32 CurrentNameIndex = Enum->GetIndexByValue(CurrentValue.Get());
-		TSharedPtr<int32> FoundNameIndexItem;
-		for ( int32 i = 0; i < VisibleEnumNameIndices.Num(); i++ )
-		{
-			if ( *VisibleEnumNameIndices[i] == CurrentNameIndex )
-			{
-				FoundNameIndexItem = VisibleEnumNameIndices[i];
-				break;
-			}
-		}
-		if ( FoundNameIndexItem.IsValid() )
-		{
-			bUpdatingSelectionInternally = true;
-			SetSelectedItem(FoundNameIndexItem);
-			bUpdatingSelectionInternally = false;
-		}
-	}	
-
-private:
-	const UEnum* Enum;
-
-	TAttribute<int32> CurrentValue;
-
-	TArray<TSharedPtr<int32>> VisibleEnumNameIndices;
-
-	bool bUpdatingSelectionInternally;
-
-	FOnEnumSelectionChanged OnEnumSelectionChangedDelegate;
-};
-
-TSharedRef<SWidget> MovieSceneToolHelpers::MakeEnumComboBox(const UEnum* InEnum, TAttribute<int32> InCurrentValue, FOnEnumSelectionChanged InOnSelectionChanged)
+TSharedRef<SWidget> MovieSceneToolHelpers::MakeEnumComboBox(const UEnum* InEnum, TAttribute<int32> InCurrentValue, SEnumCombobox::FOnEnumSelectionChanged InOnSelectionChanged)
 {
 	return SNew(SEnumCombobox, InEnum)
 		.CurrentValue(InCurrentValue)
diff --git a/Engine/Source/Editor/MovieSceneTools/Public/MovieSceneToolHelpers.h b/Engine/Source/Editor/MovieSceneTools/Public/MovieSceneToolHelpers.h
--- a/Engine/Source/Editor/MovieSceneTools/Public/MovieSceneToolHelpers.h
+++ b/Engine/Source/Editor/MovieSceneTools/Public/MovieSceneToolHelpers.h
@@ -14,6 +14,7 @@
 #include "Logging/TokenizedMessage.h"
 #include "MovieSceneTranslator.h"
 #include "MovieSceneCaptureSettings.h"
+#include "SEnumCombobox.h"
 
 
 class ISequencer;
@@ -47,8 +48,6 @@ struct FFBXInOutParameters
 	bool bForceFrontXAxisBackup;
 };
 
-DECLARE_DELEGATE_TwoParams(FOnEnumSelectionChanged, int32 /*Selection*/, ESelectInfo::Type /*SelectionType*/);
-
 class MOVIESCENETOOLS_API MovieSceneToolHelpers
 {
 public:
@@ -157,7 +156,7 @@ public:
 	 * @param OnSelectionChanged Delegate fired when selection is changed
 	 * @return The new widget
 	 */
-	static TSharedRef<SWidget> MakeEnumComboBox(const UEnum* Enum, TAttribute<int32> CurrentValue, FOnEnumSelectionChanged OnSelectionChanged);
+	static TSharedRef<SWidget> MakeEnumComboBox(const UEnum* Enum, TAttribute<int32> CurrentValue, SEnumCombobox::FOnEnumSelectionChanged OnSelectionChanged);
 
 
 	/**
