Fix link error in MovieSceneTools when using  /permissive- mode on MSVC2017

diff --git a/Engine/Source/Editor/MovieSceneTools/Private/TrackEditors/TransformTrackEditor.cpp b/Engine/Source/Editor/MovieSceneTools/Private/TrackEditors/TransformTrackEditor.cpp
--- a/Engine/Source/Editor/MovieSceneTools/Private/TrackEditors/TransformTrackEditor.cpp
+++ b/Engine/Source/Editor/MovieSceneTools/Private/TrackEditors/TransformTrackEditor.cpp
@@ -875,7 +875,7 @@ void F3DTransformTrackEditor::AddTransformKeys( UObject* ObjectToKey, const TOpt
 
 	auto OnKeyProperty = [=](FFrameNumber Time) -> FKeyPropertyResult
 	{
-		return this->AddKeysToObjects({ ObjectToKey }, Time, *GeneratedKeys,  KeyMode, UMovieScene3DTransformTrack::StaticClass(), TransformPropertyName, InitializeNewTrack);
+		return this->AddKeysToObjects(MakeArrayView(&ObjectToKey, 1), Time, *GeneratedKeys,  KeyMode, UMovieScene3DTransformTrack::StaticClass(), TransformPropertyName, InitializeNewTrack);
 	};
 
 	AnimatablePropertyChanged( FOnKeyProperty::CreateLambda(OnKeyProperty) );
diff --git a/Engine/Source/Editor/MovieSceneTools/Public/KeyframeTrackEditor.h b/Engine/Source/Editor/MovieSceneTools/Public/KeyframeTrackEditor.h
--- a/Engine/Source/Editor/MovieSceneTools/Public/KeyframeTrackEditor.h
+++ b/Engine/Source/Editor/MovieSceneTools/Public/KeyframeTrackEditor.h
@@ -377,14 +377,6 @@ protected:
 		return KeyPropertyResult;
 	}
 
-	FORCEINLINE FKeyPropertyResult AddKeysToObjects(
-		std::initializer_list<UObject*> ObjectsToKey, FFrameNumber KeyTime, FGeneratedTrackKeys& GeneratedKeys,
-		ESequencerKeyMode KeyMode, TSubclassOf<UMovieSceneTrack> TrackClass, FName PropertyName,
-		TFunction<void(TrackType*)> OnInitializeNewTrack)
-	{
-		return AddKeysToObjects(MakeArrayView(ObjectsToKey), KeyTime, GeneratedKeys, KeyMode, TrackClass, PropertyName, MoveTemp(OnInitializeNewTrack));
-	}
-
 private:
 
 	void ClearDefaults( UMovieSceneTrack* Track )
