Sequencer: Refactor sub-sequence sections to use common code.

This adds a couple of utility classes that provide basic logic for painting and editing (resizing/slipping) sub-sequence sections.
It also adds a boilerplate mixin class that can be used as a base class for any sequencer section.

The CinematicShot, SubTrack, and TemplateSequence sections are rewritten in terms of these new things, which deletes a whole bunch of code while also providing consistency between all 3 in the editor.

diff --git a/Engine/Plugins/MovieScene/TemplateSequence/Source/TemplateSequenceEditor/Private/TrackEditors/TemplateSequenceTrackEditor.cpp b/Engine/Plugins/MovieScene/TemplateSequence/Source/TemplateSequenceEditor/Private/TrackEditors/TemplateSequenceTrackEditor.cpp
--- a/Engine/Plugins/MovieScene/TemplateSequence/Source/TemplateSequenceEditor/Private/TrackEditors/TemplateSequenceTrackEditor.cpp
+++ b/Engine/Plugins/MovieScene/TemplateSequence/Source/TemplateSequenceEditor/Private/TrackEditors/TemplateSequenceTrackEditor.cpp
@@ -44,7 +44,7 @@ TSharedRef<ISequencerSection> FTemplateSequenceTrackEditor::MakeSectionInterface
 {
 	check( SupportsType( SectionObject.GetOuter()->GetClass() ) );
 	
-	return MakeShareable(new FTemplateSequenceSection(*CastChecked<UTemplateSequenceSection>(&SectionObject)));
+	return MakeShareable(new FTemplateSequenceSection(GetSequencer(), *CastChecked<UTemplateSequenceSection>(&SectionObject)));
 }
 
 void FTemplateSequenceTrackEditor::AddTemplateSequenceSubMenu(FMenuBuilder& MenuBuilder, TArray<FGuid> ObjectBindings)
@@ -132,36 +132,9 @@ FKeyPropertyResult FTemplateSequenceTrackEditor::AddKeyInternal(FFrameNumber Key
 	return KeyPropertyResult;
 }
 
-FTemplateSequenceSection::FTemplateSequenceSection(UTemplateSequenceSection& InSection)
-	: WeakSection(&InSection)
+FTemplateSequenceSection::FTemplateSequenceSection(TSharedPtr<ISequencer> InSequencer, UTemplateSequenceSection& InSection)
+	: TSubSectionMixin(InSequencer, InSection)
 {
 }
 
-UMovieSceneSection* FTemplateSequenceSection::GetSectionObject()
-{
-	return WeakSection.Get();
-}
-
-bool FTemplateSequenceSection::IsReadOnly() const
-{
-	return WeakSection.IsValid() ? WeakSection.Get()->IsReadOnly() : false;
-}
-
-FText FTemplateSequenceSection::GetSectionTitle() const
-{
-	if (const UTemplateSequenceSection* Section = WeakSection.Get())
-	{
-		if (const UMovieSceneSequence* TemplateSequence = Section->GetSequence())
-		{
-			return FText::FromString(TemplateSequence->GetName());
-		}
-	}
-	return LOCTEXT("NoTemplateSequenceSection", "No Template Animation");
-}
-
-int32 FTemplateSequenceSection::OnPaintSection(FSequencerSectionPainter& Painter) const
-{
-	return Painter.PaintSectionBackground();
-}
-
 #undef LOCTEXT_NAMESPACE
diff --git a/Engine/Plugins/MovieScene/TemplateSequence/Source/TemplateSequenceEditor/Private/TrackEditors/TemplateSequenceTrackEditor.h b/Engine/Plugins/MovieScene/TemplateSequence/Source/TemplateSequenceEditor/Private/TrackEditors/TemplateSequenceTrackEditor.h
--- a/Engine/Plugins/MovieScene/TemplateSequence/Source/TemplateSequenceEditor/Private/TrackEditors/TemplateSequenceTrackEditor.h
+++ b/Engine/Plugins/MovieScene/TemplateSequence/Source/TemplateSequenceEditor/Private/TrackEditors/TemplateSequenceTrackEditor.h
@@ -5,6 +5,7 @@
 #include "CoreMinimal.h"
 #include "Misc/Guid.h"
 #include "MovieSceneTrackEditor.h"
+#include "TrackEditors/SubTrackEditorBase.h"
 
 struct FAssetData;
 class FMenuBuilder;
@@ -14,13 +15,11 @@ class UTemplateSequenceSection;
 class FTemplateSequenceTrackEditor : public FMovieSceneTrackEditor
 {
 public:
-
 	FTemplateSequenceTrackEditor(TSharedRef<ISequencer> InSequencer);
 
 	static TSharedRef<ISequencerTrackEditor> CreateTrackEditor(TSharedRef<ISequencer> OwningSequencer);
 
 public:
-
 	// ISequencerTrackEditor interface
 	virtual bool SupportsType(TSubclassOf<class UMovieSceneTrack> TrackClass) const override;
 	virtual void BuildObjectBindingTrackMenu(FMenuBuilder& MenuBuilder, const TArray<FGuid>& ObjectBindings, const UClass* ObjectClass) override;
@@ -35,25 +34,13 @@ private:
 };
 
 class FTemplateSequenceSection
-	: public ISequencerSection
+	: public TSubSectionMixin<>
 	, public TSharedFromThis<FTemplateSequenceSection>
 {
 public:
-
 	/** Constructor. */
-	FTemplateSequenceSection(UTemplateSequenceSection& InSection);
+	FTemplateSequenceSection(TSharedPtr<ISequencer> InSequencer, UTemplateSequenceSection& InSection);
 
 	/** Virtual destructor. */
 	virtual ~FTemplateSequenceSection() {}
-
-	// ISequencerSection interface
-	virtual UMovieSceneSection* GetSectionObject() override;
-	virtual bool IsReadOnly() const override;
-	virtual FText GetSectionTitle() const override;
-	virtual int32 OnPaintSection(FSequencerSectionPainter& Painter) const override;
-
-private:
-
-	/** The section we are visualizing */
-	TWeakObjectPtr<UTemplateSequenceSection> WeakSection;
 };
diff --git a/Engine/Source/Editor/MovieSceneTools/Private/Sections/CinematicShotSection.cpp b/Engine/Source/Editor/MovieSceneTools/Private/Sections/CinematicShotSection.cpp
--- a/Engine/Source/Editor/MovieSceneTools/Private/Sections/CinematicShotSection.cpp
+++ b/Engine/Source/Editor/MovieSceneTools/Private/Sections/CinematicShotSection.cpp
@@ -49,13 +49,10 @@ FCinematicShotSection::FCinematicSectionCache::FCinematicSectionCache(UMovieScen
 }
 
 
-FCinematicShotSection::FCinematicShotSection(TSharedPtr<ISequencer> InSequencer, TSharedPtr<FTrackEditorThumbnailPool> InThumbnailPool, UMovieSceneSection& InSection, TSharedPtr<FCinematicShotTrackEditor> InCinematicShotTrackEditor)
-	: FViewportThumbnailSection(InSequencer, InThumbnailPool, InSection)
-	, SectionObject(*CastChecked<UMovieSceneCinematicShotSection>(&InSection))
+FCinematicShotSection::FCinematicShotSection(TSharedPtr<ISequencer> InSequencer, UMovieSceneCinematicShotSection& InSection, TSharedPtr<FCinematicShotTrackEditor> InCinematicShotTrackEditor, TSharedPtr<FTrackEditorThumbnailPool> InThumbnailPool)
+	: TSubSectionMixin(InSequencer, InSection, InSequencer, InThumbnailPool, InSection)
 	, CinematicShotTrackEditor(InCinematicShotTrackEditor)
-	, InitialStartOffsetDuringResize(0)
-	, InitialStartTimeDuringResize(0)
-	, ThumbnailCacheData(&SectionObject)
+	, ThumbnailCacheData(&InSection)
 {
 	AdditionalDrawEffect = ESlateDrawEffect::NoGamma;
 }
@@ -82,11 +79,11 @@ FMargin FCinematicShotSection::GetContentPadding() const
 
 void FCinematicShotSection::SetSingleTime(double GlobalTime)
 {
+	UMovieSceneCinematicShotSection& SectionObject = GetSectionObjectAs<UMovieSceneCinematicShotSection>();
 	double ReferenceOffsetSeconds = SectionObject.HasStartFrame() ? SectionObject.GetInclusiveStartFrame() / SectionObject.GetTypedOuter<UMovieScene>()->GetTickResolution() : 0;
 	SectionObject.SetThumbnailReferenceOffset(GlobalTime - ReferenceOffsetSeconds);
 }
 
-
 UCameraComponent* FindCameraCutComponentRecursive(FFrameNumber GlobalTime, FMovieSceneSequenceID InnerSequenceID, const FMovieSceneSequenceHierarchy& Hierarchy, IMovieScenePlayer& Player)
 {
 	const FMovieSceneSequenceHierarchyNode* Node    = Hierarchy.FindNode(InnerSequenceID);
@@ -152,16 +149,16 @@ UCameraComponent* FindCameraCutComponentRecursive(FFrameNumber GlobalTime, FMovi
 	return nullptr;
 }
 
-
 UCameraComponent* FCinematicShotSection::GetViewCamera()
 {
-	TSharedPtr<ISequencer> Sequencer = SequencerPtr.Pin();
+	TSharedPtr<ISequencer> Sequencer = GetSequencer();
 	if (!Sequencer.IsValid())
 	{
 		return nullptr;
 	}
 
 
+	const UMovieSceneCinematicShotSection&	SectionObject    = GetSectionObjectAs<UMovieSceneCinematicShotSection>();
 	const FMovieSceneSequenceID             ThisSequenceID   = Sequencer->GetFocusedTemplateID();
 	const FMovieSceneSequenceID             TargetSequenceID = SectionObject.GetSequenceID();
 	const FMovieSceneSequenceHierarchy&     Hierarchy        = Sequencer->GetEvaluationTemplate().GetHierarchy();
@@ -190,65 +187,6 @@ UCameraComponent* FCinematicShotSection::GetViewCamera()
 	return nullptr;
 }
 
-void FCinematicShotSection::BeginResizeSection()
-{
-	InitialStartOffsetDuringResize = SectionObject.Parameters.StartFrameOffset;
-	InitialStartTimeDuringResize = SectionObject.HasStartFrame() ? SectionObject.GetInclusiveStartFrame() : 0;
-}
-
-void FCinematicShotSection::ResizeSection(ESequencerSectionResizeMode ResizeMode, FFrameNumber ResizeTime)
-{
-	UMovieSceneSequence* InnerSequence = SectionObject.GetSequence();
-
-	// Adjust the start offset when resizing from the beginning
-	if (ResizeMode == SSRM_LeadingEdge && InnerSequence)
-	{
-		const FFrameRate    OuterFrameRate   = SectionObject.GetTypedOuter<UMovieScene>()->GetTickResolution();
-		const FFrameRate    InnerFrameRate   = InnerSequence->GetMovieScene()->GetTickResolution();
-		const FFrameNumber  ResizeDifference = ResizeTime - InitialStartTimeDuringResize;
-		const FFrameTime    InnerFrameTime   = ConvertFrameTime(ResizeDifference, OuterFrameRate, InnerFrameRate);
-		FFrameNumber		NewStartOffset   = FFrameTime::FromDecimal(InnerFrameTime.AsDecimal() * SectionObject.Parameters.TimeScale).FrameNumber;
-
-		NewStartOffset += InitialStartOffsetDuringResize;
-
-		// Ensure start offset is not less than 0
-		if (NewStartOffset < 0)
-		{
-			FFrameTime OuterFrameTimeOver = ConvertFrameTime(FFrameTime::FromDecimal(NewStartOffset.Value/SectionObject.Parameters.TimeScale), InnerFrameRate, OuterFrameRate);
-			ResizeTime = ResizeTime - OuterFrameTimeOver.GetFrame(); 
-			NewStartOffset = 0;
-		}
-		SectionObject.Parameters.StartFrameOffset = NewStartOffset;
-	}
-
-	FViewportThumbnailSection::ResizeSection(ResizeMode, ResizeTime);
-}
-
-void FCinematicShotSection::BeginSlipSection()
-{
-	BeginResizeSection();
-}
-
-void FCinematicShotSection::SlipSection(FFrameNumber SlipTime)
-{
-	UMovieSceneSequence* InnerSequence = SectionObject.GetSequence();
-
-	// Adjust the start offset when resizing from the beginning
-	if (InnerSequence)
-	{
-		const FFrameRate    OuterFrameRate   = SectionObject.GetTypedOuter<UMovieScene>()->GetTickResolution();
-		const FFrameRate    InnerFrameRate   = InnerSequence->GetMovieScene()->GetTickResolution();
-		const FFrameNumber  ResizeDifference = SlipTime - InitialStartTimeDuringResize;
-		const FFrameTime    InnerFrameTime = ConvertFrameTime(ResizeDifference, OuterFrameRate, InnerFrameRate);
-		const int32         NewStartOffset = FFrameTime::FromDecimal(InnerFrameTime.AsDecimal() * SectionObject.Parameters.TimeScale).FrameNumber.Value;
-
-		// Ensure start offset is not less than 0
-		SectionObject.Parameters.StartFrameOffset = FFrameNumber(FMath::Max(NewStartOffset, 0));
-	}
-
-	FViewportThumbnailSection::SlipSection(SlipTime);
-}
-
 bool FCinematicShotSection::IsReadOnly() const
 {
 	// Overridden to false regardless of movie scene section read only state so that we can double click into the sub section
@@ -258,6 +196,7 @@ bool FCinematicShotSection::IsReadOnly() const
 void FCinematicShotSection::Tick(const FGeometry& AllottedGeometry, const FGeometry& ClippedGeometry, const double InCurrentTime, const float InDeltaTime)
 {
 	// Set cached data
+	UMovieSceneCinematicShotSection& SectionObject = GetSectionObjectAs<UMovieSceneCinematicShotSection>();
 	FCinematicSectionCache NewCacheData(&SectionObject);
 	if (NewCacheData != ThumbnailCacheData)
 	{
@@ -286,7 +225,9 @@ int32 FCinematicShotSection::OnPaintSection(FSequencerSectionPainter& InPainter)
 	InPainter.LayerId = InPainter.PaintSectionBackground();
 
 	FVector2D LocalSectionSize = InPainter.SectionGeometry.GetLocalSize();
+	const UMovieSceneCinematicShotSection& SectionObject = GetSectionObjectAs<UMovieSceneCinematicShotSection>();
 
+	// Paint fancy-looking film border.
 	FSlateDrawElement::MakeBox(
 		InPainter.DrawElements,
 		InPainter.LayerId++,
@@ -303,194 +244,22 @@ int32 FCinematicShotSection::OnPaintSection(FSequencerSectionPainter& InPainter)
 		InPainter.bParentEnabled ? ESlateDrawEffect::None : ESlateDrawEffect::DisabledEffect
 	);
 
-	const TRange<FFrameNumber> SectionRange = SectionObject.GetRange();
-	if (SectionRange.GetLowerBound().IsOpen() || SectionRange.GetUpperBound().IsOpen())
-	{
-		return InPainter.LayerId;
-	}
-
-	const int32 SectionSize = MovieScene::DiscreteSize(SectionRange);
-	if (SectionSize <= 0)
-	{
-		return InPainter.LayerId;
-	}
-
+	// Paint the thumbnails.
 	FViewportThumbnailSection::OnPaintSection(InPainter);
 
-	UMovieSceneSequence* InnerSequence = SectionObject.GetSequence();
-	if (!InnerSequence)
-	{
-		return InPainter.LayerId;
-	}
-
-	if (SectionObject.Parameters.bCanLoop)
-	{
-		DoPaintLoopingSection(*InnerSequence, InPainter);
-	}
-	else
-	{
-		DoPaintNonLoopingSection(*InnerSequence, InPainter);
-	}
-
-	TSharedPtr<ISequencer> Sequencer = SequencerPtr.Pin();
-
-	if (InPainter.bIsSelected && Sequencer.IsValid())
-	{
-		FFrameTime CurrentTime = Sequencer->GetLocalTime().Time;
-		if (SectionRange.Contains(CurrentTime.FrameNumber))
-		{
-			UMovieScene* SubSequenceMovieScene = InnerSequence->GetMovieScene();
-			const FFrameRate DisplayRate = SubSequenceMovieScene->GetDisplayRate();
-			const FFrameRate TickResolution = SubSequenceMovieScene->GetTickResolution();
-			const FFrameNumber CurrentFrameNumber = ConvertFrameTime(CurrentTime * SectionObject.OuterToInnerTransform(), TickResolution, DisplayRate).FloorToFrame();
-
-			DrawFrameNumberHint(InPainter, CurrentTime, CurrentFrameNumber.Value);
-		}
-	}
+	// Paint the sub-sequence information/looping boundaries/etc.
+	FSubSectionPainterUtil::PaintSection(
+			GetSequencer(), SectionObject, InPainter, FSubSectionPainterParams(GetContentPadding()));
 
 	return InPainter.LayerId;
 }
 
-void FCinematicShotSection::DoPaintNonLoopingSection(UMovieSceneSequence& InnerSequence, FSequencerSectionPainter& InPainter) const
-{
-	const TRange<FFrameNumber> SectionRange = SectionObject.GetRange();
-	const FFrameNumber SectionStartFrame = SectionObject.GetInclusiveStartFrame();
-	const FFrameNumber SectionEndFrame = SectionObject.GetExclusiveEndFrame();
-	const int32 SectionSize = MovieScene::DiscreteSize(SectionRange);
-
-	const float PixelsPerFrame = InPainter.SectionGeometry.Size.X / float(SectionSize);
-
-	UMovieScene* MovieScene = InnerSequence.GetMovieScene();
-	TRange<FFrameNumber> PlaybackRange = MovieScene->GetPlaybackRange();
-
-	FMovieSceneSequenceTransform InnerToOuterTransform = SectionObject.OuterToInnerTransform().InverseLinearOnly();
-
-	const FFrameNumber PlaybackStart = (MovieScene::DiscreteInclusiveLower(PlaybackRange) * InnerToOuterTransform).FloorToFrame();
-	if (SectionRange.Contains(PlaybackStart))
-	{
-		const int32 StartOffset = (PlaybackStart-SectionStartFrame).Value;
-		// add dark tint for left out-of-bounds
-		FSlateDrawElement::MakeBox(
-			InPainter.DrawElements,
-			InPainter.LayerId++,
-			InPainter.SectionGeometry.ToPaintGeometry(
-				FVector2D(0.0f, 0.f),
-				FVector2D(StartOffset * PixelsPerFrame, InPainter.SectionGeometry.Size.Y)
-			),
-			FEditorStyle::GetBrush("WhiteBrush"),
-			ESlateDrawEffect::None,
-			FLinearColor::Black.CopyWithNewOpacity(0.5f)
-		);
-
-		// add green line for playback start
-		FSlateDrawElement::MakeBox(
-			InPainter.DrawElements,
-			InPainter.LayerId++,
-			InPainter.SectionGeometry.ToPaintGeometry(
-				FVector2D(StartOffset * PixelsPerFrame, 0.f),
-				FVector2D(1.0f, InPainter.SectionGeometry.Size.Y)
-			),
-			FEditorStyle::GetBrush("WhiteBrush"),
-			ESlateDrawEffect::None,
-			FColor(32, 128, 32)	// 120, 75, 50 (HSV)
-		);
-	}
-
-	const FFrameNumber PlaybackEnd = (MovieScene::DiscreteExclusiveUpper(PlaybackRange) * InnerToOuterTransform).FloorToFrame();
-	if (SectionRange.Contains(PlaybackEnd))
-	{
-		// add dark tint for right out-of-bounds
-		const int32 EndOffset = (PlaybackEnd-SectionStartFrame).Value;
-		FSlateDrawElement::MakeBox(
-			InPainter.DrawElements,
-			InPainter.LayerId++,
-			InPainter.SectionGeometry.ToPaintGeometry(
-				FVector2D(EndOffset * PixelsPerFrame, 0.f),
-				FVector2D((SectionSize - EndOffset) * PixelsPerFrame, InPainter.SectionGeometry.Size.Y)
-			),
-			FEditorStyle::GetBrush("WhiteBrush"),
-			ESlateDrawEffect::None,
-			FLinearColor::Black.CopyWithNewOpacity(0.5f)
-		);
-
-
-		// add red line for playback end
-		FSlateDrawElement::MakeBox(
-			InPainter.DrawElements,
-			InPainter.LayerId++,
-			InPainter.SectionGeometry.ToPaintGeometry(
-				FVector2D(EndOffset * PixelsPerFrame, 0.f),
-				FVector2D(1.0f, InPainter.SectionGeometry.Size.Y)
-			),
-			FEditorStyle::GetBrush("WhiteBrush"),
-			ESlateDrawEffect::None,
-			FColor(128, 32, 32)	// 0, 75, 50 (HSV)
-		);
-	}
-}
-
-void FCinematicShotSection::DoPaintLoopingSection(UMovieSceneSequence& InnerSequence, FSequencerSectionPainter& InPainter) const
-{
-	const FTimeToPixel& TimeToPixelConverter = InPainter.GetTimeConverter();
-
-	const FFrameNumber SectionStartFrame = SectionObject.GetInclusiveStartFrame();
-	const FFrameNumber SectionEndFrame = SectionObject.GetExclusiveEndFrame();
-	const float InvTimeScale = FMath::IsNearlyZero(SectionObject.Parameters.TimeScale) ? 1.0f : 1.0f / SectionObject.Parameters.TimeScale;
-
-	UMovieScene* MovieScene = InnerSequence.GetMovieScene();
-	TRange<FFrameNumber> InnerPlaybackRange = UMovieSceneSubSection::GetValidatedInnerPlaybackRange(SectionObject.Parameters, *MovieScene);
-
-	const FFrameNumber ShotLength = InnerPlaybackRange.Size<FFrameNumber>() * InvTimeScale;
-	const FFrameNumber FirstLoopShotLength = ShotLength - SectionObject.Parameters.FirstLoopStartFrameOffset * InvTimeScale;
-
-	FFrameNumber CurOffsetFrame = FMath::Max(FirstLoopShotLength, FFrameNumber(0));
-	const ESlateDrawEffect DrawEffects = InPainter.bParentEnabled ? ESlateDrawEffect::None : ESlateDrawEffect::DisabledEffect;
-
-	// Draw separators where the shot is looping. To be consistent with the non-looping case, we draw a red and green
-	// separator back to back.
-	uint32 MaxLoopBoundaries = 100;
-	while (CurOffsetFrame < SectionEndFrame)
-	{
-		const float CurOffsetPixel = TimeToPixelConverter.FrameToPixel(SectionStartFrame + CurOffsetFrame);
-
-		FSlateDrawElement::MakeBox(
-			InPainter.DrawElements,
-			InPainter.LayerId++,
-			InPainter.SectionGeometry.MakeChild(
-				FVector2D(1.f, InPainter.SectionGeometry.Size.Y - 2.f),
-				FSlateLayoutTransform(FVector2D(CurOffsetPixel, 1.f))
-			).ToPaintGeometry(),
-			FEditorStyle::GetBrush("WhiteBrush"),
-			DrawEffects,
-			FColor(32, 128, 32)	// 120, 75, 50 (HSV)
-		);
-		if (CurOffsetFrame > 0) {
-			FSlateDrawElement::MakeBox(
-				InPainter.DrawElements,
-				InPainter.LayerId++,
-				InPainter.SectionGeometry.MakeChild(
-					FVector2D(1.f, InPainter.SectionGeometry.Size.Y - 2.f),
-					FSlateLayoutTransform(FVector2D(CurOffsetPixel - 1.f, 1.f))
-				).ToPaintGeometry(),
-				FEditorStyle::GetBrush("WhiteBrush"),
-				DrawEffects,
-				FColor(128, 32, 32)	// 0, 75, 50 (HSV)
-			);
-		}
-
-		CurOffsetFrame += ShotLength;
-
-		if ((--MaxLoopBoundaries) == 0)
-		{
-			break;
-		}
-	}
-}
-
 void FCinematicShotSection::BuildSectionContextMenu(FMenuBuilder& MenuBuilder, const FGuid& ObjectBinding)
 {
 	FViewportThumbnailSection::BuildSectionContextMenu(MenuBuilder, ObjectBinding);
 
+	UMovieSceneCinematicShotSection& SectionObject = GetSectionObjectAs<UMovieSceneCinematicShotSection>();
+
 	MenuBuilder.BeginSection(NAME_None, LOCTEXT("ShotMenuText", "Shot"));
 	{
 		MenuBuilder.AddSubMenu(
@@ -540,9 +309,10 @@ void FCinematicShotSection::AddTakesMenu(FMenuBuilder& MenuBuilder)
 {
 	TArray<FAssetData> AssetData;
 	uint32 CurrentTakeNumber = INDEX_NONE;
+	const UMovieSceneCinematicShotSection& SectionObject = GetSectionObjectAs<UMovieSceneCinematicShotSection>();
 	MovieSceneToolHelpers::GatherTakes(&SectionObject, AssetData, CurrentTakeNumber);
 
-	AssetData.Sort([this](const FAssetData &A, const FAssetData &B) {
+	AssetData.Sort([&SectionObject](const FAssetData &A, const FAssetData &B) {
 		uint32 TakeNumberA = INDEX_NONE;
 		uint32 TakeNumberB = INDEX_NONE;
 		if (MovieSceneToolHelpers::GetTakeNumber(&SectionObject, A, TakeNumberA) && MovieSceneToolHelpers::GetTakeNumber(&SectionObject, B, TakeNumberB))
@@ -577,6 +347,7 @@ void FCinematicShotSection::AddTakesMenu(FMenuBuilder& MenuBuilder)
 
 FText FCinematicShotSection::HandleThumbnailTextBlockText() const
 {
+	const UMovieSceneCinematicShotSection& SectionObject = GetSectionObjectAs<UMovieSceneCinematicShotSection>();
 	return FText::FromString(SectionObject.GetShotDisplayName());
 }
 
@@ -585,6 +356,8 @@ void FCinematicShotSection::HandleThumbnailTextBlockTextCommitted(const FText& N
 {
 	if (CommitType == ETextCommit::OnEnter && !HandleThumbnailTextBlockText().EqualTo(NewShotName))
 	{
+		UMovieSceneCinematicShotSection& SectionObject = GetSectionObjectAs<UMovieSceneCinematicShotSection>();
+
 		SectionObject.Modify();
 
 		const FScopedTransaction Transaction(LOCTEXT("SetShotName", "Set Shot Name"));
@@ -593,24 +366,4 @@ void FCinematicShotSection::HandleThumbnailTextBlockTextCommitted(const FText& N
 	}
 }
 
-FReply FCinematicShotSection::OnSectionDoubleClicked(const FGeometry& SectionGeometry, const FPointerEvent& MouseEvent)
-{
-	if( MouseEvent.GetEffectingButton() == EKeys::LeftMouseButton )
-	{
-		if (SectionObject.GetSequence())
-		{
-			if (MouseEvent.IsControlDown())
-			{
-				GEditor->GetEditorSubsystem<UAssetEditorSubsystem>()->OpenEditorForAsset(SectionObject.GetSequence());
-			}
-			else
-			{
-				SequencerPtr.Pin()->FocusSequenceInstance(SectionObject);
-			}
-		}
-	}
-
-	return FReply::Handled();
-}
-
 #undef LOCTEXT_NAMESPACE
diff --git a/Engine/Source/Editor/MovieSceneTools/Private/Sections/CinematicShotSection.h b/Engine/Source/Editor/MovieSceneTools/Private/Sections/CinematicShotSection.h
--- a/Engine/Source/Editor/MovieSceneTools/Private/Sections/CinematicShotSection.h
+++ b/Engine/Source/Editor/MovieSceneTools/Private/Sections/CinematicShotSection.h
@@ -6,6 +6,7 @@
 #include "Input/Reply.h"
 #include "Layout/Margin.h"
 #include "Sections/ThumbnailSection.h"
+#include "TrackEditors/SubTrackEditorBase.h"
 
 class FCinematicShotTrackEditor;
 class FMenuBuilder;
@@ -17,12 +18,12 @@ class UMovieSceneCinematicShotSection;
  * CinematicShot section, which paints and ticks the appropriate section.
  */
 class FCinematicShotSection
-	: public FViewportThumbnailSection
+	: public TSubSectionMixin<FViewportThumbnailSection>
 {
 public:
 
 	/** Create and initialize a new instance. */
-	FCinematicShotSection(TSharedPtr<ISequencer> InSequencer, TSharedPtr<FTrackEditorThumbnailPool> InThumbnailPool, UMovieSceneSection& InSection, TSharedPtr<FCinematicShotTrackEditor> InCinematicShotTrackEditor);
+	FCinematicShotSection(TSharedPtr<ISequencer> InSequencer, UMovieSceneCinematicShotSection& InSection, TSharedPtr<FCinematicShotTrackEditor> InCinematicShotTrackEditor, TSharedPtr<FTrackEditorThumbnailPool> InThumbnailPool);
 
 	/** Virtual destructor. */
 	virtual ~FCinematicShotSection();
@@ -34,14 +35,9 @@ public:
 	virtual void Tick(const FGeometry& AllottedGeometry, const FGeometry& ClippedGeometry, const double InCurrentTime, const float InDeltaTime) override;
 	virtual int32 OnPaintSection( FSequencerSectionPainter& Painter ) const override;
 	virtual void BuildSectionContextMenu(FMenuBuilder& MenuBuilder, const FGuid& ObjectBinding) override;
-	virtual FReply OnSectionDoubleClicked(const FGeometry& SectionGeometry, const FPointerEvent& MouseEvent) override;
 	virtual FText GetSectionTitle() const override;
 	virtual float GetSectionHeight() const override;
 	virtual FMargin GetContentPadding() const override;
-	virtual void BeginResizeSection() override;
-	virtual void ResizeSection(ESequencerSectionResizeMode ResizeMode, FFrameNumber ResizeTime) override;
-	virtual void BeginSlipSection() override;
-	virtual void SlipSection(FFrameNumber SlipTime) override;
 	virtual bool IsReadOnly() const override;
 
 	// FThumbnail interface
@@ -55,23 +51,11 @@ private:
 	/** Add shot takes menu */
 	void AddTakesMenu(FMenuBuilder& MenuBuilder);
 
-	void DoPaintNonLoopingSection(UMovieSceneSequence& InnerSequence, FSequencerSectionPainter& InPainter) const;
-	void DoPaintLoopingSection(UMovieSceneSequence& InnerSequence, FSequencerSectionPainter& InPainter) const;
-
 private:
 
-	/** The section we are visualizing */
-	UMovieSceneCinematicShotSection& SectionObject;
-
 	/** The cinematic shot track editor that contains this section */
 	TWeakPtr<FCinematicShotTrackEditor> CinematicShotTrackEditor;
 
-	/** Cached start offset value valid only during resize */
-	FFrameNumber InitialStartOffsetDuringResize;
-
-	/** Cached start time valid only during resize */
-	FFrameNumber InitialStartTimeDuringResize;
-
 	struct FCinematicSectionCache
 	{
 		FCinematicSectionCache(UMovieSceneCinematicShotSection* Section = nullptr);
diff --git a/Engine/Source/Editor/MovieSceneTools/Private/TrackEditors/CinematicShotTrackEditor.cpp b/Engine/Source/Editor/MovieSceneTools/Private/TrackEditors/CinematicShotTrackEditor.cpp
--- a/Engine/Source/Editor/MovieSceneTools/Private/TrackEditors/CinematicShotTrackEditor.cpp
+++ b/Engine/Source/Editor/MovieSceneTools/Private/TrackEditors/CinematicShotTrackEditor.cpp
@@ -125,8 +125,9 @@ TSharedPtr<SWidget> FCinematicShotTrackEditor::BuildOutlinerEditWidget(const FGu
 TSharedRef<ISequencerSection> FCinematicShotTrackEditor::MakeSectionInterface(UMovieSceneSection& SectionObject, UMovieSceneTrack& Track, FGuid ObjectBinding)
 {
 	check(SupportsType(SectionObject.GetOuter()->GetClass()));
-
-	return MakeShareable(new FCinematicShotSection(GetSequencer(), ThumbnailPool, SectionObject, SharedThis(this)));
+	
+	UMovieSceneCinematicShotSection& SectionObjectImpl = *CastChecked<UMovieSceneCinematicShotSection>(&SectionObject);
+	return MakeShareable(new FCinematicShotSection(GetSequencer(), SectionObjectImpl, SharedThis(this), ThumbnailPool));
 }
 
 
diff --git a/Engine/Source/Editor/MovieSceneTools/Private/TrackEditors/SubTrackEditor.cpp b/Engine/Source/Editor/MovieSceneTools/Private/TrackEditors/SubTrackEditor.cpp
--- a/Engine/Source/Editor/MovieSceneTools/Private/TrackEditors/SubTrackEditor.cpp
+++ b/Engine/Source/Editor/MovieSceneTools/Private/TrackEditors/SubTrackEditor.cpp
@@ -22,6 +22,7 @@
 #include "SequencerSectionPainter.h"
 #include "ISequenceRecorder.h"
 #include "SequenceRecorderSettings.h"
+#include "TrackEditors/SubTrackEditorBase.h"
 #include "DragAndDrop/AssetDragDropOp.h"
 #include "MovieSceneToolHelpers.h"
 #include "Misc/QualifiedFrameTime.h"
@@ -44,17 +45,14 @@ namespace SubTrackEditorConstants
  * A generic implementation for displaying simple property sections.
  */
 class FSubSection
-	: public ISequencerSection
+	: public TSubSectionMixin<>
 {
 public:
 
 	FSubSection(TSharedPtr<ISequencer> InSequencer, UMovieSceneSection& InSection, const FText& InDisplayName, TSharedPtr<FSubTrackEditor> InSubTrackEditor)
-		: DisplayName(InDisplayName)
-		, SectionObject(*CastChecked<UMovieSceneSubSection>(&InSection))
-		, Sequencer(InSequencer)
+		: TSubSectionMixin(InSequencer, *CastChecked<UMovieSceneSubSection>(&InSection))
+		, DisplayName(InDisplayName)
 		, SubTrackEditor(InSubTrackEditor)
-		, InitialStartOffsetDuringResize(0)
-		, InitialStartTimeDuringResize(0)
 	{
 	}
 
@@ -67,17 +65,11 @@ public:
 		return SubTrackEditorConstants::TrackHeight;
 	}
 
-	virtual UMovieSceneSection* GetSectionObject() override
-	{
-		return &SectionObject;
-	}
 	virtual FText GetSectionTitle() const override
 	{
-		if(SectionObject.GetSequence())
-		{
-			return FText::FromString(SectionObject.GetSequence()->GetName());
-		}
-		else if(UMovieSceneSubSection::GetRecordingSection() == &SectionObject)
+		const UMovieSceneSubSection& SectionObject = GetSubSectionObject();
+		
+		if(SectionObject.GetSequence() == nullptr && UMovieSceneSubSection::GetRecordingSection() == &SectionObject)
 		{
 			AActor* ActorToRecord = UMovieSceneSubSection::GetActorToRecord();
 
@@ -107,90 +99,29 @@ public:
 		}
 		else
 		{
-			return LOCTEXT("NoSequenceSelected", "No Sequence Selected");
+			return TSubSectionMixin::GetSectionTitle();
 		}
 	}
 	
 	virtual int32 OnPaintSection( FSequencerSectionPainter& InPainter ) const override
 	{
-		int32 LayerId = InPainter.PaintSectionBackground();
+		InPainter.PaintSectionBackground();
 
-		const TRange<FFrameNumber> SectionRange = SectionObject.GetRange();
-		if (SectionRange.GetLowerBound().IsOpen() || SectionRange.GetUpperBound().IsOpen())
-		{
-			return InPainter.LayerId;
-		}
+		const UMovieSceneSubSection& SectionObject = GetSubSectionObject();
 
-		const int32 SectionSize = MovieScene::DiscreteSize(SectionRange);
-		if (SectionSize <= 0)
+		FSubSectionPainterResult PaintResult = FSubSectionPainterUtil::PaintSection(
+				GetSequencer(), SectionObject, InPainter, FSubSectionPainterParams(GetContentPadding()));
+		if (PaintResult == FSSPR_InvalidSection)
 		{
 			return InPainter.LayerId;
 		}
 
-		UMovieSceneSequence* InnerSequence = SectionObject.GetSequence();
-		const ESlateDrawEffect DrawEffects = InPainter.bParentEnabled ? ESlateDrawEffect::None : ESlateDrawEffect::DisabledEffect;
-		
-		if (InnerSequence)
-		{
-			if (SectionObject.Parameters.bCanLoop)
-			{
-				DoPaintLoopingSection(*InnerSequence, InPainter, DrawEffects);
-			}
-			else
-			{
-				DoPaintNonLoopingSection(*InnerSequence, InPainter, DrawEffects);
-			}
-
-			FMargin ContentPadding = GetContentPadding();
-			UMovieScene* MovieScene = InnerSequence->GetMovieScene();
-			const int32 NumTracks = MovieScene->GetPossessableCount() + MovieScene->GetSpawnableCount() + MovieScene->GetMasterTracks().Num();
-
-			FVector2D TopLeft = InPainter.SectionGeometry.AbsoluteToLocal(InPainter.SectionClippingRect.GetTopLeft()) + FVector2D(1.f, -1.f);
-
-			FSlateFontInfo FontInfo = FEditorStyle::GetFontStyle("NormalFont");
-
-			TSharedRef<FSlateFontCache> FontCache = FSlateApplication::Get().GetRenderer()->GetFontCache();
-
-			auto GetFontHeight = [&]
-			{
-				return FontCache->GetMaxCharacterHeight(FontInfo, 1.f) + FontCache->GetBaseline(FontInfo, 1.f);
-			};
-			while (GetFontHeight() > InPainter.SectionGeometry.Size.Y && FontInfo.Size > 11)
-			{
-				FontInfo.Size = FMath::Max(FMath::FloorToInt(FontInfo.Size - 6.f), 11);
-			}
-
-			FSlateDrawElement::MakeText(
-				InPainter.DrawElements,
-				++LayerId,
-				InPainter.SectionGeometry.MakeChild(
-					FVector2D(InPainter.SectionGeometry.Size.X, GetFontHeight()),
-					FSlateLayoutTransform(TopLeft + FVector2D(ContentPadding.Left, ContentPadding.Top) + FVector2D(11.f, GetFontHeight()*2.f))
-				).ToPaintGeometry(),
-				FText::Format(LOCTEXT("NumTracksFormat", "{0} track(s)"), FText::AsNumber(NumTracks)),
-				FontInfo,
-				DrawEffects,
-				FColor(200, 200, 200)
-			);
-
-			TSharedPtr<ISequencer> SequencerPtr = Sequencer.Pin();
-			if (InPainter.bIsSelected && SequencerPtr.IsValid())
-			{
-				FFrameTime CurrentTime = SequencerPtr->GetLocalTime().Time;
-				if (SectionRange.Contains(CurrentTime.FrameNumber))
-				{
-					UMovieScene* SubSequenceMovieScene = SectionObject.GetSequence()->GetMovieScene();
-					const FFrameRate DisplayRate = SubSequenceMovieScene->GetDisplayRate();
-					const FFrameRate TickResolution = SubSequenceMovieScene->GetTickResolution();
-					const FFrameNumber CurrentFrameNumber = ConvertFrameTime(CurrentTime * SectionObject.OuterToInnerTransform(), TickResolution, DisplayRate).FloorToFrame();
-
-					DrawFrameNumberHint(InPainter, CurrentTime, CurrentFrameNumber.Value);
-				}
-			}
+		int32 LayerId = InPainter.LayerId;
 
-		}
-		else if (UMovieSceneSubSection::GetRecordingSection() == &SectionObject)
+		if (SectionObject.GetSequence() == nullptr && UMovieSceneSubSection::GetRecordingSection() == &SectionObject)
 		{
+			const ESlateDrawEffect DrawEffects = InPainter.bParentEnabled ? ESlateDrawEffect::None : ESlateDrawEffect::DisabledEffect;
+
 			FColor SubSectionColor = FColor(180, 75, 75, 190);
 	
 			ISequenceRecorder& SequenceRecorder = FModuleManager::LoadModuleChecked<ISequenceRecorder>("SequenceRecorder");
@@ -230,157 +161,7 @@ public:
 		return LayerId;
 	}
 
-	void DoPaintNonLoopingSection(UMovieSceneSequence& InnerSequence, FSequencerSectionPainter& InPainter, ESlateDrawEffect DrawEffects) const
-	{
-		const FFrameNumber SectionStartFrame = SectionObject.GetInclusiveStartFrame();
-
-		const TRange<FFrameNumber> SectionRange = SectionObject.GetRange();
-		const int32 SectionSize = MovieScene::DiscreteSize(SectionRange);
-		const float PixelsPerFrame = InPainter.SectionGeometry.Size.X / float(SectionSize);
-
-		UMovieScene* MovieScene = InnerSequence.GetMovieScene();
-		TRange<FFrameNumber> PlaybackRange = MovieScene->GetPlaybackRange();
-
-		// We're in the non-looping case so we know we have a purely linear transform.
-		const FMovieSceneSequenceTransform InnerToOuterTransform = SectionObject.OuterToInnerTransform().InverseLinearOnly();
-		const FFrameNumber PlaybackStart = (MovieScene::DiscreteInclusiveLower(PlaybackRange) * InnerToOuterTransform).FloorToFrame();
-		if (SectionRange.Contains(PlaybackStart))
-		{
-			const int32 StartOffset = (PlaybackStart - SectionStartFrame).Value;
-			// add dark tint for left out-of-bounds
-			FSlateDrawElement::MakeBox(
-					InPainter.DrawElements,
-					InPainter.LayerId++,
-					InPainter.SectionGeometry.ToPaintGeometry(
-						FVector2D(0.0f, 0.f),
-						FVector2D(StartOffset * PixelsPerFrame, InPainter.SectionGeometry.Size.Y)
-						),
-					FEditorStyle::GetBrush("WhiteBrush"),
-					DrawEffects,
-					FLinearColor::Black.CopyWithNewOpacity(0.5f)
-					);
-
-			// add green line for playback start
-			FSlateDrawElement::MakeBox(
-					InPainter.DrawElements,
-					InPainter.LayerId++,
-					InPainter.SectionGeometry.ToPaintGeometry(
-						FVector2D(StartOffset * PixelsPerFrame, 0.f),
-						FVector2D(1.0f, InPainter.SectionGeometry.Size.Y)
-						),
-					FEditorStyle::GetBrush("WhiteBrush"),
-					DrawEffects,
-					FColor(32, 128, 32)	// 120, 75, 50 (HSV)
-					);
-		}
-
-		const FFrameNumber PlaybackEnd = (MovieScene::DiscreteExclusiveUpper(PlaybackRange) * InnerToOuterTransform) .FloorToFrame();
-		if (SectionRange.Contains(PlaybackEnd))
-		{
-			// add dark tint for right out-of-bounds
-			const int32 EndOffset = (PlaybackEnd - SectionStartFrame).Value;
-			FSlateDrawElement::MakeBox(
-					InPainter.DrawElements,
-					InPainter.LayerId++,
-					InPainter.SectionGeometry.ToPaintGeometry(
-						FVector2D(EndOffset * PixelsPerFrame, 0.f),
-						FVector2D((SectionSize - EndOffset) * PixelsPerFrame, InPainter.SectionGeometry.Size.Y)
-						),
-					FEditorStyle::GetBrush("WhiteBrush"),
-					DrawEffects,
-					FLinearColor::Black.CopyWithNewOpacity(0.5f)
-					);
-
-
-			// add red line for playback end
-			FSlateDrawElement::MakeBox(
-					InPainter.DrawElements,
-					InPainter.LayerId++,
-					InPainter.SectionGeometry.ToPaintGeometry(
-						FVector2D(EndOffset * PixelsPerFrame, 0.f),
-						FVector2D(1.0f, InPainter.SectionGeometry.Size.Y)
-						),
-					FEditorStyle::GetBrush("WhiteBrush"),
-					DrawEffects,
-					FColor(128, 32, 32)	// 0, 75, 50 (HSV)
-					);
-		}
-	}
-	
-	void DoPaintLoopingSection(UMovieSceneSequence& InnerSequence, FSequencerSectionPainter& InPainter, ESlateDrawEffect DrawEffects) const
-	{
-		const FFrameNumber SectionStartFrame = SectionObject.GetInclusiveStartFrame();
-		const FFrameNumber SectionEndFrame   = SectionObject.GetExclusiveEndFrame();
-
-		const TRange<FFrameNumber> SectionRange = SectionObject.GetRange();
-		const int32 SectionSize = MovieScene::DiscreteSize(SectionRange);
-		const float PixelsPerFrame = InPainter.SectionGeometry.Size.X / float(SectionSize);
-
-		const float InvTimeScale = FMath::IsNearlyZero(SectionObject.Parameters.TimeScale) ? 1.0f : 1.0f / SectionObject.Parameters.TimeScale;
-
-		UMovieScene* MovieScene = InnerSequence.GetMovieScene();
-		TRange<FFrameNumber> InnerPlaybackRange = UMovieSceneSubSection::GetValidatedInnerPlaybackRange(SectionObject.Parameters, *MovieScene);
-
-		const FFrameNumber SubSeqLength = InnerPlaybackRange.Size<FFrameNumber>() * InvTimeScale;
-		const FFrameNumber FirstLoopSubSeqLength = SubSeqLength - SectionObject.Parameters.FirstLoopStartFrameOffset * InvTimeScale;
-
-		FFrameNumber CurOffsetFrame = FMath::Max(FirstLoopSubSeqLength, FFrameNumber(0));
-
-		// Draw separators where the sub-sequence is looping. To be consistent with the non-looping case, we draw a red and green
-		// separator back to back.
-		uint32 MaxLoopBoundaries = 100;
-		while (CurOffsetFrame < SectionEndFrame)
-		{
-			const int32 CurOffset = CurOffsetFrame.Value;
-			FSlateDrawElement::MakeBox(
-				InPainter.DrawElements,
-				InPainter.LayerId++,
-				InPainter.SectionGeometry.ToPaintGeometry(
-					FVector2D(CurOffset * PixelsPerFrame, 0.f),
-					FVector2D(1.0f, InPainter.SectionGeometry.Size.Y)
-				),
-				FEditorStyle::GetBrush("WhiteBrush"),
-				DrawEffects,
-				FColor(32, 128, 32)	// 120, 75, 50 (HSV)
-			);
-			if (CurOffset > 0)
-			{
-				FSlateDrawElement::MakeBox(
-					InPainter.DrawElements,
-					InPainter.LayerId++,
-					InPainter.SectionGeometry.ToPaintGeometry(
-						FVector2D(CurOffset * PixelsPerFrame - 1.f, 0.f),
-						FVector2D(1.0f, InPainter.SectionGeometry.Size.Y)
-					),
-					FEditorStyle::GetBrush("WhiteBrush"),
-					DrawEffects,
-					FColor(128, 32, 32)	// 0, 75, 50 (HSV)
-				);
-			}
-
-			CurOffsetFrame += SubSeqLength;
-
-			if ((--MaxLoopBoundaries) == 0)
-			{
-				break;
-			}
-		}
-	}
-
-	virtual FReply OnSectionDoubleClicked(const FGeometry& SectionGeometry, const FPointerEvent& MouseEvent) override
-	{
-		if( MouseEvent.GetEffectingButton() == EKeys::LeftMouseButton )
-		{
-			if (SectionObject.GetSequence())
-			{
-				Sequencer.Pin()->FocusSequenceInstance(SectionObject);
-			}
-		}
-
-		return FReply::Handled();
-	}
-
-	virtual void BuildSectionContextMenu(FMenuBuilder& MenuBuilder, const FGuid& ObjectBinding)
+	virtual void BuildSectionContextMenu(FMenuBuilder& MenuBuilder, const FGuid& ObjectBinding) override
 	{
 		ISequencerSection::BuildSectionContextMenu(MenuBuilder, ObjectBinding);
 
@@ -405,13 +186,13 @@ public:
 
 	void TogglePlayableDirectly()
 	{
-		TSharedPtr<ISequencer> SequencerPtr = Sequencer.Pin();
-		if (SequencerPtr)
+		TSharedPtr<ISequencer> Sequencer = GetSequencer();
+		if (Sequencer)
 		{
 			FScopedTransaction Transaction(LOCTEXT("SetPlayableDirectly_Transaction", "Set Playable Directly"));
 
 			TArray<UMovieSceneSection*> SelectedSections;
-			SequencerPtr->GetSelectedSections(SelectedSections);
+			Sequencer->GetSelectedSections(SelectedSections);
 
 			const bool bNewPlayableDirectly = IsPlayableDirectly() != ECheckBoxState::Checked;
 
@@ -433,11 +214,11 @@ public:
 	{
 		ECheckBoxState CheckboxState = ECheckBoxState::Undetermined;
 
-		TSharedPtr<ISequencer> SequencerPtr = Sequencer.Pin();
-		if (SequencerPtr)
+		TSharedPtr<ISequencer> Sequencer = GetSequencer();
+		if (Sequencer)
 		{
 			TArray<UMovieSceneSection*> SelectedSections;
-			SequencerPtr->GetSelectedSections(SelectedSections);
+			Sequencer->GetSelectedSections(SelectedSections);
 
 			for (UMovieSceneSection* Section : SelectedSections)
 			{
@@ -462,67 +243,6 @@ public:
 		return CheckboxState;
 	}
 
-	void BeginResizeSection()
-	{
-		InitialStartOffsetDuringResize = SectionObject.Parameters.StartFrameOffset;
-		InitialStartTimeDuringResize = SectionObject.HasStartFrame() ? SectionObject.GetInclusiveStartFrame() : 0;
-	}
-
-	void ResizeSection(ESequencerSectionResizeMode ResizeMode, FFrameNumber ResizeTime)
-	{
-		UMovieSceneSequence* InnerSequence = SectionObject.GetSequence();
-
-		// Adjust the start offset when resizing from the beginning
-		if (ResizeMode == SSRM_LeadingEdge && InnerSequence)
-		{
-			const FFrameRate    OuterFrameRate   = SectionObject.GetTypedOuter<UMovieScene>()->GetTickResolution();
-			const FFrameRate    InnerFrameRate   = InnerSequence->GetMovieScene()->GetTickResolution();
-			const FFrameNumber  ResizeDifference = ResizeTime - InitialStartTimeDuringResize;
-			const FFrameTime    InnerFrameTime   = ConvertFrameTime(ResizeDifference, OuterFrameRate, InnerFrameRate);
-			FFrameNumber		NewStartOffset   = FFrameTime::FromDecimal(InnerFrameTime.AsDecimal() * SectionObject.Parameters.TimeScale).FrameNumber;
-
-			NewStartOffset += InitialStartOffsetDuringResize;
-
-			// Ensure start offset is not less than 0
-			if (NewStartOffset < 0)
-			{
-				FFrameTime OuterFrameTimeOver = ConvertFrameTime(FFrameTime::FromDecimal(NewStartOffset.Value/SectionObject.Parameters.TimeScale), InnerFrameRate, OuterFrameRate);
-				ResizeTime = ResizeTime - OuterFrameTimeOver.GetFrame(); 
-				NewStartOffset = 0;
-			}
-
-			SectionObject.Parameters.StartFrameOffset = FFrameNumber(NewStartOffset);
-		}
-
-		ISequencerSection::ResizeSection(ResizeMode, ResizeTime);
-	}
-
-	virtual void BeginSlipSection() override
-	{
-		InitialStartOffsetDuringResize = SectionObject.Parameters.StartFrameOffset;
-		InitialStartTimeDuringResize = SectionObject.HasStartFrame() ? SectionObject.GetInclusiveStartFrame() : 0;
-	}
-
-	virtual void SlipSection(FFrameNumber SlipTime) override
-	{
-		UMovieSceneSequence* InnerSequence = SectionObject.GetSequence();
-
-		// Adjust the start offset when resizing from the beginning
-		if (InnerSequence)
-		{
-			const FFrameRate    OuterFrameRate   = SectionObject.GetTypedOuter<UMovieScene>()->GetTickResolution();
-			const FFrameRate    InnerFrameRate   = InnerSequence->GetMovieScene()->GetTickResolution();
-			const FFrameNumber  ResizeDifference = SlipTime - InitialStartTimeDuringResize;
-			const FFrameTime    InnerFrameTime = ConvertFrameTime(ResizeDifference, OuterFrameRate, InnerFrameRate);
-			const int32         NewStartOffset = FFrameTime::FromDecimal(InnerFrameTime.AsDecimal() * SectionObject.Parameters.TimeScale).FrameNumber.Value;
-
-			// Ensure start offset is not less than 0
-			SectionObject.Parameters.StartFrameOffset = FFrameNumber(FMath::Max(NewStartOffset, 0));
-		}
-
-		ISequencerSection::SlipSection(SlipTime);
-	}
-
 	virtual bool IsReadOnly() const override
 	{
 		// Overridden to false regardless of movie scene section read only state so that we can double click into the sub section
@@ -535,9 +255,10 @@ private:
 	{
 		TArray<FAssetData> AssetData;
 		uint32 CurrentTakeNumber = INDEX_NONE;
+		UMovieSceneSubSection& SectionObject = GetSubSectionObject();
 		MovieSceneToolHelpers::GatherTakes(&SectionObject, AssetData, CurrentTakeNumber);
 
-		AssetData.Sort([this](const FAssetData &A, const FAssetData &B) {
+		AssetData.Sort([&SectionObject](const FAssetData &A, const FAssetData &B) {
 			uint32 TakeNumberA = INDEX_NONE;
 			uint32 TakeNumberB = INDEX_NONE;
 			if (MovieSceneToolHelpers::GetTakeNumber(&SectionObject, A, TakeNumberA) && MovieSceneToolHelpers::GetTakeNumber(&SectionObject, B, TakeNumberB))
@@ -572,18 +293,8 @@ private:
 	/** Display name of the section */
 	FText DisplayName;
 
-	UMovieSceneSubSection& SectionObject;
-	/** Sequencer interface */
-	TWeakPtr<ISequencer> Sequencer;
-
 	/** The sub track editor that contains this section */
 	TWeakPtr<FSubTrackEditor> SubTrackEditor;
-
-	/** Cached start offset value valid only during resize */
-	FFrameNumber InitialStartOffsetDuringResize;
-
-	/** Cached start time valid only during resize */
-	FFrameNumber InitialStartTimeDuringResize;
 };
 
 
diff --git a/Engine/Source/Editor/MovieSceneTools/Private/TrackEditors/SubTrackEditorBase.cpp b/Engine/Source/Editor/MovieSceneTools/Private/TrackEditors/SubTrackEditorBase.cpp
new file mode 100644
--- /dev/null
+++ b/Engine/Source/Editor/MovieSceneTools/Private/TrackEditors/SubTrackEditorBase.cpp
@@ -0,0 +1,347 @@
+// Copyright 1998-2019 Epic Games, Inc. All Rights Reserved.
+
+#include "TrackEditors/SubTrackEditorBase.h"
+
+#define LOCTEXT_NAMESPACE "FSubTrackEditorBase"
+
+FSubSectionPainterResult FSubSectionPainterUtil::PaintSection(TSharedPtr<const ISequencer> Sequencer, const UMovieSceneSubSection& SectionObject, FSequencerSectionPainter& InPainter, FSubSectionPainterParams Params)
+{
+    const TRange<FFrameNumber> SectionRange = SectionObject.GetRange();
+    if (SectionRange.GetLowerBound().IsOpen() || SectionRange.GetUpperBound().IsOpen())
+    {
+        return FSSPR_InvalidSection;
+    }
+
+    const int32 SectionSize = MovieScene::DiscreteSize(SectionRange);
+    if (SectionSize <= 0)
+    {
+        return FSSPR_InvalidSection;
+    }
+
+    UMovieSceneSequence* InnerSequence = SectionObject.GetSequence();
+    if (InnerSequence == nullptr)
+    {
+        return FSSPR_NoInnerSequence;
+    }
+
+    const ESlateDrawEffect DrawEffects = InPainter.bParentEnabled ? ESlateDrawEffect::None : ESlateDrawEffect::DisabledEffect;
+    if (SectionObject.Parameters.bCanLoop)
+    {
+        DoPaintLoopingSection(SectionObject, *InnerSequence, InPainter, DrawEffects);
+    }
+    else
+    {
+        DoPaintNonLoopingSection(SectionObject, *InnerSequence, InPainter, DrawEffects);
+    }
+
+    UMovieScene* MovieScene = InnerSequence->GetMovieScene();
+    const int32 NumTracks = MovieScene->GetPossessableCount() + MovieScene->GetSpawnableCount() + MovieScene->GetMasterTracks().Num();
+
+    FVector2D TopLeft = InPainter.SectionGeometry.AbsoluteToLocal(InPainter.SectionClippingRect.GetTopLeft()) + FVector2D(1.f, -1.f);
+
+    FSlateFontInfo FontInfo = FEditorStyle::GetFontStyle("NormalFont");
+
+    TSharedRef<FSlateFontCache> FontCache = FSlateApplication::Get().GetRenderer()->GetFontCache();
+
+    auto GetFontHeight = [&]
+    {
+        return FontCache->GetMaxCharacterHeight(FontInfo, 1.f) + FontCache->GetBaseline(FontInfo, 1.f);
+    };
+    while (GetFontHeight() > InPainter.SectionGeometry.Size.Y && FontInfo.Size > 11)
+    {
+        FontInfo.Size = FMath::Max(FMath::FloorToInt(FontInfo.Size - 6.f), 11);
+    }
+
+    uint32 LayerId = InPainter.LayerId;
+    FMargin ContentPadding = Params.ContentPadding;
+
+    if (Params.bShowTrackNum)
+    {
+        FSlateDrawElement::MakeText(
+                InPainter.DrawElements,
+                ++LayerId,
+                InPainter.SectionGeometry.MakeChild(
+                    FVector2D(InPainter.SectionGeometry.Size.X, GetFontHeight()),
+                    FSlateLayoutTransform(TopLeft + FVector2D(ContentPadding.Left, ContentPadding.Top) + FVector2D(11.f, GetFontHeight()*2.f))
+                    ).ToPaintGeometry(),
+                FText::Format(LOCTEXT("NumTracksFormat", "{0} track(s)"), FText::AsNumber(NumTracks)),
+                FontInfo,
+                DrawEffects,
+                FColor(200, 200, 200)
+                );
+    }
+
+    if (Params.bDrawFrameNumberHintWhenSelected && InPainter.bIsSelected && Sequencer.IsValid())
+    {
+        FFrameTime CurrentTime = Sequencer->GetLocalTime().Time;
+        if (SectionRange.Contains(CurrentTime.FrameNumber))
+        {
+            UMovieScene* SubSequenceMovieScene = SectionObject.GetSequence()->GetMovieScene();
+            const FFrameRate DisplayRate = SubSequenceMovieScene->GetDisplayRate();
+            const FFrameRate TickResolution = SubSequenceMovieScene->GetTickResolution();
+            const FFrameNumber CurrentFrameNumber = ConvertFrameTime(CurrentTime * SectionObject.OuterToInnerTransform(), TickResolution, DisplayRate).FloorToFrame();
+
+            DrawFrameNumberHint(InPainter, CurrentTime, CurrentFrameNumber.Value);
+        }
+    }
+
+    InPainter.LayerId = LayerId;
+
+    return FSSPR_Success;
+}
+
+void FSubSectionPainterUtil::DoPaintNonLoopingSection(const UMovieSceneSubSection& SectionObject, const UMovieSceneSequence& InnerSequence, FSequencerSectionPainter& InPainter, ESlateDrawEffect DrawEffects)
+{
+    const FFrameNumber SectionStartFrame = SectionObject.GetInclusiveStartFrame();
+
+    const TRange<FFrameNumber> SectionRange = SectionObject.GetRange();
+    const int32 SectionSize = MovieScene::DiscreteSize(SectionRange);
+    const float PixelsPerFrame = InPainter.SectionGeometry.Size.X / float(SectionSize);
+
+    UMovieScene* MovieScene = InnerSequence.GetMovieScene();
+    TRange<FFrameNumber> PlaybackRange = MovieScene->GetPlaybackRange();
+
+    // We're in the non-looping case so we know we have a purely linear transform.
+    const FMovieSceneSequenceTransform InnerToOuterTransform = SectionObject.OuterToInnerTransform().InverseLinearOnly();
+    const FFrameNumber PlaybackStart = (MovieScene::DiscreteInclusiveLower(PlaybackRange) * InnerToOuterTransform).FloorToFrame();
+    if (SectionRange.Contains(PlaybackStart))
+    {
+        const int32 StartOffset = (PlaybackStart - SectionStartFrame).Value;
+        // add dark tint for left out-of-bounds
+        FSlateDrawElement::MakeBox(
+                InPainter.DrawElements,
+                InPainter.LayerId++,
+                InPainter.SectionGeometry.ToPaintGeometry(
+                    FVector2D(0.0f, 0.f),
+                    FVector2D(StartOffset * PixelsPerFrame, InPainter.SectionGeometry.Size.Y)
+                    ),
+                FEditorStyle::GetBrush("WhiteBrush"),
+                DrawEffects,
+                FLinearColor::Black.CopyWithNewOpacity(0.5f)
+                );
+
+        // add green line for playback start
+        FSlateDrawElement::MakeBox(
+                InPainter.DrawElements,
+                InPainter.LayerId++,
+                InPainter.SectionGeometry.ToPaintGeometry(
+                    FVector2D(StartOffset * PixelsPerFrame, 0.f),
+                    FVector2D(1.0f, InPainter.SectionGeometry.Size.Y)
+                    ),
+                FEditorStyle::GetBrush("WhiteBrush"),
+                DrawEffects,
+                FColor(32, 128, 32)	// 120, 75, 50 (HSV)
+                );
+    }
+
+    const FFrameNumber PlaybackEnd = (MovieScene::DiscreteExclusiveUpper(PlaybackRange) * InnerToOuterTransform) .FloorToFrame();
+    if (SectionRange.Contains(PlaybackEnd))
+    {
+        // add dark tint for right out-of-bounds
+        const int32 EndOffset = (PlaybackEnd - SectionStartFrame).Value;
+        FSlateDrawElement::MakeBox(
+                InPainter.DrawElements,
+                InPainter.LayerId++,
+                InPainter.SectionGeometry.ToPaintGeometry(
+                    FVector2D(EndOffset * PixelsPerFrame, 0.f),
+                    FVector2D((SectionSize - EndOffset) * PixelsPerFrame, InPainter.SectionGeometry.Size.Y)
+                    ),
+                FEditorStyle::GetBrush("WhiteBrush"),
+                DrawEffects,
+                FLinearColor::Black.CopyWithNewOpacity(0.5f)
+                );
+
+
+        // add red line for playback end
+        FSlateDrawElement::MakeBox(
+                InPainter.DrawElements,
+                InPainter.LayerId++,
+                InPainter.SectionGeometry.ToPaintGeometry(
+                    FVector2D(EndOffset * PixelsPerFrame, 0.f),
+                    FVector2D(1.0f, InPainter.SectionGeometry.Size.Y)
+                    ),
+                FEditorStyle::GetBrush("WhiteBrush"),
+                DrawEffects,
+                FColor(128, 32, 32)	// 0, 75, 50 (HSV)
+                );
+    }
+}
+
+void FSubSectionPainterUtil::DoPaintLoopingSection(const UMovieSceneSubSection& SectionObject, const UMovieSceneSequence& InnerSequence, FSequencerSectionPainter& InPainter, ESlateDrawEffect DrawEffects)
+{
+    const FFrameNumber SectionStartFrame = SectionObject.GetInclusiveStartFrame();
+    const FFrameNumber SectionEndFrame   = SectionObject.GetExclusiveEndFrame();
+
+    const TRange<FFrameNumber> SectionRange = SectionObject.GetRange();
+    const int32 SectionSize = MovieScene::DiscreteSize(SectionRange);
+    const float PixelsPerFrame = InPainter.SectionGeometry.Size.X / float(SectionSize);
+
+    const float InvTimeScale = FMath::IsNearlyZero(SectionObject.Parameters.TimeScale) ? 1.0f : 1.0f / SectionObject.Parameters.TimeScale;
+
+    const UMovieScene* MovieScene = InnerSequence.GetMovieScene();
+    const TRange<FFrameNumber> InnerPlaybackRange = UMovieSceneSubSection::GetValidatedInnerPlaybackRange(SectionObject.Parameters, *MovieScene);
+
+    const FFrameNumber InnerSubSeqLength = MovieScene::DiscreteSize(InnerPlaybackRange);
+    const FFrameNumber InnerSubSeqFirstLoopLength = InnerSubSeqLength - SectionObject.Parameters.FirstLoopStartFrameOffset;
+
+    const FFrameRate OuterFrameRate   = SectionObject.GetTypedOuter<UMovieScene>()->GetTickResolution();
+    const FFrameRate InnerFrameRate   = MovieScene->GetTickResolution();
+    const FFrameNumber OuterSubSeqLength = (ConvertFrameTime(InnerSubSeqLength, InnerFrameRate, OuterFrameRate) * InvTimeScale).FrameNumber;
+    const FFrameNumber OuterSubSeqFirstLoopLength = (ConvertFrameTime(InnerSubSeqFirstLoopLength, InnerFrameRate, OuterFrameRate) * InvTimeScale).FrameNumber;
+
+    FFrameNumber CurOffsetFrame = FMath::Max(OuterSubSeqFirstLoopLength, FFrameNumber(0));
+
+    // Draw separators where the sub-sequence is looping. To be consistent with the non-looping case, we draw a red and green
+    // separator back to back.
+    uint32 MaxLoopBoundaries = 100;
+    while (CurOffsetFrame < SectionEndFrame)
+    {
+        const int32 CurOffset = CurOffsetFrame.Value;
+        FSlateDrawElement::MakeBox(
+            InPainter.DrawElements,
+            InPainter.LayerId++,
+            InPainter.SectionGeometry.ToPaintGeometry(
+                FVector2D(CurOffset * PixelsPerFrame, 0.f),
+                FVector2D(1.0f, InPainter.SectionGeometry.Size.Y)
+            ),
+            FEditorStyle::GetBrush("WhiteBrush"),
+            DrawEffects,
+            FColor(32, 128, 32)	// 120, 75, 50 (HSV)
+        );
+        if (CurOffset > 0)
+        {
+            FSlateDrawElement::MakeBox(
+                InPainter.DrawElements,
+                InPainter.LayerId++,
+                InPainter.SectionGeometry.ToPaintGeometry(
+                    FVector2D(CurOffset * PixelsPerFrame - 1.f, 0.f),
+                    FVector2D(1.0f, InPainter.SectionGeometry.Size.Y)
+                ),
+                FEditorStyle::GetBrush("WhiteBrush"),
+                DrawEffects,
+                FColor(128, 32, 32)	// 0, 75, 50 (HSV)
+            );
+        }
+
+        CurOffsetFrame += OuterSubSeqLength;
+
+        if ((--MaxLoopBoundaries) == 0)
+        {
+            break;
+        }
+    }
+}
+
+FSubSectionEditorUtil::FSubSectionEditorUtil(UMovieSceneSubSection& InSection)
+    : SectionObject(InSection)
+    , InitialStartOffsetDuringResize(0)
+    , InitialStartTimeDuringResize(0)
+{
+}
+
+void FSubSectionEditorUtil::BeginResizeSection()
+{
+    InitialStartOffsetDuringResize = SectionObject.Parameters.bCanLoop ? SectionObject.Parameters.FirstLoopStartFrameOffset : SectionObject.Parameters.StartFrameOffset;
+    InitialStartTimeDuringResize = SectionObject.HasStartFrame() ? SectionObject.GetInclusiveStartFrame() : 0;
+}
+
+FFrameNumber FSubSectionEditorUtil::ResizeSection(ESequencerSectionResizeMode ResizeMode, FFrameNumber ResizeTime)
+{
+    UMovieSceneSequence* InnerSequence = SectionObject.GetSequence();
+    UMovieScene* InnerMovieScene = InnerSequence ? InnerSequence->GetMovieScene() : nullptr;
+
+    if (ResizeMode == SSRM_LeadingEdge && InnerMovieScene != nullptr)
+    {
+        // Adjust the start offset when resizing from the beginning
+        FMovieSceneSectionParameters& SectionParameters = SectionObject.Parameters;
+
+        const FFrameRate    OuterFrameRate   = SectionObject.GetTypedOuter<UMovieScene>()->GetTickResolution();
+        const FFrameRate    InnerFrameRate   = InnerSequence->GetMovieScene()->GetTickResolution();
+        const FFrameNumber  ResizeDifference = ResizeTime - InitialStartTimeDuringResize;
+
+        const FFrameNumber InnerResizeDifference = (ConvertFrameTime(ResizeDifference, OuterFrameRate, InnerFrameRate) * SectionParameters.TimeScale).FrameNumber;
+        FFrameNumber NewStartOffset = InitialStartOffsetDuringResize + InnerResizeDifference;
+
+        const int32 InnerPlaybackLength = MovieScene::DiscreteSize(InnerMovieScene->GetPlaybackRange());
+        const float InvTimeScale = FMath::IsNearlyZero(SectionParameters.TimeScale) ? 1.0f : 1.0f / SectionParameters.TimeScale;
+        const FFrameNumber InnerLoopLength = InnerPlaybackLength - SectionParameters.StartFrameOffset - SectionParameters.EndFrameOffset;
+
+        const bool bCanLoop = SectionParameters.bCanLoop;
+        if (!bCanLoop)
+        {
+            if (NewStartOffset < 0)
+            {
+                // Ensure start offset is not less than 0, clamp resize time.
+                const FFrameTime OuterFrameTimeOver = ConvertFrameTime(FFrameTime::FromDecimal(NewStartOffset.Value / SectionParameters.TimeScale), InnerFrameRate, OuterFrameRate);
+                ResizeTime = ResizeTime - OuterFrameTimeOver.GetFrame(); 
+                NewStartOffset = 0;
+            }
+            SectionParameters.StartFrameOffset = NewStartOffset;
+        }
+        else
+        {
+			// If the start offset exceeds the length of one loop, trim it back.
+			NewStartOffset = NewStartOffset % InnerLoopLength;
+            if (NewStartOffset < 0)
+            {
+                // Move the first loop offset forward so that it's positive into the new first loop
+                // that we just "revealed" by pulling the left edge back.
+                NewStartOffset += InnerLoopLength;
+            }
+            SectionParameters.FirstLoopStartFrameOffset = NewStartOffset;
+        }
+    }
+
+    return ResizeTime;
+}
+
+void FSubSectionEditorUtil::BeginSlipSection()
+{
+    // Cache the same values as when resizing.
+    BeginResizeSection();
+}
+
+FFrameNumber FSubSectionEditorUtil::SlipSection(FFrameNumber SlipTime)
+{
+    UMovieSceneSequence* InnerSequence = SectionObject.GetSequence();
+    UMovieScene* InnerMovieScene = InnerSequence ? InnerSequence->GetMovieScene() : nullptr;
+
+    if (InnerMovieScene != nullptr)
+    {
+        FMovieSceneSectionParameters& SectionParameters = SectionObject.Parameters;
+
+        const FFrameRate    OuterFrameRate   = SectionObject.GetTypedOuter<UMovieScene>()->GetTickResolution();
+        const FFrameRate    InnerFrameRate   = InnerSequence->GetMovieScene()->GetTickResolution();
+        const FFrameNumber  ResizeDifference = SlipTime - InitialStartTimeDuringResize;
+
+        const FFrameNumber InnerResizeDifference = (ConvertFrameTime(ResizeDifference, OuterFrameRate, InnerFrameRate) * SectionParameters.TimeScale).FrameNumber;
+        FFrameNumber NewStartOffset = InitialStartOffsetDuringResize + InnerResizeDifference;
+
+        const int32 InnerPlaybackLength = MovieScene::DiscreteSize(InnerMovieScene->GetPlaybackRange());
+        const float InvTimeScale = FMath::IsNearlyZero(SectionParameters.TimeScale) ? 1.0f : 1.0f / SectionParameters.TimeScale;
+        const FFrameNumber InnerLoopLength = InnerPlaybackLength - SectionParameters.StartFrameOffset - SectionParameters.EndFrameOffset;
+
+        const bool bCanLoop = SectionParameters.bCanLoop;
+        if (!bCanLoop)
+        {
+            // Ensure start offset is not less than 0
+            SectionParameters.StartFrameOffset = FFrameNumber(FMath::Max(NewStartOffset.Value, 0));
+        }
+        else
+        {
+			// If the start offset exceeds the length of one loop, trim it back.
+            NewStartOffset = NewStartOffset % InnerLoopLength;
+            if (NewStartOffset < 0)
+            {
+                // Move the first loop offset forward so that it's positive into the new first loop
+                // that we just "revealed" by slipping too much to the side.
+                NewStartOffset += InnerLoopLength;
+            }
+            SectionParameters.FirstLoopStartFrameOffset = NewStartOffset;
+        }
+    }
+
+    return SlipTime;
+}
+
+#undef LOCTEXT_NAMESPACE
diff --git a/Engine/Source/Editor/MovieSceneTools/Public/TrackEditors/SubTrackEditorBase.h b/Engine/Source/Editor/MovieSceneTools/Public/TrackEditors/SubTrackEditorBase.h
new file mode 100644
--- /dev/null
+++ b/Engine/Source/Editor/MovieSceneTools/Public/TrackEditors/SubTrackEditorBase.h
@@ -0,0 +1,231 @@
+// Copyright 1998-2019 Epic Games, Inc. All Rights Reserved.
+
+#pragma once
+
+#include "CoreMinimal.h"
+#include "ISequencerSection.h"
+#include "Sections/MovieSceneSubSection.h"
+#include "CommonMovieSceneTools.h"
+#include "EditorStyleSet.h"
+#include "Framework/Application/SlateApplication.h"
+#include "ISequencer.h"
+#include "Misc/QualifiedFrameTime.h"
+#include "MovieSceneSequence.h"
+#include "MovieSceneTimeHelpers.h"
+#include "Rendering/DrawElements.h"
+#include "SequencerSectionPainter.h"
+
+class ISequencer;
+
+/**
+ * Parameters for painting a sub-section.
+ */
+struct FSubSectionPainterParams
+{
+    FSubSectionPainterParams() : bShowTrackNum(true), bDrawFrameNumberHintWhenSelected(true)
+    {}
+    FSubSectionPainterParams(FMargin InContentPadding)
+	: ContentPadding(InContentPadding), bShowTrackNum(true), bDrawFrameNumberHintWhenSelected(true)
+    {}
+
+    FMargin ContentPadding;
+    bool bShowTrackNum;
+    bool bDrawFrameNumberHintWhenSelected;
+};
+
+/**
+ * Enum indicating the result of painting the sub-section.
+ */
+enum FSubSectionPainterResult
+{
+    FSSPR_Success,
+    FSSPR_InvalidSection,
+    FSSPR_NoInnerSequence
+};
+
+/**
+ * Utility class for drawing sub-sequence sections.
+ */
+class MOVIESCENETOOLS_API FSubSectionPainterUtil
+{
+public:
+    /** Paints the sub-section, mostly by painting loop boundaries, if appropriate */
+    static FSubSectionPainterResult PaintSection(TSharedPtr<const ISequencer> Sequencer, const UMovieSceneSubSection& SectionObject, FSequencerSectionPainter& InPainter, FSubSectionPainterParams Params);
+
+private:
+    static void DoPaintNonLoopingSection(const UMovieSceneSubSection& SectionObject, const UMovieSceneSequence& InnerSequence, FSequencerSectionPainter& InPainter, ESlateDrawEffect DrawEffects);
+    static void DoPaintLoopingSection(const UMovieSceneSubSection& SectionObject, const UMovieSceneSequence& InnerSequence, FSequencerSectionPainter& InPainter, ESlateDrawEffect DrawEffects);
+};
+
+/**
+ * Utility class for editing (resize/slip) sub-sequence sections.
+ */
+class MOVIESCENETOOLS_API FSubSectionEditorUtil
+{
+public:
+    FSubSectionEditorUtil(UMovieSceneSubSection& InSection);
+
+    /** Starts a resize operation */
+    void BeginResizeSection();
+    /** Resizes the section, returns the new resize time */
+    FFrameNumber ResizeSection(ESequencerSectionResizeMode ResizeMode, FFrameNumber ResizeTime);
+
+    /** Starts a new slip operation */
+    void BeginSlipSection();
+    /** Slips the section, returns the new slip time */
+    FFrameNumber SlipSection(FFrameNumber SlipTime);
+
+private:
+    /** The section object this utility class is editing */
+    UMovieSceneSubSection& SectionObject;
+
+    /** Cached start offset value valid only during resize */
+    FFrameNumber InitialStartOffsetDuringResize;
+
+    /** Cached start time valid only during resize */
+    FFrameNumber InitialStartTimeDuringResize;
+};
+
+/**
+ * Mixin class for sub-sequence section interfaces.
+ *
+ * This mixin effectively wraps the 2 utility classes above, hooks them up to the ISequencerSection interface,
+ * and implements a few other ISequencerSection interface methods. As such, it's mostly handy boilerplate code,
+ * which you can avoid (if it doesn't fit your use-case) in favour of using the uility classes above directly.
+ */
+template<typename ParentSectionClass=ISequencerSection>
+class TSubSectionMixin
+	: public ParentSectionClass
+{
+public:
+    template<typename... ParentCtorParams>
+    TSubSectionMixin(TSharedPtr<ISequencer> InSequencer, UMovieSceneSubSection& InSection, ParentCtorParams&&... Params);
+
+    // ISequencerSection interface
+    virtual UMovieSceneSection* GetSectionObject() override;
+    virtual FText GetSectionTitle() const override;
+    virtual bool IsReadOnly() const override;
+    virtual int32 OnPaintSection( FSequencerSectionPainter& InPainter ) const override;
+    virtual FReply OnSectionDoubleClicked(const FGeometry& SectionGeometry, const FPointerEvent& MouseEvent) override;
+    virtual void BeginResizeSection() override;
+    virtual void ResizeSection(ESequencerSectionResizeMode ResizeMode, FFrameNumber ResizeTime) override;
+    virtual void BeginSlipSection() override;
+    virtual void SlipSection(FFrameNumber SlipTime) override;
+
+protected:
+    TSharedPtr<ISequencer> GetSequencer() const { return SequencerPtr.Pin(); }
+    UMovieSceneSubSection& GetSubSectionObject() { return SubSectionObject; }
+    const UMovieSceneSubSection& GetSubSectionObject() const { return SubSectionObject; }
+    FSubSectionEditorUtil& GetEditorUtil() { return EditorUtil; }
+
+    template<typename SectionType>
+    SectionType& GetSectionObjectAs() { return *CastChecked<SectionType>(&SubSectionObject); }
+    template<typename SectionType>
+    const SectionType& GetSectionObjectAs() const { return *CastChecked<SectionType>(&SubSectionObject); }
+
+private:
+    /** The Sequencer this section belongs to */
+    TWeakPtr<ISequencer> SequencerPtr;
+
+    /** The section object this interface represents */
+    UMovieSceneSubSection& SubSectionObject;
+
+    /** Utility class for resizing/slipping sub-sequence sections */
+    FSubSectionEditorUtil EditorUtil;
+};
+
+template<typename ParentSectionClass>
+template<typename... ParentCtorParams>
+TSubSectionMixin<ParentSectionClass>::TSubSectionMixin(TSharedPtr<ISequencer> InSequencer, UMovieSceneSubSection& InSection, ParentCtorParams&&... Params)
+    : ParentSectionClass(std::forward<ParentCtorParams>(Params)...)
+    , SequencerPtr(InSequencer)
+    , SubSectionObject(InSection)
+    , EditorUtil(InSection)
+{}
+
+template<typename ParentSectionClass>
+UMovieSceneSection* TSubSectionMixin<ParentSectionClass>::GetSectionObject()
+{
+    return &SubSectionObject;
+}
+
+#define LOCTEXT_NAMESPACE "FSubTrackEditorBase"
+template<typename ParentSectionClass>
+FText TSubSectionMixin<ParentSectionClass>::GetSectionTitle() const
+{
+    if (SubSectionObject.GetSequence())
+    {
+        return FText::FromString(SubSectionObject.GetSequence()->GetName());
+    }
+    else
+    {
+        return LOCTEXT("NoSequenceSelected", "No Sequence Selected");
+    }
+}
+#undef LOCTEXT_NAMESPACE
+
+template<typename ParentSectionClass>
+bool TSubSectionMixin<ParentSectionClass>::IsReadOnly() const
+{
+	return SubSectionObject.IsReadOnly();
+}
+
+template<typename ParentSectionClass>
+int32 TSubSectionMixin<ParentSectionClass>::OnPaintSection(FSequencerSectionPainter& InPainter) const
+{
+    InPainter.LayerId = InPainter.PaintSectionBackground();
+
+    FSubSectionPainterUtil::PaintSection(GetSequencer(), SubSectionObject, InPainter, FSubSectionPainterParams(GetContentPadding()));
+
+    return InPainter.LayerId;
+}
+
+template<typename ParentSectionClass>
+FReply TSubSectionMixin<ParentSectionClass>::OnSectionDoubleClicked(const FGeometry& SectionGeometry, const FPointerEvent& MouseEvent)
+{
+    if( MouseEvent.GetEffectingButton() == EKeys::LeftMouseButton )
+    {
+        if (SubSectionObject.GetSequence())
+        {
+	    if (MouseEvent.IsControlDown())
+	    {
+		GEditor->GetEditorSubsystem<UAssetEditorSubsystem>()->OpenEditorForAsset(SubSectionObject.GetSequence());
+	    }
+	    else
+	    {
+                GetSequencer()->FocusSequenceInstance(SubSectionObject);
+	    }
+        }
+    }
+
+    return FReply::Handled();
+}
+
+template<typename ParentSectionClass>
+void TSubSectionMixin<ParentSectionClass>::BeginResizeSection()
+{
+    EditorUtil.BeginResizeSection();
+    ParentSectionClass::BeginResizeSection();
+}
+
+template<typename ParentSectionClass>
+void TSubSectionMixin<ParentSectionClass>::ResizeSection(ESequencerSectionResizeMode ResizeMode, FFrameNumber ResizeTime)
+{
+    ResizeTime = EditorUtil.ResizeSection(ResizeMode, ResizeTime);
+    ParentSectionClass::ResizeSection(ResizeMode, ResizeTime);
+}
+
+template<typename ParentSectionClass>
+void TSubSectionMixin<ParentSectionClass>::BeginSlipSection()
+{
+    EditorUtil.BeginSlipSection();
+    ParentSectionClass::BeginSlipSection();
+}
+
+template<typename ParentSectionClass>
+void TSubSectionMixin<ParentSectionClass>::SlipSection(FFrameNumber SlipTime)
+{
+    SlipTime = EditorUtil.SlipSection(SlipTime);
+    ParentSectionClass::SlipSection(SlipTime);
+}
+
