Sequencer: Try to fix VS2019 build.

diff --git a/Engine/Source/Editor/MovieSceneTools/Public/TrackEditors/SubTrackEditorBase.h b/Engine/Source/Editor/MovieSceneTools/Public/TrackEditors/SubTrackEditorBase.h
--- a/Engine/Source/Editor/MovieSceneTools/Public/TrackEditors/SubTrackEditorBase.h
+++ b/Engine/Source/Editor/MovieSceneTools/Public/TrackEditors/SubTrackEditorBase.h
@@ -14,6 +14,7 @@
 #include "MovieSceneTimeHelpers.h"
 #include "Rendering/DrawElements.h"
 #include "SequencerSectionPainter.h"
+#include "Subsystems/AssetEditorSubsystem.h"
 
 class ISequencer;
 
@@ -175,7 +176,7 @@ int32 TSubSectionMixin<ParentSectionClass>::OnPaintSection(FSequencerSectionPain
 {
     InPainter.LayerId = InPainter.PaintSectionBackground();
 
-    FSubSectionPainterUtil::PaintSection(GetSequencer(), SubSectionObject, InPainter, FSubSectionPainterParams(GetContentPadding()));
+    FSubSectionPainterUtil::PaintSection(this->GetSequencer(), SubSectionObject, InPainter, FSubSectionPainterParams(this->GetContentPadding()));
 
     return InPainter.LayerId;
 }
@@ -193,7 +194,7 @@ FReply TSubSectionMixin<ParentSectionClass>::OnSectionDoubleClicked(const FGeome
 	    }
 	    else
 	    {
-                GetSequencer()->FocusSequenceInstance(SubSectionObject);
+                this->GetSequencer()->FocusSequenceInstance(SubSectionObject);
 	    }
         }
     }
