Sequencer: Disable drawing the number of tracks for cinematic shots (looks like it wasn't drawn before CL#9882159 and is now being drawn cutoff)

#rnx

diff --git a/Engine/Source/Editor/MovieSceneTools/Private/Sections/CinematicShotSection.cpp b/Engine/Source/Editor/MovieSceneTools/Private/Sections/CinematicShotSection.cpp
--- a/Engine/Source/Editor/MovieSceneTools/Private/Sections/CinematicShotSection.cpp
+++ b/Engine/Source/Editor/MovieSceneTools/Private/Sections/CinematicShotSection.cpp
@@ -248,8 +248,12 @@ int32 FCinematicShotSection::OnPaintSection(FSequencerSectionPainter& InPainter)
 	FViewportThumbnailSection::OnPaintSection(InPainter);
 
 	// Paint the sub-sequence information/looping boundaries/etc.
+
+	FSubSectionPainterParams SubSectionPainterParams(GetContentPadding());
+	SubSectionPainterParams.bShowTrackNum = false;
+
 	FSubSectionPainterUtil::PaintSection(
-			GetSequencer(), SectionObject, InPainter, FSubSectionPainterParams(GetContentPadding()));
+			GetSequencer(), SectionObject, InPainter, SubSectionPainterParams);
 
 	return InPainter.LayerId;
 }
