Sequencer: Only render waveform preview for whatever if visible currently.

diff --git a/Engine/Source/Editor/MovieSceneTools/Private/TrackEditors/AudioTrackEditor.cpp b/Engine/Source/Editor/MovieSceneTools/Private/TrackEditors/AudioTrackEditor.cpp
--- a/Engine/Source/Editor/MovieSceneTools/Private/TrackEditors/AudioTrackEditor.cpp
+++ b/Engine/Source/Editor/MovieSceneTools/Private/TrackEditors/AudioTrackEditor.cpp
@@ -334,49 +334,33 @@ void FAudioThumbnail::GenerateWaveformPreview(TArray<uint8>& OutData, TRange<flo
 		SectionStartTime - FrameRate.AsSeconds(AudioSection->GetStartOffset()) + DeriveUnloopedDuration(AudioSection) * (1.0f / PitchMultiplierValue));
 
 	float TrueRangeSize = AudioTrueRange.Size<float>();
-	const int32 TrueDrawOffsetPx = FMath::Max(FMath::RoundToInt((DrawRange.GetLowerBoundValue() - SectionStartTime) / DisplayScale), 0);
-	const int32 LastTrueSample = -2.f*SmoothingAmount + FMath::TruncToInt(TrueRangeSize / DisplayScale);
-
-	if (LastTrueSample <= 0)
-	{
-		return;
-	}
-
-	DrawRange = AudioTrueRange;
-
 	float DrawRangeSize = DrawRange.Size<float>();
 
 	const int32 MaxAmplitude = NumChannels == 1 ? GetSize().Y : GetSize().Y / 2;
-
-	int32 DrawOffsetPx = FMath::Max(FMath::RoundToInt((DrawRange.GetLowerBoundValue() - SectionStartTime) / DisplayScale), 0);
-
+	const int32 DrawOffsetPx = FMath::Max(FMath::RoundToInt((DrawRange.GetLowerBoundValue() - SectionStartTime) / DisplayScale), 0);
 
 	// In order to prevent flickering waveforms when moving the display position/range around, we have to lock our sample position and spline segments to the view range
-	float RangeLookupFraction = (SmoothingAmount*DisplayScale) / TrueRangeSize;
-	int32 LookupRange = FMath::Clamp(FMath::TruncToInt(RangeLookupFraction * LookupSize), 1, LookupSize);
-
-	int32 SampleLockOffset = TrueDrawOffsetPx % SmoothingAmount;
+	const float RangeLookupFraction = (SmoothingAmount * DisplayScale) / TrueRangeSize;
+	const int32 LookupRange = FMath::Clamp(FMath::TruncToInt(RangeLookupFraction * LookupSize), 1, LookupSize);
 
-	int32 FirstSample = -2.f*SmoothingAmount - SampleLockOffset;
-	int32 LastSample = LastTrueSample + 2*SmoothingAmount;
-
-	{
-		// @todo: when SampleCount <= 0, we have fewer samples than pixels, and should start to interpolate the spline by that distance, rather than a hard coded pixel density
-		int32 NumSamplesInRange = FMath::TruncToInt(LookupSize * (DrawRangeSize /LastTrueSample) / TrueRangeSize);
-		int32 SampleCount = NumSamplesInRange / NumChannels;
-	}
+	const int32 SampleLockOffset = DrawOffsetPx % SmoothingAmount;
 
+	const FIntPoint ThumbnailSize = GetSize();
+	const int32 FirstSample = -2.f * SmoothingAmount - SampleLockOffset;
+	const int32 LastSample = ThumbnailSize.X + 2.f * SmoothingAmount;
 
 	// Sample the audio one pixel to the left and right
 	for (int32 X = FirstSample; X < LastSample; ++X)
 	{
-		float LookupTime = ((float)(X - 0.5f) / (float)LastTrueSample) * DrawRangeSize + DrawRange.GetLowerBoundValue();
+		float LookupTime = ((float)(X - 0.5f) / (float)ThumbnailSize.X) * DrawRangeSize + DrawRange.GetLowerBoundValue();
 		float LookupFraction = (LookupTime - AudioTrueRange.GetLowerBoundValue()) / TrueRangeSize;
-		int32 LookupIndex = FMath::TruncToInt(LookupFraction * LookupSize);
+		float LookupFractionLooping = FMath::Fmod(LookupFraction, 1.f);
+		int32 LookupIndex = FMath::TruncToInt(LookupFractionLooping * LookupSize);
 		
-		float NextLookupTime = ((float)(X + 0.5f) / (float)LastTrueSample) * DrawRangeSize + DrawRange.GetLowerBoundValue();
+		float NextLookupTime = ((float)(X + 0.5f) / (float)ThumbnailSize.X) * DrawRangeSize + DrawRange.GetLowerBoundValue();
 		float NextLookupFraction = (NextLookupTime - AudioTrueRange.GetLowerBoundValue()) / TrueRangeSize;
-		int32 NextLookupIndex = FMath::TruncToInt(NextLookupFraction  * LookupSize);
+		float NextLookupFractionLooping = FMath::Fmod(NextLookupFraction, 1.f);
+		int32 NextLookupIndex = FMath::TruncToInt(NextLookupFractionLooping * LookupSize);
 		
 		SampleAudio(SoundWave->NumChannels, LookupData, LookupIndex, NextLookupIndex, LookupSize, MaxAmplitude);
 	}
@@ -385,14 +369,11 @@ void FAudioThumbnail::GenerateWaveformPreview(TArray<uint8>& OutData, TRange<flo
 	GenerateSpline(SoundWave->NumChannels, FirstSample);
 
 	// Now draw the spline
-	const int32 Height = GetSize().Y;
-	const int32 Width = LastTrueSample;
+	const int32 Width = ThumbnailSize.X;
+	const int32 Height = ThumbnailSize.Y;
 
 	FLinearColor BoundaryColor = BoundaryColorHSV.HSVToLinearRGB();
 
-	uint32 Size = LastTrueSample * GetSize().Y * GPixelFormats[PF_B8G8R8A8].BlockBytes;
-	TArray<uint8> TempData;
-	TempData.SetNum(Size);
 	for (int32 ChannelIndex = 0; ChannelIndex < SoundWave->NumChannels; ++ChannelIndex)
 	{
 		int32 SplineIndex = 0;
@@ -427,7 +408,7 @@ void FAudioThumbnail::GenerateWaveformPreview(TArray<uint8>& OutData, TRange<flo
 
 			for (int32 PixelIndex = 0; PixelIndex < MaxAmplitude; ++PixelIndex)
 			{
-				uint8* Pixel = LookupPixel(TempData, X, PixelIndex, Width, Height, ChannelIndex, NumChannels);
+				uint8* Pixel = LookupPixel(OutData, X, PixelIndex, Width, Height, ChannelIndex, NumChannels);
 
 				const float PixelCenter = PixelIndex + 0.5f;
 
@@ -458,28 +439,6 @@ void FAudioThumbnail::GenerateWaveformPreview(TArray<uint8>& OutData, TRange<flo
 			}
 		}
 	}
-
-	// Then loop the texture with the start offset
-	int32 TotalDrawOffsetPx = TrueDrawOffsetPx + FMath::RoundToInt(FrameRate.AsSeconds(AudioSection->GetStartOffset()) / DisplayScale);
-	for (int32 ChannelIndex = 0; ChannelIndex < SoundWave->NumChannels; ++ChannelIndex)
-	{
-		int32 SplineIndex = 0;
-		for (int32 X = 0; X < GetSize().X; ++X)
-		{
-			int32 XRot = ((X + TotalDrawOffsetPx) % LastTrueSample);
-
-			for (int32 PixelIndex = 0; PixelIndex < MaxAmplitude; ++PixelIndex)
-			{
-				uint8* PixelToCopy = LookupPixel(TempData, XRot, PixelIndex, Width, Height, ChannelIndex, NumChannels);
-				uint8* Pixel = LookupPixel(OutData, X, PixelIndex, GetSize().X, Height, ChannelIndex, NumChannels);
-
-				*Pixel++ = *PixelToCopy++;
-				*Pixel++ = *PixelToCopy++;
-				*Pixel++ = *PixelToCopy++;
-				*Pixel++ = *PixelToCopy++;
-			}
-		}
-	}
 }
 
 void FAudioThumbnail::GenerateSpline(int32 NumChannels, int32 SamplePositionOffset)
