Sequencer: Fixed movie renders not working for scales < 1.0

This check was introduced in 5677319 but has been superseded by a more robust approach.

#note Copied from 4.24[REDACTED_HANDLE]

diff --git a/Engine/Source/Editor/MovieSceneTools/Private/AutomatedLevelSequenceCapture.cpp b/Engine/Source/Editor/MovieSceneTools/Private/AutomatedLevelSequenceCapture.cpp
--- a/Engine/Source/Editor/MovieSceneTools/Private/AutomatedLevelSequenceCapture.cpp
+++ b/Engine/Source/Editor/MovieSceneTools/Private/AutomatedLevelSequenceCapture.cpp
@@ -885,13 +885,9 @@ void UAutomatedLevelSequenceCapture::SequenceUpdated(const UMovieSceneSequencePl
 
 				bool bOnLastFrame = (CurrentTime.FrameNumber >= Actor->SequencePlayer->GetStartTime().Time.FrameNumber + Actor->SequencePlayer->GetFrameDuration() - 1);
 				bool bLastShot = NumShots == 0 ? true : ShotIndex == NumShots - 1;
-				
-				// Prevent the same frame from being rendered twice
-				if (CurrentTime.FrameNumber.Value != CachedMetrics.PreviousFrame)
-				{
-					CaptureThisFrame((CurrentTime - PreviousTime) / Settings.GetFrameRate());
-					CachedMetrics.PreviousFrame = CurrentTime.FrameNumber.Value;
-				}
+
+				CaptureThisFrame((CurrentTime - PreviousTime) / Settings.GetFrameRate());
+				CachedMetrics.PreviousFrame = CurrentTime.FrameNumber.Value;
 
 				// Our callback can be called multiple times for a given frame due to how Level Sequences evaluate.
 				// For example, frame 161 is evaluated and an image is written. This isn't considered the end of the sequence
