Sequencer: Fix rendering of loop markers for audio track sections.

diff --git a/Engine/Source/Editor/MovieSceneTools/Private/TrackEditors/AudioTrackEditor.cpp b/Engine/Source/Editor/MovieSceneTools/Private/TrackEditors/AudioTrackEditor.cpp
--- a/Engine/Source/Editor/MovieSceneTools/Private/TrackEditors/AudioTrackEditor.cpp
+++ b/Engine/Source/Editor/MovieSceneTools/Private/TrackEditors/AudioTrackEditor.cpp
@@ -682,21 +682,20 @@ int32 FAudioSection::OnPaintSection( FSequencerSectionPainter& Painter ) const
 		return LayerId;
 	}
 
-	FFrameRate TickResolution = TimeToPixelConverter.GetTickResolution();
-	float AudioDuration = DeriveUnloopedDuration(AudioSection);
-
 	// Add lines where the animation starts and ends/loops
-	const float SeqLength = AudioDuration - TickResolution.AsSeconds(AudioSection->GetStartOffset());
+	FFrameRate TickResolution = TimeToPixelConverter.GetTickResolution();
+	const float AudioDuration = DeriveUnloopedDuration(AudioSection);
 
-	if (!FMath::IsNearlyZero(SeqLength, KINDA_SMALL_NUMBER) && SeqLength > 0)
+	if (AudioDuration > KINDA_SMALL_NUMBER)
 	{
-		float MaxOffset = Section.GetRange().Size<FFrameTime>() / TickResolution;
-		float OffsetTime = SeqLength;
-		float StartTime = Section.GetInclusiveStartFrame() / TickResolution;
+		const float MaxOffset = Section.GetRange().Size<FFrameTime>() / TickResolution;
+		const float StartOffsetTime = TickResolution.AsSeconds(AudioSection->GetStartOffset());
+		const float SectionStartTime = TickResolution.AsSeconds(AudioSection->GetInclusiveStartFrame());
 
+		float OffsetTime = AudioDuration - StartOffsetTime;
 		while (OffsetTime < MaxOffset)
 		{
-			float OffsetPixel = TimeToPixelConverter.SecondsToPixel(StartTime + OffsetTime) - TimeToPixelConverter.SecondsToPixel(StartTime);
+			float OffsetPixel = TimeToPixelConverter.SecondsToPixel(SectionStartTime + OffsetTime);
 
 			FSlateDrawElement::MakeBox(
 				Painter.DrawElements,
@@ -709,7 +708,7 @@ int32 FAudioSection::OnPaintSection( FSequencerSectionPainter& Painter ) const
 				DrawEffects
 			);
 
-			OffsetTime += SeqLength;
+			OffsetTime += AudioDuration;
 		}
 	}
 
