Sequencer: Remove check for null actor parent. The objects for a parent binding could be null, in which case the attach compensation will just be identity.

diff --git a/Engine/Source/Editor/MovieSceneTools/Private/TrackEditors/AttachTrackEditor.cpp b/Engine/Source/Editor/MovieSceneTools/Private/TrackEditors/AttachTrackEditor.cpp
--- a/Engine/Source/Editor/MovieSceneTools/Private/TrackEditors/AttachTrackEditor.cpp
+++ b/Engine/Source/Editor/MovieSceneTools/Private/TrackEditors/AttachTrackEditor.cpp
@@ -487,7 +487,7 @@ AActor* GetConstraintActor(TSharedPtr<ISequencer> InSequencer, const FMovieScene
 
 	TArrayView<TWeakObjectPtr<UObject>> RuntimeObjects = InSequencer->FindBoundObjects(InConstraintBindingID.GetGuid(), SequenceID);
 
-	if (RuntimeObjects.Num() == 1 && RuntimeObjects[0].IsValid())
+	if (RuntimeObjects.Num() >= 1 && RuntimeObjects[0].IsValid())
 	{
 		return Cast<AActor>(RuntimeObjects[0].Get());
 	}
@@ -1115,8 +1115,8 @@ FKeyPropertyResult F3DAttachTrackEditor::AddKeyInternal( FFrameNumber KeyTime, c
 
 	UMovieScene* MovieScene = GetSequencer()->GetFocusedMovieSceneSequence()->GetMovieScene();
 
+	// It's possible that the objects bound to this parent binding ID are null, in which case there will be no compensation
 	AActor* ParentActor = GetConstraintActor(GetSequencer(), ConstraintBindingID);	
-	check(ParentActor)
 
 	FWorldTransformEvaluator ParentTransformEval(GetSequencer(), ParentActor, SocketName);
 
