Sequencer: Corrected track count display on sub sequences

diff --git a/Engine/Source/Editor/MovieSceneTools/Private/TrackEditors/SubTrackEditor.cpp b/Engine/Source/Editor/MovieSceneTools/Private/TrackEditors/SubTrackEditor.cpp
--- a/Engine/Source/Editor/MovieSceneTools/Private/TrackEditors/SubTrackEditor.cpp
+++ b/Engine/Source/Editor/MovieSceneTools/Private/TrackEditors/SubTrackEditor.cpp
@@ -28,6 +28,7 @@
 #include "MovieSceneTimeHelpers.h"
 #include "EngineAnalytics.h"
 #include "Interfaces/IAnalyticsProvider.h"
+#include "Algo/Accumulate.h"
 
 #include "CommonMovieSceneTools.h"
 
@@ -208,7 +209,9 @@ public:
 
 			FMargin ContentPadding = GetContentPadding();
 
-			int32 NumTracks = MovieScene->GetPossessableCount() + MovieScene->GetSpawnableCount() + MovieScene->GetMasterTracks().Num();
+			const int32 NumTracks = Algo::Accumulate(MovieScene->GetBindings(), 0, [](int32 Count, const FMovieSceneBinding& Binding){ return Count + Binding.GetTracks().Num(); })
+				+ MovieScene->GetMasterTracks().Num()
+				+ (MovieScene->GetCameraCutTrack() ? 1 : 0);
 
 			FVector2D TopLeft = InPainter.SectionGeometry.AbsoluteToLocal(InPainter.SectionClippingRect.GetTopLeft()) + FVector2D(1.f, -1.f);
 
