Sequencer: Audio component and socket names with key editor for actor reference data

Add visualization for audio spatialization in the viewport for selected master audio tracks.
Refactor SComponentChooser

diff --git a/Engine/Source/Editor/MovieSceneTools/Private/Channels/BuiltInChannelEditors.cpp b/Engine/Source/Editor/MovieSceneTools/Private/Channels/BuiltInChannelEditors.cpp
--- a/Engine/Source/Editor/MovieSceneTools/Private/Channels/BuiltInChannelEditors.cpp
+++ b/Engine/Source/Editor/MovieSceneTools/Private/Channels/BuiltInChannelEditors.cpp
@@ -11,6 +11,7 @@
 #include "MovieSceneCommonHelpers.h"
 #include "GameFramework/Actor.h"
 #include "EditorStyleSet.h"
+#include "Styling/CoreStyle.h"
 #include "CurveKeyEditors/SNumericKeyEditor.h"
 #include "CurveKeyEditors/SBoolCurveKeyEditor.h"
 #include "CurveKeyEditors/SStringCurveKeyEditor.h"
@@ -23,6 +24,14 @@
 #include "Channels/FloatChannelCurveModel.h"
 #include "EventChannelCurveModel.h"
 #include "PropertyCustomizationHelpers.h"
+#include "MovieSceneObjectBindingIDCustomization.h"
+#include "MovieSceneObjectBindingIDPicker.h"
+#include "LevelEditor.h"
+#include "Modules/ModuleManager.h"
+#include "Framework/Application/MenuStack.h"
+#include "Framework/Application/SlateApplication.h"
+#include "Editor/SceneOutliner/Private/SSocketChooser.h"
+#include "SComponentChooser.h"
 
 #define LOCTEXT_NAMESPACE "BuiltInChannelEditors"
 
@@ -116,6 +125,11 @@ bool CanCreateKeyEditor(const FMovieSceneObjectPathChannel* Channel)
 	return true;
 }
 
+bool CanCreateKeyEditor(const FMovieSceneActorReferenceData* Channel)
+{
+	return true;
+}
+
 TSharedRef<SWidget> CreateKeyEditor(const TMovieSceneChannelHandle<FMovieSceneBoolChannel>&    Channel, UMovieSceneSection* Section, const FGuid& InObjectBindingID, TWeakPtr<FTrackInstancePropertyBindings> PropertyBindings, TWeakPtr<ISequencer> InSequencer)
 {
 	const TMovieSceneExternalValue<bool>* ExternalValue = Channel.GetExtendedEditorData();
@@ -242,6 +256,225 @@ TSharedRef<SWidget> CreateKeyEditor(const TMovieSceneChannelHandle<FMovieSceneOb
 	return SNullWidget::NullWidget;
 }
 
+/** Delegate used to set a class */
+DECLARE_DELEGATE_OneParam(FOnSetActorReferenceKey, FMovieSceneActorReferenceKey);
+
+class SActorReferenceBox : public SCompoundWidget, public FMovieSceneObjectBindingIDPicker
+{
+public:
+	SLATE_BEGIN_ARGS(SActorReferenceBox)
+	{}
+	SLATE_ATTRIBUTE(FMovieSceneActorReferenceKey, ActorReferenceKey)
+	SLATE_EVENT(FOnSetActorReferenceKey, OnSetActorReferenceKey)
+
+	SLATE_END_ARGS()
+
+	void Construct(const FArguments& InArgs, TWeakPtr<ISequencer> InSequencer)
+	{
+		WeakSequencer = InSequencer;
+		LocalSequenceID = InSequencer.Pin()->GetFocusedTemplateID();
+
+		Key = InArgs._ActorReferenceKey;
+		SetKey = InArgs._OnSetActorReferenceKey;
+
+		OnGlobalTimeChangedHandle = WeakSequencer.Pin()->OnGlobalTimeChanged().AddRaw(this, &SActorReferenceBox::GlobalTimeChanged);
+		OnMovieSceneDataChangedHandle = WeakSequencer.Pin()->OnMovieSceneDataChanged().AddRaw(this, &SActorReferenceBox::MovieSceneDataChanged);
+
+		ChildSlot
+		[
+			SNew(SComboButton)
+			.OnGetMenuContent(this, &SActorReferenceBox::GetPickerMenu)
+			.ContentPadding(FMargin(0.0, 0.0))
+			.ButtonStyle(FEditorStyle::Get(), "PropertyEditor.AssetComboStyle")
+			.ForegroundColor(FEditorStyle::GetColor("PropertyEditor.AssetName.ColorAndOpacity"))
+			.ButtonContent()
+			[
+				GetCurrentItemWidget(
+					SNew(STextBlock)
+					.TextStyle(FEditorStyle::Get(), "PropertyEditor.AssetClass")
+					.Font(FCoreStyle::GetDefaultFontStyle("Regular", 8))
+				)
+			]
+		];
+
+		Update();
+	}
+
+	virtual ~SActorReferenceBox()
+	{
+		if (WeakSequencer.IsValid())
+		{
+			WeakSequencer.Pin()->OnGlobalTimeChanged().Remove(OnGlobalTimeChangedHandle);
+			WeakSequencer.Pin()->OnMovieSceneDataChanged().Remove(OnMovieSceneDataChangedHandle);
+		}
+	}
+
+	virtual UMovieSceneSequence* GetSequence() const override
+	{
+		return WeakSequencer.Pin()->GetFocusedMovieSceneSequence();
+	}
+
+	/** Set the current binding ID */
+	virtual void SetCurrentValue(const FMovieSceneObjectBindingID& InBindingId) override
+	{
+		SetKey.Execute(FMovieSceneActorReferenceKey(InBindingId));
+	}
+
+	/** Get the current binding ID */
+	virtual FMovieSceneObjectBindingID GetCurrentValue() const override
+	{
+		return Key.Get().Object;
+	}
+
+	void GlobalTimeChanged()
+	{
+		Update();
+	}
+
+	void MovieSceneDataChanged(EMovieSceneDataChangeType)
+	{
+		Update();
+	}
+
+	void Update()
+	{
+		if (IsEmpty())
+		{
+			Initialize();
+		}
+		else
+		{
+			UpdateCachedData();
+		}
+	}
+
+private:
+
+	TAttribute< FMovieSceneActorReferenceKey> Key;
+
+	FOnSetActorReferenceKey SetKey;
+
+	FDelegateHandle OnGlobalTimeChangedHandle, OnMovieSceneDataChangedHandle;
+};
+
+
+TSharedRef<SWidget> CreateKeyEditor(const TMovieSceneChannelHandle<FMovieSceneActorReferenceData>& Channel, UMovieSceneSection* Section, const FGuid& InObjectBindingID, TWeakPtr<FTrackInstancePropertyBindings> PropertyBindings, TWeakPtr<ISequencer> InSequencer)
+{
+	const FMovieSceneActorReferenceData* RawChannel = Channel.Get();
+	if (!RawChannel)
+	{
+		return SNullWidget::NullWidget;
+	}
+
+	TFunction<TOptional<FMovieSceneActorReferenceKey>(UObject&, FTrackInstancePropertyBindings*)> Func;
+
+	TSequencerKeyEditor<FMovieSceneActorReferenceData, FMovieSceneActorReferenceKey> KeyEditor(InObjectBindingID, Channel, Section, InSequencer, PropertyBindings, Func);
+
+	auto OnSetCurrentValueLambda = [KeyEditor](FMovieSceneActorReferenceKey& ActorKey) mutable
+	{
+		FScopedTransaction Transaction(LOCTEXT("SetKey", "Set Actor Reference Key Value"));
+		KeyEditor.SetValueWithNotify(ActorKey, EMovieSceneDataChangeType::TrackValueChangedRefreshImmediately);
+
+		// Look for components to choose
+		ISequencer* Sequencer = KeyEditor.GetSequencer();
+		FMovieSceneEvaluationOperand ObjectOperand(ActorKey.Object.GetSequenceID(), ActorKey.Object.GetGuid());
+		TArray<USceneComponent*> ComponentsWithSockets;
+		AActor* Actor = nullptr;
+		for (TWeakObjectPtr<> WeakObject : Sequencer->FindBoundObjects(ObjectOperand))
+		{
+			Actor = Cast<AActor>(WeakObject.Get());
+			if (Actor)
+			{
+				TInlineComponentArray<USceneComponent*> Components(Actor);
+
+				for (USceneComponent* Component : Components)
+				{
+					if (Component && Component->HasAnySockets())
+					{
+						ComponentsWithSockets.Add(Component);
+					}
+				}
+				break;
+			}
+		}
+
+		if (ComponentsWithSockets.Num() == 0 || !Actor)
+		{
+			return;
+		}
+
+		// Pop up a component chooser
+		FLevelEditorModule& LevelEditorModule = FModuleManager::GetModuleChecked<FLevelEditorModule>("LevelEditor");
+		TSharedPtr< ILevelEditor > LevelEditor = LevelEditorModule.GetFirstLevelEditor();
+
+		TSharedPtr<SWidget> ComponentMenuWidget =
+			SNew(SComponentChooserPopup)
+			.Actor(Actor)
+			.OnComponentChosen_Lambda([=](FName InComponentName) mutable
+				{
+					ActorKey.ComponentName = InComponentName;
+					KeyEditor.SetValueWithNotify(ActorKey, EMovieSceneDataChangeType::TrackValueChangedRefreshImmediately);
+
+					// Look for sockets to choose
+					USceneComponent* ComponentWithSockets = nullptr;
+					TInlineComponentArray<USceneComponent*> Components(Actor);
+
+					for (USceneComponent* Component : Components)
+					{
+						if (Component && Component->GetFName() == InComponentName)
+						{
+							ComponentWithSockets = Component;
+							break;
+						}
+					}
+
+					if (!ComponentWithSockets)
+					{
+						return;
+					}
+							
+					// Pop up a socket chooser
+					TSharedPtr<SWidget> SocketMenuWidget =
+						SNew(SSocketChooserPopup)
+						.SceneComponent(ComponentWithSockets)
+						.OnSocketChosen_Lambda([=](FName InSocketName) mutable
+							{
+								ActorKey.SocketName = InSocketName;
+								KeyEditor.SetValueWithNotify(ActorKey, EMovieSceneDataChangeType::TrackValueChangedRefreshImmediately);
+							}
+						);
+
+					// Create as context menu
+					FSlateApplication::Get().PushMenu(
+						LevelEditor.ToSharedRef(),
+						FWidgetPath(),
+						SocketMenuWidget.ToSharedRef(),
+						FSlateApplication::Get().GetCursorPos(),
+						FPopupTransitionEffect(FPopupTransitionEffect::ContextMenu)
+					);
+				}
+			);
+
+		// Create as context menu
+		FSlateApplication::Get().PushMenu(
+			LevelEditor.ToSharedRef(),
+			FWidgetPath(),
+			ComponentMenuWidget.ToSharedRef(),
+			FSlateApplication::Get().GetCursorPos(),
+			FPopupTransitionEffect(FPopupTransitionEffect::ContextMenu)
+		);
+	};
+
+	auto GetCurrentValueLambda = [KeyEditor]() -> FMovieSceneActorReferenceKey
+	{
+		return KeyEditor.GetCurrentValue();
+	};
+
+	return SNew(SActorReferenceBox, InSequencer)
+		.ActorReferenceKey_Lambda(GetCurrentValueLambda)
+		.OnSetActorReferenceKey_Lambda(OnSetCurrentValueLambda);
+}
+
 UMovieSceneKeyStructType* InstanceGeneratedStruct(FMovieSceneByteChannel* Channel, FSequencerKeyStructGenerator* Generator)
 {
 	UEnum* ByteEnum = Channel->GetEnum();
diff --git a/Engine/Source/Editor/MovieSceneTools/Private/Channels/BuiltInChannelEditors.h b/Engine/Source/Editor/MovieSceneTools/Private/Channels/BuiltInChannelEditors.h
--- a/Engine/Source/Editor/MovieSceneTools/Private/Channels/BuiltInChannelEditors.h
+++ b/Engine/Source/Editor/MovieSceneTools/Private/Channels/BuiltInChannelEditors.h
@@ -44,6 +44,7 @@ bool CanCreateKeyEditor(const FMovieSceneIntegerChannel*    Channel);
 bool CanCreateKeyEditor(const FMovieSceneFloatChannel*      Channel);
 bool CanCreateKeyEditor(const FMovieSceneStringChannel*     Channel);
 bool CanCreateKeyEditor(const FMovieSceneObjectPathChannel* Channel);
+bool CanCreateKeyEditor(const FMovieSceneActorReferenceData* Channel);
 
 TSharedRef<SWidget> CreateKeyEditor(const TMovieSceneChannelHandle<FMovieSceneBoolChannel>&        Channel, UMovieSceneSection* Section, const FGuid& InObjectBindingID, TWeakPtr<FTrackInstancePropertyBindings> PropertyBindings, TWeakPtr<ISequencer> InSequencer);
 TSharedRef<SWidget> CreateKeyEditor(const TMovieSceneChannelHandle<FMovieSceneByteChannel>&        Channel, UMovieSceneSection* Section, const FGuid& InObjectBindingID, TWeakPtr<FTrackInstancePropertyBindings> PropertyBindings, TWeakPtr<ISequencer> InSequencer);
@@ -51,6 +52,7 @@ TSharedRef<SWidget> CreateKeyEditor(const TMovieSceneChannelHandle<FMovieSceneIn
 TSharedRef<SWidget> CreateKeyEditor(const TMovieSceneChannelHandle<FMovieSceneFloatChannel>&       Channel, UMovieSceneSection* Section, const FGuid& InObjectBindingID, TWeakPtr<FTrackInstancePropertyBindings> PropertyBindings, TWeakPtr<ISequencer> InSequencer);
 TSharedRef<SWidget> CreateKeyEditor(const TMovieSceneChannelHandle<FMovieSceneStringChannel>&      Channel, UMovieSceneSection* Section, const FGuid& InObjectBindingID, TWeakPtr<FTrackInstancePropertyBindings> PropertyBindings, TWeakPtr<ISequencer> InSequencer);
 TSharedRef<SWidget> CreateKeyEditor(const TMovieSceneChannelHandle<FMovieSceneObjectPathChannel>&  Channel, UMovieSceneSection* Section, const FGuid& InObjectBindingID, TWeakPtr<FTrackInstancePropertyBindings> PropertyBindings, TWeakPtr<ISequencer> InSequencer);
+TSharedRef<SWidget> CreateKeyEditor(const TMovieSceneChannelHandle<FMovieSceneActorReferenceData>&  Channel, UMovieSceneSection* Section, const FGuid& InObjectBindingID, TWeakPtr<FTrackInstancePropertyBindings> PropertyBindings, TWeakPtr<ISequencer> InSequencer);
 
 UMovieSceneKeyStructType* InstanceGeneratedStruct(FMovieSceneByteChannel* Channel,       FSequencerKeyStructGenerator* Generator);
 UMovieSceneKeyStructType* InstanceGeneratedStruct(FMovieSceneObjectPathChannel* Channel, FSequencerKeyStructGenerator* Generator);
diff --git a/Engine/Source/Editor/MovieSceneTools/Private/SComponentChooser.h b/Engine/Source/Editor/MovieSceneTools/Private/SComponentChooser.h
new file mode 100644
--- /dev/null
+++ b/Engine/Source/Editor/MovieSceneTools/Private/SComponentChooser.h
@@ -0,0 +1,118 @@
+// Copyright Epic Games, Inc. All Rights Reserved.
+
+#pragma once
+
+#include "CoreMinimal.h"
+#include "Engine/EngineTypes.h"
+#include "Components/SceneComponent.h"
+#include "Layout/SlateRect.h"
+#include "Input/Reply.h"
+#include "Widgets/DeclarativeSyntaxSupport.h"
+#include "Widgets/SBoxPanel.h"
+#include "Widgets/Text/STextBlock.h"
+#include "Widgets/SCompoundWidget.h"
+#include "Widgets/Views/STableViewBase.h"
+#include "Widgets/Layout/SBorder.h"
+#include "Widgets/Layout/SBox.h"
+#include "Widgets/Views/STableRow.h"
+#include "Widgets/Views/SListView.h"
+#include "EditorStyleSet.h"
+#include "Framework/Application/SlateApplication.h"
+#include "Widgets/Images/SImage.h"
+#include "Widgets/Input/SSearchBox.h"
+#include "Misc/TextFilterExpressionEvaluator.h"
+
+class SComponentChooserPopup : public SCompoundWidget
+{
+public:
+	DECLARE_DELEGATE_OneParam( FOnComponentChosen, FName );
+
+	SLATE_BEGIN_ARGS( SComponentChooserPopup )
+		: _Actor(NULL)
+		{}
+
+		/** An actor with components */
+		SLATE_ARGUMENT( AActor*, Actor )
+
+		/** Called when the text is chosen. */
+		SLATE_EVENT( FOnComponentChosen, OnComponentChosen )
+
+	SLATE_END_ARGS()
+
+	/** Delegate to call when component is selected */
+	FOnComponentChosen OnComponentChosen;
+
+	/** List of tag names selected in the tag containers*/
+	TArray< TSharedPtr<FName> > ComponentNames;
+
+private:
+	TSharedRef<ITableRow> MakeListViewWidget(TSharedPtr<FName> InItem, const TSharedRef<STableViewBase>& OwnerTable)
+	{
+		return SNew( STableRow< TSharedPtr<FName> >, OwnerTable )
+				[
+					SNew(STextBlock) .Text( FText::FromName(*InItem.Get()) )
+				];
+	}
+
+	void OnComponentSelected(TSharedPtr<FName> InItem, ESelectInfo::Type InSelectInfo)
+	{
+		FSlateApplication::Get().DismissAllMenus();
+
+		if(OnComponentChosen.IsBound())
+		{
+			OnComponentChosen.Execute(*InItem.Get());
+		}
+	}
+
+public:
+	void Construct( const FArguments& InArgs )
+	{
+		OnComponentChosen = InArgs._OnComponentChosen;
+		AActor* Actor = InArgs._Actor;
+
+		TInlineComponentArray<USceneComponent*> Components(Actor);
+
+		ComponentNames.Empty();
+		for(USceneComponent* Component : Components)
+		{
+			if (Component->HasAnySockets())
+			{
+				ComponentNames.Add(MakeShareable(new FName(Component->GetFName())));
+			}
+		}
+
+		// Then make widget
+		this->ChildSlot
+		[
+			SNew(SBorder)
+			.BorderImage(FEditorStyle::GetBrush(TEXT("Menu.Background")))
+			.Padding(5)
+			.Content()
+			[
+				SNew(SVerticalBox)
+				+SVerticalBox::Slot()
+				.AutoHeight()
+				.Padding(0.0f, 1.0f)
+				[
+					SNew(STextBlock)
+					.Font( FEditorStyle::GetFontStyle(TEXT("SocketChooser.TitleFont")) )
+					.Text( NSLOCTEXT("ComponentChooser", "ChooseComponentLabel", "Choose Component") )
+				]
+				+SVerticalBox::Slot()
+				.AutoHeight()
+				.MaxHeight(512)
+				[
+					SNew(SBox)
+					.WidthOverride(256)
+					.Content()
+					[
+						SNew(SListView< TSharedPtr<FName> >)
+						.ListItemsSource( &ComponentNames)
+						.OnGenerateRow( this, &SComponentChooserPopup::MakeListViewWidget )
+						.OnSelectionChanged( this, &SComponentChooserPopup::OnComponentSelected )
+					]
+				]
+			]
+		];
+	}
+};
\ No newline at end of file
diff --git a/Engine/Source/Editor/MovieSceneTools/Private/TrackEditors/ActorPickerTrackEditor.cpp b/Engine/Source/Editor/MovieSceneTools/Private/TrackEditors/ActorPickerTrackEditor.cpp
--- a/Engine/Source/Editor/MovieSceneTools/Private/TrackEditors/ActorPickerTrackEditor.cpp
+++ b/Engine/Source/Editor/MovieSceneTools/Private/TrackEditors/ActorPickerTrackEditor.cpp
@@ -23,6 +23,7 @@
 #include "LevelEditor.h"
 #include "MovieSceneObjectBindingIDPicker.h"
 #include "MovieSceneToolHelpers.h"
+#include "SComponentChooser.h"
 
 #define LOCTEXT_NAMESPACE "FActorPickerTrackEditor"
 
@@ -159,100 +160,6 @@ void FActorPickerTrackEditor::ShowActorSubMenu(FMenuBuilder& MenuBuilder, TArray
 	}
 }
 
-class SComponentChooserPopup : public SCompoundWidget
-{
-public:
-	DECLARE_DELEGATE_OneParam( FOnComponentChosen, FString );
-
-	SLATE_BEGIN_ARGS( SComponentChooserPopup )
-		: _Actor(NULL)
-		{}
-
-		/** An actor with components */
-		SLATE_ARGUMENT( AActor*, Actor )
-
-		/** Called when the text is chosen. */
-		SLATE_EVENT( FOnComponentChosen, OnComponentChosen )
-
-	SLATE_END_ARGS()
-
-	/** Delegate to call when component is selected */
-	FOnComponentChosen OnComponentChosen;
-
-	/** List of tag names selected in the tag containers*/
-	TArray< TSharedPtr<FString> > ComponentNames;
-
-private:
-	TSharedRef<ITableRow> MakeListViewWidget(TSharedPtr<FString> InItem, const TSharedRef<STableViewBase>& OwnerTable)
-	{
-		return SNew( STableRow< TSharedPtr<FString> >, OwnerTable )
-				[
-					SNew(STextBlock) .Text( FText::FromString(*InItem.Get()) )
-				];
-	}
-
-	void OnComponentSelected(TSharedPtr<FString> InItem, ESelectInfo::Type InSelectInfo)
-	{
-		FSlateApplication::Get().DismissAllMenus();
-
-		if(OnComponentChosen.IsBound())
-		{
-			OnComponentChosen.Execute(*InItem.Get());
-		}
-	}
-
-public:
-	void Construct( const FArguments& InArgs )
-	{
-		OnComponentChosen = InArgs._OnComponentChosen;
-		AActor* Actor = InArgs._Actor;
-
-		TInlineComponentArray<USceneComponent*> Components(Actor);
-
-		ComponentNames.Empty();
-		for(USceneComponent* Component : Components)
-		{
-			if (Component->HasAnySockets())
-			{
-				ComponentNames.Add(MakeShareable(new FString(Component->GetName())));
-			}
-		}
-
-		// Then make widget
-		this->ChildSlot
-		[
-			SNew(SBorder)
-			.BorderImage(FEditorStyle::GetBrush(TEXT("Menu.Background")))
-			.Padding(5)
-			.Content()
-			[
-				SNew(SVerticalBox)
-				+SVerticalBox::Slot()
-				.AutoHeight()
-				.Padding(0.0f, 1.0f)
-				[
-					SNew(STextBlock)
-					.Font( FEditorStyle::GetFontStyle(TEXT("SocketChooser.TitleFont")) )
-					.Text( NSLOCTEXT("ComponentChooser", "ChooseComponentLabel", "Choose Component") )
-				]
-				+SVerticalBox::Slot()
-				.AutoHeight()
-				.MaxHeight(512)
-				[
-					SNew(SBox)
-					.WidthOverride(256)
-					.Content()
-					[
-						SNew(SListView< TSharedPtr<FString> >)
-						.ListItemsSource( &ComponentNames)
-						.OnGenerateRow( this, &SComponentChooserPopup::MakeListViewWidget )
-						.OnSelectionChanged( this, &SComponentChooserPopup::OnComponentSelected )
-					]
-				]
-			]
-		];
-	}
-};
 
 void FActorPickerTrackEditor::ActorPicked(AActor* ParentActor, TArray<FGuid> ObjectGuids, UMovieSceneSection* Section)
 {
@@ -337,12 +244,12 @@ void FActorPickerTrackEditor::ActorPickerIDPicked(FActorPickerID ActorPickerID,
 	}
 	else
 	{
-		ActorComponentPicked(ComponentsWithSockets[0]->GetName(), ActorPickerID, ObjectGuids, Section);
+		ActorComponentPicked(ComponentsWithSockets[0]->GetFName(), ActorPickerID, ObjectGuids, Section);
 	}
 }
 
 
-void FActorPickerTrackEditor::ActorComponentPicked(FString ComponentName, FActorPickerID ActorPickerID, TArray<FGuid> ObjectGuids, UMovieSceneSection* Section)
+void FActorPickerTrackEditor::ActorComponentPicked(FName ComponentName, FActorPickerID ActorPickerID, TArray<FGuid> ObjectGuids, UMovieSceneSection* Section)
 {
 	USceneComponent* ComponentWithSockets = nullptr;
 	if (ActorPickerID.ActorPicked.IsValid())
@@ -351,7 +258,7 @@ void FActorPickerTrackEditor::ActorComponentPicked(FString ComponentName, FActor
 	
 		for(USceneComponent* Component : Components)
 		{
-			if (Component->GetName() == ComponentName)
+			if (Component->GetFName() == ComponentName)
 			{
 				ComponentWithSockets = Component;
 				break;
diff --git a/Engine/Source/Editor/MovieSceneTools/Private/TrackEditors/ActorPickerTrackEditor.h b/Engine/Source/Editor/MovieSceneTools/Private/TrackEditors/ActorPickerTrackEditor.h
--- a/Engine/Source/Editor/MovieSceneTools/Private/TrackEditors/ActorPickerTrackEditor.h
+++ b/Engine/Source/Editor/MovieSceneTools/Private/TrackEditors/ActorPickerTrackEditor.h
@@ -51,7 +51,7 @@ private:
 	void ActorPickerIDPicked(FActorPickerID, const TArray<FGuid>& ObjectBindings, UMovieSceneSection* Section);
 
 	/** Actor component was picked */
-	void ActorComponentPicked(FString ComponentName, FActorPickerID ActorPickerID, TArray<FGuid> ObjectBindings, UMovieSceneSection* Section);
+	void ActorComponentPicked(FName ComponentName, FActorPickerID ActorPickerID, TArray<FGuid> ObjectBindings, UMovieSceneSection* Section);
 
 	/** Interactively pick an actor from the viewport */
 	void PickActorInteractive(const TArray<FGuid>& ObjectBindings, UMovieSceneSection* Section);
diff --git a/Engine/Source/Editor/MovieSceneTools/Public/MovieSceneObjectBindingIDPicker.h b/Engine/Source/Editor/MovieSceneTools/Public/MovieSceneObjectBindingIDPicker.h
--- a/Engine/Source/Editor/MovieSceneTools/Public/MovieSceneObjectBindingIDPicker.h
+++ b/Engine/Source/Editor/MovieSceneTools/Public/MovieSceneObjectBindingIDPicker.h
@@ -93,6 +93,9 @@ protected:
 	/** The ID of the sequence to generate IDs relative to */
 	FMovieSceneSequenceID LocalSequenceID;
 
+	/** Update the cached text, tooltip and icon */
+	void UpdateCachedData();
+
 private:
 
 	/** Get the currently set binding ID, remapped to the root sequence if necessary */
@@ -101,9 +104,6 @@ private:
 	/** Set the binding ID, remapped to the local sequence if possible */
 	void SetRemappedCurrentValue(FMovieSceneObjectBindingID InValue);
 
-	/** UPdate the cached text, tooltip and icon */
-	void UpdateCachedData();
-
 	/** Called when the combo box has been clicked to populate its menu content */
 	void OnGetMenuContent(FMenuBuilder& MenuBuilder, TSharedPtr<FSequenceBindingNode> Node);
 
diff --git a/Engine/Source/Editor/Sequencer/Private/SequencerEdMode.cpp b/Engine/Source/Editor/Sequencer/Private/SequencerEdMode.cpp
--- a/Engine/Source/Editor/Sequencer/Private/SequencerEdMode.cpp
+++ b/Engine/Source/Editor/Sequencer/Private/SequencerEdMode.cpp
@@ -15,6 +15,8 @@
 #include "MovieSceneHitProxy.h"
 #include "Tracks/MovieScene3DTransformTrack.h"
 #include "Sections/MovieScene3DTransformSection.h"
+#include "Tracks/MovieSceneAudioTrack.h"
+#include "Sections/MovieSceneAudioSection.h"
 #include "SubtitleManager.h"
 #include "SequencerMeshTrail.h"
 #include "SequencerKeyActor.h"
@@ -43,6 +45,9 @@ FSequencerEdMode::FSequencerEdMode()
 	{
 		bDrawMeshTrails = Var->GetBool();
 	}));
+
+	AudioTexture = LoadObject<UTexture2D>(NULL, TEXT("/Engine/EditorResources/AudioIcons/S_AudioComponent.S_AudioComponent"));
+	check(AudioTexture);
 }
 
 FSequencerEdMode::~FSequencerEdMode()
@@ -101,6 +106,11 @@ void FSequencerEdMode::Render(const FSceneView* View, FViewport* Viewport, FPrim
 	FEdMode::Render(View, Viewport, PDI);
 
 #if WITH_EDITORONLY_DATA
+	if (PDI)
+	{
+		DrawAudioTracks(PDI);
+	}
+
 	// Draw spline trails using the PDI
 	if (View->Family->EngineShowFlags.Splines)
 	{
@@ -682,6 +692,89 @@ void FSequencerEdMode::DrawTracks3D(FPrimitiveDrawInterface* PDI)
 	}
 }
 
+void FSequencerEdMode::DrawAudioTracks(FPrimitiveDrawInterface* PDI)
+{
+	for (TWeakPtr<FSequencer> WeakSequencer : Sequencers)
+	{
+		TSharedPtr<FSequencer> Sequencer = WeakSequencer.Pin();
+		if (!Sequencer.IsValid())
+		{
+			continue;
+		}
+
+		UMovieSceneSequence* Sequence = Sequencer->GetFocusedMovieSceneSequence();
+		if (!Sequence)
+		{
+			continue;
+		}
+
+		FQualifiedFrameTime CurrentTime = Sequencer->GetLocalTime();
+
+		const FSequencerSelection& Selection = Sequencer->GetSelection();
+		for (UMovieSceneTrack* Track : Selection.GetSelectedTracks())
+		{
+			UMovieSceneAudioTrack* AudioTrack = Cast<UMovieSceneAudioTrack>(Track);
+
+			if (!AudioTrack || !AudioTrack->IsAMasterTrack())
+			{
+				continue;
+			}
+
+			for (UMovieSceneSection* Section : AudioTrack->GetAudioSections())
+			{
+				UMovieSceneAudioSection* AudioSection = Cast<UMovieSceneAudioSection>(Section);
+				const FMovieSceneActorReferenceData& AttachActorData = AudioSection->GetAttachActorData();
+
+				TMovieSceneChannelData<const FMovieSceneActorReferenceKey> ChannelData = AttachActorData.GetData();
+
+				TArrayView<const FFrameNumber> Times = ChannelData.GetTimes();
+				TArrayView<const FMovieSceneActorReferenceKey> Values = ChannelData.GetValues();
+		
+				FMovieSceneActorReferenceKey CurrentValue;
+				AttachActorData.Evaluate(CurrentTime.Time, CurrentValue);
+
+				for (int32 Index = 0; Index < Times.Num(); ++Index)
+				{
+					FMovieSceneObjectBindingID AttachBindingID = Values[Index].Object;
+					FName AttachSocketName = Values[Index].SocketName;
+
+					FMovieSceneSequenceID SequenceID = Sequencer->GetFocusedTemplateID();
+					if (AttachBindingID.GetSequenceID().IsValid())
+					{
+						// Ensure that this ID is resolvable from the root, based on the current local sequence ID
+						FMovieSceneObjectBindingID RootBindingID = AttachBindingID.ResolveLocalToRoot(SequenceID, Sequencer->GetEvaluationTemplate().GetHierarchy());
+						SequenceID = RootBindingID.GetSequenceID();
+					}
+
+					// If the transform is set, otherwise use the bound actor's transform
+					FMovieSceneEvaluationOperand ObjectOperand(SequenceID, AttachBindingID.GetGuid());
+
+					for (TWeakObjectPtr<> WeakObject : Sequencer->FindBoundObjects(ObjectOperand))
+					{
+						AActor* AttachActor = Cast<AActor>(WeakObject.Get());
+						if (AttachActor)
+						{
+							USceneComponent* AttachComponent = AudioSection->GetAttachComponent(AttachActor, Values[Index]);
+							if (AttachComponent)
+							{
+								FVector Location = AttachComponent->GetSocketLocation(AttachSocketName);
+								bool bIsActive = CurrentValue == Values[Index];
+								FColor Color = bIsActive ? FColor::Green : FColor::White;
+
+								float Scale = PDI->View->WorldToScreen(Location).W * (4.0f / PDI->View->UnscaledViewRect.Width() / PDI->View->ViewMatrices.GetProjectionMatrix().M[0][0]);
+								Scale *= bIsActive ? 15.f : 10.f;
+
+								PDI->DrawSprite(Location, Scale, Scale, AudioTexture->Resource, Color, SDPG_Foreground, 0.0, 0.0, 0.0, 0.0, SE_BLEND_Masked);
+								break;
+							}
+						}
+					}
+				}
+			}
+		}
+	}
+}
+
 FSequencerEdModeTool::FSequencerEdModeTool(FSequencerEdMode* InSequencerEdMode) :
 	SequencerEdMode(InSequencerEdMode)
 {
diff --git a/Engine/Source/Editor/Sequencer/Private/SequencerEdMode.h b/Engine/Source/Editor/Sequencer/Private/SequencerEdMode.h
--- a/Engine/Source/Editor/Sequencer/Private/SequencerEdMode.h
+++ b/Engine/Source/Editor/Sequencer/Private/SequencerEdMode.h
@@ -7,6 +7,7 @@
 #include "EditorModeTools.h"
 #include "EdMode.h"
 #include "Misc/FrameTime.h"
+#include "Engine/Texture2D.h"
 
 class FCanvas;
 class FEditorViewportClient;
@@ -72,6 +73,7 @@ public:
 protected:
 	void DrawTracks3D(FPrimitiveDrawInterface* PDI);
 	void DrawTransformTrack(const TSharedPtr<FSequencer>& Sequencer, FPrimitiveDrawInterface* PDI, UMovieScene3DTransformTrack* TransformTrack, TArrayView<const TWeakObjectPtr<>> BoundObjects, const bool bIsSelected);
+	void DrawAudioTracks(FPrimitiveDrawInterface* PDI);
 
 protected:
 	static void GetLocationAtTime(FMovieSceneEvaluationTrack* Track, UObject* BoundObject, FFrameTime KeyTime, FVector& KeyPos, FRotator& KeyRot, const TSharedPtr<FSequencer>& Sequencer);
@@ -89,6 +91,9 @@ private:
 
 	/** If true, draw mesh trails instead of debug lines*/
 	bool bDrawMeshTrails;
+
+	/** The audio texture used for drawing the audio spatialization points */
+	UTexture2D* AudioTexture;
 };
 
 /**
diff --git a/Engine/Source/Runtime/MovieSceneTracks/Private/Evaluation/MovieSceneActorReferenceTemplate.cpp b/Engine/Source/Runtime/MovieSceneTracks/Private/Evaluation/MovieSceneActorReferenceTemplate.cpp
--- a/Engine/Source/Runtime/MovieSceneTracks/Private/Evaluation/MovieSceneActorReferenceTemplate.cpp
+++ b/Engine/Source/Runtime/MovieSceneTracks/Private/Evaluation/MovieSceneActorReferenceTemplate.cpp
@@ -58,6 +58,7 @@ void FMovieSceneActorReferenceSectionTemplate::Evaluate(const FMovieSceneEvaluat
 {
 	using namespace PropertyTemplate;
 
-	FMovieSceneActorReferenceKey ObjectBinding = ActorReferenceData.Evaluate(Context.GetTime());
+	FMovieSceneActorReferenceKey ObjectBinding;
+	ActorReferenceData.Evaluate(Context.GetTime(), ObjectBinding);
 	ExecutionTokens.Add(TPropertyTrackExecutionToken<UObject*, FMovieSceneObjectBindingID>(ObjectBinding.Object));
 }
diff --git a/Engine/Source/Runtime/MovieSceneTracks/Private/Evaluation/MovieSceneAudioTemplate.cpp b/Engine/Source/Runtime/MovieSceneTracks/Private/Evaluation/MovieSceneAudioTemplate.cpp
--- a/Engine/Source/Runtime/MovieSceneTracks/Private/Evaluation/MovieSceneAudioTemplate.cpp
+++ b/Engine/Source/Runtime/MovieSceneTracks/Private/Evaluation/MovieSceneAudioTemplate.cpp
@@ -250,6 +250,7 @@ struct FAudioSectionExecutionToken : IMovieSceneExecutionToken
 		: AudioSection(InAudioSection), SectionKey(InAudioSection)
 	{}
 
+
 	virtual void Execute(const FMovieSceneContext& Context, const FMovieSceneEvaluationOperand& Operand, FPersistentEvaluationData& PersistentData, IMovieScenePlayer& Player)
 	{
 		FCachedAudioTrackData& TrackData = PersistentData.GetOrAddTrackData<FCachedAudioTrackData>();
@@ -268,8 +269,9 @@ struct FAudioSectionExecutionToken : IMovieSceneExecutionToken
 			const FMovieSceneActorReferenceData& AttachActorData = AudioSection->GetAttachActorData();
 
 			USceneComponent* AttachComponent = nullptr;
-
-			FMovieSceneObjectBindingID AttachBindingID = AttachActorData.Evaluate(Context.GetTime()).Object;
+			FMovieSceneActorReferenceKey AttachKey;
+			AttachActorData.Evaluate(Context.GetTime(), AttachKey);
+			FMovieSceneObjectBindingID AttachBindingID = AttachKey.Object;
 			if (AttachBindingID.IsValid())
 			{
 				FMovieSceneSequenceID SequenceID = Operand.SequenceID;
@@ -285,19 +287,10 @@ struct FAudioSectionExecutionToken : IMovieSceneExecutionToken
 
 				for (TWeakObjectPtr<> WeakObject : Player.FindBoundObjects(ObjectOperand))
 				{
-					UObject* AttachObject = WeakObject.Get();
-					if (AttachObject)
+					AActor* AttachActor = Cast<AActor>(WeakObject.Get());
+					if (AttachActor)
 					{
-						if (AttachObject->IsA<AActor>())
-						{
-							AActor* Actor = Cast<AActor>(AttachObject);
-							AttachComponent = Actor->GetRootComponent();
-						}
-						else if (AttachObject->IsA<UActorComponent>())
-						{
-							UActorComponent* ActorComponent = Cast<UActorComponent>(AttachObject);
-							AttachComponent = Cast<USceneComponent>(ActorComponent);
-						}
+						AttachComponent = AudioSection->GetAttachComponent(AttachActor, AttachKey);
 					}
 					if (AttachComponent)
 					{
@@ -331,9 +324,9 @@ struct FAudioSectionExecutionToken : IMovieSceneExecutionToken
 
 			if (AudioComponent)
 			{
-				if (AttachComponent)
+				if (AttachComponent && AudioComponent->GetAttachParent() != AttachComponent || AudioComponent->GetAttachSocketName() != AttachKey.SocketName)
 				{
-					AudioComponent->AttachToComponent(AttachComponent, FAttachmentTransformRules::KeepRelativeTransform);
+					AudioComponent->AttachToComponent(AttachComponent, FAttachmentTransformRules::KeepRelativeTransform, AttachKey.SocketName);
 				}
 
 				EnsureAudioIsPlaying(*AudioComponent, PersistentData, Context, AudioComponent->GetAttachParent() != nullptr, Player);
diff --git a/Engine/Source/Runtime/MovieSceneTracks/Private/Sections/MovieSceneActorReferenceSection.cpp b/Engine/Source/Runtime/MovieSceneTracks/Private/Sections/MovieSceneActorReferenceSection.cpp
--- a/Engine/Source/Runtime/MovieSceneTracks/Private/Sections/MovieSceneActorReferenceSection.cpp
+++ b/Engine/Source/Runtime/MovieSceneTracks/Private/Sections/MovieSceneActorReferenceSection.cpp
@@ -4,15 +4,17 @@
 #include "Channels/MovieSceneChannelProxy.h"
 
 
-FMovieSceneActorReferenceKey FMovieSceneActorReferenceData::Evaluate(FFrameTime InTime) const
+bool FMovieSceneActorReferenceData::Evaluate(FFrameTime InTime, FMovieSceneActorReferenceKey& OutValue) const
 {
 	if (KeyTimes.Num())
 	{
 		const int32 Index = FMath::Max(0, Algo::UpperBound(KeyTimes, InTime.FrameNumber)-1);
-		return KeyValues[Index];
+		OutValue = KeyValues[Index];
+		return true;
 	}
 
-	return DefaultValue;
+	OutValue = DefaultValue;
+	return true;
 }
 
 void FMovieSceneActorReferenceData::GetKeys(const TRange<FFrameNumber>& WithinRange, TArray<FFrameNumber>* OutKeyTimes, TArray<FKeyHandle>* OutKeyHandles)
@@ -45,7 +47,8 @@ void FMovieSceneActorReferenceData::DeleteKeysFrom(FFrameNumber InTime, bool bDe
 	// Insert a key at the current time to maintain evaluation
 	if (GetData().GetTimes().Num() > 0)
 	{
-		FMovieSceneActorReferenceKey Value = Evaluate(InTime);
+		FMovieSceneActorReferenceKey Value;
+		Evaluate(InTime, Value);
 		GetData().UpdateOrAddKey(InTime, Value);
 	}
 
diff --git a/Engine/Source/Runtime/MovieSceneTracks/Private/Sections/MovieSceneAudioSection.cpp b/Engine/Source/Runtime/MovieSceneTracks/Private/Sections/MovieSceneAudioSection.cpp
--- a/Engine/Source/Runtime/MovieSceneTracks/Private/Sections/MovieSceneAudioSection.cpp
+++ b/Engine/Source/Runtime/MovieSceneTracks/Private/Sections/MovieSceneAudioSection.cpp
@@ -1,6 +1,7 @@
 // Copyright Epic Games, Inc. All Rights Reserved.
 
 #include "Sections/MovieSceneAudioSection.h"
+#include "Tracks/MovieSceneAudioTrack.h"
 #include "Sound/SoundBase.h"
 #include "Evaluation/MovieSceneAudioTemplate.h"
 #include "UObject/SequencerObjectVersion.h"
@@ -49,27 +50,58 @@ UMovieSceneAudioSection::UMovieSceneAudioSection( const FObjectInitializer& Obje
 	SoundVolume.SetDefault(1.f);
 	PitchMultiplier.SetDefault(1.f);
 
+	UpdateChannelProxy();
+}
+
+void UMovieSceneAudioSection::UpdateChannelProxy()
+{
 	// Set up the channel proxy
 	FMovieSceneChannelProxyData Channels;
 
+	UMovieSceneAudioTrack* AudioTrack = Cast<UMovieSceneAudioTrack>(GetOuter());
+
 #if WITH_EDITOR
 
-	static const FAudioChannelEditorData EditorData;
+	FAudioChannelEditorData EditorData;
 	Channels.Add(SoundVolume,     EditorData.Data[0], TMovieSceneExternalValue<float>());
 	Channels.Add(PitchMultiplier, EditorData.Data[1], TMovieSceneExternalValue<float>());
-	Channels.Add(AttachActorData, EditorData.Data[2]);
+
+	if (AudioTrack && AudioTrack->IsAMasterTrack())
+	{
+		Channels.Add(AttachActorData, EditorData.Data[2]);
+	}
 
 #else
 
 	Channels.Add(SoundVolume);
 	Channels.Add(PitchMultiplier);
-	Channels.Add(AttachActorData);
+	if (AudioTrack && AudioTrack->IsAMasterTrack())
+	{
+		Channels.Add(AttachActorData);
+	}
 
 #endif
 
 	ChannelProxy = MakeShared<FMovieSceneChannelProxy>(MoveTemp(Channels));
 }
 
+void UMovieSceneAudioSection::Serialize(FArchive& Ar)
+{
+	Super::Serialize(Ar);
+
+	if (Ar.IsLoading())
+	{
+		UpdateChannelProxy();
+	}
+}
+
+void UMovieSceneAudioSection::PostEditImport()
+{
+	Super::PostEditImport();
+
+	UpdateChannelProxy();
+}
+
 FMovieSceneEvalTemplatePtr UMovieSceneAudioSection::GenerateTemplate() const
 {
 	return FMovieSceneAudioSectionTemplate(*this);
@@ -192,3 +224,50 @@ UMovieSceneSection* UMovieSceneAudioSection::SplitSection(FQualifiedFrameTime Sp
 
 	return NewSection;
 }
+
+
+USceneComponent* UMovieSceneAudioSection::GetAttachComponent(const AActor* InParentActor, const FMovieSceneActorReferenceKey& Key) const
+{
+	FName AttachComponentName = Key.ComponentName;
+	FName AttachSocketName = Key.SocketName;
+
+	if (AttachSocketName != NAME_None)
+	{
+		if (AttachComponentName != NAME_None)
+		{
+			TInlineComponentArray<USceneComponent*> PotentialAttachComponents(InParentActor);
+			for (USceneComponent* PotentialAttachComponent : PotentialAttachComponents)
+			{
+				if (PotentialAttachComponent->GetFName() == AttachComponentName && PotentialAttachComponent->DoesSocketExist(AttachSocketName))
+				{
+					return PotentialAttachComponent;
+				}
+			}
+		}
+		else if (InParentActor->GetRootComponent()->DoesSocketExist(AttachSocketName))
+		{
+			return InParentActor->GetRootComponent();
+		}
+	}
+	else if (AttachComponentName != NAME_None)
+	{
+		TInlineComponentArray<USceneComponent*> PotentialAttachComponents(InParentActor);
+		for (USceneComponent* PotentialAttachComponent : PotentialAttachComponents)
+		{
+			if (PotentialAttachComponent->GetFName() == AttachComponentName)
+			{
+				return PotentialAttachComponent;
+			}
+		}
+	}
+
+	if (InParentActor->GetDefaultAttachComponent())
+	{
+		return InParentActor->GetDefaultAttachComponent();
+	}
+	else
+	{
+		return InParentActor->GetRootComponent();
+	}
+}
+
diff --git a/Engine/Source/Runtime/MovieSceneTracks/Public/Sections/MovieSceneActorReferenceSection.h b/Engine/Source/Runtime/MovieSceneTracks/Public/Sections/MovieSceneActorReferenceSection.h
--- a/Engine/Source/Runtime/MovieSceneTracks/Public/Sections/MovieSceneActorReferenceSection.h
+++ b/Engine/Source/Runtime/MovieSceneTracks/Public/Sections/MovieSceneActorReferenceSection.h
@@ -28,16 +28,22 @@ struct FMovieSceneActorReferenceKey
 
 	friend bool operator==(const FMovieSceneActorReferenceKey& A, const FMovieSceneActorReferenceKey& B)
 	{
-		return A.Object == B.Object;
+		return A.Object == B.Object && A.ComponentName == B.ComponentName && A.SocketName == B.SocketName;
 	}
 
 	friend bool operator!=(const FMovieSceneActorReferenceKey& A, const FMovieSceneActorReferenceKey& B)
 	{
-		return A.Object != B.Object;
+		return A.Object != B.Object || A.ComponentName != B.ComponentName || A.SocketName != B.SocketName;
 	}
 
 	UPROPERTY(EditAnywhere, Category="Key")
 	FMovieSceneObjectBindingID Object;
+
+	UPROPERTY(EditAnywhere, Category="Key")
+	FName ComponentName;
+
+	UPROPERTY(EditAnywhere, Category="Key")
+	FName SocketName;
 };
 
 /** A curve of events */
@@ -74,9 +80,10 @@ struct MOVIESCENETRACKS_API FMovieSceneActorReferenceData : public FMovieSceneCh
 	 * Evaluate this channel
 	 *
 	 * @param InTime     The time to evaluate at
-	 * @return the result of the evaluation
+	 * @param OutValue   A value to receive the result
+	 * @return true if the channel was evaluated successfully, false otherwise
 	 */
-	FMovieSceneActorReferenceKey Evaluate(FFrameTime InTime) const;
+	bool Evaluate(FFrameTime InTime, FMovieSceneActorReferenceKey& OutValue) const;
 
 public:
 
@@ -179,6 +186,5 @@ private:
 
 inline bool EvaluateChannel(const FMovieSceneActorReferenceData* InChannel, FFrameTime InTime, FMovieSceneActorReferenceKey& OutValue)
 {
-	OutValue = InChannel->Evaluate(InTime);
-	return true;
+	return InChannel->Evaluate(InTime, OutValue);
 }
\ No newline at end of file
diff --git a/Engine/Source/Runtime/MovieSceneTracks/Public/Sections/MovieSceneAudioSection.h b/Engine/Source/Runtime/MovieSceneTracks/Public/Sections/MovieSceneAudioSection.h
--- a/Engine/Source/Runtime/MovieSceneTracks/Public/Sections/MovieSceneAudioSection.h
+++ b/Engine/Source/Runtime/MovieSceneTracks/Public/Sections/MovieSceneAudioSection.h
@@ -104,8 +104,16 @@ public:
 		return AttenuationSettings;
 	}
 
+	/*
+	 * @return The attach actor data
+	 */
 	const FMovieSceneActorReferenceData& GetAttachActorData() const { return AttachActorData; }
 
+	/*
+	 * @return The attach component given the bound actor and the actor attach key with the component and socket names
+	 */
+	USceneComponent* GetAttachComponent(const AActor* InParentActor, const FMovieSceneActorReferenceKey& Key) const;
+
 	/** ~UObject interface */
 	virtual void PostLoad() override;
 
@@ -152,6 +160,12 @@ public:
 	virtual TOptional<FFrameTime> GetOffsetTime() const override;
 	virtual FMovieSceneEvalTemplatePtr GenerateTemplate() const override;
 
+protected:
+
+	virtual void Serialize(FArchive& Ar) override;
+	virtual void PostEditImport() override;
+	void UpdateChannelProxy();
+
 private:
 
 	/** The sound cue or wave that this section plays */
