Sequencer: Fix ActorReferenceData AddOrUpdateKey so that it evalutes the current value if there's no current actor binding

#rnx

diff --git a/Engine/Source/Editor/MovieSceneTools/Private/Channels/BuiltInChannelEditors.cpp b/Engine/Source/Editor/MovieSceneTools/Private/Channels/BuiltInChannelEditors.cpp
--- a/Engine/Source/Editor/MovieSceneTools/Private/Channels/BuiltInChannelEditors.cpp
+++ b/Engine/Source/Editor/MovieSceneTools/Private/Channels/BuiltInChannelEditors.cpp
@@ -78,8 +78,6 @@ FKeyHandle AddOrUpdateKey(FMovieSceneFloatChannel* Channel, UMovieSceneSection*
 
 FKeyHandle AddOrUpdateKey(FMovieSceneActorReferenceData* Channel, UMovieSceneSection* SectionToKey, FFrameNumber InTime, ISequencer& Sequencer, const FGuid& InObjectBindingID, FTrackInstancePropertyBindings* PropertyBindings)
 {
-	AActor* CurrentActor = nullptr;
-
 	if (PropertyBindings && InObjectBindingID.IsValid())
 	{
 		for (TWeakObjectPtr<> WeakObject : Sequencer.FindBoundObjects(InObjectBindingID, Sequencer.GetFocusedTemplateID()))
@@ -87,17 +85,26 @@ FKeyHandle AddOrUpdateKey(FMovieSceneActorReferenceData* Channel, UMovieSceneSec
 			if (UObject* Object = WeakObject.Get())
 			{
 				// Care is taken here to ensure that we call GetCurrentValue with the correct instantiation of UObject* rather than AActor*
-				CurrentActor = Cast<AActor>(PropertyBindings->GetCurrentValue<UObject*>(*Object));
-				break;
+				AActor* CurrentActor = Cast<AActor>(PropertyBindings->GetCurrentValue<UObject*>(*Object));
+				if (CurrentActor)
+				{
+					FGuid ThisGuid = Sequencer.FindObjectId(*CurrentActor, Sequencer.GetFocusedTemplateID());
+
+					FMovieSceneObjectBindingID NewValue(ThisGuid, MovieSceneSequenceID::Root, EMovieSceneObjectBindingSpace::Local);
+
+					int32 NewIndex = Channel->GetData().AddKey(InTime, NewValue);
+
+					return Channel->GetData().GetHandle(NewIndex);
+				}
 			}
 		}
 	}
 
-	FGuid ThisGuid = CurrentActor ? Sequencer.FindObjectId(*CurrentActor, Sequencer.GetFocusedTemplateID()) : FGuid();
+	FMovieSceneActorReferenceKey NewValue;
+
+	Channel->Evaluate(InTime, NewValue);
 
-	FMovieSceneObjectBindingID NewValue(ThisGuid, MovieSceneSequenceID::Root, EMovieSceneObjectBindingSpace::Local);
-	int32 NewIndex = Channel->GetData().AddKey(InTime, NewValue);
-	return Channel->GetData().GetHandle(NewIndex);
+	return Channel->GetData().UpdateOrAddKey(InTime, NewValue);
 }
 
 bool CanCreateKeyEditor(const FMovieSceneBoolChannel*    Channel)
