Sequencer: Don't display thumbnails over blending areas of the camera cut track.

diff --git a/Engine/Source/Editor/MovieSceneTools/Private/Sections/ThumbnailSection.cpp b/Engine/Source/Editor/MovieSceneTools/Private/Sections/ThumbnailSection.cpp
--- a/Engine/Source/Editor/MovieSceneTools/Private/Sections/ThumbnailSection.cpp
+++ b/Engine/Source/Editor/MovieSceneTools/Private/Sections/ThumbnailSection.cpp
@@ -224,7 +224,14 @@ int32 FThumbnailSection::OnPaintSection( FSequencerSectionPainter& InPainter ) c
 
 	const float TimePerPx = GenerationRange.Size<double>() / InPainter.SectionGeometry.GetLocalSize().X;
 
-	FSlateRect ThumbnailClipRect = SectionGeometry.GetLayoutBoundingRect().InsetBy(FMargin(SectionThumbnailPadding, 0.f)).IntersectionWith(InPainter.SectionClippingRect);
+	const FFrameRate TickResolution = Section->GetTypedOuter<UMovieScene>()->GetTickResolution();
+	const double SectionEaseInDuration = TickResolution.AsSeconds(Section->Easing.GetEaseInDuration()) / TimePerPx;
+	const double SectionEaseOutDuration = TickResolution.AsSeconds(Section->Easing.GetEaseOutDuration()) / TimePerPx;
+
+	const FSlateRect ThumbnailClipRect = SectionGeometry.GetLayoutBoundingRect()
+		.InsetBy(FMargin(SectionThumbnailPadding, 0.f))
+		.InsetBy(FMargin(SectionEaseInDuration, 0.f, SectionEaseOutDuration, 0.f))
+		.IntersectionWith(InPainter.SectionClippingRect);
 
 	for (const TSharedPtr<FTrackEditorThumbnail>& Thumbnail : ThumbnailCache.GetThumbnails())
 	{
@@ -251,7 +258,6 @@ int32 FThumbnailSection::OnPaintSection( FSequencerSectionPainter& InPainter ) c
 
 		const float PositionY = (SectionGeometry.GetLocalSize().Y - ThumbnailCropSize.Y)*.5f;
 
-
 		FPaintGeometry PaintGeometry = SectionGeometry.ToPaintGeometry(
 			ThumbnailRTSize,
 			FSlateLayoutTransform(ThumbnailScale, FVector2D(PositionX-HorizontalCropOffset, PositionY))
@@ -276,7 +282,7 @@ int32 FThumbnailSection::OnPaintSection( FSequencerSectionPainter& InPainter ) c
 		FGeometry ClipGeometry = SectionGeometry.MakeChild(
 			ThumbnailCropSize,
 			FSlateLayoutTransform(
-				FVector2D( PositionX, PositionY)
+				FVector2D(PositionX, PositionY)
 			)
 		);
 
