Sequencer: Add missing track editors for new shake source track.

diff --git a/Engine/Source/Editor/MovieSceneTools/Private/TrackEditors/CameraShakeSourceShakeTrackEditor.cpp b/Engine/Source/Editor/MovieSceneTools/Private/TrackEditors/CameraShakeSourceShakeTrackEditor.cpp
new file mode 100644
--- /dev/null
+++ b/Engine/Source/Editor/MovieSceneTools/Private/TrackEditors/CameraShakeSourceShakeTrackEditor.cpp
@@ -0,0 +1,349 @@
+// Copyright Epic Games, Inc. All Rights Reserved.
+
+#include "TrackEditors/CameraShakeSourceShakeTrackEditor.h"
+#include "AssetRegistryModule.h"
+#include "Camera/CameraShake.h"
+#include "Camera/CameraShakeSourceComponent.h"
+#include "ContentBrowserModule.h"
+#include "ContentBrowserModule.h"
+#include "Engine/Blueprint.h"
+#include "Framework/Application/SlateApplication.h"
+#include "Framework/MultiBox/MultiBoxBuilder.h"
+#include "GameFramework/Actor.h"
+#include "IContentBrowserSingleton.h"
+#include "Sections/MovieSceneCameraShakeSourceShakeSection.h"
+#include "SequencerUtilities.h"
+#include "Tracks/MovieSceneCameraShakeSourceShakeTrack.h"
+#include "Widgets/Input/SButton.h"
+#include "Widgets/Layout/SBox.h"
+#include "Widgets/SBoxPanel.h"
+
+#define LOCTEXT_NAMESPACE "FCameraShakeSourceShakeTrackEditor"
+
+static const FString CameraShakeClassPath(TEXT("Class'/Script/Engine.CameraShake'"));
+
+class FCameraShakeSection : public FSequencerSection
+{
+public:
+	FCameraShakeSection(UMovieSceneSection& InSection)
+		: FSequencerSection(InSection)
+	{}
+
+	virtual FText GetSectionTitle() const override
+	{
+		if (const UMovieSceneCameraShakeSourceShakeSection* ShakeSection = Cast<UMovieSceneCameraShakeSourceShakeSection>(WeakSection.Get()))
+		{
+			const UClass* ShakeClass = ShakeSection ? ShakeSection->ShakeData.ShakeClass : nullptr;
+			if (ShakeClass != nullptr)
+			{
+				return FText::FromString(ShakeClass->GetName());
+			}
+			else
+			{
+				// TODO-ludovic: display the class name of the shake setup on the source actor.
+				return LOCTEXT("AutoCameraShake", "Automatic");
+			}
+		}
+		return LOCTEXT("NoCameraShake", "No Camera Shake");
+	}
+};
+
+FCameraShakeSourceShakeTrackEditor::FCameraShakeSourceShakeTrackEditor(TSharedRef<ISequencer> InSequencer)
+	: FMovieSceneTrackEditor(InSequencer) 
+{ 
+}
+
+TSharedRef<ISequencerTrackEditor> FCameraShakeSourceShakeTrackEditor::CreateTrackEditor(TSharedRef<ISequencer> InSequencer)
+{
+	return MakeShareable(new FCameraShakeSourceShakeTrackEditor(InSequencer));
+}
+
+bool FCameraShakeSourceShakeTrackEditor::SupportsType(TSubclassOf<UMovieSceneTrack> Type) const
+{
+	return Type == UMovieSceneCameraShakeSourceShakeTrack::StaticClass();
+}
+
+TSharedRef<ISequencerSection> FCameraShakeSourceShakeTrackEditor::MakeSectionInterface(UMovieSceneSection& SectionObject, UMovieSceneTrack& Track, FGuid ObjectBinding)
+{
+	check(SupportsType(SectionObject.GetOuter()->GetClass()));
+	return MakeShareable(new FCameraShakeSection(SectionObject));
+}
+
+void FCameraShakeSourceShakeTrackEditor::AddKey(const FGuid& ObjectGuid)
+{
+	TArray<TWeakObjectPtr<>> OutObjects;
+	for (TWeakObjectPtr<> Object : GetSequencer()->FindObjectsInCurrentSequence(ObjectGuid))
+	{
+		OutObjects.Add(Object);
+	}
+	
+	// Add a key that defaults to "auto camera shake".
+	const FScopedTransaction Transaction(LOCTEXT("AddCameraShakeSourceShake_Transaction", "Add Camera Shake"));
+
+	AnimatablePropertyChanged(FOnKeyProperty::CreateSP(this, &FCameraShakeSourceShakeTrackEditor::AddKeyInternal, OutObjects));
+}
+
+void FCameraShakeSourceShakeTrackEditor::BuildObjectBindingTrackMenu(FMenuBuilder& MenuBuilder, const TArray<FGuid>& ObjectBindings, const UClass* ObjectClass)
+{
+	if (const UCameraShakeSourceComponent* ShakeSourceComponent = AcquireCameraShakeSourceComponentFromGuid(ObjectBindings[0]))
+	{
+		MenuBuilder.AddMenuEntry(
+			LOCTEXT("AddShakeSourceShake", "Camera Shake"),
+			LOCTEXT("AddShakeSourceShakeTooltip", "Adds a camera shake originating from the parent camera shake source."),
+			FSlateIcon(),
+			FUIAction( 
+				FExecuteAction::CreateSP( this, &FCameraShakeSourceShakeTrackEditor::AddCameraShakeSection, ObjectBindings)
+			)
+		);
+	}
+}
+
+void FCameraShakeSourceShakeTrackEditor::AddCameraShakeSection(TArray<FGuid> ObjectHandles)
+{
+	TSharedPtr<ISequencer> SequencerPtr = GetSequencer();
+	if (!SequencerPtr.IsValid() || !SequencerPtr->IsAllowedToChange())
+	{
+		return;
+	}
+
+	TArray<TWeakObjectPtr<UObject>> Objects;
+	for (FGuid ObjectHandle : ObjectHandles)
+	{
+		for (TWeakObjectPtr<UObject> Object : SequencerPtr->FindObjectsInCurrentSequence(ObjectHandle))
+		{
+			Objects.Add(Object);
+		}
+	}
+
+	auto OnAddShakeSourceShakeSection = [=](FFrameNumber Time) -> FKeyPropertyResult
+	{
+		return this->AddKeyInternal(Time, Objects);
+	};
+
+	const FScopedTransaction Transaction(LOCTEXT("AddCameraShakeSourceShake_Transaction", "Add Camera Shake"));
+
+	AnimatablePropertyChanged(FOnKeyProperty::CreateLambda(OnAddShakeSourceShakeSection));
+}
+
+TSharedPtr<SWidget> FCameraShakeSourceShakeTrackEditor::BuildOutlinerEditWidget(const FGuid& ObjectBinding, UMovieSceneTrack* Track, const FBuildEditWidgetParams& Params)
+{
+	return SNew(SHorizontalBox)
+		+SHorizontalBox::Slot()
+		.AutoWidth()
+		.VAlign(VAlign_Center)
+		[
+			FSequencerUtilities::MakeAddButton(
+					LOCTEXT("AddShakeSourceShakeSection", "Camera Shake"),
+					FOnGetContent::CreateSP(this, &FCameraShakeSourceShakeTrackEditor::BuildCameraShakeSubMenu, ObjectBinding),
+					Params.NodeIsHovered, GetSequencer())
+		];
+}
+
+FKeyPropertyResult FCameraShakeSourceShakeTrackEditor::AddKeyInternal(FFrameNumber KeyTime, const TArray<TWeakObjectPtr<UObject>> Objects)
+{
+	return AddKeyInternal(KeyTime, Objects, TSubclassOf<UCameraShake>());
+}
+
+FKeyPropertyResult FCameraShakeSourceShakeTrackEditor::AddKeyInternal(FFrameNumber KeyTime, const TArray<TWeakObjectPtr<UObject>> Objects, TSubclassOf<UCameraShake> CameraShake)
+{
+	FKeyPropertyResult KeyPropertyResult;
+
+	for (int32 ObjectIndex = 0; ObjectIndex < Objects.Num(); ++ObjectIndex)
+	{
+		UObject* Object = Objects[ObjectIndex].Get();
+
+		if (AActor* Actor = Cast<AActor>(Object))
+		{
+			UCameraShakeSourceComponent* Component = Actor->FindComponentByClass<UCameraShakeSourceComponent>();
+			if (ensure(Component != nullptr))
+			{
+				Object = Component;
+			}
+		}
+
+		FFindOrCreateHandleResult HandleResult = FindOrCreateHandleToObject(Object);
+		FGuid ObjectHandle = HandleResult.Handle;
+		KeyPropertyResult.bHandleCreated |= HandleResult.bWasCreated;
+
+		if (ObjectHandle.IsValid())
+		{
+			FFindOrCreateTrackResult TrackResult = FindOrCreateTrackForObject(ObjectHandle, UMovieSceneCameraShakeSourceShakeTrack::StaticClass());
+			UMovieSceneTrack* Track = TrackResult.Track;
+			KeyPropertyResult.bTrackCreated |= TrackResult.bWasCreated;
+
+			if (ensure(Track))
+			{
+				UMovieSceneSection* NewSection = Cast<UMovieSceneCameraShakeSourceShakeTrack>(Track)->AddNewCameraShake(KeyTime, CameraShake);
+				KeyPropertyResult.bTrackModified = true;
+				
+				GetSequencer()->EmptySelection();
+				GetSequencer()->SelectSection(NewSection);
+				GetSequencer()->ThrobSectionSelection();
+			}
+		}
+	}
+
+	return KeyPropertyResult;
+}
+
+TSharedRef<SWidget> FCameraShakeSourceShakeTrackEditor::BuildCameraShakeSubMenu(FGuid ObjectBinding)
+{
+	FMenuBuilder MenuBuilder(true, nullptr);
+
+	TArray<FGuid> ObjectBindings;
+	ObjectBindings.Add(ObjectBinding);
+
+	AddCameraShakeSubMenu(MenuBuilder, ObjectBindings);
+
+	return MenuBuilder.MakeWidget();
+}
+
+void FCameraShakeSourceShakeTrackEditor::AddCameraShakeSubMenu(FMenuBuilder& MenuBuilder, TArray<FGuid> ObjectBindings)
+{
+	MenuBuilder.AddMenuEntry(
+			LOCTEXT("AutoShake", "Automatic Shake"),
+			LOCTEXT("AutoShakeTooltip", "Adds a section that plays the camera shake already configured on the shake source component."),
+			FSlateIcon(),
+			FUIAction( 
+				FExecuteAction::CreateSP(this, &FCameraShakeSourceShakeTrackEditor::OnAutoCameraShakeSelected, ObjectBindings)
+			));
+
+	MenuBuilder.AddSubMenu(
+			LOCTEXT("OtherShake", "Other Shake"),
+			LOCTEXT("OtherShakeTooltip", "Adds a section that plays a specific camera shake originating from the shake source component."),
+			FNewMenuDelegate::CreateSP(this, &FCameraShakeSourceShakeTrackEditor::AddOtherCameraShakeBrowserSubMenu, ObjectBindings));
+}
+
+void FCameraShakeSourceShakeTrackEditor::AddOtherCameraShakeBrowserSubMenu(FMenuBuilder& MenuBuilder, TArray<FGuid> ObjectBindings)
+{
+	FAssetPickerConfig AssetPickerConfig;
+	{
+		AssetPickerConfig.OnAssetSelected = FOnAssetSelected::CreateSP(this, &FCameraShakeSourceShakeTrackEditor::OnCameraShakeAssetSelected, ObjectBindings);
+		AssetPickerConfig.OnAssetEnterPressed = FOnAssetEnterPressed::CreateSP(this, &FCameraShakeSourceShakeTrackEditor::OnCameraShakeAssetEnterPressed, ObjectBindings);
+		AssetPickerConfig.OnShouldFilterAsset = FOnShouldFilterAsset::CreateSP(this, &FCameraShakeSourceShakeTrackEditor::OnShouldFilterCameraShake);
+		AssetPickerConfig.bAllowNullSelection = false;
+		AssetPickerConfig.InitialAssetViewType = EAssetViewType::List;
+		AssetPickerConfig.Filter.ClassNames.Add(UBlueprint::StaticClass()->GetFName());
+		AssetPickerConfig.Filter.TagsAndValues.Add(FBlueprintTags::ParentClassPath, CameraShakeClassPath);
+	}
+
+	FContentBrowserModule& ContentBrowserModule = FModuleManager::Get().LoadModuleChecked<FContentBrowserModule>(TEXT("ContentBrowser"));
+
+	TSharedPtr<SBox> MenuEntry = SNew(SBox)
+		.WidthOverride(300.0f)
+		.HeightOverride(400.f)
+		[
+			ContentBrowserModule.Get().CreateAssetPicker(AssetPickerConfig)
+		];
+
+	MenuBuilder.AddWidget(MenuEntry.ToSharedRef(), FText::GetEmpty(), true);
+}
+
+void FCameraShakeSourceShakeTrackEditor::OnCameraShakeAssetSelected(const FAssetData& AssetData, TArray<FGuid> ObjectBindings)
+{
+	FSlateApplication::Get().DismissAllMenus();
+
+	TSharedPtr<ISequencer> SequencerPtr = GetSequencer();
+	UBlueprint* SelectedObject = Cast<UBlueprint>(AssetData.GetAsset());
+
+	if (SelectedObject && SelectedObject->GeneratedClass && SelectedObject->GeneratedClass->IsChildOf(UCameraShake::StaticClass()))
+	{
+		TSubclassOf<UCameraShake> CameraShakeClass = *(SelectedObject->GeneratedClass);
+
+		TArray<TWeakObjectPtr<>> OutObjects;
+		for (FGuid ObjectBinding : ObjectBindings)
+		{
+			for (TWeakObjectPtr<> Object : SequencerPtr->FindObjectsInCurrentSequence(ObjectBinding))
+			{
+				OutObjects.Add(Object);
+			}
+		}
+		
+		const FScopedTransaction Transaction(LOCTEXT("AddCameraShakeSourceShake_Transaction", "Add Camera Shake"));
+
+		AnimatablePropertyChanged(FOnKeyProperty::CreateSP(this, &FCameraShakeSourceShakeTrackEditor::AddKeyInternal, OutObjects, CameraShakeClass));
+	}
+}
+
+void FCameraShakeSourceShakeTrackEditor::OnCameraShakeAssetEnterPressed(const TArray<FAssetData>& AssetData, TArray<FGuid> ObjectBindings)
+{
+	if (AssetData.Num() > 0)
+	{
+		OnCameraShakeAssetSelected(AssetData[0].GetAsset(), ObjectBindings);
+	}
+}
+
+void FCameraShakeSourceShakeTrackEditor::OnAutoCameraShakeSelected(TArray<FGuid> ObjectBindings)
+{
+	TArray<TWeakObjectPtr<>> OutObjects;
+	TSharedPtr<ISequencer> SequencerPtr = GetSequencer();
+
+	for (FGuid ObjectBinding : ObjectBindings)
+	{
+		for (TWeakObjectPtr<> Object : SequencerPtr->FindObjectsInCurrentSequence(ObjectBinding))
+		{
+			OutObjects.Add(Object);
+		}
+	}
+
+	const FScopedTransaction Transaction(LOCTEXT("AddCameraShakeSourceShake_Transaction", "Add Camera Shake"));
+
+	AnimatablePropertyChanged(FOnKeyProperty::CreateSP(this, &FCameraShakeSourceShakeTrackEditor::AddKeyInternal, OutObjects));
+
+	//return FReply::Handled();
+}
+
+bool FCameraShakeSourceShakeTrackEditor::OnShouldFilterCameraShake(const FAssetData& AssetData)
+{
+	const UBlueprint* SelectedObject = Cast<UBlueprint>(AssetData.GetAsset());
+	if (SelectedObject && SelectedObject->GeneratedClass && SelectedObject->GeneratedClass->IsChildOf(UCameraShake::StaticClass()))
+	{
+		TSubclassOf<UCameraShake> CameraShakeClass = *(SelectedObject->GeneratedClass);
+		if (const UCameraShake* CameraShakeCDO = Cast<UCameraShake>(CameraShakeClass->ClassDefaultObject))
+		{
+			return CameraShakeCDO->bSingleInstance;
+		}
+	}
+	return true;
+}
+
+UCameraShakeSourceComponent* FCameraShakeSourceShakeTrackEditor::AcquireCameraShakeSourceComponentFromGuid(const FGuid& Guid)
+{
+	TArray<UCameraShakeSourceComponent*> ShakeSourceComponents;
+
+	for (TWeakObjectPtr<> WeakObject : GetSequencer()->FindObjectsInCurrentSequence(Guid))
+	{
+		if (UObject* Obj = WeakObject.Get())
+		{
+			if (AActor* Actor = Cast<AActor>(Obj))
+			{
+				TArray<UCameraShakeSourceComponent*> CurShakeSourceComponents;
+				Actor->GetComponents<UCameraShakeSourceComponent>(CurShakeSourceComponents);
+				ShakeSourceComponents.Append(CurShakeSourceComponents);
+			}
+			else if (UCameraShakeSourceComponent* ShakeSourceComponent = Cast<UCameraShakeSourceComponent>(Obj))
+			{
+				ShakeSourceComponents.Add(ShakeSourceComponent);
+			}
+		}
+	}
+
+	UCameraShakeSourceComponent** ActiveComponent = ShakeSourceComponents.FindByPredicate([](UCameraShakeSourceComponent* Component)
+			{
+				return Component->IsActive();
+			});
+	if (ActiveComponent != nullptr)
+	{
+		return *ActiveComponent;
+	}
+
+	if (ShakeSourceComponents.Num() > 0)
+	{
+		return ShakeSourceComponents[0];
+	}
+
+	return nullptr;
+}
+
+#undef LOCTEXT_NAMESPACE
+
diff --git a/Engine/Source/Editor/MovieSceneTools/Private/TrackEditors/CameraShakeSourceShakeTrackEditor.h b/Engine/Source/Editor/MovieSceneTools/Private/TrackEditors/CameraShakeSourceShakeTrackEditor.h
new file mode 100644
--- /dev/null
+++ b/Engine/Source/Editor/MovieSceneTools/Private/TrackEditors/CameraShakeSourceShakeTrackEditor.h
@@ -0,0 +1,51 @@
+// Copyright Epic Games, Inc. All Rights Reserved.
+
+#pragma once
+
+#include "CoreMinimal.h"
+#include "Misc/Guid.h"
+#include "Templates/SubclassOf.h"
+#include "Widgets/SWidget.h"
+#include "ISequencer.h"
+#include "MovieSceneTrack.h"
+#include "ISequencerSection.h"
+#include "ISequencerTrackEditor.h"
+#include "MovieSceneTrackEditor.h"
+#include "Camera/CameraShake.h"
+
+struct FAssetData;
+class FMenuBuilder;
+class UCameraShakeSourceComponent;
+
+class FCameraShakeSourceShakeTrackEditor : public FMovieSceneTrackEditor
+{
+public:
+	static TSharedRef<ISequencerTrackEditor> CreateTrackEditor(TSharedRef<ISequencer> OwningSequencer);
+
+	FCameraShakeSourceShakeTrackEditor(TSharedRef<ISequencer> InSequencer);
+	virtual ~FCameraShakeSourceShakeTrackEditor() {}
+
+	// ISequencerTrackEditor interface
+	virtual void AddKey(const FGuid& ObjectGuid) override;
+	virtual void BuildObjectBindingTrackMenu(FMenuBuilder& MenuBuilder, const TArray<FGuid>& ObjectBindings, const UClass* ObjectClass) override;
+	virtual TSharedRef<ISequencerSection> MakeSectionInterface(UMovieSceneSection& SectionObject, UMovieSceneTrack& Track, FGuid ObjectBinding) override;
+	virtual bool SupportsType(TSubclassOf<UMovieSceneTrack> Type) const override;
+	virtual TSharedPtr<SWidget> BuildOutlinerEditWidget(const FGuid& ObjectBinding, UMovieSceneTrack* Track, const FBuildEditWidgetParams& Params) override;
+
+private:
+	FKeyPropertyResult AddKeyInternal(FFrameNumber KeyTime, const TArray<TWeakObjectPtr<UObject>> Objects);
+	FKeyPropertyResult AddKeyInternal(FFrameNumber KeyTime, const TArray<TWeakObjectPtr<UObject>> Objects, TSubclassOf<UCameraShake> CameraShake);
+
+	void AddCameraShakeSection(TArray<FGuid> ObjectHandles);
+
+	TSharedRef<SWidget> BuildCameraShakeSubMenu(FGuid ObjectBinding);
+	void AddCameraShakeSubMenu(FMenuBuilder& MenuBuilder, TArray<FGuid> ObjectBindings);
+	void AddOtherCameraShakeBrowserSubMenu(FMenuBuilder& MenuBuilder, TArray<FGuid> ObjectBindings);
+	void OnCameraShakeAssetSelected(const FAssetData& AssetData, TArray<FGuid> ObjectBindings);
+	void OnCameraShakeAssetEnterPressed(const TArray<FAssetData>& AssetData, TArray<FGuid> ObjectBindings);
+	void OnAutoCameraShakeSelected(TArray<FGuid> ObjectBindings);
+	bool OnShouldFilterCameraShake(const FAssetData& AssetData);
+	
+	UCameraShakeSourceComponent* AcquireCameraShakeSourceComponentFromGuid(const FGuid& Guid);
+};
+
