Sequencer: Fix unity build conflicts.

diff --git a/Engine/Source/Editor/MovieSceneTools/Private/TrackEditors/CameraShakeSourceShakeTrackEditor.cpp b/Engine/Source/Editor/MovieSceneTools/Private/TrackEditors/CameraShakeSourceShakeTrackEditor.cpp
--- a/Engine/Source/Editor/MovieSceneTools/Private/TrackEditors/CameraShakeSourceShakeTrackEditor.cpp
+++ b/Engine/Source/Editor/MovieSceneTools/Private/TrackEditors/CameraShakeSourceShakeTrackEditor.cpp
@@ -20,12 +20,10 @@
 
 #define LOCTEXT_NAMESPACE "FCameraShakeSourceShakeTrackEditor"
 
-static const FString CameraShakeClassPath(TEXT("Class'/Script/Engine.CameraShake'"));
-
-class FCameraShakeSection : public FSequencerSection
+class FCameraShakeSourceShakeSection : public FSequencerSection
 {
 public:
-	FCameraShakeSection(UMovieSceneSection& InSection)
+	FCameraShakeSourceShakeSection(UMovieSceneSection& InSection)
 		: FSequencerSection(InSection)
 	{}
 
@@ -66,7 +64,7 @@ bool FCameraShakeSourceShakeTrackEditor::SupportsType(TSubclassOf<UMovieSceneTra
 TSharedRef<ISequencerSection> FCameraShakeSourceShakeTrackEditor::MakeSectionInterface(UMovieSceneSection& SectionObject, UMovieSceneTrack& Track, FGuid ObjectBinding)
 {
 	check(SupportsType(SectionObject.GetOuter()->GetClass()));
-	return MakeShareable(new FCameraShakeSection(SectionObject));
+	return MakeShareable(new FCameraShakeSourceShakeSection(SectionObject));
 }
 
 void FCameraShakeSourceShakeTrackEditor::AddKey(const FGuid& ObjectGuid)
@@ -216,6 +214,8 @@ void FCameraShakeSourceShakeTrackEditor::AddCameraShakeSubMenu(FMenuBuilder& Men
 
 void FCameraShakeSourceShakeTrackEditor::AddOtherCameraShakeBrowserSubMenu(FMenuBuilder& MenuBuilder, TArray<FGuid> ObjectBindings)
 {
+	static const FString CameraShakeClassPath(TEXT("Class'/Script/Engine.CameraShake'"));
+
 	FAssetPickerConfig AssetPickerConfig;
 	{
 		AssetPickerConfig.OnAssetSelected = FOnAssetSelected::CreateSP(this, &FCameraShakeSourceShakeTrackEditor::OnCameraShakeAssetSelected, ObjectBindings);
