Sequencer: Add some delta-time methods to the time converter.

diff --git a/Engine/Source/Editor/MovieSceneTools/Public/CommonMovieSceneTools.h b/Engine/Source/Editor/MovieSceneTools/Public/CommonMovieSceneTools.h
--- a/Engine/Source/Editor/MovieSceneTools/Public/CommonMovieSceneTools.h
+++ b/Engine/Source/Editor/MovieSceneTools/Public/CommonMovieSceneTools.h
@@ -35,6 +35,17 @@ public:
 		return (Time - ViewRangeStartSeconds) * PixelsPerSecond;
 	}
 
+	/**
+	 * Converts a time delta to a pixel delta
+	 *
+	 * @param TimeDelta  The time delta to convert
+	 * @return           The pixel equivalent of the delta time
+	 */
+	float SecondsDeltaToPixel( double TimeDelta ) const
+	{
+		return TimeDelta * PixelsPerSecond;
+	}
+
 	/**
 	 * Converts a pixel value to time 
 	 *
@@ -57,6 +68,17 @@ public:
 		return (Time / TickResolution - ViewRangeStartSeconds) * PixelsPerSecond;
 	}
 
+	/**
+	 * Converts a frame delta value to pixel delta
+	 *
+	 * @param TimeDelta   The time delta to convert
+	 * @return            The pixel equivalent of the delta time
+	 */
+	float FrameDeltaToPixel( const FFrameTime& TimeDelta) const
+	{
+		return (TimeDelta / TickResolution) * PixelsPerSecond;
+	}
+
 	/**
 	 * Converts a pixel value to frame time 
 	 *
