Sequencer: improvements to camera shake source tracks.

Add a second track so that there are 2 ways to start and preview camera shakes: a "controlled" track where sections control when the shake starts and stops, and a "trigger" track which is a "fire-and-forget" way to trigger shakes.
Moved all the shake previewing into a separate class so that both tracks can use it.

Controlled shakes track improvements:
- Handle blending camera shakes out.
- Show a marker for the end of the shake if it has a duration, for when the section is longer than that.
- Show shake intensity curve if the shake has blend in/out times.
- Give new sections an appropriate length based on their shake's duration.

diff --git a/Engine/Source/Editor/MovieSceneTools/Private/MovieSceneToolsModule.cpp b/Engine/Source/Editor/MovieSceneTools/Private/MovieSceneToolsModule.cpp
--- a/Engine/Source/Editor/MovieSceneTools/Private/MovieSceneToolsModule.cpp
+++ b/Engine/Source/Editor/MovieSceneTools/Private/MovieSceneToolsModule.cpp
@@ -65,6 +65,7 @@
 #include "Channels/BuiltInChannelEditors.h"
 #include "Channels/MovieSceneObjectPathChannel.h"
 #include "Channels/MovieSceneEventChannel.h"
+#include "Channels/MovieSceneCameraShakeSourceTriggerChannel.h"
 #include "Channels/EventChannelCurveModel.h"
 #include "Channels/SCurveEditorEventChannelView.h"
 #include "Sections/MovieSceneEventSection.h"
@@ -147,6 +148,8 @@ void FMovieSceneToolsModule::StartupModule()
 
 		SequencerModule.RegisterChannelInterface<FMovieSceneEventChannel>();
 
+		SequencerModule.RegisterChannelInterface<FMovieSceneCameraShakeSourceTriggerChannel>();
+
 		ICurveEditorModule& CurveEditorModule = FModuleManager::LoadModuleChecked<ICurveEditorModule>("CurveEditor");
 
 		FEventChannelCurveModel::EventView = CurveEditorModule.RegisterView(FOnCreateCurveEditorView::CreateStatic(
diff --git a/Engine/Source/Editor/MovieSceneTools/Private/TrackEditors/CameraShakeSourceShakeTrackEditor.cpp b/Engine/Source/Editor/MovieSceneTools/Private/TrackEditors/CameraShakeSourceShakeTrackEditor.cpp
--- a/Engine/Source/Editor/MovieSceneTools/Private/TrackEditors/CameraShakeSourceShakeTrackEditor.cpp
+++ b/Engine/Source/Editor/MovieSceneTools/Private/TrackEditors/CameraShakeSourceShakeTrackEditor.cpp
@@ -4,48 +4,337 @@
 #include "AssetRegistryModule.h"
 #include "Camera/CameraShake.h"
 #include "Camera/CameraShakeSourceComponent.h"
+#include "CommonMovieSceneTools.h"
 #include "ContentBrowserModule.h"
-#include "ContentBrowserModule.h"
+#include "EditorStyleSet.h"
 #include "Engine/Blueprint.h"
+#include "Fonts/FontMeasure.h"
 #include "Framework/Application/SlateApplication.h"
 #include "Framework/MultiBox/MultiBoxBuilder.h"
 #include "GameFramework/Actor.h"
 #include "IContentBrowserSingleton.h"
+#include "MovieSceneTimeHelpers.h"
+#include "Rendering/DrawElements.h"
 #include "Sections/MovieSceneCameraShakeSourceShakeSection.h"
+#include "Sections/MovieSceneCameraShakeSourceTriggerSection.h"
+#include "SequencerSectionPainter.h"
 #include "SequencerUtilities.h"
 #include "Tracks/MovieSceneCameraShakeSourceShakeTrack.h"
+#include "Tracks/MovieSceneCameraShakeSourceTriggerTrack.h"
 #include "Widgets/Input/SButton.h"
 #include "Widgets/Layout/SBox.h"
 #include "Widgets/SBoxPanel.h"
 
 #define LOCTEXT_NAMESPACE "FCameraShakeSourceShakeTrackEditor"
 
-class FCameraShakeSourceShakeSection : public FSequencerSection
+/**
+ * Section interface for shake sections
+ */
+class FCameraShakeSourceShakeSection : public ISequencerSection
 {
 public:
-	FCameraShakeSourceShakeSection(UMovieSceneSection& InSection)
-		: FSequencerSection(InSection)
+	FCameraShakeSourceShakeSection(const TSharedPtr<ISequencer> InSequencer, UMovieSceneCameraShakeSourceShakeSection& InSection, const FGuid& InObjectBinding)
+		: SequencerPtr(InSequencer)
+		, SectionPtr(&InSection)
+		, ObjectBinding(InObjectBinding)
 	{}
 
 	virtual FText GetSectionTitle() const override
 	{
-		if (const UMovieSceneCameraShakeSourceShakeSection* ShakeSection = Cast<UMovieSceneCameraShakeSourceShakeSection>(WeakSection.Get()))
+		const TSubclassOf<UCameraShake> ShakeClass = GetCameraShakeClass();
+		if (const UClass* ShakeClassPtr = ShakeClass.Get())
+		{
+			return FText::FromString(ShakeClassPtr->GetName());
+		}
+		return LOCTEXT("NoCameraShake", "No Camera Shake");
+	}
+
+	virtual UMovieSceneSection* GetSectionObject() override
+	{
+		return SectionPtr.Get();
+	}
+
+	virtual bool IsReadOnly() const override 
+	{ 
+		return SectionPtr.IsValid() ? SectionPtr.Get()->IsReadOnly() : false; 
+	}
+
+	virtual int32 OnPaintSection(FSequencerSectionPainter& Painter) const override
+	{
+		static const FSlateBrush* GenericDivider = FEditorStyle::GetBrush("Sequencer.GenericDivider");
+
+		Painter.LayerId = Painter.PaintSectionBackground();
+
+		const UMovieSceneCameraShakeSourceShakeSection* SectionObject = SectionPtr.Get();
+		const TSharedPtr<ISequencer> Sequencer = SequencerPtr.Pin();
+		if (SectionObject == nullptr || !Sequencer.IsValid())
+		{
+			return Painter.LayerId;
+		}
+
+		const UMovieSceneSequence* FocusedSequence = Sequencer->GetFocusedMovieSceneSequence();
+		if (FocusedSequence == nullptr)
 		{
-			const UClass* ShakeClass = ShakeSection ? ShakeSection->ShakeData.ShakeClass : nullptr;
-			if (ShakeClass != nullptr)
+			return Painter.LayerId;
+		}
+
+		const FTimeToPixel& TimeConverter = Painter.GetTimeConverter();
+		const FFrameRate TickResolution = FocusedSequence->GetMovieScene()->GetTickResolution();
+
+		const ESlateDrawEffect DrawEffects = Painter.bParentEnabled ? ESlateDrawEffect::None : ESlateDrawEffect::DisabledEffect;
+		const TRange<FFrameNumber> SectionRange = SectionObject->GetRange();
+		const int32 SectionSize = MovieScene::DiscreteSize(SectionRange);
+		const float SectionDuration = FFrameNumber(SectionSize) / TickResolution;
+
+		const float SectionStartTime = SectionObject->GetInclusiveStartFrame() / TickResolution;
+		const float SectionStartTimeInPixels = TimeConverter.SecondsToPixel(SectionStartTime);
+		const float SectionEndTime = SectionObject->GetExclusiveEndFrame() / TickResolution;
+		const float SectionEndTimeInPixels = TimeConverter.SecondsToPixel(SectionEndTime);
+
+		float ShakeDuration = -1.f;
+		const TSubclassOf<UCameraShake> CameraShakeClass = GetCameraShakeClass();
+		UCameraShake::GetCameraShakeDuration(CameraShakeClass, ShakeDuration);
+
+		const float ShakeEndInPixels = (ShakeDuration > SMALL_NUMBER) ? 
+			TimeConverter.SecondsToPixel(FMath::Min(SectionStartTime + ShakeDuration, SectionEndTime)) :
+			SectionEndTimeInPixels;
+		const bool bSectionContainsEntireShake = (ShakeDuration > SMALL_NUMBER && SectionDuration > ShakeDuration);
+
+		if (bSectionContainsEntireShake && SectionRange.HasLowerBound())
+		{
+			// Add some separator where the shake ends.
+			float OffsetPixel = ShakeEndInPixels - SectionStartTimeInPixels;
+
+			FSlateDrawElement::MakeBox(
+					Painter.DrawElements,
+					Painter.LayerId++,
+					Painter.SectionGeometry.MakeChild(
+						FVector2D(2.f, Painter.SectionGeometry.Size.Y - 2.f),
+						FSlateLayoutTransform(FVector2D(OffsetPixel, 1.f))
+						).ToPaintGeometry(),
+					GenericDivider,
+					DrawEffects
+					);
+
+			// Draw the rest in a "muted" color.
+			const float OverflowSizeInPixels = SectionEndTimeInPixels - ShakeEndInPixels;
+
+			FSlateDrawElement::MakeBox(
+					Painter.DrawElements,
+					Painter.LayerId++,
+					Painter.SectionGeometry.MakeChild(
+						FVector2D(OverflowSizeInPixels, Painter.SectionGeometry.Size.Y),
+						FSlateLayoutTransform(FVector2D(OffsetPixel, 0))
+						).ToPaintGeometry(),
+					FEditorStyle::GetBrush("WhiteBrush"),
+					ESlateDrawEffect::None,
+					FLinearColor::Black.CopyWithNewOpacity(0.5f)
+					);
+		}
+
+		float ShakeBlendIn = 0.f, ShakeBlendOut = 0.f;
+		UCameraShake::GetCameraShakeBlendTimes(CameraShakeClass, ShakeBlendIn, ShakeBlendOut);
+		{
+			// Draw the shake "intensity" as a line that goes up and down according to
+			// blend in and out times.
+			const FLinearColor LineColor(0.25f, 0.25f, 1.f, 0.75f);
+
+			const bool bHasBlendIn = (ShakeBlendIn > SMALL_NUMBER);
+			const bool bHasBlendOut = (ShakeBlendOut > SMALL_NUMBER);
+
+			float ShakeBlendInEndInPixels = TimeConverter.SecondsToPixel(SectionStartTime + ShakeBlendIn);
+			float ShakeBlendOutStartInPixels = ShakeEndInPixels - TimeConverter.SecondsDeltaToPixel(ShakeBlendOut);
+			if (ShakeBlendInEndInPixels > ShakeBlendOutStartInPixels)
 			{
-				return FText::FromString(ShakeClass->GetName());
+				// If we have to blend out before we're done blending in, let's switch over
+				// at the half mark.
+				ShakeBlendInEndInPixels = ShakeBlendOutStartInPixels = (ShakeBlendInEndInPixels + ShakeBlendOutStartInPixels) / 2.f;
+			}
+
+			TArray<FVector2D> LinePoints;
+
+			if (bHasBlendIn)
+			{
+				LinePoints.Add(FVector2D(SectionStartTimeInPixels, Painter.SectionGeometry.Size.Y - 2.f));
+				LinePoints.Add(FVector2D(ShakeBlendInEndInPixels, 2.f));
 			}
 			else
 			{
-				// TODO-ludovic: display the class name of the shake setup on the source actor.
-				return LOCTEXT("AutoCameraShake", "Automatic");
+				LinePoints.Add(FVector2D(SectionStartTimeInPixels, 2.f));
+			}
+
+			if (bHasBlendOut)
+			{
+				LinePoints.Add(FVector2D(ShakeBlendOutStartInPixels, 2.f));
+				LinePoints.Add(FVector2D(ShakeEndInPixels, Painter.SectionGeometry.Size.Y - 2.f));
 			}
+			else
+			{
+				LinePoints.Add(FVector2D(ShakeEndInPixels, 2.f));
+			}
+
+			FSlateDrawElement::MakeLines(
+					Painter.DrawElements,
+					Painter.LayerId++,
+					Painter.SectionGeometry.ToPaintGeometry(),
+					LinePoints,
+					DrawEffects,
+					LineColor);
 		}
-		return LOCTEXT("NoCameraShake", "No Camera Shake");
+
+		return Painter.LayerId;
 	}
+
+private:
+	TSubclassOf<UCameraShake> GetCameraShakeClass() const
+	{
+		if (const UMovieSceneCameraShakeSourceShakeSection* SectionObject = SectionPtr.Get())
+		{
+			if (SectionObject->ShakeData.ShakeClass.Get() != nullptr)
+			{
+				return SectionObject->ShakeData.ShakeClass;
+			}
+		}
+
+		TSharedPtr<ISequencer> Sequencer = SequencerPtr.Pin();
+		TArrayView<TWeakObjectPtr<>> BoundObjects = Sequencer->FindBoundObjects(ObjectBinding, Sequencer->GetFocusedTemplateID());
+		if (BoundObjects.Num() > 0)
+		{
+			if (const UCameraShakeSourceComponent* Component = Cast<UCameraShakeSourceComponent>(BoundObjects[0].Get()))
+			{
+				return Component->CameraShake;
+			}
+		}
+
+		return TSubclassOf<UCameraShake>();
+	}
+
+	const UCameraShake* GetCameraShakeDefaultObject() const
+	{
+		const TSubclassOf<UCameraShake> ShakeClass = GetCameraShakeClass();
+		if (const UClass* ShakeClassPtr = ShakeClass.Get())
+		{
+			return ShakeClassPtr->GetDefaultObject<UCameraShake>();
+		}
+		return nullptr;
+	}
+
+private:
+	TWeakPtr<ISequencer> SequencerPtr;
+	TWeakObjectPtr<UMovieSceneCameraShakeSourceShakeSection> SectionPtr;
+	FGuid ObjectBinding;
 };
 
+/**
+ * Section interface for shake triggers.
+ */
+class FCameraShakeSourceTriggerSection : public FSequencerSection
+{
+public:
+	FCameraShakeSourceTriggerSection(const TSharedPtr<ISequencer> InSequencer, UMovieSceneCameraShakeSourceTriggerSection& InSectionObject)
+		: FSequencerSection(InSectionObject)
+		, Sequencer(InSequencer)
+	{}
+
+private:
+	virtual int32 OnPaintSection(FSequencerSectionPainter& Painter) const override;
+
+	bool IsTrackSelected() const;
+	void PaintShakeName(FSequencerSectionPainter& Painter, int32 LayerId, TSubclassOf<UCameraShake> ShakeClass, float PixelPos) const;
+
+	TWeakPtr<ISequencer> Sequencer;
+};
+
+bool FCameraShakeSourceTriggerSection::IsTrackSelected() const
+{
+	TSharedPtr<ISequencer> SequencerPtr = Sequencer.Pin();
+
+	TArray<UMovieSceneTrack*> SelectedTracks;
+	SequencerPtr->GetSelectedTracks(SelectedTracks);
+
+	UMovieSceneSection* Section = WeakSection.Get();
+	UMovieSceneTrack*   Track   = Section ? CastChecked<UMovieSceneTrack>(Section->GetOuter()) : nullptr;
+	return Track && SelectedTracks.Contains(Track);
+}
+
+void FCameraShakeSourceTriggerSection::PaintShakeName(FSequencerSectionPainter& Painter, int32 LayerId, TSubclassOf<UCameraShake> ShakeClass, float PixelPos) const
+{
+	static const int32   FontSize      = 10;
+	static const float   BoxOffsetPx   = 10.f;
+	static const FString AutoShakeText = LOCTEXT("AutoShake", "(Automatic)").ToString();
+
+	const FSlateFontInfo FontAwesomeFont = FEditorStyle::Get().GetFontStyle("FontAwesome.10");
+	const FSlateFontInfo SmallLayoutFont = FCoreStyle::GetDefaultFontStyle("Bold", 10);
+	const FLinearColor   DrawColor       = FEditorStyle::GetSlateColor("SelectionColor").GetColor(FWidgetStyle());
+
+	const FString ShakeText = (ShakeClass.Get() != nullptr) ? ShakeClass.Get()->GetName() : AutoShakeText;
+
+	TSharedRef<FSlateFontMeasure> FontMeasureService = FSlateApplication::Get().GetRenderer()->GetFontMeasureService();
+
+	const  FMargin   BoxPadding     = FMargin(4.0f, 2.0f);
+	const FVector2D  TextSize       = FontMeasureService->Measure(ShakeText, SmallLayoutFont);
+
+	// Flip the text position if getting near the end of the view range
+	const bool bDrawLeft = (Painter.SectionGeometry.Size.X - PixelPos) < (TextSize.X + 22.f) - BoxOffsetPx;
+	const float BoxPositionX = FMath::Max(
+			bDrawLeft ? PixelPos - TextSize.X - BoxOffsetPx : PixelPos + BoxOffsetPx,
+			0.f);
+
+	const FVector2D BoxOffset  = FVector2D(BoxPositionX, Painter.SectionGeometry.Size.Y*.5f - TextSize.Y*.5f);
+	const FVector2D TextOffset = FVector2D(BoxPadding.Left, 0);
+
+	// Draw the background box.
+	FSlateDrawElement::MakeBox(
+		Painter.DrawElements,
+		LayerId + 1,
+		Painter.SectionGeometry.ToPaintGeometry(BoxOffset, TextSize),
+		FEditorStyle::GetBrush("WhiteBrush"),
+		ESlateDrawEffect::None,
+		FLinearColor::Black.CopyWithNewOpacity(0.5f)
+	);
+
+	// Draw shake name.
+	FSlateDrawElement::MakeText(
+		Painter.DrawElements,
+		LayerId + 2,
+		Painter.SectionGeometry.ToPaintGeometry(BoxOffset + TextOffset, TextSize),
+		ShakeText,
+		SmallLayoutFont,
+		Painter.bParentEnabled ? ESlateDrawEffect::None : ESlateDrawEffect::DisabledEffect,
+		DrawColor
+	);
+}
+
+int32 FCameraShakeSourceTriggerSection::OnPaintSection(FSequencerSectionPainter& Painter) const
+{
+	int32 LayerId = Painter.PaintSectionBackground();
+
+	UMovieSceneCameraShakeSourceTriggerSection* TriggerSection = Cast<UMovieSceneCameraShakeSourceTriggerSection>(WeakSection.Get());
+	if (!TriggerSection || !IsTrackSelected())
+	{
+		return LayerId;
+	}
+
+	const FTimeToPixel& TimeToPixelConverter = Painter.GetTimeConverter();
+	const FMovieSceneCameraShakeSourceTriggerChannel& TriggerChannel = TriggerSection->GetChannel();
+	const TArrayView<const FFrameNumber> Times = TriggerChannel.GetData().GetTimes();
+	const TArrayView<const FMovieSceneCameraShakeSourceTrigger> Values = TriggerChannel.GetData().GetValues();
+	const TRange<FFrameNumber> SectionRange = TriggerSection->GetRange();
+
+	for (int32 KeyIndex = 0; KeyIndex < Times.Num(); ++KeyIndex)
+	{
+		const FFrameNumber Time = Times[KeyIndex];
+		if (SectionRange.Contains(Time))
+		{
+			const FMovieSceneCameraShakeSourceTrigger& Value = Values[KeyIndex];
+			const float PixelPos = TimeToPixelConverter.FrameToPixel(Time);
+			PaintShakeName(Painter, LayerId, Value.ShakeClass, PixelPos);
+		}
+	}
+
+	return LayerId + 3;
+}
+
 FCameraShakeSourceShakeTrackEditor::FCameraShakeSourceShakeTrackEditor(TSharedRef<ISequencer> InSequencer)
 	: FMovieSceneTrackEditor(InSequencer) 
 { 
@@ -58,42 +347,72 @@ TSharedRef<ISequencerTrackEditor> FCameraShakeSourceShakeTrackEditor::CreateTrac
 
 bool FCameraShakeSourceShakeTrackEditor::SupportsType(TSubclassOf<UMovieSceneTrack> Type) const
 {
-	return Type == UMovieSceneCameraShakeSourceShakeTrack::StaticClass();
+	return (
+		Type == UMovieSceneCameraShakeSourceShakeTrack::StaticClass() ||
+		Type == UMovieSceneCameraShakeSourceTriggerTrack::StaticClass()
+		);
 }
 
 TSharedRef<ISequencerSection> FCameraShakeSourceShakeTrackEditor::MakeSectionInterface(UMovieSceneSection& SectionObject, UMovieSceneTrack& Track, FGuid ObjectBinding)
 {
-	check(SupportsType(SectionObject.GetOuter()->GetClass()));
-	return MakeShareable(new FCameraShakeSourceShakeSection(SectionObject));
+	if (UMovieSceneCameraShakeSourceShakeSection* ShakeSection = Cast<UMovieSceneCameraShakeSourceShakeSection>(&SectionObject))
+	{
+		return MakeShareable(new FCameraShakeSourceShakeSection(GetSequencer(), *ShakeSection, ObjectBinding));
+	}
+	else if (UMovieSceneCameraShakeSourceTriggerSection* TriggerSection = Cast<UMovieSceneCameraShakeSourceTriggerSection>(&SectionObject))
+	{
+		return MakeShared<FCameraShakeSourceTriggerSection>(GetSequencer(), *TriggerSection);
+	}
+
+	check(false);
+	return MakeShareable(new FSequencerSection(SectionObject));
 }
 
-void FCameraShakeSourceShakeTrackEditor::AddKey(const FGuid& ObjectGuid)
+UMovieSceneTrack* FCameraShakeSourceShakeTrackEditor::AddTrack(UMovieScene* FocusedMovieScene, const FGuid& ObjectHandle, TSubclassOf<UMovieSceneTrack> TrackClass, FName UniqueTypeName)
 {
-	TArray<TWeakObjectPtr<>> OutObjects;
-	for (TWeakObjectPtr<> Object : GetSequencer()->FindObjectsInCurrentSequence(ObjectGuid))
+	UMovieSceneTrack* NewTrack = FMovieSceneTrackEditor::AddTrack(FocusedMovieScene, ObjectHandle, TrackClass, UniqueTypeName);
+
+	if (UMovieSceneCameraShakeSourceTriggerTrack* ShakeTrack = Cast<UMovieSceneCameraShakeSourceTriggerTrack>(NewTrack))
 	{
-		OutObjects.Add(Object);
+		// If it's a trigger track, auto-add an infinite section in which we can place our trigger keyframes.
+		UMovieSceneSection* NewSection = ShakeTrack->CreateNewSection();
+		NewSection->SetRange(TRange<FFrameNumber>::All());
+		ShakeTrack->AddSection(*NewSection);
 	}
-	
-	// Add a key that defaults to "auto camera shake".
-	const FScopedTransaction Transaction(LOCTEXT("AddCameraShakeSourceShake_Transaction", "Add Camera Shake"));
 
-	AnimatablePropertyChanged(FOnKeyProperty::CreateSP(this, &FCameraShakeSourceShakeTrackEditor::AddKeyInternal, OutObjects));
+	return NewTrack;
 }
 
 void FCameraShakeSourceShakeTrackEditor::BuildObjectBindingTrackMenu(FMenuBuilder& MenuBuilder, const TArray<FGuid>& ObjectBindings, const UClass* ObjectClass)
 {
 	if (const UCameraShakeSourceComponent* ShakeSourceComponent = AcquireCameraShakeSourceComponentFromGuid(ObjectBindings[0]))
 	{
-		MenuBuilder.AddMenuEntry(
-			LOCTEXT("AddShakeSourceShake", "Camera Shake"),
-			LOCTEXT("AddShakeSourceShakeTooltip", "Adds a camera shake originating from the parent camera shake source."),
+		MenuBuilder.AddSubMenu(
+				LOCTEXT("AddShakeSourceShake", "Camera Shake"),
+				LOCTEXT("AddShakeSourceShakeTooltip", "Adds a camera shake originating from the parent camera shake source."),
+				FNewMenuDelegate::CreateSP(this, &FCameraShakeSourceShakeTrackEditor::AddCameraShakeTracksMenu, ObjectBindings));
+	}
+}
+
+void FCameraShakeSourceShakeTrackEditor::AddCameraShakeTracksMenu(FMenuBuilder& MenuBuilder, TArray<FGuid> ObjectBindings)
+{
+	MenuBuilder.AddMenuEntry(
+			LOCTEXT("AddShakeSourceShakeControlled", "Controlled"),
+			LOCTEXT("AddShakeSourceShakeControlledTooltip", "Adds a track that lets you start and stop camera shakes originating from the parent camera shake source."),
 			FSlateIcon(),
 			FUIAction( 
 				FExecuteAction::CreateSP( this, &FCameraShakeSourceShakeTrackEditor::AddCameraShakeSection, ObjectBindings)
-			)
-		);
-	}
+				)
+			);
+
+	MenuBuilder.AddMenuEntry(
+			LOCTEXT("AddShakeSourceShakeTrigger", "Trigger"),
+			LOCTEXT("AddShakeSourceShakeTriggerTooltip", "Adds a track that lets you trigger camera shakes originating from the parent camera shake source."),
+			FSlateIcon(),
+			FUIAction(
+				FExecuteAction::CreateSP(this, &FCameraShakeSourceShakeTrackEditor::AddCameraShakeTriggerTrack, ObjectBindings)
+				)
+			);
 }
 
 void FCameraShakeSourceShakeTrackEditor::AddCameraShakeSection(TArray<FGuid> ObjectHandles)
@@ -115,7 +434,7 @@ void FCameraShakeSourceShakeTrackEditor::AddCameraShakeSection(TArray<FGuid> Obj
 
 	auto OnAddShakeSourceShakeSection = [=](FFrameNumber Time) -> FKeyPropertyResult
 	{
-		return this->AddKeyInternal(Time, Objects);
+		return this->AddCameraShakeSectionKeyInternal(Time, Objects, true);
 	};
 
 	const FScopedTransaction Transaction(LOCTEXT("AddCameraShakeSourceShake_Transaction", "Add Camera Shake"));
@@ -125,27 +444,46 @@ void FCameraShakeSourceShakeTrackEditor::AddCameraShakeSection(TArray<FGuid> Obj
 
 TSharedPtr<SWidget> FCameraShakeSourceShakeTrackEditor::BuildOutlinerEditWidget(const FGuid& ObjectBinding, UMovieSceneTrack* Track, const FBuildEditWidgetParams& Params)
 {
-	return SNew(SHorizontalBox)
-		+SHorizontalBox::Slot()
-		.AutoWidth()
-		.VAlign(VAlign_Center)
-		[
-			FSequencerUtilities::MakeAddButton(
-					LOCTEXT("AddShakeSourceShakeSection", "Camera Shake"),
-					FOnGetContent::CreateSP(this, &FCameraShakeSourceShakeTrackEditor::BuildCameraShakeSubMenu, ObjectBinding),
-					Params.NodeIsHovered, GetSequencer())
-		];
+	if (UMovieSceneCameraShakeSourceShakeTrack* ShakeTrack = Cast<UMovieSceneCameraShakeSourceShakeTrack>(Track))
+	{
+		return SNew(SHorizontalBox)
+			+SHorizontalBox::Slot()
+			.AutoWidth()
+			.VAlign(VAlign_Center)
+			[
+				FSequencerUtilities::MakeAddButton(
+						LOCTEXT("AddShakeSourceShakeSection", "Camera Shake"),
+						FOnGetContent::CreateSP(this, &FCameraShakeSourceShakeTrackEditor::BuildCameraShakeSubMenu, ObjectBinding),
+						Params.NodeIsHovered, GetSequencer())
+			];
+	}
+	else
+	{
+		return SNew(SHorizontalBox)
+			+ SHorizontalBox::Slot()
+			.AutoWidth()
+			.VAlign(VAlign_Center)
+			[
+				FSequencerUtilities::MakeAddButton(
+						LOCTEXT("AddSection", "Section"),
+						FOnGetContent::CreateSP(this, &FCameraShakeSourceShakeTrackEditor::BuildCameraShakeTracksMenu, ObjectBinding),
+						Params.NodeIsHovered, 
+						GetSequencer())
+			];
+	}
 }
 
-FKeyPropertyResult FCameraShakeSourceShakeTrackEditor::AddKeyInternal(FFrameNumber KeyTime, const TArray<TWeakObjectPtr<UObject>> Objects)
+FKeyPropertyResult FCameraShakeSourceShakeTrackEditor::AddCameraShakeSectionKeyInternal(FFrameNumber KeyTime, const TArray<TWeakObjectPtr<UObject>> Objects, bool bSelect)
 {
-	return AddKeyInternal(KeyTime, Objects, TSubclassOf<UCameraShake>());
+	return AddCameraShakeSectionKeyInternal(KeyTime, Objects, TSubclassOf<UCameraShake>(), bSelect);
 }
 
-FKeyPropertyResult FCameraShakeSourceShakeTrackEditor::AddKeyInternal(FFrameNumber KeyTime, const TArray<TWeakObjectPtr<UObject>> Objects, TSubclassOf<UCameraShake> CameraShake)
+FKeyPropertyResult FCameraShakeSourceShakeTrackEditor::AddCameraShakeSectionKeyInternal(FFrameNumber KeyTime, const TArray<TWeakObjectPtr<UObject>> Objects, TSubclassOf<UCameraShake> CameraShake, bool bSelect)
 {
 	FKeyPropertyResult KeyPropertyResult;
 
+	TArray<UMovieSceneSection*> SectionsToSelect;
+
 	for (int32 ObjectIndex = 0; ObjectIndex < Objects.Num(); ++ObjectIndex)
 	{
 		UObject* Object = Objects[ObjectIndex].Get();
@@ -159,6 +497,15 @@ FKeyPropertyResult FCameraShakeSourceShakeTrackEditor::AddKeyInternal(FFrameNumb
 			}
 		}
 
+		const bool bIsAutomaticShake = (CameraShake.Get() == nullptr);
+		if (bIsAutomaticShake)
+		{
+			if (UCameraShakeSourceComponent* ShakeSourceComponent = Cast<UCameraShakeSourceComponent>(Object))
+			{
+				CameraShake = ShakeSourceComponent->CameraShake;
+			}
+		}
+
 		FFindOrCreateHandleResult HandleResult = FindOrCreateHandleToObject(Object);
 		FGuid ObjectHandle = HandleResult.Handle;
 		KeyPropertyResult.bHandleCreated |= HandleResult.bWasCreated;
@@ -171,16 +518,24 @@ FKeyPropertyResult FCameraShakeSourceShakeTrackEditor::AddKeyInternal(FFrameNumb
 
 			if (ensure(Track))
 			{
-				UMovieSceneSection* NewSection = Cast<UMovieSceneCameraShakeSourceShakeTrack>(Track)->AddNewCameraShake(KeyTime, CameraShake);
+				UMovieSceneSection* NewSection = Cast<UMovieSceneCameraShakeSourceShakeTrack>(Track)->AddNewCameraShake(KeyTime, CameraShake, bIsAutomaticShake);
 				KeyPropertyResult.bTrackModified = true;
-				
-				GetSequencer()->EmptySelection();
-				GetSequencer()->SelectSection(NewSection);
-				GetSequencer()->ThrobSectionSelection();
+				SectionsToSelect.Add(NewSection);
 			}
 		}
 	}
 
+	if (bSelect)
+	{
+		const TSharedPtr<ISequencer> SequencerPtr = GetSequencer();
+		SequencerPtr->EmptySelection();
+		for (UMovieSceneSection* SectionToSelect : SectionsToSelect)
+		{
+			SequencerPtr->SelectSection(SectionToSelect);
+		}
+		SequencerPtr->ThrobSectionSelection();
+	}
+
 	return KeyPropertyResult;
 }
 
@@ -212,6 +567,18 @@ void FCameraShakeSourceShakeTrackEditor::AddCameraShakeSubMenu(FMenuBuilder& Men
 			FNewMenuDelegate::CreateSP(this, &FCameraShakeSourceShakeTrackEditor::AddOtherCameraShakeBrowserSubMenu, ObjectBindings));
 }
 
+TSharedRef<SWidget> FCameraShakeSourceShakeTrackEditor::BuildCameraShakeTracksMenu(FGuid ObjectBinding)
+{
+	FMenuBuilder MenuBuilder(true, nullptr);
+
+	TArray<FGuid> ObjectBindings;
+	ObjectBindings.Add(ObjectBinding);
+
+	AddCameraShakeTracksMenu(MenuBuilder, ObjectBindings);
+
+	return MenuBuilder.MakeWidget();
+}
+
 void FCameraShakeSourceShakeTrackEditor::AddOtherCameraShakeBrowserSubMenu(FMenuBuilder& MenuBuilder, TArray<FGuid> ObjectBindings)
 {
 	static const FString CameraShakeClassPath(TEXT("Class'/Script/Engine.CameraShake'"));
@@ -261,7 +628,7 @@ void FCameraShakeSourceShakeTrackEditor::OnCameraShakeAssetSelected(const FAsset
 		
 		const FScopedTransaction Transaction(LOCTEXT("AddCameraShakeSourceShake_Transaction", "Add Camera Shake"));
 
-		AnimatablePropertyChanged(FOnKeyProperty::CreateSP(this, &FCameraShakeSourceShakeTrackEditor::AddKeyInternal, OutObjects, CameraShakeClass));
+		AnimatablePropertyChanged(FOnKeyProperty::CreateSP(this, &FCameraShakeSourceShakeTrackEditor::AddCameraShakeSectionKeyInternal, OutObjects, CameraShakeClass, true));
 	}
 }
 
@@ -288,9 +655,7 @@ void FCameraShakeSourceShakeTrackEditor::OnAutoCameraShakeSelected(TArray<FGuid>
 
 	const FScopedTransaction Transaction(LOCTEXT("AddCameraShakeSourceShake_Transaction", "Add Camera Shake"));
 
-	AnimatablePropertyChanged(FOnKeyProperty::CreateSP(this, &FCameraShakeSourceShakeTrackEditor::AddKeyInternal, OutObjects));
-
-	//return FReply::Handled();
+	AnimatablePropertyChanged(FOnKeyProperty::CreateSP(this, &FCameraShakeSourceShakeTrackEditor::AddCameraShakeSectionKeyInternal, OutObjects, true));
 }
 
 bool FCameraShakeSourceShakeTrackEditor::OnShouldFilterCameraShake(const FAssetData& AssetData)
@@ -307,6 +672,87 @@ bool FCameraShakeSourceShakeTrackEditor::OnShouldFilterCameraShake(const FAssetD
 	return true;
 }
 
+FKeyPropertyResult FCameraShakeSourceShakeTrackEditor::AddCameraShakeTriggerTrackInternal(FFrameNumber Time, const TArray<TWeakObjectPtr<UObject>> Objects, TSubclassOf<UCameraShake> CameraShake)
+{
+	FKeyPropertyResult KeyPropertyResult;
+
+	for (int32 ObjectIndex = 0; ObjectIndex < Objects.Num(); ++ObjectIndex)
+	{
+		UObject* Object = Objects[ObjectIndex].Get();
+
+		if (AActor* Actor = Cast<AActor>(Object))
+		{
+			UCameraShakeSourceComponent* Component = Actor->FindComponentByClass<UCameraShakeSourceComponent>();
+			if (ensure(Component != nullptr))
+			{
+				Object = Component;
+			}
+		}
+
+		const bool bIsAutomaticShake = (CameraShake.Get() == nullptr);
+		if (bIsAutomaticShake)
+		{
+			if (UCameraShakeSourceComponent* ShakeSourceComponent = Cast<UCameraShakeSourceComponent>(Object))
+			{
+				CameraShake = ShakeSourceComponent->CameraShake;
+			}
+		}
+
+		FFindOrCreateHandleResult HandleResult = FindOrCreateHandleToObject(Object);
+		FGuid ObjectHandle = HandleResult.Handle;
+		KeyPropertyResult.bHandleCreated |= HandleResult.bWasCreated;
+
+		if (ObjectHandle.IsValid())
+		{
+			FFindOrCreateTrackResult TrackResult = FindOrCreateTrackForObject(ObjectHandle, UMovieSceneCameraShakeSourceTriggerTrack::StaticClass());
+			UMovieSceneTrack* Track = TrackResult.Track;
+			KeyPropertyResult.bTrackCreated |= TrackResult.bWasCreated;
+
+			if (ensure(Track))
+			{
+				TArray<UMovieSceneSection*> AllSections = Track->GetAllSections();
+				if (ensure(AllSections.Num() > 0))
+				{
+					UMovieSceneCameraShakeSourceTriggerSection* FirstSection = Cast<UMovieSceneCameraShakeSourceTriggerSection>(AllSections[0]);
+					// TODO: add trigger key at given time.
+					GetSequencer()->EmptySelection();
+					GetSequencer()->SelectSection(FirstSection);
+					GetSequencer()->ThrobSectionSelection();
+				}
+			}
+		}
+	}
+
+	return KeyPropertyResult;
+}
+
+void FCameraShakeSourceShakeTrackEditor::AddCameraShakeTriggerTrack(const TArray<FGuid> ObjectBindings)
+{
+	TSharedPtr<ISequencer> SequencerPtr = GetSequencer();
+	if (!SequencerPtr.IsValid() || !SequencerPtr->IsAllowedToChange())
+	{
+		return;
+	}
+
+	TArray<TWeakObjectPtr<UObject>> Objects;
+	for (FGuid ObjectBinding : ObjectBindings)
+	{
+		for (TWeakObjectPtr<UObject> Object : SequencerPtr->FindObjectsInCurrentSequence(ObjectBinding))
+		{
+			Objects.Add(Object);
+		}
+	}
+
+	auto OnAddShakeSourceShakeSection = [=](FFrameNumber Time) -> FKeyPropertyResult
+	{
+		return this->AddCameraShakeTriggerTrackInternal(Time, Objects, nullptr);
+	};
+
+	const FScopedTransaction Transaction(LOCTEXT("AddCameraShakeSourceShake_Transaction", "Add Camera Shake"));
+
+	AnimatablePropertyChanged(FOnKeyProperty::CreateLambda(OnAddShakeSourceShakeSection));
+}
+
 UCameraShakeSourceComponent* FCameraShakeSourceShakeTrackEditor::AcquireCameraShakeSourceComponentFromGuid(const FGuid& Guid)
 {
 	TArray<UCameraShakeSourceComponent*> ShakeSourceComponents;
diff --git a/Engine/Source/Editor/MovieSceneTools/Private/TrackEditors/CameraShakeSourceShakeTrackEditor.h b/Engine/Source/Editor/MovieSceneTools/Private/TrackEditors/CameraShakeSourceShakeTrackEditor.h
--- a/Engine/Source/Editor/MovieSceneTools/Private/TrackEditors/CameraShakeSourceShakeTrackEditor.h
+++ b/Engine/Source/Editor/MovieSceneTools/Private/TrackEditors/CameraShakeSourceShakeTrackEditor.h
@@ -26,15 +26,16 @@ public:
 	virtual ~FCameraShakeSourceShakeTrackEditor() {}
 
 	// ISequencerTrackEditor interface
-	virtual void AddKey(const FGuid& ObjectGuid) override;
+	virtual UMovieSceneTrack* AddTrack(UMovieScene* FocusedMovieScene, const FGuid& ObjectHandle, TSubclassOf<class UMovieSceneTrack> TrackClass, FName UniqueTypeName) override;
 	virtual void BuildObjectBindingTrackMenu(FMenuBuilder& MenuBuilder, const TArray<FGuid>& ObjectBindings, const UClass* ObjectClass) override;
 	virtual TSharedRef<ISequencerSection> MakeSectionInterface(UMovieSceneSection& SectionObject, UMovieSceneTrack& Track, FGuid ObjectBinding) override;
 	virtual bool SupportsType(TSubclassOf<UMovieSceneTrack> Type) const override;
 	virtual TSharedPtr<SWidget> BuildOutlinerEditWidget(const FGuid& ObjectBinding, UMovieSceneTrack* Track, const FBuildEditWidgetParams& Params) override;
 
 private:
-	FKeyPropertyResult AddKeyInternal(FFrameNumber KeyTime, const TArray<TWeakObjectPtr<UObject>> Objects);
-	FKeyPropertyResult AddKeyInternal(FFrameNumber KeyTime, const TArray<TWeakObjectPtr<UObject>> Objects, TSubclassOf<UCameraShake> CameraShake);
+	// Shake section members
+	FKeyPropertyResult AddCameraShakeSectionKeyInternal(FFrameNumber KeyTime, const TArray<TWeakObjectPtr<UObject>> Objects, bool bSelect);
+	FKeyPropertyResult AddCameraShakeSectionKeyInternal(FFrameNumber KeyTime, const TArray<TWeakObjectPtr<UObject>> Objects, TSubclassOf<UCameraShake> CameraShake, bool bSelect);
 
 	void AddCameraShakeSection(TArray<FGuid> ObjectHandles);
 
@@ -45,7 +46,15 @@ private:
 	void OnCameraShakeAssetEnterPressed(const TArray<FAssetData>& AssetData, TArray<FGuid> ObjectBindings);
 	void OnAutoCameraShakeSelected(TArray<FGuid> ObjectBindings);
 	bool OnShouldFilterCameraShake(const FAssetData& AssetData);
+
+	// Shake trigger members
+	FKeyPropertyResult AddCameraShakeTriggerTrackInternal(FFrameNumber Time, const TArray<TWeakObjectPtr<UObject>> Objects, TSubclassOf<UCameraShake> CameraShake);
+
+	void AddCameraShakeTriggerTrack(const TArray<FGuid> ObjectBindings);
 	
+	// Utility
 	UCameraShakeSourceComponent* AcquireCameraShakeSourceComponentFromGuid(const FGuid& Guid);
+	void AddCameraShakeTracksMenu(FMenuBuilder& MenuBuilder, TArray<FGuid> ObjectBindings);
+	TSharedRef<SWidget> BuildCameraShakeTracksMenu(FGuid ObjectBinding);
 };
 
diff --git a/Engine/Source/Editor/MovieSceneTools/Public/ClipboardTypes.h b/Engine/Source/Editor/MovieSceneTools/Public/ClipboardTypes.h
--- a/Engine/Source/Editor/MovieSceneTools/Public/ClipboardTypes.h
+++ b/Engine/Source/Editor/MovieSceneTools/Public/ClipboardTypes.h
@@ -10,6 +10,7 @@
 #include "Channels/MovieSceneFloatChannel.h"
 #include "Channels/MovieSceneEvent.h"
 #include "Channels/MovieSceneObjectPathChannel.h"
+#include "Channels/MovieSceneCameraShakeSourceTriggerChannel.h"
 #include "Sections/MovieSceneEventSection.h"
 #include "Sections/MovieSceneActorReferenceSection.h"
 #include "MovieSceneClipboard.h"
@@ -68,6 +69,10 @@ namespace MovieSceneClipboard
 	{
 		return "MovieSceneEvent";
 	}
+	template<> inline FName GetKeyTypeName<FMovieSceneCameraShakeSourceTrigger>()
+	{
+		return "FMovieSceneCameraShakeSourceTrigger";
+	}
 }
 
 #endif
diff --git a/Engine/Source/Runtime/MovieSceneTracks/Private/Channels/MovieSceneCameraShakeSourceTriggerChannel.cpp b/Engine/Source/Runtime/MovieSceneTracks/Private/Channels/MovieSceneCameraShakeSourceTriggerChannel.cpp
new file mode 100644
--- /dev/null
+++ b/Engine/Source/Runtime/MovieSceneTracks/Private/Channels/MovieSceneCameraShakeSourceTriggerChannel.cpp
@@ -0,0 +1,61 @@
+// Copyright Epic Games, Inc. All Rights Reserved.
+
+#include "Channels/MovieSceneCameraShakeSourceTriggerChannel.h"
+
+void FMovieSceneCameraShakeSourceTriggerChannel::GetKeys(const TRange<FFrameNumber>& WithinRange, TArray<FFrameNumber>* OutKeyTimes, TArray<FKeyHandle>* OutKeyHandles)
+{
+	GetData().GetKeys(WithinRange, OutKeyTimes, OutKeyHandles);
+}
+
+void FMovieSceneCameraShakeSourceTriggerChannel::GetKeyTimes(TArrayView<const FKeyHandle> InHandles, TArrayView<FFrameNumber> OutKeyTimes)
+{
+	GetData().GetKeyTimes(InHandles, OutKeyTimes);
+}
+
+void FMovieSceneCameraShakeSourceTriggerChannel::SetKeyTimes(TArrayView<const FKeyHandle> InHandles, TArrayView<const FFrameNumber> InKeyTimes)
+{
+	GetData().SetKeyTimes(InHandles, InKeyTimes);
+}
+
+void FMovieSceneCameraShakeSourceTriggerChannel::DuplicateKeys(TArrayView<const FKeyHandle> InHandles, TArrayView<FKeyHandle> OutNewHandles)
+{
+	GetData().DuplicateKeys(InHandles, OutNewHandles);
+}
+
+void FMovieSceneCameraShakeSourceTriggerChannel::DeleteKeys(TArrayView<const FKeyHandle> InHandles)
+{
+	GetData().DeleteKeys(InHandles);
+}
+
+void FMovieSceneCameraShakeSourceTriggerChannel::DeleteKeysFrom(FFrameNumber InTime, bool bDeleteKeysBefore)
+{
+	GetData().DeleteKeysFrom(InTime, bDeleteKeysBefore);
+}
+
+void FMovieSceneCameraShakeSourceTriggerChannel::ChangeFrameResolution(FFrameRate SourceRate, FFrameRate DestinationRate)
+{
+	GetData().ChangeFrameResolution(SourceRate, DestinationRate);
+}
+
+TRange<FFrameNumber> FMovieSceneCameraShakeSourceTriggerChannel::ComputeEffectiveRange() const
+{
+	return GetData().GetTotalRange();
+}
+
+int32 FMovieSceneCameraShakeSourceTriggerChannel::GetNumKeys() const
+{
+	return KeyTimes.Num();
+}
+
+void FMovieSceneCameraShakeSourceTriggerChannel::Reset()
+{
+	KeyTimes.Reset();
+	KeyValues.Reset();
+	KeyHandles.Reset();
+}
+
+void FMovieSceneCameraShakeSourceTriggerChannel::Offset(FFrameNumber DeltaPosition)
+{
+	GetData().Offset(DeltaPosition);
+}
+
diff --git a/Engine/Source/Runtime/MovieSceneTracks/Private/Evaluation/MovieSceneCameraShakePreviewer.cpp b/Engine/Source/Runtime/MovieSceneTracks/Private/Evaluation/MovieSceneCameraShakePreviewer.cpp
new file mode 100644
--- /dev/null
+++ b/Engine/Source/Runtime/MovieSceneTracks/Private/Evaluation/MovieSceneCameraShakePreviewer.cpp
@@ -0,0 +1,109 @@
+// Copyright Epic Games, Inc. All Rights Reserved.
+
+#include "Evaluation/MovieSceneCameraShakePreviewer.h"
+
+#if WITH_EDITOR
+
+#include "LevelEditorViewport.h"
+#include "Camera/CameraModifier_CameraShake.h"
+
+FCameraShakePreviewer::FCameraShakePreviewer()
+	: PreviewCameraShake(NewObject<UCameraModifier_CameraShake>())
+	, LastLocationModifier(FVector::ZeroVector)
+	, LastRotationModifier(FRotator::ZeroRotator)
+	, LastFOVModifier(0.f)
+{
+}
+
+FCameraShakePreviewer::~FCameraShakePreviewer()
+{
+	if (!ensureMsgf(RegisteredViewportClients.Num() == 0, TEXT("Forgot to call UnRegisterViewModifier!")))
+	{
+		UnRegisterViewModifier();
+	}
+
+	PreviewCameraShake = nullptr;
+}
+
+void FCameraShakePreviewer::Update(float DeltaTime, bool bIsPlaying)
+{
+	LastDeltaTime = DeltaTime;
+
+	if (!bIsPlaying)
+	{
+		LastLocationModifier = FVector::ZeroVector;
+		LastRotationModifier = FRotator::ZeroRotator;
+		LastFOVModifier = 0.f;
+	}
+}
+
+void FCameraShakePreviewer::OnModifyView(FMinimalViewInfo& InOutPOV)
+{
+	const float DeltaTime = LastDeltaTime.Get(-1.f);
+	if (DeltaTime > 0.f)
+	{
+		FMinimalViewInfo InPOV(InOutPOV);
+		PreviewCameraShake->ModifyCamera(DeltaTime, InOutPOV);
+
+		LastLocationModifier = InOutPOV.Location - InPOV.Location;
+		LastRotationModifier = InOutPOV.Rotation - InPOV.Rotation;
+		LastFOVModifier = InOutPOV.FOV - InPOV.FOV;
+
+		LastDeltaTime.Reset();
+	}
+	else
+	{
+		InOutPOV.Location += LastLocationModifier;
+		InOutPOV.Rotation += LastRotationModifier;
+		InOutPOV.FOV += LastFOVModifier;
+	}
+}
+
+void FCameraShakePreviewer::RegisterViewModifier()
+{
+	if (GEditor == nullptr)
+	{
+		return;
+	}
+
+	// Register our view modifier on all appropriate viewports, and remember which viewports we did that on.
+	// We will later make sure to unregister on the same list, except for any viewport that somehow disappeared since,
+	// which we will be notified about with the OnLevelViewportClientListChanged event.
+	RegisteredViewportClients.Reset();
+	for (FLevelEditorViewportClient* LevelVC : GEditor->GetLevelViewportClients())
+	{		
+		if (LevelVC && LevelVC->AllowsCinematicControl() && LevelVC->GetViewMode() != VMI_Unknown)
+		{
+			RegisteredViewportClients.Add(LevelVC);
+			LevelVC->ViewModifiers.AddRaw(this, &FCameraShakePreviewer::OnModifyView);
+		}
+	}
+
+	GEditor->OnLevelViewportClientListChanged().AddRaw(this, &FCameraShakePreviewer::OnLevelViewportClientListChanged);
+}
+
+void FCameraShakePreviewer::UnRegisterViewModifier()
+{
+	GEditor->OnLevelViewportClientListChanged().RemoveAll(this);
+
+	for (FLevelEditorViewportClient* ViewportClient : RegisteredViewportClients)
+	{
+		ViewportClient->ViewModifiers.RemoveAll(this);
+	}
+	RegisteredViewportClients.Reset();
+}
+
+void FCameraShakePreviewer::OnLevelViewportClientListChanged()
+{
+	if (GEditor != nullptr)
+	{
+		// If any viewports were removed while we were playing, simply get rid of them from our list of
+		// registered viewports.
+		TSet<FLevelEditorViewportClient*> PreviousViewportClients(RegisteredViewportClients);
+		TSet<FLevelEditorViewportClient*> NewViewportClients(GEditor->GetLevelViewportClients());
+		RegisteredViewportClients = PreviousViewportClients.Intersect(NewViewportClients).Array();
+	}
+}
+
+#endif
+
diff --git a/Engine/Source/Runtime/MovieSceneTracks/Private/Evaluation/MovieSceneCameraShakePreviewer.h b/Engine/Source/Runtime/MovieSceneTracks/Private/Evaluation/MovieSceneCameraShakePreviewer.h
new file mode 100644
--- /dev/null
+++ b/Engine/Source/Runtime/MovieSceneTracks/Private/Evaluation/MovieSceneCameraShakePreviewer.h
@@ -0,0 +1,49 @@
+// Copyright Epic Games, Inc. All Rights Reserved.
+
+#pragma once
+
+#if WITH_EDITOR
+
+#include "CoreMinimal.h"
+#include "UObject/GCObject.h"
+
+class FLevelEditorViewportClient;
+class UCameraModifier_CameraShake;
+struct FMinimalViewInfo;
+
+/**
+ * A class that owns a gameplay camera shake manager, so that we can us it to preview shakes in editor.
+ */
+class FCameraShakePreviewer : public FGCObject
+{
+public:
+	FCameraShakePreviewer();
+	~FCameraShakePreviewer();
+
+	void RegisterViewModifier();
+	void UnRegisterViewModifier();
+	void Update(float DeltaTime, bool bIsPlaying);
+
+	UCameraModifier_CameraShake* GetCameraShake() { return PreviewCameraShake; }
+
+private:
+	// FGCObject interface
+	virtual void AddReferencedObjects(FReferenceCollector& Collector) override { Collector.AddReferencedObject(PreviewCameraShake); }
+	virtual FString GetReferencerName() const override { return TEXT("SCameraShakePreviewer"); }
+
+private:
+	void OnModifyView(FMinimalViewInfo& InOutPOV);
+	void OnLevelViewportClientListChanged();
+
+private:
+	UCameraModifier_CameraShake* PreviewCameraShake;
+	TArray<FLevelEditorViewportClient*> RegisteredViewportClients;
+
+	TOptional<float> LastDeltaTime;
+	FVector LastLocationModifier;
+	FRotator LastRotationModifier;
+	float LastFOVModifier;
+};
+
+#endif
+
diff --git a/Engine/Source/Runtime/MovieSceneTracks/Private/Evaluation/MovieSceneCameraShakeSourceShakeTemplate.cpp b/Engine/Source/Runtime/MovieSceneTracks/Private/Evaluation/MovieSceneCameraShakeSourceShakeTemplate.cpp
--- a/Engine/Source/Runtime/MovieSceneTracks/Private/Evaluation/MovieSceneCameraShakeSourceShakeTemplate.cpp
+++ b/Engine/Source/Runtime/MovieSceneTracks/Private/Evaluation/MovieSceneCameraShakeSourceShakeTemplate.cpp
@@ -1,149 +1,23 @@
 // Copyright Epic Games, Inc. All Rights Reserved.
 
 #include "Evaluation/MovieSceneCameraShakeSourceShakeTemplate.h"
+#include "Evaluation/MovieSceneCameraShakePreviewer.h"
 #include "Evaluation/MovieSceneEvaluation.h"
-#include "Camera/CameraShakeSourceComponent.h"
-
-#if WITH_EDITOR
-
-#include "LevelEditorViewport.h"
 #include "Camera/CameraModifier_CameraShake.h"
+#include "Camera/CameraShakeSourceComponent.h"
 
-#endif
-
-#if WITH_EDITOR
-
-/**
- * A class that owns a gameplay camera shake manager, so that we can us it to preview shakes in editor.
- */
-class FCameraShakePreviewer : public FGCObject
+enum class ECameraShakeSourceShakeStatus : uint8
 {
-public:
-	FCameraShakePreviewer();
-	~FCameraShakePreviewer();
-
-	void RegisterViewModifier();
-	void UnRegisterViewModifier();
-	void Update(float DeltaTime, bool bIsPlaying);
-
-	UCameraModifier_CameraShake* GetCameraShake() { return PreviewCameraShake; }
-
-private:
-	// FGCObject interface
-	virtual void AddReferencedObjects(FReferenceCollector& Collector) override { Collector.AddReferencedObject(PreviewCameraShake); }
-	virtual FString GetReferencerName() const override { return TEXT("SCameraShakePreviewer"); }
-
-private:
-	void OnModifyView(FMinimalViewInfo& InOutPOV);
-	void OnLevelViewportClientListChanged();
-
-private:
-	UCameraModifier_CameraShake* PreviewCameraShake;
-	TArray<FLevelEditorViewportClient*> RegisteredViewportClients;
-
-	TOptional<float> LastDeltaTime;
-	FVector LastLocationModifier;
-	FRotator LastRotationModifier;
-	float LastFOVModifier;
+	NotStarted,
+	Started,
+	BlendingOut
 };
 
-FCameraShakePreviewer::FCameraShakePreviewer()
-	: PreviewCameraShake(NewObject<UCameraModifier_CameraShake>())
-	, LastLocationModifier(FVector::ZeroVector)
-	, LastRotationModifier(FRotator::ZeroRotator)
-	, LastFOVModifier(0.f)
-{
-}
-
-FCameraShakePreviewer::~FCameraShakePreviewer()
-{
-	PreviewCameraShake = nullptr;
-}
-
-void FCameraShakePreviewer::Update(float DeltaTime, bool bIsPlaying)
-{
-	LastDeltaTime = DeltaTime;
-
-	if (!bIsPlaying)
-	{
-		LastLocationModifier = FVector::ZeroVector;
-		LastRotationModifier = FRotator::ZeroRotator;
-		LastFOVModifier = 0.f;
-	}
-}
-
-void FCameraShakePreviewer::OnModifyView(FMinimalViewInfo& InOutPOV)
-{
-	const float DeltaTime = LastDeltaTime.Get(-1.f);
-	if (DeltaTime > 0.f)
-	{
-		FMinimalViewInfo InPOV(InOutPOV);
-		PreviewCameraShake->ModifyCamera(DeltaTime, InOutPOV);
-
-		LastLocationModifier = InOutPOV.Location - InPOV.Location;
-		LastRotationModifier = InOutPOV.Rotation - InPOV.Rotation;
-		LastFOVModifier = InOutPOV.FOV - InPOV.FOV;
-
-		LastDeltaTime.Reset();
-	}
-	else
-	{
-		InOutPOV.Location += LastLocationModifier;
-		InOutPOV.Rotation += LastRotationModifier;
-		InOutPOV.FOV += LastFOVModifier;
-	}
-}
-
-void FCameraShakePreviewer::RegisterViewModifier()
-{
-	if (GEditor == nullptr)
-	{
-		return;
-	}
-
-	// Register our view modifier on all appropriate viewports, and remember which viewports we did that on.
-	// We will later make sure to unregister on the same list, except for any viewport that somehow disappeared since,
-	// which we will be notified about with the OnLevelViewportClientListChanged event.
-	RegisteredViewportClients.Reset();
-	for (FLevelEditorViewportClient* LevelVC : GEditor->GetLevelViewportClients())
-	{		
-		if (LevelVC && LevelVC->AllowsCinematicControl() && LevelVC->GetViewMode() != VMI_Unknown)
-		{
-			RegisteredViewportClients.Add(LevelVC);
-			LevelVC->ViewModifiers.AddRaw(this, &FCameraShakePreviewer::OnModifyView);
-		}
-	}
-
-	GEditor->OnLevelViewportClientListChanged().AddRaw(this, &FCameraShakePreviewer::OnLevelViewportClientListChanged);
-}
-
-void FCameraShakePreviewer::UnRegisterViewModifier()
-{
-	GEditor->OnLevelViewportClientListChanged().RemoveAll(this);
-
-	for (FLevelEditorViewportClient* ViewportClient : RegisteredViewportClients)
-	{
-		ViewportClient->ViewModifiers.RemoveAll(this);
-	}
-}
-
-void FCameraShakePreviewer::OnLevelViewportClientListChanged()
-{
-	if (GEditor != nullptr)
-	{
-		// If any viewports were removed while we were playing, simply get rid of them from our list of
-		// registered viewports.
-		TSet<FLevelEditorViewportClient*> PreviousViewportClients(RegisteredViewportClients);
-		TSet<FLevelEditorViewportClient*> NewViewportClients(GEditor->GetLevelViewportClients());
-		RegisteredViewportClients = PreviousViewportClients.Intersect(NewViewportClients).Array();
-	}
-}
-
-#endif
-
 struct FCameraShakeSourceShakeSectionInstanceData : IPersistentEvaluationData
 {
-	bool bStarted = false;
+	ECameraShakeSourceShakeStatus Status = ECameraShakeSourceShakeStatus::NotStarted;
+	TOptional<FFrameTime> Duration;
+	TOptional<FFrameTime> BlendOutTime;
 
 #if WITH_EDITOR
 	FCameraShakePreviewer Previewer;
@@ -186,21 +60,18 @@ struct FPreAnimatedCameraShakeSourceShakeTokenProducer : IMovieScenePreAnimatedG
 	}
 };
 
-struct FCameraShakeSourceShakeExecutionToken : IMovieSceneExecutionToken
+struct FCameraShakeSourceShakeStartExecutionToken : IMovieSceneExecutionToken
 {
 	const FMovieSceneCameraShakeSectionData& SourceData;
 
-	FCameraShakeSourceShakeExecutionToken(const FMovieSceneCameraShakeSectionData& InSourceData)
+	FCameraShakeSourceShakeStartExecutionToken(const FMovieSceneCameraShakeSectionData& InSourceData)
 		: SourceData(InSourceData)
 	{
 	}
 
 	virtual void Execute(const FMovieSceneContext& Context, const FMovieSceneEvaluationOperand& Operand, FPersistentEvaluationData& PersistentData, IMovieScenePlayer& Player) override
 	{
-#if WITH_EDITOR
 		FCameraShakeSourceShakeSectionInstanceData& InstanceData = PersistentData.GetSectionData<FCameraShakeSourceShakeSectionInstanceData>();
-		UCameraModifier_CameraShake* const PreviewCameraShake = InstanceData.Previewer.GetCameraShake();
-#endif
 
 		for (TWeakObjectPtr<> BoundObject : Player.FindBoundObjects(Operand))
 		{
@@ -208,34 +79,101 @@ struct FCameraShakeSourceShakeExecutionToken : IMovieSceneExecutionToken
 			{
 				if (UCameraShakeSourceComponent* ShakeSourceComponent = CastChecked<UCameraShakeSourceComponent>(BoundObject.Get()))
 				{
-					static FMovieSceneAnimTypeID ShakeTypeID = TMovieSceneAnimTypeID<FCameraShakeSourceShakeExecutionToken, 0>();
+					static FMovieSceneAnimTypeID ShakeTypeID = TMovieSceneAnimTypeID<FCameraShakeSourceShakeStartExecutionToken, 0>();
 					Player.SavePreAnimatedState(ShakeTypeID, FPreAnimatedCameraShakeSourceShakeTokenProducer(Operand));
 
-					if (SourceData.ShakeClass)
+					TSubclassOf<UCameraShake> ShakeClass = SourceData.ShakeClass;
+					if (ShakeClass.Get() == nullptr)
 					{
-						ShakeSourceComponent->PlayCameraShake(SourceData.ShakeClass);
+						ShakeClass = ShakeSourceComponent->CameraShake;
+					}
+
+					if (ShakeClass.Get() != nullptr)
+					{
+						// Get the duration of the shake and store it in the instance data.
+						float Duration = -1.f;
+						UCameraShake::GetCameraShakeDuration(ShakeClass, Duration);
+						if (Duration > SMALL_NUMBER)
+						{
+							InstanceData.Duration = Context.GetFrameRate().AsFrameTime(Duration);
+						}
+
+						// Get the blend out duration and also store it in the instance data.
+						float BlendIn = 0.f, BlendOut = 0.f;
+						UCameraShake::GetCameraShakeBlendTimes(ShakeClass, BlendIn, BlendOut);
+						if (BlendOut > SMALL_NUMBER)
+						{
+							InstanceData.BlendOutTime = Context.GetFrameRate().AsFrameTime(BlendOut);
+						}
+
+						// Start playing the shake.
+						ShakeSourceComponent->PlayCameraShake(ShakeClass);
 
 #if WITH_EDITOR
+						// Also start playing the shake in our editor preview.
+						UCameraModifier_CameraShake* const PreviewCameraShake = InstanceData.Previewer.GetCameraShake();
+
 						FAddCameraShakeParams Params;
 						Params.SourceComponent = ShakeSourceComponent;
-						PreviewCameraShake->AddCameraShake(SourceData.ShakeClass, Params);
+						PreviewCameraShake->AddCameraShake(ShakeClass, Params);
 #endif
 					}
-					else
+				}
+			}
+
+			// TODO-ludovic: only support a single binding for now.
+			break;
+		}
+	}
+};
+
+struct FCameraShakeSourceShakeBlendOutExecutionToken : IMovieSceneExecutionToken
+{
+	const FMovieSceneCameraShakeSectionData& SourceData;
+
+	FCameraShakeSourceShakeBlendOutExecutionToken(const FMovieSceneCameraShakeSectionData& InSourceData)
+		: SourceData(InSourceData)
+	{
+	}
+
+	virtual void Execute(const FMovieSceneContext& Context, const FMovieSceneEvaluationOperand& Operand, FPersistentEvaluationData& PersistentData, IMovieScenePlayer& Player) override
+	{
+		FCameraShakeSourceShakeSectionInstanceData& InstanceData = PersistentData.GetSectionData<FCameraShakeSourceShakeSectionInstanceData>();
+
+		for (TWeakObjectPtr<> BoundObject : Player.FindBoundObjects(Operand))
+		{
+			if (BoundObject.IsValid())
+			{
+				if (UCameraShakeSourceComponent* ShakeSourceComponent = CastChecked<UCameraShakeSourceComponent>(BoundObject.Get()))
+				{
+					TSubclassOf<UCameraShake> ShakeClass = SourceData.ShakeClass;
+					if (ShakeClass.Get() == nullptr)
+					{
+						ShakeClass = ShakeSourceComponent->CameraShake;
+					}
+
+					if (ShakeClass.Get() != nullptr)
 					{
-						ShakeSourceComponent->Play();
+						// Don't stop shakes immediately... let them gracefully blend out.
+						const bool bImmediately = false;
+
+						// TODO-ludovic: this isn't exactly correct...
+						// We could be stopping other shakes of the same type started by other means.. but doing
+						// the correct thing would require storing multiple weak shake instance pointers mapped 
+						// to multiple player controllers, themselves mapped to multiple bound objects. Let's only
+						// do that if we run into the (quite unlikely) case where we need it?
+						ShakeSourceComponent->StopAllCameraShakesOfType(ShakeClass, bImmediately);
 
 #if WITH_EDITOR
-						if (ShakeSourceComponent->CameraShake.Get() != nullptr)
-						{
-							FAddCameraShakeParams Params;
-							Params.SourceComponent = ShakeSourceComponent;
-							PreviewCameraShake->AddCameraShake(ShakeSourceComponent->CameraShake, Params);
-						}
+						UCameraModifier_CameraShake* const PreviewCameraShake = InstanceData.Previewer.GetCameraShake();
+						PreviewCameraShake->RemoveAllCameraShakesOfClassFromSource(ShakeClass, ShakeSourceComponent, bImmediately);
 #endif
 					}
 				}
 			}
+
+			// TODO-ludovic: only support a single binding for now.
+			break;
 		}
 	}
 };
@@ -247,6 +185,7 @@ FMovieSceneCameraShakeSourceShakeSectionTemplate::FMovieSceneCameraShakeSourceSh
 FMovieSceneCameraShakeSourceShakeSectionTemplate::FMovieSceneCameraShakeSourceShakeSectionTemplate(const UMovieSceneCameraShakeSourceShakeSection& Section)
 	: SourceData(Section.ShakeData)
 	, SectionStartTime(Section.HasStartFrame() ? Section.GetInclusiveStartFrame() : 0)
+	, SectionEndTime(Section.HasEndFrame() ? Section.GetExclusiveEndFrame() : 0)
 {
 }
 
@@ -258,7 +197,7 @@ void FMovieSceneCameraShakeSourceShakeSectionTemplate::SetupOverrides()
 void FMovieSceneCameraShakeSourceShakeSectionTemplate::Setup(FPersistentEvaluationData& PersistentData, IMovieScenePlayer& Player) const
 {
 	FCameraShakeSourceShakeSectionInstanceData& InstanceData = PersistentData.AddSectionData<FCameraShakeSourceShakeSectionInstanceData>();
-	InstanceData.bStarted = false;
+	InstanceData.Status = ECameraShakeSourceShakeStatus::NotStarted;
 
 #if WITH_EDITOR
 	InstanceData.Previewer.RegisterViewModifier();
@@ -268,10 +207,42 @@ void FMovieSceneCameraShakeSourceShakeSectionTemplate::Setup(FPersistentEvaluati
 void FMovieSceneCameraShakeSourceShakeSectionTemplate::Evaluate(const FMovieSceneEvaluationOperand& Operand, const FMovieSceneContext& Context, const FPersistentEvaluationData& PersistentData, FMovieSceneExecutionTokens& ExecutionTokens) const
 {
 	FCameraShakeSourceShakeSectionInstanceData& InstanceData = PersistentData.GetSectionData<FCameraShakeSourceShakeSectionInstanceData>();
-	if (!InstanceData.bStarted)
+	if (InstanceData.Status == ECameraShakeSourceShakeStatus::NotStarted)
 	{
-		ExecutionTokens.Add(FCameraShakeSourceShakeExecutionToken(SourceData));
-		InstanceData.bStarted = true;
+		// We haven't started playing the shake yet... add an execution token which will trigger it.
+		ExecutionTokens.Add(FCameraShakeSourceShakeStartExecutionToken(SourceData));
+		InstanceData.Status = ECameraShakeSourceShakeStatus::Started;
+	}
+	else if (InstanceData.Status == ECameraShakeSourceShakeStatus::Started && InstanceData.BlendOutTime.IsSet())
+	{
+		// We have started playing, but we have a blend-out time that we need to watch out for.
+		// There can be 3 situations here:
+		//
+		//   1. The shake has a duration, and our source section is long enough that the shake will
+		//		blend out and finish naturally on its own. In this case we have nothing to do.
+		//	 2. The shake has a duration, but our source section's size is cutting this short. We 
+		//	    will want to start making the shake blend out manually.
+		//	 3. The shake has no duration, so we need to make it blend out manually near the end of
+		//	    our source section.
+		//
+		// Cases 2 and 3 require us to stop the shake ourselves. Let's see which case we are in.
+		//
+		const bool bHasDuration = InstanceData.Duration.IsSet();
+		const bool bNeedsManualStop = 
+			// Case 3
+			!bHasDuration ||
+			// Case 2
+			(SectionStartTime + InstanceData.Duration.GetValue()) > SectionEndTime;
+		if (bNeedsManualStop)
+		{
+			// Let's see if we have reached the time when we need to start blending out.
+			const FFrameTime BlendOutStartTime = SectionEndTime - InstanceData.BlendOutTime.GetValue();
+			if (Context.GetPreviousTime() < BlendOutStartTime && Context.GetTime() >= BlendOutStartTime)
+			{
+				ExecutionTokens.Add(FCameraShakeSourceShakeBlendOutExecutionToken(SourceData));
+				InstanceData.Status = ECameraShakeSourceShakeStatus::BlendingOut;
+			}
+		}
 	}
 
 #if WITH_EDITOR
diff --git a/Engine/Source/Runtime/MovieSceneTracks/Private/Evaluation/MovieSceneCameraShakeSourceShakeTemplate.h b/Engine/Source/Runtime/MovieSceneTracks/Private/Evaluation/MovieSceneCameraShakeSourceShakeTemplate.h
--- a/Engine/Source/Runtime/MovieSceneTracks/Private/Evaluation/MovieSceneCameraShakeSourceShakeTemplate.h
+++ b/Engine/Source/Runtime/MovieSceneTracks/Private/Evaluation/MovieSceneCameraShakeSourceShakeTemplate.h
@@ -31,4 +31,8 @@ private:
 	/** Cached section start time */
 	UPROPERTY()
 	FFrameNumber SectionStartTime;
+
+	/** Cached section end time */
+	UPROPERTY()
+	FFrameNumber SectionEndTime;
 };
diff --git a/Engine/Source/Runtime/MovieSceneTracks/Private/Evaluation/MovieSceneCameraShakeSourceTriggerTemplate.cpp b/Engine/Source/Runtime/MovieSceneTracks/Private/Evaluation/MovieSceneCameraShakeSourceTriggerTemplate.cpp
new file mode 100644
--- /dev/null
+++ b/Engine/Source/Runtime/MovieSceneTracks/Private/Evaluation/MovieSceneCameraShakeSourceTriggerTemplate.cpp
@@ -0,0 +1,152 @@
+// Copyright Epic Games, Inc. All Rights Reserved.
+
+#include "Evaluation/MovieSceneCameraShakeSourceTriggerTemplate.h"
+#include "Camera/CameraModifier_CameraShake.h"
+#include "Camera/CameraShakeSourceComponent.h"
+#include "Channels/MovieSceneCameraShakeSourceTriggerChannel.h"
+#include "Evaluation/MovieSceneCameraShakePreviewer.h"
+#include "Sections/MovieSceneCameraShakeSourceTriggerSection.h"
+
+struct FCameraShakeSourceTriggerInstanceData : IPersistentEvaluationData
+{
+#if WITH_EDITOR
+	FCameraShakePreviewer Previewer;
+#endif
+};
+
+struct FCameraShakeSourceTriggerExecutionToken : IMovieSceneExecutionToken
+{
+	TArray<FMovieSceneCameraShakeSourceTrigger> Triggers;
+
+	FCameraShakeSourceTriggerExecutionToken(TArray<FMovieSceneCameraShakeSourceTrigger>&& InTriggers)
+		: Triggers(MoveTemp(InTriggers))
+	{
+	}
+
+	virtual void Execute(const FMovieSceneContext& Context, const FMovieSceneEvaluationOperand& Operand, FPersistentEvaluationData& PersistentData, IMovieScenePlayer& Player) override
+	{
+		FCameraShakeSourceTriggerInstanceData& InstanceData = PersistentData.GetSectionData<FCameraShakeSourceTriggerInstanceData>();
+
+		for (TWeakObjectPtr<> BoundObject : Player.FindBoundObjects(Operand))
+		{
+			if (BoundObject.IsValid())
+			{
+				if (UCameraShakeSourceComponent* ShakeSourceComponent = CastChecked<UCameraShakeSourceComponent>(BoundObject.Get()))
+				{
+					for (const FMovieSceneCameraShakeSourceTrigger& Trigger : Triggers)
+					{
+						TSubclassOf<UCameraShake> ShakeClass = Trigger.ShakeClass;
+						if (ShakeClass.Get() == nullptr)
+						{
+							ShakeClass = ShakeSourceComponent->CameraShake;
+						}
+
+						if (ShakeClass.Get() != nullptr)
+						{
+							// Start playing the shake.
+							ShakeSourceComponent->PlayCameraShake(ShakeClass);
+
+#if WITH_EDITOR
+							// Also start playing the shake in our editor preview.
+							UCameraModifier_CameraShake* const PreviewCameraShake = InstanceData.Previewer.GetCameraShake();
+
+							FAddCameraShakeParams Params;
+							Params.SourceComponent = ShakeSourceComponent;
+							PreviewCameraShake->AddCameraShake(ShakeClass, Params);
+#endif
+						}
+					}
+				}
+			}
+
+			// TODO-ludovic: only support a single binding for now.
+			break;
+		}
+	}
+};
+
+
+FMovieSceneCameraShakeSourceTriggerSectionTemplate::FMovieSceneCameraShakeSourceTriggerSectionTemplate(const UMovieSceneCameraShakeSourceTriggerSection& Section)
+{
+	TMovieSceneChannelData<const FMovieSceneCameraShakeSourceTrigger> TriggerData = Section.GetChannel().GetData();
+	TArrayView<const FFrameNumber> Times = TriggerData.GetTimes();
+	TArrayView<const FMovieSceneCameraShakeSourceTrigger> Values = TriggerData.GetValues();
+
+	TriggerTimes.Reserve(Times.Num());
+	TriggerValues.Reserve(Values.Num());
+
+	for (int32 Idx = 0; Idx < Times.Num(); ++Idx)
+	{
+		TriggerTimes.Add(Times[Idx]);
+		TriggerValues.Add(Values[Idx]);
+	}
+}
+
+void FMovieSceneCameraShakeSourceTriggerSectionTemplate::SetupOverrides()
+{
+	EnableOverrides(RequiresSetupFlag | RequiresTearDownFlag);
+}
+
+void FMovieSceneCameraShakeSourceTriggerSectionTemplate::Setup(FPersistentEvaluationData& PersistentData, IMovieScenePlayer& Player) const
+{
+	FCameraShakeSourceTriggerInstanceData& InstanceData = PersistentData.AddSectionData<FCameraShakeSourceTriggerInstanceData>();
+#if WITH_EDITOR
+	InstanceData.Previewer.RegisterViewModifier();
+#endif
+}
+
+void FMovieSceneCameraShakeSourceTriggerSectionTemplate::EvaluateSwept(const FMovieSceneEvaluationOperand& Operand, const FMovieSceneContext& Context, const TRange<FFrameNumber>& SweptRange, const FPersistentEvaluationData& PersistentData, FMovieSceneExecutionTokens& ExecutionTokens) const
+{
+#if WITH_EDITOR
+	{
+		const float DeltaTime = Context.GetFrameRate().AsSeconds(Context.GetDelta());
+		const bool bIsPlaying = Context.GetStatus() == EMovieScenePlayerStatus::Playing;
+		FCameraShakeSourceTriggerInstanceData& InstanceData = PersistentData.GetSectionData<FCameraShakeSourceTriggerInstanceData>();
+		InstanceData.Previewer.Update(DeltaTime, bIsPlaying);
+	}
+#endif
+
+	if (Context.GetStatus() == EMovieScenePlayerStatus::Stopped || Context.IsSilent())
+	{
+		return;
+	}
+
+	const bool bBackwards = Context.GetDirection() == EPlayDirection::Backwards;
+	if (bBackwards)
+	{
+		return;
+	}
+
+	TArray<FMovieSceneCameraShakeSourceTrigger> ShakesToTrigger;
+
+	const int32 TriggerCount = TriggerTimes.Num();
+	const float PositionInSeconds = Context.GetTime() * Context.GetRootToSequenceTransform().Inverse() / Context.GetFrameRate();
+
+	for (int32 KeyIndex = 0; KeyIndex < TriggerCount; ++KeyIndex)
+	{
+		const FFrameNumber Time = TriggerTimes[KeyIndex];
+		if (SweptRange.Contains(Time))
+		{
+			ShakesToTrigger.Add(TriggerValues[KeyIndex]);
+		}
+	}
+
+	if (ShakesToTrigger.Num() > 0)
+	{
+		ExecutionTokens.Add(FCameraShakeSourceTriggerExecutionToken(MoveTemp(ShakesToTrigger)));
+	}
+}
+
+FMovieSceneAnimTypeID FMovieSceneCameraShakeSourceTriggerSectionTemplate::GetAnimTypeID()
+{
+	return TMovieSceneAnimTypeID<FMovieSceneCameraShakeSourceTriggerSectionTemplate>();
+}
+
+void FMovieSceneCameraShakeSourceTriggerSectionTemplate::TearDown(FPersistentEvaluationData& PersistentData, IMovieScenePlayer& Player) const
+{
+	FCameraShakeSourceTriggerInstanceData& InstanceData = PersistentData.GetSectionData<FCameraShakeSourceTriggerInstanceData>();
+#if WITH_EDITOR
+	InstanceData.Previewer.UnRegisterViewModifier();
+#endif
+}
+
diff --git a/Engine/Source/Runtime/MovieSceneTracks/Private/Evaluation/MovieSceneCameraShakeSourceTriggerTemplate.h b/Engine/Source/Runtime/MovieSceneTracks/Private/Evaluation/MovieSceneCameraShakeSourceTriggerTemplate.h
new file mode 100644
--- /dev/null
+++ b/Engine/Source/Runtime/MovieSceneTracks/Private/Evaluation/MovieSceneCameraShakeSourceTriggerTemplate.h
@@ -0,0 +1,38 @@
+// Copyright Epic Games, Inc. All Rights Reserved.
+
+#pragma once
+
+#include "CoreMinimal.h"
+#include "Evaluation/MovieSceneAnimTypeID.h"
+#include "Evaluation/MovieSceneEvalTemplate.h"
+#include "Sections/MovieSceneCameraShakeSourceTriggerSection.h"
+#include "UObject/ObjectMacros.h"
+#include "MovieSceneCameraShakeSourceTriggerTemplate.generated.h"
+
+class UMovieSceneCameraShakeSourceTriggerSection;
+
+USTRUCT()
+struct MOVIESCENETRACKS_API FMovieSceneCameraShakeSourceTriggerSectionTemplate : public FMovieSceneEvalTemplate
+{
+	GENERATED_BODY()
+	
+	FMovieSceneCameraShakeSourceTriggerSectionTemplate() {}
+	FMovieSceneCameraShakeSourceTriggerSectionTemplate(const UMovieSceneCameraShakeSourceTriggerSection& Section);
+
+	static FMovieSceneAnimTypeID GetAnimTypeID();
+	
+private:
+	virtual UScriptStruct& GetScriptStructImpl() const override { return *StaticStruct(); }
+	virtual void SetupOverrides() override;
+	virtual void Setup(FPersistentEvaluationData& PersistentData, IMovieScenePlayer& Player) const override;
+	virtual void EvaluateSwept(const FMovieSceneEvaluationOperand& Operand, const FMovieSceneContext& Context, const TRange<FFrameNumber>& SweptRange, const FPersistentEvaluationData& PersistentData, FMovieSceneExecutionTokens& ExecutionTokens) const override;
+	virtual void TearDown(FPersistentEvaluationData& PersistentData, IMovieScenePlayer& Player) const override;
+
+private:
+	UPROPERTY()
+	TArray<FFrameNumber> TriggerTimes;
+
+	UPROPERTY()
+	TArray<FMovieSceneCameraShakeSourceTrigger> TriggerValues;
+};
+
diff --git a/Engine/Source/Runtime/MovieSceneTracks/Private/Sections/MovieSceneCameraShakeSourceTriggerSection.cpp b/Engine/Source/Runtime/MovieSceneTracks/Private/Sections/MovieSceneCameraShakeSourceTriggerSection.cpp
new file mode 100644
--- /dev/null
+++ b/Engine/Source/Runtime/MovieSceneTracks/Private/Sections/MovieSceneCameraShakeSourceTriggerSection.cpp
@@ -0,0 +1,21 @@
+// Copyright Epic Games, Inc. All Rights Reserved.
+
+#include "Sections/MovieSceneCameraShakeSourceTriggerSection.h"
+#include "Channels/MovieSceneChannelProxy.h"
+#include "Evaluation/MovieSceneCameraShakeSourceTriggerTemplate.h"
+
+UMovieSceneCameraShakeSourceTriggerSection::UMovieSceneCameraShakeSourceTriggerSection(const FObjectInitializer& Init)
+	: Super(Init)
+{
+#if WITH_EDITOR
+	ChannelProxy = MakeShared<FMovieSceneChannelProxy>(Channel, FMovieSceneChannelMetaData());
+#else
+	ChannelProxy = MakeShared<FMovieSceneChannelProxy>(Channel);
+#endif
+}
+
+FMovieSceneEvalTemplatePtr UMovieSceneCameraShakeSourceTriggerSection::GenerateTemplate() const
+{
+	return FMovieSceneCameraShakeSourceTriggerSectionTemplate(*this);
+}
+
diff --git a/Engine/Source/Runtime/MovieSceneTracks/Private/Tracks/MovieSceneCameraShakeSourceShakeTrack.cpp b/Engine/Source/Runtime/MovieSceneTracks/Private/Tracks/MovieSceneCameraShakeSourceShakeTrack.cpp
--- a/Engine/Source/Runtime/MovieSceneTracks/Private/Tracks/MovieSceneCameraShakeSourceShakeTrack.cpp
+++ b/Engine/Source/Runtime/MovieSceneTracks/Private/Tracks/MovieSceneCameraShakeSourceShakeTrack.cpp
@@ -1,22 +1,46 @@
 // Copyright Epic Games, Inc. All Rights Reserved.
 
 #include "Tracks/MovieSceneCameraShakeSourceShakeTrack.h"
-#include "Sections/MovieSceneCameraShakeSourceShakeSection.h"
+#include "Camera/CameraShakeSourceComponent.h"
 #include "Compilation/MovieSceneCompilerRules.h"
+#include "Sections/MovieSceneCameraShakeSourceShakeSection.h"
 
 #define LOCTEXT_NAMESPACE "MovieSceneCameraShakeSourceShakeTrack"
 
-UMovieSceneSection* UMovieSceneCameraShakeSourceShakeTrack::AddNewCameraShake(const FFrameNumber KeyTime, const TSubclassOf<UCameraShake> ShakeClass)
+UMovieSceneSection* UMovieSceneCameraShakeSourceShakeTrack::AddNewCameraShake(const FFrameNumber KeyTime, const UCameraShakeSourceComponent& ShakeSourceComponent)
+{
+	if (ensure(ShakeSourceComponent.CameraShake))
+	{
+		return AddNewCameraShake(KeyTime, ShakeSourceComponent.CameraShake, true);
+	}
+	return nullptr;
+}
+
+UMovieSceneSection* UMovieSceneCameraShakeSourceShakeTrack::AddNewCameraShake(const FFrameNumber KeyTime, const TSubclassOf<UCameraShake> ShakeClass, bool bIsAutomaticShake)
 {
+	ensure(ShakeClass);
+
 	Modify();
 
 	UMovieSceneCameraShakeSourceShakeSection* const NewSection = Cast<UMovieSceneCameraShakeSourceShakeSection>(CreateNewSection());
 	if (NewSection)
 	{
-		// #fixme get length
-		FFrameTime Duration = 5.0 * GetTypedOuter<UMovieScene>()->GetTickResolution();
-		NewSection->InitialPlacement(CameraShakeSections, KeyTime, Duration.FrameNumber.Value, SupportsMultipleRows());
-		NewSection->ShakeData.ShakeClass = ShakeClass;
+		float ShakeDuration = -1.f;
+		UCameraShake::GetCameraShakeDuration(ShakeClass, ShakeDuration);
+		if (ShakeDuration <= SMALL_NUMBER)
+		{
+			// Default endless shakes to 5 seconds totally arbitrarily.
+			ShakeDuration = 5.f;
+		}
+
+		const FFrameTime ShakeDurationTime = ShakeDuration * GetTypedOuter<UMovieScene>()->GetTickResolution();
+		NewSection->InitialPlacement(CameraShakeSections, KeyTime, ShakeDurationTime.FrameNumber.Value, SupportsMultipleRows());
+		if (!bIsAutomaticShake)
+		{
+			// Only set the shake class if we want to specifically run the provided shake. Otherwise, we leave
+			// it null and will automatically fallback to whatever is set on the shake source component.
+			NewSection->ShakeData.ShakeClass = ShakeClass;
+		}
 		
 		AddSection(*NewSection);
 	}
diff --git a/Engine/Source/Runtime/MovieSceneTracks/Private/Tracks/MovieSceneCameraShakeSourceTriggerTrack.cpp b/Engine/Source/Runtime/MovieSceneTracks/Private/Tracks/MovieSceneCameraShakeSourceTriggerTrack.cpp
new file mode 100644
--- /dev/null
+++ b/Engine/Source/Runtime/MovieSceneTracks/Private/Tracks/MovieSceneCameraShakeSourceTriggerTrack.cpp
@@ -0,0 +1,81 @@
+// Copyright Epic Games, Inc. All Rights Reserved.
+
+#include "Tracks/MovieSceneCameraShakeSourceTriggerTrack.h"
+#include "Sections/MovieSceneCameraShakeSourceTriggerSection.h"
+
+#define LOCTEXT_NAMESPACE "MovieSceneCameraShakeSourceTrigger"
+
+UMovieSceneCameraShakeSourceTriggerTrack::UMovieSceneCameraShakeSourceTriggerTrack(const FObjectInitializer& Obj)
+	: Super(Obj)
+{
+#if WITH_EDITORONLY_DATA
+	TrackTint = FColor(43, 43, 155, 65);
+#endif
+}
+
+bool UMovieSceneCameraShakeSourceTriggerTrack::SupportsType(TSubclassOf<UMovieSceneSection> SectionClass) const
+{
+	return SectionClass == UMovieSceneCameraShakeSourceTriggerSection::StaticClass();
+}
+
+UMovieSceneSection* UMovieSceneCameraShakeSourceTriggerTrack::CreateNewSection()
+{
+	return NewObject<UMovieSceneCameraShakeSourceTriggerSection>(this, NAME_None, RF_Transactional);
+}
+
+
+bool UMovieSceneCameraShakeSourceTriggerTrack::HasSection(const UMovieSceneSection& Section) const
+{
+	return Sections.ContainsByPredicate([&](const UMovieSceneSection* In){ return In == &Section; });
+}
+
+
+void UMovieSceneCameraShakeSourceTriggerTrack::AddSection(UMovieSceneSection& Section)
+{
+	Sections.Add(&Section);
+}
+
+
+void UMovieSceneCameraShakeSourceTriggerTrack::RemoveSection(UMovieSceneSection& Section)
+{
+	Sections.RemoveAll([&](const UMovieSceneSection* In) { return In == &Section; });
+}
+
+void UMovieSceneCameraShakeSourceTriggerTrack::RemoveSectionAt(int32 SectionIndex)
+{
+	Sections.RemoveAt(SectionIndex);
+}
+
+void UMovieSceneCameraShakeSourceTriggerTrack::RemoveAllAnimationData()
+{
+	Sections.Empty();
+}
+
+
+bool UMovieSceneCameraShakeSourceTriggerTrack::IsEmpty() const
+{
+	return Sections.Num() == 0;
+}
+
+
+const TArray<UMovieSceneSection*>& UMovieSceneCameraShakeSourceTriggerTrack::GetAllSections() const
+{
+	return Sections;
+}
+
+void UMovieSceneCameraShakeSourceTriggerTrack::PostCompile(FMovieSceneEvaluationTrack& Track, const FMovieSceneTrackCompilerArgs& Args) const
+{
+	Track.SetEvaluationMethod(EEvaluationMethod::Swept);
+}
+
+#if WITH_EDITORONLY_DATA
+
+FText UMovieSceneCameraShakeSourceTriggerTrack::GetDisplayName() const
+{
+	return LOCTEXT("TrackName", "Camera Shake Trigger");
+}
+
+#endif
+
+#undef LOCTEXT_NAMESPACE
+
diff --git a/Engine/Source/Runtime/MovieSceneTracks/Public/Channels/MovieSceneCameraShakeSourceTriggerChannel.h b/Engine/Source/Runtime/MovieSceneTracks/Public/Channels/MovieSceneCameraShakeSourceTriggerChannel.h
new file mode 100644
--- /dev/null
+++ b/Engine/Source/Runtime/MovieSceneTracks/Public/Channels/MovieSceneCameraShakeSourceTriggerChannel.h
@@ -0,0 +1,76 @@
+// Copyright Epic Games, Inc. All Rights Reserved.
+
+#pragma once
+
+#include "Camera/CameraShake.h"
+#include "CoreTypes.h"
+#include "Curves/KeyHandle.h"
+#include "Channels/MovieSceneChannel.h"
+#include "Channels/MovieSceneChannelData.h"
+#include "Channels/MovieSceneChannelTraits.h"
+#include "MovieSceneCameraShakeSourceTriggerChannel.generated.h"
+
+USTRUCT(BlueprintType)
+struct FMovieSceneCameraShakeSourceTrigger
+{
+	GENERATED_BODY()
+
+	FMovieSceneCameraShakeSourceTrigger() {}
+	FMovieSceneCameraShakeSourceTrigger(TSubclassOf<UCameraShake> InShakeClass) : ShakeClass(InShakeClass) {}
+
+	UPROPERTY(EditAnywhere, BlueprintReadWrite, Category="Camera Shake")
+	TSubclassOf<UCameraShake> ShakeClass;
+};
+
+USTRUCT()
+struct MOVIESCENETRACKS_API FMovieSceneCameraShakeSourceTriggerChannel : public FMovieSceneChannel
+{
+	GENERATED_BODY()
+
+	FORCEINLINE TMovieSceneChannelData<FMovieSceneCameraShakeSourceTrigger> GetData()
+	{
+		return TMovieSceneChannelData<FMovieSceneCameraShakeSourceTrigger>(&KeyTimes, &KeyValues, &KeyHandles);
+	}
+
+	FORCEINLINE TMovieSceneChannelData<const FMovieSceneCameraShakeSourceTrigger> GetData() const
+	{
+		return TMovieSceneChannelData<const FMovieSceneCameraShakeSourceTrigger>(&KeyTimes, &KeyValues);
+	}
+
+public:
+	// ~ FMovieSceneChannel Interface
+	virtual void GetKeys(const TRange<FFrameNumber>& WithinRange, TArray<FFrameNumber>* OutKeyTimes, TArray<FKeyHandle>* OutKeyHandles) override;
+	virtual void GetKeyTimes(TArrayView<const FKeyHandle> InHandles, TArrayView<FFrameNumber> OutKeyTimes) override;
+	virtual void SetKeyTimes(TArrayView<const FKeyHandle> InHandles, TArrayView<const FFrameNumber> InKeyTimes) override;
+	virtual void DuplicateKeys(TArrayView<const FKeyHandle> InHandles, TArrayView<FKeyHandle> OutNewHandles) override;
+	virtual void DeleteKeys(TArrayView<const FKeyHandle> InHandles) override;
+	virtual void DeleteKeysFrom(FFrameNumber InTime, bool bDeleteKeysBefore) override;
+	virtual void ChangeFrameResolution(FFrameRate SourceRate, FFrameRate DestinationRate) override;
+	virtual TRange<FFrameNumber> ComputeEffectiveRange() const override;
+	virtual int32 GetNumKeys() const override;
+	virtual void Reset() override;
+	virtual void Offset(FFrameNumber DeltaPosition) override;
+
+private:
+	/** Array of times for each key */
+	UPROPERTY(meta=(KeyTimes))
+	TArray<FFrameNumber> KeyTimes;
+
+	/** Array of values that correspond to each key time */
+	UPROPERTY(meta=(KeyValues))
+	TArray<FMovieSceneCameraShakeSourceTrigger> KeyValues;
+
+	FMovieSceneKeyHandleMap KeyHandles;
+};
+
+template<>
+struct TMovieSceneChannelTraits<FMovieSceneCameraShakeSourceTriggerChannel> : TMovieSceneChannelTraitsBase<FMovieSceneCameraShakeSourceTriggerChannel>
+{
+	enum { SupportsDefaults = false };
+};
+
+inline bool EvaluateChannel(const FMovieSceneCameraShakeSourceTriggerChannel* InChannel, FFrameTime InTime, FMovieSceneCameraShakeSourceTrigger& OutValue)
+{
+	return false;
+}
+
diff --git a/Engine/Source/Runtime/MovieSceneTracks/Public/Sections/MovieSceneCameraShakeSourceTriggerSection.h b/Engine/Source/Runtime/MovieSceneTracks/Public/Sections/MovieSceneCameraShakeSourceTriggerSection.h
new file mode 100644
--- /dev/null
+++ b/Engine/Source/Runtime/MovieSceneTracks/Public/Sections/MovieSceneCameraShakeSourceTriggerSection.h
@@ -0,0 +1,28 @@
+// Copyright Epic Games, Inc. All Rights Reserved.
+
+#pragma once
+
+#include "Camera/CameraShake.h"
+#include "Channels/MovieSceneCameraShakeSourceTriggerChannel.h"
+#include "CoreMinimal.h"
+#include "UObject/ObjectMacros.h"
+#include "MovieSceneCameraShakeSourceTriggerSection.generated.h"
+
+UCLASS(MinimalAPI)
+class UMovieSceneCameraShakeSourceTriggerSection 
+	: public UMovieSceneSection
+{
+public:
+	GENERATED_BODY()
+
+	UMovieSceneCameraShakeSourceTriggerSection(const FObjectInitializer& Init);
+
+	const FMovieSceneCameraShakeSourceTriggerChannel& GetChannel() const { return Channel; }
+
+	virtual FMovieSceneEvalTemplatePtr GenerateTemplate() const override;
+
+private:
+	UPROPERTY()
+	FMovieSceneCameraShakeSourceTriggerChannel Channel;
+};
+
diff --git a/Engine/Source/Runtime/MovieSceneTracks/Public/Tracks/MovieSceneCameraShakeSourceShakeTrack.h b/Engine/Source/Runtime/MovieSceneTracks/Public/Tracks/MovieSceneCameraShakeSourceShakeTrack.h
--- a/Engine/Source/Runtime/MovieSceneTracks/Public/Tracks/MovieSceneCameraShakeSourceShakeTrack.h
+++ b/Engine/Source/Runtime/MovieSceneTracks/Public/Tracks/MovieSceneCameraShakeSourceShakeTrack.h
@@ -12,6 +12,7 @@
 
 struct FMovieSceneEvaluationTrack;
 struct FMovieSceneSegmentCompilerRules;
+class UCameraShakeSourceComponent;
 
 /**
  * 
@@ -22,7 +23,8 @@ class UMovieSceneCameraShakeSourceShakeTrack : public UMovieSceneNameableTrack
 	GENERATED_BODY()
 
 public:
-	virtual UMovieSceneSection* AddNewCameraShake(const FFrameNumber KeyTime, const TSubclassOf<UCameraShake> ShakeClass);
+	MOVIESCENETRACKS_API UMovieSceneSection* AddNewCameraShake(const FFrameNumber KeyTime, const UCameraShakeSourceComponent& ShakeSourceComponent);
+	MOVIESCENETRACKS_API UMovieSceneSection* AddNewCameraShake(const FFrameNumber KeyTime, const TSubclassOf<UCameraShake> ShakeClass, bool bIsAutomaticShake);
 	
 public:
 	// UMovieSceneTrack interface
diff --git a/Engine/Source/Runtime/MovieSceneTracks/Public/Tracks/MovieSceneCameraShakeSourceTriggerTrack.h b/Engine/Source/Runtime/MovieSceneTracks/Public/Tracks/MovieSceneCameraShakeSourceTriggerTrack.h
new file mode 100644
--- /dev/null
+++ b/Engine/Source/Runtime/MovieSceneTracks/Public/Tracks/MovieSceneCameraShakeSourceTriggerTrack.h
@@ -0,0 +1,42 @@
+// Copyright Epic Games, Inc. All Rights Reserved.
+
+#pragma once
+
+#include "CoreMinimal.h"
+#include "UObject/ObjectMacros.h"
+#include "MovieSceneTrack.h"
+#include "MovieSceneCameraShakeSourceTriggerTrack.generated.h"
+
+UCLASS()
+class MOVIESCENETRACKS_API UMovieSceneCameraShakeSourceTriggerTrack
+	: public UMovieSceneTrack
+{
+public:
+
+	GENERATED_BODY()
+
+	UMovieSceneCameraShakeSourceTriggerTrack(const FObjectInitializer& Obj);
+
+	// UMovieSceneTrack interface
+	virtual bool SupportsType(TSubclassOf<UMovieSceneSection> SectionClass) const override;
+	virtual UMovieSceneSection* CreateNewSection() override;
+	virtual bool HasSection(const UMovieSceneSection& Section) const override;
+	virtual void AddSection(UMovieSceneSection& Section) override;
+	virtual void RemoveSection(UMovieSceneSection& Section) override;
+	virtual void RemoveSectionAt(int32 SectionIndex) override;
+	virtual void RemoveAllAnimationData() override;
+	virtual bool IsEmpty() const override;
+	virtual const TArray<UMovieSceneSection*>& GetAllSections() const override;
+	virtual void PostCompile(FMovieSceneEvaluationTrack& Track, const FMovieSceneTrackCompilerArgs& Args) const override;
+
+#if WITH_EDITORONLY_DATA
+	virtual FText GetDisplayName() const override;
+#endif
+
+protected:
+
+	/** All the sections in this track */
+	UPROPERTY()
+	TArray<UMovieSceneSection*> Sections;
+};
+
