Camera Shake: Fix duplicate loctext strings

diff --git a/Engine/Source/Editor/MovieSceneTools/Private/TrackEditors/CameraShakeSourceShakeTrackEditor.cpp b/Engine/Source/Editor/MovieSceneTools/Private/TrackEditors/CameraShakeSourceShakeTrackEditor.cpp
--- a/Engine/Source/Editor/MovieSceneTools/Private/TrackEditors/CameraShakeSourceShakeTrackEditor.cpp
+++ b/Engine/Source/Editor/MovieSceneTools/Private/TrackEditors/CameraShakeSourceShakeTrackEditor.cpp
@@ -554,16 +554,16 @@ TSharedRef<SWidget> FCameraShakeSourceShakeTrackEditor::BuildCameraShakeSubMenu(
 void FCameraShakeSourceShakeTrackEditor::AddCameraShakeSubMenu(FMenuBuilder& MenuBuilder, TArray<FGuid> ObjectBindings)
 {
 	MenuBuilder.AddMenuEntry(
-			LOCTEXT("AutoShake", "Automatic Shake"),
-			LOCTEXT("AutoShakeTooltip", "Adds a section that plays the camera shake already configured on the shake source component."),
+			LOCTEXT("AddAutoShake", "Automatic Shake"),
+			LOCTEXT("AddAutoShakeTooltip", "Adds a section that plays the camera shake already configured on the shake source component."),
 			FSlateIcon(),
 			FUIAction( 
 				FExecuteAction::CreateSP(this, &FCameraShakeSourceShakeTrackEditor::OnAutoCameraShakeSelected, ObjectBindings)
 			));
 
 	MenuBuilder.AddSubMenu(
-			LOCTEXT("OtherShake", "Other Shake"),
-			LOCTEXT("OtherShakeTooltip", "Adds a section that plays a specific camera shake originating from the shake source component."),
+			LOCTEXT("AddOtherShake", "Other Shake"),
+			LOCTEXT("AddOtherShakeTooltip", "Adds a section that plays a specific camera shake originating from the shake source component."),
 			FNewMenuDelegate::CreateSP(this, &FCameraShakeSourceShakeTrackEditor::AddOtherCameraShakeBrowserSubMenu, ObjectBindings));
 }
 
