fix two validation issues in vulkan caused by the Track Editor Thumbnail

diff --git a/Engine/Source/Editor/MovieSceneTools/Private/TrackEditorThumbnail/TrackEditorThumbnail.cpp b/Engine/Source/Editor/MovieSceneTools/Private/TrackEditorThumbnail/TrackEditorThumbnail.cpp
--- a/Engine/Source/Editor/MovieSceneTools/Private/TrackEditorThumbnail/TrackEditorThumbnail.cpp
+++ b/Engine/Source/Editor/MovieSceneTools/Private/TrackEditorThumbnail/TrackEditorThumbnail.cpp
@@ -345,10 +345,13 @@ void FTrackEditorThumbnailCache::DrawThumbnail(FTrackEditorThumbnail& TrackEdito
 		ViewportThumbnailClient->PostDraw(TrackEditorThumbnail);
 	}
 
+	FSlateTextureRenderTarget2DResource* pSlateResource = TrackEditorThumbnail.GetRenderTarget();
 	FThreadSafeBool* bHasFinishedDrawingPtr = &TrackEditorThumbnail.bHasFinishedDrawing;
 	ENQUEUE_RENDER_COMMAND(SetFinishedDrawing)(
-		[bHasFinishedDrawingPtr](FRHICommandList& RHICmdList)
+		[bHasFinishedDrawingPtr, pSlateResource](FRHICommandList& RHICmdList)
 		{
+			FRHITexture* InTexture = pSlateResource->GetTextureRHI();
+			RHICmdList.TransitionResource(EResourceTransitionAccess::EReadable, InTexture);
 			*bHasFinishedDrawingPtr = true;
 		}
 	);
diff --git a/Engine/Source/Runtime/VulkanRHI/Private/VulkanRenderTarget.cpp b/Engine/Source/Runtime/VulkanRHI/Private/VulkanRenderTarget.cpp
--- a/Engine/Source/Runtime/VulkanRHI/Private/VulkanRenderTarget.cpp
+++ b/Engine/Source/Runtime/VulkanRHI/Private/VulkanRenderTarget.cpp
@@ -1299,11 +1299,7 @@ bool FVulkanCommandListContext::FPendingTransition::GatherBarriers(FTransitionAn
 			VkBufferMemoryBarrier& Barrier = OutBufferBarriers[OutBufferBarriers.AddUninitialized()];
 			if (BUF_DrawIndirect == (UAV->SourceVertexBuffer->GetUEUsage() & BUF_DrawIndirect)) //for indirect read we translate Read INDIRECT_COMMAND_READ
 			{
-				if (SrcAccess == VK_ACCESS_SHADER_READ_BIT)
-				{
-					SrcAccess = VK_ACCESS_INDIRECT_COMMAND_READ_BIT;
-				}
-				else if (DestAccess == VK_ACCESS_SHADER_READ_BIT)
+				if (DestAccess == VK_ACCESS_SHADER_READ_BIT)
 				{
 					DestAccess = VK_ACCESS_INDIRECT_COMMAND_READ_BIT;
 				}
