Sequencer: Always add a camera button to the transform track's object binding edit buttons. The visibility of the button is based on whether the bound object has a camera component. This fixes an issue where looking at the spawnable template isn't enough to determine whether a camera exists because the camera component might not exist at the time the button is created.

diff --git a/Engine/Source/Editor/MovieSceneTools/Private/TrackEditors/TransformTrackEditor.cpp b/Engine/Source/Editor/MovieSceneTools/Private/TrackEditors/TransformTrackEditor.cpp
--- a/Engine/Source/Editor/MovieSceneTools/Private/TrackEditors/TransformTrackEditor.cpp
+++ b/Engine/Source/Editor/MovieSceneTools/Private/TrackEditors/TransformTrackEditor.cpp
@@ -399,49 +399,28 @@ void F3DTransformTrackEditor::OnAddTransformKeysForSelectedObjects( EMovieSceneT
 
 void F3DTransformTrackEditor::BuildObjectBindingEditButtons(TSharedPtr<SHorizontalBox> EditBox, const FGuid& ObjectGuid, const UClass* ObjectClass)
 {
-	bool bHasCameraComponent = false;
-
-	UObject* Object = GetSequencer()->FindSpawnedObjectOrTemplate(ObjectGuid);
-	if (Object && Object->IsA<AActor>())
-	{
-		AActor* Actor = Cast<AActor>(Object);
-		if (Actor != nullptr)
-		{
-			UCameraComponent* CameraComponent = MovieSceneHelpers::CameraComponentFromActor(Actor);
-			if (CameraComponent)
-			{
-				bHasCameraComponent = true;
-			}
-		}
-	}
-
-	if (bHasCameraComponent)
-	{
-		// If this is a camera track, add a button to lock the viewport to the camera
-		EditBox.Get()->AddSlot()
-			.VAlign(VAlign_Center)
-			.HAlign(HAlign_Right)
-			.AutoWidth()
-			.Padding(4, 0, 0, 0)
-			[
-				SNew(SCheckBox)		
-					.IsFocusable(false)
-					.Visibility(this, &F3DTransformTrackEditor::IsCameraVisible, ObjectGuid)
-					.IsChecked(this, &F3DTransformTrackEditor::IsCameraLocked, ObjectGuid)
-					.OnCheckStateChanged(this, &F3DTransformTrackEditor::OnLockCameraClicked, ObjectGuid)
-					.ToolTipText(this, &F3DTransformTrackEditor::GetLockCameraToolTip, ObjectGuid)
-					.ForegroundColor(FLinearColor::White)
-					.CheckedImage(FEditorStyle::GetBrush("Sequencer.LockCamera"))
-					.CheckedHoveredImage(FEditorStyle::GetBrush("Sequencer.LockCamera"))
-					.CheckedPressedImage(FEditorStyle::GetBrush("Sequencer.LockCamera"))
-					.UncheckedImage(FEditorStyle::GetBrush("Sequencer.UnlockCamera"))
-					.UncheckedHoveredImage(FEditorStyle::GetBrush("Sequencer.UnlockCamera"))
-					.UncheckedPressedImage(FEditorStyle::GetBrush("Sequencer.UnlockCamera"))
-			];
-	}
-}
-
-
+	// If this is a camera track, add a button to lock the viewport to the camera
+	EditBox.Get()->AddSlot()
+		.VAlign(VAlign_Center)
+		.HAlign(HAlign_Right)
+		.AutoWidth()
+		.Padding(4, 0, 0, 0)
+		[
+			SNew(SCheckBox)		
+				.IsFocusable(false)
+				.Visibility(this, &F3DTransformTrackEditor::IsCameraVisible, ObjectGuid)
+				.IsChecked(this, &F3DTransformTrackEditor::IsCameraLocked, ObjectGuid)
+				.OnCheckStateChanged(this, &F3DTransformTrackEditor::OnLockCameraClicked, ObjectGuid)
+				.ToolTipText(this, &F3DTransformTrackEditor::GetLockCameraToolTip, ObjectGuid)
+				.ForegroundColor(FLinearColor::White)
+				.CheckedImage(FEditorStyle::GetBrush("Sequencer.LockCamera"))
+				.CheckedHoveredImage(FEditorStyle::GetBrush("Sequencer.LockCamera"))
+				.CheckedPressedImage(FEditorStyle::GetBrush("Sequencer.LockCamera"))
+				.UncheckedImage(FEditorStyle::GetBrush("Sequencer.UnlockCamera"))
+				.UncheckedHoveredImage(FEditorStyle::GetBrush("Sequencer.UnlockCamera"))
+				.UncheckedPressedImage(FEditorStyle::GetBrush("Sequencer.UnlockCamera"))
+		];
+};
 void F3DTransformTrackEditor::BuildObjectBindingTrackMenu(FMenuBuilder& MenuBuilder, const TArray<FGuid>& ObjectBindings, const UClass* ObjectClass)
 {
 	if (ObjectClass != nullptr && (ObjectClass->IsChildOf(AActor::StaticClass()) || ObjectClass->IsChildOf(USceneComponent::StaticClass())))
@@ -475,7 +454,11 @@ EVisibility F3DTransformTrackEditor::IsCameraVisible(FGuid ObjectGuid) const
 		
 		if (Actor != nullptr)
 		{
-			return EVisibility::Visible;
+			UCameraComponent* CameraComponent = MovieSceneHelpers::CameraComponentFromActor(Actor);
+			if (CameraComponent)
+			{
+				return EVisibility::Visible;
+			}
 		}
 	}
 
