Sequencer - Prevent a null pointer crash in the thumbnail track when the sequence has bad data.

diff --git a/Engine/Source/Editor/MovieSceneTools/Private/TrackEditorThumbnail/TrackEditorThumbnail.cpp b/Engine/Source/Editor/MovieSceneTools/Private/TrackEditorThumbnail/TrackEditorThumbnail.cpp
--- a/Engine/Source/Editor/MovieSceneTools/Private/TrackEditorThumbnail/TrackEditorThumbnail.cpp
+++ b/Engine/Source/Editor/MovieSceneTools/Private/TrackEditorThumbnail/TrackEditorThumbnail.cpp
@@ -346,15 +346,18 @@ void FTrackEditorThumbnailCache::DrawThumbnail(FTrackEditorThumbnail& TrackEdito
 	}
 
 	FSlateTextureRenderTarget2DResource* pSlateResource = TrackEditorThumbnail.GetRenderTarget();
-	FThreadSafeBool* bHasFinishedDrawingPtr = &TrackEditorThumbnail.bHasFinishedDrawing;
-	ENQUEUE_RENDER_COMMAND(SetFinishedDrawing)(
-		[bHasFinishedDrawingPtr, pSlateResource](FRHICommandList& RHICmdList)
-		{
-			FRHITexture* InTexture = pSlateResource->GetTextureRHI();
-			RHICmdList.TransitionResource(EResourceTransitionAccess::EReadable, InTexture);
-			*bHasFinishedDrawingPtr = true;
-		}
-	);
+	if (pSlateResource != nullptr)
+	{
+		FThreadSafeBool* bHasFinishedDrawingPtr = &TrackEditorThumbnail.bHasFinishedDrawing;
+		ENQUEUE_RENDER_COMMAND(SetFinishedDrawing)(
+			[bHasFinishedDrawingPtr, pSlateResource](FRHICommandList& RHICmdList)
+			{
+				FRHITexture* InTexture = pSlateResource->GetTextureRHI();
+				RHICmdList.TransitionResource(EResourceTransitionAccess::EReadable, InTexture);
+				*bHasFinishedDrawingPtr = true;
+			}
+		);
+	}
 }
 void FTrackEditorThumbnailCache::DrawViewportThumbnail(FTrackEditorThumbnail& TrackEditorThumbnail)
 {
