Sequencer: Set camera button as collapsed when there is no camera component so that it dosen't take up unnecessary space.

#rnx

diff --git a/Engine/Source/Editor/MovieSceneTools/Private/TrackEditors/TransformTrackEditor.cpp b/Engine/Source/Editor/MovieSceneTools/Private/TrackEditors/TransformTrackEditor.cpp
--- a/Engine/Source/Editor/MovieSceneTools/Private/TrackEditors/TransformTrackEditor.cpp
+++ b/Engine/Source/Editor/MovieSceneTools/Private/TrackEditors/TransformTrackEditor.cpp
@@ -462,7 +462,7 @@ EVisibility F3DTransformTrackEditor::IsCameraVisible(FGuid ObjectGuid) const
 		}
 	}
 
-	return EVisibility::Hidden;
+	return EVisibility::Collapsed;
 }
 
 ECheckBoxState F3DTransformTrackEditor::IsCameraLocked(FGuid ObjectGuid) const
