Sequencer: Import onto objects with the focused template ID. This fixes an issue preventing importing onto an object in a subsequence.

diff --git a/Engine/Plugins/MovieScene/SequencerScripting/Source/SequencerScriptingEditor/Private/SequencerTools.cpp b/Engine/Plugins/MovieScene/SequencerScripting/Source/SequencerScriptingEditor/Private/SequencerTools.cpp
--- a/Engine/Plugins/MovieScene/SequencerScripting/Source/SequencerScriptingEditor/Private/SequencerTools.cpp
+++ b/Engine/Plugins/MovieScene/SequencerScripting/Source/SequencerScriptingEditor/Private/SequencerTools.cpp
@@ -425,7 +425,7 @@ bool USequencerToolsFunctionLibrary::ImportFBX(UWorld* World, ULevelSequence* Se
 	ImportFBXCamera(FbxImporter, World, Sequence, MovieScene, Player, MovieSceneSequenceID::Root, ObjectBindingMap, bMatchByNameOnly,
 		ImportFBXSettings->bCreateCameras);
 
-	bool bValid = MovieSceneToolHelpers::ImportFBXIfReady(World, MovieScene, Player, ObjectBindingMap, ImportFBXSettings, InOutParams);
+	bool bValid = MovieSceneToolHelpers::ImportFBXIfReady(World, MovieScene, Player, MovieSceneSequenceID::Root, ObjectBindingMap, ImportFBXSettings, InOutParams);
 
 	return bValid;
 }
diff --git a/Engine/Source/Editor/MovieSceneTools/Private/MovieSceneToolHelpers.cpp b/Engine/Source/Editor/MovieSceneTools/Private/MovieSceneToolHelpers.cpp
--- a/Engine/Source/Editor/MovieSceneTools/Private/MovieSceneToolHelpers.cpp
+++ b/Engine/Source/Editor/MovieSceneTools/Private/MovieSceneToolHelpers.cpp
@@ -1629,7 +1629,7 @@ private:
 		ImportFBXCamera(FbxImporter, MovieScene, *Sequencer, ObjectBindingMap, bMatchByNameOnly, bCreateCameras.IsSet() ? bCreateCameras.GetValue() : ImportFBXSettings->bCreateCameras);
 
 		UWorld* World = Cast<UWorld>(Sequencer->GetPlaybackContext());
-		bool bValid = MovieSceneToolHelpers::ImportFBXIfReady(World, MovieScene, Sequencer, ObjectBindingMap, ImportFBXSettings, InOutParams);
+		bool bValid = MovieSceneToolHelpers::ImportFBXIfReady(World, MovieScene, Sequencer, Sequencer->GetFocusedTemplateID(), ObjectBindingMap, ImportFBXSettings, InOutParams);
 	
 		Sequencer->NotifyMovieSceneDataChanged(EMovieSceneDataChangeType::MovieSceneStructureItemAdded);
 
@@ -1678,7 +1678,7 @@ bool MovieSceneToolHelpers::ReadyFBXForImport(const FString&  ImportFilename, UM
 }
 
 
-bool MovieSceneToolHelpers::ImportFBXIfReady(UWorld* World, UMovieScene* MovieScene, IMovieScenePlayer* Player, TMap<FGuid, FString>& ObjectBindingMap, UMovieSceneUserImportFBXSettings* ImportFBXSettings,
+bool MovieSceneToolHelpers::ImportFBXIfReady(UWorld* World, UMovieScene* MovieScene, IMovieScenePlayer* Player, FMovieSceneSequenceIDRef TemplateID, TMap<FGuid, FString>& ObjectBindingMap, UMovieSceneUserImportFBXSettings* ImportFBXSettings,
 	const FFBXInOutParameters& InParams)
 {
 	UMovieSceneUserImportFBXSettings* CurrentImportFBXSettings = GetMutableDefault<UMovieSceneUserImportFBXSettings>();
@@ -1718,7 +1718,7 @@ bool MovieSceneToolHelpers::ImportFBXIfReady(UWorld* World, UMovieScene* MovieSc
 		{
 			if (FCString::Strcmp(*It.Value().ToUpper(), *NodeName.ToUpper()) == 0)
 			{
-				MovieSceneToolHelpers::ImportFBXNode(NodeName, CurveAPI, MovieScene, Player, MovieSceneSequenceID::Root, It.Key());
+				MovieSceneToolHelpers::ImportFBXNode(NodeName, CurveAPI, MovieScene, Player, TemplateID, It.Key());
 
 				ObjectBindingMap.Remove(It.Key());
 				AllNodeNames.RemoveAt(NodeIndex);
@@ -1751,7 +1751,7 @@ bool MovieSceneToolHelpers::ImportFBXIfReady(UWorld* World, UMovieScene* MovieSc
 			auto It = ObjectBindingMap.CreateConstIterator();
 			if (It)
 			{
-				MovieSceneToolHelpers::ImportFBXNode(NodeName, CurveAPI, MovieScene, Player, MovieSceneSequenceID::Root, It.Key());
+				MovieSceneToolHelpers::ImportFBXNode(NodeName, CurveAPI, MovieScene, Player, TemplateID, It.Key());
 
 				UE_LOG(LogMovieScene, Warning, TEXT("Fbx Import: Failed to find any matching node for (%s). Defaulting to first available (%s)."), *NodeName, *It.Value());
 				ObjectBindingMap.Remove(It.Key());
diff --git a/Engine/Source/Editor/MovieSceneTools/Public/MovieSceneToolHelpers.h b/Engine/Source/Editor/MovieSceneTools/Public/MovieSceneToolHelpers.h
--- a/Engine/Source/Editor/MovieSceneTools/Public/MovieSceneToolHelpers.h
+++ b/Engine/Source/Editor/MovieSceneTools/Public/MovieSceneToolHelpers.h
@@ -263,14 +263,15 @@ public:
 	* loading the data.
 	*
 	* @param World The world to import the fbx into
-	* @param World The movie scene to import the fbx into
+	* @param InMovieScene The movie scene to import the fbx into
 	* @param ObjectBindingMap Map relating binding id's to track names. 
+	* @param TemplateID Id of the sequence that contains the objects being imported onto 
 	* @param ImportFBXSettings FBX Import Settings to enforce.
 	* @param InFBXParams Paremter from ImportReadiedFbx used to reset some internal fbx settings that we override.
 	* @return Whether the fbx file was ready and is ready to be import.
 	*/
 
-	static bool ImportFBXIfReady(UWorld* World, UMovieScene* MovieScene, IMovieScenePlayer* Player, TMap<FGuid, FString>& ObjectBindingMap, UMovieSceneUserImportFBXSettings* ImportFBXSettings,
+	static bool ImportFBXIfReady(UWorld* World, UMovieScene* InMovieScene, IMovieScenePlayer* Player, FMovieSceneSequenceIDRef TemplateID, TMap<FGuid, FString>& ObjectBindingMap, UMovieSceneUserImportFBXSettings* ImportFBXSettings,
 		const FFBXInOutParameters& InFBXParams);
 
 	/**
@@ -279,7 +280,7 @@ public:
 	* @param FbxImporter The Fbx importer
 	* @param InMovieScene The movie scene to import the fbx into
 	* @param Player The player we are getting objects from.
-	* @param TemplateID Id of the sequence template ID.
+	* @param TemplateID Id of the sequence that contains the objects being imported onto
 	* @param InObjectBindingNameMap The object binding to name map to map import fbx animation onto
 	* @param bCreateCameras Whether to allow creation of cameras if found in the fbx file.
 	* @param bNotifySlate  If an issue show's up, whether or not to notify the UI.
