Sequencer: NLA blend Fix.

Merge of CL 14520067 so this can get tested in Frosty asap.

diff --git a/Engine/Source/Editor/MovieSceneTools/Private/EditModes/SkeletalAnimationTrackEditMode.cpp b/Engine/Source/Editor/MovieSceneTools/Private/EditModes/SkeletalAnimationTrackEditMode.cpp
--- a/Engine/Source/Editor/MovieSceneTools/Private/EditModes/SkeletalAnimationTrackEditMode.cpp
+++ b/Engine/Source/Editor/MovieSceneTools/Private/EditModes/SkeletalAnimationTrackEditMode.cpp
@@ -334,95 +334,98 @@ void FSkeletalAnimationTrackEditMode::Render(const FSceneView* View, FViewport*
 		TArrayView<TWeakObjectPtr<UObject>> BoundObjects = Sequencer->FindObjectsInCurrentSequence(Binding.GetObjectGuid());
 		for (UMovieSceneTrack* Track : Binding.GetTracks())
 		{
-			if (UMovieSceneSkeletalAnimationTrack* SkelAnimTrack = Cast<UMovieSceneSkeletalAnimationTrack>(Track))
-			{
-				for (TWeakObjectPtr<> WeakBinding : BoundObjects)
+			if(!Track->IsEvalDisabled())
+			{ 
+				if (UMovieSceneSkeletalAnimationTrack* SkelAnimTrack = Cast<UMovieSceneSkeletalAnimationTrack>(Track))
 				{
-					UObject* BoundObject = WeakBinding.Get();
-					if (BoundObject)
+					for (TWeakObjectPtr<> WeakBinding : BoundObjects)
 					{
-						//show root motions
-						if (SkelAnimTrack->bShowRootMotionTrail)
+						UObject* BoundObject = WeakBinding.Get();
+						if (BoundObject)
 						{
-							//Get Transform Track if so then exit
-							//TODO a bunch, what if no transform track, what if multiple what if external transform?
-							for (UMovieSceneTrack* TrackTrack : Binding.GetTracks())
+							//show root motions
+							if (SkelAnimTrack->bShowRootMotionTrail)
 							{
-								UMovieScene3DTransformTrack* TransformTrack = Cast<UMovieScene3DTransformTrack>(TrackTrack);
-								if (!TransformTrack)
+								//Get Transform Track if so then exit
+								//TODO a bunch, what if no transform track, what if multiple what if external transform?
+								for (UMovieSceneTrack* TrackTrack : Binding.GetTracks())
 								{
-									continue;
+									UMovieScene3DTransformTrack* TransformTrack = Cast<UMovieScene3DTransformTrack>(TrackTrack);
+									if (!TransformTrack)
+									{
+										continue;
+									}
+									RenderTrail(BoundObject, SkelAnimTrack, InterrogationLinker, TransformTrack, Sequencer, PDI);
+
 								}
-								RenderTrail(BoundObject, SkelAnimTrack, InterrogationLinker,TransformTrack, Sequencer, PDI);
-								
 							}
-						}
-						//show skeletons
-						USkeletalMeshComponent* SkelMeshComp = AcquireSkeletalMeshFromObject(BoundObject);
-						if (SkelMeshComp && SkelMeshComp->GetAnimInstance())
-						{
-							const FLinearColor Colors[5] = { FLinearColor(1.0f,0.0f,1.0f,1.0f), FLinearColor(0.0f,1.0f,0.0f,1.0f), FLinearColor(0.0f, 0.0f, 1.0f, 0.0f),
-								FLinearColor(0.5f, 0.5f, 0.0f, 1.0f), FLinearColor(0.0f, 0.5f, 0.5f, 1.0f) };
+							//show skeletons
+							USkeletalMeshComponent* SkelMeshComp = AcquireSkeletalMeshFromObject(BoundObject);
+							if (SkelMeshComp && SkelMeshComp->GetAnimInstance())
+							{
+								const FLinearColor Colors[5] = { FLinearColor(1.0f,0.0f,1.0f,1.0f), FLinearColor(0.0f,1.0f,0.0f,1.0f), FLinearColor(0.0f, 0.0f, 1.0f, 0.0f),
+									FLinearColor(0.5f, 0.5f, 0.0f, 1.0f), FLinearColor(0.0f, 0.5f, 0.5f, 1.0f) };
 
-							int32 SectionIndex = 0;
-							FQualifiedFrameTime CurrentTime = Sequencer->GetLocalTime();
+								int32 SectionIndex = 0;
+								FQualifiedFrameTime CurrentTime = Sequencer->GetLocalTime();
 
-							for (UMovieSceneSection* Section : Track->GetAllSections())
-							{
-								UMovieSceneSkeletalAnimationSection* AnimSection = Cast<UMovieSceneSkeletalAnimationSection>(Section);
-								if (AnimSection && AnimSection->bShowSkeleton && AnimSection->IsActive())
+								for (UMovieSceneSection* Section : Track->GetAllSections())
 								{
-									UAnimSequence* AnimSequence = Cast<UAnimSequence>(AnimSection->Params.Animation);
-									if (AnimSequence)
+									UMovieSceneSkeletalAnimationSection* AnimSection = Cast<UMovieSceneSkeletalAnimationSection>(Section);
+									if (AnimSection && AnimSection->bShowSkeleton && AnimSection->IsActive())
 									{
-										FCompactPose OutPose;
-										OutPose.ResetToRefPose(SkelMeshComp->GetAnimInstance()->GetRequiredBones());
+										UAnimSequence* AnimSequence = Cast<UAnimSequence>(AnimSection->Params.Animation);
+										if (AnimSequence)
+										{
+											FCompactPose OutPose;
+											OutPose.ResetToRefPose(SkelMeshComp->GetAnimInstance()->GetRequiredBones());
 
-										FBlendedCurve OutCurve;
-										FStackCustomAttributes TempAttributes;
-										FAnimationPoseData OutAnimationPoseData(OutPose, OutCurve, TempAttributes);
+											FBlendedCurve OutCurve;
+											FStackCustomAttributes TempAttributes;
+											FAnimationPoseData OutAnimationPoseData(OutPose, OutCurve, TempAttributes);
 
-										const float Seconds = AnimSection->MapTimeToAnimation(CurrentTime.Time, CurrentTime.Rate);
-										FAnimExtractContext ExtractionContext(Seconds, false);
+											const float Seconds = AnimSection->MapTimeToAnimation(CurrentTime.Time, CurrentTime.Rate);
+											FAnimExtractContext ExtractionContext(Seconds, false);
 
-										AnimSequence->GetAnimationPose(OutAnimationPoseData, ExtractionContext);
+											AnimSequence->GetAnimationPose(OutAnimationPoseData, ExtractionContext);
 
-										FFrameTime TimeToSample = CurrentTime.Time;
-										if (TimeToSample < AnimSection->GetRange().GetLowerBoundValue())
-										{
-											TimeToSample = AnimSection->GetRange().GetLowerBoundValue();
-										}
-										else if (TimeToSample > AnimSection->GetRange().GetUpperBoundValue())
-										{
-											TimeToSample = AnimSection->GetRange().GetUpperBoundValue();
-										}
-										FTransform RootMotionAtStart;
-										float Weight;
-										AnimSection->GetRootMotionTransform(TimeToSample, MovieScene->GetTickResolution(), RootMotionAtStart, Weight);
-										RootMotionAtStart = RootMotionAtStart * AnimSection->TempOffsetTransform;
-										OutPose[FCompactPoseBoneIndex(0)] = RootMotionAtStart;
-
-										FVector RootLocation;
-										DrawBonesFromCompactPose(OutPose, SkelMeshComp, PDI, Colors[SectionIndex % 5], RootLocation);
-										if (IsRootSelected(AnimSection))
-										{
-											RootTransform = RootMotionAtStart;
-											RootTransform.SetLocation(RootLocation);
-										}
-										if (PDI != nullptr)
-										{
-											if (PDI->IsHitTesting())
+											FFrameTime TimeToSample = CurrentTime.Time;
+											if (TimeToSample < AnimSection->GetRange().GetLowerBoundValue())
+											{
+												TimeToSample = AnimSection->GetRange().GetLowerBoundValue();
+											}
+											else if (TimeToSample > AnimSection->GetRange().GetUpperBoundValue())
 											{
-												PDI->SetHitProxy(new HMovieSceneSkeletalAnimationRootHitProxy(AnimSection, SkelMeshComp));
+												TimeToSample = AnimSection->GetRange().GetUpperBoundValue();
 											}
-											PDI->DrawPoint(RootLocation, Colors[SectionIndex % 5], 10.f, SDPG_Foreground);
-											if (PDI->IsHitTesting())
+											FTransform RootMotionAtStart;
+											float Weight;
+											AnimSection->GetRootMotionTransform(TimeToSample, MovieScene->GetTickResolution(), RootMotionAtStart, Weight);
+											RootMotionAtStart = RootMotionAtStart * AnimSection->TempOffsetTransform;
+											OutPose[FCompactPoseBoneIndex(0)] = RootMotionAtStart;
+
+											FVector RootLocation;
+											DrawBonesFromCompactPose(OutPose, SkelMeshComp, PDI, Colors[SectionIndex % 5], RootLocation);
+											if (IsRootSelected(AnimSection))
 											{
-												PDI->SetHitProxy(nullptr);
+												RootTransform = RootMotionAtStart;
+												RootTransform.SetLocation(RootLocation);
 											}
+											if (PDI != nullptr)
+											{
+												if (PDI->IsHitTesting())
+												{
+													PDI->SetHitProxy(new HMovieSceneSkeletalAnimationRootHitProxy(AnimSection, SkelMeshComp));
+												}
+												PDI->DrawPoint(RootLocation, Colors[SectionIndex % 5], 10.f, SDPG_Foreground);
+												if (PDI->IsHitTesting())
+												{
+													PDI->SetHitProxy(nullptr);
+												}
 
+											}
+											++SectionIndex;
 										}
-										++SectionIndex;
 									}
 								}
 							}
diff --git a/Engine/Source/Editor/MovieSceneTools/Private/TrackEditors/SkeletalAnimationTrackEditor.cpp b/Engine/Source/Editor/MovieSceneTools/Private/TrackEditors/SkeletalAnimationTrackEditor.cpp
--- a/Engine/Source/Editor/MovieSceneTools/Private/TrackEditors/SkeletalAnimationTrackEditor.cpp
+++ b/Engine/Source/Editor/MovieSceneTools/Private/TrackEditors/SkeletalAnimationTrackEditor.cpp
@@ -802,7 +802,7 @@ void FSkeletalAnimationSection::BuildSectionContextMenu(FMenuBuilder& MenuBuilde
 				FSlateIcon(),
 				FUIAction(
 					FExecuteAction::CreateLambda([=]()->void {
-						FScopedTransaction MatchTransaction(LOCTEXT("MatchSectionByBone_Transaction", "Match Section By Bone"));
+						FScopedTransaction MatchTransaction(LOCTEXT("MatchTranslation_Transaction", "Match Translation"));
 						Section.Modify();	
 						Section.ToggleMatchTranslation();
 						Section.MatchSectionByBoneTransform(SkelMeshComp, SequencerPtr->GetLocalTime().Time, SequencerPtr->GetLocalTime().Rate, Section.MatchedBoneName);
@@ -815,11 +815,11 @@ void FSkeletalAnimationSection::BuildSectionContextMenu(FMenuBuilder& MenuBuilde
 
 			MenuBuilder.AddMenuEntry(
 				LOCTEXT("MatchZHeight", "Match Z Height"),
-				LOCTEXT("MatchZHeightTooltip", "Match the Z Height, may want this off for better matching "),
+				LOCTEXT("MatchZHeightTooltip", "Match the Z Height, may want this off for better matching"),
 				FSlateIcon(),
 				FUIAction(
 					FExecuteAction::CreateLambda([=]()->void { 
-						FScopedTransaction MatchTransaction(LOCTEXT("MatchSectionByBone_Transaction", "Match Section By Bone"));
+						FScopedTransaction MatchTransaction(LOCTEXT("MatchZHeight_Transaction", "Match Include Z Height"));
 						Section.Modify();
 						Section.ToggleMatchIncludeZHeight(); 
 						Section.MatchSectionByBoneTransform(SkelMeshComp, SequencerPtr->GetLocalTime().Time, SequencerPtr->GetLocalTime().Rate, Section.MatchedBoneName);
@@ -829,41 +829,59 @@ void FSkeletalAnimationSection::BuildSectionContextMenu(FMenuBuilder& MenuBuilde
 				NAME_None,
 				EUserInterfaceActionType::ToggleButton
 			);
+
 			MenuBuilder.AddMenuEntry(
-				LOCTEXT("MatchRotation", "Match Yaw And Roll Rotation"),
-				LOCTEXT("MatchRotationTooltip", "Match the UYaw and Roll Rotation to the Specified Bone"),
+				LOCTEXT("MatchYawRotation", "Match Yaw Rotation"),
+				LOCTEXT("MatchYawRotationTooltip", "Match the Yaw Rotation, may want this off for better matching"),
 				FSlateIcon(),
 				FUIAction(
-					FExecuteAction::CreateLambda([=]()->void { 
-						FScopedTransaction MatchTransaction(LOCTEXT("MatchSectionByBone_Transaction", "Match Section By Bone"));
+					FExecuteAction::CreateLambda([=]()->void {
+						FScopedTransaction MatchTransaction(LOCTEXT("MatchYawRotation_Transaction", "Match Yaw Rotation"));
 						Section.Modify();
-						Section.ToggleMatchRotation(); 
+						Section.ToggleMatchIncludeYawRotation();
 						Section.MatchSectionByBoneTransform(SkelMeshComp, SequencerPtr->GetLocalTime().Time, SequencerPtr->GetLocalTime().Rate, Section.MatchedBoneName);
 						SequencerPtr->NotifyMovieSceneDataChanged(EMovieSceneDataChangeType::TrackValueChanged); }),
 					FCanExecuteAction::CreateLambda([=]()->bool { return true; }),
-					FIsActionChecked::CreateLambda([=]()->bool { return Section.bMatchRotation; })),
+							FIsActionChecked::CreateLambda([=]()->bool { return Section.bMatchRotationYaw; })),
 				NAME_None,
-				EUserInterfaceActionType::ToggleButton	
-			);
-
+							EUserInterfaceActionType::ToggleButton
+							);
 
 			MenuBuilder.AddMenuEntry(
 				LOCTEXT("MatchPitchRotation", "Match Pitch Rotation"),
-				LOCTEXT("MatchPitchRotationTooltip", "Match the Pitch Rotation, may want this off for better matching "),
+				LOCTEXT("MatchPitchRotationTooltip", "Match the Pitch Rotation, may want this off for better matching"),
 				FSlateIcon(),
 				FUIAction(
 					FExecuteAction::CreateLambda([=]()->void { 
-						FScopedTransaction MatchTransaction(LOCTEXT("MatchSectionByBone_Transaction", "Match Section By Bone"));
+						FScopedTransaction MatchTransaction(LOCTEXT("MatchPitchRotation_Transaction", "Match Pitch Rotation"));
 						Section.Modify();
 						Section.ToggleMatchIncludePitchRotation();
 						Section.MatchSectionByBoneTransform(SkelMeshComp, SequencerPtr->GetLocalTime().Time, SequencerPtr->GetLocalTime().Rate, Section.MatchedBoneName);
 						SequencerPtr->NotifyMovieSceneDataChanged(EMovieSceneDataChangeType::TrackValueChanged); }),
 					FCanExecuteAction::CreateLambda([=]()->bool { return true; }),
-					FIsActionChecked::CreateLambda([=]()->bool { return Section.bMatchIncludePitchRotation;})),
+					FIsActionChecked::CreateLambda([=]()->bool { return Section.bMatchRotationPitch;})),
 				NAME_None,
 				EUserInterfaceActionType::ToggleButton
 			);
 
+			MenuBuilder.AddMenuEntry(
+				LOCTEXT("MatchRollRotation", "Match Roll Rotation"),
+				LOCTEXT("MatchRollRotationTooltip", "Match the Roll Rotation, may want this off for better matching"),
+				FSlateIcon(),
+				FUIAction(
+					FExecuteAction::CreateLambda([=]()->void {
+						FScopedTransaction MatchTransaction(LOCTEXT("MatchRollRotation_Transaction", "Match Roll Rotation"));
+						Section.Modify();
+						Section.ToggleMatchIncludeRollRotation();
+						Section.MatchSectionByBoneTransform(SkelMeshComp, SequencerPtr->GetLocalTime().Time, SequencerPtr->GetLocalTime().Rate, Section.MatchedBoneName);
+						SequencerPtr->NotifyMovieSceneDataChanged(EMovieSceneDataChangeType::TrackValueChanged); }),
+					FCanExecuteAction::CreateLambda([=]()->bool { return true; }),
+							FIsActionChecked::CreateLambda([=]()->bool { return Section.bMatchRotationRoll; })),
+				NAME_None,
+							EUserInterfaceActionType::ToggleButton
+				);
+
+
 
 			/*
 			MenuBuilder.AddMenuEntry(
diff --git a/Engine/Source/Runtime/MovieSceneTracks/Private/Sections/MovieSceneSkeletalAnimationSection.cpp b/Engine/Source/Runtime/MovieSceneTracks/Private/Sections/MovieSceneSkeletalAnimationSection.cpp
--- a/Engine/Source/Runtime/MovieSceneTracks/Private/Sections/MovieSceneSkeletalAnimationSection.cpp
+++ b/Engine/Source/Runtime/MovieSceneTracks/Private/Sections/MovieSceneSkeletalAnimationSection.cpp
@@ -59,9 +59,10 @@ UMovieSceneSkeletalAnimationSection::UMovieSceneSkeletalAnimationSection( const
 	MatchedRotationOffset = FRotator::ZeroRotator;
 
 	bMatchTranslation = true;
-	bMatchRotation = true;
+	bMatchRotationYaw = true;
+	bMatchRotationRoll = false;
+	bMatchRotationPitch = false;
 	bMatchIncludeZHeight = false;
-	bMatchIncludePitchRotation = false;
 
 #if WITH_EDITOR
 
@@ -436,15 +437,20 @@ bool UMovieSceneSkeletalAnimationSection::GetRootMotionTransform(FFrameTime Curr
 {
 	UAnimSequence* AnimSequence = Cast<UAnimSequence>(Params.Animation);
 	FTransform OffsetTransform = GetOffsetTransform();
-	if (AnimSequence && AnimSequence->bForceRootLock == false)
+	if (AnimSequence)
 	{
 		float ManualWeight = 1.f;
 		Params.Weight.Evaluate(CurrentTime, ManualWeight);
 		OutWeight = ManualWeight * EvaluateEasing(CurrentTime);
 
 		float CurrentTimeSeconds = MapTimeToAnimation(CurrentTime, FrameRate);
+		if (AnimSequence->bForceRootLock) //if root lock is on just get the first frame since ExtractRootTrackTransform will still extract(it ignores the flag)
+		{
+			CurrentTimeSeconds = 0.0f;
+		}
 		OutTransform = AnimSequence->ExtractRootTrackTransform(CurrentTimeSeconds,nullptr);
 		OutTransform = OutTransform * OffsetTransform;
+
 		return true;
 	}
 	//for safety always return true for now
@@ -461,17 +467,59 @@ void UMovieSceneSkeletalAnimationSection::FindBestBlendPoint(USkeletalMeshCompon
 	}
 }
 
+void UMovieSceneSkeletalAnimationSection::MultiplyOutInverseOnNextClips(FVector PreviousMatchedLocationOffset, FRotator PreviousMatchedRotationOffset)
+{
+	UMovieSceneSkeletalAnimationTrack* Track = GetTypedOuter<UMovieSceneSkeletalAnimationTrack>();
+	if (Track)
+	{
+		bool bMultiplyOutInverse = false;
+		//calculate the diff here....
+		FTransform Previous(PreviousMatchedRotationOffset.Quaternion(), PreviousMatchedLocationOffset);
+		FTransform Matched(MatchedRotationOffset.Quaternion(), MatchedLocationOffset);
+		FTransform Inverse = Previous.GetRelativeTransformReverse(Matched);
+		for (UMovieSceneSection* Section : Track->AnimationSections)
+		{
+			if (bMultiplyOutInverse)
+			{
+				UMovieSceneSkeletalAnimationSection* AnimSection = Cast<UMovieSceneSkeletalAnimationSection>(Section);
+				if (AnimSection)
+				{
+					//then for all of the next sections we need to multiply that diff through.
+					FTransform CurrentMatched(AnimSection->MatchedRotationOffset, AnimSection->MatchedLocationOffset);
+					FTransform NewMatched = Inverse.GetRelativeTransformReverse(CurrentMatched);
+					AnimSection->MatchedLocationOffset = NewMatched.GetTranslation();
+					AnimSection->MatchedRotationOffset = NewMatched.GetRotation().Rotator();
+				}
+				break;
+			}
+			if (Section == this) //next ones we multiply out
+			{
+				bMultiplyOutInverse = true;
+			}
+		}
+	}
+}
+
 void UMovieSceneSkeletalAnimationSection::ClearMatchedOffsetTransforms()
 {
-	bMatchWithPrevious = true;
-	MatchedBoneName = NAME_None;
+	//need to store the previous since we may need to apply the change we made to the next clips so they don't move
+	FVector PreviousMatchedLocationOffset = MatchedLocationOffset;
+	FRotator PreviousMatchedRotationOffset = MatchedRotationOffset;
 	MatchedLocationOffset = FVector::ZeroVector;
 	MatchedRotationOffset = FRotator::ZeroRotator;
+	if (bMatchWithPrevious == false)
+	{
+		MultiplyOutInverseOnNextClips(PreviousMatchedLocationOffset, PreviousMatchedRotationOffset);
+	}
+	bMatchWithPrevious = true;
+	MatchedBoneName = NAME_None;
 	if (GetRootMotionParams())
 	{
 		GetRootMotionParams()->bRootMotionsDirty = true;
 	}
 }
+
+
 void UMovieSceneSkeletalAnimationSection::MatchSectionByBoneTransform(USkeletalMeshComponent* SkelMeshComp, FFrameTime CurrentFrame, FFrameRate FrameRate,
 	const FName& BoneName)
 {
@@ -481,63 +529,25 @@ void UMovieSceneSkeletalAnimationSection::MatchSectionByBoneTransform(USkeletalM
 	{
 		FTransform DiffTransform;
 		FVector DiffTranslate;
-		FQuat DiffRotate;
-		Track->MatchSectionByBoneTransform(bMatchWithPrevious,SkelMeshComp, this, CurrentFrame, FrameRate, BoneName, DiffTransform,DiffTranslate,DiffRotate);
+		FQuat DiffRotate; 
+		//need to store the previous since we may need to apply the change we made to the next clips so they don't move
+		FVector PreviousMatchedLocationOffset = MatchedLocationOffset;
+		FRotator PreviousMatchedRotationOffset = MatchedRotationOffset;
 
-		/*
-		StartLocationOffset = bMatchTranslation ? DiffTranslate : FVector::ZeroVector;
-		if (!bMatchIncludeZHeight)
-		{
-			StartLocationOffset.Z = 0.0f;
-		}
+		Track->MatchSectionByBoneTransform(bMatchWithPrevious,SkelMeshComp, this, CurrentFrame, FrameRate, BoneName, DiffTransform,DiffTranslate,DiffRotate);
 
-		StartRotationOffset = bMatchRotation ? DiffRotate.Rotator() : FRotator::ZeroRotator;
-		if (!bMatchIncludePitchRotation)
-		{
-			StartRotationOffset.Pitch = 0.0f;
-		}
-		*/
-		
-		
-		
 		MatchedLocationOffset = bMatchTranslation ? DiffTranslate : FVector::ZeroVector;
-		if (!bMatchIncludeZHeight)
-		{
-			MatchedLocationOffset.Z = 0.0f;
-		}
-
-		MatchedRotationOffset = bMatchRotation ? DiffRotate.Rotator() : FRotator::ZeroRotator;
-		if (!bMatchIncludePitchRotation)
-		{
-			MatchedRotationOffset.Pitch = 0.0f;
-		}
+		MatchedRotationOffset = DiffRotate.Rotator();
 		
-		//if we previous we need to multiply the inverses of these out so they stay at their current locations
 		if (bMatchWithPrevious == false)
 		{
-			bool bMultiplyOutInverse = false;
-			FVector InverseLoc(-MatchedLocationOffset);
-			FQuat InverseRot(MatchedRotationOffset.Quaternion().Inverse());
-			for (UMovieSceneSection* Section : Track->AnimationSections)
-			{
-				if (bMultiplyOutInverse)
-				{
-					UMovieSceneSkeletalAnimationSection* AnimSection = Cast<UMovieSceneSkeletalAnimationSection>(Section);
-					if (AnimSection)
-					{
-						AnimSection->MatchedLocationOffset += InverseLoc;
-						AnimSection->MatchedRotationOffset = (AnimSection->MatchedRotationOffset.Quaternion() * InverseRot).Rotator();
-					}
-					break;
-				}
-				if (Section == this) //next ones we multily out
-				{
-					bMultiplyOutInverse = true;
-				}
-			}
+			MultiplyOutInverseOnNextClips(PreviousMatchedLocationOffset, PreviousMatchedRotationOffset);
 		}
 		
-		GetRootMotionParams()->bRootMotionsDirty = true;
+		if (GetRootMotionParams())
+		{
+			GetRootMotionParams()->bRootMotionsDirty = true;
+		}
 	}
 }
 
@@ -548,21 +558,27 @@ void UMovieSceneSkeletalAnimationSection::ToggleMatchTranslation()
 	GetRootMotionParams()->bRootMotionsDirty = true;
 }
 
-void UMovieSceneSkeletalAnimationSection::ToggleMatchRotation()
+void UMovieSceneSkeletalAnimationSection::ToggleMatchIncludeZHeight()
 {
-	bMatchRotation = bMatchRotation ? false : true;
+	bMatchIncludeZHeight = bMatchIncludeZHeight ? false : true;
 	GetRootMotionParams()->bRootMotionsDirty = true;
 }
 
-void UMovieSceneSkeletalAnimationSection::ToggleMatchIncludeZHeight()
+void UMovieSceneSkeletalAnimationSection::ToggleMatchIncludeYawRotation()
 {
-	bMatchIncludeZHeight = bMatchIncludeZHeight ? false : true;
+	bMatchRotationYaw = bMatchRotationYaw ? false : true;
 	GetRootMotionParams()->bRootMotionsDirty = true;
 }
 
 void UMovieSceneSkeletalAnimationSection::ToggleMatchIncludePitchRotation()
 {
-	bMatchIncludePitchRotation = bMatchIncludePitchRotation ? false : true;
+	bMatchRotationPitch = bMatchRotationPitch ? false : true;
+	GetRootMotionParams()->bRootMotionsDirty = true;
+}
+
+void UMovieSceneSkeletalAnimationSection::ToggleMatchIncludeRollRotation()
+{
+	bMatchRotationRoll = bMatchRotationRoll ? false : true;
 	GetRootMotionParams()->bRootMotionsDirty = true;
 }
 
diff --git a/Engine/Source/Runtime/MovieSceneTracks/Private/Tracks/MovieSceneSkeletalAnimationTrack.cpp b/Engine/Source/Runtime/MovieSceneTracks/Private/Tracks/MovieSceneSkeletalAnimationTrack.cpp
--- a/Engine/Source/Runtime/MovieSceneTracks/Private/Tracks/MovieSceneSkeletalAnimationTrack.cpp
+++ b/Engine/Source/Runtime/MovieSceneTracks/Private/Tracks/MovieSceneSkeletalAnimationTrack.cpp
@@ -490,7 +490,7 @@ void UMovieSceneSkeletalAnimationTrack::SetUpRootMotions(bool bForce)
 					UMovieSceneSkeletalAnimationSection* AnimSection = Cast<UMovieSceneSkeletalAnimationSection>(Section);
 					if (AnimSection->GetRootMotionTransform( FrameNumber.FrameNumber, TickResolution, CurrentTransform, CurrentWeight))
 					{
-						CurrentTransform = CurrentTransform * AnimSection->TempOffsetTransform;
+						CurrentTransform =  CurrentTransform * AnimSection->TempOffsetTransform;
 						CurrentTransforms.Add(CurrentTransform);
 						CurrentWeights.Add(CurrentWeight);
 					}
@@ -579,6 +579,10 @@ void UMovieSceneSkeletalAnimationTrack::MatchSectionByBoneTransform(bool bMatchW
 		PrevSection = Section;
 	}
 
+	TranslationDiff = FVector(0.0f, 0.0f, 0.0f);
+	RotationDiff = FQuat::Identity;
+	SecondSectionRootDiff = FTransform::Identity;
+
 	if (bMatchWithPrevious && PrevSection)
 	{
 		UMovieSceneSkeletalAnimationSection* FirstSection = Cast<UMovieSceneSkeletalAnimationSection>(PrevSection);
@@ -589,16 +593,45 @@ void UMovieSceneSkeletalAnimationTrack::MatchSectionByBoneTransform(bool bMatchW
 		{
 			float FirstSectionTime = FirstSection->MapTimeToAnimation(CurrentFrame, FrameRate);
 			FTransform  FirstTransform = GetWorldTransformForBone(FirstAnimSequence, SkelMeshComp, BoneName, FirstSectionTime);
-			//multiply it by the offset transform
-			//We know do this is SetUpRootMotions so we can move clips around before this one
-			//FTransform OffsetTransform = FirstSection->GetOffsetTransform();
-			//FirstTransform = FirstTransform * OffsetTransform;
 			float SecondSectionTime = CurrentSection->MapTimeToAnimation(CurrentFrame, FrameRate);
 			FTransform  SecondTransform = GetWorldTransformForBone(SecondAnimSequence, SkelMeshComp, BoneName, SecondSectionTime);
-			FTransform SecondInverse = SecondTransform.Inverse();
-			SecondSectionRootDiff = SecondInverse * FirstTransform;
-			TranslationDiff = - SecondTransform.GetTranslation() + FirstTransform.GetTranslation();
-			RotationDiff = SecondInverse.GetRotation() * FirstTransform.GetRotation();
+
+			//Need to match the translations and rotations here 
+			FVector FirstTransformTranslation = FirstTransform.GetTranslation();
+			FVector SecondTransformTranslation = SecondTransform.GetTranslation();
+			FRotator FirstTransformRotation = FirstTransform.GetRotation().Rotator();
+			FRotator SecondTransformRotation = SecondTransform.GetRotation().Rotator();
+			if (!CurrentSection->bMatchTranslation)
+			{
+				FirstTransformTranslation.X = SecondTransformTranslation.X;
+				FirstTransformTranslation.Y = SecondTransformTranslation.Y;
+				FirstTransformTranslation.Z = SecondTransformTranslation.Z;
+			}
+			if (!CurrentSection->bMatchIncludeZHeight)
+			{
+				FirstTransformTranslation.Z = SecondTransformTranslation.Z;
+			}
+			FirstTransform.SetTranslation(FirstTransformTranslation);
+
+			if (!CurrentSection->bMatchRotationYaw)
+			{
+				FirstTransformRotation.Yaw = SecondTransformRotation.Yaw;
+			}
+			if (!CurrentSection->bMatchRotationPitch)
+			{
+				FirstTransformRotation.Pitch = SecondTransformRotation.Pitch;
+			}
+			if (!CurrentSection->bMatchRotationRoll)
+			{
+				FirstTransformRotation.Roll = SecondTransformRotation.Roll;
+			}
+			FirstTransform.SetRotation(FirstTransformRotation.Quaternion());
+
+			// Below is the match but we need to use GetRelativeTransformReverse since Inverse doesn't work as expected.
+			//	* GetRelativeTransformReverse returns this(-1)* Other, and parameter is Other.
+			SecondSectionRootDiff = SecondTransform.GetRelativeTransformReverse(FirstTransform);
+			TranslationDiff = SecondSectionRootDiff.GetTranslation();
+			RotationDiff = SecondSectionRootDiff.GetRotation();
 
 		}
 	}
@@ -612,16 +645,45 @@ void UMovieSceneSkeletalAnimationTrack::MatchSectionByBoneTransform(bool bMatchW
 		{
 			float FirstSectionTime = CurrentSection->MapTimeToAnimation(CurrentFrame, FrameRate);
 			FTransform  FirstTransform = GetWorldTransformForBone(FirstAnimSequence, SkelMeshComp, BoneName, FirstSectionTime);
-			//multiply it by the offset transform
-			//We know do this is SetUpRootMotions so we can move clips around before this one
-			//FTransform OffsetTransform = FirstSection->GetOffsetTransform();
-			//FirstTransform = FirstTransform * OffsetTransform;
 			float SecondSectionTime = SecondSection->MapTimeToAnimation(CurrentFrame, FrameRate);
 			FTransform  SecondTransform = GetWorldTransformForBone(SecondAnimSequence, SkelMeshComp, BoneName, SecondSectionTime);
-			FTransform FirstInverse = FirstTransform.Inverse();
-			SecondSectionRootDiff = SecondTransform * FirstInverse;
-			TranslationDiff = SecondTransform.GetTranslation() - FirstTransform.GetTranslation();
-			RotationDiff = SecondTransform.GetRotation() * FirstInverse.GetRotation();
+
+			//Need to match the translations and rotations here 
+			FVector FirstTransformTranslation = FirstTransform.GetTranslation();
+			FVector SecondTransformTranslation = SecondTransform.GetTranslation();
+			FRotator FirstTransformRotation = FirstTransform.GetRotation().Rotator();
+			FRotator SecondTransformRotation = SecondTransform.GetRotation().Rotator();
+			if (!CurrentSection->bMatchTranslation)
+			{
+				SecondTransformTranslation.X = FirstTransformTranslation.X;
+				SecondTransformTranslation.Y = FirstTransformTranslation.Y;
+				SecondTransformTranslation.Z = FirstTransformTranslation.Z;
+			}
+			if (!CurrentSection->bMatchIncludeZHeight)
+			{
+				SecondTransformTranslation.Z = FirstTransformTranslation.Z;
+			}
+			SecondTransform.SetTranslation(SecondTransformTranslation);
+
+			if (!CurrentSection->bMatchRotationYaw)
+			{
+				SecondTransformRotation.Yaw = FirstTransformRotation.Yaw;
+			}
+			if (!CurrentSection->bMatchRotationPitch)
+			{
+				SecondTransformRotation.Pitch = FirstTransformRotation.Pitch;
+			}
+			if (!CurrentSection->bMatchRotationRoll)
+			{
+				SecondTransformRotation.Roll = FirstTransformRotation.Roll;
+			}
+			SecondTransform.SetRotation(SecondTransformRotation.Quaternion());
+
+			//GetRelativeTransformReverse returns this(-1)* Other, and parameter is Other.
+			SecondSectionRootDiff = FirstTransform.GetRelativeTransformReverse(SecondTransform);
+			TranslationDiff = SecondSectionRootDiff.GetTranslation();
+			RotationDiff = SecondSectionRootDiff.GetRotation();
+
 		}
 	}
 }
diff --git a/Engine/Source/Runtime/MovieSceneTracks/Public/Sections/MovieSceneSkeletalAnimationSection.h b/Engine/Source/Runtime/MovieSceneTracks/Public/Sections/MovieSceneSkeletalAnimationSection.h
--- a/Engine/Source/Runtime/MovieSceneTracks/Public/Sections/MovieSceneSkeletalAnimationSection.h
+++ b/Engine/Source/Runtime/MovieSceneTracks/Public/Sections/MovieSceneSkeletalAnimationSection.h
@@ -172,13 +172,17 @@ public:
 	bool bMatchTranslation;
 
 	UPROPERTY()
-	bool bMatchRotation;
+	bool bMatchIncludeZHeight;
 
 	UPROPERTY()
-	bool bMatchIncludeZHeight;
+	bool bMatchRotationYaw;
 
 	UPROPERTY()
-	bool bMatchIncludePitchRotation;
+	bool bMatchRotationPitch;
+
+	UPROPERTY()
+	bool bMatchRotationRoll;
+
 #if WITH_EDITORONLY_DATA
 	/** Whether to show the underlying skeleton for this section. */
 	UPROPERTY(EditAnywhere, Category = "Root Motions")
@@ -206,15 +210,20 @@ public:
 
 	MOVIESCENETRACKS_API void ToggleMatchTranslation();
 
-	MOVIESCENETRACKS_API void ToggleMatchRotation();
-
 	MOVIESCENETRACKS_API void ToggleMatchIncludeZHeight();
 
+	MOVIESCENETRACKS_API void ToggleMatchIncludeYawRotation();
+
 	MOVIESCENETRACKS_API void ToggleMatchIncludePitchRotation();
 
+	MOVIESCENETRACKS_API void ToggleMatchIncludeRollRotation();
+
 #if WITH_EDITORONLY_DATA
 	MOVIESCENETRACKS_API void ToggleShowSkeleton();
 #endif
 
+private:
+	void MultiplyOutInverseOnNextClips(FVector PreviousMatchedLocationOffset, FRotator PreviousMatchedRotationOffset);
+
 
 };
