Control Rig: Sequencer: When baking to an aim sequence, CR Components don't tick, so we now listen to the post evaluate event and broadcast an event that the control rig track editor listens to and so we then upate/tick.

diff --git a/Engine/Plugins/Experimental/ControlRig/Source/ControlRigEditor/Private/Sequencer/ControlRigParameterTrackEditor.cpp b/Engine/Plugins/Experimental/ControlRig/Source/ControlRigEditor/Private/Sequencer/ControlRigParameterTrackEditor.cpp
--- a/Engine/Plugins/Experimental/ControlRig/Source/ControlRigEditor/Private/Sequencer/ControlRigParameterTrackEditor.cpp
+++ b/Engine/Plugins/Experimental/ControlRig/Source/ControlRigEditor/Private/Sequencer/ControlRigParameterTrackEditor.cpp
@@ -167,6 +167,8 @@ FControlRigParameterTrackEditor::FControlRigParameterTrackEditor(TSharedRef<ISeq
 	: FKeyframeTrackEditor<UMovieSceneControlRigParameterTrack>(InSequencer), bIsDoingSelection(false),  bFilterAssetBySkeleton(true), bFilterAssetByAnimatableControls(true)
 
 {
+	FMovieSceneToolsModule::Get().RegisterAnimationBakeHelper(this);
+
 	UMovieScene* MovieScene = InSequencer->GetFocusedMovieSceneSequence()->GetMovieScene();
 
 	SelectionChangedHandle = InSequencer->GetSelectionChangedTracks().AddRaw(this, &FControlRigParameterTrackEditor::OnSelectionChanged);
@@ -291,6 +293,7 @@ FControlRigParameterTrackEditor::~FControlRigParameterTrackEditor()
 		//REMOVE ME IN UE5
 		GetSequencer()->GetObjectChangeListener().GetOnPropagateObjectChanges().RemoveAll(this);
 	}
+	FMovieSceneToolsModule::Get().UnregisterAnimationBakeHelper(this);
 }
 
 void FControlRigParameterTrackEditor::ObjectImplicitlyAdded(UObject* InObject)
@@ -1318,6 +1321,30 @@ void FControlRigParameterTrackEditor::OnTreeViewChanged()
 	}
 }
 
+void FControlRigParameterTrackEditor::PostEvaluation(UMovieScene* MovieScene, FFrameNumber Frame)
+{
+	if (MovieScene)
+	{
+		const TArray<FMovieSceneBinding>& Bindings = MovieScene->GetBindings();
+		for (const FMovieSceneBinding& Binding : Bindings)
+		{
+			if (UMovieSceneControlRigParameterTrack* Track = Cast<UMovieSceneControlRigParameterTrack>(MovieScene->FindTrack(UMovieSceneControlRigParameterTrack::StaticClass(), Binding.GetObjectGuid(), NAME_None)))
+			{
+				if (UControlRig* ControlRig = Track->GetControlRig())
+				{
+					if (ControlRig->GetObjectBinding())
+					{
+						if (UControlRigComponent* ControlRigComponent = Cast<UControlRigComponent>(ControlRig->GetObjectBinding()->GetBoundObject()))
+						{
+							ControlRigComponent->Update(.1); //delta time doesn't matter.
+						}
+					}
+				}
+			}
+		}
+	}
+}
+
 void FControlRigParameterTrackEditor::OnSelectionChanged(TArray<UMovieSceneTrack*> InTracks)
 {
 	if (bIsDoingSelection || GetSequencer().IsValid() == false)
diff --git a/Engine/Plugins/Experimental/ControlRig/Source/ControlRigEditor/Private/Sequencer/ControlRigParameterTrackEditor.h b/Engine/Plugins/Experimental/ControlRig/Source/ControlRigEditor/Private/Sequencer/ControlRigParameterTrackEditor.h
--- a/Engine/Plugins/Experimental/ControlRig/Source/ControlRigEditor/Private/Sequencer/ControlRigParameterTrackEditor.h
+++ b/Engine/Plugins/Experimental/ControlRig/Source/ControlRigEditor/Private/Sequencer/ControlRigParameterTrackEditor.h
@@ -15,6 +15,7 @@
 #include "Sections/MovieScene3DTransformSection.h"
 #include "AcquiredResources.h"
 #include "MovieSceneToolHelpers.h"
+#include "MovieSceneToolsModule.h"
 
 struct FAssetData;
 class FMenuBuilder;
@@ -25,7 +26,7 @@ class UFKControlRig;
 /**
  * Tools for animation tracks
  */
-class FControlRigParameterTrackEditor : public FKeyframeTrackEditor<UMovieSceneControlRigParameterTrack>
+class FControlRigParameterTrackEditor : public FKeyframeTrackEditor<UMovieSceneControlRigParameterTrack>, public IMovieSceneToolsAnimationBakeHelper
 {
 public:
 	/**
@@ -62,6 +63,9 @@ public:
 	virtual void BuildTrackContextMenu(FMenuBuilder& MenuBuilder, UMovieSceneTrack* InTrack) override;
 	virtual bool HandleAssetAdded(UObject* Asset, const FGuid& TargetObjectGuid) override;
 
+	//IMovieSceneToolsAnimationBakeHelper
+	virtual void PostEvaluation(UMovieScene* MovieScene, FFrameNumber Frame);
+
 private:
 
 	void HandleAddTrackSubMenu(FMenuBuilder& MenuBuilder, TArray<FGuid> ObjectBindings, UMovieSceneTrack* Track);
@@ -82,6 +86,8 @@ private:
 	/** Delegate for Selection Changed Event */
 	void OnSelectionChanged(TArray<UMovieSceneTrack*> InTracks);
 
+
+
 	/** Delegate for  Tree View Changed Event */
 	void OnTreeViewChanged();
 
diff --git a/Engine/Source/Editor/MovieSceneTools/Private/MovieSceneToolHelpers.cpp b/Engine/Source/Editor/MovieSceneTools/Private/MovieSceneToolHelpers.cpp
--- a/Engine/Source/Editor/MovieSceneTools/Private/MovieSceneToolHelpers.cpp
+++ b/Engine/Source/Editor/MovieSceneTools/Private/MovieSceneToolHelpers.cpp
@@ -3420,6 +3420,8 @@ bool MovieSceneToolHelpers::BakeToSkelMeshToCallbacks(UMovieScene* MovieScene, I
 	USkeletalMeshComponent* SkelMeshComp, FMovieSceneSequenceIDRef& Template, FMovieSceneSequenceTransform& RootToLocalTransform,
 	FInitAnimationCB InitCallback, FStartAnimationCB StartCallback, FTickAnimationCB TickCallback, FEndAnimationCB EndCallback)
 {
+	const TArray<IMovieSceneToolsAnimationBakeHelper*>&  BakeHelpers = FMovieSceneToolsModule::Get().GetAnimationBakeHelpers();
+
 	//if we have no allocated bone space transforms something wrong so try to recalc them
 	if (SkelMeshComp->GetBoneSpaceTransforms().Num() <= 0)
 	{
@@ -3456,10 +3458,32 @@ bool MovieSceneToolHelpers::BakeToSkelMeshToCallbacks(UMovieScene* MovieScene, I
 		CVarAlwaysSendInterpolatedLiveLink->Set(1, ECVF_SetByConsole);
 	}
 
+	for (IMovieSceneToolsAnimationBakeHelper* BakeHelper : BakeHelpers)
+	{
+		if (BakeHelper)
+		{
+			BakeHelper->StartBaking(MovieScene);
+		}
+	}
 	InitCallback.ExecuteIfBound();
 
 	//Begin records a frame so need to set things up first
+	for (IMovieSceneToolsAnimationBakeHelper* BakeHelper : BakeHelpers)
+	{
+		if (BakeHelper)
+		{
+			BakeHelper->PreEvaluation(MovieScene,LocalStartFrame);
+		}
+	}
+	// This evaluaties the MoviePlayer
 	AnimTrackAdapter.UpdateAnimation(LocalStartFrame);
+	for (IMovieSceneToolsAnimationBakeHelper* BakeHelper : BakeHelpers)
+	{
+		if (BakeHelper)
+		{
+			BakeHelper->PostEvaluation(MovieScene, LocalStartFrame);
+		}
+	}
 	SkelMeshComp->TickAnimation(0.03f, false);
 	SkelMeshComp->RefreshBoneTransforms();
 	SkelMeshComp->RefreshSlaveComponents();
@@ -3502,8 +3526,23 @@ bool MovieSceneToolHelpers::BakeToSkelMeshToCallbacks(UMovieScene* MovieScene, I
 	{
 		int32 LocalFrame = LocalStartFrame + FrameCount;
 
+		for (IMovieSceneToolsAnimationBakeHelper* BakeHelper : BakeHelpers)
+		{
+			if (BakeHelper)
+			{
+				BakeHelper->PreEvaluation(MovieScene, LocalStartFrame);
+			}
+		}
 		// This will call UpdateSkelPose on the skeletal mesh component to move bones based on animations in the matinee group
+		// This also evaluaties the MoviePlayer
 		AnimTrackAdapter.UpdateAnimation(LocalFrame);
+		for (IMovieSceneToolsAnimationBakeHelper* BakeHelper : BakeHelpers)
+		{
+			if (BakeHelper)
+			{
+				BakeHelper->PostEvaluation(MovieScene, LocalStartFrame);
+			}
+		}
 
 		//Live Link sourcer can show up at any time so we unfortunately need to check for it
 		if (LiveLinkClient)
@@ -3552,6 +3591,13 @@ bool MovieSceneToolHelpers::BakeToSkelMeshToCallbacks(UMovieScene* MovieScene, I
 		TickCallback.ExecuteIfBound(DeltaTime);
 	}
 
+	for (IMovieSceneToolsAnimationBakeHelper* BakeHelper : BakeHelpers)
+	{
+		if (BakeHelper)
+		{
+			BakeHelper->StopBaking(MovieScene);
+		}
+	}
 	EndCallback.ExecuteIfBound();
 
 	//now do any sequencer live link cleanup
diff --git a/Engine/Source/Editor/MovieSceneTools/Private/MovieSceneToolsModule.cpp b/Engine/Source/Editor/MovieSceneTools/Private/MovieSceneToolsModule.cpp
--- a/Engine/Source/Editor/MovieSceneTools/Private/MovieSceneToolsModule.cpp
+++ b/Engine/Source/Editor/MovieSceneTools/Private/MovieSceneToolsModule.cpp
@@ -460,6 +460,18 @@ void FMovieSceneToolsModule::RegisterClipboardConversions()
 	FSequencerClipboardReconciler::AddTrackAlias("W", "A");
 }
 
+void FMovieSceneToolsModule::RegisterAnimationBakeHelper(IMovieSceneToolsAnimationBakeHelper* InBakeHelper)
+{
+	checkf(!BakeHelpers.Contains(InBakeHelper), TEXT("Bake Helper is already registered"));
+	BakeHelpers.Add(InBakeHelper);
+}
+
+void FMovieSceneToolsModule::UnregisterAnimationBakeHelper(IMovieSceneToolsAnimationBakeHelper* InBakeHelper)
+{
+	checkf(BakeHelpers.Contains(InBakeHelper), TEXT("Bake Helper is not registered"));
+	BakeHelpers.Remove(InBakeHelper);
+}
+
 void FMovieSceneToolsModule::RegisterTakeData(IMovieSceneToolsTakeData* InTakeData)
 {
 	checkf(!TakeDatas.Contains(InTakeData), TEXT("Take Data is already registered"));
diff --git a/Engine/Source/Editor/MovieSceneTools/Public/MovieSceneToolsModule.h b/Engine/Source/Editor/MovieSceneTools/Public/MovieSceneToolsModule.h
--- a/Engine/Source/Editor/MovieSceneTools/Public/MovieSceneToolsModule.h
+++ b/Engine/Source/Editor/MovieSceneTools/Public/MovieSceneToolsModule.h
@@ -27,6 +27,16 @@ public:
 	virtual bool SetTakeNumber(const UMovieSceneSection*, uint32 InTakeNumber) = 0;
 };
 
+//Interface to get notifications when an animation bake happens in case in needs to run custom code
+class IMovieSceneToolsAnimationBakeHelper
+{
+public:
+	virtual void StartBaking(UMovieScene* MovieScene) {};
+	virtual void PreEvaluation(UMovieScene* MovieScene, FFrameNumber Frame) {};
+	virtual void PostEvaluation(UMovieScene* MovieScene, FFrameNumber Frame) {};
+	virtual void StopBaking(UMovieScene* MovieScene) {};
+};
+
 /**
 * Implements the MovieSceneTools module.
 */
@@ -44,6 +54,9 @@ public:
 	virtual void StartupModule() override;
 	virtual void ShutdownModule() override;
 
+	void RegisterAnimationBakeHelper(IMovieSceneToolsAnimationBakeHelper* BakeHelper);
+	void UnregisterAnimationBakeHelper(IMovieSceneToolsAnimationBakeHelper* BakeHelper);
+	const TArray<IMovieSceneToolsAnimationBakeHelper*>& GetAnimationBakeHelpers() { return BakeHelpers; }
 	void RegisterTakeData(IMovieSceneToolsTakeData*);
 	void UnregisterTakeData(IMovieSceneToolsTakeData*);
 
@@ -111,4 +124,6 @@ private:
 
 	TArray<IMovieSceneToolsTakeData*> TakeDatas;
 	TArray<IMovieSceneToolsTrackImporter*> TrackImporters;
+
+	TArray<IMovieSceneToolsAnimationBakeHelper*> BakeHelpers;
 };
