Make sure the sequencer code always try to load bindings in order to properly resolve unloaded objects.
Also converted lots of code to go through LoadSequence instead of manually calling TryLoad.

diff --git a/Engine/Plugins/MovieScene/MatineeToLevelSequence/Source/MatineeToLevelSequence/Private/MatineeToLevelSequenceConverter.cpp b/Engine/Plugins/MovieScene/MatineeToLevelSequence/Source/MatineeToLevelSequence/Private/MatineeToLevelSequenceConverter.cpp
--- a/Engine/Plugins/MovieScene/MatineeToLevelSequence/Source/MatineeToLevelSequence/Private/MatineeToLevelSequenceConverter.cpp
+++ b/Engine/Plugins/MovieScene/MatineeToLevelSequence/Source/MatineeToLevelSequence/Private/MatineeToLevelSequenceConverter.cpp
@@ -78,7 +78,7 @@ void FMatineeToLevelSequenceConverter::ConvertMatineeToLevelSequence(TArray<TWea
 	// Edit the first asset
 	if (NewActors.Num())
 	{
-		UObject* NewAsset = NewActors[0]->LevelSequence.TryLoad();
+		UObject* NewAsset = NewActors[0]->LoadSequence();
 		if (NewAsset)
 		{
 			GEditor->GetEditorSubsystem<UAssetEditorSubsystem>()->OpenEditorForAsset(NewAsset);
diff --git a/Engine/Source/Editor/DetailCustomizations/Private/LevelSequenceActorDetails.cpp b/Engine/Source/Editor/DetailCustomizations/Private/LevelSequenceActorDetails.cpp
--- a/Engine/Source/Editor/DetailCustomizations/Private/LevelSequenceActorDetails.cpp
+++ b/Engine/Source/Editor/DetailCustomizations/Private/LevelSequenceActorDetails.cpp
@@ -150,7 +150,7 @@ FReply FLevelSequenceActorDetails::OnOpenLevelSequenceForActor()
 {
 	if( LevelSequenceActor.IsValid() )
 	{
-		UObject* LoadedObject = LevelSequenceActor.Get()->LevelSequence.TryLoad();
+		UObject* LoadedObject = LevelSequenceActor.Get()->LoadSequence();
 		if (LoadedObject != nullptr)
 		{
 			GEditor->GetEditorSubsystem<UAssetEditorSubsystem>()->OpenEditorForAsset(LoadedObject);
diff --git a/Engine/Source/Editor/LevelEditor/Private/LevelEditorToolBar.cpp b/Engine/Source/Editor/LevelEditor/Private/LevelEditorToolBar.cpp
--- a/Engine/Source/Editor/LevelEditor/Private/LevelEditorToolBar.cpp
+++ b/Engine/Source/Editor/LevelEditor/Private/LevelEditorToolBar.cpp
@@ -2026,7 +2026,7 @@ void FLevelEditorToolBar::OnCinematicsActorPicked( AActor* Actor )
 		FScopedSlowTask SlowTask(1.f, NSLOCTEXT("LevelToolBarCinematicsMenu", "LoadSequenceSlowTask", "Loading Level Sequence..."));
 		SlowTask.MakeDialog();
 		SlowTask.EnterProgressFrame();
-		UObject* Asset = LevelSequenceActor->LevelSequence.TryLoad();
+		UObject* Asset = LevelSequenceActor->LoadSequence();
 
 		if (Asset != nullptr)
 		{
diff --git a/Engine/Source/Editor/MovieSceneTools/Private/AutomatedLevelSequenceCapture.cpp b/Engine/Source/Editor/MovieSceneTools/Private/AutomatedLevelSequenceCapture.cpp
--- a/Engine/Source/Editor/MovieSceneTools/Private/AutomatedLevelSequenceCapture.cpp
+++ b/Engine/Source/Editor/MovieSceneTools/Private/AutomatedLevelSequenceCapture.cpp
@@ -82,7 +82,7 @@ UMovieScene* GetMovieScene(TWeakObjectPtr<ALevelSequenceActor> LevelSequenceActo
 		return nullptr;
 	}
 
-	ULevelSequence* LevelSequence = Cast<ULevelSequence>(Actor->LevelSequence.TryLoad());
+	ULevelSequence* LevelSequence = Cast<ULevelSequence>(Actor->LoadSequence());
 	if (!LevelSequence)
 	{
 		return nullptr;
@@ -539,7 +539,7 @@ void UAutomatedLevelSequenceCapture::SetupFrameRange()
 	ALevelSequenceActor* Actor = LevelSequenceActor.Get();
 	if( Actor )
 	{
-		ULevelSequence* LevelSequence = Cast<ULevelSequence>( Actor->LevelSequence.TryLoad() );
+		ULevelSequence* LevelSequence = Actor->LoadSequence();
 		if( LevelSequence != nullptr )
 		{
 			UMovieScene* MovieScene = LevelSequence->GetMovieScene();
diff --git a/Engine/Source/Editor/UnrealEd/Private/Factories/ActorFactory.cpp b/Engine/Source/Editor/UnrealEd/Private/Factories/ActorFactory.cpp
--- a/Engine/Source/Editor/UnrealEd/Private/Factories/ActorFactory.cpp
+++ b/Engine/Source/Editor/UnrealEd/Private/Factories/ActorFactory.cpp
@@ -2059,7 +2059,7 @@ UObject* UActorFactoryMovieScene::GetAssetFromActorInstance(AActor* Instance)
 {
 	if (ALevelSequenceActor* LevelSequenceActor = Cast<ALevelSequenceActor>(Instance))
 	{
-		return LevelSequenceActor->LevelSequence.TryLoad();
+		return LevelSequenceActor->LoadSequence();
 	}
 
 	return nullptr;
diff --git a/Engine/Source/Runtime/LevelSequence/Private/LevelSequenceBindingReference.cpp b/Engine/Source/Runtime/LevelSequence/Private/LevelSequenceBindingReference.cpp
--- a/Engine/Source/Runtime/LevelSequence/Private/LevelSequenceBindingReference.cpp
+++ b/Engine/Source/Runtime/LevelSequence/Private/LevelSequenceBindingReference.cpp
@@ -79,7 +79,7 @@ UObject* FLevelSequenceBindingReference::Resolve(UObject* InContext, FName Strea
 		}
 	#endif
 
-		return TempPath.ResolveObject();
+		return TempPath.TryLoad();
 	}
 
 	return nullptr;
