Sequencer: Added missing Track->Modify()

diff --git a/Engine/Source/Editor/MovieSceneTools/Private/MovieSceneToolHelpers.cpp b/Engine/Source/Editor/MovieSceneTools/Private/MovieSceneToolHelpers.cpp
--- a/Engine/Source/Editor/MovieSceneTools/Private/MovieSceneToolHelpers.cpp
+++ b/Engine/Source/Editor/MovieSceneTools/Private/MovieSceneToolHelpers.cpp
@@ -131,6 +131,8 @@ void MovieSceneToolHelpers::TrimSection(const TSet<TWeakObjectPtr<UMovieSceneSec
 
 void MovieSceneToolHelpers::TrimOrExtendSection(UMovieSceneTrack* Track, TOptional<int32> SpecifiedRowIndex, FQualifiedFrameTime Time, bool bTrimOrExtendLeft, bool bDeleteKeys)
 {
+	Track->Modify();
+
 	int32 StartRowIndex = SpecifiedRowIndex.IsSet() ? SpecifiedRowIndex.GetValue() : 0;
 	int32 EndRowIndex = SpecifiedRowIndex.IsSet() ? SpecifiedRowIndex.GetValue() : Track->GetMaxRowIndex();
 
