Non-Unity fix for same function name.
#trivial

diff --git a/Engine/Source/Editor/MovieSceneTools/Private/EditModes/SkeletalAnimationTrackEditMode.cpp b/Engine/Source/Editor/MovieSceneTools/Private/EditModes/SkeletalAnimationTrackEditMode.cpp
--- a/Engine/Source/Editor/MovieSceneTools/Private/EditModes/SkeletalAnimationTrackEditMode.cpp
+++ b/Engine/Source/Editor/MovieSceneTools/Private/EditModes/SkeletalAnimationTrackEditMode.cpp
@@ -125,28 +125,6 @@ void FSkeletalAnimationTrackEditMode::Tick(FEditorViewportClient* ViewportClient
 	ViewportClient->Invalidate();
 }
 
-static USkeletalMeshComponent* AcquireSkeletalMeshFromObject(UObject* BoundObject)
-{
-	if (AActor* Actor = Cast<AActor>(BoundObject))
-	{
-		for (UActorComponent* Component : Actor->GetComponents())
-		{
-			if (USkeletalMeshComponent* SkeletalMeshComp = Cast<USkeletalMeshComponent>(Component))
-			{
-				return SkeletalMeshComp;
-			}
-		}
-	}
-	else if (USkeletalMeshComponent* SkeletalMeshComponent = Cast<USkeletalMeshComponent>(BoundObject))
-	{
-		if (SkeletalMeshComponent->SkeletalMesh)
-		{
-			return SkeletalMeshComponent;
-		}
-	}
-	return nullptr;
-}
-
 static void DrawBones(const TArray<FBoneIndexType>& RequiredBones, const FReferenceSkeleton& RefSkeleton, const TArray<FTransform>& WorldTransforms,
 	FSkeletalMeshLODRenderData* LODData,
 	FPrimitiveDrawInterface* PDI, const TArray<FLinearColor>& BoneColours, float BoundRadius, float LineThickness/*=0.f*/)
@@ -359,7 +337,7 @@ void FSkeletalAnimationTrackEditMode::Render(const FSceneView* View, FViewport*
 								}
 							}
 							//show skeletons
-							USkeletalMeshComponent* SkelMeshComp = AcquireSkeletalMeshFromObject(BoundObject);
+							USkeletalMeshComponent* SkelMeshComp = MovieSceneToolHelpers::AcquireSkeletalMeshFromObject(BoundObject);
 							if (SkelMeshComp && SkelMeshComp->GetAnimInstance())
 							{
 								const FLinearColor Colors[5] = { FLinearColor(1.0f,0.0f,1.0f,1.0f), FLinearColor(0.0f,1.0f,0.0f,1.0f), FLinearColor(0.0f, 0.0f, 1.0f, 0.0f),
diff --git a/Engine/Source/Editor/MovieSceneTools/Private/MovieSceneToolHelpers.cpp b/Engine/Source/Editor/MovieSceneTools/Private/MovieSceneToolHelpers.cpp
--- a/Engine/Source/Editor/MovieSceneTools/Private/MovieSceneToolHelpers.cpp
+++ b/Engine/Source/Editor/MovieSceneTools/Private/MovieSceneToolHelpers.cpp
@@ -3924,7 +3924,7 @@ void MovieSceneToolHelpers::GetLocationAtTime(const FMovieSceneEvaluationTrack*
 }
 
 
-static USkeletalMeshComponent* AcquireSkeletalMeshFromObject(UObject* BoundObject)
+USkeletalMeshComponent* MovieSceneToolHelpers::AcquireSkeletalMeshFromObject(UObject* BoundObject)
 {
 	if (AActor* Actor = Cast<AActor>(BoundObject))
 	{
@@ -3951,7 +3951,7 @@ static void GetSequencerActorWorldTransforms(ISequencer* Sequencer, const FActor
 	if (AActor* Actor = ActorSelection.Actor.Get())
 	{
 
-		USkeletalMeshComponent* SkelMeshComp = AcquireSkeletalMeshFromObject(Actor);
+		USkeletalMeshComponent* SkelMeshComp = MovieSceneToolHelpers::AcquireSkeletalMeshFromObject(Actor);
 
 		if (ActorSelection.SocketName != NAME_None && SkelMeshComp)
 		{
diff --git a/Engine/Source/Editor/MovieSceneTools/Public/MovieSceneToolHelpers.h b/Engine/Source/Editor/MovieSceneTools/Public/MovieSceneToolHelpers.h
--- a/Engine/Source/Editor/MovieSceneTools/Public/MovieSceneToolHelpers.h
+++ b/Engine/Source/Editor/MovieSceneTools/Public/MovieSceneToolHelpers.h
@@ -502,6 +502,12 @@ public:
 	static bool ImportFBXIntoControlRigChannels(UMovieScene* MovieScene, const FString& ImportFilename,  UMovieSceneUserImportFBXControlRigSettings *ControlRigSettings,
 		TArray<FFBXNodeAndChannels>* NodeAndChannels, const TArray<FName>& SelectedControlNames, FFrameRate FrameRate);
 
+	/*
+	* Acquire first SkeletalMeshComponent from the Object
+	* @param BoundObject Object to get SkeletalMeshComponent from.If actor checks it's components, if component checks itself then child components.
+	* @return Returns the USkeletalMeshComponent if one is found
+	*/
+	static USkeletalMeshComponent* AcquireSkeletalMeshFromObject(UObject* BoundObject);
 	/*
 	*  Get an actors word transforms at the specified times
 	* @param Sequencer Sequencer to evaluation
