Sequencer: Remove UE4 -> UE

diff --git a/Engine/Source/Editor/MovieSceneTools/Private/FCPXML/FCPXMLExport.cpp b/Engine/Source/Editor/MovieSceneTools/Private/FCPXML/FCPXMLExport.cpp
--- a/Engine/Source/Editor/MovieSceneTools/Private/FCPXML/FCPXMLExport.cpp
+++ b/Engine/Source/Editor/MovieSceneTools/Private/FCPXML/FCPXMLExport.cpp
@@ -1304,19 +1304,19 @@ bool FFCPXMLExportVisitor::CreateCinematicSectionMetadata(const UMovieSceneCinem
 		return false;
 	}
 
-	OutMetadata = TEXT("[UE4ShotSection=") + InSection->GetPathName() + TEXT("]");
+	OutMetadata = TEXT("[UEShotSection=") + InSection->GetPathName() + TEXT("]");
 
 	// Store the start offset and the handle frames for round-tripping to compute the new start offset
 	int32 HandleFrames = ExportData->GetHandleFrames();
 	FFrameRate TickResolution = InSection->GetTypedOuter<UMovieScene>()->GetTickResolution();
 	int32 StartFrameOffset = ConvertFrameTime(InSection->Parameters.StartFrameOffset, TickResolution, ExportData->GetFrameRate()).CeilToFrame().Value;
 
-	OutMetadata += TEXT("[UE4ShotStartOffset=") + FString::FromInt(StartFrameOffset) + TEXT("]");
-	OutMetadata += TEXT("[UE4ShotHandleFrames=") + FString::FromInt(HandleFrames) + TEXT("]");
+	OutMetadata += TEXT("[UEShotStartOffset=") + FString::FromInt(StartFrameOffset) + TEXT("]");
+	OutMetadata += TEXT("[UEShotHandleFrames=") + FString::FromInt(HandleFrames) + TEXT("]");
 	return true;
 }
 
-/** Get metadata section name from sequencer shot name - format is "[UE4SoundWave=soundwaveobjectname][UE4SoundSectionTopLevel=toplevelobjectname][UE4SoundSection=sectionobjectname]", whitespace ok. */
+/** Get metadata section name from sequencer shot name - format is "[UESoundWave=soundwaveobjectname][UESoundSectionTopLevel=toplevelobjectname][UESoundSection=sectionobjectname]", whitespace ok. */
 bool FFCPXMLExportVisitor::CreateSoundWaveMetadata(const USoundWave* InSoundWave, const TArray<const UMovieSceneAudioSection*> InAudioSections, FString& OutMetadata) const
 {
 	if (InSoundWave == nullptr)
@@ -1326,12 +1326,12 @@ bool FFCPXMLExportVisitor::CreateSoundWaveMetadata(const USoundWave* InSoundWave
 
 	TArray<FString> SectionsAdded;
 	bool bTopLevelAdded = false;
-	OutMetadata = TEXT("[UE4SoundWave=") + InSoundWave->GetPathName() + TEXT("]");
+	OutMetadata = TEXT("[UESoundWave=") + InSoundWave->GetPathName() + TEXT("]");
 	for (const UMovieSceneAudioSection* AudioSection : InAudioSections)
 	{
 		if (!bTopLevelAdded)
 		{
-			OutMetadata += TEXT("[UE4AudioSectionTopLevel=") + FFCPXMLExportVisitor::GetAudioSectionTopLevelName(AudioSection) + TEXT("]");
+			OutMetadata += TEXT("[UEAudioSectionTopLevel=") + FFCPXMLExportVisitor::GetAudioSectionTopLevelName(AudioSection) + TEXT("]");
 			bTopLevelAdded = true;
 		}
 
@@ -1339,7 +1339,7 @@ bool FFCPXMLExportVisitor::CreateSoundWaveMetadata(const USoundWave* InSoundWave
 		FString SectionName = FFCPXMLExportVisitor::GetAudioSectionName(AudioSection);
 		if (SectionsAdded.Num() == 0 || !SectionsAdded.Contains(SectionName))
 		{
-			OutMetadata += TEXT("[UE4AudioSection=") + SectionName + TEXT("]");
+			OutMetadata += TEXT("[UEAudioSection=") + SectionName + TEXT("]");
 			SectionsAdded.Add(SectionName);
 		}
 	}
diff --git a/Engine/Source/Editor/MovieSceneTools/Private/FCPXML/FCPXMLExport.h b/Engine/Source/Editor/MovieSceneTools/Private/FCPXML/FCPXMLExport.h
--- a/Engine/Source/Editor/MovieSceneTools/Private/FCPXML/FCPXMLExport.h
+++ b/Engine/Source/Editor/MovieSceneTools/Private/FCPXML/FCPXMLExport.h
@@ -118,9 +118,9 @@ protected:
 	/** Returns true if audio track data contains sections with 2 channels */
 	bool HasStereoAudioSections(const  TArray<TSharedPtr<FMovieSceneExportAudioSectionData>>& InAudioTrackData) const;
 
-	/** Get metadata section name from sequencer shot name - format is "[UE4Section=sectionobjectname]", whitespace ok. */
+	/** Get metadata section name from sequencer shot name - format is "[UESection=sectionobjectname]", whitespace ok. */
 	bool CreateCinematicSectionMetadata(const UMovieSceneCinematicShotSection* InSection, FString& OutMetadata) const;
-	/** Get metadata section name from sequencer shot name - format is "[UE4SoundWave=soundwaveobjectname]", whitespace ok. */
+	/** Get metadata section name from sequencer shot name - format is "[UESoundWave=soundwaveobjectname]", whitespace ok. */
 	bool CreateSoundWaveMetadata(const USoundWave* InSoundWave, const TArray<const UMovieSceneAudioSection*> InAudioSections, FString& OutMetadata) const;
 	/** Get id for audio top level section */
 	static FString GetAudioSectionTopLevelName(const UMovieSceneAudioSection* InAudioSection);
diff --git a/Engine/Source/Editor/MovieSceneTools/Private/FCPXML/FCPXMLImport.cpp b/Engine/Source/Editor/MovieSceneTools/Private/FCPXML/FCPXMLImport.cpp
--- a/Engine/Source/Editor/MovieSceneTools/Private/FCPXML/FCPXMLImport.cpp
+++ b/Engine/Source/Editor/MovieSceneTools/Private/FCPXML/FCPXMLImport.cpp
@@ -802,17 +802,17 @@ bool FFCPXMLImportVisitor::ParseMetadata(const FString& InMetadata, const FStrin
 	return ParseMetadata(InMetadata, InKey, OutValue, MetadataRemaining);
 }
 
-/** Get sequencer section name from section metadata. Format is "[UE4ShotSection=sectionobjectname]", whitespace ok. */
+/** Get sequencer section name from section metadata. Format is "[UEShotSection=sectionobjectname]", whitespace ok. */
 bool FFCPXMLImportVisitor::GetCinematicSectionPathNameFromMetadata(const FString& InMetadata, FString& OutSectionObjectName) const
 {
-	return ParseMetadata(InMetadata, TEXT("UE4ShotSection"), OutSectionObjectName);
+	return ParseMetadata(InMetadata, TEXT("UEShotSection"), OutSectionObjectName);
 }
 
-/** Get sequencer shot handle frames from section metadata. Format is "[UE4ShotHandleFrames=handleframes]", whitespace ok. */
+/** Get sequencer shot handle frames from section metadata. Format is "[UEShotHandleFrames=handleframes]", whitespace ok. */
 bool FFCPXMLImportVisitor::GetCinematicSectionHandleFramesFromMetadata(const FString& InMetadata, int32& OutHandleFrames) const
 {
 	FString HandleFrameData;
-	bool bSuccess = ParseMetadata(InMetadata, TEXT("UE4ShotHandleFrames"), HandleFrameData);
+	bool bSuccess = ParseMetadata(InMetadata, TEXT("UEShotHandleFrames"), HandleFrameData);
 	if (bSuccess)
 	{
 		OutHandleFrames = FCString::Atoi(*HandleFrameData);
@@ -820,11 +820,11 @@ bool FFCPXMLImportVisitor::GetCinematicSectionHandleFramesFromMetadata(const FSt
 	return bSuccess;
 }
 
-/** Get sequencer shot start offset frame from section metadata. Format is "[UE4ShotStartOffset=startoffset]", whitespace ok. */
+/** Get sequencer shot start offset frame from section metadata. Format is "[UEShotStartOffset=startoffset]", whitespace ok. */
 bool FFCPXMLImportVisitor::GetCinematicSectionStartOffsetFromMetadata(const FString& InMetadata, int32& OutStartOffset) const
 {
 	FString StartOffsetData;
-	bool bSuccess = ParseMetadata(InMetadata, TEXT("UE4ShotStartOffset"), StartOffsetData);
+	bool bSuccess = ParseMetadata(InMetadata, TEXT("UEShotStartOffset"), StartOffsetData);
 	if (bSuccess)
 	{
 		OutStartOffset = FCString::Atoi(*StartOffsetData);
@@ -832,24 +832,24 @@ bool FFCPXMLImportVisitor::GetCinematicSectionStartOffsetFromMetadata(const FStr
 	return bSuccess;
 }
 
-/** Get sequencer track name from track metadata. Format is "[UE4Track=trackobjectname][UE4Row=rowindex]", whitespace ok. */
+/** Get sequencer track name from track metadata. Format is "[UETrack=trackobjectname][UERow=rowindex]", whitespace ok. */
 bool FFCPXMLImportVisitor::GetAudioFromMetadata(const FString& InMetadata, TSharedPtr<FFCPXMLImportAudioMetadata>& OutAudioMetadata) const
 {
 	FString Metadata1{ TEXT("") };
 	FString SoundWavePathName;
-	bool bSuccess = ParseMetadata(InMetadata, TEXT("UE4SoundWave"), SoundWavePathName, Metadata1);
+	bool bSuccess = ParseMetadata(InMetadata, TEXT("UESoundWave"), SoundWavePathName, Metadata1);
 	OutAudioMetadata = MakeShared<FFCPXMLImportAudioMetadata>(SoundWavePathName);
 
 	if (bSuccess)
 	{
 		FString Metadata2{ TEXT("") };
 		FString AudioSectionTopLevel{ TEXT("") };
-		bSuccess = ParseMetadata(Metadata1, TEXT("UE4AudioSectionTopLevel"), AudioSectionTopLevel, Metadata2);
+		bSuccess = ParseMetadata(Metadata1, TEXT("UEAudioSectionTopLevel"), AudioSectionTopLevel, Metadata2);
 
 		if (bSuccess)
 		{
 			FString AudioSection{ TEXT("") };
-			while (ParseMetadata(Metadata2, TEXT("UE4AudioSection"), AudioSection, Metadata1))
+			while (ParseMetadata(Metadata2, TEXT("UEAudioSection"), AudioSection, Metadata1))
 			{
 				FString AudioSectionPathName = AudioSectionTopLevel + TEXT(".") + AudioSection;
 				TSharedPtr<FFCPXMLImportAudioSectionMetadata> AudioSectionMetadata = MakeShared<FFCPXMLImportAudioSectionMetadata>(AudioSectionPathName);
diff --git a/Engine/Source/Editor/MovieSceneTools/Private/FCPXML/FCPXMLImport.h b/Engine/Source/Editor/MovieSceneTools/Private/FCPXML/FCPXMLImport.h
--- a/Engine/Source/Editor/MovieSceneTools/Private/FCPXML/FCPXMLImport.h
+++ b/Engine/Source/Editor/MovieSceneTools/Private/FCPXML/FCPXMLImport.h
@@ -108,13 +108,13 @@ private:
 	bool ParseMetadata(const FString& InMetadata, const FString& InKey, FString& OutValue, FString& OutMetadata) const;
 	/** parse metadata of the format "[tag=value]", whitespace ok. */
 	bool ParseMetadata(const FString& InMetadata, const FString& InKey, FString& OutValue) const;
-	/** Get sequencer section id from section metadata. Format is "[UE4ShotSection=sectionobjectname]", whitespace ok. */
+	/** Get sequencer section id from section metadata. Format is "[UEShotSection=sectionobjectname]", whitespace ok. */
 	bool GetCinematicSectionPathNameFromMetadata(const FString& InMetadata, FString& OutSectionObjectName) const;
-	/** Get sequencer shot handle frames from section metadata. Format is "[UE4ShotHandleFrames=handleframes]", whitespace ok. */
+	/** Get sequencer shot handle frames from section metadata. Format is "[UEShotHandleFrames=handleframes]", whitespace ok. */
 	bool GetCinematicSectionHandleFramesFromMetadata(const FString& InMetadata, int32& OutHandleFrames) const;
-	/** Get sequencer shot start offset frame from section metadata. Format is "[UE4ShotStartOffset=startoffset]", whitespace ok. */
+	/** Get sequencer shot start offset frame from section metadata. Format is "[UEShotStartOffset=startoffset]", whitespace ok. */
 	bool GetCinematicSectionStartOffsetFromMetadata(const FString& InMetadata, int32& OutStartOffset) const;
-	/** Get sequencer sound wave id and audio section ids from metadata. Format is "[UE4SoundWave=trackobjectname][UE4AudioSectionTopLevel=toplevelobjectname][UE4AudioSection=audiosectionobjectname]", whitespace ok. */
+	/** Get sequencer sound wave id and audio section ids from metadata. Format is "[UESoundWave=trackobjectname][UEAudioSectionTopLevel=toplevelobjectname][UEAudioSection=audiosectionobjectname]", whitespace ok. */
 	bool GetAudioFromMetadata(const FString& InMetadata, TSharedPtr<FFCPXMLImportAudioMetadata>& OutAudioMetadata) const;
 	
 public:
diff --git a/Engine/Source/Editor/MovieSceneTools/Private/MovieSceneEventCustomization.cpp b/Engine/Source/Editor/MovieSceneTools/Private/MovieSceneEventCustomization.cpp
--- a/Engine/Source/Editor/MovieSceneTools/Private/MovieSceneEventCustomization.cpp
+++ b/Engine/Source/Editor/MovieSceneTools/Private/MovieSceneEventCustomization.cpp
@@ -51,7 +51,7 @@
 #define LOCTEXT_NAMESPACE "MovieSceneEventCustomization"
 
 
-namespace UE4_MovieSceneEventCustomization
+namespace UE_MovieSceneEventCustomization
 {
 	void CollectQuickBindActions(FGraphActionListBuilderBase& OutAllActions, UBlueprint* Blueprint, UClass* BoundObjectPinClass)
 	{
@@ -917,7 +917,7 @@ void FMovieSceneEventCustomization::PopulateQuickBindSubMenu(FMenuBuilder& MenuB
 
 	TSharedRef<SGraphActionMenu> ActionMenu = SNew(SGraphActionMenu)
 		.OnCreateCustomRowExpander_Static([](const FCustomExpanderData& Data) -> TSharedRef<SExpanderArrow> { return SNew(SExpanderArrow, Data.TableRow); })
-		.OnCollectAllActions_Static(UE4_MovieSceneEventCustomization::CollectQuickBindActions, Blueprint, EventParams.BoundObjectPinClass)
+		.OnCollectAllActions_Static(UE_MovieSceneEventCustomization::CollectQuickBindActions, Blueprint, EventParams.BoundObjectPinClass)
 		.OnActionSelected(this, &FMovieSceneEventCustomization::HandleQuickBindActionSelected, Blueprint, EventParams);
 
 	ActionMenu->RegisterActiveTimer(0.f, FWidgetActiveTimerDelegate::CreateLambda(
@@ -1061,7 +1061,7 @@ void FMovieSceneEventCustomization::PopulateRebindSubMenu(FMenuBuilder& MenuBuil
 
 	TSharedRef<SGraphActionMenu> ActionMenu = SNew(SGraphActionMenu)
 		.OnCreateCustomRowExpander_Static([](const FCustomExpanderData& Data) -> TSharedRef<SExpanderArrow> { return SNew(SExpanderArrow, Data.TableRow); })
-		.OnCollectAllActions_Static(UE4_MovieSceneEventCustomization::CollectAllRebindActions, Blueprint)
+		.OnCollectAllActions_Static(UE_MovieSceneEventCustomization::CollectAllRebindActions, Blueprint)
 		.OnActionSelected(this, &FMovieSceneEventCustomization::HandleRebindActionSelected, Blueprint, EventParams.BoundObjectPinClass);
 
 	ActionMenu->RegisterActiveTimer(0.f, FWidgetActiveTimerDelegate::CreateLambda(
diff --git a/Engine/Source/Editor/MovieSceneTools/Private/MovieSceneTranslator.cpp b/Engine/Source/Editor/MovieSceneTools/Private/MovieSceneTranslator.cpp
--- a/Engine/Source/Editor/MovieSceneTools/Private/MovieSceneTranslator.cpp
+++ b/Engine/Source/Editor/MovieSceneTools/Private/MovieSceneTranslator.cpp
@@ -54,9 +54,9 @@ FMovieSceneExportData::FMovieSceneExportData(const UMovieScene* InMovieScene, FF
 	SaveFilename = InSaveFilename;
 	MovieExtension = InMovieExtension;
 
-	// preferred sample rate in UE4
+	// preferred sample rate in UE
 	DefaultAudioSampleRate = 44100;
-	// all audio in UE4 is has depth 16
+	// all audio in UE is has depth 16
 	DefaultAudioDepth = 16;
 
 	bExportDataIsValid = ConstructData(InMovieScene);
diff --git a/Engine/Source/Editor/MovieSceneTools/Public/MovieSceneToolHelpers.h b/Engine/Source/Editor/MovieSceneTools/Public/MovieSceneToolHelpers.h
--- a/Engine/Source/Editor/MovieSceneTools/Public/MovieSceneToolHelpers.h
+++ b/Engine/Source/Editor/MovieSceneTools/Public/MovieSceneToolHelpers.h
@@ -364,7 +364,7 @@ public:
 	* Import FBX Camera to existing camera's
 	*
 	* @param CameraNode The Fbx camera
-	* @param InCameraActor Ue4 actor
+	* @param InCameraActor UE actor
 	*/
 	static void CopyCameraProperties(fbxsdk::FbxCamera* CameraNode, AActor* InCameraActor);
 
diff --git a/Engine/Source/Editor/MovieSceneTools/Public/MovieSceneToolsUserSettings.h b/Engine/Source/Editor/MovieSceneTools/Public/MovieSceneToolsUserSettings.h
--- a/Engine/Source/Editor/MovieSceneTools/Public/MovieSceneToolsUserSettings.h
+++ b/Engine/Source/Editor/MovieSceneTools/Public/MovieSceneToolsUserSettings.h
@@ -64,7 +64,7 @@ public:
 	bool bMatchByNameOnly;
 
 	/** Whether to force the front axis to be align with X instead of -Y. */
-	UPROPERTY(EditAnywhere, config, Category=Import, meta= (ToolTip = "Convert the scene from FBX coordinate system to UE4 coordinate system with front X axis instead of -Y"))
+	UPROPERTY(EditAnywhere, config, Category=Import, meta= (ToolTip = "Convert the scene from FBX coordinate system to UE coordinate system with front X axis instead of -Y"))
 	bool bForceFrontXAxis;
 
 	/** Whether to create cameras if they don't already exist in the level. */
