Sequencer: Reomve UE4 -> UE

diff --git a/Engine/Source/Editor/MovieSceneTools/Public/MovieSceneToolsUserSettings.h b/Engine/Source/Editor/MovieSceneTools/Public/MovieSceneToolsUserSettings.h
--- a/Engine/Source/Editor/MovieSceneTools/Public/MovieSceneToolsUserSettings.h
+++ b/Engine/Source/Editor/MovieSceneTools/Public/MovieSceneToolsUserSettings.h
@@ -202,7 +202,7 @@ public:
 	TArray< FControlFindReplaceString> FindAndReplaceStrings;
 
 	/** Whether to force the front axis to be align with X instead of -Y. */
-	UPROPERTY(EditAnywhere, config, Category = "Import Options", meta = (ToolTip = "Convert the scene from FBX coordinate system to UE4 coordinate system with front X axis instead of -Y"))
+	UPROPERTY(EditAnywhere, config, Category = "Import Options", meta = (ToolTip = "Convert the scene from FBX coordinate system to UE coordinate system with front X axis instead of -Y"))
 	bool bForceFrontXAxis;
 
 	/** Whether or not import onto selected controls or all controls*/
