UAutomatedLevelSequenceCapture: Replace FlushLevelStreaming call with BlockTillLevelStreamingCompleted which waits properly on streaming to complete.

diff --git a/Engine/Source/Editor/MovieSceneTools/Private/AutomatedLevelSequenceCapture.cpp b/Engine/Source/Editor/MovieSceneTools/Private/AutomatedLevelSequenceCapture.cpp
--- a/Engine/Source/Editor/MovieSceneTools/Private/AutomatedLevelSequenceCapture.cpp
+++ b/Engine/Source/Editor/MovieSceneTools/Private/AutomatedLevelSequenceCapture.cpp
@@ -645,15 +645,14 @@ void UAutomatedLevelSequenceCapture::OnTick(float DeltaSeconds)
 		return;
 	}
 
-	// Flush the level streaming system. This would cause hitches under normal gameplay, but because we already run at slower-than-real-time
+	// Block till level streaming is completed. This would cause hitches under normal gameplay, but because we already run at slower-than-real-time
 	// it doesn't matter for movie captures. This solves situations where games have systems that pull in sublevels via level streaming that cannot
-	// be normally controlled via the Sequencer Level Visibility track. If all levels are already loaded, flushing will have no effect.
-	if (GetWorld())
+	// be normally controlled via the Sequencer Level Visibility track. If all levels are already loaded, blocking will have no effect.
+	if (Actor->GetWorld())
 	{
-		GetWorld()->FlushLevelStreaming(EFlushLevelStreamingType::Full);
+		Actor->GetWorld()->BlockTillLevelStreamingCompleted();
 	}
 
-
 	if (GShaderCompilingManager && GShaderCompilingManager->GetNumRemainingJobs() > 0)
 	{
 		UE_LOG(LogMovieSceneCapture, Log, TEXT("[%d] Waiting for %d shaders to finish compiling..."), GFrameCounter, GShaderCompilingManager->GetNumRemainingJobs());
