Sequencer: Control Rig: When saving linked anim sequence also save link back to level sequence, which allows copy edit workflows better.
#trivial

diff --git a/Engine/Source/Editor/MovieSceneTools/Private/TrackEditors/SkeletalAnimationTrackEditor.cpp b/Engine/Source/Editor/MovieSceneTools/Private/TrackEditors/SkeletalAnimationTrackEditor.cpp
--- a/Engine/Source/Editor/MovieSceneTools/Private/TrackEditors/SkeletalAnimationTrackEditor.cpp
+++ b/Engine/Source/Editor/MovieSceneTools/Private/TrackEditors/SkeletalAnimationTrackEditor.cpp
@@ -1123,6 +1123,17 @@ void FSkeletalAnimationTrackEditor::OnSequencerSaved(ISequencer& )
 				for (FLevelSequenceAnimSequenceLinkItem& Item : LevelAnimLink->AnimSequenceLinks)
 				{
 					UAnimSequence* AnimSequence = Item.ResolveAnimSequence();
+					if (IInterface_AssetUserData* AnimAssetUserData = Cast< IInterface_AssetUserData >(AnimSequence))
+					{
+						UAnimSequenceLevelSequenceLink* AnimLevelLink = AnimAssetUserData->GetAssetUserData< UAnimSequenceLevelSequenceLink >();
+						if (!AnimLevelLink)
+						{
+							AnimLevelLink = NewObject<UAnimSequenceLevelSequenceLink>(AnimSequence, NAME_None, RF_Public | RF_Transactional);
+							AnimAssetUserData->AddAssetUserData(AnimLevelLink);
+						}
+						AnimLevelLink->SetLevelSequence(LevelSequence);
+						AnimLevelLink->SkelTrackGuid = Item.SkelTrackGuid;
+					}
 					USkeletalMeshComponent* SkelMeshComp = AcquireSkeletalMeshFromObjectGuid(Item.SkelTrackGuid, GetSequencer());
 					if (AnimSequence && SkelMeshComp)
 					{
