Fix Automation Test: Sequencer Baking

diff --git a/Engine/Source/Editor/MovieSceneTools/Private/MovieSceneToolHelpers.cpp b/Engine/Source/Editor/MovieSceneTools/Private/MovieSceneToolHelpers.cpp
--- a/Engine/Source/Editor/MovieSceneTools/Private/MovieSceneToolHelpers.cpp
+++ b/Engine/Source/Editor/MovieSceneTools/Private/MovieSceneToolHelpers.cpp
@@ -4038,7 +4038,10 @@ static void GetSequencerActorWorldTransforms(ISequencer* Sequencer, const FActor
 					}
 
 					FMovieSceneContext Context = FMovieSceneContext(FMovieSceneEvaluationRange(GlobalTime, TickResolution), Sequencer->GetPlaybackStatus()).SetHasJumped(true);
-
+					if (Index == 0) // similar with baking first time in we need to evaluate twice (think due to double buffering that happens with skel mesh components).
+					{
+						Sequencer->GetEvaluationTemplate().Evaluate(Context, *Sequencer);
+					}
 					Sequencer->GetEvaluationTemplate().Evaluate(Context, *Sequencer);
 
 					for (IMovieSceneToolsAnimationBakeHelper* BakeHelper : BakeHelpers)
