Sequencer :Control Rig: Take into account start frame of imported animation when inserting/replacing.

diff --git a/Engine/Source/Editor/MovieSceneTools/Private/MovieSceneToolHelpers.cpp b/Engine/Source/Editor/MovieSceneTools/Private/MovieSceneToolHelpers.cpp
--- a/Engine/Source/Editor/MovieSceneTools/Private/MovieSceneToolHelpers.cpp
+++ b/Engine/Source/Editor/MovieSceneTools/Private/MovieSceneToolHelpers.cpp
@@ -1879,6 +1879,7 @@ static void PrepForInsertReplaceAnimation(bool bInsert, const FFBXNodeAndChannel
 	}
 
 	FFrameNumber Diff = EndFrame - StartFrame;
+	FrameToInsertOrReplace += StartFrame;
 	if (bInsert)
 	{
 		for (FMovieSceneChannel* Channel : Channels)
