Matinee/Editor: Removed Sequencer's ability to import Matinee tracks from the Matinee Copy Paste buffer which can no longer be filled as there is no Matinee UI anymore.

diff --git a/Engine/Source/Editor/MovieSceneTools/Private/TrackEditors/AudioTrackEditor.cpp b/Engine/Source/Editor/MovieSceneTools/Private/TrackEditors/AudioTrackEditor.cpp
--- a/Engine/Source/Editor/MovieSceneTools/Private/TrackEditors/AudioTrackEditor.cpp
+++ b/Engine/Source/Editor/MovieSceneTools/Private/TrackEditors/AudioTrackEditor.cpp
@@ -960,28 +960,6 @@ void CopyInterpSoundTrack(TSharedRef<ISequencer> Sequencer, UInterpTrackSound* M
 	}
 }
 
-void FAudioTrackEditor::BuildTrackContextMenu( FMenuBuilder& MenuBuilder, UMovieSceneTrack* Track )
-{
-	UInterpTrackSound* MatineeSoundTrack = nullptr;
-	for ( UObject* CopyPasteObject : GUnrealEd->MatineeCopyPasteBuffer )
-	{
-		MatineeSoundTrack = Cast<UInterpTrackSound>( CopyPasteObject );
-		if ( MatineeSoundTrack != nullptr )
-		{
-			break;
-		}
-	}
-	UMovieSceneAudioTrack* AudioTrack = Cast<UMovieSceneAudioTrack>( Track );
-	MenuBuilder.AddMenuEntry(
-		NSLOCTEXT( "Sequencer", "PasteMatineeSoundTrack", "Paste Matinee Sound Track" ),
-		NSLOCTEXT( "Sequencer", "PasteMatineeSoundTrackTooltip", "Pastes keys from a Matinee sound track into this track." ),
-		FSlateIcon(),
-		FUIAction(
-		FExecuteAction::CreateStatic( &CopyInterpSoundTrack, GetSequencer().ToSharedRef(), MatineeSoundTrack, AudioTrack ),
-		FCanExecuteAction::CreateLambda( [=]()->bool { return MatineeSoundTrack != nullptr && MatineeSoundTrack->Sounds.Num() > 0 && AudioTrack != nullptr; } ) ) );
-}
-
-
 const FSlateBrush* FAudioTrackEditor::GetIconBrush() const
 {
 	return FEditorStyle::GetBrush("Sequencer.Tracks.Audio");
diff --git a/Engine/Source/Editor/MovieSceneTools/Private/TrackEditors/ParticleTrackEditor.cpp b/Engine/Source/Editor/MovieSceneTools/Private/TrackEditors/ParticleTrackEditor.cpp
--- a/Engine/Source/Editor/MovieSceneTools/Private/TrackEditors/ParticleTrackEditor.cpp
+++ b/Engine/Source/Editor/MovieSceneTools/Private/TrackEditors/ParticleTrackEditor.cpp
@@ -305,26 +305,6 @@ void CopyInterpParticleTrack(TSharedRef<ISequencer> Sequencer, UInterpTrackToggl
 	}
 }
 
-void FParticleTrackEditor::BuildTrackContextMenu( FMenuBuilder& MenuBuilder, UMovieSceneTrack* Track )
-{
-	UInterpTrackToggle* MatineeToggleTrack = nullptr;
-	for ( UObject* CopyPasteObject : GUnrealEd->MatineeCopyPasteBuffer )
-	{
-		MatineeToggleTrack = Cast<UInterpTrackToggle>( CopyPasteObject );
-		if ( MatineeToggleTrack != nullptr )
-		{
-			break;
-		}
-	}
-	UMovieSceneParticleTrack* ParticleTrack = Cast<UMovieSceneParticleTrack>( Track );
-	MenuBuilder.AddMenuEntry(
-		NSLOCTEXT( "Sequencer", "PasteMatineeToggleTrack", "Paste Matinee Particle Track" ),
-		NSLOCTEXT( "Sequencer", "PasteMatineeToggleTrackTooltip", "Pastes keys from a Matinee particle track into this track." ),
-		FSlateIcon(),
-		FUIAction(
-			FExecuteAction::CreateStatic( &CopyInterpParticleTrack, GetSequencer().ToSharedRef(), MatineeToggleTrack, ParticleTrack ),
-			FCanExecuteAction::CreateLambda( [=]()->bool { return MatineeToggleTrack != nullptr && MatineeToggleTrack->ToggleTrack.Num() > 0 && ParticleTrack != nullptr; } ) ) );
-}
 
 
 #undef LOCTEXT_NAMESPACE
diff --git a/Engine/Source/Editor/MovieSceneTools/Private/TrackEditors/ParticleTrackEditor.h b/Engine/Source/Editor/MovieSceneTools/Private/TrackEditors/ParticleTrackEditor.h
--- a/Engine/Source/Editor/MovieSceneTools/Private/TrackEditors/ParticleTrackEditor.h
+++ b/Engine/Source/Editor/MovieSceneTools/Private/TrackEditors/ParticleTrackEditor.h
@@ -52,7 +52,6 @@ public:
 	virtual void BuildObjectBindingTrackMenu(FMenuBuilder& MenuBuilder, const TArray<FGuid>& ObjectBindings, const UClass* ObjectClass) override;
 	virtual TSharedRef<ISequencerSection> MakeSectionInterface( UMovieSceneSection& SectionObject, UMovieSceneTrack& Track, FGuid ObjectBinding ) override;
 	virtual bool SupportsType( TSubclassOf<UMovieSceneTrack> Type ) const override;
-	virtual void BuildTrackContextMenu( FMenuBuilder& MenuBuilder, UMovieSceneTrack* Track ) override;
 	
 	void AddParticleKey(TArray<FGuid> ObjectGuids);
 private:
diff --git a/Engine/Source/Editor/MovieSceneTools/Private/TrackEditors/PropertyTrackEditors/ColorPropertyTrackEditor.cpp b/Engine/Source/Editor/MovieSceneTools/Private/TrackEditors/PropertyTrackEditors/ColorPropertyTrackEditor.cpp
--- a/Engine/Source/Editor/MovieSceneTools/Private/TrackEditors/PropertyTrackEditors/ColorPropertyTrackEditor.cpp
+++ b/Engine/Source/Editor/MovieSceneTools/Private/TrackEditors/PropertyTrackEditors/ColorPropertyTrackEditor.cpp
@@ -95,35 +95,6 @@ void CopyInterpLinearColorTrack(TSharedRef<ISequencer> Sequencer, UInterpTrackLi
 	}
 }
 
-void FColorPropertyTrackEditor::BuildTrackContextMenu( FMenuBuilder& MenuBuilder, UMovieSceneTrack* Track )
-{
-	UInterpTrackColorProp* ColorPropTrack = nullptr;
-	UInterpTrackLinearColorProp* LinearColorPropTrack = nullptr;
-	for ( UObject* CopyPasteObject : GUnrealEd->MatineeCopyPasteBuffer )
-	{
-		ColorPropTrack = Cast<UInterpTrackColorProp>( CopyPasteObject );
-		LinearColorPropTrack = Cast<UInterpTrackLinearColorProp>( CopyPasteObject );
-		if ( ColorPropTrack != nullptr || LinearColorPropTrack != nullptr )
-		{
-			break;
-		}
-	}
-	UMovieSceneColorTrack* ColorTrack = Cast<UMovieSceneColorTrack>( Track );
-	MenuBuilder.AddMenuEntry(
-		NSLOCTEXT( "Sequencer", "PasteMatineeColorTrack", "Paste Matinee Color Track" ),
-		NSLOCTEXT( "Sequencer", "PasteMatineeColorTrackTooltip", "Pastes keys from a Matinee color track into this track." ),
-		FSlateIcon(),
-		FUIAction(
-			ColorPropTrack != nullptr ? 
-			FExecuteAction::CreateStatic( &CopyInterpColorTrack, GetSequencer().ToSharedRef(), ColorPropTrack, ColorTrack ) : 
-			FExecuteAction::CreateStatic( &CopyInterpLinearColorTrack, GetSequencer().ToSharedRef(), LinearColorPropTrack, ColorTrack ),			
-			FCanExecuteAction::CreateLambda( [=]()->bool { return ((ColorPropTrack != nullptr && ColorPropTrack->GetNumKeys() > 0) || (LinearColorPropTrack != nullptr && LinearColorPropTrack->GetNumKeys() > 0)) && ColorTrack != nullptr; } ) ) );
-
-	MenuBuilder.AddMenuSeparator();
-	FKeyframeTrackEditor::BuildTrackContextMenu(MenuBuilder, Track);
-}
-
-
 bool FColorPropertyTrackEditor::ModifyGeneratedKeysByCurrentAndWeight(UObject *Object, UMovieSceneTrack *Track, UMovieSceneSection* SectionToKey, FFrameNumber KeyTime, FGeneratedTrackKeys& GeneratedTotalKeys, float Weight) const
 {
 	using namespace UE::MovieScene;
diff --git a/Engine/Source/Editor/MovieSceneTools/Private/TrackEditors/PropertyTrackEditors/ColorPropertyTrackEditor.h b/Engine/Source/Editor/MovieSceneTools/Private/TrackEditors/PropertyTrackEditors/ColorPropertyTrackEditor.h
--- a/Engine/Source/Editor/MovieSceneTools/Private/TrackEditors/PropertyTrackEditors/ColorPropertyTrackEditor.h
+++ b/Engine/Source/Editor/MovieSceneTools/Private/TrackEditors/PropertyTrackEditors/ColorPropertyTrackEditor.h
@@ -59,13 +59,6 @@ protected:
 
 	virtual void GenerateKeysFromPropertyChanged(const FPropertyChangedParams& PropertyChangedParams, UMovieSceneSection* SectionToKey, FGeneratedTrackKeys& OutGeneratedKeys) override;
 	virtual bool ModifyGeneratedKeysByCurrentAndWeight(UObject *Object, UMovieSceneTrack *Track, UMovieSceneSection* SectionToKey, FFrameNumber KeyTime, FGeneratedTrackKeys& GeneratedTotalKeys, float Weight) const override;
-
-protected:
-
-	//~ ISequencerTrackEditor interface
-
-	virtual void BuildTrackContextMenu(FMenuBuilder& MenuBuilder, UMovieSceneTrack* Track) override;
-
 private:
 
 	static FName RedName;
diff --git a/Engine/Source/Editor/MovieSceneTools/Private/TrackEditors/PropertyTrackEditors/FloatPropertyTrackEditor.cpp b/Engine/Source/Editor/MovieSceneTools/Private/TrackEditors/PropertyTrackEditors/FloatPropertyTrackEditor.cpp
--- a/Engine/Source/Editor/MovieSceneTools/Private/TrackEditors/PropertyTrackEditors/FloatPropertyTrackEditor.cpp
+++ b/Engine/Source/Editor/MovieSceneTools/Private/TrackEditors/PropertyTrackEditors/FloatPropertyTrackEditor.cpp
@@ -33,30 +33,6 @@ void CopyInterpFloatTrack(TSharedRef<ISequencer> Sequencer, UInterpTrackFloatBas
 	}
 }
 
-void FFloatPropertyTrackEditor::BuildTrackContextMenu( FMenuBuilder& MenuBuilder, UMovieSceneTrack* Track )
-{
-	UInterpTrackFloatBase* MatineeFloatTrack = nullptr;
-	for ( UObject* CopyPasteObject : GUnrealEd->MatineeCopyPasteBuffer )
-	{
-		MatineeFloatTrack = Cast<UInterpTrackFloatBase>( CopyPasteObject );
-		if ( MatineeFloatTrack != nullptr )
-		{
-			break;
-		}
-	}
-	UMovieSceneFloatTrack* FloatTrack = Cast<UMovieSceneFloatTrack>( Track );
-	MenuBuilder.AddMenuEntry(
-		NSLOCTEXT( "Sequencer", "PasteMatineeFloatTrack", "Paste Matinee Float Track" ),
-		NSLOCTEXT( "Sequencer", "PasteMatineeFloatTrackTooltip", "Pastes keys from a Matinee float track into this track." ),
-		FSlateIcon(),
-		FUIAction(
-			FExecuteAction::CreateStatic( &CopyInterpFloatTrack, GetSequencer().ToSharedRef(), MatineeFloatTrack, FloatTrack ),
-			FCanExecuteAction::CreateLambda( [=]()->bool { return MatineeFloatTrack != nullptr && MatineeFloatTrack->GetNumKeys() > 0 && FloatTrack != nullptr; } ) ) );
-
-	MenuBuilder.AddMenuSeparator();
-	FKeyframeTrackEditor::BuildTrackContextMenu(MenuBuilder, Track);
-}
-
 
 float FFloatPropertyTrackEditor::RecomposeFloat(float InCurrentValue, UObject* AnimatedObject, UMovieSceneSection* Section)
 {
diff --git a/Engine/Source/Editor/MovieSceneTools/Private/TrackEditors/PropertyTrackEditors/FloatPropertyTrackEditor.h b/Engine/Source/Editor/MovieSceneTools/Private/TrackEditors/PropertyTrackEditors/FloatPropertyTrackEditor.h
--- a/Engine/Source/Editor/MovieSceneTools/Private/TrackEditors/PropertyTrackEditors/FloatPropertyTrackEditor.h
+++ b/Engine/Source/Editor/MovieSceneTools/Private/TrackEditors/PropertyTrackEditors/FloatPropertyTrackEditor.h
@@ -43,13 +43,6 @@ public:
 	 * @return The new instance of this class
 	 */
 	static TSharedRef<ISequencerTrackEditor> CreateTrackEditor(TSharedRef<ISequencer> OwningSequencer);
-
-public:
-
-	//~ FMovieSceneTrackEditor interface
-
-	virtual void BuildTrackContextMenu(FMenuBuilder& MenuBuilder, UMovieSceneTrack* Track) override;
-
 protected:
 
 	virtual void GenerateKeysFromPropertyChanged(const FPropertyChangedParams& PropertyChangedParams, UMovieSceneSection* SectionToKey, FGeneratedTrackKeys& OutGeneratedKeys) override;
diff --git a/Engine/Source/Editor/MovieSceneTools/Private/TrackEditors/PropertyTrackEditors/VectorPropertyTrackEditor.cpp b/Engine/Source/Editor/MovieSceneTools/Private/TrackEditors/PropertyTrackEditors/VectorPropertyTrackEditor.cpp
--- a/Engine/Source/Editor/MovieSceneTools/Private/TrackEditors/PropertyTrackEditors/VectorPropertyTrackEditor.cpp
+++ b/Engine/Source/Editor/MovieSceneTools/Private/TrackEditors/PropertyTrackEditors/VectorPropertyTrackEditor.cpp
@@ -107,31 +107,6 @@ void CopyInterpVectorTrack(TSharedRef<ISequencer> Sequencer, UInterpTrackVectorP
 	}
 }
 
-void FVectorPropertyTrackEditor::BuildTrackContextMenu( FMenuBuilder& MenuBuilder, UMovieSceneTrack* Track )
-{
-	UInterpTrackVectorProp* MatineeVectorTrack = nullptr;
-	for ( UObject* CopyPasteObject : GUnrealEd->MatineeCopyPasteBuffer )
-	{
-		MatineeVectorTrack = Cast<UInterpTrackVectorProp>( CopyPasteObject );
-		if ( MatineeVectorTrack != nullptr )
-		{
-			break;
-		}
-	}
-	UMovieSceneVectorTrack* VectorTrack = Cast<UMovieSceneVectorTrack>( Track );
-	MenuBuilder.AddMenuEntry(
-		NSLOCTEXT( "Sequencer", "PasteMatineeVectorTrack", "Paste Matinee Vector Track" ),
-		NSLOCTEXT( "Sequencer", "PasteMatineeVectorTrackTooltip", "Pastes keys from a Matinee vector track into this track." ),
-		FSlateIcon(),
-		FUIAction(
-			FExecuteAction::CreateStatic( &CopyInterpVectorTrack, GetSequencer().ToSharedRef(), MatineeVectorTrack, VectorTrack ),
-			FCanExecuteAction::CreateLambda( [=]()->bool { return MatineeVectorTrack != nullptr && MatineeVectorTrack->GetNumKeys() > 0 && VectorTrack != nullptr && VectorTrack->GetNumChannelsUsed() == 3; } ) ) );
-
-	MenuBuilder.AddMenuSeparator();
-	FKeyframeTrackEditor::BuildTrackContextMenu(MenuBuilder, Track);
-}
-
-
 bool FVectorPropertyTrackEditor::ModifyGeneratedKeysByCurrentAndWeight(UObject *Object, UMovieSceneTrack *Track, UMovieSceneSection* SectionToKey, FFrameNumber KeyTime, FGeneratedTrackKeys& GeneratedTotalKeys, float Weight) const
 {
 	using namespace UE::MovieScene;
diff --git a/Engine/Source/Editor/MovieSceneTools/Private/TrackEditors/PropertyTrackEditors/VectorPropertyTrackEditor.h b/Engine/Source/Editor/MovieSceneTools/Private/TrackEditors/PropertyTrackEditors/VectorPropertyTrackEditor.h
--- a/Engine/Source/Editor/MovieSceneTools/Private/TrackEditors/PropertyTrackEditors/VectorPropertyTrackEditor.h
+++ b/Engine/Source/Editor/MovieSceneTools/Private/TrackEditors/PropertyTrackEditors/VectorPropertyTrackEditor.h
@@ -51,11 +51,6 @@ public:
 	static TSharedRef<ISequencerTrackEditor> CreateTrackEditor(TSharedRef<ISequencer> OwningSequencer);
 
 protected:
-
-	//~ FMovieSceneTrackEditor interface
-
-	virtual void BuildTrackContextMenu(FMenuBuilder& MenuBuilder, UMovieSceneTrack* Track) override;
-
 	//~ FPropertyTrackEditor interface
 
 	virtual void GenerateKeysFromPropertyChanged(const FPropertyChangedParams& PropertyChangedParams, UMovieSceneSection* SectionToKey, FGeneratedTrackKeys& OutGeneratedKeys) override;
diff --git a/Engine/Source/Editor/MovieSceneTools/Private/TrackEditors/SkeletalAnimationTrackEditor.cpp b/Engine/Source/Editor/MovieSceneTools/Private/TrackEditors/SkeletalAnimationTrackEditor.cpp
--- a/Engine/Source/Editor/MovieSceneTools/Private/TrackEditors/SkeletalAnimationTrackEditor.cpp
+++ b/Engine/Source/Editor/MovieSceneTools/Private/TrackEditors/SkeletalAnimationTrackEditor.cpp
@@ -1644,24 +1644,7 @@ void FSkeletalAnimationTrackEditor::BuildTrackContextMenu( FMenuBuilder& MenuBui
 
 	TSharedPtr<ISequencer> SequencerPtr = GetSequencer();
 
-	UInterpTrackAnimControl* MatineeAnimControlTrack = nullptr;
-	for ( UObject* CopyPasteObject : GUnrealEd->MatineeCopyPasteBuffer )
-	{
-		MatineeAnimControlTrack = Cast<UInterpTrackAnimControl>( CopyPasteObject );
-		if ( MatineeAnimControlTrack != nullptr )
-		{
-			break;
-		}
-	}
 	UMovieSceneSkeletalAnimationTrack* SkeletalAnimationTrack = Cast<UMovieSceneSkeletalAnimationTrack>( Track );
-	MenuBuilder.AddMenuEntry(
-		NSLOCTEXT( "Sequencer", "PasteMatineeAnimControlTrack", "Paste Matinee SkeletalAnimation Track" ),
-		NSLOCTEXT( "Sequencer", "PasteMatineeAnimControlTrackTooltip", "Pastes keys from a Matinee float track into this track." ),
-		FSlateIcon(),
-		FUIAction(
-			FExecuteAction::CreateStatic( &CopyInterpAnimControlTrack, GetSequencer().ToSharedRef(), MatineeAnimControlTrack, SkeletalAnimationTrack ),
-			FCanExecuteAction::CreateLambda( [=]()->bool { return MatineeAnimControlTrack != nullptr && MatineeAnimControlTrack->AnimSeqs.Num() > 0 && SkeletalAnimationTrack != nullptr; } ) ) );
-
 	/** Put this back when and if it works
 	MenuBuilder.BeginSection(NAME_None, LOCTEXT("MotionBlendingOptions", "Motion Blending Options"));
 	{
diff --git a/Engine/Source/Editor/MovieSceneTools/Private/TrackEditors/TransformTrackEditor.cpp b/Engine/Source/Editor/MovieSceneTools/Private/TrackEditors/TransformTrackEditor.cpp
--- a/Engine/Source/Editor/MovieSceneTools/Private/TrackEditors/TransformTrackEditor.cpp
+++ b/Engine/Source/Editor/MovieSceneTools/Private/TrackEditors/TransformTrackEditor.cpp
@@ -137,25 +137,7 @@ bool CanCopyInterpMoveTrack(UInterpTrackMove* MoveTrack, UMovieScene3DTransformT
 
 void F3DTransformTrackEditor::BuildTrackContextMenu( FMenuBuilder& MenuBuilder, UMovieSceneTrack* Track )
 {
-	UInterpTrackMove* MoveTrack = nullptr;
-	for ( UObject* CopyPasteObject : GUnrealEd->MatineeCopyPasteBuffer )
-	{
-		MoveTrack = Cast<UInterpTrackMove>( CopyPasteObject );
-		if ( MoveTrack != nullptr )
-		{
-			break;
-		}
-	}
 	UMovieScene3DTransformTrack* TransformTrack = Cast<UMovieScene3DTransformTrack>( Track );
-	MenuBuilder.AddMenuEntry(
-		NSLOCTEXT("Sequencer", "PasteMatineeMoveTrack", "Paste Matinee Move Track"),
-		NSLOCTEXT("Sequencer", "PasteMatineeMoveTrackTooltip", "Pastes keys from a Matinee move track into this track."),
-		FSlateIcon(),
-		FUIAction(
-			FExecuteAction::CreateStatic(&CopyInterpMoveTrack, GetSequencer().ToSharedRef(), MoveTrack, TransformTrack),
-			FCanExecuteAction::CreateStatic(&CanCopyInterpMoveTrack, MoveTrack, TransformTrack)));
-
-	//		FCanExecuteAction::CreateLambda( [=]()->bool { return MoveTrack != nullptr && MoveTrack->GetNumKeys() > 0 && TransformTrack != nullptr; } ) ) );
 
 	auto AnimSubMenuDelegate = [](FMenuBuilder& InMenuBuilder, TSharedRef<ISequencer> InSequencer, UMovieScene3DTransformTrack* InTransformTrack)
 	{
diff --git a/Engine/Source/Editor/MovieSceneTools/Public/TrackEditors/AudioTrackEditor.h b/Engine/Source/Editor/MovieSceneTools/Public/TrackEditors/AudioTrackEditor.h
--- a/Engine/Source/Editor/MovieSceneTools/Public/TrackEditors/AudioTrackEditor.h
+++ b/Engine/Source/Editor/MovieSceneTools/Public/TrackEditors/AudioTrackEditor.h
@@ -57,7 +57,6 @@ public:
 	virtual TSharedRef<ISequencerSection> MakeSectionInterface(UMovieSceneSection& SectionObject, UMovieSceneTrack& Track, FGuid ObjectBinding) override;
 	virtual bool SupportsType(TSubclassOf<UMovieSceneTrack> Type) const override;
 	virtual bool SupportsSequence(UMovieSceneSequence* InSequence) const override;
-	virtual void BuildTrackContextMenu( FMenuBuilder& MenuBuilder, UMovieSceneTrack* Track ) override;
 	virtual const FSlateBrush* GetIconBrush() const override;
 	virtual bool IsResizable(UMovieSceneTrack* InTrack) const override;
 	virtual void Resize(float NewSize, UMovieSceneTrack* InTrack) override;
