Sequencer: Fixed incorrect desired state when adding unloaded data layer tracks

diff --git a/Engine/Source/Editor/MovieSceneTools/Private/TrackEditors/DataLayerTrackEditor.cpp b/Engine/Source/Editor/MovieSceneTools/Private/TrackEditors/DataLayerTrackEditor.cpp
--- a/Engine/Source/Editor/MovieSceneTools/Private/TrackEditors/DataLayerTrackEditor.cpp
+++ b/Engine/Source/Editor/MovieSceneTools/Private/TrackEditors/DataLayerTrackEditor.cpp
@@ -246,7 +246,7 @@ TSharedRef<SWidget> FDataLayerTrackEditor::BuildAddDataLayerMenu(UMovieSceneTrac
 		LOCTEXT("AddUnloadedDataLayer_Tip", "Instruct a data layer to be unloaded for a duration."),
 		FSlateIcon(),
 		FUIAction(FExecuteAction::CreateSP(
-			this, &FDataLayerTrackEditor::HandleAddNewSection, DataLayerTrack, EDataLayerState::Loaded)));
+			this, &FDataLayerTrackEditor::HandleAddNewSection, DataLayerTrack, EDataLayerState::Unloaded)));
 
 	return MenuBuilder.MakeWidget();
 }
