Fixing no unity compile errors
#rnx

diff --git a/Engine/Source/Editor/MovieSceneTools/Private/Sections/TransformPropertySection.cpp b/Engine/Source/Editor/MovieSceneTools/Private/Sections/TransformPropertySection.cpp
--- a/Engine/Source/Editor/MovieSceneTools/Private/Sections/TransformPropertySection.cpp
+++ b/Engine/Source/Editor/MovieSceneTools/Private/Sections/TransformPropertySection.cpp
@@ -4,6 +4,7 @@
 #include "ISectionLayoutBuilder.h"
 #include "Sections/MovieScene3DTransformSection.h"
 #include "SequencerSectionPainter.h"
+#include "Framework/Docking/TabManager.h"
 #include "Framework/MultiBox/MultiBoxBuilder.h"
 #include "ScopedTransaction.h"
 #include "ISequencer.h"
diff --git a/Engine/Source/Editor/Sequencer/Private/Tools/SMotionTrailOptions.cpp b/Engine/Source/Editor/Sequencer/Private/Tools/SMotionTrailOptions.cpp
--- a/Engine/Source/Editor/Sequencer/Private/Tools/SMotionTrailOptions.cpp
+++ b/Engine/Source/Editor/Sequencer/Private/Tools/SMotionTrailOptions.cpp
@@ -2,6 +2,9 @@
 
 #include "Tools/SMotionTrailOptions.h"
 #include "Tools/MotionTrailOptions.h"
+#include "Modules/moduleManager.h"
+#include "PropertyEditorModule.h"
+#include "Widgets/Layout/SBorder.h"
 
 #define LOCTEXT_NAMESPACE "MotionTrail"
 
diff --git a/Engine/Source/Editor/Sequencer/Public/Tools/MotionTrailOptions.h b/Engine/Source/Editor/Sequencer/Public/Tools/MotionTrailOptions.h
--- a/Engine/Source/Editor/Sequencer/Public/Tools/MotionTrailOptions.h
+++ b/Engine/Source/Editor/Sequencer/Public/Tools/MotionTrailOptions.h
@@ -2,6 +2,7 @@
 
 #pragma once
 #include "CoreMinimal.h"
+#include "HAL/IConsoleManager.h"
 #include "MotionTrailOptions.generated.h"
 
 //if true still use old motion trails for sequencer objects.
