Sequencer: Add options to normalize or specify a floating point range for the curve display

diff --git a/Engine/Source/Editor/MovieSceneTools/Private/Channels/BuiltInChannelEditors.cpp b/Engine/Source/Editor/MovieSceneTools/Private/Channels/BuiltInChannelEditors.cpp
--- a/Engine/Source/Editor/MovieSceneTools/Private/Channels/BuiltInChannelEditors.cpp
+++ b/Engine/Source/Editor/MovieSceneTools/Private/Channels/BuiltInChannelEditors.cpp
@@ -4,9 +4,11 @@
 #include "MovieSceneSequence.h"
 #include "MovieSceneSequenceEditor.h"
 #include "MovieSceneEventUtils.h"
+#include "Channels/MovieSceneChannelHandle.h"
 #include "Sections/MovieSceneEventSectionBase.h"
 #include "ISequencerChannelInterface.h"
 #include "Widgets/SNullWidget.h"
+#include "IKeyArea.h"
 #include "ISequencer.h"
 #include "SequencerSettings.h"
 #include "MovieSceneCommonHelpers.h"
@@ -1039,25 +1041,41 @@ struct FFloatChannelSectionMenuExtension : FExtender, TSharedFromThis<FFloatChan
 	{
 		TSharedRef<FFloatChannelSectionMenuExtension> SharedThis = AsShared();
 
-		// Menu entry for key area height
-		auto OnKeyAreaHeightChanged = [=](int32 NewValue) {
-
-			if (ISequencer* Sequencer = WeakSequencer.Pin().Get())
-			{
-				Sequencer->GetSequencerSettings()->SetKeyAreaHeightWithCurves((float)NewValue);
-			}
-		};
+		ISequencer* Sequencer = WeakSequencer.Pin().Get();
+		if (!Sequencer)
+		{
+			return;
+		}
 
-		auto GetKeyAreaHeight = [=]() {
+		USequencerSettings* Settings = Sequencer->GetSequencerSettings();
+		if (!Settings)
+		{
+			return;
+		}
 
-			if (ISequencer* Sequencer = WeakSequencer.Pin().Get())
+		// Menu entry for key area height
+		auto OnKeyAreaHeightChanged = [=](int32 NewValue) { Settings->SetKeyAreaHeightWithCurves((float)NewValue); };
+		auto GetKeyAreaHeight = [=]() { return (int)Settings->GetKeyAreaHeightWithCurves(); };
+
+		auto OnKeyAreaCurveNormalized = [=](FString KeyAreaName) 
+		{ 
+			if (Settings->HasKeyAreaCurveExtents(KeyAreaName))
+			{	
+				Settings->RemoveKeyAreaCurveExtents(KeyAreaName);
+			}
+			else
 			{
-				return (int)Sequencer->GetSequencerSettings()->GetKeyAreaHeightWithCurves();
+				// Initialize to some arbitrary value
+				Settings->SetKeyAreaCurveExtents(KeyAreaName, 0.f, 6.f); 
 			}
-
-			return 15;
 		};
+		auto GetKeyAreaCurveNormalized = [=](FString KeyAreaName) { return !Settings->HasKeyAreaCurveExtents(KeyAreaName); };
+
+		auto OnKeyAreaCurveMinChanged = [=](float NewValue, FString KeyAreaName) { float CurveMin = 0.f; float CurveMax = 0.f; Settings->GetKeyAreaCurveExtents(KeyAreaName, CurveMin, CurveMax); Settings->SetKeyAreaCurveExtents(KeyAreaName, NewValue, CurveMax); };
+		auto GetKeyAreaCurveMin = [=](FString KeyAreaName) { float CurveMin = 0.f; float CurveMax = 0.f; Settings->GetKeyAreaCurveExtents(KeyAreaName, CurveMin, CurveMax); return CurveMin; };
 
+		auto OnKeyAreaCurveMaxChanged = [=](float NewValue, FString KeyAreaName) { float CurveMin = 0.f; float CurveMax = 0.f; Settings->GetKeyAreaCurveExtents(KeyAreaName, CurveMin, CurveMax); Settings->SetKeyAreaCurveExtents(KeyAreaName, CurveMin, NewValue); };
+		auto GetKeyAreaCurveMax = [=](FString KeyAreaName) { float CurveMin = 0.f; float CurveMax = 0.f; Settings->GetKeyAreaCurveExtents(KeyAreaName, CurveMin, CurveMax); return CurveMax; };
 
 		MenuBuilder.AddMenuEntry(
 			LOCTEXT("ToggleShowCurve", "Show Curve"),
@@ -1071,7 +1089,69 @@ struct FFloatChannelSectionMenuExtension : FExtender, TSharedFromThis<FFloatChan
 			NAME_None,
 			EUserInterfaceActionType::ToggleButton
 		);	
+		
+		FString KeyAreaName;
+		TArray<const IKeyArea*> SelectedKeyAreas;
+		Sequencer->GetSelectedKeyAreas(SelectedKeyAreas);
+		for (const IKeyArea* KeyArea : SelectedKeyAreas)
+		{			
+			if (KeyArea)
+			{
+				KeyAreaName = KeyArea->GetName().ToString();
+				break;
+			}
+		}
 
+		MenuBuilder.AddMenuEntry(
+			LOCTEXT("ToggleKeyAreaCurveNormalized", "Key Area Curve Normalized"),
+			LOCTEXT("ToggleKeyAreaCurveNormalizedTooltip", "Toggle showing the curve in the track area as normalized"),
+			FSlateIcon(),
+			FUIAction(
+				FExecuteAction::CreateLambda([=] { OnKeyAreaCurveNormalized(KeyAreaName); }),
+				FCanExecuteAction(FCanExecuteAction::CreateLambda([SharedThis]{ return SharedThis->IsShowCurve(); })),
+				FIsActionChecked::CreateLambda([=] { return GetKeyAreaCurveNormalized(KeyAreaName); })
+			),
+			NAME_None,
+			EUserInterfaceActionType::ToggleButton
+		);	
+
+		MenuBuilder.AddWidget(
+			SNew(SHorizontalBox)
+			+ SHorizontalBox::Slot()
+				[
+					SNew(SSpacer)
+				]
+			+ SHorizontalBox::Slot()
+				.AutoWidth()
+				[
+					SNew(SBox)
+					.WidthOverride(50.f)
+					.IsEnabled_Lambda([=]() { return SharedThis->IsShowCurve() && Settings->HasKeyAreaCurveExtents(KeyAreaName); })
+					[
+						SNew(SSpinBox<float>)
+						.Style(&FEditorStyle::GetWidgetStyle<FSpinBoxStyle>("Sequencer.HyperlinkSpinBox"))
+						.OnValueCommitted_Lambda([=](float NewValue, ETextCommit::Type CommitType) { OnKeyAreaCurveMinChanged(NewValue, KeyAreaName); })
+						.OnValueChanged_Lambda([=](float NewValue) { OnKeyAreaCurveMinChanged(NewValue, KeyAreaName); })
+						.Value_Lambda([=]() -> float { return GetKeyAreaCurveMin(KeyAreaName); })
+					]
+				]
+			+ SHorizontalBox::Slot()
+				.AutoWidth()
+				[
+					SNew(SBox)
+					.WidthOverride(50.f)
+					.IsEnabled_Lambda([=]() { return SharedThis->IsShowCurve() && Settings->HasKeyAreaCurveExtents(KeyAreaName); })
+					[
+						SNew(SSpinBox<float>)
+						.Style(&FEditorStyle::GetWidgetStyle<FSpinBoxStyle>("Sequencer.HyperlinkSpinBox"))
+						.OnValueCommitted_Lambda([=](float NewValue, ETextCommit::Type CommitType) { OnKeyAreaCurveMaxChanged(NewValue, KeyAreaName); })
+						.OnValueChanged_Lambda([=](float NewValue) { OnKeyAreaCurveMaxChanged(NewValue, KeyAreaName); })
+						.Value_Lambda([=]() -> float { return GetKeyAreaCurveMax(KeyAreaName); })
+					]
+				],
+				LOCTEXT("KeyAreaCurveRangeText", "Key Area Curve Range")
+		);		
+		
 		MenuBuilder.AddWidget(
 			SNew(SHorizontalBox)
 			+ SHorizontalBox::Slot()
@@ -1082,7 +1162,7 @@ struct FFloatChannelSectionMenuExtension : FExtender, TSharedFromThis<FFloatChan
 				.AutoWidth()
 				[
 					SNew(SBox)
-					.WidthOverride(30.f)
+					.WidthOverride(50.f)
 					[
 						SNew(SSpinBox<int32>)
 						.Style(&FEditorStyle::GetWidgetStyle<FSpinBoxStyle>("Sequencer.HyperlinkSpinBox"))
diff --git a/Engine/Source/Editor/Sequencer/Private/SequencerKeyRenderer.cpp b/Engine/Source/Editor/Sequencer/Private/SequencerKeyRenderer.cpp
--- a/Engine/Source/Editor/Sequencer/Private/SequencerKeyRenderer.cpp
+++ b/Engine/Source/Editor/Sequencer/Private/SequencerKeyRenderer.cpp
@@ -481,7 +481,7 @@ void FKeyRenderer::FKeyDrawBatch::UpdateViewDependentData(FSequencer* Sequencer,
 					MinValue = FMath::Min(MinValue, Value);
 
 					FCurveKey NewKey;
-					NewKey.NormalizedValue = Value;
+					NewKey.Value = Value;
 					NewKey.FinalKeyPositionSeconds = EvalTime;
 
 					PrecomputedCurve.Add(NewKey);
@@ -501,13 +501,46 @@ void FKeyRenderer::FKeyDrawBatch::UpdateViewDependentData(FSequencer* Sequencer,
 			FrameNumber = FMath::Min(FrameNumber, UpperBoundFrame);
 		}
 
-		// Normalize
+		TOptional<float> SpecifiedMinValue;
+		TOptional<float> SpecifiedMaxValue;
+		
+		FString KeyAreaName = ThisKeyArea.Get()->GetName().ToString();
+		if (Sequencer->GetSequencerSettings()->HasKeyAreaCurveExtents(KeyAreaName))
+		{
+			float CurveMin = 0.f;
+			float CurveMax = 0.f;
+			Sequencer->GetSequencerSettings()->GetKeyAreaCurveExtents(KeyAreaName, CurveMin, CurveMax);
+			SpecifiedMinValue = CurveMin;
+			SpecifiedMaxValue = CurveMax;
+		}
+
+		if (SpecifiedMinValue.IsSet())
+		{
+			MinValue = SpecifiedMinValue.GetValue();
+		}
+
+		if (SpecifiedMaxValue.IsSet())
+		{
+			MaxValue = SpecifiedMaxValue.GetValue();
+		}
+
+		// Normalize or clamp
 		if (PrecomputedCurve.Num() > 0)
 		{
 			float DiffValue = MaxValue - MinValue;
 			for (FCurveKey& CurveKey : PrecomputedCurve)
 			{
-				CurveKey.NormalizedValue = (CurveKey.NormalizedValue - MinValue)/DiffValue;
+				if (SpecifiedMaxValue.IsSet())
+				{
+					CurveKey.Value = FMath::Min(CurveKey.Value, SpecifiedMaxValue.GetValue());
+				}
+
+				if (SpecifiedMinValue.IsSet())
+				{
+					CurveKey.Value = FMath::Max(CurveKey.Value, SpecifiedMinValue.GetValue());
+				}
+
+				CurveKey.Value = (CurveKey.Value - MinValue)/DiffValue;
 			}
 		}
 	}
@@ -529,7 +562,7 @@ void FKeyRenderer::FKeyDrawBatch::DrawCurve(FSequencer* Sequencer, FSequencerSec
 	{		
 		CurvePoints.Add(
 		FVector2D(TimeToPixelConverter.SecondsToPixel(CurveKey.FinalKeyPositionSeconds),
-				(1.0f - CurveKey.NormalizedValue) * KeyGeometry.GetLocalSize().Y ) );
+				(1.0f - CurveKey.Value) * KeyGeometry.GetLocalSize().Y ) );
 	}
 
 	const float CurveThickness = 1.f;
diff --git a/Engine/Source/Editor/Sequencer/Private/SequencerKeyRenderer.h b/Engine/Source/Editor/Sequencer/Private/SequencerKeyRenderer.h
--- a/Engine/Source/Editor/Sequencer/Private/SequencerKeyRenderer.h
+++ b/Engine/Source/Editor/Sequencer/Private/SequencerKeyRenderer.h
@@ -181,8 +181,8 @@ private:
 
 	struct FCurveKey
 	{
-		/** Normalized key value */
-		float NormalizedValue;
+		/** Key value */
+		float Value;
 
 		/** The time in seconds that this key should be drawn - represents the average time for overlapping keys */
 		double FinalKeyPositionSeconds;
diff --git a/Engine/Source/Editor/Sequencer/Private/SequencerSettings.cpp b/Engine/Source/Editor/Sequencer/Private/SequencerSettings.cpp
--- a/Engine/Source/Editor/Sequencer/Private/SequencerSettings.cpp
+++ b/Engine/Source/Editor/Sequencer/Private/SequencerSettings.cpp
@@ -552,6 +552,78 @@ void USequencerSettings::SetShowChannelColors(bool InbShowChannelColors)
 	}
 }
 
+bool USequencerSettings::HasKeyAreaCurveExtents(const FString& ChannelName) const
+{
+	TArray<FString> ChannelsArray;
+	KeyAreaCurveExtents.ParseIntoArray(ChannelsArray, TEXT(":"));
+
+	for (int32 ChannelIndex = 0; ChannelIndex < ChannelsArray.Num(); ++ChannelIndex)
+	{
+		TArray<FString> ExtentsArray;
+		ChannelsArray[ChannelIndex].ParseIntoArray(ExtentsArray, TEXT(","));	
+
+		if (ExtentsArray.Num() == 3 && ExtentsArray[0] == ChannelName)
+		{
+			return true;
+		}
+	}
+	return false;
+}
+
+void USequencerSettings::RemoveKeyAreaCurveExtents(const FString& ChannelName)
+{
+	TArray<FString> ChannelsArray;
+	KeyAreaCurveExtents.ParseIntoArray(ChannelsArray, TEXT(":"));
+
+	FString NewKeyAreaCurveExtents;
+	for (int32 ChannelIndex = 0; ChannelIndex < ChannelsArray.Num(); ++ChannelIndex)
+	{
+		TArray<FString> ExtentsArray;
+		ChannelsArray[ChannelIndex].ParseIntoArray(ExtentsArray, TEXT(","));	
+
+		if (ExtentsArray.Num() == 3 && ExtentsArray[0] == ChannelName)
+		{
+			continue;
+		}
+
+		NewKeyAreaCurveExtents.Append(TEXT(":"));
+		NewKeyAreaCurveExtents.Append(ChannelsArray[ChannelIndex]);
+	}
+
+	KeyAreaCurveExtents = NewKeyAreaCurveExtents;
+	SaveConfig();
+}
+
+void USequencerSettings::SetKeyAreaCurveExtents(const FString& ChannelName, float InMin, float InMax)
+{
+	RemoveKeyAreaCurveExtents(ChannelName);
+
+	FString NewChannelExtents = FString::Printf(TEXT("%s,%0.3f,%0.3f"), *ChannelName, InMin, InMax);
+	KeyAreaCurveExtents.Append(TEXT(":"));
+	KeyAreaCurveExtents.Append(NewChannelExtents);
+
+	SaveConfig();
+}
+
+void USequencerSettings::GetKeyAreaCurveExtents(const FString& ChannelName, float& InMin, float& InMax) const
+{
+	TArray<FString> ChannelsArray;
+	KeyAreaCurveExtents.ParseIntoArray(ChannelsArray, TEXT(":"));
+
+	for (int32 ChannelIndex = 0; ChannelIndex < ChannelsArray.Num(); ++ChannelIndex)
+	{
+		TArray<FString> ExtentsArray;
+		ChannelsArray[ChannelIndex].ParseIntoArray(ExtentsArray, TEXT(","));	
+
+		if (ExtentsArray.Num() == 3 && ExtentsArray[0] == ChannelName)
+		{
+			InMin = FCString::Atof(*ExtentsArray[1]);
+			InMax = FCString::Atof(*ExtentsArray[2]);
+			return;
+		}
+	}
+}
+
 float USequencerSettings::GetKeyAreaHeightWithCurves() const
 {
 	return KeyAreaHeightWithCurves;
diff --git a/Engine/Source/Editor/Sequencer/Public/SequencerSettings.h b/Engine/Source/Editor/Sequencer/Public/SequencerSettings.h
--- a/Engine/Source/Editor/Sequencer/Public/SequencerSettings.h
+++ b/Engine/Source/Editor/Sequencer/Public/SequencerSettings.h
@@ -274,6 +274,15 @@ public:
 	/** Set whether to show channel colors */
 	void SetShowChannelColors(bool bInShowChannelColors);
 
+	/** @return Whether the given channel has curve extents */
+	bool HasKeyAreaCurveExtents(const FString& ChannelName) const;
+	/** @ Remove curve extents for the given channel */
+	void RemoveKeyAreaCurveExtents(const FString& ChannelName);
+	/** @return Get the key area curve extents for the given channel */
+	void GetKeyAreaCurveExtents(const FString& ChannelName, float& InMin, float& InMax) const;
+	/** Set the key area curve extents for the given channel */
+	void SetKeyAreaCurveExtents(const FString& ChannelName, float InMin, float InMax);
+
 	/** @return The key area height when showing curves */
 	float GetKeyAreaHeightWithCurves() const;
 	/** Set the key area height when showing curves */
@@ -502,6 +511,10 @@ protected:
 	UPROPERTY(config, EditAnywhere, Category = Timeline)
 	bool bShowChannelColors;
 
+	/** The key area curve extents, stored per channel name */
+	UPROPERTY(config, EditAnywhere, Category = Timeline)
+	FString KeyAreaCurveExtents;
+
 	/** The key area height when showing curves */
 	UPROPERTY(config, EditAnywhere, Category = Timeline)
 	float KeyAreaHeightWithCurves;
