MatineeToSequencerUpgrader: Animation sections now default to 'Keep State' to match Matinee behavior. Sequencer's behavior for keep state is now much more accurate than it used to be so this should better reflect the original Matinee animations.

diff --git a/Engine/Source/Editor/MovieSceneTools/Private/MatineeImportTools.cpp b/Engine/Source/Editor/MovieSceneTools/Private/MatineeImportTools.cpp
--- a/Engine/Source/Editor/MovieSceneTools/Private/MatineeImportTools.cpp
+++ b/Engine/Source/Editor/MovieSceneTools/Private/MatineeImportTools.cpp
@@ -812,6 +812,10 @@ bool FMatineeImportTools::CopyInterpAnimControlTrack( UInterpTrackAnimControl* M
 		NewSection->Params.Animation = AnimSeq.AnimSeq;
 		NewSection->Params.SlotName = MatineeAnimControlTrack->SlotName;
 
+		// Matinee behavior is that animation sections hold their value after they end but Sequencer properly handles save/load restore
+		// so we need to change the default behavior of the section to mimick Matinee more.
+		NewSection->SetCompletionMode(EMovieSceneCompletionMode::KeepState);
+
 		SkeletalAnimationTrack->AddSection( *NewSection );
 		bSectionCreated = true;
 	}
