Fix sequencer camera cut thumbnails over-exposed bug.
Normal and high quality have expected exposure.

diff --git a/Engine/Source/Editor/MovieSceneTools/Private/TrackEditorThumbnail/TrackEditorThumbnail.cpp b/Engine/Source/Editor/MovieSceneTools/Private/TrackEditorThumbnail/TrackEditorThumbnail.cpp
--- a/Engine/Source/Editor/MovieSceneTools/Private/TrackEditorThumbnail/TrackEditorThumbnail.cpp
+++ b/Engine/Source/Editor/MovieSceneTools/Private/TrackEditorThumbnail/TrackEditorThumbnail.cpp
@@ -385,6 +385,8 @@ void FTrackEditorThumbnailCache::DrawViewportThumbnail(FTrackEditorThumbnail& Tr
 		.SetWorldTimes(FApp::GetCurrentTime() - GStartTime, FApp::GetDeltaTime(), FApp::GetCurrentTime() - GStartTime)
 		.SetResolveScene(true));
 
+	FSceneViewStateInterface* ViewStateInterface = nullptr;
+
 	// Screen percentage is not supported in thumbnail.
 	ViewFamily.EngineShowFlags.ScreenPercentage = false;
 
@@ -398,11 +400,25 @@ void FTrackEditorThumbnailCache::DrawViewportThumbnail(FTrackEditorThumbnail& Tr
 	case EThumbnailQuality::Normal:
 	case EThumbnailQuality::Best:
 		ViewFamily.EngineShowFlags.SetMotionBlur(false);
+
+		// Default eye adaptation requires a viewstate.
+		ViewFamily.EngineShowFlags.EyeAdaptation = true;
+		UMovieSceneUserThumbnailSettings* ThumbnailSettings = GetMutableDefault<UMovieSceneUserThumbnailSettings>();
+		FSceneViewStateInterface* Ref = ThumbnailSettings->ViewState.GetReference();
+		if (!Ref)
+		{
+			ThumbnailSettings->ViewState.Allocate();
+		}
+		ViewStateInterface = ThumbnailSettings->ViewState.GetReference();
 		break;
 	}
 
 	FSceneViewInitOptions ViewInitOptions;
 
+	// Use target exposure without blend. 
+	ViewInitOptions.bInCameraCut = true;
+	ViewInitOptions.SceneViewStateInterface = ViewStateInterface;
+
 	ViewInitOptions.BackgroundColor = FLinearColor::Black;
 	ViewInitOptions.SetViewRectangle(FIntRect(FIntPoint::ZeroValue, RTSize));
 	ViewInitOptions.ViewFamily = &ViewFamily;
diff --git a/Engine/Source/Editor/MovieSceneTools/Public/MovieSceneToolsUserSettings.h b/Engine/Source/Editor/MovieSceneTools/Public/MovieSceneToolsUserSettings.h
--- a/Engine/Source/Editor/MovieSceneTools/Public/MovieSceneToolsUserSettings.h
+++ b/Engine/Source/Editor/MovieSceneTools/Public/MovieSceneToolsUserSettings.h
@@ -41,6 +41,9 @@ public:
 	UPROPERTY(EditAnywhere, config, Category=General, meta=(EditCondition=bDrawThumbnails))
 	EThumbnailQuality Quality;
 
+	/** Temporal history for the view required for advanced features(e.g., eye adaptation) on all thumbnails*/
+	FSceneViewStateReference ViewState;
+
 	DECLARE_EVENT(UMovieSceneUserThumbnailSettings, FOnForceRedraw)
 	FOnForceRedraw& OnForceRedraw() { return OnForceRedrawEvent; }
 	void BroadcastRedrawThumbnails() const { OnForceRedrawEvent.Broadcast(); }
