Sequencer: Rename "Subscenes" to "Subsequences"

diff --git a/Engine/Source/Editor/MovieSceneTools/Private/TrackEditors/SubTrackEditor.cpp b/Engine/Source/Editor/MovieSceneTools/Private/TrackEditors/SubTrackEditor.cpp
--- a/Engine/Source/Editor/MovieSceneTools/Private/TrackEditors/SubTrackEditor.cpp
+++ b/Engine/Source/Editor/MovieSceneTools/Private/TrackEditors/SubTrackEditor.cpp
@@ -313,7 +313,7 @@ FSubTrackEditor::FSubTrackEditor(TSharedRef<ISequencer> InSequencer)
 void FSubTrackEditor::BuildAddTrackMenu(FMenuBuilder& MenuBuilder)
 {
 	MenuBuilder.AddMenuEntry(
-		LOCTEXT("AddSubTrack", "Subscenes Track"),
+		LOCTEXT("AddSubTrack", "Subsequences Track"),
 		LOCTEXT("AddSubTooltip", "Adds a new track that can contain other sequences."),
 		FSlateIcon(FEditorStyle::GetStyleSetName(), "Sequencer.Tracks.Sub"),
 		FUIAction(
@@ -379,7 +379,7 @@ bool FSubTrackEditor::HandleAssetAdded(UObject* Asset, const FGuid& TargetObject
 
 	if (CanAddSubSequence(*Sequence))
 	{
-		const FScopedTransaction Transaction(LOCTEXT("AddSubScene_Transaction", "Add Subscene"));
+		const FScopedTransaction Transaction(LOCTEXT("AddSubSequence_Transaction", "Add Subsequence"));
 
 		int32 RowIndex = INDEX_NONE;
 		AnimatablePropertyChanged(FOnKeyProperty::CreateRaw(this, &FSubTrackEditor::HandleSequenceAdded, Sequence, RowIndex));
diff --git a/Engine/Source/Runtime/MovieScene/Private/Tracks/MovieSceneSubTrack.cpp b/Engine/Source/Runtime/MovieScene/Private/Tracks/MovieSceneSubTrack.cpp
--- a/Engine/Source/Runtime/MovieScene/Private/Tracks/MovieSceneSubTrack.cpp
+++ b/Engine/Source/Runtime/MovieScene/Private/Tracks/MovieSceneSubTrack.cpp
@@ -219,7 +219,7 @@ bool UMovieSceneSubTrack::SupportsMultipleRows() const
 #if WITH_EDITORONLY_DATA
 FText UMovieSceneSubTrack::GetDefaultDisplayName() const
 {
-	return LOCTEXT("TrackName", "Subscenes");
+	return LOCTEXT("TrackName", "Subsequences");
 }
 #endif
 
