Sequencer: Disable match by name when importing onto one object

diff --git a/Engine/Source/Editor/MovieSceneTools/Private/MovieSceneToolHelpers.cpp b/Engine/Source/Editor/MovieSceneTools/Private/MovieSceneToolHelpers.cpp
--- a/Engine/Source/Editor/MovieSceneTools/Private/MovieSceneToolHelpers.cpp
+++ b/Engine/Source/Editor/MovieSceneTools/Private/MovieSceneToolHelpers.cpp
@@ -2747,7 +2747,7 @@ void MovieSceneToolHelpers::ImportFBXCameraToExisting(UnFbx::FFbxImporter* FbxIm
 					FString CameraName = GetCameraName(CameraNode);
 					FNotificationInfo Info(FText::Format(NSLOCTEXT("MovieSceneTools", "NoMatchingCameraWarning", "Failed to find any matching camera for {0}. Importing onto first camera from fbx {1}"), FText::FromString(ObjectName), FText::FromString(CameraName)));
 					Info.ExpireDuration = 5.0f;
-					FSlateNotificationManager::Get().AddNotification(Info)->SetCompletionState(SNotificationItem::CS_Fail);
+					FSlateNotificationManager::Get().AddNotification(Info)->SetCompletionState(SNotificationItem::CS_None);
 				}
 			}
 		}
@@ -3105,7 +3105,13 @@ private:
 		const FScopedTransaction Transaction(NSLOCTEXT("MovieSceneTools", "ImportFBXTransaction", "Import FBX"));
 		UnFbx::FFbxImporter* FbxImporter = UnFbx::FFbxImporter::GetInstance();
 
-		const bool bMatchByNameOnly = ImportFBXSettings->bMatchByNameOnly;
+		bool bMatchByNameOnly = ImportFBXSettings->bMatchByNameOnly;
+		if (ObjectBindingMap.Num() == 1 && bMatchByNameOnly)
+		{
+			UE_LOG(LogMovieScene, Display, TEXT("Fbx Import: Importing onto one selected binding, disabling match by name only."));
+			bMatchByNameOnly = false;
+		}
+
 		// Import static cameras first
 		ImportFBXCamera(FbxImporter, Sequence, *Sequencer, ObjectBindingMap, bMatchByNameOnly, bCreateCameras.IsSet() ? bCreateCameras.GetValue() : ImportFBXSettings->bCreateCameras);
 
