Sequencer: Fix double conversion to tick rate on importing float properties. The conversion is already implemented in FMatineeTools::SetOrAddKey

diff --git a/Engine/Source/Editor/MovieSceneTools/Private/MovieSceneToolHelpers.cpp b/Engine/Source/Editor/MovieSceneTools/Private/MovieSceneToolHelpers.cpp
--- a/Engine/Source/Editor/MovieSceneTools/Private/MovieSceneToolHelpers.cpp
+++ b/Engine/Source/Editor/MovieSceneTools/Private/MovieSceneToolHelpers.cpp
@@ -1084,7 +1084,7 @@ bool ImportFBXProperty(FString NodeName, FString AnimatedPropertyName, FGuid Obj
 					if (Source.IsKeyHandleValid(PrevKeyHandle))
 					{
 						FRichCurveKey &PrevKey = Source.GetKey(PrevKeyHandle);
-						ArriveTangent = ArriveTangent / ((Key.Time - PrevKey.Time) * DecimalRate);
+						ArriveTangent = ArriveTangent / (Key.Time - PrevKey.Time);
 
 					}
 					float LeaveTangent = Key.LeaveTangent;
@@ -1092,7 +1092,7 @@ bool ImportFBXProperty(FString NodeName, FString AnimatedPropertyName, FGuid Obj
 					if (Source.IsKeyHandleValid(NextKeyHandle))
 					{
 						FRichCurveKey &NextKey = Source.GetKey(NextKeyHandle);
-						LeaveTangent = LeaveTangent / ((NextKey.Time - Key.Time) * DecimalRate);
+						LeaveTangent = LeaveTangent / (NextKey.Time - Key.Time);
 					}
 
 					FFrameNumber KeyTime = (Key.Time * FrameRate).RoundToFrame();
