Sequencer: Large world coordinates support

diff --git a/Engine/Source/Editor/MovieSceneTools/Private/Channels/BuiltInChannelEditors.cpp b/Engine/Source/Editor/MovieSceneTools/Private/Channels/BuiltInChannelEditors.cpp
--- a/Engine/Source/Editor/MovieSceneTools/Private/Channels/BuiltInChannelEditors.cpp
+++ b/Engine/Source/Editor/MovieSceneTools/Private/Channels/BuiltInChannelEditors.cpp
@@ -24,6 +24,7 @@
 #include "Framework/MultiBox/MultiBoxBuilder.h"
 #include "Channels/MovieSceneChannelProxy.h"
 #include "Channels/MovieSceneChannelEditorData.h"
+#include "Channels/DoubleChannelCurveModel.h"
 #include "Channels/FloatChannelCurveModel.h"
 #include "Channels/IntegerChannelCurveModel.h"
 #include "Channels/BoolChannelCurveModel.h"
@@ -37,6 +38,7 @@
 #include "Framework/Application/SlateApplication.h"
 #include "Editor/SceneOutliner/Private/SSocketChooser.h"
 #include "SComponentChooser.h"
+#include "EntitySystem/MovieSceneDecompositionQuery.h"
 #include "EntitySystem/Interrogation/MovieSceneInterrogationLinker.h"
 #include "EntitySystem/Interrogation/MovieSceneInterrogatedPropertyInstantiator.h"
 #include "Systems/MovieScenePropertyInstantiator.h"
@@ -48,7 +50,8 @@
 #define LOCTEXT_NAMESPACE "BuiltInChannelEditors"
 
 
-FKeyHandle AddOrUpdateKey(FMovieSceneFloatChannel* Channel, UMovieSceneSection* SectionToKey, const TMovieSceneExternalValue<float>& ExternalValue, FFrameNumber InTime, ISequencer& Sequencer, const FGuid& InObjectBindingID, FTrackInstancePropertyBindings* PropertyBindings)
+template<typename ChannelType, typename ValueType>
+FKeyHandle AddOrUpdateKeyImpl(ChannelType* Channel, UMovieSceneSection* SectionToKey, const TMovieSceneExternalValue<ValueType>& ExternalValue, FFrameNumber InTime, ISequencer& Sequencer, const FGuid& InObjectBindingID, FTrackInstancePropertyBindings* PropertyBindings)
 {
 	using namespace UE::MovieScene;
 
@@ -56,7 +59,7 @@ FKeyHandle AddOrUpdateKey(FMovieSceneFloatChannel* Channel, UMovieSceneSection*
 
 	// Find the first bound object so we can get the current property channel value on it.
 	UObject* FirstBoundObject = nullptr;
-	TOptional<float> CurrentBoundObjectValue;
+	TOptional<ValueType> CurrentBoundObjectValue;
 	if (InObjectBindingID.IsValid())
 	{
 		for (TWeakObjectPtr<> WeakObject : Sequencer.FindBoundObjects(InObjectBindingID, SequenceID))
@@ -79,7 +82,7 @@ FKeyHandle AddOrUpdateKey(FMovieSceneFloatChannel* Channel, UMovieSceneSection*
 	// object *would* be at if we scrubbed here and let the sequence evaluation do its thing). This will help us figure out the difference between the current object value
 	// and the evaluated sequencer value: we will compute a new value for the channel so that a new sequence evaluation would come out at the "desired" value, which is
 	// what the current object value.
-	float NewValue = Channel->GetDefault().Get(0.f);
+	ValueType NewValue = Channel->GetDefault().Get(0.f);
 
 	const bool bWasEvaluated = Channel->Evaluate(InTime, NewValue);
 
@@ -88,7 +91,7 @@ FKeyHandle AddOrUpdateKey(FMovieSceneFloatChannel* Channel, UMovieSceneSection*
 		if (ExternalValue.OnGetCurrentValueAndWeight)
 		{
 			// We have a custom callback that can provide us with the evaluated value of this channel.
-			float CurrentValue = CurrentBoundObjectValue.Get(0.0f);
+			ValueType CurrentValue = CurrentBoundObjectValue.Get(0.0f);
 			float CurrentWeight = 1.0f;
 			FMovieSceneRootEvaluationTemplateInstance& EvaluationTemplate = Sequencer.GetEvaluationTemplate();
 			ExternalValue.OnGetCurrentValueAndWeight(FirstBoundObject, SectionToKey, InTime, Sequencer.GetFocusedTickResolution(), EvaluationTemplate, CurrentValue, CurrentWeight);
@@ -97,7 +100,7 @@ FKeyHandle AddOrUpdateKey(FMovieSceneFloatChannel* Channel, UMovieSceneSection*
 			{
 				if (bWasEvaluated)
 				{
-					float CurrentGlobalValue = CurrentBoundObjectValue.GetValue();
+					ValueType CurrentGlobalValue = CurrentBoundObjectValue.GetValue();
 					NewValue = (CurrentBoundObjectValue.Get(0.0f) - CurrentValue) * CurrentWeight + NewValue;
 				}
 				else //Nothing set (key or default) on channel so use external value
@@ -179,9 +182,9 @@ FKeyHandle AddOrUpdateKey(FMovieSceneFloatChannel* Channel, UMovieSceneSection*
 
 					FIntermediate3DTransform InTransformData;
 
-					TRecompositionResult<float> RecomposeResult = System->RecomposeBlendFloatChannel(BoundPropertyDefinition, ChannelIndex, Query, CurrentBoundObjectValue.Get(0.f));
+					TRecompositionResult<double> RecomposeResult = System->RecomposeBlendChannel(BoundPropertyDefinition, ChannelIndex, Query, (double)CurrentBoundObjectValue.Get(0.f));
 
-					NewValue = RecomposeResult.Values[0];
+					NewValue = (ValueType)RecomposeResult.Values[0];
 				}
 			}
 		}
@@ -191,6 +194,16 @@ FKeyHandle AddOrUpdateKey(FMovieSceneFloatChannel* Channel, UMovieSceneSection*
 	return AddKeyToChannel(Channel, InTime, NewValue, Sequencer.GetKeyInterpolation());
 }
 
+FKeyHandle AddOrUpdateKey(FMovieSceneFloatChannel* Channel, UMovieSceneSection* SectionToKey, const TMovieSceneExternalValue<float>& ExternalValue, FFrameNumber InTime, ISequencer& Sequencer, const FGuid& InObjectBindingID, FTrackInstancePropertyBindings* PropertyBindings)
+{
+	return AddOrUpdateKeyImpl<FMovieSceneFloatChannel, float>(Channel, SectionToKey, ExternalValue, InTime, Sequencer, InObjectBindingID, PropertyBindings);
+}
+
+FKeyHandle AddOrUpdateKey(FMovieSceneDoubleChannel* Channel, UMovieSceneSection* SectionToKey, const TMovieSceneExternalValue<double>& ExternalValue, FFrameNumber InTime, ISequencer& Sequencer, const FGuid& InObjectBindingID, FTrackInstancePropertyBindings* PropertyBindings)
+{
+	return AddOrUpdateKeyImpl<FMovieSceneDoubleChannel, double>(Channel, SectionToKey, ExternalValue, InTime, Sequencer, InObjectBindingID, PropertyBindings);
+}
+
 FKeyHandle AddOrUpdateKey(FMovieSceneActorReferenceData* Channel, UMovieSceneSection* SectionToKey, FFrameNumber InTime, ISequencer& Sequencer, const FGuid& InObjectBindingID, FTrackInstancePropertyBindings* PropertyBindings)
 {
 	if (PropertyBindings && InObjectBindingID.IsValid())
@@ -248,6 +261,10 @@ bool CanCreateKeyEditor(const FMovieSceneFloatChannel*   Channel)
 {
 	return true;
 }
+bool CanCreateKeyEditor(const FMovieSceneDoubleChannel*   Channel)
+{
+	return true;
+}
 bool CanCreateKeyEditor(const FMovieSceneStringChannel*  Channel)
 {
 	return true;
@@ -315,6 +332,24 @@ TSharedRef<SWidget> CreateKeyEditor(const TMovieSceneChannelHandle<FMovieSceneFl
 }
 
 
+TSharedRef<SWidget> CreateKeyEditor(const TMovieSceneChannelHandle<FMovieSceneDoubleChannel>&   Channel, UMovieSceneSection* Section, const FGuid& InObjectBindingID, TWeakPtr<FTrackInstancePropertyBindings> PropertyBindings, TWeakPtr<ISequencer> InSequencer)
+{
+	const TMovieSceneExternalValue<double>* ExternalValue = Channel.GetExtendedEditorData();
+	if (!ExternalValue)
+	{
+		return SNullWidget::NullWidget;
+	}
+
+	TSequencerKeyEditor<FMovieSceneDoubleChannel, double> KeyEditor(
+		InObjectBindingID, Channel,
+		Section, InSequencer, PropertyBindings, ExternalValue->OnGetExternalValue
+		);
+
+	typedef SNumericKeyEditor<FMovieSceneDoubleChannel, double> KeyEditorType;
+	return SNew(KeyEditorType, KeyEditor);
+}
+
+
 TSharedRef<SWidget> CreateKeyEditor(const TMovieSceneChannelHandle<FMovieSceneStringChannel>&  Channel, UMovieSceneSection* Section, const FGuid& InObjectBindingID, TWeakPtr<FTrackInstancePropertyBindings> PropertyBindings, TWeakPtr<ISequencer> InSequencer)
 {
 	const TMovieSceneExternalValue<FString>* ExternalValue = Channel.GetExtendedEditorData();
@@ -719,8 +754,11 @@ void PostConstructKeyInstance(const TMovieSceneChannelHandle<FMovieSceneObjectPa
 	KeyStruct->OnPropertyChangedEvent = CopyInstanceToKeyLambda;
 }
 
-void DrawKeys(FMovieSceneFloatChannel* Channel, TArrayView<const FKeyHandle> InKeyHandles, const UMovieSceneSection* InOwner, TArrayView<FKeyDrawParams> OutKeyDrawParams)
+template<typename ChannelType>
+void DrawKeysImpl(ChannelType* Channel, TArrayView<const FKeyHandle> InKeyHandles, const UMovieSceneSection* InOwner, TArrayView<FKeyDrawParams> OutKeyDrawParams)
 {
+	using ChannelValueType = typename ChannelType::ChannelValueType;
+
 	static const FName CircleKeyBrushName("Sequencer.KeyCircle");
 	static const FName DiamondKeyBrushName("Sequencer.KeyDiamond");
 	static const FName SquareKeyBrushName("Sequencer.KeySquare");
@@ -731,8 +769,8 @@ void DrawKeys(FMovieSceneFloatChannel* Channel, TArrayView<const FKeyHandle> InK
 	const FSlateBrush* SquareKeyBrush = FEditorStyle::GetBrush(SquareKeyBrushName);
 	const FSlateBrush* TriangleKeyBrush = FEditorStyle::GetBrush(TriangleKeyBrushName);
 
-	TMovieSceneChannelData<FMovieSceneFloatValue> ChannelData = Channel->GetData();
-	TArrayView<const FMovieSceneFloatValue> Values = ChannelData.GetValues();
+	TMovieSceneChannelData<ChannelValueType> ChannelData = Channel->GetData();
+	TArrayView<const ChannelValueType> Values = ChannelData.GetValues();
 
 	FKeyDrawParams TempParams;
 	TempParams.BorderBrush = TempParams.FillBrush = DiamondKeyBrush;
@@ -783,6 +821,16 @@ void DrawKeys(FMovieSceneFloatChannel* Channel, TArrayView<const FKeyHandle> InK
 	}
 }
 
+void DrawKeys(FMovieSceneFloatChannel* Channel, TArrayView<const FKeyHandle> InKeyHandles, const UMovieSceneSection* InOwner, TArrayView<FKeyDrawParams> OutKeyDrawParams)
+{
+	DrawKeysImpl(Channel, InKeyHandles, InOwner, OutKeyDrawParams);
+}
+
+void DrawKeys(FMovieSceneDoubleChannel* Channel, TArrayView<const FKeyHandle> InKeyHandles, const UMovieSceneSection* InOwner, TArrayView<FKeyDrawParams> OutKeyDrawParams)
+{
+	DrawKeysImpl(Channel, InKeyHandles, InOwner, OutKeyDrawParams);
+}
+
 void DrawKeys(FMovieSceneParticleChannel* Channel, TArrayView<const FKeyHandle> InKeyHandles, const UMovieSceneSection* InOwner, TArrayView<FKeyDrawParams> OutKeyDrawParams)
 {
 	static const FName KeyLeftBrushName("Sequencer.KeyLeft");
@@ -856,9 +904,12 @@ void DrawKeys(FMovieSceneEventChannel* Channel, TArrayView<const FKeyHandle> InK
 	}
 }
 
-struct FFloatChannelKeyMenuExtension : FExtender, TSharedFromThis<FFloatChannelKeyMenuExtension>
+template<typename ChannelType>
+struct TCurveChannelKeyMenuExtension : FExtender, TSharedFromThis<TCurveChannelKeyMenuExtension<ChannelType>>
 {
-	FFloatChannelKeyMenuExtension(TWeakPtr<ISequencer> InSequencer, TArray<TExtendKeyMenuParams<FMovieSceneFloatChannel>>&& InChannels)
+	using ChannelValueType = typename ChannelType::ChannelValueType;
+
+	TCurveChannelKeyMenuExtension(TWeakPtr<ISequencer> InSequencer, TArray<TExtendKeyMenuParams<ChannelType>>&& InChannels)
 		: WeakSequencer(InSequencer)
 		, ChannelAndHandles(MoveTemp(InChannels))
 	{}
@@ -871,7 +922,7 @@ struct FFloatChannelKeyMenuExtension : FExtender, TSharedFromThis<FFloatChannelK
 			return;
 		}
 
-		TSharedRef<FFloatChannelKeyMenuExtension> SharedThis = AsShared();
+		TSharedRef<TCurveChannelKeyMenuExtension<ChannelType>> SharedThis = this->AsShared();
 
 		MenuBuilder.BeginSection("SequencerInterpolation", LOCTEXT("KeyInterpolationMenu", "Key Interpolation"));
 		{
@@ -943,17 +994,17 @@ struct FFloatChannelKeyMenuExtension : FExtender, TSharedFromThis<FFloatChannelK
 		FScopedTransaction SetInterpTangentModeTransaction(NSLOCTEXT("Sequencer", "SetInterpTangentMode_Transaction", "Set Interpolation and Tangent Mode"));
 		bool bAnythingChanged = false;
 
-		for (const TExtendKeyMenuParams<FMovieSceneFloatChannel>& Channel : ChannelAndHandles)
+		for (const TExtendKeyMenuParams<ChannelType>& Channel : ChannelAndHandles)
 		{
 			UMovieSceneSection* Section = Channel.Section.Get();
-			FMovieSceneFloatChannel* ChannelPtr = Channel.Channel.Get();
+			ChannelType* ChannelPtr = Channel.Channel.Get();
 
 			if (Section && ChannelPtr)
 			{
 				Section->Modify();
 
-				TMovieSceneChannelData<FMovieSceneFloatValue> ChannelData = ChannelPtr->GetData();
-				TArrayView<FMovieSceneFloatValue> Values = ChannelData.GetValues();
+				TMovieSceneChannelData<ChannelValueType> ChannelData = ChannelPtr->GetData();
+				TArrayView<ChannelValueType> Values = ChannelData.GetValues();
 
 				for (FKeyHandle Handle : Channel.Handles)
 				{
@@ -981,13 +1032,13 @@ struct FFloatChannelKeyMenuExtension : FExtender, TSharedFromThis<FFloatChannelK
 
 	bool IsInterpTangentModeSelected(ERichCurveInterpMode InterpMode, ERichCurveTangentMode TangentMode) const
 	{
-		for (const TExtendKeyMenuParams<FMovieSceneFloatChannel>& Channel : ChannelAndHandles)
+		for (const TExtendKeyMenuParams<ChannelType>& Channel : ChannelAndHandles)
 		{
-			FMovieSceneFloatChannel* ChannelPtr = Channel.Channel.Get();
+			ChannelType* ChannelPtr = Channel.Channel.Get();
 			if (ChannelPtr)
 			{
-				TMovieSceneChannelData<FMovieSceneFloatValue> ChannelData = ChannelPtr->GetData();
-				TArrayView<FMovieSceneFloatValue> Values = ChannelData.GetValues();
+				TMovieSceneChannelData<ChannelValueType> ChannelData = ChannelPtr->GetData();
+				TArrayView<ChannelValueType> Values = ChannelData.GetValues();
 
 				for (FKeyHandle Handle : Channel.Handles)
 				{
@@ -1005,16 +1056,33 @@ struct FFloatChannelKeyMenuExtension : FExtender, TSharedFromThis<FFloatChannelK
 private:
 
 	/** Hidden AsShared() methods to prevent CreateSP delegate use since this extender disappears with its menu. */
-	using TSharedFromThis::AsShared;
+	using TSharedFromThis<TCurveChannelKeyMenuExtension<ChannelType>>::AsShared;
 
 	TWeakPtr<ISequencer> WeakSequencer;
-	TArray<TExtendKeyMenuParams<FMovieSceneFloatChannel>> ChannelAndHandles;
+	TArray<TExtendKeyMenuParams<ChannelType>> ChannelAndHandles;
 };
 
+struct FFloatChannelKeyMenuExtension : TCurveChannelKeyMenuExtension<FMovieSceneFloatChannel>
+{
+	FFloatChannelKeyMenuExtension(TWeakPtr<ISequencer> InSequencer, TArray<TExtendKeyMenuParams<FMovieSceneFloatChannel>>&& InChannels)
+		: TCurveChannelKeyMenuExtension<FMovieSceneFloatChannel>(InSequencer, MoveTemp(InChannels))
+	{}
+};
 
-struct FFloatChannelSectionMenuExtension : FExtender, TSharedFromThis<FFloatChannelSectionMenuExtension>
+struct FDoubleChannelKeyMenuExtension : TCurveChannelKeyMenuExtension<FMovieSceneDoubleChannel>
 {
-	FFloatChannelSectionMenuExtension(TWeakPtr<ISequencer> InSequencer, TArray<TMovieSceneChannelHandle<FMovieSceneFloatChannel>>&& InChannels, TArrayView<UMovieSceneSection* const> InSections)
+	FDoubleChannelKeyMenuExtension(TWeakPtr<ISequencer> InSequencer, TArray<TExtendKeyMenuParams<FMovieSceneDoubleChannel>>&& InChannels)
+		: TCurveChannelKeyMenuExtension<FMovieSceneDoubleChannel>(InSequencer, MoveTemp(InChannels))
+	{}
+};
+
+template<typename ChannelType>
+struct TCurveChannelSectionMenuExtension : FExtender, TSharedFromThis<TCurveChannelSectionMenuExtension<ChannelType>>
+{
+	using ChannelValueType = typename ChannelType::ChannelValueType;
+	using CurveValueType = typename ChannelType::CurveValueType;
+
+	TCurveChannelSectionMenuExtension(TWeakPtr<ISequencer> InSequencer, TArray<TMovieSceneChannelHandle<ChannelType>>&& InChannels, TArrayView<UMovieSceneSection* const> InSections)
 		: WeakSequencer(InSequencer)
 		, Channels(MoveTemp(InChannels))
 	{
@@ -1033,7 +1101,7 @@ struct FFloatChannelSectionMenuExtension : FExtender, TSharedFromThis<FFloatChan
 			return;
 		}
 
-		TSharedRef<FFloatChannelSectionMenuExtension> SharedThis = AsShared();
+		TSharedRef<TCurveChannelSectionMenuExtension<ChannelType>> SharedThis = this->AsShared();
 
 		MenuBuilder.AddSubMenu(
 			LOCTEXT("SetPreInfinityExtrap", "Pre-Infinity"),
@@ -1056,7 +1124,7 @@ struct FFloatChannelSectionMenuExtension : FExtender, TSharedFromThis<FFloatChan
 
 	void AddDisplayOptionsMenu(FMenuBuilder& MenuBuilder)
 	{
-		TSharedRef<FFloatChannelSectionMenuExtension> SharedThis = AsShared();
+		TSharedRef<TCurveChannelSectionMenuExtension<ChannelType>> SharedThis = this->AsShared();
 
 		ISequencer* Sequencer = WeakSequencer.Pin().Get();
 		if (!Sequencer)
@@ -1197,7 +1265,7 @@ struct FFloatChannelSectionMenuExtension : FExtender, TSharedFromThis<FFloatChan
 
 	void AddExtrapolationMenu(FMenuBuilder& MenuBuilder, bool bPreInfinity)
 	{
-		TSharedRef<FFloatChannelSectionMenuExtension> SharedThis = AsShared();
+		TSharedRef<TCurveChannelSectionMenuExtension<ChannelType>> SharedThis = this->AsShared();
 
 		MenuBuilder.AddMenuEntry(
 			LOCTEXT("SetExtrapCycle", "Cycle"),
@@ -1281,9 +1349,9 @@ struct FFloatChannelSectionMenuExtension : FExtender, TSharedFromThis<FFloatChan
 		}
 
 		// Apply to all channels
-		for (const TMovieSceneChannelHandle<FMovieSceneFloatChannel>& Handle : Channels)
+		for (const TMovieSceneChannelHandle<ChannelType>& Handle : Channels)
 		{
-			FMovieSceneFloatChannel* Channel = Handle.Get();
+			ChannelType* Channel = Handle.Get();
 
 			if (Channel)
 			{
@@ -1309,9 +1377,9 @@ struct FFloatChannelSectionMenuExtension : FExtender, TSharedFromThis<FFloatChan
 
 	bool IsExtrapolationModeSelected(ERichCurveExtrapolation ExtrapMode, bool bPreInfinity) const
 	{
-		for (const TMovieSceneChannelHandle<FMovieSceneFloatChannel>& Handle : Channels)
+		for (const TMovieSceneChannelHandle<ChannelType>& Handle : Channels)
 		{
-			FMovieSceneFloatChannel* Channel = Handle.Get();
+			ChannelType* Channel = Handle.Get();
 
 			if (Channel)
 			{
@@ -1343,9 +1411,9 @@ struct FFloatChannelSectionMenuExtension : FExtender, TSharedFromThis<FFloatChan
 		}
 
 		// Apply to all channels
-		for (const TMovieSceneChannelHandle<FMovieSceneFloatChannel>& Handle : Channels)
+		for (const TMovieSceneChannelHandle<ChannelType>& Handle : Channels)
 		{
-			FMovieSceneFloatChannel* Channel = Handle.Get();
+			ChannelType* Channel = Handle.Get();
 
 			if (Channel)
 			{
@@ -1362,9 +1430,9 @@ struct FFloatChannelSectionMenuExtension : FExtender, TSharedFromThis<FFloatChan
 
 	bool IsShowCurve() const
 	{
-		for (const TMovieSceneChannelHandle<FMovieSceneFloatChannel>& Handle : Channels)
+		for (const TMovieSceneChannelHandle<ChannelType>& Handle : Channels)
 		{
-			FMovieSceneFloatChannel* Channel = Handle.Get();
+			ChannelType* Channel = Handle.Get();
 
 			if (Channel && !Channel->GetShowCurve())
 			{
@@ -1378,13 +1446,27 @@ struct FFloatChannelSectionMenuExtension : FExtender, TSharedFromThis<FFloatChan
 private:
 
 	/** Hidden AsShared() methods to prevent CreateSP delegate use since this extender disappears with its menu. */
-	using TSharedFromThis::AsShared;
+	using TSharedFromThis<TCurveChannelSectionMenuExtension<ChannelType>>::AsShared;
 
 	TWeakPtr<ISequencer> WeakSequencer;
-	TArray<TMovieSceneChannelHandle<FMovieSceneFloatChannel>> Channels;
+	TArray<TMovieSceneChannelHandle<ChannelType>> Channels;
 	TArray<TWeakObjectPtr<UMovieSceneSection>> Sections;
 };
 
+struct FFloatChannelSectionMenuExtension : TCurveChannelSectionMenuExtension<FMovieSceneFloatChannel>
+{
+	FFloatChannelSectionMenuExtension(TWeakPtr<ISequencer> InSequencer, TArray<TMovieSceneChannelHandle<FMovieSceneFloatChannel>>&& InChannels, TArrayView<UMovieSceneSection* const> InSections)
+		: TCurveChannelSectionMenuExtension<FMovieSceneFloatChannel>(InSequencer, MoveTemp(InChannels), InSections)
+	{}
+};
+
+struct FDoubleChannelSectionMenuExtension : TCurveChannelSectionMenuExtension<FMovieSceneDoubleChannel>
+{
+	FDoubleChannelSectionMenuExtension(TWeakPtr<ISequencer> InSequencer, TArray<TMovieSceneChannelHandle<FMovieSceneDoubleChannel>>&& InChannels, TArrayView<UMovieSceneSection* const> InSections)
+		: TCurveChannelSectionMenuExtension<FMovieSceneDoubleChannel>(InSequencer, MoveTemp(InChannels), InSections)
+	{}
+};
+
 void ExtendSectionMenu(FMenuBuilder& OuterMenuBuilder, TArray<TMovieSceneChannelHandle<FMovieSceneFloatChannel>>&& Channels, TArrayView<UMovieSceneSection* const> Sections, TWeakPtr<ISequencer> InSequencer)
 {
 	TSharedRef<FFloatChannelSectionMenuExtension> Extension = MakeShared<FFloatChannelSectionMenuExtension>(InSequencer, MoveTemp(Channels), Sections);
@@ -1394,6 +1476,15 @@ void ExtendSectionMenu(FMenuBuilder& OuterMenuBuilder, TArray<TMovieSceneChannel
 	OuterMenuBuilder.PushExtender(Extension);
 }
 
+void ExtendSectionMenu(FMenuBuilder& OuterMenuBuilder, TArray<TMovieSceneChannelHandle<FMovieSceneDoubleChannel>>&& Channels, TArrayView<UMovieSceneSection* const> Sections, TWeakPtr<ISequencer> InSequencer)
+{
+	TSharedRef<FDoubleChannelSectionMenuExtension> Extension = MakeShared<FDoubleChannelSectionMenuExtension>(InSequencer, MoveTemp(Channels), Sections);
+
+	Extension->AddMenuExtension("SequencerSections", EExtensionHook::First, nullptr, FMenuExtensionDelegate::CreateLambda([Extension](FMenuBuilder& MenuBuilder) { Extension->ExtendMenu(MenuBuilder); }));
+
+	OuterMenuBuilder.PushExtender(Extension);
+}
+
 void ExtendKeyMenu(FMenuBuilder& OuterMenuBuilder, TArray<TExtendKeyMenuParams<FMovieSceneFloatChannel>>&& Channels, TWeakPtr<ISequencer> InSequencer)
 {
 	TSharedRef<FFloatChannelKeyMenuExtension> Extension = MakeShared<FFloatChannelKeyMenuExtension>(InSequencer, MoveTemp(Channels));
@@ -1403,11 +1494,25 @@ void ExtendKeyMenu(FMenuBuilder& OuterMenuBuilder, TArray<TExtendKeyMenuParams<F
 	OuterMenuBuilder.PushExtender(Extension);
 }
 
+void ExtendKeyMenu(FMenuBuilder& OuterMenuBuilder, TArray<TExtendKeyMenuParams<FMovieSceneDoubleChannel>>&& Channels, TWeakPtr<ISequencer> InSequencer)
+{
+	TSharedRef<FDoubleChannelKeyMenuExtension> Extension = MakeShared<FDoubleChannelKeyMenuExtension>(InSequencer, MoveTemp(Channels));
+
+	Extension->AddMenuExtension("SequencerKeyEdit", EExtensionHook::After, nullptr, FMenuExtensionDelegate::CreateLambda([Extension](FMenuBuilder& MenuBuilder) { Extension->ExtendMenu(MenuBuilder); }));
+
+	OuterMenuBuilder.PushExtender(Extension);
+}
+
 TUniquePtr<FCurveModel> CreateCurveEditorModel(const TMovieSceneChannelHandle<FMovieSceneFloatChannel>& FloatChannel, UMovieSceneSection* OwningSection, TSharedRef<ISequencer> InSequencer)
 {
 	return MakeUnique<FFloatChannelCurveModel>(FloatChannel, OwningSection, InSequencer);
 }
 
+TUniquePtr<FCurveModel> CreateCurveEditorModel(const TMovieSceneChannelHandle<FMovieSceneDoubleChannel>& DoubleChannel, UMovieSceneSection* OwningSection, TSharedRef<ISequencer> InSequencer)
+{
+	return MakeUnique<FDoubleChannelCurveModel>(DoubleChannel, OwningSection, InSequencer);
+}
+
 TUniquePtr<FCurveModel> CreateCurveEditorModel(const TMovieSceneChannelHandle<FMovieSceneIntegerChannel>& IntegerChannel, UMovieSceneSection* OwningSection, TSharedRef<ISequencer> InSequencer)
 {
 	return MakeUnique<FIntegerChannelCurveModel>(IntegerChannel, OwningSection, InSequencer);
diff --git a/Engine/Source/Editor/MovieSceneTools/Private/Channels/BuiltInChannelEditors.h b/Engine/Source/Editor/MovieSceneTools/Private/Channels/BuiltInChannelEditors.h
--- a/Engine/Source/Editor/MovieSceneTools/Private/Channels/BuiltInChannelEditors.h
+++ b/Engine/Source/Editor/MovieSceneTools/Private/Channels/BuiltInChannelEditors.h
@@ -15,6 +15,7 @@
 
 #include "Channels/MovieSceneBoolChannel.h"
 #include "Channels/MovieSceneByteChannel.h"
+#include "Channels/MovieSceneDoubleChannel.h"
 #include "Channels/MovieSceneFloatChannel.h"
 #include "Channels/MovieSceneIntegerChannel.h"
 #include "Channels/MovieSceneEventChannel.h"
@@ -33,6 +34,7 @@ class FStructOnScope;
 class ISectionLayoutBuilder;
 
 /** Overrides for adding or updating a key for non-standard channels */
+FKeyHandle AddOrUpdateKey(FMovieSceneDoubleChannel* Channel, UMovieSceneSection* SectionToKey,  const TMovieSceneExternalValue<double>& EditorData, FFrameNumber InTime, ISequencer& Sequencer, const FGuid& InObjectBindingID, FTrackInstancePropertyBindings* PropertyBindings);
 FKeyHandle AddOrUpdateKey(FMovieSceneFloatChannel* Channel, UMovieSceneSection* SectionToKey,  const TMovieSceneExternalValue<float>& EditorData, FFrameNumber InTime, ISequencer& Sequencer, const FGuid& InObjectBindingID, FTrackInstancePropertyBindings* PropertyBindings);
 FKeyHandle AddOrUpdateKey(FMovieSceneActorReferenceData* Channel, UMovieSceneSection* SectionToKey, FFrameNumber InTime, ISequencer& Sequencer, const FGuid& InObjectBindingID, FTrackInstancePropertyBindings* PropertyBindings);
 
@@ -40,6 +42,7 @@ FKeyHandle AddOrUpdateKey(FMovieSceneActorReferenceData* Channel, UMovieSceneSec
 bool CanCreateKeyEditor(const FMovieSceneBoolChannel*       Channel);
 bool CanCreateKeyEditor(const FMovieSceneByteChannel*       Channel);
 bool CanCreateKeyEditor(const FMovieSceneIntegerChannel*    Channel);
+bool CanCreateKeyEditor(const FMovieSceneDoubleChannel*      Channel);
 bool CanCreateKeyEditor(const FMovieSceneFloatChannel*      Channel);
 bool CanCreateKeyEditor(const FMovieSceneStringChannel*     Channel);
 bool CanCreateKeyEditor(const FMovieSceneObjectPathChannel* Channel);
@@ -48,6 +51,7 @@ bool CanCreateKeyEditor(const FMovieSceneActorReferenceData* Channel);
 TSharedRef<SWidget> CreateKeyEditor(const TMovieSceneChannelHandle<FMovieSceneBoolChannel>&        Channel, UMovieSceneSection* Section, const FGuid& InObjectBindingID, TWeakPtr<FTrackInstancePropertyBindings> PropertyBindings, TWeakPtr<ISequencer> InSequencer);
 TSharedRef<SWidget> CreateKeyEditor(const TMovieSceneChannelHandle<FMovieSceneByteChannel>&        Channel, UMovieSceneSection* Section, const FGuid& InObjectBindingID, TWeakPtr<FTrackInstancePropertyBindings> PropertyBindings, TWeakPtr<ISequencer> InSequencer);
 TSharedRef<SWidget> CreateKeyEditor(const TMovieSceneChannelHandle<FMovieSceneIntegerChannel>&     Channel, UMovieSceneSection* Section, const FGuid& InObjectBindingID, TWeakPtr<FTrackInstancePropertyBindings> PropertyBindings, TWeakPtr<ISequencer> InSequencer);
+TSharedRef<SWidget> CreateKeyEditor(const TMovieSceneChannelHandle<FMovieSceneDoubleChannel>&       Channel, UMovieSceneSection* Section, const FGuid& InObjectBindingID, TWeakPtr<FTrackInstancePropertyBindings> PropertyBindings, TWeakPtr<ISequencer> InSequencer);
 TSharedRef<SWidget> CreateKeyEditor(const TMovieSceneChannelHandle<FMovieSceneFloatChannel>&       Channel, UMovieSceneSection* Section, const FGuid& InObjectBindingID, TWeakPtr<FTrackInstancePropertyBindings> PropertyBindings, TWeakPtr<ISequencer> InSequencer);
 TSharedRef<SWidget> CreateKeyEditor(const TMovieSceneChannelHandle<FMovieSceneStringChannel>&      Channel, UMovieSceneSection* Section, const FGuid& InObjectBindingID, TWeakPtr<FTrackInstancePropertyBindings> PropertyBindings, TWeakPtr<ISequencer> InSequencer);
 TSharedRef<SWidget> CreateKeyEditor(const TMovieSceneChannelHandle<FMovieSceneObjectPathChannel>&  Channel, UMovieSceneSection* Section, const FGuid& InObjectBindingID, TWeakPtr<FTrackInstancePropertyBindings> PropertyBindings, TWeakPtr<ISequencer> InSequencer);
@@ -59,20 +63,25 @@ UMovieSceneKeyStructType* InstanceGeneratedStruct(FMovieSceneObjectPathChannel*
 void PostConstructKeyInstance(const TMovieSceneChannelHandle<FMovieSceneObjectPathChannel>& ChannelHandle, FKeyHandle InHandle, FStructOnScope* Struct);
 
 /** Key drawing overrides */
+void DrawKeys(FMovieSceneDoubleChannel*    Channel, TArrayView<const FKeyHandle> InKeyHandles, const UMovieSceneSection* InOwner, TArrayView<FKeyDrawParams> OutKeyDrawParams);
 void DrawKeys(FMovieSceneFloatChannel*    Channel, TArrayView<const FKeyHandle> InKeyHandles, const UMovieSceneSection* InOwner, TArrayView<FKeyDrawParams> OutKeyDrawParams);
 void DrawKeys(FMovieSceneParticleChannel* Channel, TArrayView<const FKeyHandle> InKeyHandles, const UMovieSceneSection* InOwner, TArrayView<FKeyDrawParams> OutKeyDrawParams);
 void DrawKeys(FMovieSceneEventChannel*    Channel, TArrayView<const FKeyHandle> InKeyHandles, const UMovieSceneSection* InOwner, TArrayView<FKeyDrawParams> OutKeyDrawParams);
 
 /** Context menu overrides */
+void ExtendSectionMenu(FMenuBuilder& OuterMenuBuilder, TArray<TMovieSceneChannelHandle<FMovieSceneDoubleChannel>>&& Channels, TArrayView<UMovieSceneSection* const> Sections, TWeakPtr<ISequencer> InSequencer);
 void ExtendSectionMenu(FMenuBuilder& OuterMenuBuilder, TArray<TMovieSceneChannelHandle<FMovieSceneFloatChannel>>&& Channels, TArrayView<UMovieSceneSection* const> Sections, TWeakPtr<ISequencer> InSequencer);
+void ExtendKeyMenu(FMenuBuilder& OuterMenuBuilder, TArray<TExtendKeyMenuParams<FMovieSceneDoubleChannel>>&& Channels, TWeakPtr<ISequencer> InSequencer);
 void ExtendKeyMenu(FMenuBuilder& OuterMenuBuilder, TArray<TExtendKeyMenuParams<FMovieSceneFloatChannel>>&& Channels, TWeakPtr<ISequencer> InSequencer);
 
 /** Curve editor models */
+inline bool SupportsCurveEditorModels(const TMovieSceneChannelHandle<FMovieSceneDoubleChannel>& DoubleChannel) { return true; }
 inline bool SupportsCurveEditorModels(const TMovieSceneChannelHandle<FMovieSceneFloatChannel>& FloatChannel) { return true; }
 inline bool SupportsCurveEditorModels(const TMovieSceneChannelHandle<FMovieSceneIntegerChannel>& IntegerChannel) { return true; }
 inline bool SupportsCurveEditorModels(const TMovieSceneChannelHandle<FMovieSceneBoolChannel>& BoolChannel) { return true; }
 inline bool SupportsCurveEditorModels(const TMovieSceneChannelHandle<FMovieSceneEventChannel>& EventChannel) { return true; }
 
+TUniquePtr<FCurveModel> CreateCurveEditorModel(const TMovieSceneChannelHandle<FMovieSceneDoubleChannel>& DoubleChannel, UMovieSceneSection* OwningSection, TSharedRef<ISequencer> InSequencer);
 TUniquePtr<FCurveModel> CreateCurveEditorModel(const TMovieSceneChannelHandle<FMovieSceneFloatChannel>& FloatChannel, UMovieSceneSection* OwningSection, TSharedRef<ISequencer> InSequencer);
 TUniquePtr<FCurveModel> CreateCurveEditorModel(const TMovieSceneChannelHandle<FMovieSceneIntegerChannel>& IntegerChannel, UMovieSceneSection* OwningSection, TSharedRef<ISequencer> InSequencer);
 TUniquePtr<FCurveModel> CreateCurveEditorModel(const TMovieSceneChannelHandle<FMovieSceneBoolChannel>& BoolChannel, UMovieSceneSection* OwningSection, TSharedRef<ISequencer> InSequencer);
diff --git a/Engine/Source/Editor/MovieSceneTools/Private/Channels/ChannelCurveModel.cpp b/Engine/Source/Editor/MovieSceneTools/Private/Channels/ChannelCurveModel.cpp
--- a/Engine/Source/Editor/MovieSceneTools/Private/Channels/ChannelCurveModel.cpp
+++ b/Engine/Source/Editor/MovieSceneTools/Private/Channels/ChannelCurveModel.cpp
@@ -412,6 +412,7 @@ void FChannelCurveModel<ChannelType, ChannelValue, KeyType>::FixupCurve()
 }
 
 // Explicit template instantiation
+template class FChannelCurveModel<FMovieSceneDoubleChannel, FMovieSceneDoubleValue, double>;
 template class FChannelCurveModel<FMovieSceneFloatChannel, FMovieSceneFloatValue, float>;
 template class FChannelCurveModel<FMovieSceneIntegerChannel, int32, int32>;
 template class FChannelCurveModel<FMovieSceneBoolChannel, bool, bool>;
diff --git a/Engine/Source/Editor/MovieSceneTools/Private/Channels/DoubleChannelKeyProxy.cpp b/Engine/Source/Editor/MovieSceneTools/Private/Channels/DoubleChannelKeyProxy.cpp
new file mode 100644
--- /dev/null
+++ b/Engine/Source/Editor/MovieSceneTools/Private/Channels/DoubleChannelKeyProxy.cpp
@@ -0,0 +1,23 @@
+// Copyright Epic Games, Inc. All Rights Reserved.
+
+#include "Channels/DoubleChannelKeyProxy.h"
+
+void UDoubleChannelKeyProxy::Initialize(FKeyHandle InKeyHandle, TMovieSceneChannelHandle<FMovieSceneDoubleChannel> InChannelHandle, TWeakObjectPtr<UMovieSceneSection> InWeakSection)
+{
+	KeyHandle     = InKeyHandle;
+	ChannelHandle = InChannelHandle;
+	WeakSection   = InWeakSection;
+}
+
+
+void UDoubleChannelKeyProxy::PostEditChangeProperty(FPropertyChangedEvent& PropertyChangedEvent)
+{
+	Super::PostEditChangeProperty(PropertyChangedEvent);
+
+	OnProxyValueChanged(ChannelHandle, WeakSection.Get(), KeyHandle, Value, Time);
+}
+
+void UDoubleChannelKeyProxy::UpdateValuesFromRawData()
+{
+	RefreshCurrentValue(ChannelHandle, KeyHandle, Value, Time);
+}
diff --git a/Engine/Source/Editor/MovieSceneTools/Private/Channels/FloatChannelCurveModel.cpp b/Engine/Source/Editor/MovieSceneTools/Private/Channels/FloatChannelCurveModel.cpp
--- a/Engine/Source/Editor/MovieSceneTools/Private/Channels/FloatChannelCurveModel.cpp
+++ b/Engine/Source/Editor/MovieSceneTools/Private/Channels/FloatChannelCurveModel.cpp
@@ -1,22 +1,8 @@
 // Copyright Epic Games, Inc. All Rights Reserved.
 
 #include "Channels/FloatChannelCurveModel.h"
-#include "Math/Vector2D.h"
-#include "HAL/PlatformMath.h"
-#include "Channels/MovieSceneFloatChannel.h"
 #include "Channels/FloatChannelKeyProxy.h"
-#include "MovieSceneSection.h"
-#include "MovieScene.h"
-#include "CurveDrawInfo.h"
-#include "CurveDataAbstraction.h"
-#include "CurveEditor.h"
 #include "CurveEditorScreenSpace.h"
-#include "CurveEditorSnapMetrics.h"
-#include "EditorStyleSet.h"
-#include "BuiltInChannelEditors.h"
-#include "SequencerChannelTraits.h"
-#include "ISequencer.h"
-#include "Channels/MovieSceneChannelProxy.h"
 
 /**
  * Buffered curve implementation for a float channel curve model, stores a copy of the float channel in order to draw itself.
@@ -55,317 +41,8 @@ private:
 };
 
 FFloatChannelCurveModel::FFloatChannelCurveModel(TMovieSceneChannelHandle<FMovieSceneFloatChannel> InChannel, UMovieSceneSection* OwningSection, TWeakPtr<ISequencer> InWeakSequencer)
-	: FChannelCurveModel<FMovieSceneFloatChannel, FMovieSceneFloatValue, float>(InChannel, OwningSection, InWeakSequencer)
+	: FBezierChannelCurveModel<FMovieSceneFloatChannel, FMovieSceneFloatValue, float>(InChannel, OwningSection, InWeakSequencer)
 {
-	FMovieSceneFloatChannel* Channel = InChannel.Get();
-
-	if (Channel && OwningSection && OwningSection->GetTypedOuter<UMovieScene>())
-	{
-		Channel->SetTickResolution(OwningSection->GetTypedOuter<UMovieScene>()->GetTickResolution());
-	}
-}
-
-void FFloatChannelCurveModel::DrawCurve(const FCurveEditor& CurveEditor, const FCurveEditorScreenSpace& ScreenSpace, TArray<TTuple<double, double>>& InterpolatingPoints) const
-{
-	FMovieSceneFloatChannel* Channel = GetChannelHandle().Get();
-	UMovieSceneSection*      Section = Cast<UMovieSceneSection>(GetOwningObject());
-
-	if (Channel && Section && Section->GetTypedOuter<UMovieScene>())
-	{
-		FFrameRate   TickResolution   = Section->GetTypedOuter<UMovieScene>()->GetTickResolution();
-
-		const double DisplayOffset    = GetInputDisplayOffset();
-		const double StartTimeSeconds = ScreenSpace.GetInputMin() - DisplayOffset;
-		const double EndTimeSeconds   = ScreenSpace.GetInputMax() - DisplayOffset;
-		const double TimeThreshold    = FMath::Max(0.0001, 1.0 / ScreenSpace.PixelsPerInput());
-		const double ValueThreshold   = FMath::Max(0.0001, 1.0 / ScreenSpace.PixelsPerOutput());
-
-		Channel->PopulateCurvePoints(StartTimeSeconds, EndTimeSeconds, TimeThreshold, ValueThreshold, TickResolution, InterpolatingPoints);
-	}
-}
-
-void FFloatChannelCurveModel::GetKeyDrawInfo(ECurvePointType PointType, const FKeyHandle InKeyHandle, FKeyDrawInfo& OutDrawInfo) const
-{
-	if (PointType == ECurvePointType::ArriveTangent || PointType == ECurvePointType::LeaveTangent)
-	{
-		OutDrawInfo.Brush = FEditorStyle::GetBrush("GenericCurveEditor.TangentHandle");
-		OutDrawInfo.ScreenSize = FVector2D(8, 8);
-	}
-	else
-	{
-		// All keys are the same size by default
-		OutDrawInfo.ScreenSize = FVector2D(11, 11);
-
-		ERichCurveInterpMode KeyType = RCIM_None;
-		ERichCurveTangentWeightMode KeyTWType = RCTWM_WeightedNone;
-
-		// Get the key type from the supplied key handle if it's valid
-		FMovieSceneFloatChannel* Channel = GetChannelHandle().Get();
-		if (Channel && InKeyHandle != FKeyHandle::Invalid())
-		{
-			TMovieSceneChannelData<FMovieSceneFloatValue> ChannelData = Channel->GetData();
-			const int32 KeyIndex = ChannelData.GetIndex(InKeyHandle);
-			if (KeyIndex != INDEX_NONE)
-			{
-				KeyType = ChannelData.GetValues()[KeyIndex].InterpMode;
-				KeyTWType = ChannelData.GetValues()[KeyIndex].Tangent.TangentWeightMode;
-			}
-		}
-
-		switch (KeyType)
-		{
-		case ERichCurveInterpMode::RCIM_Constant:
-			OutDrawInfo.Brush = FEditorStyle::GetBrush("GenericCurveEditor.ConstantKey");
-			OutDrawInfo.Tint = FLinearColor(0, 0.45f, 0.70f);
-			break;
-		case ERichCurveInterpMode::RCIM_Linear:
-			OutDrawInfo.Brush = FEditorStyle::GetBrush("GenericCurveEditor.LinearKey");
-			OutDrawInfo.Tint = FLinearColor(0, 0.62f, 0.46f);
-			break;
-		case ERichCurveInterpMode::RCIM_Cubic:
-			if (KeyTWType == ERichCurveTangentWeightMode::RCTWM_WeightedBoth)
-			{
-				OutDrawInfo.Brush = FEditorStyle::GetBrush("GenericCurveEditor.WeightedTangentCubicKey");
-			}
-			else
-			{
-				OutDrawInfo.Brush = FEditorStyle::GetBrush("GenericCurveEditor.CubicKey");
-			}
-
-			OutDrawInfo.Tint = FLinearColor::White;
-			break;
-		default:
-			OutDrawInfo.Brush = FEditorStyle::GetBrush("GenericCurveEditor.Key");
-			OutDrawInfo.Tint = FLinearColor::White;
-			break;
-		}
-
-		if (IsReadOnly())
-		{
-			OutDrawInfo.Tint = OutDrawInfo.Tint * 0.5f;
-		}
-	}
-}
-
-void FFloatChannelCurveModel::GetKeyAttributes(TArrayView<const FKeyHandle> InKeys, TArrayView<FKeyAttributes> OutAttributes) const
-{
-	FMovieSceneFloatChannel* Channel = GetChannelHandle().Get();
-	UMovieSceneSection*      Section = Cast<UMovieSceneSection>(GetOwningObject());
-	if (Channel && Section && Section->GetTypedOuter<UMovieScene>())
-	{
-		TMovieSceneChannelData<FMovieSceneFloatValue> ChannelData = Channel->GetData();
-		TArrayView<const FFrameNumber>    Times  = ChannelData.GetTimes();
-		TArrayView<FMovieSceneFloatValue> Values = ChannelData.GetValues();
-
-		float TimeInterval = Section->GetTypedOuter<UMovieScene>()->GetTickResolution().AsInterval();
-
-		for (int32 Index = 0; Index < InKeys.Num(); ++Index)
-		{
-			const int32 KeyIndex = ChannelData.GetIndex(InKeys[Index]);
-			if (KeyIndex != INDEX_NONE)
-			{
-				const FMovieSceneFloatValue& KeyValue    = Values[KeyIndex];
-				FKeyAttributes&              Attributes  = OutAttributes[Index];
-
-				Attributes.SetInterpMode(KeyValue.InterpMode);
-
-				if (KeyValue.InterpMode != RCIM_Constant && KeyValue.InterpMode != RCIM_Linear)
-				{
-					Attributes.SetTangentMode(KeyValue.TangentMode);
-					if (KeyIndex != 0)
-					{
-						Attributes.SetArriveTangent(KeyValue.Tangent.ArriveTangent / TimeInterval);
-					}
-
-					if (KeyIndex != Times.Num()-1)
-					{
-						Attributes.SetLeaveTangent(KeyValue.Tangent.LeaveTangent / TimeInterval);
-					}
-					if (KeyValue.InterpMode == RCIM_Cubic)
-					{
-						Attributes.SetTangentWeightMode(KeyValue.Tangent.TangentWeightMode);
-						if (KeyValue.Tangent.TangentWeightMode != RCTWM_WeightedNone)
-						{
-							Attributes.SetArriveTangentWeight(KeyValue.Tangent.ArriveTangentWeight);
-							Attributes.SetLeaveTangentWeight(KeyValue.Tangent.LeaveTangentWeight);
-						}
-					}
-				}
-			}
-		}
-	}
-}
-
-void FFloatChannelCurveModel::SetKeyAttributes(TArrayView<const FKeyHandle> InKeys, TArrayView<const FKeyAttributes> InAttributes, EPropertyChangeType::Type ChangeType)
-{
-	FMovieSceneFloatChannel* Channel = GetChannelHandle().Get();
-	UMovieSceneSection*      Section = Cast<UMovieSceneSection>(GetOwningObject());
-	if (Channel && Section && Section->GetTypedOuter<UMovieScene>() && !IsReadOnly())
-	{
-		bool bAutoSetTangents = false;
-		Section->MarkAsChanged();
-
-		TMovieSceneChannelData<FMovieSceneFloatValue> ChannelData = Channel->GetData();
-		TArrayView<FMovieSceneFloatValue> Values = ChannelData.GetValues();
-
-		FFrameRate TickResolution = Section->GetTypedOuter<UMovieScene>()->GetTickResolution();
-		float TimeInterval = TickResolution.AsInterval();
-
-		for (int32 Index = 0; Index < InKeys.Num(); ++Index)
-		{
-			const int32 KeyIndex = ChannelData.GetIndex(InKeys[Index]);
-			if (KeyIndex != INDEX_NONE)
-			{
-				const FKeyAttributes&  Attributes = InAttributes[Index];
-				FMovieSceneFloatValue& KeyValue   = Values[KeyIndex];
-				if (Attributes.HasInterpMode())    { KeyValue.InterpMode  = Attributes.GetInterpMode();  bAutoSetTangents = true; }
-				if (Attributes.HasTangentMode())
-				{
-					KeyValue.TangentMode = Attributes.GetTangentMode();
-					if (KeyValue.TangentMode == RCTM_Auto)
-					{
-						KeyValue.Tangent.TangentWeightMode = RCTWM_WeightedNone;
-					}
-					bAutoSetTangents = true;
-				}
-				if (Attributes.HasTangentWeightMode()) 
-				{ 
-					if (KeyValue.Tangent.TangentWeightMode == RCTWM_WeightedNone) //set tangent weights to default use
-					{
-						TArrayView<const FFrameNumber> Times = Channel->GetTimes();
-						const float OneThird = 1.0f / 3.0f;
-
-						//calculate a tangent weight based upon tangent and time difference
-						//calculate arrive tangent weight
-						if (KeyIndex > 0 )
-						{
-							const float X = TickResolution.AsSeconds(Times[KeyIndex].Value - Times[KeyIndex - 1].Value);
-							const float ArriveTangentNormal = KeyValue.Tangent.ArriveTangent / (TimeInterval);
-							const float Y = ArriveTangentNormal * X;
-							KeyValue.Tangent.ArriveTangentWeight = FMath::Sqrt(X*X + Y * Y) * OneThird;
-						}
-						//calculate leave weight
-						if(KeyIndex < ( Times.Num() - 1))
-						{
-							const float X = TickResolution.AsSeconds(Times[KeyIndex].Value - Times[KeyIndex + 1].Value);
-							const float LeaveTangentNormal = KeyValue.Tangent.LeaveTangent / (TimeInterval);
-							const float Y = LeaveTangentNormal * X;
-							KeyValue.Tangent.LeaveTangentWeight = FMath::Sqrt(X*X + Y*Y) * OneThird;
-						}
-					}
-					KeyValue.Tangent.TangentWeightMode = Attributes.GetTangentWeightMode();
-
-					if( KeyValue.Tangent.TangentWeightMode != RCTWM_WeightedNone )
-					{
-						if (KeyValue.TangentMode != RCTM_User && KeyValue.TangentMode != RCTM_Break)
-						{
-							KeyValue.TangentMode = RCTM_User;
-						}
-					}
-				}
-
-				if (Attributes.HasArriveTangent())
-				{
-					if (KeyValue.TangentMode == RCTM_Auto)
-					{
-						KeyValue.TangentMode = RCTM_User;
-						KeyValue.Tangent.TangentWeightMode = RCTWM_WeightedNone;
-					}
-
-					KeyValue.Tangent.ArriveTangent = Attributes.GetArriveTangent() * TimeInterval;
-					if (KeyValue.InterpMode == RCIM_Cubic && KeyValue.TangentMode != RCTM_Break)
-					{
-						KeyValue.Tangent.LeaveTangent = KeyValue.Tangent.ArriveTangent;
-					}
-				}
-
-				if (Attributes.HasLeaveTangent())
-				{
-					if (KeyValue.TangentMode == RCTM_Auto)
-					{
-						KeyValue.TangentMode = RCTM_User;
-						KeyValue.Tangent.TangentWeightMode = RCTWM_WeightedNone;
-					}
-
-					KeyValue.Tangent.LeaveTangent = Attributes.GetLeaveTangent() * TimeInterval;
-					if (KeyValue.InterpMode == RCIM_Cubic && KeyValue.TangentMode != RCTM_Break)
-					{
-						KeyValue.Tangent.ArriveTangent = KeyValue.Tangent.LeaveTangent;
-					}
-				}
-
-				if (Attributes.HasArriveTangentWeight())
-				{
-					if (KeyValue.TangentMode == RCTM_Auto)
-					{
-						KeyValue.TangentMode = RCTM_User;
-						KeyValue.Tangent.TangentWeightMode = RCTWM_WeightedNone;
-					}
-
-					KeyValue.Tangent.ArriveTangentWeight = Attributes.GetArriveTangentWeight(); 
-					if (KeyValue.InterpMode == RCIM_Cubic && KeyValue.TangentMode != RCTM_Break)
-					{
-						KeyValue.Tangent.LeaveTangentWeight = KeyValue.Tangent.ArriveTangentWeight;
-					}
-				}
-
-				if (Attributes.HasLeaveTangentWeight())
-				{
-				
-					if (KeyValue.TangentMode == RCTM_Auto)
-					{
-						KeyValue.TangentMode = RCTM_User;
-						KeyValue.Tangent.TangentWeightMode = RCTWM_WeightedNone;
-					}
-
-					KeyValue.Tangent.LeaveTangentWeight = Attributes.GetLeaveTangentWeight();
-					if (KeyValue.InterpMode == RCIM_Cubic && KeyValue.TangentMode != RCTM_Break)
-					{
-						KeyValue.Tangent.ArriveTangentWeight = KeyValue.Tangent.LeaveTangentWeight;
-					}
-				}
-			}
-		}
-
-		if (bAutoSetTangents)
-		{
-			Channel->AutoSetTangents();
-		}
-
-		CurveModifiedDelegate.Broadcast();
-	}
-}
-
-void FFloatChannelCurveModel::GetCurveAttributes(FCurveAttributes& OutCurveAttributes) const
-{
-	FMovieSceneFloatChannel* Channel = GetChannelHandle().Get();
-	if (Channel)
-	{
-		OutCurveAttributes.SetPreExtrapolation(Channel->PreInfinityExtrap);
-		OutCurveAttributes.SetPostExtrapolation(Channel->PostInfinityExtrap);
-	}
-}
-
-void FFloatChannelCurveModel::SetCurveAttributes(const FCurveAttributes& InCurveAttributes)
-{
-	FMovieSceneFloatChannel* Channel = GetChannelHandle().Get();
-	UMovieSceneSection*      Section = Cast<UMovieSceneSection>(GetOwningObject());
-	if (Channel && Section && !IsReadOnly())
-	{
-		Section->MarkAsChanged();
-
-		if (InCurveAttributes.HasPreExtrapolation())
-		{
-			Channel->PreInfinityExtrap = InCurveAttributes.GetPreExtrapolation();
-		}
-
-		if (InCurveAttributes.HasPostExtrapolation())
-		{
-			Channel->PostInfinityExtrap = InCurveAttributes.GetPostExtrapolation();
-		}
-
-		CurveModifiedDelegate.Broadcast();
-	}
 }
 
 void FFloatChannelCurveModel::CreateKeyProxies(TArrayView<const FKeyHandle> InKeyHandles, TArrayView<UObject*> OutObjects)
@@ -405,96 +82,3 @@ TUniquePtr<IBufferedCurveModel> FFloatChannelCurveModel::CreateBufferedCurveCopy
 	return nullptr;
 }
 
-/*	 Finds min/max for cubic curves:
-Looks for feature points in the signal(determined by change in direction of local tangent), these locations are then re-examined in closer detail recursively
-Similar to function in RichCurve but usees the Channel::Evaluate function, instead of CurveModel::Eval*/
-
-void FFloatChannelCurveModel::FeaturePointMethod(double StartTime, double EndTime, double StartValue, double Mu, int Depth, int MaxDepth, double& MaxV, double& MinVal) const
-{
-	if (Depth >= MaxDepth)
-	{
-		return;
-	}
-	double PrevValue = StartValue;
-	double EvalValue;
-	Evaluate(StartTime - Mu, EvalValue);
-	double PrevTangent = StartValue - EvalValue;
-	EndTime += Mu;
-	for (double f = StartTime + Mu; f < EndTime; f += Mu)
-	{
-		double Value;
-		Evaluate(f, Value);
-
-		MaxV = FMath::Max(Value, MaxV);
-		MinVal = FMath::Min(Value, MinVal);
-		double CurTangent = Value - PrevValue;
-		if (FMath::Sign(CurTangent) != FMath::Sign(PrevTangent))
-		{
-			//feature point centered around the previous tangent
-			double FeaturePointTime = f - Mu * 2.0f;
-			double NewVal;
-			Evaluate(FeaturePointTime, NewVal);
-			FeaturePointMethod(FeaturePointTime, f,NewVal, Mu*0.4f, Depth + 1, MaxDepth, MaxV, MinVal);
-		}
-		PrevTangent = CurTangent;
-		PrevValue = Value;
-	}
-}
-
-void FFloatChannelCurveModel::GetValueRange(double& MinValue, double& MaxValue) const
-{
-	FMovieSceneFloatChannel* Channel = GetChannelHandle().Get();
-	UMovieSceneSection*      Section = Cast<UMovieSceneSection>(GetOwningObject());
-
-	if (Channel && Section && Section->GetTypedOuter<UMovieScene>())
-	{
-		TArrayView<const FFrameNumber> Times = Channel->GetData().GetTimes();
-		TArrayView<const FMovieSceneFloatValue>   Values = Channel->GetData().GetValues();
-
-		if (Times.Num() == 0)
-		{
-			// If there are no keys we just use the default value for the channel, defaulting to zero if there is no default.
-			MinValue = MaxValue = Channel->GetDefault().Get(0.f);
-		}
-		else
-		{
-			FFrameRate TickResolution = Section->GetTypedOuter<UMovieScene>()->GetTickResolution();
-			double ToTime = TickResolution.AsInterval();
-			int32 LastKeyIndex = Values.Num() - 1;
-			MinValue = MaxValue = Values[0].Value;
-
-			for (int32 i = 0; i < Values.Num(); i++)
-			{
-				const FMovieSceneFloatValue& Key = Values[i];
-
-				MinValue = FMath::Min(MinValue,(double) Key.Value);
-				MaxValue = FMath::Max(MaxValue, (double)Key.Value);
-
-				if (Key.InterpMode == RCIM_Cubic && i != LastKeyIndex)
-				{
-					const FMovieSceneFloatValue& NextKey = Values[i + 1];
-					double KeyTime = static_cast<double>(Times[i].Value) *ToTime;
-					double NextTime = static_cast<double>(Times[i +1].Value) *ToTime;
-					double TimeStep = (NextTime - KeyTime) * 0.2f;
-					FeaturePointMethod(KeyTime, NextTime, Key.Value, TimeStep, 0, 3, MaxValue, MinValue);
-				}
-			}
-		}
-	}
-}
-
-double FFloatChannelCurveModel::GetKeyValue(TArrayView<const FMovieSceneFloatValue> Values, int32 Index) const
-{
-	return Values[Index].Value;
-}
-
-void FFloatChannelCurveModel::SetKeyValue(int32 Index, double KeyValue) const
-{
-	FMovieSceneFloatChannel* Channel = GetChannelHandle().Get();
-
-	if (Channel)
-	{
-		TMovieSceneChannelData<FMovieSceneFloatValue> ChannelData = Channel->GetData();
-		ChannelData.GetValues()[Index].Value = KeyValue;
-	}
-}
diff --git a/Engine/Source/Editor/MovieSceneTools/Private/MatineeImportTools.cpp b/Engine/Source/Editor/MovieSceneTools/Private/MatineeImportTools.cpp
--- a/Engine/Source/Editor/MovieSceneTools/Private/MatineeImportTools.cpp
+++ b/Engine/Source/Editor/MovieSceneTools/Private/MatineeImportTools.cpp
@@ -107,6 +107,16 @@ ERichCurveTangentMode FMatineeImportTools::MatineeInterpolationToRichCurveTangen
 	}
 }
 
+
+void CleanupCurveKeys(FMovieSceneDoubleChannel* InChannel)
+{
+	FKeyDataOptimizationParams Params;
+	Params.bAutoSetInterpolation = true;
+
+	UE::MovieScene::Optimize(InChannel, Params);
+}
+
+
 void CleanupCurveKeys(FMovieSceneFloatChannel* InChannel)
 {
 	FKeyDataOptimizationParams Params;
@@ -153,6 +163,25 @@ void FMatineeImportTools::SetOrAddKey(TMovieSceneChannelData<FMovieSceneFloatVal
 }
 
 
+void FMatineeImportTools::SetOrAddKey(TMovieSceneChannelData<FMovieSceneDoubleValue>& ChannelData, FFrameNumber Time, double Value, float ArriveTangent, float LeaveTangent, EInterpCurveMode MatineeInterpMode, FFrameRate FrameRate
+	, ERichCurveTangentWeightMode WeightedMode, float ArriveTangentWeight, float LeaveTangentWeight)
+{
+	if (ChannelData.FindKey(Time) == INDEX_NONE)
+	{
+		FMovieSceneDoubleValue NewKey(Value);
+
+		NewKey.InterpMode = MatineeInterpolationToRichCurveInterpolation( MatineeInterpMode );
+		NewKey.TangentMode = MatineeInterpolationToRichCurveTangent( MatineeInterpMode );
+		NewKey.Tangent.ArriveTangent = ArriveTangent / FrameRate.AsDecimal();
+		NewKey.Tangent.LeaveTangent = LeaveTangent / FrameRate.AsDecimal();
+		NewKey.Tangent.TangentWeightMode = WeightedMode;
+		NewKey.Tangent.ArriveTangentWeight = ArriveTangentWeight;
+		NewKey.Tangent.LeaveTangentWeight = LeaveTangentWeight;
+		ChannelData.AddKey( Time, NewKey );
+	}
+}
+
+
 bool FMatineeImportTools::CopyInterpBoolTrack( UInterpTrackBoolProp* MatineeBoolTrack, UMovieSceneBoolTrack* BoolTrack )
 {
 	const FScopedTransaction Transaction( NSLOCTEXT( "Sequencer", "PasteMatineeFBoolTrack", "Paste Matinee Bool Track" ) );
@@ -381,7 +410,7 @@ bool FMatineeImportTools::CopyInterpMaterialParamTrack(UInterpTrackVectorMateria
 	return bSectionCreated;
 }
 
-bool FMatineeImportTools::CopyInterpVectorTrack( UInterpTrackVectorProp* MatineeVectorTrack, UMovieSceneVectorTrack* VectorTrack )
+bool FMatineeImportTools::CopyInterpVectorTrack( UInterpTrackVectorProp* MatineeVectorTrack, UMovieSceneFloatVectorTrack* VectorTrack )
 {
 	const FScopedTransaction Transaction( NSLOCTEXT( "Sequencer", "PasteMatineeVectorTrack", "Paste Matinee Vector Track" ) );
 	bool bSectionCreated = false;
@@ -391,10 +420,10 @@ bool FMatineeImportTools::CopyInterpVectorTrack( UInterpTrackVectorProp* Matinee
 	FFrameRate   FrameRate    = VectorTrack->GetTypedOuter<UMovieScene>()->GetTickResolution();
 	FFrameNumber FirstKeyTime = (MatineeVectorTrack->GetKeyframeTime( 0 ) * FrameRate).RoundToFrame();
 
-	UMovieSceneVectorSection* Section = Cast<UMovieSceneVectorSection>( MovieSceneHelpers::FindSectionAtTime( VectorTrack->GetAllSections(), FirstKeyTime ) );
+	UMovieSceneFloatVectorSection* Section = Cast<UMovieSceneFloatVectorSection>( MovieSceneHelpers::FindSectionAtTime( VectorTrack->GetAllSections(), FirstKeyTime ) );
 	if ( Section == nullptr )
 	{
-		Section = Cast<UMovieSceneVectorSection>( VectorTrack->CreateNewSection() );
+		Section = Cast<UMovieSceneFloatVectorSection>( VectorTrack->CreateNewSection() );
 		VectorTrack->AddSection( *Section );
 		Section->SetRange(TRange<FFrameNumber>::All());
 		bSectionCreated = true;
@@ -610,18 +639,18 @@ bool FMatineeImportTools::CopyInterpMoveTrack( UInterpTrackMove* MoveTrack, UMov
 		Section->SetRange(TRange<FFrameNumber>::All());
 		bSectionCreated = true;
 
-		TArrayView<FMovieSceneFloatChannel*> FloatChannels = Section->GetChannelProxy().GetChannels<FMovieSceneFloatChannel>();
-		FloatChannels[6]->SetDefault(DefaultScale.X);
-		FloatChannels[7]->SetDefault(DefaultScale.Y);
-		FloatChannels[8]->SetDefault(DefaultScale.Z);
+		TArrayView<FMovieSceneDoubleChannel*> DoubleChannels = Section->GetChannelProxy().GetChannels<FMovieSceneDoubleChannel>();
+		DoubleChannels[6]->SetDefault(DefaultScale.X);
+		DoubleChannels[7]->SetDefault(DefaultScale.Y);
+		DoubleChannels[8]->SetDefault(DefaultScale.Z);
 	}
 
 	if (Section->TryModify())
 	{
 		TRange<FFrameNumber> KeyRange = TRange<FFrameNumber>::Empty();
 
-		TArrayView<FMovieSceneFloatChannel*> Channels = Section->GetChannelProxy().GetChannels<FMovieSceneFloatChannel>();
-		TMovieSceneChannelData<FMovieSceneFloatValue> ChannelData[6] = {
+		TArrayView<FMovieSceneDoubleChannel*> Channels = Section->GetChannelProxy().GetChannels<FMovieSceneDoubleChannel>();
+		TMovieSceneChannelData<FMovieSceneDoubleValue> ChannelData[6] = {
 			Channels[0]->GetData(), Channels[1]->GetData(), Channels[2]->GetData(),
 			Channels[3]->GetData(), Channels[4]->GetData(), Channels[5]->GetData()
 		};
diff --git a/Engine/Source/Editor/MovieSceneTools/Private/MovieSceneToolHelpers.cpp b/Engine/Source/Editor/MovieSceneTools/Private/MovieSceneToolHelpers.cpp
--- a/Engine/Source/Editor/MovieSceneTools/Private/MovieSceneToolHelpers.cpp
+++ b/Engine/Source/Editor/MovieSceneTools/Private/MovieSceneToolHelpers.cpp
@@ -20,6 +20,8 @@
 #include "EditorStyleSet.h"
 #include "Compilation/MovieSceneCompiledDataManager.h"
 #include "EditorDirectories.h"
+#include "Sections/MovieSceneDoubleSection.h"
+#include "Tracks/MovieSceneDoubleTrack.h"
 #include "Sections/MovieSceneFloatSection.h"
 #include "Tracks/MovieSceneFloatTrack.h"
 #include "Tracks/MovieSceneCameraCutTrack.h"
@@ -982,6 +984,89 @@ static FGuid GetHandleToObject(UObject* InObject, UMovieSceneSequence* InSequenc
 	return PropertyOwnerGuid;
 }
 
+template<typename TrackType, typename ChannelType>
+bool ImportFBXPropertyToPropertyTrack(UMovieScene* MovieScene, const FGuid PropertyOwnerGuid, const FMovieSceneToolsFbxSettings& FbxSetting, FRichCurve& Source)
+{
+	TrackType* PropertyTrack = MovieScene->FindTrack<TrackType>(PropertyOwnerGuid, *FbxSetting.PropertyPath.PropertyName);
+	if (!PropertyTrack)
+	{
+		MovieScene->Modify();
+		PropertyTrack = MovieScene->AddTrack<TrackType>(PropertyOwnerGuid);
+		PropertyTrack->SetPropertyNameAndPath(*FbxSetting.PropertyPath.PropertyName, *FbxSetting.PropertyPath.PropertyName);
+	}
+
+	if (!PropertyTrack)
+	{
+		return false;
+	}
+
+	PropertyTrack->Modify();
+	PropertyTrack->RemoveAllAnimationData();
+
+	FFrameRate FrameRate = PropertyTrack->template GetTypedOuter<UMovieScene>()->GetTickResolution();
+
+	bool bSectionAdded = false;
+	UMovieSceneSection* Section = Cast<UMovieSceneSection>(PropertyTrack->FindOrAddSection(0, bSectionAdded));
+	if (!Section)
+	{
+		return false;
+	}
+
+	Section->Modify();
+
+	if (bSectionAdded)
+	{
+		Section->SetRange(TRange<FFrameNumber>::All());
+	}
+
+	ChannelType* Channel = Section->GetChannelProxy().GetChannel<ChannelType>(0);
+	TMovieSceneChannelData<typename ChannelType::ChannelValueType> ChannelData = Channel->GetData();
+
+	ChannelData.Reset();
+	double DecimalRate = FrameRate.AsDecimal();
+
+	for (auto SourceIt = Source.GetKeyHandleIterator(); SourceIt; ++SourceIt)
+	{
+		FRichCurveKey &Key = Source.GetKey(*SourceIt);
+		float ArriveTangent = Key.ArriveTangent;
+		FKeyHandle PrevKeyHandle = Source.GetPreviousKey(*SourceIt);
+		if (Source.IsKeyHandleValid(PrevKeyHandle))
+		{
+			FRichCurveKey &PrevKey = Source.GetKey(PrevKeyHandle);
+			ArriveTangent = ArriveTangent / (Key.Time - PrevKey.Time);
+
+		}
+		float LeaveTangent = Key.LeaveTangent;
+		FKeyHandle NextKeyHandle = Source.GetNextKey(*SourceIt);
+		if (Source.IsKeyHandleValid(NextKeyHandle))
+		{
+			FRichCurveKey &NextKey = Source.GetKey(NextKeyHandle);
+			LeaveTangent = LeaveTangent / (NextKey.Time - Key.Time);
+		}
+
+		FFrameNumber KeyTime = (Key.Time * FrameRate).RoundToFrame();
+		FMatineeImportTools::SetOrAddKey(ChannelData, KeyTime, Key.Value, ArriveTangent, LeaveTangent,
+				MovieSceneToolHelpers::RichCurveInterpolationToMatineeInterpolation(Key.InterpMode, Key.TangentMode), FrameRate, Key.TangentWeightMode,
+				Key.ArriveTangentWeight, Key.LeaveTangentWeight);
+
+	}
+
+	Channel->AutoSetTangents();
+
+	const UMovieSceneUserImportFBXSettings* ImportFBXSettings = GetDefault<UMovieSceneUserImportFBXSettings>();
+
+	if (ImportFBXSettings->bReduceKeys)
+	{
+		FKeyDataOptimizationParams Params;
+		Params.Tolerance = ImportFBXSettings->ReduceKeysTolerance;
+		Params.DisplayRate = FrameRate;
+		Params.bAutoSetInterpolation = true; //we use this to perform the AutoSetTangents after the keys are reduced.
+		Channel->Optimize(Params);
+	}
+
+	return true;
+}
+
 bool ImportFBXProperty(FString NodeName, FString AnimatedPropertyName, FGuid ObjectBinding, UnFbx::FFbxCurvesAPI& CurveAPI, UMovieSceneSequence* InSequence, IMovieScenePlayer* Player, FMovieSceneSequenceIDRef TemplateID)
 {
 	UMovieScene* MovieScene = InSequence->GetMovieScene();
@@ -999,7 +1084,6 @@ bool ImportFBXProperty(FString NodeName, FString AnimatedPropertyName, FGuid Obj
 	}
 
 	const UMovieSceneToolsProjectSettings* ProjectSettings = GetDefault<UMovieSceneToolsProjectSettings>();
-	const UMovieSceneUserImportFBXSettings* ImportFBXSettings = GetDefault<UMovieSceneUserImportFBXSettings>();
 
 	TArrayView<TWeakObjectPtr<>> BoundObjects = Player->FindBoundObjects(ObjectBinding, TemplateID);
 
@@ -1041,78 +1125,18 @@ bool ImportFBXProperty(FString NodeName, FString AnimatedPropertyName, FGuid Obj
 				continue;
 			}
 
-			UMovieSceneFloatTrack* FloatTrack = MovieScene->FindTrack<UMovieSceneFloatTrack>(PropertyOwnerGuid, *FbxSetting.PropertyPath.PropertyName);
-			if (!FloatTrack)
+			bool bImported = false;
+			switch (FbxSetting.PropertyType)
 			{
-				MovieScene->Modify();
-				FloatTrack = MovieScene->AddTrack<UMovieSceneFloatTrack>(PropertyOwnerGuid);
-				FloatTrack->SetPropertyNameAndPath(*FbxSetting.PropertyPath.PropertyName, *FbxSetting.PropertyPath.PropertyName);
+				case EMovieSceneToolsPropertyTrackType::FloatTrack:
+					bImported = ImportFBXPropertyToPropertyTrack<UMovieSceneFloatTrack, FMovieSceneFloatChannel>(MovieScene, PropertyOwnerGuid, FbxSetting, Source);
+					break;
+				case EMovieSceneToolsPropertyTrackType::DoubleTrack:
+					bImported = ImportFBXPropertyToPropertyTrack<UMovieSceneDoubleTrack, FMovieSceneDoubleChannel>(MovieScene, PropertyOwnerGuid, FbxSetting, Source);
+					break;
 			}
-
-			if (FloatTrack)
+			if (bImported)
 			{
-				FloatTrack->Modify();
-				FloatTrack->RemoveAllAnimationData();
-
-				FFrameRate FrameRate = FloatTrack->GetTypedOuter<UMovieScene>()->GetTickResolution();
-
-				bool bSectionAdded = false;
-				UMovieSceneFloatSection* FloatSection = Cast<UMovieSceneFloatSection>(FloatTrack->FindOrAddSection(0, bSectionAdded));
-				if (!FloatSection)
-				{
-					continue;
-				}
-
-				FloatSection->Modify();
-
-				if (bSectionAdded)
-				{
-					FloatSection->SetRange(TRange<FFrameNumber>::All());
-				}
-
-				FMovieSceneFloatChannel* Channel = FloatSection->GetChannelProxy().GetChannel<FMovieSceneFloatChannel>(0);
-				TMovieSceneChannelData<FMovieSceneFloatValue> ChannelData = Channel->GetData();
-
-				ChannelData.Reset();
-				double DecimalRate = FrameRate.AsDecimal();
-
-				for (auto SourceIt = Source.GetKeyHandleIterator(); SourceIt; ++SourceIt)
-				{
-					FRichCurveKey &Key = Source.GetKey(*SourceIt);
-					float ArriveTangent = Key.ArriveTangent;
-					FKeyHandle PrevKeyHandle = Source.GetPreviousKey(*SourceIt);
-					if (Source.IsKeyHandleValid(PrevKeyHandle))
-					{
-						FRichCurveKey &PrevKey = Source.GetKey(PrevKeyHandle);
-						ArriveTangent = ArriveTangent / (Key.Time - PrevKey.Time);
-
-					}
-					float LeaveTangent = Key.LeaveTangent;
-					FKeyHandle NextKeyHandle = Source.GetNextKey(*SourceIt);
-					if (Source.IsKeyHandleValid(NextKeyHandle))
-					{
-						FRichCurveKey &NextKey = Source.GetKey(NextKeyHandle);
-						LeaveTangent = LeaveTangent / (NextKey.Time - Key.Time);
-					}
-
-					FFrameNumber KeyTime = (Key.Time * FrameRate).RoundToFrame();
-					FMatineeImportTools::SetOrAddKey(ChannelData, KeyTime, Key.Value, ArriveTangent, LeaveTangent,
-						MovieSceneToolHelpers::RichCurveInterpolationToMatineeInterpolation(Key.InterpMode, Key.TangentMode), FrameRate, Key.TangentWeightMode,
-						Key.ArriveTangentWeight, Key.LeaveTangentWeight);
-
-				}
-
-				Channel->AutoSetTangents();
-
-				if (ImportFBXSettings->bReduceKeys)
-				{
-					FKeyDataOptimizationParams Params;
-					Params.Tolerance = ImportFBXSettings->ReduceKeysTolerance;
-					Params.DisplayRate = FrameRate;
-					Params.bAutoSetInterpolation = true; //we use this to perform the AutoSetTangents after the keys are reduced.
-					Channel->Optimize(Params);
-				}
-
 				return true;
 			}
 		}
@@ -1180,7 +1204,8 @@ void MovieSceneToolHelpers::CreateCameraCutSectionForCamera(UMovieScene* OwnerMo
 	}
 }
 
-void ImportTransformChannel(const FRichCurve& Source, FMovieSceneFloatChannel* Dest, FFrameRate DestFrameRate, bool bNegateTangents, bool bClearChannel,FFrameNumber StartFrame = 0, bool bNegateValue = false)
+template<typename ChannelType>
+void ImportTransformChannelToBezierChannel(const FRichCurve& Source, ChannelType* Dest, FFrameRate DestFrameRate, bool bNegateTangents, bool bClearChannel,FFrameNumber StartFrame = 0, bool bNegateValue = false)
 {
 	// If there are no keys, don't clear the existing channel
 	if (!Source.GetNumKeys())
@@ -1188,7 +1213,7 @@ void ImportTransformChannel(const FRichCurve& Source, FMovieSceneFloatChannel* D
 		return;
 	}
 
-	TMovieSceneChannelData<FMovieSceneFloatValue> ChannelData = Dest->GetData();
+	TMovieSceneChannelData<typename ChannelType::ChannelValueType> ChannelData = Dest->GetData();
 
 	if (bClearChannel)
 	{
@@ -1239,6 +1264,16 @@ void ImportTransformChannel(const FRichCurve& Source, FMovieSceneFloatChannel* D
 	}
 }
 
+void ImportTransformChannelToDouble(const FRichCurve& Source, FMovieSceneDoubleChannel* Dest, FFrameRate DestFrameRate, bool bNegateTangents, bool bClearChannel, FFrameNumber StartFrame = 0, bool bNegateValue = false)
+{
+	ImportTransformChannelToBezierChannel(Source, Dest, DestFrameRate, bNegateTangents, bClearChannel, StartFrame, bNegateValue);
+}
+
+void ImportTransformChannelToFloat(const FRichCurve& Source, FMovieSceneFloatChannel* Dest, FFrameRate DestFrameRate, bool bNegateTangents, bool bClearChannel, FFrameNumber StartFrame = 0, bool bNegateValue = false)
+{
+	ImportTransformChannelToBezierChannel(Source, Dest, DestFrameRate, bNegateTangents, bClearChannel, StartFrame, bNegateValue);
+}
+
 void ImportTransformChannelToBool(const FRichCurve& Source, FMovieSceneBoolChannel* Dest, FFrameRate DestFrameRate, bool bClearChannel, FFrameNumber StartFrame)
 {
 	// If there are no keys, don't clear the existing channel
@@ -1327,7 +1362,7 @@ void ImportTransformChannelToInteger(const FRichCurve& Source, FMovieSceneIntege
 	}
 }
 
-void SetChannelValue(FMovieSceneFloatChannel* Channel, FMovieSceneBoolChannel *BoolChannel, FMovieSceneByteChannel *EnumChannel, FMovieSceneIntegerChannel* IntegerChannel,
+void SetChannelValue(FMovieSceneDoubleChannel* DoubleChannel, FMovieSceneFloatChannel* FloatChannel, FMovieSceneBoolChannel *BoolChannel, FMovieSceneByteChannel *EnumChannel, FMovieSceneIntegerChannel* IntegerChannel,
 	FFrameRate FrameRate, FFrameNumber StartFrame,
 	FControlRigChannelEnum ChannelEnum, UMovieSceneUserImportFBXControlRigSettings* ImportFBXControlRigSettings,
 	FTransform& DefaultTransform,
@@ -1362,11 +1397,11 @@ void SetChannelValue(FMovieSceneFloatChannel* Channel, FMovieSceneBoolChannel *B
 					IntegerChannel->SetDefault(bDefault);
 					ImportTransformChannelToInteger(TranslationX, IntegerChannel, FrameRate, false, StartFrame);
 				}
-				else if (Channel)
+				else if (DoubleChannel)
 				{
 					float Default = bNegate ? -Location.X : Location.X;
-					Channel->SetDefault(Default);
-					ImportTransformChannel(TranslationX, Channel, FrameRate, false, false, StartFrame, bNegate);
+					DoubleChannel->SetDefault(Default);
+					ImportTransformChannelToDouble(TranslationX, DoubleChannel, FrameRate, false, false, StartFrame, bNegate);
 				}
 			}
 			else if (Mapping.FBXChannel == FTransformChannelEnum::TranslateY)
@@ -1389,12 +1424,12 @@ void SetChannelValue(FMovieSceneFloatChannel* Channel, FMovieSceneBoolChannel *B
 					IntegerChannel->SetDefault(bDefault);
 					ImportTransformChannelToInteger(TranslationY, IntegerChannel, FrameRate, false, StartFrame);
 				}
-				else if (Channel)
+				else if (DoubleChannel)
 				{
 					bNegate = !bNegate;
 					float Default = bNegate ? -Location.Y : Location.Y;
-					Channel->SetDefault(Default);
-					ImportTransformChannel(TranslationY, Channel, FrameRate, false, false, StartFrame, bNegate);
+					DoubleChannel->SetDefault(Default);
+					ImportTransformChannelToDouble(TranslationY, DoubleChannel, FrameRate, false, false, StartFrame, bNegate);
 				}
 			}
 			else if (Mapping.FBXChannel == FTransformChannelEnum::TranslateZ)
@@ -1417,11 +1452,11 @@ void SetChannelValue(FMovieSceneFloatChannel* Channel, FMovieSceneBoolChannel *B
 					IntegerChannel->SetDefault(bDefault);
 					ImportTransformChannelToInteger(TranslationZ, IntegerChannel, FrameRate, false, StartFrame);
 				}
-				else if (Channel)
+				else if (DoubleChannel)
 				{
 					float Default = bNegate ? -Location.Z : Location.Z;
-					Channel->SetDefault(Default);
-					ImportTransformChannel(TranslationZ, Channel, FrameRate, false, false, StartFrame, bNegate);
+					DoubleChannel->SetDefault(Default);
+					ImportTransformChannelToDouble(TranslationZ, DoubleChannel, FrameRate, false, false, StartFrame, bNegate);
 				}
 			}
 			else if (Mapping.FBXChannel == FTransformChannelEnum::RotateX)
@@ -1444,11 +1479,11 @@ void SetChannelValue(FMovieSceneFloatChannel* Channel, FMovieSceneBoolChannel *B
 					IntegerChannel->SetDefault(bDefault);
 					ImportTransformChannelToInteger(EulerRotationX, IntegerChannel, FrameRate, false, StartFrame);
 				}
-				else if (Channel)
+				else if (DoubleChannel)
 				{
 					float Default = bNegate ? -Rotation.X : Rotation.X;
-					Channel->SetDefault(Default);
-					ImportTransformChannel(EulerRotationX, Channel, FrameRate, false, false, StartFrame, bNegate);
+					DoubleChannel->SetDefault(Default);
+					ImportTransformChannelToDouble(EulerRotationX, DoubleChannel, FrameRate, false, false, StartFrame, bNegate);
 				}
 			}
 			else if (Mapping.FBXChannel == FTransformChannelEnum::RotateY)
@@ -1471,11 +1506,11 @@ void SetChannelValue(FMovieSceneFloatChannel* Channel, FMovieSceneBoolChannel *B
 					IntegerChannel->SetDefault(bDefault);
 					ImportTransformChannelToInteger(EulerRotationY, IntegerChannel, FrameRate, false, StartFrame);
 				}
-				else if (Channel)
+				else if (DoubleChannel)
 				{
 					float Default = bNegate ? -Rotation.Y : Rotation.Y;
-					Channel->SetDefault(Default);
-					ImportTransformChannel(EulerRotationY, Channel, FrameRate, false, false, StartFrame, bNegate);
+					DoubleChannel->SetDefault(Default);
+					ImportTransformChannelToDouble(EulerRotationY, DoubleChannel, FrameRate, false, false, StartFrame, bNegate);
 				}
 			}
 			else if (Mapping.FBXChannel == FTransformChannelEnum::RotateZ)
@@ -1498,11 +1533,11 @@ void SetChannelValue(FMovieSceneFloatChannel* Channel, FMovieSceneBoolChannel *B
 					IntegerChannel->SetDefault(bDefault);
 					ImportTransformChannelToInteger(EulerRotationZ, IntegerChannel, FrameRate, false, StartFrame);
 				}
-				else if (Channel)
+				else if (DoubleChannel)
 				{
 					float Default = bNegate ? -Rotation.Z : Rotation.Z;
-					Channel->SetDefault(Default);
-					ImportTransformChannel(EulerRotationZ, Channel, FrameRate, false, false, StartFrame, bNegate);
+					DoubleChannel->SetDefault(Default);
+					ImportTransformChannelToDouble(EulerRotationZ, DoubleChannel, FrameRate, false, false, StartFrame, bNegate);
 				}
 			}
 			else if (Mapping.FBXChannel == FTransformChannelEnum::ScaleX)
@@ -1525,11 +1560,11 @@ void SetChannelValue(FMovieSceneFloatChannel* Channel, FMovieSceneBoolChannel *B
 					IntegerChannel->SetDefault(bDefault);
 					ImportTransformChannelToInteger(ScaleX, IntegerChannel, FrameRate, false, StartFrame);
 				}
-				else if (Channel)
+				else if (DoubleChannel)
 				{
 					float Default = bNegate ? -Scale3D.X : Scale3D.X;
-					Channel->SetDefault(Default);
-					ImportTransformChannel(ScaleX, Channel, FrameRate, false, false, StartFrame, bNegate);
+					DoubleChannel->SetDefault(Default);
+					ImportTransformChannelToDouble(ScaleX, DoubleChannel, FrameRate, false, false, StartFrame, bNegate);
 				}
 			}
 			else if (Mapping.FBXChannel == FTransformChannelEnum::ScaleY)
@@ -1552,11 +1587,11 @@ void SetChannelValue(FMovieSceneFloatChannel* Channel, FMovieSceneBoolChannel *B
 					IntegerChannel->SetDefault(bDefault);
 					ImportTransformChannelToInteger(ScaleY, IntegerChannel, FrameRate, false, StartFrame);
 				}
-				else if (Channel)
+				else if (DoubleChannel)
 				{
 					float Default = bNegate ? -Scale3D.Y : Scale3D.Y;
-					Channel->SetDefault(Default);
-					ImportTransformChannel(ScaleY, Channel, FrameRate, false, false, StartFrame, bNegate);
+					DoubleChannel->SetDefault(Default);
+					ImportTransformChannelToDouble(ScaleY, DoubleChannel, FrameRate, false, false, StartFrame, bNegate);
 				}
 			}
 			else if (Mapping.FBXChannel == FTransformChannelEnum::ScaleZ)
@@ -1579,11 +1614,11 @@ void SetChannelValue(FMovieSceneFloatChannel* Channel, FMovieSceneBoolChannel *B
 					IntegerChannel->SetDefault(bDefault);
 					ImportTransformChannelToInteger(ScaleZ, IntegerChannel, FrameRate, false, StartFrame);
 				}
-				else if (Channel)
+				else if (DoubleChannel)
 				{
 					float Default = bNegate ? -Scale3D.Z : Scale3D.Z;
-					Channel->SetDefault(Default);
-					ImportTransformChannel(ScaleZ, Channel, FrameRate, false, false, StartFrame, bNegate);
+					DoubleChannel->SetDefault(Default);
+					ImportTransformChannelToDouble(ScaleZ, DoubleChannel, FrameRate, false, false, StartFrame, bNegate);
 				}
 			}
 		}
@@ -1614,85 +1649,85 @@ void SetChannelValue(FMovieSceneFloatChannel* Channel, FMovieSceneBoolChannel *B
 	{
 		bool bNegate = false;
 		float Default = Location.X;
-		Channel->SetDefault(Default);
-		ImportTransformChannel(TranslationX, Channel, FrameRate, false,false, StartFrame, bNegate);
+		FloatChannel->SetDefault(Default);
+		ImportTransformChannelToFloat(TranslationX, FloatChannel, FrameRate, false,false, StartFrame, bNegate);
 	}
 	else if (ChannelEnum == FControlRigChannelEnum::Vector2DX)
 	{
 		bool bNegate = false;
 		float Default = Location.X;
-		Channel->SetDefault(Default);
-		ImportTransformChannel(TranslationX, Channel, FrameRate, false, false,StartFrame, bNegate);
+		FloatChannel->SetDefault(Default);
+		ImportTransformChannelToFloat(TranslationX, FloatChannel, FrameRate, false, false,StartFrame, bNegate);
 	}
 	else if (ChannelEnum == FControlRigChannelEnum::Vector2DY)
 	{
 		bool bNegate = true;
 		float Default = -Location.Y;
-		Channel->SetDefault(Default);
-		ImportTransformChannel(TranslationY, Channel, FrameRate, false, false,StartFrame, bNegate);
+		FloatChannel->SetDefault(Default);
+		ImportTransformChannelToFloat(TranslationY, FloatChannel, FrameRate, false, false,StartFrame, bNegate);
 	}
 	else if (ChannelEnum == FControlRigChannelEnum::PositionX)
 	{
 		bool bNegate = false;
 		float Default = Location.X;
-		Channel->SetDefault(Default);
-		ImportTransformChannel(TranslationX, Channel, FrameRate, false, false, StartFrame, bNegate);
+		DoubleChannel->SetDefault(Default);
+		ImportTransformChannelToDouble(TranslationX, DoubleChannel, FrameRate, false, false, StartFrame, bNegate);
 	}
 	else if (ChannelEnum == FControlRigChannelEnum::PositionY)
 	{
 		bool bNegate = true;
 		float Default = -Location.Y;
-		Channel->SetDefault(Default);
-		ImportTransformChannel(TranslationX, Channel, FrameRate, false, false, StartFrame, bNegate);
+		DoubleChannel->SetDefault(Default);
+		ImportTransformChannelToDouble(TranslationX, DoubleChannel, FrameRate, false, false, StartFrame, bNegate);
 	}
 	else if (ChannelEnum == FControlRigChannelEnum::PositionZ)
 	{
 		bool bNegate = false;
 		float Default = Location.Z;
-		Channel->SetDefault(Default);
-		ImportTransformChannel(TranslationZ, Channel, FrameRate, false, false, StartFrame, bNegate);
+		DoubleChannel->SetDefault(Default);
+		ImportTransformChannelToDouble(TranslationZ, DoubleChannel, FrameRate, false, false, StartFrame, bNegate);
 	}
 	else if (ChannelEnum == FControlRigChannelEnum::RotatorX)
 	{
 		bool bNegate = false;
 		float Default = Rotation.X;
-		Channel->SetDefault(Default);
-		ImportTransformChannel(EulerRotationX, Channel, FrameRate, false, false, StartFrame, bNegate);
+		DoubleChannel->SetDefault(Default);
+		ImportTransformChannelToDouble(EulerRotationX, DoubleChannel, FrameRate, false, false, StartFrame, bNegate);
 	}
 	else if (ChannelEnum == FControlRigChannelEnum::RotatorY)
 	{
 		bool bNegate = false;
 		float Default = Rotation.Y;
-		Channel->SetDefault(Default);		
-		ImportTransformChannel(EulerRotationY, Channel, FrameRate, false, false, StartFrame, bNegate);
+		DoubleChannel->SetDefault(Default);		
+		ImportTransformChannelToDouble(EulerRotationY, DoubleChannel, FrameRate, false, false, StartFrame, bNegate);
 	}
 	else if (ChannelEnum == FControlRigChannelEnum::RotatorZ)
 	{
 		bool bNegate = false;
 		float Default =  Rotation.Z;
-		Channel->SetDefault(Default);
-		ImportTransformChannel(EulerRotationZ, Channel, FrameRate, false, false, StartFrame, bNegate);
+		DoubleChannel->SetDefault(Default);
+		ImportTransformChannelToDouble(EulerRotationZ, DoubleChannel, FrameRate, false, false, StartFrame, bNegate);
 	}
 	else if (ChannelEnum == FControlRigChannelEnum::ScaleX)
 	{
 		bool bNegate = false;
 		float Default = Scale3D.X;
-		Channel->SetDefault(Default); 
-		ImportTransformChannel(ScaleX, Channel, FrameRate, false, false, StartFrame, bNegate);
+		DoubleChannel->SetDefault(Default); 
+		ImportTransformChannelToDouble(ScaleX, DoubleChannel, FrameRate, false, false, StartFrame, bNegate);
 	}
 	else if (ChannelEnum == FControlRigChannelEnum::ScaleY)
 	{
 		bool bNegate = false;
 		float Default = Scale3D.Y;
-		Channel->SetDefault(Default); 
-		ImportTransformChannel(ScaleY, Channel, FrameRate, false, false, StartFrame, bNegate);
+		DoubleChannel->SetDefault(Default); 
+		ImportTransformChannelToDouble(ScaleY, DoubleChannel, FrameRate, false, false, StartFrame, bNegate);
 	}
 	else if (ChannelEnum == FControlRigChannelEnum::ScaleZ)
 	{
 		bool bNegate = false;
 		float Default = Scale3D.Z;
-		Channel->SetDefault(Default);	
-		ImportTransformChannel(ScaleZ, Channel, FrameRate, false,false,  StartFrame, bNegate);
+		DoubleChannel->SetDefault(Default);	
+		ImportTransformChannelToDouble(ScaleZ, DoubleChannel, FrameRate, false,false,  StartFrame, bNegate);
 	}
 
 }
@@ -1704,7 +1739,8 @@ static bool ImportFBXTransformToChannels(FString NodeName, const UMovieSceneUser
 	 UnFbx::FFbxCurvesAPI& CurveAPI)
 {
 
-	TArray<FMovieSceneFloatChannel*>& Channels = NodeAndChannels.Channels;
+	TArray<FMovieSceneDoubleChannel*>& DoubleChannels = NodeAndChannels.DoubleChannels;
+	TArray<FMovieSceneFloatChannel*>& FloatChannels = NodeAndChannels.FloatChannels;
 	TArray<FMovieSceneBoolChannel*>& BoolChannels = NodeAndChannels.BoolChannels;
 	TArray<FMovieSceneByteChannel*>& EnumChannels = NodeAndChannels.EnumChannels;
 	TArray<FMovieSceneIntegerChannel*>& IntegerChannels = NodeAndChannels.IntegerChannels;
@@ -1726,7 +1762,7 @@ static bool ImportFBXTransformToChannels(FString NodeName, const UMovieSceneUser
 	if (BoolChannels.Num() == 1)
 	{
 		FControlRigChannelEnum Channel = FControlRigChannelEnum::Bool;
-		SetChannelValue(nullptr, BoolChannels[0], nullptr,nullptr, FrameRate, StartFrame,
+		SetChannelValue(nullptr, nullptr, BoolChannels[0], nullptr,nullptr, FrameRate, StartFrame,
 			Channel, ImportFBXControlRigSettings, DefaultTransform,
 			Translation[0], Translation[1], Translation[2], EulerRotation[0],
 			EulerRotation[1], EulerRotation[2], Scale[0], Scale[1], Scale[2]);
@@ -1735,7 +1771,7 @@ static bool ImportFBXTransformToChannels(FString NodeName, const UMovieSceneUser
 	if (EnumChannels.Num() == 1)
 	{
 		FControlRigChannelEnum Channel = FControlRigChannelEnum::Enum;
-		SetChannelValue(nullptr, nullptr, EnumChannels[0], nullptr, FrameRate, StartFrame,
+		SetChannelValue(nullptr, nullptr, nullptr, EnumChannels[0], nullptr, FrameRate, StartFrame,
 			Channel, ImportFBXControlRigSettings, DefaultTransform,
 			Translation[0], Translation[1], Translation[2], EulerRotation[0],
 			EulerRotation[1], EulerRotation[2], Scale[0], Scale[1], Scale[2]);
@@ -1744,54 +1780,54 @@ static bool ImportFBXTransformToChannels(FString NodeName, const UMovieSceneUser
 	if (IntegerChannels.Num() == 1)
 	{
 		FControlRigChannelEnum Channel = FControlRigChannelEnum::Integer;
-		SetChannelValue(nullptr,nullptr, nullptr, IntegerChannels[0], FrameRate, StartFrame,
+		SetChannelValue(nullptr, nullptr, nullptr, nullptr, IntegerChannels[0], FrameRate, StartFrame,
 			Channel,  ImportFBXControlRigSettings,DefaultTransform,
 			Translation[0], Translation[1], Translation[2], EulerRotation[0], 
 			EulerRotation[1], EulerRotation[2], Scale[0], Scale[1], Scale[2]);
 	}
 
-	if (Channels.Num() == 1)
+	if (FloatChannels.Num() == 1)
 	{
 		FControlRigChannelEnum Channel = FControlRigChannelEnum::Float;
-		SetChannelValue(Channels[0], nullptr, nullptr, nullptr, FrameRate, StartFrame,
+		SetChannelValue(nullptr, FloatChannels[0], nullptr, nullptr, nullptr, FrameRate, StartFrame,
 			Channel, ImportFBXControlRigSettings, DefaultTransform,
 			Translation[0], Translation[1], Translation[2], EulerRotation[0],
 			EulerRotation[1], EulerRotation[2], Scale[0], Scale[1], Scale[2]);
 	}
-	else if (Channels.Num() == 2)
+	else if (DoubleChannels.Num() == 2)
 	{
 		FControlRigChannelEnum Channel = FControlRigChannelEnum::Vector2DX;
 
-		SetChannelValue(Channels[0],nullptr, nullptr, nullptr, FrameRate, StartFrame,
+		SetChannelValue(nullptr, FloatChannels[0], nullptr, nullptr, nullptr, FrameRate, StartFrame,
 			Channel, ImportFBXControlRigSettings, DefaultTransform,
 			Translation[0], Translation[1], Translation[2], EulerRotation[0],
 			EulerRotation[1], EulerRotation[2], Scale[0], Scale[1], Scale[2]);
 
 		Channel = FControlRigChannelEnum::Vector2DY;
-		SetChannelValue(Channels[1],nullptr, nullptr, nullptr, FrameRate, StartFrame,
+		SetChannelValue(nullptr, FloatChannels[1], nullptr, nullptr, nullptr, FrameRate, StartFrame,
 			Channel, ImportFBXControlRigSettings, DefaultTransform,
 			Translation[0], Translation[1], Translation[2], EulerRotation[0],
 			EulerRotation[1], EulerRotation[2], Scale[0], Scale[1], Scale[2]);
 	}
-	else if (Channels.Num() == 3)
+	else if (DoubleChannels.Num() == 3)
 	{
 		if (NodeAndChannels.ControlType == FFBXControlRigTypeProxyEnum::Position)
 		{
 			FControlRigChannelEnum Channel = FControlRigChannelEnum::PositionX;
 
-			SetChannelValue(Channels[0],nullptr, nullptr, nullptr, FrameRate, StartFrame,
+			SetChannelValue(DoubleChannels[0], nullptr, nullptr, nullptr, nullptr, FrameRate, StartFrame,
 				Channel, ImportFBXControlRigSettings, DefaultTransform,
 				Translation[0], Translation[1], Translation[2], EulerRotation[0],
 				EulerRotation[1], EulerRotation[2], Scale[0], Scale[1], Scale[2]);
 
 			Channel = FControlRigChannelEnum::PositionY;
-			SetChannelValue(Channels[1], nullptr, nullptr, nullptr, FrameRate, StartFrame,
+			SetChannelValue(DoubleChannels[1], nullptr, nullptr, nullptr, nullptr, FrameRate, StartFrame,
 				Channel, ImportFBXControlRigSettings, DefaultTransform,
 				Translation[0], Translation[1], Translation[2], EulerRotation[0],
 				EulerRotation[1], EulerRotation[2], Scale[0], Scale[1], Scale[2]);
 
 			Channel = FControlRigChannelEnum::PositionZ;
-			SetChannelValue(Channels[2],nullptr, nullptr, nullptr, FrameRate, StartFrame,
+			SetChannelValue(DoubleChannels[2], nullptr,nullptr, nullptr, nullptr, FrameRate, StartFrame,
 				Channel, ImportFBXControlRigSettings, DefaultTransform,
 				Translation[0], Translation[1], Translation[2], EulerRotation[0],
 				EulerRotation[1], EulerRotation[2], Scale[0], Scale[1], Scale[2]);
@@ -1799,19 +1835,19 @@ static bool ImportFBXTransformToChannels(FString NodeName, const UMovieSceneUser
 		else if (NodeAndChannels.ControlType == FFBXControlRigTypeProxyEnum::Rotator)
 		{
 			FControlRigChannelEnum Channel = FControlRigChannelEnum::RotatorX;
-			SetChannelValue(Channels[0],nullptr, nullptr, nullptr, FrameRate, StartFrame,
+			SetChannelValue(DoubleChannels[0], nullptr,nullptr, nullptr, nullptr, FrameRate, StartFrame,
 				Channel, ImportFBXControlRigSettings, DefaultTransform,
 				Translation[0], Translation[1], Translation[2], EulerRotation[0],
 				EulerRotation[1], EulerRotation[2], Scale[0], Scale[1], Scale[2]);
 
 			Channel = FControlRigChannelEnum::RotatorY;
-			SetChannelValue(Channels[1], nullptr, nullptr, nullptr, FrameRate, StartFrame,
+			SetChannelValue(DoubleChannels[1], nullptr, nullptr, nullptr, nullptr, FrameRate, StartFrame,
 				Channel, ImportFBXControlRigSettings, DefaultTransform,
 				Translation[0], Translation[1], Translation[2], EulerRotation[0],
 				EulerRotation[1], EulerRotation[2], Scale[0], Scale[1], Scale[2]);
 
 			Channel = FControlRigChannelEnum::RotatorZ;
-			SetChannelValue(Channels[2], nullptr, nullptr, nullptr, FrameRate, StartFrame,
+			SetChannelValue(DoubleChannels[2], nullptr, nullptr, nullptr, nullptr, FrameRate, StartFrame,
 				Channel, ImportFBXControlRigSettings, DefaultTransform,
 				Translation[0], Translation[1], Translation[2], EulerRotation[0],
 				EulerRotation[1], EulerRotation[2], Scale[0], Scale[1], Scale[2]);
@@ -1819,52 +1855,52 @@ static bool ImportFBXTransformToChannels(FString NodeName, const UMovieSceneUser
 		else if (NodeAndChannels.ControlType == FFBXControlRigTypeProxyEnum::Scale)
 		{
 			FControlRigChannelEnum Channel = FControlRigChannelEnum::ScaleX;
-			SetChannelValue(Channels[0], nullptr, nullptr, nullptr, FrameRate, StartFrame,
+			SetChannelValue(DoubleChannels[0], nullptr, nullptr, nullptr, nullptr, FrameRate, StartFrame,
 				Channel, ImportFBXControlRigSettings, DefaultTransform,
 				Translation[0], Translation[1], Translation[2], EulerRotation[0],
 				EulerRotation[1], EulerRotation[2], Scale[0], Scale[1], Scale[2]);
 			Channel = FControlRigChannelEnum::ScaleY;
-			SetChannelValue(Channels[1], nullptr, nullptr, nullptr, FrameRate, StartFrame,
+			SetChannelValue(DoubleChannels[1], nullptr, nullptr, nullptr, nullptr, FrameRate, StartFrame,
 				Channel, ImportFBXControlRigSettings, DefaultTransform,
 				Translation[0], Translation[1], Translation[2], EulerRotation[0],
 				EulerRotation[1], EulerRotation[2], Scale[0], Scale[1], Scale[2]);
 			Channel = FControlRigChannelEnum::ScaleZ;
-			SetChannelValue(Channels[2], nullptr, nullptr, nullptr, FrameRate, StartFrame,
+			SetChannelValue(DoubleChannels[2], nullptr, nullptr, nullptr, nullptr, FrameRate, StartFrame,
 				Channel, ImportFBXControlRigSettings, DefaultTransform,
 				Translation[0], Translation[1], Translation[2], EulerRotation[0],
 				EulerRotation[1], EulerRotation[2], Scale[0], Scale[1], Scale[2]);
 		}
 	}
-	else if (Channels.Num() == 9 || Channels.Num() == 6)
+	else if (DoubleChannels.Num() == 9 || DoubleChannels.Num() == 6)
 	{
-		Channels[0]->SetDefault(Location.X);
-		Channels[1]->SetDefault(Location.Y);
-		Channels[2]->SetDefault(Location.Z);
+		DoubleChannels[0]->SetDefault(Location.X);
+		DoubleChannels[1]->SetDefault(Location.Y);
+		DoubleChannels[2]->SetDefault(Location.Z);
 
-		Channels[3]->SetDefault(Rotation.X);
-		Channels[4]->SetDefault(Rotation.Y);
-		Channels[5]->SetDefault(Rotation.Z);
+		DoubleChannels[3]->SetDefault(Rotation.X);
+		DoubleChannels[4]->SetDefault(Rotation.Y);
+		DoubleChannels[5]->SetDefault(Rotation.Z);
 
-		if (Channels.Num() > 6) //noscale
+		if (DoubleChannels.Num() > 6) //noscale
 		{
-			Channels[6]->SetDefault(Scale3D.X);
-			Channels[7]->SetDefault(Scale3D.Y);
-			Channels[8]->SetDefault(Scale3D.Z);
+			DoubleChannels[6]->SetDefault(Scale3D.X);
+			DoubleChannels[7]->SetDefault(Scale3D.Y);
+			DoubleChannels[8]->SetDefault(Scale3D.Z);
 		}
 
-		ImportTransformChannel(Translation[0], Channels[0], FrameRate, false, false, StartFrame);
-		ImportTransformChannel(Translation[1], Channels[1], FrameRate, true, false, StartFrame);
-		ImportTransformChannel(Translation[2], Channels[2], FrameRate, false, false, StartFrame);
+		ImportTransformChannelToDouble(Translation[0], DoubleChannels[0], FrameRate, false, false, StartFrame);
+		ImportTransformChannelToDouble(Translation[1], DoubleChannels[1], FrameRate, true, false, StartFrame);
+		ImportTransformChannelToDouble(Translation[2], DoubleChannels[2], FrameRate, false, false, StartFrame);
 
-		ImportTransformChannel(EulerRotation[0], Channels[3], FrameRate, false,false, StartFrame);
-		ImportTransformChannel(EulerRotation[1], Channels[4], FrameRate, true, false, StartFrame);
-		ImportTransformChannel(EulerRotation[2], Channels[5], FrameRate, true, false, StartFrame);
+		ImportTransformChannelToDouble(EulerRotation[0], DoubleChannels[3], FrameRate, false,false, StartFrame);
+		ImportTransformChannelToDouble(EulerRotation[1], DoubleChannels[4], FrameRate, true, false, StartFrame);
+		ImportTransformChannelToDouble(EulerRotation[2], DoubleChannels[5], FrameRate, true, false, StartFrame);
 
-		if (Channels.Num() > 6) //noscale
+		if (DoubleChannels.Num() > 6) //noscale
 		{
-			ImportTransformChannel(Scale[0], Channels[6], FrameRate, false, false, StartFrame);
-			ImportTransformChannel(Scale[1], Channels[7], FrameRate, false, false, StartFrame);
-			ImportTransformChannel(Scale[2], Channels[8], FrameRate, false, false, StartFrame);
+			ImportTransformChannelToDouble(Scale[0], DoubleChannels[6], FrameRate, false, false, StartFrame);
+			ImportTransformChannelToDouble(Scale[1], DoubleChannels[7], FrameRate, false, false, StartFrame);
+			ImportTransformChannelToDouble(Scale[2], DoubleChannels[8], FrameRate, false, false, StartFrame);
 		}
 	}
 	return true;
@@ -1884,7 +1920,11 @@ static void PrepForInsertReplaceAnimation(bool bInsert, const FFBXNodeAndChannel
 {
 
 	TArray<FMovieSceneChannel*> Channels;
-	for (FMovieSceneFloatChannel* FChannel : NodeAndChannel.Channels)
+	for (FMovieSceneDoubleChannel* DChannel : NodeAndChannel.DoubleChannels)
+	{
+		Channels.Add(DChannel);
+	}
+	for (FMovieSceneFloatChannel* FChannel : NodeAndChannel.FloatChannels)
 	{
 		Channels.Add(FChannel);
 	}
@@ -2277,7 +2317,7 @@ bool MovieSceneToolHelpers::ImportFBXIntoControlRigChannels(UMovieScene* MovieSc
 
 		UMovieSceneUserImportFBXSettings* CurrentImportFBXSettings = GetMutableDefault<UMovieSceneUserImportFBXSettings>();
 		TArray<uint8> OriginalSettings;
-		FObjectWriter(CurrentImportFBXSettings, OriginalSettings);
+		FObjectWriter ObjWriter(CurrentImportFBXSettings, OriginalSettings);
 
 		CurrentImportFBXSettings->bMatchByNameOnly = false;
 		CurrentImportFBXSettings->bConvertSceneUnit = ImportFBXControlRigSettings->bConvertSceneUnit;
@@ -2368,7 +2408,7 @@ bool MovieSceneToolHelpers::ImportFBXIntoControlRigChannels(UMovieScene* MovieSc
 		}
 
 		// restore
-		FObjectReader(GetMutableDefault<UMovieSceneUserImportFBXSettings>(), OriginalSettings);
+		FObjectReader ObjReader(GetMutableDefault<UMovieSceneUserImportFBXSettings>(), OriginalSettings);
 		FbxImporter->ReleaseScene();
 	}
 
@@ -2521,7 +2561,7 @@ bool ImportFBXTransform(FString NodeName, FGuid ObjectBinding, UnFbx::FFbxCurves
 
 	FVector Location = DefaultTransform.GetLocation(), Rotation = DefaultTransform.GetRotation().Euler(), Scale3D = DefaultTransform.GetScale3D();
 
-	TArrayView<FMovieSceneFloatChannel*> Channels = TransformSection->GetChannelProxy().GetChannels<FMovieSceneFloatChannel>();
+	TArrayView<FMovieSceneDoubleChannel*> Channels = TransformSection->GetChannelProxy().GetChannels<FMovieSceneDoubleChannel>();
 
 	Channels[0]->SetDefault(Location.X);
 	Channels[1]->SetDefault(Location.Y);
@@ -2535,17 +2575,17 @@ bool ImportFBXTransform(FString NodeName, FGuid ObjectBinding, UnFbx::FFbxCurves
 	Channels[7]->SetDefault(Scale3D.Y);
 	Channels[8]->SetDefault(Scale3D.Z);
 
-	ImportTransformChannel(Translation[0],   Channels[0], FrameRate, false, true);
-	ImportTransformChannel(Translation[1],   Channels[1], FrameRate, true, true);
-	ImportTransformChannel(Translation[2],   Channels[2], FrameRate, false, true);
+	ImportTransformChannelToDouble(Translation[0],   Channels[0], FrameRate, false, true);
+	ImportTransformChannelToDouble(Translation[1],   Channels[1], FrameRate, true, true);
+	ImportTransformChannelToDouble(Translation[2],   Channels[2], FrameRate, false, true);
 
-	ImportTransformChannel(EulerRotation[0], Channels[3], FrameRate, false, true);
-	ImportTransformChannel(EulerRotation[1], Channels[4], FrameRate, true, true);
-	ImportTransformChannel(EulerRotation[2], Channels[5], FrameRate, true, true);
+	ImportTransformChannelToDouble(EulerRotation[0], Channels[3], FrameRate, false, true);
+	ImportTransformChannelToDouble(EulerRotation[1], Channels[4], FrameRate, true, true);
+	ImportTransformChannelToDouble(EulerRotation[2], Channels[5], FrameRate, true, true);
 
-	ImportTransformChannel(Scale[0],         Channels[6], FrameRate, false, true);
-	ImportTransformChannel(Scale[1],         Channels[7], FrameRate, false, true);
-	ImportTransformChannel(Scale[2],         Channels[8], FrameRate, false, true);
+	ImportTransformChannelToDouble(Scale[0],         Channels[6], FrameRate, false, true);
+	ImportTransformChannelToDouble(Scale[1],         Channels[7], FrameRate, false, true);
+	ImportTransformChannelToDouble(Scale[2],         Channels[8], FrameRate, false, true);
 
 	return true;
 }
@@ -3175,7 +3215,7 @@ bool ImportFBXOntoControlRigs(UWorld* World, UMovieScene* MovieScene, IMovieScen
 {
 	UMovieSceneUserImportFBXSettings* CurrentImportFBXSettings = GetMutableDefault<UMovieSceneUserImportFBXSettings>();
 	TArray<uint8> OriginalSettings;
-	FObjectWriter(CurrentImportFBXSettings, OriginalSettings);
+	FObjectWriter ObjWriter(CurrentImportFBXSettings, OriginalSettings);
 
 	CurrentImportFBXSettings->bMatchByNameOnly = ImportFBXSettings->bMatchByNameOnly;
 	CurrentImportFBXSettings->bForceFrontXAxis = ImportFBXSettings->bForceFrontXAxis;
@@ -3198,7 +3238,7 @@ bool MovieSceneToolHelpers::ImportFBXIfReady(UWorld* World, UMovieSceneSequence*
 
 	UMovieSceneUserImportFBXSettings* CurrentImportFBXSettings = GetMutableDefault<UMovieSceneUserImportFBXSettings>();
 	TArray<uint8> OriginalSettings;
-	FObjectWriter(CurrentImportFBXSettings, OriginalSettings);
+	FObjectWriter ObjWriter(CurrentImportFBXSettings, OriginalSettings);
 
 	CurrentImportFBXSettings->bMatchByNameOnly = ImportFBXSettings->bMatchByNameOnly;
 	CurrentImportFBXSettings->bForceFrontXAxis = ImportFBXSettings->bForceFrontXAxis;
@@ -3285,7 +3325,7 @@ bool MovieSceneToolHelpers::ImportFBXIfReady(UWorld* World, UMovieSceneSequence*
 	}
 
 	// restore
-	FObjectReader(GetMutableDefault<UMovieSceneUserImportFBXSettings>(), OriginalSettings);
+	FObjectReader ObjReader(GetMutableDefault<UMovieSceneUserImportFBXSettings>(), OriginalSettings);
 
 	FbxImporter->ReleaseScene();
 	UnFbx::FBXImportOptions* ImportOptions = FbxImporter->GetImportOptions();
@@ -3374,19 +3414,19 @@ EInterpCurveMode MovieSceneToolHelpers::RichCurveInterpolationToMatineeInterpola
 	}
 }
 
-void MovieSceneToolHelpers::CopyKeyDataToMoveAxis(const TMovieSceneChannelData<FMovieSceneFloatValue>& Channel, UInterpTrackMoveAxis* MoveAxis, FFrameRate InFrameRate)
+void MovieSceneToolHelpers::CopyKeyDataToMoveAxis(const TMovieSceneChannelData<FMovieSceneDoubleValue>& Channel, UInterpTrackMoveAxis* MoveAxis, FFrameRate InFrameRate)
 {
 	MoveAxis->FloatTrack.Points.Reset();
 
 	static FName LookupName(NAME_None);
 	
-	TArrayView<const FFrameNumber>          Times  = Channel.GetTimes();
-	TArrayView<const FMovieSceneFloatValue> Values = Channel.GetValues();
+	TArrayView<const FFrameNumber>           Times  = Channel.GetTimes();
+	TArrayView<const FMovieSceneDoubleValue> Values = Channel.GetValues();
 
 	for (int32 KeyIndex = 0; KeyIndex < Times.Num(); ++KeyIndex)
 	{
 		const float Time = Times[KeyIndex] / InFrameRate;
-		const FMovieSceneFloatValue& Value = Values[KeyIndex];
+		const FMovieSceneDoubleValue& Value = Values[KeyIndex];
 
 		const int32 PointIndex = MoveAxis->FloatTrack.AddPoint(Time, Value.Value);
 		MoveAxis->LookupTrack.AddPoint(Time, LookupName);
@@ -3463,14 +3503,14 @@ UObject* MovieSceneToolHelpers::ExportToCameraAnim(UMovieScene* InMovieScene, FG
 			{
 				FFrameRate TickResolution = InMovieScene->GetTickResolution();
 				UMovieScene3DTransformSection* TransformSection = Cast<UMovieScene3DTransformSection>(Sections[0]);
-				TArrayView<FMovieSceneFloatChannel*> FloatChannels = TransformSection->GetChannelProxy().GetChannels<FMovieSceneFloatChannel>();
-
-				CopyKeyDataToMoveAxis(FloatChannels[0]->GetData(), MoveAxies[AXIS_TranslationX], TickResolution);
-				CopyKeyDataToMoveAxis(FloatChannels[1]->GetData(), MoveAxies[AXIS_TranslationY], TickResolution);
-				CopyKeyDataToMoveAxis(FloatChannels[2]->GetData(), MoveAxies[AXIS_TranslationZ], TickResolution);
-				CopyKeyDataToMoveAxis(FloatChannels[3]->GetData(), MoveAxies[AXIS_RotationX],    TickResolution);
-				CopyKeyDataToMoveAxis(FloatChannels[4]->GetData(), MoveAxies[AXIS_RotationY],    TickResolution);
-				CopyKeyDataToMoveAxis(FloatChannels[5]->GetData(), MoveAxies[AXIS_RotationZ],    TickResolution);
+				TArrayView<FMovieSceneDoubleChannel*> DoubleChannels = TransformSection->GetChannelProxy().GetChannels<FMovieSceneDoubleChannel>();
+
+				CopyKeyDataToMoveAxis(DoubleChannels[0]->GetData(), MoveAxies[AXIS_TranslationX], TickResolution);
+				CopyKeyDataToMoveAxis(DoubleChannels[1]->GetData(), MoveAxies[AXIS_TranslationY], TickResolution);
+				CopyKeyDataToMoveAxis(DoubleChannels[2]->GetData(), MoveAxies[AXIS_TranslationZ], TickResolution);
+				CopyKeyDataToMoveAxis(DoubleChannels[3]->GetData(), MoveAxies[AXIS_RotationX],    TickResolution);
+				CopyKeyDataToMoveAxis(DoubleChannels[4]->GetData(), MoveAxies[AXIS_RotationY],    TickResolution);
+				CopyKeyDataToMoveAxis(DoubleChannels[5]->GetData(), MoveAxies[AXIS_RotationZ],    TickResolution);
 			}
 		}
 	}
@@ -4292,6 +4332,29 @@ void MovieSceneToolHelpers::SetOrAddKey(TMovieSceneChannelData<FMovieSceneFloatV
 	}
 }
 
+void MovieSceneToolHelpers::SetOrAddKey(TMovieSceneChannelData<FMovieSceneDoubleValue>& ChannelData, FFrameNumber Time, double Value)
+{
+	int32 ExistingIndex = ChannelData.FindKey(Time);
+	if (ExistingIndex != INDEX_NONE)
+	{
+		FMovieSceneDoubleValue& DoubleValue = ChannelData.GetValues()[ExistingIndex]; //-V758
+		DoubleValue.Value = Value;
+	}
+	else
+	{
+		FMovieSceneDoubleValue NewKey(Value);
+		ERichCurveTangentWeightMode WeightedMode = RCTWM_WeightedNone;
+		NewKey.InterpMode = ERichCurveInterpMode::RCIM_Cubic;
+		NewKey.TangentMode = ERichCurveTangentMode::RCTM_Auto;
+		NewKey.Tangent.ArriveTangent = 0.0f;
+		NewKey.Tangent.LeaveTangent = 0.0f;
+		NewKey.Tangent.TangentWeightMode = WeightedMode;
+		NewKey.Tangent.ArriveTangentWeight = 0.0f;
+		NewKey.Tangent.LeaveTangentWeight = 0.0f;
+		ChannelData.AddKey(Time, NewKey);
+	}
+}
+
 void MovieSceneToolHelpers::GetActorWorldTransforms(IMovieScenePlayer* Player, UMovieSceneSequence* InSequence, FMovieSceneSequenceIDRef Template, const FActorForWorldTransforms& ActorSelection, const TArray<FFrameNumber>& Frames, TArray<FTransform>& OutWorldTransforms)
 {
 	FGuid ActorHandle = GetHandleToObject(ActorSelection.Actor.Get(), InSequence, Player, Template,false);
