Sequencer: Large world coordinates support

diff --git a/Engine/Source/Editor/MovieSceneTools/Private/MovieSceneToolsModule.cpp b/Engine/Source/Editor/MovieSceneTools/Private/MovieSceneToolsModule.cpp
--- a/Engine/Source/Editor/MovieSceneTools/Private/MovieSceneToolsModule.cpp
+++ b/Engine/Source/Editor/MovieSceneTools/Private/MovieSceneToolsModule.cpp
@@ -19,6 +19,7 @@
 #include "TrackEditors/PropertyTrackEditors/BytePropertyTrackEditor.h"
 #include "TrackEditors/PropertyTrackEditors/ColorPropertyTrackEditor.h"
 #include "TrackEditors/PropertyTrackEditors/FloatPropertyTrackEditor.h"
+#include "TrackEditors/PropertyTrackEditors/DoublePropertyTrackEditor.h"
 #include "TrackEditors/PropertyTrackEditors/IntegerPropertyTrackEditor.h"
 #include "TrackEditors/PropertyTrackEditors/VectorPropertyTrackEditor.h"
 #include "TrackEditors/PropertyTrackEditors/TransformPropertyTrackEditor.h"
@@ -105,8 +106,10 @@ void FMovieSceneToolsModule::StartupModule()
 		BytePropertyTrackCreateEditorHandle = SequencerModule.RegisterPropertyTrackEditor<FBytePropertyTrackEditor>();
 		ColorPropertyTrackCreateEditorHandle = SequencerModule.RegisterPropertyTrackEditor<FColorPropertyTrackEditor>();
 		FloatPropertyTrackCreateEditorHandle = SequencerModule.RegisterPropertyTrackEditor<FFloatPropertyTrackEditor>();
+		DoublePropertyTrackCreateEditorHandle = SequencerModule.RegisterPropertyTrackEditor<FDoublePropertyTrackEditor>();
 		IntegerPropertyTrackCreateEditorHandle = SequencerModule.RegisterPropertyTrackEditor<FIntegerPropertyTrackEditor>();
-		VectorPropertyTrackCreateEditorHandle = SequencerModule.RegisterPropertyTrackEditor<FVectorPropertyTrackEditor>();
+		FloatVectorPropertyTrackCreateEditorHandle = SequencerModule.RegisterPropertyTrackEditor<FFloatVectorPropertyTrackEditor>();
+		DoubleVectorPropertyTrackCreateEditorHandle = SequencerModule.RegisterPropertyTrackEditor<FDoubleVectorPropertyTrackEditor>();
 		TransformPropertyTrackCreateEditorHandle = SequencerModule.RegisterPropertyTrackEditor<FTransformPropertyTrackEditor>();
 		EulerTransformPropertyTrackCreateEditorHandle = SequencerModule.RegisterPropertyTrackEditor<FEulerTransformPropertyTrackEditor>();
 		VisibilityPropertyTrackCreateEditorHandle = SequencerModule.RegisterPropertyTrackEditor<FVisibilityPropertyTrackEditor>();
@@ -151,6 +154,7 @@ void FMovieSceneToolsModule::StartupModule()
 		SequencerModule.RegisterChannelInterface<FMovieSceneByteChannel>();
 		SequencerModule.RegisterChannelInterface<FMovieSceneIntegerChannel>();
 		SequencerModule.RegisterChannelInterface<FMovieSceneFloatChannel>();
+		SequencerModule.RegisterChannelInterface<FMovieSceneDoubleChannel>();
 		SequencerModule.RegisterChannelInterface<FMovieSceneStringChannel>();
 		SequencerModule.RegisterChannelInterface<FMovieSceneParticleChannel>();
 		SequencerModule.RegisterChannelInterface<FMovieSceneActorReferenceData>();
@@ -231,8 +235,10 @@ void FMovieSceneToolsModule::ShutdownModule()
 	SequencerModule.UnRegisterTrackEditor( BytePropertyTrackCreateEditorHandle );
 	SequencerModule.UnRegisterTrackEditor( ColorPropertyTrackCreateEditorHandle );
 	SequencerModule.UnRegisterTrackEditor( FloatPropertyTrackCreateEditorHandle );
+	SequencerModule.UnRegisterTrackEditor( DoublePropertyTrackCreateEditorHandle );
 	SequencerModule.UnRegisterTrackEditor( IntegerPropertyTrackCreateEditorHandle );
-	SequencerModule.UnRegisterTrackEditor( VectorPropertyTrackCreateEditorHandle );
+	SequencerModule.UnRegisterTrackEditor( FloatVectorPropertyTrackCreateEditorHandle );
+	SequencerModule.UnRegisterTrackEditor( DoubleVectorPropertyTrackCreateEditorHandle );
 	SequencerModule.UnRegisterTrackEditor( TransformPropertyTrackCreateEditorHandle );
 	SequencerModule.UnRegisterTrackEditor( EulerTransformPropertyTrackCreateEditorHandle );
 	SequencerModule.UnRegisterTrackEditor( VisibilityPropertyTrackCreateEditorHandle );
@@ -427,6 +433,12 @@ void FMovieSceneToolsModule::RegisterClipboardConversions()
 	DefineExplicitConversion<FMovieSceneFloatValue, uint8>([](const FMovieSceneFloatValue& In) -> uint8 { return In.Value; 					});
 	DefineExplicitConversion<FMovieSceneFloatValue, bool>([](const FMovieSceneFloatValue& In) -> bool	{ return !!In.Value; 				});
 
+	DefineExplicitConversion<int32, FMovieSceneDoubleValue>([](const int32& In) -> FMovieSceneDoubleValue { return FMovieSceneDoubleValue(In);	});
+	DefineExplicitConversion<uint8, FMovieSceneDoubleValue>([](const uint8& In) -> FMovieSceneDoubleValue { return FMovieSceneDoubleValue(In);	});
+	DefineExplicitConversion<FMovieSceneDoubleValue, int32>([](const FMovieSceneDoubleValue& In) -> int32 { return In.Value; 					});
+	DefineExplicitConversion<FMovieSceneDoubleValue, uint8>([](const FMovieSceneDoubleValue& In) -> uint8 { return In.Value; 					});
+	DefineExplicitConversion<FMovieSceneDoubleValue, bool>([](const FMovieSceneDoubleValue& In) -> bool	  { return !!In.Value;					});
+
 	FSequencerClipboardReconciler::AddTrackAlias("Location.X", "R");
 	FSequencerClipboardReconciler::AddTrackAlias("Location.Y", "G");
 	FSequencerClipboardReconciler::AddTrackAlias("Location.Z", "B");
diff --git a/Engine/Source/Editor/MovieSceneTools/Private/TrackEditors/AttachTrackEditor.cpp b/Engine/Source/Editor/MovieSceneTools/Private/TrackEditors/AttachTrackEditor.cpp
--- a/Engine/Source/Editor/MovieSceneTools/Private/TrackEditors/AttachTrackEditor.cpp
+++ b/Engine/Source/Editor/MovieSceneTools/Private/TrackEditors/AttachTrackEditor.cpp
@@ -384,9 +384,9 @@ void F3DAttachTrackEditor::FindOrCreateTransformTrack(const TRange<FFrameNumber>
 /**
  * Helper method to safely return an array allocated to store the proper the number of float channels if not already allocated
  */
-TArray<FMovieSceneFloatValue>& ResizeAndAddKey(const FFrameNumber& InKey, int32 InNum, TMap<FFrameNumber, TArray<FMovieSceneFloatValue>>& OutTransformMap, TSet<FFrameNumber>* OutTimesAdded)
+TArray<FMovieSceneDoubleValue>& ResizeAndAddKey(const FFrameNumber& InKey, int32 InNum, TMap<FFrameNumber, TArray<FMovieSceneDoubleValue>>& OutTransformMap, TSet<FFrameNumber>* OutTimesAdded)
 {
-	TArray<FMovieSceneFloatValue>& Transform = OutTransformMap.FindOrAdd(InKey);
+	TArray<FMovieSceneDoubleValue>& Transform = OutTransformMap.FindOrAdd(InKey);
 	if (Transform.Num() == 0)
 	{
 		Transform.SetNum(InNum);
@@ -401,7 +401,7 @@ TArray<FMovieSceneFloatValue>& ResizeAndAddKey(const FFrameNumber& InKey, int32
 /**
  * Helper method which adds keys from a list of float channels to a map mapping the time to a full transform
  */
-void AddKeysFromChannels(TArrayView<FMovieSceneFloatChannel*> InChannels, const TRange<FFrameNumber>& InAttachRange, TMap<FFrameNumber, TArray<FMovieSceneFloatValue>>& OutTransformMap, TSet<FFrameNumber>& OutTimesAdded)
+void AddKeysFromChannels(TArrayView<FMovieSceneDoubleChannel*> InChannels, const TRange<FFrameNumber>& InAttachRange, TMap<FFrameNumber, TArray<FMovieSceneDoubleValue>>& OutTransformMap, TSet<FFrameNumber>& OutTimesAdded)
 {
 	const int32 NumChannels = 9;
 	for (int32 ChannelIndex = 0; ChannelIndex < NumChannels; ChannelIndex++)
@@ -421,10 +421,10 @@ void AddKeysFromChannels(TArrayView<FMovieSceneFloatChannel*> InChannels, const
 		}
 
 		const int32 NumValsInRange = TimesInRange.Num();
-		TArrayView<const FMovieSceneFloatValue> ValuesInRange = InChannels[ChannelIndex]->GetValues().Slice(BeginRangeIndex, NumValsInRange);
+		TArrayView<const FMovieSceneDoubleValue> ValuesInRange = InChannels[ChannelIndex]->GetValues().Slice(BeginRangeIndex, NumValsInRange);
 		for (int32 KeyIndex = 0; KeyIndex < ValuesInRange.Num(); KeyIndex++)
 		{
-			TArray<FMovieSceneFloatValue>& Transform = ResizeAndAddKey(TimesInRange[KeyIndex], InChannels.Num(), OutTransformMap, &OutTimesAdded);
+			TArray<FMovieSceneDoubleValue>& Transform = ResizeAndAddKey(TimesInRange[KeyIndex], InChannels.Num(), OutTransformMap, &OutTimesAdded);
 			Transform[ChannelIndex] = ValuesInRange[KeyIndex];
 		}
 	}
@@ -434,28 +434,28 @@ void AddKeysFromChannels(TArrayView<FMovieSceneFloatChannel*> InChannels, const
  * Helper method which updates the values in each channel in a list of movie scene float values given a 
  * transform, preserving the interpolation style and other attributes
  */
-void UpdateFloatValueTransform(const FTransform& InTransform, TArrayView<FMovieSceneFloatValue> OutFloatValueTransform)
+void UpdateDoubleValueTransform(const FTransform& InTransform, TArrayView<FMovieSceneDoubleValue> OutDoubleValueTransform)
 {
-	OutFloatValueTransform[0].Value = InTransform.GetTranslation().X;
-	OutFloatValueTransform[1].Value = InTransform.GetTranslation().Y;
-	OutFloatValueTransform[2].Value = InTransform.GetTranslation().Z;
+	OutDoubleValueTransform[0].Value = InTransform.GetTranslation().X;
+	OutDoubleValueTransform[1].Value = InTransform.GetTranslation().Y;
+	OutDoubleValueTransform[2].Value = InTransform.GetTranslation().Z;
 
-	OutFloatValueTransform[3].Value = InTransform.GetRotation().Euler().X;
-	OutFloatValueTransform[4].Value = InTransform.GetRotation().Euler().Y;
-	OutFloatValueTransform[5].Value = InTransform.GetRotation().Euler().Z;
+	OutDoubleValueTransform[3].Value = InTransform.GetRotation().Euler().X;
+	OutDoubleValueTransform[4].Value = InTransform.GetRotation().Euler().Y;
+	OutDoubleValueTransform[5].Value = InTransform.GetRotation().Euler().Z;
 
-	OutFloatValueTransform[6].Value = InTransform.GetScale3D().X;
-	OutFloatValueTransform[7].Value = InTransform.GetScale3D().Y;
-	OutFloatValueTransform[8].Value = InTransform.GetScale3D().Z;
+	OutDoubleValueTransform[6].Value = InTransform.GetScale3D().X;
+	OutDoubleValueTransform[7].Value = InTransform.GetScale3D().Y;
+	OutDoubleValueTransform[8].Value = InTransform.GetScale3D().Z;
 }
 
 /**
  * Helper method which converts a list of float values to a transform
  */
-FORCEINLINE FTransform FloatValuesToTransform(TArrayView<const FMovieSceneFloatValue> InFloatValues)
+FORCEINLINE FTransform DoubleValuesToTransform(TArrayView<const FMovieSceneDoubleValue> InDoubleValues)
 {
-	return FTransform(FRotator::MakeFromEuler(FVector(InFloatValues[3].Value, InFloatValues[4].Value, InFloatValues[5].Value)), 
-		FVector(InFloatValues[0].Value, InFloatValues[1].Value, InFloatValues[2].Value), FVector(InFloatValues[6].Value, InFloatValues[7].Value, InFloatValues[8].Value));
+	return FTransform(FRotator::MakeFromEuler(FVector(InDoubleValues[3].Value, InDoubleValues[4].Value, InDoubleValues[5].Value)), 
+		FVector(InDoubleValues[0].Value, InDoubleValues[1].Value, InDoubleValues[2].Value), FVector(InDoubleValues[6].Value, InDoubleValues[7].Value, InDoubleValues[8].Value));
 }
 
 /**
@@ -807,10 +807,10 @@ private:
 /**
  * Updates an array of float channels with the keys in a given transform map mapping times to float values
  */
-void UpdateChannelTransforms(const TRange<FFrameNumber>& InAttachRange, TMap<FFrameNumber, TArray<FMovieSceneFloatValue>>& InTransformMap, TArrayView<FMovieSceneFloatChannel*>& InChannels, int32 InNumChannels, bool bInBakedData)
+void UpdateChannelTransforms(const TRange<FFrameNumber>& InAttachRange, TMap<FFrameNumber, TArray<FMovieSceneDoubleValue>>& InTransformMap, TArrayView<FMovieSceneDoubleChannel*>& InChannels, int32 InNumChannels, bool bInBakedData)
 {
 	// Remove all handles in range so we can add the new ones
-	for (FMovieSceneFloatChannel* Channel : InChannels)
+	for (FMovieSceneDoubleChannel* Channel : InChannels)
 	{
 		TArray<FKeyHandle> KeysToRemove;
 		Channel->GetKeys(InAttachRange, nullptr, &KeysToRemove);
@@ -824,7 +824,7 @@ void UpdateChannelTransforms(const TRange<FFrameNumber>& InAttachRange, TMap<FFr
 	InTransformMap.KeySort([](const FFrameNumber& LHS, const FFrameNumber& RHS) { return LHS.Value < RHS.Value; });
 	TArray<FFrameNumber> NewKeyFrames;
 	InTransformMap.GetKeys(NewKeyFrames);
-	TArray<FMovieSceneFloatValue> NewKeyValues;
+	TArray<FMovieSceneDoubleValue> NewKeyValues;
 
 	// Update keys in channels
 	for (int32 ChannelIndex = 0; ChannelIndex < InNumChannels; ChannelIndex++)
@@ -837,7 +837,7 @@ void UpdateChannelTransforms(const TRange<FFrameNumber>& InAttachRange, TMap<FFr
 		// and simply remove and re-add all of the keys from first to last
 		TArray<FKeyHandle> LowerKeyHandles, UpperKeyHandles;
 		TArray<FFrameNumber> LowerKeyTimes, UpperKeyTimes;
-		TArray<FMovieSceneFloatValue> PrevKeyValues;
+		TArray<FMovieSceneDoubleValue> PrevKeyValues;
 		Algo::Copy(InChannels[ChannelIndex]->GetValues(), PrevKeyValues);
 
 		// Get the keys contained in before attach and after attach ranges
@@ -850,8 +850,8 @@ void UpdateChannelTransforms(const TRange<FFrameNumber>& InAttachRange, TMap<FFr
 		if (ExcludedRanges.Num() > 0 && ExcludedRanges[0].GetUpperBoundValue() <= InAttachRange.GetLowerBoundValue() && LowerKeyTimes.Num() > 0)
 		{
 			InChannels[ChannelIndex]->DeleteKeys(LowerKeyHandles);
-			TArray<FMovieSceneFloatValue> ValuesToAdd;
-			Algo::Copy(TArrayView<const FMovieSceneFloatValue>(PrevKeyValues).Slice(ValueIndex, LowerKeyTimes.Num()), ValuesToAdd);
+			TArray<FMovieSceneDoubleValue> ValuesToAdd;
+			Algo::Copy(TArrayView<const FMovieSceneDoubleValue>(PrevKeyValues).Slice(ValueIndex, LowerKeyTimes.Num()), ValuesToAdd);
 			InChannels[ChannelIndex]->AddKeys(LowerKeyTimes, ValuesToAdd);
 			ValueIndex += LowerKeyTimes.Num();
 		}
@@ -863,8 +863,8 @@ void UpdateChannelTransforms(const TRange<FFrameNumber>& InAttachRange, TMap<FFr
 		if (ExcludedRanges.Num() > 0 && ExcludedRanges.Top().GetLowerBoundValue() >= InAttachRange.GetUpperBoundValue() && ValueIndex < PrevKeyValues.Num() && UpperKeyTimes.Num() > 0)
 		{
 			InChannels[ChannelIndex]->DeleteKeys(UpperKeyHandles);
-			TArray<FMovieSceneFloatValue> ValuesToAdd;
-			Algo::Copy(TArrayView<const FMovieSceneFloatValue>(PrevKeyValues).Slice(ValueIndex, UpperKeyTimes.Num()), ValuesToAdd);
+			TArray<FMovieSceneDoubleValue> ValuesToAdd;
+			Algo::Copy(TArrayView<const FMovieSceneDoubleValue>(PrevKeyValues).Slice(ValueIndex, UpperKeyTimes.Num()), ValuesToAdd);
 			InChannels[ChannelIndex]->AddKeys(UpperKeyTimes, ValuesToAdd);
 		}
 
@@ -925,7 +925,7 @@ void F3DAttachTrackEditor::TrimAndPreserve(FGuid InObjectBinding, UMovieSceneSec
 			return;
 		}
 
-		TArrayView<FMovieSceneFloatChannel*> Channels = TransformSection->GetChannelProxy().GetChannels<FMovieSceneFloatChannel>();
+		TArrayView<FMovieSceneDoubleChannel*> Channels = TransformSection->GetChannelProxy().GetChannels<FMovieSceneDoubleChannel>();
 
 		FLocalTransformEvaluator LocalTransformEval(SharedThis(this), Object, TransformTrack);
 
@@ -933,12 +933,12 @@ void F3DAttachTrackEditor::TrimAndPreserve(FGuid InObjectBinding, UMovieSceneSec
 		{
 			FWorldTransformEvaluator ReAttachParentEvaluator(SharedThis(this), AttachSection->ReAttachOnDetach.Get());
 
-			CompensateChildTrack(ExcludedRange, Channels, TOptional<TArrayView<FMovieSceneFloatChannel*>>(), ReAttachParentEvaluator, LocalTransformEval, ETransformPreserveType::CurrentKey, RevertModifier);
+			CompensateChildTrack(ExcludedRange, Channels, TOptional<TArrayView<FMovieSceneDoubleChannel*>>(), ReAttachParentEvaluator, LocalTransformEval, ETransformPreserveType::CurrentKey, RevertModifier);
 		}
 		else
 		{
 			TSet<FFrameNumber> KeyTimesToCompensate;
-			TMap<FFrameNumber, TArray<FMovieSceneFloatValue>> TransformMap;
+			TMap<FFrameNumber, TArray<FMovieSceneDoubleValue>> TransformMap;
 
 			// Add all keys already existing in the range to the transform map
 			AddKeysFromChannels(Channels, ExcludedRange, TransformMap, KeyTimesToCompensate);
@@ -967,20 +967,20 @@ void F3DAttachTrackEditor::TrimAndPreserve(FGuid InObjectBinding, UMovieSceneSec
 			// Evaluate the transform at all times with keys
 			for (auto Itr = TransformMap.CreateIterator(); Itr; ++Itr)
 			{
-				UpdateFloatValueTransform(LocalTransformEval(Itr->Key), Itr->Value);
+				UpdateDoubleValueTransform(LocalTransformEval(Itr->Key), Itr->Value);
 			}
 
 			// Modify each transform
 			for (const FFrameNumber& CompTime : KeyTimesToCompensate)
 			{
-				const FTransform RevertedTransform = RevertModifier(FloatValuesToTransform(TransformMap[CompTime]), CompTime);
-				UpdateFloatValueTransform(RevertedTransform, TransformMap[CompTime]);
+				const FTransform RevertedTransform = RevertModifier(DoubleValuesToTransform(TransformMap[CompTime]), CompTime);
+				UpdateDoubleValueTransform(RevertedTransform, TransformMap[CompTime]);
 			}
 
 			// Manually set edge keys to have linear interpolation
 			for (const FFrameNumber& EdgeKey : EdgeKeys)
 			{
-				for (FMovieSceneFloatValue& Key : TransformMap[EdgeKey])
+				for (FMovieSceneDoubleValue& Key : TransformMap[EdgeKey])
 				{
 					Key.InterpMode = ERichCurveInterpMode::RCIM_Linear;
 				}
@@ -991,7 +991,7 @@ void F3DAttachTrackEditor::TrimAndPreserve(FGuid InObjectBinding, UMovieSceneSec
 		}
 
 		// Remove previous boundary keys
-		for (FMovieSceneFloatChannel* Channel : Channels)
+		for (FMovieSceneDoubleChannel* Channel : Channels)
 		{
 			TArray<FKeyHandle> KeyAtTime;
 
@@ -1006,7 +1006,7 @@ void F3DAttachTrackEditor::TrimAndPreserve(FGuid InObjectBinding, UMovieSceneSec
 }
 
 template<typename ModifierFuncType>
-void F3DAttachTrackEditor::CompensateChildTrack(const TRange<FFrameNumber>& InAttachRange, TArrayView<FMovieSceneFloatChannel*> Channels, TOptional<TArrayView<FMovieSceneFloatChannel*>> ParentChannels,
+void F3DAttachTrackEditor::CompensateChildTrack(const TRange<FFrameNumber>& InAttachRange, TArrayView<FMovieSceneDoubleChannel*> Channels, TOptional<TArrayView<FMovieSceneDoubleChannel*>> ParentChannels,
 	const ITransformEvaluator& InParentTransformEval, const ITransformEvaluator& InChildTransformEval,
 	ETransformPreserveType InPreserveType, ModifierFuncType InModifyTransform)
 {
@@ -1015,7 +1015,7 @@ void F3DAttachTrackEditor::CompensateChildTrack(const TRange<FFrameNumber>& InAt
 	const int32 NumChannels = 9;
 
 	TSet<FFrameNumber> KeyTimesToCompensate;
-	TMap<FFrameNumber, TArray<FMovieSceneFloatValue>> TransformMap;
+	TMap<FFrameNumber, TArray<FMovieSceneDoubleValue>> TransformMap;
 
 	// Add all times with keys to the map
 	if (PreserveType == ETransformPreserveType::Bake)
@@ -1026,9 +1026,9 @@ void F3DAttachTrackEditor::CompensateChildTrack(const TRange<FFrameNumber>& InAt
 			FrameItr += FMath::RoundToInt(TickResolution.AsDecimal() / DisplayRate.AsDecimal()))
 		{
 			ResizeAndAddKey(FrameItr, Channels.Num(), TransformMap, &KeyTimesToCompensate);
-			for (FMovieSceneFloatValue& FloatVal : TransformMap[FrameItr])
+			for (FMovieSceneDoubleValue& DoubleVal : TransformMap[FrameItr])
 			{
-				FloatVal.InterpMode = ERichCurveInterpMode::RCIM_Linear;
+				DoubleVal.InterpMode = ERichCurveInterpMode::RCIM_Linear;
 			}
 		}
 	}
@@ -1060,7 +1060,7 @@ void F3DAttachTrackEditor::CompensateChildTrack(const TRange<FFrameNumber>& InAt
 	for (auto Itr = TransformMap.CreateIterator(); Itr; ++Itr)
 	{
 		const FTransform TempTransform = InChildTransformEval(Itr->Key);
-		UpdateFloatValueTransform(TempTransform, Itr->Value);
+		UpdateDoubleValueTransform(TempTransform, Itr->Value);
 	}
 
 	if (InPreserveType == ETransformPreserveType::AllKeys || InPreserveType == ETransformPreserveType::Bake)
@@ -1069,15 +1069,15 @@ void F3DAttachTrackEditor::CompensateChildTrack(const TRange<FFrameNumber>& InAt
 		for (const FFrameNumber& CompTime : KeyTimesToCompensate)
 		{
 			const FTransform ParentTransformAtTime = InParentTransformEval(CompTime);
-			const FTransform NewTransform = InModifyTransform(FloatValuesToTransform(TransformMap[CompTime]), CompTime);
+			const FTransform NewTransform = InModifyTransform(DoubleValuesToTransform(TransformMap[CompTime]), CompTime);
 			const FTransform RelativeTransform = NewTransform.GetRelativeTransform(ParentTransformAtTime);
-			UpdateFloatValueTransform(RelativeTransform, TransformMap[CompTime]);
+			UpdateDoubleValueTransform(RelativeTransform, TransformMap[CompTime]);
 		}
 	}
 	else if (InPreserveType == ETransformPreserveType::CurrentKey)
 	{
 		// Find the relative transform on the first frame of the attach
-		const FTransform BeginChildTransform = InModifyTransform(FloatValuesToTransform(TransformMap[KeyTime]), KeyTime);
+		const FTransform BeginChildTransform = InModifyTransform(DoubleValuesToTransform(TransformMap[KeyTime]), KeyTime);
 		const FTransform BeginParentTransform = InParentTransformEval(KeyTime);
 
 		const FTransform BeginRelativeTransform = BeginChildTransform.GetRelativeTransform(BeginParentTransform);
@@ -1085,20 +1085,20 @@ void F3DAttachTrackEditor::CompensateChildTrack(const TRange<FFrameNumber>& InAt
 		// offset each transform by initial relative transform calculated before
 		for (const FFrameNumber& CompTime : KeyTimesToCompensate)
 		{
-			const FTransform ChildTransformAtTime = InModifyTransform(FloatValuesToTransform(TransformMap[CompTime]), CompTime);
+			const FTransform ChildTransformAtTime = InModifyTransform(DoubleValuesToTransform(TransformMap[CompTime]), CompTime);
 			const FTransform StartToCurrentTransform = ChildTransformAtTime.GetRelativeTransform(BeginChildTransform);
 
-			UpdateFloatValueTransform(BeginRelativeTransform * StartToCurrentTransform, TransformMap[CompTime]);
+			UpdateDoubleValueTransform(BeginRelativeTransform * StartToCurrentTransform, TransformMap[CompTime]);
 		}
 
 		const FTransform EndParentTransform = InParentTransformEval(AttachEndTime);
-		UpdateFloatValueTransform(EndParentTransform * FloatValuesToTransform(TransformMap[BeforeDetachTime]), TransformMap[AttachEndTime]);
+		UpdateDoubleValueTransform(EndParentTransform * DoubleValuesToTransform(TransformMap[BeforeDetachTime]), TransformMap[AttachEndTime]);
 	}
 
 	// Manually set edge keys to have linear interpolation
 	for (const FFrameNumber& EdgeKey : EdgeKeys)
 	{
-		for (FMovieSceneFloatValue& Key : TransformMap[EdgeKey])
+		for (FMovieSceneDoubleValue& Key : TransformMap[EdgeKey])
 		{
 			Key.InterpMode = ERichCurveInterpMode::RCIM_Linear;
 		}
@@ -1148,7 +1148,7 @@ FKeyPropertyResult F3DAttachTrackEditor::AddKeyInternal( FFrameNumber KeyTime, c
 
 	FWorldTransformEvaluator ParentTransformEval(SharedThis(this), ParentActor, SocketName);
 
-	TOptional<TArrayView<FMovieSceneFloatChannel*>> ParentChannels;
+	TOptional<TArrayView<FMovieSceneDoubleChannel*>> ParentChannels;
 
 	// If the constraint exists within this sequence, we can perform transform compensation
 	if (ConstraintBindingID.ResolveSequenceID(CurrentSequenceID,*GetSequencer()) == CurrentSequenceID)
@@ -1156,7 +1156,7 @@ FKeyPropertyResult F3DAttachTrackEditor::AddKeyInternal( FFrameNumber KeyTime, c
 		UMovieScene3DTransformTrack* ParentTransformTrack = MovieScene->FindTrack<UMovieScene3DTransformTrack>(ConstraintBindingID.GetGuid());
 		if (ParentTransformTrack && ParentTransformTrack->GetAllSections().Num() == 1)
 		{
-			ParentChannels = ParentTransformTrack->GetAllSections()[0]->GetChannelProxy().GetChannels<FMovieSceneFloatChannel>();
+			ParentChannels = ParentTransformTrack->GetAllSections()[0]->GetChannelProxy().GetChannels<FMovieSceneDoubleChannel>();
 		}
 	}
 
@@ -1250,7 +1250,7 @@ FKeyPropertyResult F3DAttachTrackEditor::AddKeyInternal( FFrameNumber KeyTime, c
 		}
 
 		// Get transform track channels
-		TArrayView<FMovieSceneFloatChannel*> Channels = TransformSection->GetChannelProxy().GetChannels<FMovieSceneFloatChannel>();
+		TArrayView<FMovieSceneDoubleChannel*> Channels = TransformSection->GetChannelProxy().GetChannels<FMovieSceneDoubleChannel>();
 
 		// find intersecting section
 		TOptional<UMovieSceneSection*> IntersectingSection;
diff --git a/Engine/Source/Editor/MovieSceneTools/Private/TrackEditors/AttachTrackEditor.h b/Engine/Source/Editor/MovieSceneTools/Private/TrackEditors/AttachTrackEditor.h
--- a/Engine/Source/Editor/MovieSceneTools/Private/TrackEditors/AttachTrackEditor.h
+++ b/Engine/Source/Editor/MovieSceneTools/Private/TrackEditors/AttachTrackEditor.h
@@ -15,7 +15,7 @@
 #include "Templates/Function.h"
 #include "EntitySystem/Interrogation/MovieSceneInterrogationLinker.h"
 
-struct FMovieSceneFloatChannel;
+struct FMovieSceneDoubleChannel;
 
 class AActor;
 class FMenuBuilder;
@@ -84,7 +84,7 @@ private:
 
 	/** Helper for AddKeyInternal to offset child track's keys */
 	template<typename ModifierFuncType>
-	void CompensateChildTrack(const TRange<FFrameNumber>& InAttachRange, TArrayView<FMovieSceneFloatChannel*> Channels, TOptional<TArrayView<FMovieSceneFloatChannel*>> ParentChannels,
+	void CompensateChildTrack(const TRange<FFrameNumber>& InAttachRange, TArrayView<FMovieSceneDoubleChannel*> Channels, TOptional<TArrayView<FMovieSceneDoubleChannel*>> ParentChannels,
 		const ITransformEvaluator& InParentTransformEval, const ITransformEvaluator& InChildTransformEval, ETransformPreserveType InPreserveType, ModifierFuncType InModifyTransform);
 
 	/** Delegate for AnimatablePropertyChanged in AddKey */
diff --git a/Engine/Source/Editor/MovieSceneTools/Private/TrackEditors/PropertyTrackEditors/DoublePropertyTrackEditor.cpp b/Engine/Source/Editor/MovieSceneTools/Private/TrackEditors/PropertyTrackEditors/DoublePropertyTrackEditor.cpp
new file mode 100644
--- /dev/null
+++ b/Engine/Source/Editor/MovieSceneTools/Private/TrackEditors/PropertyTrackEditors/DoublePropertyTrackEditor.cpp
@@ -0,0 +1,53 @@
+// Copyright Epic Games, Inc. All Rights Reserved.
+
+#include "TrackEditors/PropertyTrackEditors/DoublePropertyTrackEditor.h"
+#include "Editor/UnrealEdEngine.h"
+#include "UnrealEdGlobals.h"
+#include "MovieSceneToolHelpers.h"
+#include "Evaluation/MovieScenePropertyTemplate.h"
+
+#include "Systems/MovieScenePropertyInstantiator.h"
+#include "MovieSceneTracksComponentTypes.h"
+
+TSharedRef<ISequencerTrackEditor> FDoublePropertyTrackEditor::CreateTrackEditor( TSharedRef<ISequencer> OwningSequencer )
+{
+	return MakeShareable(new FDoublePropertyTrackEditor(OwningSequencer));
+}
+
+
+void FDoublePropertyTrackEditor::GenerateKeysFromPropertyChanged( const FPropertyChangedParams& PropertyChangedParams, UMovieSceneSection* SectionToKey, FGeneratedTrackKeys& OutGeneratedKeys )
+{
+	const double KeyedValue = PropertyChangedParams.GetPropertyValue<double>();
+	const double NewValue   = RecomposeDouble(KeyedValue, PropertyChangedParams.ObjectsThatChanged[0], SectionToKey);
+	OutGeneratedKeys.Add(FMovieSceneChannelValueSetter::Create<FMovieSceneDoubleChannel>(0, NewValue, true));
+}
+
+
+double FDoublePropertyTrackEditor::RecomposeDouble(double InCurrentValue, UObject* AnimatedObject, UMovieSceneSection* Section)
+{
+	using namespace UE::MovieScene;
+
+	double Recomposed = InCurrentValue;
+
+	const FMovieSceneRootEvaluationTemplateInstance& EvaluationTemplate = GetSequencer()->GetEvaluationTemplate();
+
+	UMovieSceneEntitySystemLinker* EntityLinker = EvaluationTemplate.GetEntitySystemLinker();
+	FMovieSceneEntityID EntityID = EvaluationTemplate.FindEntityFromOwner(Section, 0, GetSequencer()->GetFocusedTemplateID());
+
+	if (EntityLinker && EntityID)
+	{
+		TGuardValue<FEntityManager*> DebugVizGuard(GEntityManagerForDebuggingVisualizers, &EntityLinker->EntityManager);
+
+		UMovieScenePropertyInstantiatorSystem* System = EntityLinker->FindSystem<UMovieScenePropertyInstantiatorSystem>();
+		if (System)
+		{
+			FDecompositionQuery Query;
+			Query.Entities = MakeArrayView(&EntityID, 1);
+			Query.Object   = AnimatedObject;
+
+			Recomposed = System->RecomposeBlend(FMovieSceneTracksComponentTypes::Get()->Double, Query, InCurrentValue).Values[0];
+		}
+	}
+
+	return Recomposed;
+}
diff --git a/Engine/Source/Editor/MovieSceneTools/Private/TrackEditors/PropertyTrackEditors/DoublePropertyTrackEditor.h b/Engine/Source/Editor/MovieSceneTools/Private/TrackEditors/PropertyTrackEditors/DoublePropertyTrackEditor.h
new file mode 100644
--- /dev/null
+++ b/Engine/Source/Editor/MovieSceneTools/Private/TrackEditors/PropertyTrackEditors/DoublePropertyTrackEditor.h
@@ -0,0 +1,53 @@
+// Copyright Epic Games, Inc. All Rights Reserved.
+
+#pragma once
+
+#include "CoreMinimal.h"
+#include "Misc/Guid.h"
+#include "ISequencer.h"
+#include "ISequencerSection.h"
+#include "ISequencerTrackEditor.h"
+#include "PropertyTrackEditor.h"
+#include "Tracks/MovieSceneDoubleTrack.h"
+#include "Sections/MovieSceneDoubleSection.h"
+
+/**
+ * A property track editor for doubles.
+ */
+class FDoublePropertyTrackEditor
+	: public FPropertyTrackEditor<UMovieSceneDoubleTrack>
+{
+public:
+
+	/**
+	 * Constructor.
+	 *
+	 * @param InSequencer The sequencer instance to be used by this tool.
+	 */
+	FDoublePropertyTrackEditor(TSharedRef<ISequencer> InSequencer)
+		: FPropertyTrackEditor(InSequencer, GetAnimatedPropertyTypes())
+	{ }
+
+	/**
+	 * Retrieve a list of all property types that this track editor animates
+	 */
+	static TArray<FAnimatedPropertyKey, TInlineAllocator<1>> GetAnimatedPropertyTypes()
+	{
+		return TArray<FAnimatedPropertyKey, TInlineAllocator<1>>({ FAnimatedPropertyKey::FromPropertyTypeName(NAME_DoubleProperty) });
+	}
+
+	/**
+	 * Creates an instance of this class (called by a sequencer).
+	 *
+	 * @param OwningSequencer The sequencer instance to be used by this tool
+	 * @return The new instance of this class
+	 */
+	static TSharedRef<ISequencerTrackEditor> CreateTrackEditor(TSharedRef<ISequencer> OwningSequencer);
+protected:
+
+	virtual void GenerateKeysFromPropertyChanged(const FPropertyChangedParams& PropertyChangedParams, UMovieSceneSection* SectionToKey, FGeneratedTrackKeys& OutGeneratedKeys) override;
+
+private:
+
+	double RecomposeDouble(double InCurrentValue, UObject* AnimatedObject, UMovieSceneSection* Section);
+};
diff --git a/Engine/Source/Editor/MovieSceneTools/Private/TrackEditors/PropertyTrackEditors/EulerTransformPropertyTrackEditor.cpp b/Engine/Source/Editor/MovieSceneTools/Private/TrackEditors/PropertyTrackEditors/EulerTransformPropertyTrackEditor.cpp
--- a/Engine/Source/Editor/MovieSceneTools/Private/TrackEditors/PropertyTrackEditors/EulerTransformPropertyTrackEditor.cpp
+++ b/Engine/Source/Editor/MovieSceneTools/Private/TrackEditors/PropertyTrackEditors/EulerTransformPropertyTrackEditor.cpp
@@ -138,11 +138,11 @@ void FEulerTransformPropertyTrackEditor::ProcessKeyOperation(UObject* ObjectToKe
 			for (TSharedPtr<IKeyArea> KeyArea : SectionsToKey[Index].KeyAreas)
 			{
 				FMovieSceneChannelHandle Handle  = KeyArea->GetChannel();
-				if (Handle.GetChannelTypeName() == FMovieSceneFloatChannel::StaticStruct()->GetFName() && Handle.GetChannelIndex() < 9)
+				if (Handle.GetChannelTypeName() == FMovieSceneDoubleChannel::StaticStruct()->GetFName() && Handle.GetChannelIndex() < 9)
 				{
-					FMovieSceneFloatChannel* Channel = static_cast<FMovieSceneFloatChannel*>(Handle.Get());
+					FMovieSceneDoubleChannel* Channel = static_cast<FMovieSceneDoubleChannel*>(Handle.Get());
 
-					float Value = RecomposedTransform[Handle.GetChannelIndex()];
+					double Value = RecomposedTransform[Handle.GetChannelIndex()];
 					AddKeyToChannel(Channel, KeyTime, Value, InSequencer.GetKeyInterpolation());
 				}
 				else
@@ -203,20 +203,20 @@ void FEulerTransformPropertyTrackEditor::GenerateKeysFromPropertyChanged( const
 	FEulerTransform Transform           = PropertyChangedParams.GetPropertyValue<FEulerTransform>();
 	FEulerTransform RecomposedTransform = RecomposeTransform(Transform, PropertyChangedParams.ObjectsThatChanged[0], SectionToKey);
 
-	FVector3f Translation = RecomposedTransform.Location;
-	OutGeneratedKeys.Add(FMovieSceneChannelValueSetter::Create<FMovieSceneFloatChannel>(0, Translation.X, bKeyChannels[0]));
-	OutGeneratedKeys.Add(FMovieSceneChannelValueSetter::Create<FMovieSceneFloatChannel>(1, Translation.Y, bKeyChannels[1]));
-	OutGeneratedKeys.Add(FMovieSceneChannelValueSetter::Create<FMovieSceneFloatChannel>(2, Translation.Z, bKeyChannels[2]));
+	FVector Translation = RecomposedTransform.Location;
+	OutGeneratedKeys.Add(FMovieSceneChannelValueSetter::Create<FMovieSceneDoubleChannel>(0, (double)Translation.X, bKeyChannels[0]));
+	OutGeneratedKeys.Add(FMovieSceneChannelValueSetter::Create<FMovieSceneDoubleChannel>(1, (double)Translation.Y, bKeyChannels[1]));
+	OutGeneratedKeys.Add(FMovieSceneChannelValueSetter::Create<FMovieSceneDoubleChannel>(2, (double)Translation.Z, bKeyChannels[2]));
 
 	FRotator Rotator = RecomposedTransform.Rotation;
-	OutGeneratedKeys.Add(FMovieSceneChannelValueSetter::Create<FMovieSceneFloatChannel>(3, Rotator.Roll,  bKeyChannels[3]));
-	OutGeneratedKeys.Add(FMovieSceneChannelValueSetter::Create<FMovieSceneFloatChannel>(4, Rotator.Pitch, bKeyChannels[4]));
-	OutGeneratedKeys.Add(FMovieSceneChannelValueSetter::Create<FMovieSceneFloatChannel>(5, Rotator.Yaw,   bKeyChannels[5]));
-
-	FVector3f Scale = RecomposedTransform.Scale;
-	OutGeneratedKeys.Add(FMovieSceneChannelValueSetter::Create<FMovieSceneFloatChannel>(6, Scale.X, bKeyChannels[6]));
-	OutGeneratedKeys.Add(FMovieSceneChannelValueSetter::Create<FMovieSceneFloatChannel>(7, Scale.Y, bKeyChannels[7]));
-	OutGeneratedKeys.Add(FMovieSceneChannelValueSetter::Create<FMovieSceneFloatChannel>(8, Scale.Z, bKeyChannels[8]));
+	OutGeneratedKeys.Add(FMovieSceneChannelValueSetter::Create<FMovieSceneDoubleChannel>(3, (double)Rotator.Roll,  bKeyChannels[3]));
+	OutGeneratedKeys.Add(FMovieSceneChannelValueSetter::Create<FMovieSceneDoubleChannel>(4, (double)Rotator.Pitch, bKeyChannels[4]));
+	OutGeneratedKeys.Add(FMovieSceneChannelValueSetter::Create<FMovieSceneDoubleChannel>(5, (double)Rotator.Yaw,   bKeyChannels[5]));
+
+	FVector Scale = RecomposedTransform.Scale;
+	OutGeneratedKeys.Add(FMovieSceneChannelValueSetter::Create<FMovieSceneDoubleChannel>(6, (double)Scale.X, bKeyChannels[6]));
+	OutGeneratedKeys.Add(FMovieSceneChannelValueSetter::Create<FMovieSceneDoubleChannel>(7, (double)Scale.Y, bKeyChannels[7]));
+	OutGeneratedKeys.Add(FMovieSceneChannelValueSetter::Create<FMovieSceneDoubleChannel>(8, (double)Scale.Z, bKeyChannels[8]));
 }
 
 
@@ -245,4 +245,4 @@ FEulerTransform FEulerTransformPropertyTrackEditor::RecomposeTransform(const FEu
 	}
 
 	return InTransform;
-}
\ No newline at end of file
+}
diff --git a/Engine/Source/Editor/MovieSceneTools/Private/TrackEditors/PropertyTrackEditors/TransformPropertyTrackEditor.cpp b/Engine/Source/Editor/MovieSceneTools/Private/TrackEditors/PropertyTrackEditors/TransformPropertyTrackEditor.cpp
--- a/Engine/Source/Editor/MovieSceneTools/Private/TrackEditors/PropertyTrackEditors/TransformPropertyTrackEditor.cpp
+++ b/Engine/Source/Editor/MovieSceneTools/Private/TrackEditors/PropertyTrackEditors/TransformPropertyTrackEditor.cpp
@@ -58,17 +58,17 @@ void FTransformPropertyTrackEditor::GenerateKeysFromPropertyChanged( const FProp
 
 	FIntermediate3DTransform Recomposed = RecomposeTransform(CurrentTransform, PropertyChangedParams.ObjectsThatChanged[0], SectionToKey);
 
-	OutGeneratedKeys.Add(FMovieSceneChannelValueSetter::Create<FMovieSceneFloatChannel>(0, Recomposed.T_X, true));
-	OutGeneratedKeys.Add(FMovieSceneChannelValueSetter::Create<FMovieSceneFloatChannel>(1, Recomposed.T_Y, true));
-	OutGeneratedKeys.Add(FMovieSceneChannelValueSetter::Create<FMovieSceneFloatChannel>(2, Recomposed.T_Z, true));
+	OutGeneratedKeys.Add(FMovieSceneChannelValueSetter::Create<FMovieSceneDoubleChannel>(0, Recomposed.T_X, true));
+	OutGeneratedKeys.Add(FMovieSceneChannelValueSetter::Create<FMovieSceneDoubleChannel>(1, Recomposed.T_Y, true));
+	OutGeneratedKeys.Add(FMovieSceneChannelValueSetter::Create<FMovieSceneDoubleChannel>(2, Recomposed.T_Z, true));
 
-	OutGeneratedKeys.Add(FMovieSceneChannelValueSetter::Create<FMovieSceneFloatChannel>(3, Recomposed.R_X, true));
-	OutGeneratedKeys.Add(FMovieSceneChannelValueSetter::Create<FMovieSceneFloatChannel>(4, Recomposed.R_Y, true));
-	OutGeneratedKeys.Add(FMovieSceneChannelValueSetter::Create<FMovieSceneFloatChannel>(5, Recomposed.R_Z, true));
+	OutGeneratedKeys.Add(FMovieSceneChannelValueSetter::Create<FMovieSceneDoubleChannel>(3, Recomposed.R_X, true));
+	OutGeneratedKeys.Add(FMovieSceneChannelValueSetter::Create<FMovieSceneDoubleChannel>(4, Recomposed.R_Y, true));
+	OutGeneratedKeys.Add(FMovieSceneChannelValueSetter::Create<FMovieSceneDoubleChannel>(5, Recomposed.R_Z, true));
 
-	OutGeneratedKeys.Add(FMovieSceneChannelValueSetter::Create<FMovieSceneFloatChannel>(6, Recomposed.S_X, true));
-	OutGeneratedKeys.Add(FMovieSceneChannelValueSetter::Create<FMovieSceneFloatChannel>(7, Recomposed.S_Y, true));
-	OutGeneratedKeys.Add(FMovieSceneChannelValueSetter::Create<FMovieSceneFloatChannel>(8, Recomposed.S_Z, true));
+	OutGeneratedKeys.Add(FMovieSceneChannelValueSetter::Create<FMovieSceneDoubleChannel>(6, Recomposed.S_X, true));
+	OutGeneratedKeys.Add(FMovieSceneChannelValueSetter::Create<FMovieSceneDoubleChannel>(7, Recomposed.S_Y, true));
+	OutGeneratedKeys.Add(FMovieSceneChannelValueSetter::Create<FMovieSceneDoubleChannel>(8, Recomposed.S_Z, true));
 }
 
 UE::MovieScene::FIntermediate3DTransform FTransformPropertyTrackEditor::RecomposeTransform(const UE::MovieScene::FIntermediate3DTransform& InTransformData, UObject* AnimatedObject, UMovieSceneSection* Section)
@@ -94,4 +94,4 @@ UE::MovieScene::FIntermediate3DTransform FTransformPropertyTrackEditor::Recompos
 	}
 
 	return InTransformData;
-}
\ No newline at end of file
+}
diff --git a/Engine/Source/Editor/MovieSceneTools/Private/TrackEditors/PropertyTrackEditors/VectorPropertyTrackEditor.cpp b/Engine/Source/Editor/MovieSceneTools/Private/TrackEditors/PropertyTrackEditors/VectorPropertyTrackEditor.cpp
--- a/Engine/Source/Editor/MovieSceneTools/Private/TrackEditors/PropertyTrackEditors/VectorPropertyTrackEditor.cpp
+++ b/Engine/Source/Editor/MovieSceneTools/Private/TrackEditors/PropertyTrackEditors/VectorPropertyTrackEditor.cpp
@@ -11,18 +11,18 @@
 #include "EntitySystem/Interrogation/MovieSceneInterrogationLinker.h"
 
 
-FName FVectorPropertyTrackEditor::XName( "X" );
-FName FVectorPropertyTrackEditor::YName( "Y" );
-FName FVectorPropertyTrackEditor::ZName( "Z" );
-FName FVectorPropertyTrackEditor::WName( "W" );
+FName FFloatVectorPropertyTrackEditor::XName( "X" );
+FName FFloatVectorPropertyTrackEditor::YName( "Y" );
+FName FFloatVectorPropertyTrackEditor::ZName( "Z" );
+FName FFloatVectorPropertyTrackEditor::WName( "W" );
 
 
-TSharedRef<ISequencerTrackEditor> FVectorPropertyTrackEditor::CreateTrackEditor( TSharedRef<ISequencer> InSequencer )
+TSharedRef<ISequencerTrackEditor> FFloatVectorPropertyTrackEditor::CreateTrackEditor( TSharedRef<ISequencer> InSequencer )
 {
-	return MakeShareable( new FVectorPropertyTrackEditor( InSequencer ) );
+	return MakeShareable( new FFloatVectorPropertyTrackEditor( InSequencer ) );
 }
 
-void FVectorPropertyTrackEditor::GenerateKeysFromPropertyChanged( const FPropertyChangedParams& PropertyChangedParams, UMovieSceneSection* SectionToKey, FGeneratedTrackKeys& OutGeneratedKeys )
+void FFloatVectorPropertyTrackEditor::GenerateKeysFromPropertyChanged( const FPropertyChangedParams& PropertyChangedParams, UMovieSceneSection* SectionToKey, FGeneratedTrackKeys& OutGeneratedKeys )
 {
 	const FStructProperty* StructProp = CastField<const FStructProperty>( PropertyChangedParams.PropertyPath.GetLeafMostProperty().Property.Get() );
 	if (!StructProp)
@@ -32,7 +32,11 @@ void FVectorPropertyTrackEditor::GenerateKeysFromPropertyChanged( const FPropert
 	FName StructName = StructProp->Struct->GetFName();
 
 	bool bIsVector2D = StructName == NAME_Vector2D;
-	bool bIsVector = StructName == NAME_Vector;
+	bool bIsVector = (StructName == NAME_Vector3f
+#if UE_LARGE_WORLD_COORDINATES_DISABLED
+			|| StructName == NAME_Vector
+#endif
+			);
 	bool bIsVector4 = StructName == NAME_Vector4;
 
 	FVector4 VectorValues;
@@ -79,7 +83,7 @@ void FVectorPropertyTrackEditor::GenerateKeysFromPropertyChanged( const FPropert
 	}
 }
 
-void FVectorPropertyTrackEditor::InitializeNewTrack( UMovieSceneVectorTrack* NewTrack, FPropertyChangedParams PropertyChangedParams )
+void FFloatVectorPropertyTrackEditor::InitializeNewTrack( UMovieSceneFloatVectorTrack* NewTrack, FPropertyChangedParams PropertyChangedParams )
 {
 	FPropertyTrackEditor::InitializeNewTrack( NewTrack, PropertyChangedParams );
 	const FStructProperty* StructProp = CastField<const FStructProperty>( PropertyChangedParams.PropertyPath.GetLeafMostProperty().Property.Get() );
@@ -89,7 +93,11 @@ void FVectorPropertyTrackEditor::InitializeNewTrack( UMovieSceneVectorTrack* New
 	{
 		NewTrack->SetNumChannelsUsed( 2 );
 	}
-	if ( StructName == NAME_Vector )
+	if ( StructName == NAME_Vector3f
+#if UE_LARGE_WORLD_COORDINATES_DISABLED
+			|| StructName == NAME_Vector 
+#endif
+			)
 	{
 		NewTrack->SetNumChannelsUsed( 3 );
 	}
@@ -99,7 +107,7 @@ void FVectorPropertyTrackEditor::InitializeNewTrack( UMovieSceneVectorTrack* New
 	}
 }
 
-void CopyInterpVectorTrack(TSharedRef<ISequencer> Sequencer, UInterpTrackVectorProp* MatineeVectorTrack, UMovieSceneVectorTrack* VectorTrack)
+void CopyInterpVectorTrack(TSharedRef<ISequencer> Sequencer, UInterpTrackVectorProp* MatineeVectorTrack, UMovieSceneFloatVectorTrack* VectorTrack)
 {
 	if (FMatineeImportTools::CopyInterpVectorTrack(MatineeVectorTrack, VectorTrack))
 	{
@@ -107,11 +115,11 @@ void CopyInterpVectorTrack(TSharedRef<ISequencer> Sequencer, UInterpTrackVectorP
 	}
 }
 
-bool FVectorPropertyTrackEditor::ModifyGeneratedKeysByCurrentAndWeight(UObject *Object, UMovieSceneTrack *Track, UMovieSceneSection* SectionToKey, FFrameNumber KeyTime, FGeneratedTrackKeys& GeneratedTotalKeys, float Weight) const
+bool FFloatVectorPropertyTrackEditor::ModifyGeneratedKeysByCurrentAndWeight(UObject *Object, UMovieSceneTrack *Track, UMovieSceneSection* SectionToKey, FFrameNumber KeyTime, FGeneratedTrackKeys& GeneratedTotalKeys, float Weight) const
 {
 	using namespace UE::MovieScene;
 
-	UMovieSceneVectorTrack* VectorTrack = Cast<UMovieSceneVectorTrack>(Track);
+	UMovieSceneFloatVectorTrack* VectorTrack = Cast<UMovieSceneFloatVectorTrack>(Track);
 
 	if (VectorTrack)
 	{
@@ -126,8 +134,8 @@ bool FVectorPropertyTrackEditor::ModifyGeneratedKeysByCurrentAndWeight(UObject *
 		Interrogator.Update();
 
 		const FMovieSceneTracksComponentTypes* ComponentTypes = FMovieSceneTracksComponentTypes::Get();
-		TArray<FIntermediateVector> InterrogatedValues;
-		Interrogator.QueryPropertyValues(ComponentTypes->Vector, InterrogationChannel, InterrogatedValues);
+		TArray<FFloatIntermediateVector> InterrogatedValues;
+		Interrogator.QueryPropertyValues(ComponentTypes->FloatVector, InterrogationChannel, InterrogatedValues);
 
 		switch (VectorTrack->GetNumChannelsUsed())
 		{
@@ -165,4 +173,103 @@ bool FVectorPropertyTrackEditor::ModifyGeneratedKeysByCurrentAndWeight(UObject *
 	}
 
 	return false;
-}
\ No newline at end of file
+}
+
+
+FName FDoubleVectorPropertyTrackEditor::XName( "X" );
+FName FDoubleVectorPropertyTrackEditor::YName( "Y" );
+FName FDoubleVectorPropertyTrackEditor::ZName( "Z" );
+FName FDoubleVectorPropertyTrackEditor::WName( "W" );
+
+
+TSharedRef<ISequencerTrackEditor> FDoubleVectorPropertyTrackEditor::CreateTrackEditor( TSharedRef<ISequencer> InSequencer )
+{
+	return MakeShareable( new FDoubleVectorPropertyTrackEditor( InSequencer ) );
+}
+
+void FDoubleVectorPropertyTrackEditor::GenerateKeysFromPropertyChanged( const FPropertyChangedParams& PropertyChangedParams, UMovieSceneSection* SectionToKey, FGeneratedTrackKeys& OutGeneratedKeys )
+{
+	const FStructProperty* StructProp = CastField<const FStructProperty>( PropertyChangedParams.PropertyPath.GetLeafMostProperty().Property.Get() );
+	if (!StructProp)
+	{
+		return;
+	}
+	FName StructName = StructProp->Struct->GetFName();
+
+	bool bIsVector = (StructName == NAME_Vector3d
+#if !UE_LARGE_WORLD_COORDINATES_DISABLED
+			|| StructName == NAME_Vector
+#endif
+			);
+	ensure(bIsVector);
+
+	FVector3d VectorValues = PropertyChangedParams.GetPropertyValue<FVector3d>();
+
+	FPropertyPath StructPath = PropertyChangedParams.StructPathToKey;
+	FName ChannelName = StructPath.GetNumProperties() != 0 ? StructPath.GetLeafMostProperty().Property->GetFName() : NAME_None;
+
+	const bool bKeyX = ChannelName == NAME_None || ChannelName == XName;
+	const bool bKeyY = ChannelName == NAME_None || ChannelName == YName;
+	const bool bKeyZ = ChannelName == NAME_None || ChannelName == ZName;
+
+	OutGeneratedKeys.Add(FMovieSceneChannelValueSetter::Create<FMovieSceneDoubleChannel>(0, VectorValues.X, bKeyX));
+	OutGeneratedKeys.Add(FMovieSceneChannelValueSetter::Create<FMovieSceneDoubleChannel>(1, VectorValues.Y, bKeyY));
+	OutGeneratedKeys.Add(FMovieSceneChannelValueSetter::Create<FMovieSceneDoubleChannel>(2, VectorValues.Z, bKeyZ));
+}
+
+void FDoubleVectorPropertyTrackEditor::InitializeNewTrack( UMovieSceneDoubleVectorTrack* NewTrack, FPropertyChangedParams PropertyChangedParams )
+{
+	FPropertyTrackEditor::InitializeNewTrack( NewTrack, PropertyChangedParams );
+	const FStructProperty* StructProp = CastField<const FStructProperty>( PropertyChangedParams.PropertyPath.GetLeafMostProperty().Property.Get() );
+	FName StructName = StructProp->Struct->GetFName();
+
+	if (StructName == NAME_Vector3d
+#if !UE_LARGE_WORLD_COORDINATES_DISABLED
+			|| StructName == NAME_Vector
+#endif
+			)
+	{
+		NewTrack->SetNumChannelsUsed( 3 );
+	}
+}
+
+bool FDoubleVectorPropertyTrackEditor::ModifyGeneratedKeysByCurrentAndWeight(UObject *Object, UMovieSceneTrack *Track, UMovieSceneSection* SectionToKey, FFrameNumber KeyTime, FGeneratedTrackKeys& GeneratedTotalKeys, float Weight) const
+{
+	using namespace UE::MovieScene;
+
+	UMovieSceneDoubleVectorTrack* VectorTrack = Cast<UMovieSceneDoubleVectorTrack>(Track);
+
+	if (VectorTrack)
+	{
+		FSystemInterrogator Interrogator;
+
+		TGuardValue<FEntityManager*> DebugVizGuard(GEntityManagerForDebuggingVisualizers, &Interrogator.GetLinker()->EntityManager);
+
+		Interrogator.ImportTrack(VectorTrack, FInterrogationChannel::Default());
+		Interrogator.AddInterrogation(KeyTime);
+
+		Interrogator.Update();
+
+		const FMovieSceneTracksComponentTypes* ComponentTypes = FMovieSceneTracksComponentTypes::Get();
+		TArray<FDoubleIntermediateVector> InterrogatedValues;
+		Interrogator.QueryPropertyValues(ComponentTypes->DoubleVector, InterrogatedValues);
+
+		switch (VectorTrack->GetNumChannelsUsed())
+		{
+		case 3:
+			{
+				FVector3d Val(InterrogatedValues[0].X, InterrogatedValues[0].Y, InterrogatedValues[0].Z);
+				FMovieSceneChannelProxy& Proxy = SectionToKey->GetChannelProxy();
+				GeneratedTotalKeys[0]->ModifyByCurrentAndWeight(Proxy, KeyTime, (void *)&Val.X, Weight);
+				GeneratedTotalKeys[1]->ModifyByCurrentAndWeight(Proxy, KeyTime, (void *)&Val.Y, Weight);
+				GeneratedTotalKeys[2]->ModifyByCurrentAndWeight(Proxy, KeyTime, (void *)&Val.Z, Weight);
+			}
+			break;
+		}
+
+		return true;
+	}
+
+	return false;
+}
+
diff --git a/Engine/Source/Editor/MovieSceneTools/Private/TrackEditors/PropertyTrackEditors/VectorPropertyTrackEditor.h b/Engine/Source/Editor/MovieSceneTools/Private/TrackEditors/PropertyTrackEditors/VectorPropertyTrackEditor.h
--- a/Engine/Source/Editor/MovieSceneTools/Private/TrackEditors/PropertyTrackEditors/VectorPropertyTrackEditor.h
+++ b/Engine/Source/Editor/MovieSceneTools/Private/TrackEditors/PropertyTrackEditors/VectorPropertyTrackEditor.h
@@ -10,13 +10,13 @@
 #include "ISequencerTrackEditor.h"
 #include "PropertyTrackEditor.h"
 #include "Tracks/MovieSceneVectorTrack.h"
-#include "Sections/MovieSceneVectorSection.h"
+#include "Misc/LargeWorldCoordinates.h"
 
 /**
- * A property track editor for vectors.
+ * A property track editor for float vectors.
  */
-class FVectorPropertyTrackEditor
-	: public FPropertyTrackEditor<UMovieSceneVectorTrack>
+class FFloatVectorPropertyTrackEditor
+	: public FPropertyTrackEditor<UMovieSceneFloatVectorTrack>
 {
 public:
 
@@ -25,7 +25,7 @@ public:
 	 *
 	 * @param InSequencer The sequencer instance to be used by this tool.
 	 */
-	FVectorPropertyTrackEditor(TSharedRef<ISequencer> InSequencer)
+	FFloatVectorPropertyTrackEditor(TSharedRef<ISequencer> InSequencer)
 		: FPropertyTrackEditor(InSequencer, GetAnimatedPropertyTypes())
 	{
 	}
@@ -36,7 +36,10 @@ public:
 	static TArray<FAnimatedPropertyKey, TInlineAllocator<3>> GetAnimatedPropertyTypes()
 	{
 		return TArray<FAnimatedPropertyKey, TInlineAllocator<3>>({
+#if UE_LARGE_WORLD_COORDINATES_DISABLED
 			FAnimatedPropertyKey::FromStructType(NAME_Vector),
+#endif
+			FAnimatedPropertyKey::FromStructType(NAME_Vector3f),
 			FAnimatedPropertyKey::FromStructType(NAME_Vector4),
 			FAnimatedPropertyKey::FromStructType(NAME_Vector2D)
 		});
@@ -54,7 +57,61 @@ protected:
 	//~ FPropertyTrackEditor interface
 
 	virtual void GenerateKeysFromPropertyChanged(const FPropertyChangedParams& PropertyChangedParams, UMovieSceneSection* SectionToKey, FGeneratedTrackKeys& OutGeneratedKeys) override;
-	virtual void InitializeNewTrack(UMovieSceneVectorTrack* NewTrack, FPropertyChangedParams PropertyChangedParams) override;
+	virtual void InitializeNewTrack(UMovieSceneFloatVectorTrack* NewTrack, FPropertyChangedParams PropertyChangedParams) override;
+	virtual bool ModifyGeneratedKeysByCurrentAndWeight(UObject* Object, UMovieSceneTrack *Track, UMovieSceneSection* SectionToKey, FFrameNumber Time, FGeneratedTrackKeys& GeneratedTotalKeys, float Weight) const override;
+
+private:
+
+	static FName XName;
+	static FName YName;
+	static FName ZName;
+	static FName WName;
+};
+
+/**
+ * A property track editor for double vectors.
+ */
+class FDoubleVectorPropertyTrackEditor
+	: public FPropertyTrackEditor<UMovieSceneDoubleVectorTrack>
+{
+public:
+
+	/**
+	 * Constructor.
+	 *
+	 * @param InSequencer The sequencer instance to be used by this tool.
+	 */
+	FDoubleVectorPropertyTrackEditor(TSharedRef<ISequencer> InSequencer)
+		: FPropertyTrackEditor(InSequencer, GetAnimatedPropertyTypes())
+	{
+	}
+
+	/**
+	 * Retrieve a list of all property types that this track editor animates
+	 */
+	static TArray<FAnimatedPropertyKey, TInlineAllocator<1>> GetAnimatedPropertyTypes()
+	{
+		return TArray<FAnimatedPropertyKey, TInlineAllocator<1>>({
+#if !UE_LARGE_WORLD_COORDINATES_DISABLED
+			FAnimatedPropertyKey::FromStructType(NAME_Vector),
+#endif
+			FAnimatedPropertyKey::FromStructType(NAME_Vector3d)
+		});
+	}
+
+	/**
+	 * Creates an instance of this class (called by a sequence).
+	 *
+	 * @param OwningSequencer The sequencer instance to be used by this tool.
+	 * @return The new instance of this class.
+	 */
+	static TSharedRef<ISequencerTrackEditor> CreateTrackEditor(TSharedRef<ISequencer> OwningSequencer);
+
+protected:
+	//~ FPropertyTrackEditor interface
+
+	virtual void GenerateKeysFromPropertyChanged(const FPropertyChangedParams& PropertyChangedParams, UMovieSceneSection* SectionToKey, FGeneratedTrackKeys& OutGeneratedKeys) override;
+	virtual void InitializeNewTrack(UMovieSceneDoubleVectorTrack* NewTrack, FPropertyChangedParams PropertyChangedParams) override;
 	virtual bool ModifyGeneratedKeysByCurrentAndWeight(UObject* Object, UMovieSceneTrack *Track, UMovieSceneSection* SectionToKey, FFrameNumber Time, FGeneratedTrackKeys& GeneratedTotalKeys, float Weight) const override;
 
 private:
diff --git a/Engine/Source/Editor/MovieSceneTools/Private/TrackEditors/TransformTrackEditor.cpp b/Engine/Source/Editor/MovieSceneTools/Private/TrackEditors/TransformTrackEditor.cpp
--- a/Engine/Source/Editor/MovieSceneTools/Private/TrackEditors/TransformTrackEditor.cpp
+++ b/Engine/Source/Editor/MovieSceneTools/Private/TrackEditors/TransformTrackEditor.cpp
@@ -422,7 +422,7 @@ void F3DTransformTrackEditor::BuildObjectBindingEditButtons(TSharedPtr<SHorizont
 		.Padding(4, 0, 0, 0)
 		[
 			SNew(SCheckBox)		
-                .Style( &FAppStyle::Get().GetWidgetStyle<FCheckBoxStyle>("ToggleButtonCheckBoxAlt"))
+				.Style( &FAppStyle::Get().GetWidgetStyle<FCheckBoxStyle>("ToggleButtonCheckBoxAlt"))
 				.Type(ESlateCheckBoxType::CheckBox)
 				.IsFocusable(false)
 				.Visibility(this, &F3DTransformTrackEditor::IsCameraVisible, ObjectGuid)
@@ -643,7 +643,7 @@ FText F3DTransformTrackEditor::GetLockCameraToolTip(FGuid ObjectGuid) const
 	return FText();
 }
 
-float UnwindChannel(const float& OldValue, float NewValue)
+double UnwindChannel(const double& OldValue, double NewValue)
 {
 	while( NewValue - OldValue > 180.0f )
 	{
@@ -716,11 +716,11 @@ void F3DTransformTrackEditor::GetTransformKeys( const TOptional<FTransformData>&
 			bKeyZ = false;
 		}
 
-		FVector3f KeyVector = RecomposedTransform.Translation;
+		FVector KeyVector = RecomposedTransform.Translation;
 
-		OutGeneratedKeys.Add(FMovieSceneChannelValueSetter::Create<FMovieSceneFloatChannel>(0, KeyVector.X, bKeyX));
-		OutGeneratedKeys.Add(FMovieSceneChannelValueSetter::Create<FMovieSceneFloatChannel>(1, KeyVector.Y, bKeyY));
-		OutGeneratedKeys.Add(FMovieSceneChannelValueSetter::Create<FMovieSceneFloatChannel>(2, KeyVector.Z, bKeyZ));
+		OutGeneratedKeys.Add(FMovieSceneChannelValueSetter::Create<FMovieSceneDoubleChannel>(0, (double)KeyVector.X, bKeyX));
+		OutGeneratedKeys.Add(FMovieSceneChannelValueSetter::Create<FMovieSceneDoubleChannel>(1, (double)KeyVector.Y, bKeyY));
+		OutGeneratedKeys.Add(FMovieSceneChannelValueSetter::Create<FMovieSceneDoubleChannel>(2, (double)KeyVector.Z, bKeyZ));
 	}
 
 	// Set rotation keys/defaults
@@ -759,9 +759,9 @@ void F3DTransformTrackEditor::GetTransformKeys( const TOptional<FTransformData>&
 
 		// Do we need to unwind re-composed rotations?
 		KeyRotator = UnwindRotator(CurrentTransform.Rotation, RecomposedTransform.Rotation);
-		OutGeneratedKeys.Add(FMovieSceneChannelValueSetter::Create<FMovieSceneFloatChannel>(3, KeyRotator.Roll, bKeyX));
-		OutGeneratedKeys.Add(FMovieSceneChannelValueSetter::Create<FMovieSceneFloatChannel>(4, KeyRotator.Pitch, bKeyY));
-		OutGeneratedKeys.Add(FMovieSceneChannelValueSetter::Create<FMovieSceneFloatChannel>(5, KeyRotator.Yaw, bKeyZ));
+		OutGeneratedKeys.Add(FMovieSceneChannelValueSetter::Create<FMovieSceneDoubleChannel>(3, (double)KeyRotator.Roll, bKeyX));
+		OutGeneratedKeys.Add(FMovieSceneChannelValueSetter::Create<FMovieSceneDoubleChannel>(4, (double)KeyRotator.Pitch, bKeyY));
+		OutGeneratedKeys.Add(FMovieSceneChannelValueSetter::Create<FMovieSceneDoubleChannel>(5, (double)KeyRotator.Yaw, bKeyZ));
 
 	}
 
@@ -796,10 +796,10 @@ void F3DTransformTrackEditor::GetTransformKeys( const TOptional<FTransformData>&
 			bKeyZ = false;
 		}
 
-		FVector3f KeyVector = RecomposedTransform.Scale;
-		OutGeneratedKeys.Add(FMovieSceneChannelValueSetter::Create<FMovieSceneFloatChannel>(6, KeyVector.X, bKeyX));
-		OutGeneratedKeys.Add(FMovieSceneChannelValueSetter::Create<FMovieSceneFloatChannel>(7, KeyVector.Y, bKeyY));
-		OutGeneratedKeys.Add(FMovieSceneChannelValueSetter::Create<FMovieSceneFloatChannel>(8, KeyVector.Z, bKeyZ));
+		FVector KeyVector = RecomposedTransform.Scale;
+		OutGeneratedKeys.Add(FMovieSceneChannelValueSetter::Create<FMovieSceneDoubleChannel>(6, (double)KeyVector.X, bKeyX));
+		OutGeneratedKeys.Add(FMovieSceneChannelValueSetter::Create<FMovieSceneDoubleChannel>(7, (double)KeyVector.Y, bKeyY));
+		OutGeneratedKeys.Add(FMovieSceneChannelValueSetter::Create<FMovieSceneDoubleChannel>(8, (double)KeyVector.Z, bKeyZ));
 	}
 }
 
@@ -989,11 +989,11 @@ void F3DTransformTrackEditor::ProcessKeyOperation(UObject* ObjectToKey, TArrayVi
 			for (TSharedPtr<IKeyArea> KeyArea : SectionsToKey[Index].KeyAreas)
 			{
 				FMovieSceneChannelHandle Handle  = KeyArea->GetChannel();
-				if (Handle.GetChannelTypeName() == FMovieSceneFloatChannel::StaticStruct()->GetFName() && Handle.GetChannelIndex() < 9)
+				if (Handle.GetChannelTypeName() == FMovieSceneDoubleChannel::StaticStruct()->GetFName() && Handle.GetChannelIndex() < 9)
 				{
-					FMovieSceneFloatChannel* Channel = static_cast<FMovieSceneFloatChannel*>(Handle.Get());
+					FMovieSceneDoubleChannel* Channel = static_cast<FMovieSceneDoubleChannel*>(Handle.Get());
 
-					float Value = RecomposedTransform[Handle.GetChannelIndex()];
+					double Value = RecomposedTransform[Handle.GetChannelIndex()];
 					AddKeyToChannel(Channel, KeyTime, Value, InSequencer.GetKeyInterpolation());
 				}
 				else
@@ -1005,15 +1005,15 @@ void F3DTransformTrackEditor::ProcessKeyOperation(UObject* ObjectToKey, TArrayVi
 	}
 }
 
-void AddUnwoundKey(FMovieSceneFloatChannel& Channel, FFrameNumber Time, float Value)
+void AddUnwoundKey(FMovieSceneDoubleChannel& Channel, FFrameNumber Time, double Value)
 {
 	int32 Index = Channel.AddLinearKey(Time, Value);
 
-	TArrayView<FMovieSceneFloatValue> Values = Channel.GetData().GetValues();
+	TArrayView<FMovieSceneDoubleValue> Values = Channel.GetData().GetValues();
 	if (Index >= 1)
 	{
-		const float PreviousValue = Values[Index - 1].Value;
-		float NewValue = Value;
+		const double PreviousValue = Values[Index - 1].Value;
+		double NewValue = Value;
 
 		while (NewValue - PreviousValue > 180.0f)
 		{
@@ -1083,17 +1083,17 @@ void F3DTransformTrackEditor::ImportAnimSequenceTransforms(const FAssetData& Ass
 		
 		FFrameRate TickResolution = Section->GetTypedOuter<UMovieScene>()->GetTickResolution();
 
-		TArrayView<FMovieSceneFloatChannel*> FloatChannels = Section->GetChannelProxy().GetChannels<FMovieSceneFloatChannel>();
+		TArrayView<FMovieSceneDoubleChannel*> DoubleChannels = Section->GetChannelProxy().GetChannels<FMovieSceneDoubleChannel>();
 
 		// Set default translation and rotation
 		for (int32 Index = 0; Index < 6; ++Index)
 		{
-			FloatChannels[Index]->SetDefault(0.f);
+			DoubleChannels[Index]->SetDefault(0.f);
 		}
 		// Set default scale
 		for (int32 Index = 6; Index < 9; ++Index)
 		{
-			FloatChannels[Index]->SetDefault(1.f);
+			DoubleChannels[Index]->SetDefault(1.f);
 		}
 
 		TransformTrack->AddSection(*Section);
@@ -1175,7 +1175,7 @@ void F3DTransformTrackEditor::ImportAnimSequenceTransforms(const FAssetData& Ass
 				const FVector3f Rotation = TempKey.WoundRotation.Euler();
 				const FVector3f Scale = TempKey.Transform.GetScale3D();
 
-				TArrayView<FMovieSceneFloatChannel*> Channels = Section->GetChannelProxy().GetChannels<FMovieSceneFloatChannel>();
+				TArrayView<FMovieSceneDoubleChannel*> Channels = Section->GetChannelProxy().GetChannels<FMovieSceneDoubleChannel>();
 
 				Channels[0]->AddLinearKey(KeyTime, Translation.X);
 				Channels[1]->AddLinearKey(KeyTime, Translation.Y);
diff --git a/Engine/Source/Editor/MovieSceneTools/Public/Channels/FloatChannelCurveModel.h b/Engine/Source/Editor/MovieSceneTools/Public/Channels/FloatChannelCurveModel.h
--- a/Engine/Source/Editor/MovieSceneTools/Public/Channels/FloatChannelCurveModel.h
+++ b/Engine/Source/Editor/MovieSceneTools/Public/Channels/FloatChannelCurveModel.h
@@ -2,45 +2,17 @@
 
 #pragma once
 
-#include "CoreTypes.h"
-#include "UObject/WeakObjectPtr.h"
-#include "UObject/WeakObjectPtrTemplates.h"
-#include "Channels/ChannelCurveModel.h"
-#include "IBufferedCurveModel.h"
-#include "MovieSceneSection.h"
+#include "Channels/BezierChannelCurveModel.h"
 #include "Channels/MovieSceneFloatChannel.h"
-#include "Channels/MovieSceneChannelHandle.h"
 
 struct FMovieSceneFloatChannel;
-class UMovieSceneSection;
-class ISequencer;
 
-class FFloatChannelCurveModel : public FChannelCurveModel<FMovieSceneFloatChannel, FMovieSceneFloatValue, float>
+class FFloatChannelCurveModel : public FBezierChannelCurveModel<FMovieSceneFloatChannel, FMovieSceneFloatValue, float>
 {
 public:
 	FFloatChannelCurveModel(TMovieSceneChannelHandle<FMovieSceneFloatChannel> InChannel, UMovieSceneSection* InOwningSection, TWeakPtr<ISequencer> InWeakSequencer);
 
 	// FCurveModel
-	virtual void DrawCurve(const FCurveEditor& CurveEditor, const FCurveEditorScreenSpace& ScreenSpace, TArray<TTuple<double, double>>& InterpolatingPoints) const override;
-	virtual void GetKeyDrawInfo(ECurvePointType PointType, const FKeyHandle InKeyHandle, FKeyDrawInfo& OutDrawInfo) const override;
-
-	virtual void GetKeyAttributes(TArrayView<const FKeyHandle> InKeys, TArrayView<FKeyAttributes> OutAttributes) const override;
-	virtual void SetKeyAttributes(TArrayView<const FKeyHandle> InKeys, TArrayView<const FKeyAttributes> InAttributes, EPropertyChangeType::Type ChangeType = EPropertyChangeType::Unspecified) override;
-
-	virtual void GetCurveAttributes(FCurveAttributes& OutCurveAttributes) const override;
-	virtual void SetCurveAttributes(const FCurveAttributes& InCurveAttributes) override;
-	virtual void GetValueRange(double& MinValue, double& MaxValue) const override;
-
 	virtual void CreateKeyProxies(TArrayView<const FKeyHandle> InKeyHandles, TArrayView<UObject*> OutObjects) override;
-
 	virtual TUniquePtr<IBufferedCurveModel> CreateBufferedCurveCopy() const override;
-
-protected:
-	// FChannelCurveModel
-	virtual double GetKeyValue(TArrayView<const FMovieSceneFloatValue> Values, int32 Index) const override;
-	virtual void SetKeyValue(int32 Index, double KeyValue) const override;
-
-private:
-	void FeaturePointMethod(double StartTime, double EndTime, double StartValue, double Mu, int Depth, int MaxDepth, double& MaxV, double& MinVal) const;
-
 };
\ No newline at end of file
diff --git a/Engine/Source/Editor/MovieSceneTools/Public/ClipboardTypes.h b/Engine/Source/Editor/MovieSceneTools/Public/ClipboardTypes.h
--- a/Engine/Source/Editor/MovieSceneTools/Public/ClipboardTypes.h
+++ b/Engine/Source/Editor/MovieSceneTools/Public/ClipboardTypes.h
@@ -8,6 +8,7 @@
 #if WITH_EDITOR
 
 #include "Channels/MovieSceneFloatChannel.h"
+#include "Channels/MovieSceneDoubleChannel.h"
 #include "Channels/MovieSceneEvent.h"
 #include "Channels/MovieSceneObjectPathChannel.h"
 #include "Channels/MovieSceneCameraShakeSourceTriggerChannel.h"
@@ -37,6 +38,11 @@ namespace MovieSceneClipboard
 		static FName Name("Float");
 		return Name;
 	}
+	template<> inline FName GetKeyTypeName<FMovieSceneDoubleValue>()
+	{
+		static FName Name("Double");
+		return Name;
+	}
 	template<> inline FName GetKeyTypeName<FName>()
 	{
 		static FName Name("Name");
diff --git a/Engine/Source/Editor/MovieSceneTools/Public/KeyframeTrackEditor.h b/Engine/Source/Editor/MovieSceneTools/Public/KeyframeTrackEditor.h
--- a/Engine/Source/Editor/MovieSceneTools/Public/KeyframeTrackEditor.h
+++ b/Engine/Source/Editor/MovieSceneTools/Public/KeyframeTrackEditor.h
@@ -19,6 +19,7 @@
 #include "Channels/MovieSceneChannelProxy.h"
 #include "Evaluation/MovieSceneEvalTemplate.h"
 #include "Evaluation/MovieSceneEvaluation.h"
+#include "Channels/MovieSceneDoubleChannel.h"
 #include "Channels/MovieSceneFloatChannel.h"
 #include "Channels/MovieSceneIntegerChannel.h"
 
@@ -97,6 +98,83 @@ struct TAddKeyImpl : IImpl
 
 //Specializations for channels that SUPPORT Blending
 template<>
+struct TAddKeyImpl<FMovieSceneDoubleChannel, double> : IImpl
+{
+	int32 ChannelIndex;
+	bool bAddKey;
+	double ValueToSet;
+
+	TAddKeyImpl(int32 InChannelIndex, bool bInAddKey,  const double& InValue)
+		: ChannelIndex(InChannelIndex), bAddKey(bInAddKey), ValueToSet(InValue)
+	{}
+
+	virtual bool Apply(UMovieSceneSection* Section, FMovieSceneChannelProxy& Proxy, FFrameNumber InTime, EMovieSceneKeyInterpolation InterpolationMode, bool bKeyEvenIfUnchanged, bool bKeyEvenIfEmpty) const override
+	{
+		bool bKeyCreated = false;
+		using namespace UE::MovieScene;
+
+		FMovieSceneDoubleChannel* Channel = Proxy.GetChannel<FMovieSceneDoubleChannel>(ChannelIndex);
+		if (bAddKey && Channel)
+		{
+			bool bShouldKeyChannel = bKeyEvenIfUnchanged;
+			if (!bShouldKeyChannel)
+			{
+				bShouldKeyChannel = !ValueExistsAtTime(Channel, InTime, ValueToSet);
+			}
+
+			if (bShouldKeyChannel)
+			{
+				if (Channel->GetNumKeys() != 0 || bKeyEvenIfEmpty)
+				{
+					if (Section->TryModify())
+					{
+						AddKeyToChannel(Channel, InTime, ValueToSet, InterpolationMode);
+						bKeyCreated = true;
+					}
+				}
+			}
+		}
+
+		return bKeyCreated;
+	}
+
+	virtual void ApplyDefault(UMovieSceneSection* Section, FMovieSceneChannelProxy& Proxy) const override
+	{
+		using namespace UE::MovieScene;
+
+		FMovieSceneDoubleChannel* Channel = Proxy.GetChannel<FMovieSceneDoubleChannel>(ChannelIndex);
+		if (Channel && Channel->GetData().GetTimes().Num() == 0  && Channel->GetDefault() != ValueToSet)
+		{
+			if (Section->TryModify())
+			{
+				using namespace UE::MovieScene;
+				SetChannelDefault(Channel, ValueToSet);
+			}
+		}
+	}
+
+	virtual bool ModifyByCurrentAndWeight(FMovieSceneChannelProxy& Proxy, FFrameNumber InTime, void* VCurrentValue, float Weight) override
+	{
+		using namespace UE::MovieScene;
+		double CurrentValue = *(double*)(VCurrentValue);
+		FMovieSceneDoubleChannel* Channel = Proxy.GetChannel<FMovieSceneDoubleChannel>(ChannelIndex);
+		if (Channel)
+		{
+			double LocalValue;
+			using namespace UE::MovieScene;
+			if (!EvaluateChannel(Channel, InTime, LocalValue))
+			{
+				TOptional<double> OptDouble = Channel->GetDefault();
+				LocalValue = OptDouble.IsSet() ? OptDouble.GetValue() : 0.0f;
+			}
+			ValueToSet = (ValueToSet - CurrentValue) * Weight + LocalValue;
+			return true;
+		}
+
+		return false;
+	}
+};
+template<>
 struct TAddKeyImpl<FMovieSceneFloatChannel, float> : IImpl
 {
 	int32 ChannelIndex;
@@ -172,7 +250,6 @@ struct TAddKeyImpl<FMovieSceneFloatChannel, float> : IImpl
 
 		return false;
 	}
-
 };
 template<>
 struct TAddKeyImpl<FMovieSceneIntegerChannel, int32> : IImpl
@@ -250,7 +327,6 @@ struct TAddKeyImpl<FMovieSceneIntegerChannel, int32> : IImpl
 
 		return false;
 	}
-
 };
 struct FMovieSceneChannelValueSetter
 {
diff --git a/Engine/Source/Editor/MovieSceneTools/Public/MatineeImportTools.h b/Engine/Source/Editor/MovieSceneTools/Public/MatineeImportTools.h
--- a/Engine/Source/Editor/MovieSceneTools/Public/MatineeImportTools.h
+++ b/Engine/Source/Editor/MovieSceneTools/Public/MatineeImportTools.h
@@ -4,6 +4,7 @@
 #include "CoreMinimal.h"
 #include "Matinee/InterpTrackToggle.h"
 #include "Sections/MovieSceneParticleSection.h"
+#include "Channels/MovieSceneDoubleChannel.h"
 #include "Channels/MovieSceneFloatChannel.h"
 
 class AMatineeActor;
@@ -34,7 +35,7 @@ class UMovieSceneFadeTrack;
 class UMovieSceneFloatTrack;
 class UMovieSceneParticleTrack;
 class UMovieSceneSkeletalAnimationTrack;
-class UMovieSceneVectorTrack;
+class UMovieSceneFloatVectorTrack;
 class UMovieSceneVisibilityTrack;
 class UMovieSceneSlomoTrack;
 
@@ -67,6 +68,10 @@ public:
 	static void SetOrAddKey( TMovieSceneChannelData<FMovieSceneFloatValue>& Curve, FFrameNumber Time, float Value, float ArriveTangent, float LeaveTangent, EInterpCurveMode MatineeInterpMode, FFrameRate FrameRate
 	, ERichCurveTangentWeightMode WeightedMode = RCTWM_WeightedNone, float ArriveTangentWeight = 0.0f, float LeaveTangentWeight = 0.0f);
 
+	/** Adds a key to a rich curve based on matinee curve key data. */
+	static void SetOrAddKey( TMovieSceneChannelData<FMovieSceneDoubleValue>& Curve, FFrameNumber Time, double Value, float ArriveTangent, float LeaveTangent, EInterpCurveMode MatineeInterpMode, FFrameRate FrameRate
+	, ERichCurveTangentWeightMode WeightedMode = RCTWM_WeightedNone, float ArriveTangentWeight = 0.0f, float LeaveTangentWeight = 0.0f);
+
 	/** Copies keys from a matinee bool track to a sequencer bool track. */
 	static bool CopyInterpBoolTrack( UInterpTrackBoolProp* MatineeBoolTrack, UMovieSceneBoolTrack* BoolTrack );
 
@@ -80,7 +85,7 @@ public:
 	static bool CopyInterpMaterialParamTrack( UInterpTrackVectorMaterialParam* MatineeMaterialParamTrack, UMovieSceneComponentMaterialTrack* MaterialTrack);
 
 	/** Copies keys from a matinee vector track to a sequencer vector track. */
-	static bool CopyInterpVectorTrack( UInterpTrackVectorProp* MatineeVectorTrack, UMovieSceneVectorTrack* VectorTrack );
+	static bool CopyInterpVectorTrack( UInterpTrackVectorProp* MatineeVectorTrack, UMovieSceneFloatVectorTrack* VectorTrack );
 
 	/** Copies keys from a matinee move track to a sequencer transform track. */
 	static bool CopyInterpMoveTrack( UInterpTrackMove* MoveTrack, UMovieScene3DTransformTrack* TransformTrack, const FVector& DefaultScale = FVector(1.f) );
diff --git a/Engine/Source/Editor/MovieSceneTools/Public/MovieSceneToolHelpers.h b/Engine/Source/Editor/MovieSceneTools/Public/MovieSceneToolHelpers.h
--- a/Engine/Source/Editor/MovieSceneTools/Public/MovieSceneToolHelpers.h
+++ b/Engine/Source/Editor/MovieSceneTools/Public/MovieSceneToolHelpers.h
@@ -30,6 +30,7 @@ class UMovieSceneTrack;
 struct FMovieSceneEvaluationTrack;
 class UMovieSceneUserImportFBXSettings;
 class UMovieSceneUserImportFBXControlRigSettings;
+struct FMovieSceneDoubleValue;
 struct FMovieSceneFloatValue;
 class INodeNameAdapter;
 struct FMovieSceneSequenceTransform;
@@ -404,7 +405,7 @@ public:
 	 * @param MoveAxis The move axis to copy to
 	 * @param FrameRate The frame rate of the source channel
 	 */
-	static void CopyKeyDataToMoveAxis(const TMovieSceneChannelData<FMovieSceneFloatValue>& KeyData, UInterpTrackMoveAxis* MoveAxis, FFrameRate FrameRate);
+	static void CopyKeyDataToMoveAxis(const TMovieSceneChannelData<FMovieSceneDoubleValue>& KeyData, UInterpTrackMoveAxis* MoveAxis, FFrameRate FrameRate);
 
 	/*
 	 * Export the object binding to a camera anim
@@ -546,6 +547,12 @@ public:
 	*/
 	static void SetOrAddKey(TMovieSceneChannelData<FMovieSceneFloatValue>& ChannelData, FFrameNumber Time, float Value);
 
+	/* Set or add a key onto a double channel.
+	* @param ChannelData Channel to set or add
+	* @param Time Frame to add or set the value
+	* @param Value  Value to Set
+	*/
+	static void SetOrAddKey(TMovieSceneChannelData<FMovieSceneDoubleValue>& ChannelData, FFrameNumber Time, double Value);
 	
 	/*
 	*  Get an actors world transforms at the specified times using a player
diff --git a/Engine/Source/Editor/MovieSceneTools/Public/MovieSceneToolsModule.h b/Engine/Source/Editor/MovieSceneTools/Public/MovieSceneToolsModule.h
--- a/Engine/Source/Editor/MovieSceneTools/Public/MovieSceneToolsModule.h
+++ b/Engine/Source/Editor/MovieSceneTools/Public/MovieSceneToolsModule.h
@@ -85,8 +85,10 @@ private:
 	FDelegateHandle BytePropertyTrackCreateEditorHandle;
 	FDelegateHandle ColorPropertyTrackCreateEditorHandle;
 	FDelegateHandle FloatPropertyTrackCreateEditorHandle;
+	FDelegateHandle DoublePropertyTrackCreateEditorHandle;
 	FDelegateHandle IntegerPropertyTrackCreateEditorHandle;
-	FDelegateHandle VectorPropertyTrackCreateEditorHandle;
+	FDelegateHandle FloatVectorPropertyTrackCreateEditorHandle;
+	FDelegateHandle DoubleVectorPropertyTrackCreateEditorHandle;
 	FDelegateHandle TransformPropertyTrackCreateEditorHandle;
 	FDelegateHandle EulerTransformPropertyTrackCreateEditorHandle;
 	FDelegateHandle VisibilityPropertyTrackCreateEditorHandle;
diff --git a/Engine/Source/Editor/MovieSceneTools/Public/MovieSceneToolsProjectSettings.h b/Engine/Source/Editor/MovieSceneTools/Public/MovieSceneToolsProjectSettings.h
--- a/Engine/Source/Editor/MovieSceneTools/Public/MovieSceneToolsProjectSettings.h
+++ b/Engine/Source/Editor/MovieSceneTools/Public/MovieSceneToolsProjectSettings.h
@@ -7,6 +7,13 @@
 #include "UObject/Object.h"
 #include "MovieSceneToolsProjectSettings.generated.h"
 
+UENUM()
+enum class EMovieSceneToolsPropertyTrackType
+{
+	FloatTrack,
+	DoubleTrack
+};
+
 USTRUCT()
 struct FMovieSceneToolsPropertyTrackSettings
 {
@@ -33,6 +40,10 @@ struct FMovieSceneToolsFbxSettings
 	/** The property track setting to map to */
 	UPROPERTY(config, EditAnywhere, Category=TrackSettings)
 	FMovieSceneToolsPropertyTrackSettings PropertyPath;
+
+	/** Type of property track to map to */
+	UPROPERTY(config, EditAnywhere, Category=TrackSettings)
+	EMovieSceneToolsPropertyTrackType PropertyType = EMovieSceneToolsPropertyTrackType::FloatTrack;
 };
 
 // Settings for the level sequences
diff --git a/Engine/Source/Editor/SequenceRecorder/Private/Sections/MovieScene3DTransformSectionRecorder.cpp b/Engine/Source/Editor/SequenceRecorder/Private/Sections/MovieScene3DTransformSectionRecorder.cpp
--- a/Engine/Source/Editor/SequenceRecorder/Private/Sections/MovieScene3DTransformSectionRecorder.cpp
+++ b/Engine/Source/Editor/SequenceRecorder/Private/Sections/MovieScene3DTransformSectionRecorder.cpp
@@ -78,16 +78,16 @@ void FMovieScene3DTransformSectionRecorder::CreateSection(UObject* InObjectToRec
 		FVector EulerRotation = DefaultTransform.GetRotation().Rotator().Euler();
 		FVector Scale         = DefaultTransform.GetScale3D();
 
-		TArrayView<FMovieSceneFloatChannel*> FloatChannels = MovieSceneSection->GetChannelProxy().GetChannels<FMovieSceneFloatChannel>();
-		FloatChannels[0]->SetDefault(Translation.X);
-		FloatChannels[1]->SetDefault(Translation.Y);
-		FloatChannels[2]->SetDefault(Translation.Z);
-		FloatChannels[3]->SetDefault(EulerRotation.X);
-		FloatChannels[4]->SetDefault(EulerRotation.Y);
-		FloatChannels[5]->SetDefault(EulerRotation.Z);
-		FloatChannels[6]->SetDefault(Scale.X);
-		FloatChannels[7]->SetDefault(Scale.Y);
-		FloatChannels[8]->SetDefault(Scale.Z);
+		TArrayView<FMovieSceneDoubleChannel*> DoubleChannels = MovieSceneSection->GetChannelProxy().GetChannels<FMovieSceneDoubleChannel>();
+		DoubleChannels[0]->SetDefault(Translation.X);
+		DoubleChannels[1]->SetDefault(Translation.Y);
+		DoubleChannels[2]->SetDefault(Translation.Z);
+		DoubleChannels[3]->SetDefault(EulerRotation.X);
+		DoubleChannels[4]->SetDefault(EulerRotation.Y);
+		DoubleChannels[5]->SetDefault(EulerRotation.Z);
+		DoubleChannels[6]->SetDefault(Scale.X);
+		DoubleChannels[7]->SetDefault(Scale.Y);
+		DoubleChannels[8]->SetDefault(Scale.Z);
 
 		FFrameRate   TickResolution = MovieSceneSection->GetTypedOuter<UMovieScene>()->GetTickResolution();
 		FFrameNumber CurrentFrame    = (Time * TickResolution).FloorToFrame();
@@ -229,51 +229,51 @@ void FMovieScene3DTransformSectionRecorder::FinalizeSection(float CurrentTime)
 	const ERichCurveInterpMode Interpolation = AnimRecorder.IsValid() ? RCIM_Linear : RCIM_Cubic;
 
 	// add buffered transforms
-	TArrayView<FMovieSceneFloatChannel*> FloatChannels = MovieSceneSection->GetChannelProxy().GetChannels<FMovieSceneFloatChannel>();
+	TArrayView<FMovieSceneDoubleChannel*> DoubleChannels = MovieSceneSection->GetChannelProxy().GetChannels<FMovieSceneDoubleChannel>();
 
-	auto Transformation = [Interpolation](float In)
+	auto Transformation = [Interpolation](double In)
 	{
-		FMovieSceneFloatValue NewValue(In);
+		FMovieSceneDoubleValue NewValue(In);
 		NewValue.InterpMode = Interpolation;
 		return NewValue;
 	};
-	TArray<FMovieSceneFloatValue> FloatValues;
+	TArray<FMovieSceneDoubleValue> DoubleValues;
 
-	FloatValues.Reset(BufferedTransforms.LocationX.Num());
-	Algo::Transform(BufferedTransforms.LocationX, FloatValues, Transformation);
-	FloatChannels[0]->Set(BufferedTransforms.Times, MoveTemp(FloatValues));
+	DoubleValues.Reset(BufferedTransforms.LocationX.Num());
+	Algo::Transform(BufferedTransforms.LocationX, DoubleValues, Transformation);
+	DoubleChannels[0]->Set(BufferedTransforms.Times, MoveTemp(DoubleValues));
 
-	FloatValues.Reset(BufferedTransforms.LocationY.Num());
-	Algo::Transform(BufferedTransforms.LocationY, FloatValues, Transformation);
-	FloatChannels[1]->Set(BufferedTransforms.Times, MoveTemp(FloatValues));
+	DoubleValues.Reset(BufferedTransforms.LocationY.Num());
+	Algo::Transform(BufferedTransforms.LocationY, DoubleValues, Transformation);
+	DoubleChannels[1]->Set(BufferedTransforms.Times, MoveTemp(DoubleValues));
 
-	FloatValues.Reset(BufferedTransforms.LocationZ.Num());
-	Algo::Transform(BufferedTransforms.LocationZ, FloatValues, Transformation);
-	FloatChannels[2]->Set(BufferedTransforms.Times, MoveTemp(FloatValues));
+	DoubleValues.Reset(BufferedTransforms.LocationZ.Num());
+	Algo::Transform(BufferedTransforms.LocationZ, DoubleValues, Transformation);
+	DoubleChannels[2]->Set(BufferedTransforms.Times, MoveTemp(DoubleValues));
 
-	FloatValues.Reset(BufferedTransforms.RotationX.Num());
-	Algo::Transform(BufferedTransforms.RotationX, FloatValues, Transformation);
-	FloatChannels[3]->Set(BufferedTransforms.Times, MoveTemp(FloatValues));
+	DoubleValues.Reset(BufferedTransforms.RotationX.Num());
+	Algo::Transform(BufferedTransforms.RotationX, DoubleValues, Transformation);
+	DoubleChannels[3]->Set(BufferedTransforms.Times, MoveTemp(DoubleValues));
 
-	FloatValues.Reset(BufferedTransforms.RotationY.Num());
-	Algo::Transform(BufferedTransforms.RotationY, FloatValues, Transformation);
-	FloatChannels[4]->Set(BufferedTransforms.Times, MoveTemp(FloatValues));
+	DoubleValues.Reset(BufferedTransforms.RotationY.Num());
+	Algo::Transform(BufferedTransforms.RotationY, DoubleValues, Transformation);
+	DoubleChannels[4]->Set(BufferedTransforms.Times, MoveTemp(DoubleValues));
 
-	FloatValues.Reset(BufferedTransforms.RotationZ.Num());
-	Algo::Transform(BufferedTransforms.RotationZ, FloatValues, Transformation);
-	FloatChannels[5]->Set(BufferedTransforms.Times, MoveTemp(FloatValues));
+	DoubleValues.Reset(BufferedTransforms.RotationZ.Num());
+	Algo::Transform(BufferedTransforms.RotationZ, DoubleValues, Transformation);
+	DoubleChannels[5]->Set(BufferedTransforms.Times, MoveTemp(DoubleValues));
 
-	FloatValues.Reset(BufferedTransforms.ScaleX.Num());
-	Algo::Transform(BufferedTransforms.ScaleX, FloatValues, Transformation);
-	FloatChannels[6]->Set(BufferedTransforms.Times, MoveTemp(FloatValues));
+	DoubleValues.Reset(BufferedTransforms.ScaleX.Num());
+	Algo::Transform(BufferedTransforms.ScaleX, DoubleValues, Transformation);
+	DoubleChannels[6]->Set(BufferedTransforms.Times, MoveTemp(DoubleValues));
 
-	FloatValues.Reset(BufferedTransforms.ScaleY.Num());
-	Algo::Transform(BufferedTransforms.ScaleY, FloatValues, Transformation);
-	FloatChannels[7]->Set(BufferedTransforms.Times, MoveTemp(FloatValues));
+	DoubleValues.Reset(BufferedTransforms.ScaleY.Num());
+	Algo::Transform(BufferedTransforms.ScaleY, DoubleValues, Transformation);
+	DoubleChannels[7]->Set(BufferedTransforms.Times, MoveTemp(DoubleValues));
 
-	FloatValues.Reset(BufferedTransforms.ScaleZ.Num());
-	Algo::Transform(BufferedTransforms.ScaleZ, FloatValues, Transformation);
-	FloatChannels[8]->Set(BufferedTransforms.Times, MoveTemp(FloatValues));
+	DoubleValues.Reset(BufferedTransforms.ScaleZ.Num());
+	Algo::Transform(BufferedTransforms.ScaleZ, DoubleValues, Transformation);
+	DoubleChannels[8]->Set(BufferedTransforms.Times, MoveTemp(DoubleValues));
 
 	FTransform FirstTransform = FTransform::Identity;
 	if (BufferedTransforms.Times.Num())
@@ -293,14 +293,14 @@ void FMovieScene3DTransformSectionRecorder::FinalizeSection(float CurrentTime)
 	{
 		FKeyDataOptimizationParams Params;
 
-		for (FMovieSceneFloatChannel* Channel : FloatChannels)
+		for (FMovieSceneDoubleChannel* Channel : DoubleChannels)
 		{
 			Channel->Optimize(Params);
 		}
 	}
 	else
 	{
-		for (FMovieSceneFloatChannel* Channel : FloatChannels)
+		for (FMovieSceneDoubleChannel* Channel : DoubleChannels)
 		{
 			Channel->AutoSetTangents();
 		}
@@ -311,14 +311,14 @@ void FMovieScene3DTransformSectionRecorder::FinalizeSection(float CurrentTime)
 	if(!bWasAttached)
 	{
 		bool bCanRemoveTrack = true;
-		for (FMovieSceneFloatChannel* Channel : MovieSceneSection->GetChannelProxy().GetChannels<FMovieSceneFloatChannel>())
+		for (FMovieSceneDoubleChannel* Channel : MovieSceneSection->GetChannelProxy().GetChannels<FMovieSceneDoubleChannel>())
 		{
 			if (Channel)
 			{
 				int32 NumKeys = Channel->GetTimes().Num();
 				if (NumKeys == 1)
 				{
-					*Channel = FMovieSceneFloatChannel();
+					*Channel = FMovieSceneDoubleChannel();
 				}
 				else if (NumKeys > 1)
 				{
diff --git a/Engine/Source/Editor/SequenceRecorder/Private/Sections/MovieScenePropertyRecorder.cpp b/Engine/Source/Editor/SequenceRecorder/Private/Sections/MovieScenePropertyRecorder.cpp
--- a/Engine/Source/Editor/SequenceRecorder/Private/Sections/MovieScenePropertyRecorder.cpp
+++ b/Engine/Source/Editor/SequenceRecorder/Private/Sections/MovieScenePropertyRecorder.cpp
@@ -8,6 +8,8 @@
 #include "Tracks/MovieSceneByteTrack.h"
 #include "Sections/MovieSceneEnumSection.h"
 #include "Tracks/MovieSceneEnumTrack.h"
+#include "Sections/MovieSceneDoubleSection.h"
+#include "Tracks/MovieSceneDoubleTrack.h"
 #include "Sections/MovieSceneFloatSection.h"
 #include "Tracks/MovieSceneFloatTrack.h"
 #include "Sections/MovieSceneColorSection.h"
@@ -347,13 +349,13 @@ void FMovieScenePropertyRecorder<FColor>::ReduceKeys(UMovieSceneSection* InSecti
 }
 
 template <>
-bool FMovieScenePropertyRecorder<FVector>::ShouldAddNewKey(const FVector& InNewValue) const
+bool FMovieScenePropertyRecorder<FVector3f>::ShouldAddNewKey(const FVector3f& InNewValue) const
 {
 	return true;
 }
 
 template <>
-UMovieSceneSection* FMovieScenePropertyRecorder<FVector>::AddSection(UObject* InObjectToRecord, UMovieScene* InMovieScene, const FGuid& InGuid, float InTime)
+UMovieSceneSection* FMovieScenePropertyRecorder<FVector3f>::AddSection(UObject* InObjectToRecord, UMovieScene* InMovieScene, const FGuid& InGuid, float InTime)
 {
 	if (!InObjectToRecord)
 	{
@@ -361,10 +363,10 @@ UMovieSceneSection* FMovieScenePropertyRecorder<FVector>::AddSection(UObject* In
 	}
 
 	FName TrackName = *Binding.GetPropertyPath();
-	UMovieSceneVectorTrack* Track = InMovieScene->FindTrack<UMovieSceneVectorTrack>(InGuid, TrackName);
+	UMovieSceneFloatVectorTrack* Track = InMovieScene->FindTrack<UMovieSceneFloatVectorTrack>(InGuid, TrackName);
 	if (!Track)
 	{
-		Track = InMovieScene->AddTrack<UMovieSceneVectorTrack>(InGuid);
+		Track = InMovieScene->AddTrack<UMovieSceneFloatVectorTrack>(InGuid);
 	}
 	else
 	{
@@ -376,7 +378,7 @@ UMovieSceneSection* FMovieScenePropertyRecorder<FVector>::AddSection(UObject* In
 		Track->SetNumChannelsUsed(3);
 		Track->SetPropertyNameAndPath(Binding.GetPropertyName(), Binding.GetPropertyPath());
 
-		UMovieSceneVectorSection* Section = Cast<UMovieSceneVectorSection>(Track->CreateNewSection());
+		UMovieSceneFloatVectorSection* Section = Cast<UMovieSceneFloatVectorSection>(Track->CreateNewSection());
 
 		FFrameRate   TickResolution  = Section->GetTypedOuter<UMovieScene>()->GetTickResolution();
 		FFrameNumber CurrentFrame    = (InTime * TickResolution).FloorToFrame();
@@ -405,7 +407,7 @@ UMovieSceneSection* FMovieScenePropertyRecorder<FVector>::AddSection(UObject* In
 }
 
 template <>
-void FMovieScenePropertyRecorder<FVector>::AddKeyToSection(UMovieSceneSection* InSection, const FPropertyKey<FVector>& InKey)
+void FMovieScenePropertyRecorder<FVector3f>::AddKeyToSection(UMovieSceneSection* InSection, const FPropertyKey<FVector3f>& InKey)
 {
 	TArrayView<FMovieSceneFloatChannel*> FloatChannels = InSection->GetChannelProxy().GetChannels<FMovieSceneFloatChannel>();
 	FloatChannels[0]->AddCubicKey(InKey.Time, InKey.Value.X);
@@ -414,7 +416,7 @@ void FMovieScenePropertyRecorder<FVector>::AddKeyToSection(UMovieSceneSection* I
 }
 
 template <>
-void FMovieScenePropertyRecorder<FVector>::ReduceKeys(UMovieSceneSection* InSection)
+void FMovieScenePropertyRecorder<FVector3f>::ReduceKeys(UMovieSceneSection* InSection)
 {
 	TArrayView<FMovieSceneFloatChannel*> FloatChannels = InSection->GetChannelProxy().GetChannels<FMovieSceneFloatChannel>();
 
@@ -424,8 +426,87 @@ void FMovieScenePropertyRecorder<FVector>::ReduceKeys(UMovieSceneSection* InSect
 	UE::MovieScene::Optimize(FloatChannels[2], Params);
 }
 
+template <>
+bool FMovieScenePropertyRecorder<FVector3d>::ShouldAddNewKey(const FVector3d& InNewValue) const
+{
+	return true;
+}
+
+template <>
+UMovieSceneSection* FMovieScenePropertyRecorder<FVector3d>::AddSection(UObject* InObjectToRecord, UMovieScene* InMovieScene, const FGuid& InGuid, float InTime)
+{
+	if (!InObjectToRecord)
+	{
+		return nullptr;
+	}
+
+	FName TrackName = *Binding.GetPropertyPath();
+	UMovieSceneDoubleVectorTrack* Track = InMovieScene->FindTrack<UMovieSceneDoubleVectorTrack>(InGuid, TrackName);
+	if (!Track)
+	{
+		Track = InMovieScene->AddTrack<UMovieSceneDoubleVectorTrack>(InGuid);
+	}
+	else
+	{
+		Track->RemoveAllAnimationData();
+	}
+
+	if (Track)
+	{
+		Track->SetNumChannelsUsed(3);
+		Track->SetPropertyNameAndPath(Binding.GetPropertyName(), Binding.GetPropertyPath());
+
+		UMovieSceneDoubleVectorSection* Section = Cast<UMovieSceneDoubleVectorSection>(Track->CreateNewSection());
+
+		FFrameRate   TickResolution  = Section->GetTypedOuter<UMovieScene>()->GetTickResolution();
+		FFrameNumber CurrentFrame    = (InTime * TickResolution).FloorToFrame();
+
+		Section->SetRange(TRange<FFrameNumber>::Inclusive(CurrentFrame, CurrentFrame));
+
+		Section->TimecodeSource = SequenceRecorderUtils::GetTimecodeSource();
+
+		TArrayView<FMovieSceneDoubleChannel*> DoubleChannels = Section->GetChannelProxy().GetChannels<FMovieSceneDoubleChannel>();
+
+		DoubleChannels[0]->SetDefault(PreviousValue.X);
+		DoubleChannels[0]->AddCubicKey(CurrentFrame, PreviousValue.X, RCTM_Break);
+
+		DoubleChannels[1]->SetDefault(PreviousValue.Y);
+		DoubleChannels[1]->AddCubicKey(CurrentFrame, PreviousValue.Y, RCTM_Break);
+
+		DoubleChannels[2]->SetDefault(PreviousValue.Z);
+		DoubleChannels[2]->AddCubicKey(CurrentFrame, PreviousValue.Z, RCTM_Break);
+
+		Track->AddSection(*Section);
+
+		return Section;
+	}
+
+	return nullptr;
+}
+
+template <>
+void FMovieScenePropertyRecorder<FVector3d>::AddKeyToSection(UMovieSceneSection* InSection, const FPropertyKey<FVector3d>& InKey)
+{
+	TArrayView<FMovieSceneDoubleChannel*> DoubleChannels = InSection->GetChannelProxy().GetChannels<FMovieSceneDoubleChannel>();
+	DoubleChannels[0]->AddCubicKey(InKey.Time, InKey.Value.X);
+	DoubleChannels[1]->AddCubicKey(InKey.Time, InKey.Value.Y);
+	DoubleChannels[2]->AddCubicKey(InKey.Time, InKey.Value.Z);
+}
+
+template <>
+void FMovieScenePropertyRecorder<FVector3d>::ReduceKeys(UMovieSceneSection* InSection)
+{
+	TArrayView<FMovieSceneDoubleChannel*> DoubleChannels = InSection->GetChannelProxy().GetChannels<FMovieSceneDoubleChannel>();
+
+	FKeyDataOptimizationParams Params;
+	UE::MovieScene::Optimize(DoubleChannels[0], Params);
+	UE::MovieScene::Optimize(DoubleChannels[1], Params);
+	UE::MovieScene::Optimize(DoubleChannels[2], Params);
+}
+
 template class FMovieScenePropertyRecorder<bool>;
 template class FMovieScenePropertyRecorder<uint8>;
 template class FMovieScenePropertyRecorder<float>;
 template class FMovieScenePropertyRecorder<FColor>;
-template class FMovieScenePropertyRecorder<FVector>;
+template class FMovieScenePropertyRecorder<FVector3f>;
+template class FMovieScenePropertyRecorder<FVector3d>;
diff --git a/Engine/Source/Editor/Sequencer/Private/Sequencer.cpp b/Engine/Source/Editor/Sequencer/Private/Sequencer.cpp
--- a/Engine/Source/Editor/Sequencer/Private/Sequencer.cpp
+++ b/Engine/Source/Editor/Sequencer/Private/Sequencer.cpp
@@ -1556,12 +1556,19 @@ void FSequencer::SetInterpTangentMode(ERichCurveInterpMode InterpMode, ERichCurv
 	TSet<UMovieSceneSection*> ModifiedSections;
 
 	const FName FloatChannelTypeName = FMovieSceneFloatChannel::StaticStruct()->GetFName();
+	const FName DoubleChannelTypeName = FMovieSceneDoubleChannel::StaticStruct()->GetFName();
 
 	// @todo: sequencer-timecode: move this float-specific logic elsewhere to make it extensible for any channel type
 	for (const FSelectedChannelInfo& ChannelInfo : KeysByChannel.SelectedChannels)
 	{
 		FMovieSceneChannel* ChannelPtr = ChannelInfo.Channel.Get();
-		if (ChannelInfo.Channel.GetChannelTypeName() == FloatChannelTypeName && ChannelPtr)
+		if (!ChannelPtr)
+		{
+			continue;
+		}
+
+		const FName ChannelTypeName = ChannelInfo.Channel.GetChannelTypeName();
+		if (ChannelTypeName == FloatChannelTypeName || ChannelTypeName == DoubleChannelTypeName)
 		{
 			if (!ModifiedSections.Contains(ChannelInfo.OwningSection))
 			{
@@ -1569,23 +1576,46 @@ void FSequencer::SetInterpTangentMode(ERichCurveInterpMode InterpMode, ERichCurv
 				ModifiedSections.Add(ChannelInfo.OwningSection);
 			}
 
-			FMovieSceneFloatChannel* Channel = static_cast<FMovieSceneFloatChannel*>(ChannelPtr);
-			TMovieSceneChannelData<FMovieSceneFloatValue> ChannelData = Channel->GetData();
+			if (ChannelTypeName == FloatChannelTypeName)
+			{
+				FMovieSceneFloatChannel* Channel = static_cast<FMovieSceneFloatChannel*>(ChannelPtr);
+				TMovieSceneChannelData<FMovieSceneFloatValue> ChannelData = Channel->GetData();
 
-			TArrayView<FMovieSceneFloatValue> Values = ChannelData.GetValues();
+				TArrayView<FMovieSceneFloatValue> Values = ChannelData.GetValues();
 
-			for (FKeyHandle Handle : ChannelInfo.KeyHandles)
-			{
-				const int32 KeyIndex = ChannelData.GetIndex(Handle);
-				if (KeyIndex != INDEX_NONE)
+				for (FKeyHandle Handle : ChannelInfo.KeyHandles)
 				{
-					Values[KeyIndex].InterpMode = InterpMode;
-					Values[KeyIndex].TangentMode = TangentMode;
-					bAnythingChanged = true;
+					const int32 KeyIndex = ChannelData.GetIndex(Handle);
+					if (KeyIndex != INDEX_NONE)
+					{
+						Values[KeyIndex].InterpMode = InterpMode;
+						Values[KeyIndex].TangentMode = TangentMode;
+						bAnythingChanged = true;
+					}
 				}
+
+				Channel->AutoSetTangents();
 			}
+			else if (ChannelTypeName == DoubleChannelTypeName)
+			{
+				FMovieSceneDoubleChannel* Channel = static_cast<FMovieSceneDoubleChannel*>(ChannelPtr);
+				TMovieSceneChannelData<FMovieSceneDoubleValue> ChannelData = Channel->GetData();
+
+				TArrayView<FMovieSceneDoubleValue> Values = ChannelData.GetValues();
 
-			Channel->AutoSetTangents();
+				for (FKeyHandle Handle : ChannelInfo.KeyHandles)
+				{
+					const int32 KeyIndex = ChannelData.GetIndex(Handle);
+					if (KeyIndex != INDEX_NONE)
+					{
+						Values[KeyIndex].InterpMode = InterpMode;
+						Values[KeyIndex].TangentMode = TangentMode;
+						bAnythingChanged = true;
+					}
+				}
+
+				Channel->AutoSetTangents();
+			}
 		}
 	}
 
@@ -1612,6 +1642,7 @@ void FSequencer::ToggleInterpTangentWeightMode()
 	TSet<UMovieSceneSection*> ModifiedSections;
 
 	const FName FloatChannelTypeName = FMovieSceneFloatChannel::StaticStruct()->GetFName();
+	const FName DoubleChannelTypeName = FMovieSceneDoubleChannel::StaticStruct()->GetFName();
 
 	// Remove all tangent weights unless we find a compatible key that does not have weights yet
 	ERichCurveTangentWeightMode WeightModeToApply = RCTWM_WeightedNone;
@@ -1620,13 +1651,35 @@ void FSequencer::ToggleInterpTangentWeightMode()
 	for (const FSelectedChannelInfo& ChannelInfo : KeysByChannel.SelectedChannels)
 	{
 		FMovieSceneChannel* ChannelPtr = ChannelInfo.Channel.Get();
-		if (ChannelInfo.Channel.GetChannelTypeName() == FloatChannelTypeName && ChannelPtr)
+		if (!ChannelPtr)
+		{
+			continue;
+		}
+
+		if (ChannelInfo.Channel.GetChannelTypeName() == FloatChannelTypeName)
 		{
 			FMovieSceneFloatChannel* Channel = static_cast<FMovieSceneFloatChannel*>(ChannelPtr);
 			TMovieSceneChannelData<FMovieSceneFloatValue> ChannelData = Channel->GetData();
 
 			TArrayView<FMovieSceneFloatValue> Values = ChannelData.GetValues();
 
+			for (FKeyHandle Handle : ChannelInfo.KeyHandles)
+			{
+				const int32 KeyIndex = ChannelData.GetIndex(Handle);
+				if (KeyIndex != INDEX_NONE && Values[KeyIndex].InterpMode == RCIM_Cubic && Values[KeyIndex].Tangent.TangentWeightMode == RCTWM_WeightedNone)
+				{
+					WeightModeToApply = RCTWM_WeightedBoth;
+					goto assign_weights;
+				}
+			}
+		}
+		else if (ChannelInfo.Channel.GetChannelTypeName() == DoubleChannelTypeName)
+		{
+			FMovieSceneDoubleChannel* Channel = static_cast<FMovieSceneDoubleChannel*>(ChannelPtr);
+			TMovieSceneChannelData<FMovieSceneDoubleValue> ChannelData = Channel->GetData();
+
+			TArrayView<FMovieSceneDoubleValue> Values = ChannelData.GetValues();
+
 			for (FKeyHandle Handle : ChannelInfo.KeyHandles)
 			{
 				const int32 KeyIndex = ChannelData.GetIndex(Handle);
@@ -1645,7 +1698,13 @@ assign_weights:
 	for (const FSelectedChannelInfo& ChannelInfo : KeysByChannel.SelectedChannels)
 	{
 		FMovieSceneChannel* ChannelPtr = ChannelInfo.Channel.Get();
-		if (ChannelInfo.Channel.GetChannelTypeName() == FloatChannelTypeName && ChannelPtr)
+		if (!ChannelPtr)
+		{
+			continue;
+		}
+
+		const FName ChannelTypeName = ChannelInfo.Channel.GetChannelTypeName();
+		if (ChannelTypeName == FloatChannelTypeName || ChannelTypeName == DoubleChannelTypeName)
 		{
 			if (!ModifiedSections.Contains(ChannelInfo.OwningSection))
 			{
@@ -1653,22 +1712,44 @@ assign_weights:
 				ModifiedSections.Add(ChannelInfo.OwningSection);
 			}
 
-			FMovieSceneFloatChannel* Channel = static_cast<FMovieSceneFloatChannel*>(ChannelPtr);
-			TMovieSceneChannelData<FMovieSceneFloatValue> ChannelData = Channel->GetData();
+			if (ChannelTypeName == FloatChannelTypeName)
+			{
+				FMovieSceneFloatChannel* Channel = static_cast<FMovieSceneFloatChannel*>(ChannelPtr);
+				TMovieSceneChannelData<FMovieSceneFloatValue> ChannelData = Channel->GetData();
 
-			TArrayView<FMovieSceneFloatValue> Values = ChannelData.GetValues();
+				TArrayView<FMovieSceneFloatValue> Values = ChannelData.GetValues();
 
-			for (FKeyHandle Handle : ChannelInfo.KeyHandles)
-			{
-				const int32 KeyIndex = ChannelData.GetIndex(Handle);
-				if (KeyIndex != INDEX_NONE && Values[KeyIndex].InterpMode == RCIM_Cubic)
+				for (FKeyHandle Handle : ChannelInfo.KeyHandles)
 				{
-					Values[KeyIndex].Tangent.TangentWeightMode = WeightModeToApply;
-					bAnythingChanged = true;
+					const int32 KeyIndex = ChannelData.GetIndex(Handle);
+					if (KeyIndex != INDEX_NONE && Values[KeyIndex].InterpMode == RCIM_Cubic)
+					{
+						Values[KeyIndex].Tangent.TangentWeightMode = WeightModeToApply;
+						bAnythingChanged = true;
+					}
 				}
+
+				Channel->AutoSetTangents();
 			}
+			else if (ChannelTypeName == DoubleChannelTypeName)
+			{
+				FMovieSceneDoubleChannel* Channel = static_cast<FMovieSceneDoubleChannel*>(ChannelPtr);
+				TMovieSceneChannelData<FMovieSceneDoubleValue> ChannelData = Channel->GetData();
 
-			Channel->AutoSetTangents();
+				TArrayView<FMovieSceneDoubleValue> Values = ChannelData.GetValues();
+
+				for (FKeyHandle Handle : ChannelInfo.KeyHandles)
+				{
+					const int32 KeyIndex = ChannelData.GetIndex(Handle);
+					if (KeyIndex != INDEX_NONE && Values[KeyIndex].InterpMode == RCIM_Cubic)
+					{
+						Values[KeyIndex].Tangent.TangentWeightMode = WeightModeToApply;
+						bAnythingChanged = true;
+					}
+				}
+
+				Channel->AutoSetTangents();
+			}
 		}
 	}
 
@@ -2406,16 +2487,16 @@ void FSequencer::BakeTransform()
 
 			TransformSection->SetRange(TRange<FFrameNumber>::All());
 
-			TArrayView<FMovieSceneFloatChannel*> FloatChannels = TransformSection->GetChannelProxy().GetChannels<FMovieSceneFloatChannel>();
-			FloatChannels[0]->SetDefault(DefaultLocation.X);
-			FloatChannels[1]->SetDefault(DefaultLocation.Y);
-			FloatChannels[2]->SetDefault(DefaultLocation.Z);
-			FloatChannels[3]->SetDefault(DefaultRotation.X);
-			FloatChannels[4]->SetDefault(DefaultRotation.Y);
-			FloatChannels[5]->SetDefault(DefaultRotation.Z);
-			FloatChannels[6]->SetDefault(DefaultScale.X);
-			FloatChannels[7]->SetDefault(DefaultScale.Y);
-			FloatChannels[8]->SetDefault(DefaultScale.Z);
+			TArrayView<FMovieSceneDoubleChannel*> DoubleChannels = TransformSection->GetChannelProxy().GetChannels<FMovieSceneDoubleChannel>();
+			DoubleChannels[0]->SetDefault(DefaultLocation.X);
+			DoubleChannels[1]->SetDefault(DefaultLocation.Y);
+			DoubleChannels[2]->SetDefault(DefaultLocation.Z);
+			DoubleChannels[3]->SetDefault(DefaultRotation.X);
+			DoubleChannels[4]->SetDefault(DefaultRotation.Y);
+			DoubleChannels[5]->SetDefault(DefaultRotation.Z);
+			DoubleChannels[6]->SetDefault(DefaultScale.X);
+			DoubleChannels[7]->SetDefault(DefaultScale.Y);
+			DoubleChannels[8]->SetDefault(DefaultScale.Z);
 
 			TArray<FVector> LocalTranslations, LocalRotations, LocalScales;
 			LocalTranslations.SetNum(BakeData.Value.KeyTimes.Num());
@@ -2441,15 +2522,15 @@ void FSequencer::BakeTransform()
 			for (int32 Counter = 0; Counter < BakeData.Value.KeyTimes.Num(); ++Counter)
 			{
 				FFrameNumber KeyTime = BakeData.Value.KeyTimes[Counter];
-				FloatChannels[0]->AddLinearKey(KeyTime, LocalTranslations[Counter].X);
-				FloatChannels[1]->AddLinearKey(KeyTime, LocalTranslations[Counter].Y);
-				FloatChannels[2]->AddLinearKey(KeyTime, LocalTranslations[Counter].Z);
-				FloatChannels[3]->AddLinearKey(KeyTime, LocalRotations[Counter].X);
-				FloatChannels[4]->AddLinearKey(KeyTime, LocalRotations[Counter].Y);
-				FloatChannels[5]->AddLinearKey(KeyTime, LocalRotations[Counter].Z);
-				FloatChannels[6]->AddLinearKey(KeyTime, LocalScales[Counter].X);
-				FloatChannels[7]->AddLinearKey(KeyTime, LocalScales[Counter].Y);
-				FloatChannels[8]->AddLinearKey(KeyTime, LocalScales[Counter].Z);
+				DoubleChannels[0]->AddLinearKey(KeyTime, LocalTranslations[Counter].X);
+				DoubleChannels[1]->AddLinearKey(KeyTime, LocalTranslations[Counter].Y);
+				DoubleChannels[2]->AddLinearKey(KeyTime, LocalTranslations[Counter].Z);
+				DoubleChannels[3]->AddLinearKey(KeyTime, LocalRotations[Counter].X);
+				DoubleChannels[4]->AddLinearKey(KeyTime, LocalRotations[Counter].Y);
+				DoubleChannels[5]->AddLinearKey(KeyTime, LocalRotations[Counter].Z);
+				DoubleChannels[6]->AddLinearKey(KeyTime, LocalScales[Counter].X);
+				DoubleChannels[7]->AddLinearKey(KeyTime, LocalScales[Counter].Y);
+				DoubleChannels[8]->AddLinearKey(KeyTime, LocalScales[Counter].Z);
 			}
 		}
 	}
diff --git a/Engine/Source/Editor/Sequencer/Private/SequencerContextMenus.cpp b/Engine/Source/Editor/Sequencer/Private/SequencerContextMenus.cpp
--- a/Engine/Source/Editor/Sequencer/Private/SequencerContextMenus.cpp
+++ b/Engine/Source/Editor/Sequencer/Private/SequencerContextMenus.cpp
@@ -926,6 +926,21 @@ void FSectionContextMenu::SetInterpTangentMode(ERichCurveInterpMode InterpMode,
 
 					FloatChannel->AutoSetTangents();
 				}
+
+				for (FMovieSceneDoubleChannel* DoubleChannel : Section->GetChannelProxy().GetChannels<FMovieSceneDoubleChannel>())
+				{
+					TMovieSceneChannelData<FMovieSceneDoubleValue> ChannelData = DoubleChannel->GetData();
+					TArrayView<FMovieSceneDoubleValue> Values = ChannelData.GetValues();
+
+					for (int32 KeyIndex = 0; KeyIndex < DoubleChannel->GetNumKeys(); ++KeyIndex)
+					{
+						Values[KeyIndex].InterpMode = InterpMode;
+						Values[KeyIndex].TangentMode = TangentMode;
+						bAnythingChanged = true;
+					}
+
+					DoubleChannel->AutoSetTangents();
+				}
 			}
 		}
 	}
@@ -960,7 +975,8 @@ bool FSectionContextMenu::CanSetInterpTangentMode() const
 			UMovieSceneSection* Section = KeyArea->GetOwningSection();
 			if (Section)
 			{
-				return Section->GetChannelProxy().GetChannels<FMovieSceneFloatChannel>().Num() != 0;
+				return (Section->GetChannelProxy().GetChannels<FMovieSceneFloatChannel>().Num() != 0 ||
+					    Section->GetChannelProxy().GetChannels<FMovieSceneDoubleChannel>().Num() != 0);
 			}
 		}
 	}
diff --git a/Engine/Source/Editor/Sequencer/Private/SequencerKeyActor.cpp b/Engine/Source/Editor/Sequencer/Private/SequencerKeyActor.cpp
--- a/Engine/Source/Editor/Sequencer/Private/SequencerKeyActor.cpp
+++ b/Engine/Source/Editor/Sequencer/Private/SequencerKeyActor.cpp
@@ -80,12 +80,12 @@ void ASequencerKeyActor::PropagateKeyChange()
 		FFrameRate   TickResolution  = TrackSection->GetTypedOuter<UMovieScene>()->GetTickResolution();
 		FFrameNumber FrameNumber     = (KeyTime * TickResolution).RoundToFrame();
 
-		TArrayView<FMovieSceneFloatChannel*> FloatChannels = TrackSection->GetChannelProxy().GetChannels<FMovieSceneFloatChannel>();
+		TArrayView<FMovieSceneDoubleChannel*> DoubleChannels = TrackSection->GetChannelProxy().GetChannels<FMovieSceneDoubleChannel>();
 
 		const FVector Translation = GetActorTransform().GetLocation();
-		FloatChannels[0]->GetData().UpdateOrAddKey(FrameNumber, FMovieSceneFloatValue(Translation.X));
-		FloatChannels[1]->GetData().UpdateOrAddKey(FrameNumber, FMovieSceneFloatValue(Translation.Y));
-		FloatChannels[2]->GetData().UpdateOrAddKey(FrameNumber, FMovieSceneFloatValue(Translation.Z));
+		DoubleChannels[0]->GetData().UpdateOrAddKey(FrameNumber, FMovieSceneDoubleValue(Translation.X));
+		DoubleChannels[1]->GetData().UpdateOrAddKey(FrameNumber, FMovieSceneDoubleValue(Translation.Y));
+		DoubleChannels[2]->GetData().UpdateOrAddKey(FrameNumber, FMovieSceneDoubleValue(Translation.Z));
 
 		// Draw a single transform track based on the data from this key
 		FEditorViewportClient* ViewportClient = StaticCast<FEditorViewportClient*>(GEditor->GetActiveViewport()->GetClient());
