Sequencer: Check for invalid hierarchy when setting value

diff --git a/Engine/Source/Editor/MovieSceneTools/Private/MovieSceneObjectBindingIDPicker.cpp b/Engine/Source/Editor/MovieSceneTools/Private/MovieSceneObjectBindingIDPicker.cpp
--- a/Engine/Source/Editor/MovieSceneTools/Private/MovieSceneObjectBindingIDPicker.cpp
+++ b/Engine/Source/Editor/MovieSceneTools/Private/MovieSceneObjectBindingIDPicker.cpp
@@ -285,7 +285,7 @@ void FMovieSceneObjectBindingIDPicker::SetCurrentValueFromFixed(UE::MovieScene::
 	const FMovieSceneSequenceHierarchy* Hierarchy = Sequencer.IsValid() ? Sequencer->GetEvaluationTemplate().GetHierarchy() : nullptr;
 
 	// If we don't know the local sequence ID, or we have no hierarchy, or we're resetting the binding; just set the ID directly
-	if (LocalSequenceID == MovieSceneSequenceID::Invalid || !InValue.Guid.IsValid())
+	if (LocalSequenceID == MovieSceneSequenceID::Invalid || !InValue.Guid.IsValid() || Hierarchy == nullptr)
 	{
 		SetCurrentValue(InValue);
 	}
