Sequencer: Don't create linked animation sequence if cancelled

diff --git a/Engine/Source/Editor/MovieSceneTools/Private/TrackEditors/SkeletalAnimationTrackEditor.cpp b/Engine/Source/Editor/MovieSceneTools/Private/TrackEditors/SkeletalAnimationTrackEditor.cpp
--- a/Engine/Source/Editor/MovieSceneTools/Private/TrackEditors/SkeletalAnimationTrackEditor.cpp
+++ b/Engine/Source/Editor/MovieSceneTools/Private/TrackEditors/SkeletalAnimationTrackEditor.cpp
@@ -1204,14 +1204,17 @@ bool FSkeletalAnimationTrackEditor::CreateAnimationSequence(const TArray<UObject
 
 			FSlateApplication::Get().AddModalWindow(Window, ParentWindow, false);
 
-			const TSharedPtr<ISequencer> ParentSequencer = GetSequencer();
-			UMovieScene* MovieScene = ParentSequencer->GetFocusedMovieSceneSequence()->GetMovieScene();
-			FMovieSceneSequenceIDRef Template = ParentSequencer->GetFocusedTemplateID();
-			FMovieSceneSequenceTransform RootToLocalTransform = ParentSequencer->GetFocusedMovieSceneSequenceTransform();
+			if (OptionWindow->ShouldExport())
+			{
+				const TSharedPtr<ISequencer> ParentSequencer = GetSequencer();
+				UMovieScene* MovieScene = ParentSequencer->GetFocusedMovieSceneSequence()->GetMovieScene();
+				FMovieSceneSequenceIDRef Template = ParentSequencer->GetFocusedTemplateID();
+				FMovieSceneSequenceTransform RootToLocalTransform = ParentSequencer->GetFocusedMovieSceneSequenceTransform();
 
-			bResult  = MovieSceneToolHelpers::ExportToAnimSequence(AnimSequence, AnimSeqExportOption,MovieScene, ParentSequencer.Get(), SkelMeshComp, Template, RootToLocalTransform);
+				bResult  = MovieSceneToolHelpers::ExportToAnimSequence(AnimSequence, AnimSeqExportOption,MovieScene, ParentSequencer.Get(), SkelMeshComp, Template, RootToLocalTransform);
+			}
 		}
-		
+
 		if (bResult && bCreateSoftLink)
 		{
 			TSharedPtr<ISequencer> SequencerPtr = GetSequencer();
