Sequencer: Temporarily disable remove unused events

diff --git a/Engine/Source/Editor/MovieSceneTools/Private/MovieSceneEventBlueprintExtension.cpp b/Engine/Source/Editor/MovieSceneTools/Private/MovieSceneEventBlueprintExtension.cpp
--- a/Engine/Source/Editor/MovieSceneTools/Private/MovieSceneEventBlueprintExtension.cpp
+++ b/Engine/Source/Editor/MovieSceneTools/Private/MovieSceneEventBlueprintExtension.cpp
@@ -72,6 +72,7 @@ void UMovieSceneEventBlueprintExtension::HandleGenerateFunctionGraphs(FKismetCom
 	// Don't remove while transacting because it can remove when an undo is in progress
 	if (!GIsTransacting)
 	{
-		FMovieSceneEventUtils::RemoveUnusedCustomEvents(EventSections, CompilerContext->Blueprint);
+		///// Temporarily disabled for UE-132130
+		/////FMovieSceneEventUtils::RemoveUnusedCustomEvents(EventSections, CompilerContext->Blueprint);
 	}
 }
\ No newline at end of file
