Export the FAudioSection symbol from MovieSceneTools module

diff --git a/Engine/Source/Editor/MovieSceneTools/Public/TrackEditors/AudioTrackEditor.h b/Engine/Source/Editor/MovieSceneTools/Public/TrackEditors/AudioTrackEditor.h
--- a/Engine/Source/Editor/MovieSceneTools/Public/TrackEditors/AudioTrackEditor.h
+++ b/Engine/Source/Editor/MovieSceneTools/Public/TrackEditors/AudioTrackEditor.h
@@ -99,7 +99,7 @@ private:
 /**
  * Class for audio sections, handles drawing of all waveform previews.
  */
-class FAudioSection
+class MOVIESCENETOOLS_API FAudioSection
 	: public ISequencerSection
 	, public TSharedFromThis<FAudioSection>
 {
