Sequencer: Send a track value changed if any property changed for skeletal animations. This fixes a bug where if you change any property on a skeletal mesh, the pose will revert to a tpose until you scrub Sequencer again. This is because the change causes the skeletal mesh to update/tick, but since Sequencer data isn't actually changed (unless you autokeky), Sequencer doesn't evaluate again.

diff --git a/Engine/Source/Editor/MovieSceneTools/Private/TrackEditors/SkeletalAnimationTrackEditor.cpp b/Engine/Source/Editor/MovieSceneTools/Private/TrackEditors/SkeletalAnimationTrackEditor.cpp
--- a/Engine/Source/Editor/MovieSceneTools/Private/TrackEditors/SkeletalAnimationTrackEditor.cpp
+++ b/Engine/Source/Editor/MovieSceneTools/Private/TrackEditors/SkeletalAnimationTrackEditor.cpp
@@ -1051,6 +1051,7 @@ FSkeletalAnimationTrackEditor::FSkeletalAnimationTrackEditor( TSharedRef<ISequen
 
 	SequencerSavedHandle = InSequencer->OnPostSave().AddRaw(this, &FSkeletalAnimationTrackEditor::OnSequencerSaved);
 
+	FCoreUObjectDelegates::OnObjectPropertyChanged.AddRaw(this, &FSkeletalAnimationTrackEditor::OnPostPropertyChanged);
 }
 
 void FSkeletalAnimationTrackEditor::OnInitialize()
@@ -1067,6 +1068,8 @@ void FSkeletalAnimationTrackEditor::OnInitialize()
 }
 void FSkeletalAnimationTrackEditor::OnRelease()
 {
+	FCoreUObjectDelegates::OnObjectPropertyChanged.RemoveAll(this);
+
 	--FSkeletalAnimationTrackEditor::NumberActive;
 
 	if (GetSequencer().IsValid() && SequencerSavedHandle.IsValid())
@@ -1206,6 +1209,21 @@ void FSkeletalAnimationTrackEditor::OnSequencerSaved(ISequencer& )
 	}
 }
 
+void FSkeletalAnimationTrackEditor::OnPostPropertyChanged(UObject* InObject, FPropertyChangedEvent& InPropertyChangedEvent)
+{
+	// If the object changed has any animation track, notify sequencer to update because animations tick on their own and sequencer needs to evaluate again
+	const bool bCreateIfMissing = false;
+	FFindOrCreateHandleResult HandleResult = FindOrCreateHandleToObject(InObject, bCreateIfMissing );
+	FGuid ObjectHandle = HandleResult.Handle;
+	if (ObjectHandle.IsValid())
+	{
+		FFindOrCreateTrackResult TrackResult = FindOrCreateTrackForObject(ObjectHandle, UMovieSceneSkeletalAnimationTrack::StaticClass(), NAME_None, bCreateIfMissing);
+		if (TrackResult.Track)
+		{
+			GetSequencer()->NotifyMovieSceneDataChanged(EMovieSceneDataChangeType::TrackValueChanged);
+		}
+	}
+}
 
 bool FSkeletalAnimationTrackEditor::CreateAnimationSequence(const TArray<UObject*> NewAssets, USkeletalMeshComponent* SkelMeshComp, FGuid Binding, bool bCreateSoftLink)
 {
diff --git a/Engine/Source/Editor/MovieSceneTools/Private/TrackEditors/SkeletalAnimationTrackEditor.h b/Engine/Source/Editor/MovieSceneTools/Private/TrackEditors/SkeletalAnimationTrackEditor.h
--- a/Engine/Source/Editor/MovieSceneTools/Private/TrackEditors/SkeletalAnimationTrackEditor.h
+++ b/Engine/Source/Editor/MovieSceneTools/Private/TrackEditors/SkeletalAnimationTrackEditor.h
@@ -122,7 +122,7 @@ private:
 	FDelegateHandle SequencerSavedHandle;
 	void OnSequencerSaved(ISequencer& InSequence);
 
-
+	void OnPostPropertyChanged(UObject* InObject, struct FPropertyChangedEvent& InPropertyChangedEvent);
 };
 
 
