Sequencer: Path track section to end at the current sequence's playback range, like the attach section

diff --git a/Engine/Source/Editor/MovieSceneTools/Private/TrackEditors/PathTrackEditor.cpp b/Engine/Source/Editor/MovieSceneTools/Private/TrackEditors/PathTrackEditor.cpp
--- a/Engine/Source/Editor/MovieSceneTools/Private/TrackEditors/PathTrackEditor.cpp
+++ b/Engine/Source/Editor/MovieSceneTools/Private/TrackEditors/PathTrackEditor.cpp
@@ -245,6 +245,8 @@ FKeyPropertyResult F3DPathTrackEditor::AddKeyInternal( FFrameNumber KeyTime, con
 		return KeyPropertyResult;
 	}
 
+	UMovieScene* MovieScene = GetSequencer()->GetFocusedMovieSceneSequence()->GetMovieScene();
+
 	for( int32 ObjectIndex = 0; ObjectIndex < Objects.Num(); ++ObjectIndex )
 	{
 		UObject* Object = Objects[ObjectIndex].Get();
@@ -260,8 +262,8 @@ FKeyPropertyResult F3DPathTrackEditor::AddKeyInternal( FFrameNumber KeyTime, con
 
 			if (ensure(Track))
 			{
-				// Clamp to next path section's start time or the end of the current sequencer view range
-				FFrameNumber PathEndTime = ( GetSequencer()->GetViewRange().GetUpperBoundValue() * Track->GetTypedOuter<UMovieScene>()->GetTickResolution() ).FrameNumber;
+				// Clamp to next path section's start time or the end of the current movie scene range
+				FFrameNumber PathEndTime = MovieScene->GetPlaybackRange().GetUpperBoundValue();
 	
 				for (UMovieSceneSection* Section : Track->GetAllSections())
 				{
