Sequencer: Don't re-apply a default if the value is the same this will dirty the movie scene section when it's not changing which can cause issues with moving controls with motion trails on. Had to add specialization for bool and add the check for ints.

diff --git a/Engine/Source/Editor/MovieSceneTools/Public/KeyframeTrackEditor.h b/Engine/Source/Editor/MovieSceneTools/Public/KeyframeTrackEditor.h
--- a/Engine/Source/Editor/MovieSceneTools/Public/KeyframeTrackEditor.h
+++ b/Engine/Source/Editor/MovieSceneTools/Public/KeyframeTrackEditor.h
@@ -22,6 +22,8 @@
 #include "Channels/MovieSceneDoubleChannel.h"
 #include "Channels/MovieSceneFloatChannel.h"
 #include "Channels/MovieSceneIntegerChannel.h"
+#include "Channels/MovieSceneBoolChannel.h"
+
 
 struct IImpl
 {
@@ -96,6 +98,65 @@ struct TAddKeyImpl : IImpl
 
 };
 
+//bool specialization
+template<>
+struct TAddKeyImpl<FMovieSceneBoolChannel, bool> : IImpl
+{
+	int32 ChannelIndex;
+	bool bAddKey;
+	bool ValueToSet;
+
+	TAddKeyImpl(int32 InChannelIndex, bool bInAddKey, const bool& InValue)
+		: ChannelIndex(InChannelIndex), bAddKey(bInAddKey), ValueToSet(InValue)
+	{}
+
+	virtual bool Apply(UMovieSceneSection* Section, FMovieSceneChannelProxy& Proxy, FFrameNumber InTime, EMovieSceneKeyInterpolation InterpolationMode, bool bKeyEvenIfUnchanged, bool bKeyEvenIfEmpty) const override
+	{
+		bool bKeyCreated = false;
+		using namespace UE::MovieScene;
+
+		FMovieSceneBoolChannel* Channel = Proxy.GetChannel<FMovieSceneBoolChannel>(ChannelIndex);
+		if (bAddKey && Channel)
+		{
+			bool bShouldKeyChannel = bKeyEvenIfUnchanged;
+			if (!bShouldKeyChannel)
+			{
+				bShouldKeyChannel = !ValueExistsAtTime(Channel, InTime, ValueToSet);
+			}
+
+			if (bShouldKeyChannel)
+			{
+				if (Channel->GetNumKeys() != 0 || bKeyEvenIfEmpty)
+				{
+					if (Section->TryModify())
+					{
+						AddKeyToChannel(Channel, InTime, ValueToSet, InterpolationMode);
+						bKeyCreated = true;
+					}
+				}
+			}
+		}
+
+		return bKeyCreated;
+	}
+
+	virtual void ApplyDefault(UMovieSceneSection* Section, FMovieSceneChannelProxy& Proxy) const override
+	{
+		using namespace UE::MovieScene;
+
+		FMovieSceneBoolChannel* Channel = Proxy.GetChannel<FMovieSceneBoolChannel>(ChannelIndex);
+		if (Channel && Channel->GetData().GetTimes().Num() == 0 && Channel->GetDefault() != ValueToSet)
+		{
+			if (Section->TryModify())
+			{
+				using namespace UE::MovieScene;
+				SetChannelDefault(Channel, ValueToSet);
+			}
+		}
+	}
+	virtual bool ModifyByCurrentAndWeight(FMovieSceneChannelProxy& Proxy, FFrameNumber InTime, void* VCurrentValue, float Weight) { return false; }
+
+};
 //Specializations for channels that SUPPORT Blending
 template<>
 struct TAddKeyImpl<FMovieSceneDoubleChannel, double> : IImpl
@@ -297,7 +358,7 @@ struct TAddKeyImpl<FMovieSceneIntegerChannel, int32> : IImpl
 		using namespace UE::MovieScene;
 
 		FMovieSceneIntegerChannel* Channel = Proxy.GetChannel<FMovieSceneIntegerChannel>(ChannelIndex);
-		if (Channel && Channel->GetData().GetTimes().Num() == 0)
+		if (Channel && Channel->GetData().GetTimes().Num() == 0 && Channel->GetDefault() != ValueToSet)
 		{
 			if (Section->TryModify())
 			{
