Partner cinematic tools changes : exposing MovieSceneCinematicShotSection/Track and CinematicShotTrack editor as public API.

diff --git a/Engine/Source/Editor/MovieSceneTools/Private/TrackEditors/CinematicShotTrackEditor.h b/Engine/Source/Editor/MovieSceneTools/Private/TrackEditors/CinematicShotTrackEditor.h
deleted file mode 100644
diff --git a/Engine/Source/Editor/MovieSceneTools/Public/TrackEditors/CinematicShotTrackEditor.h b/Engine/Source/Editor/MovieSceneTools/Public/TrackEditors/CinematicShotTrackEditor.h
new file mode 100644
--- /dev/null
+++ b/Engine/Source/Editor/MovieSceneTools/Public/TrackEditors/CinematicShotTrackEditor.h
@@ -0,0 +1,201 @@
+// Copyright Epic Games, Inc. All Rights Reserved.
+
+#pragma once
+
+#include "CoreMinimal.h"
+#include "Misc/Guid.h"
+#include "Templates/SubclassOf.h"
+#include "Widgets/SWidget.h"
+#include "ISequencer.h"
+#include "ISequencerSection.h"
+#include "MovieSceneTrack.h"
+#include "ISequencerTrackEditor.h"
+#include "MovieSceneTrackEditor.h"
+
+class AActor;
+struct FAssetData;
+class FMenuBuilder;
+class FTrackEditorThumbnailPool;
+class UMovieSceneCinematicShotSection;
+class UMovieSceneCinematicShotTrack;
+class UMovieSceneSubSection;
+
+/**
+ * Tools for cinematic shots.
+ */
+class MOVIESCENETOOLS_API FCinematicShotTrackEditor : public FMovieSceneTrackEditor
+{
+public:
+
+	/**
+	 * Constructor
+	 *
+	 * @param InSequencer The sequencer instance to be used by this tool.
+	 */
+	FCinematicShotTrackEditor(TSharedRef<ISequencer> InSequencer);
+
+	/** Virtual destructor. */
+	virtual ~FCinematicShotTrackEditor() { }
+
+	/**
+	 * Creates an instance of this class.  Called by a sequencer .
+	 *
+	 * @param OwningSequencer The sequencer instance to be used by this tool.
+	 * @return The new instance of this class.
+	 */
+	static TSharedRef<ISequencerTrackEditor> CreateTrackEditor(TSharedRef<ISequencer> OwningSequencer);
+
+	TWeakObjectPtr<AActor> GetCinematicShotCamera() const { return CinematicShotCamera; }
+
+public:
+
+	// ISequencerTrackEditor interface
+	virtual void OnInitialize() override;
+	virtual void OnRelease() override;
+	virtual void BuildAddTrackMenu(FMenuBuilder& MenuBuilder) override;
+	virtual TSharedPtr<SWidget> BuildOutlinerEditWidget(const FGuid& ObjectBinding, UMovieSceneTrack* Track, const FBuildEditWidgetParams& Params) override;
+	virtual TSharedRef<ISequencerSection> MakeSectionInterface(UMovieSceneSection& SectionObject, UMovieSceneTrack& Track, FGuid ObjectBinding) override;
+	virtual bool HandleAssetAdded(UObject* Asset, const FGuid& TargetObjectGuid) override;
+	virtual bool SupportsSequence(UMovieSceneSequence* InSequence) const override;
+	virtual bool SupportsType(TSubclassOf<UMovieSceneTrack> Type ) const override;
+	virtual void Tick(float DeltaTime) override;
+	virtual void BuildTrackContextMenu( FMenuBuilder& MenuBuilder, UMovieSceneTrack* Track ) override;
+	virtual const FSlateBrush* GetIconBrush() const override;
+	virtual bool OnAllowDrop(const FDragDropEvent& DragDropEvent, FSequencerDragDropParams& DragDropParams) override;
+	virtual FReply OnDrop(const FDragDropEvent& DragDropEvent, const FSequencerDragDropParams& DragDropParams) override;
+
+	/*
+	 * Insert shot. 
+	 */
+	void InsertShot();
+
+	/*
+	 * Insert filler.
+	 */
+	void InsertFiller();
+
+	/*
+	 * Duplicate shot. 
+	 *
+	 * @param Section The section to duplicate
+	 */
+	void DuplicateShot(UMovieSceneCinematicShotSection* Section);
+
+	/*
+	 * Render shots. 
+	 *
+	 * @param Sections The sections to render
+	 */
+	void RenderShots(const TArray<UMovieSceneCinematicShotSection*>& Sections);
+
+	/*
+	 * Rename shot. 
+	 *
+	 * @param Section The section to rename.
+	 */
+	void RenameShot(UMovieSceneCinematicShotSection* Section);
+
+	/*
+	 * New take. 
+	 *
+	 * @param Section The section to create a new take of.
+	 */
+	void NewTake(UMovieSceneCinematicShotSection* Section);
+
+	/*
+	* Switch take for the selected sections
+	*
+	* @param TakeObject The take object to switch to.
+	*/
+	void SwitchTake(UObject* TakeObject);
+
+private:
+
+	/*
+	 * Create shot
+	 *
+	 * @param NewShotName The new shot name.
+	 * @param NewShotStartTime The time to start the new shot at.
+	 * @param ShotToDuplicate The shot to duplicate.
+	 * @return The new shot.
+	 */
+	UMovieSceneSubSection* CreateShotInternal(FString& NewShotName, FFrameNumber NewShotStartTime, UMovieSceneCinematicShotSection* ShotToDuplicate = nullptr);
+	
+	/** Callback for determining whether the "Add Shot" menu entry can execute. */
+	bool HandleAddCinematicShotTrackMenuEntryCanExecute() const;
+
+	/** Callback for executing the "Add Shot Track" menu entry. */
+	void HandleAddCinematicShotTrackMenuEntryExecute();
+	
+protected:
+
+	/** Callback for generating the menu of the "Add Shot" combo button. */
+	virtual TSharedRef<SWidget> HandleAddCinematicShotComboButtonGetMenuContent();
+	
+	/** Find or create a cinematic shot track in the currently focused movie scene. */
+	virtual UMovieSceneCinematicShotTrack* FindOrCreateCinematicShotTrack();
+
+private:
+
+	/** Callback for executing a menu entry in the "Add Shot" combo button. */
+	virtual void HandleAddCinematicShotComboButtonMenuEntryExecute(const FAssetData& AssetData);
+
+	/** Callback for executing a menu entry in the "Add Shot" combo button when enter pressed. */
+	virtual void HandleAddCinematicShotComboButtonMenuEntryEnterPressed(const TArray<FAssetData>& AssetData);
+
+	/** Delegate for AnimatablePropertyChanged in AddKey */
+	virtual FKeyPropertyResult AddKeyInternal(FFrameNumber KeyTime, UMovieSceneSequence* InMovieSceneSequence, int32 RowIndex);
+
+	/** Delegate for shots button lock state */
+	ECheckBoxState AreShotsLocked() const;
+
+	/** Delegate for locked shots button */
+	void OnLockShotsClicked(ECheckBoxState CheckBoxState);
+	
+	/** Delegate for shots button lock tooltip */
+	FText GetLockShotsToolTip() const;
+
+protected:
+	
+	/**
+	 * Check whether the given sequence can be added as a sub-sequence.
+	 *
+	 * The purpose of this method is to disallow circular references
+	 * between sub-sequences in the focused movie scene.
+	 *
+	 * @param Sequence The sequence to check.
+	 * @return true if the sequence can be added as a sub-sequence, false otherwise.
+	 */
+	bool CanAddSubSequence(const UMovieSceneSequence& Sequence) const;
+	
+private:
+	
+	/** Called when our sequencer wants to switch cameras */
+	void OnUpdateCameraCut(UObject* CameraObject, bool bJumpCut);
+
+	/** Callback for AnimatablePropertyChanged in HandleAssetAdded. */
+	FKeyPropertyResult HandleSequenceAdded(FFrameNumber KeyTime, UMovieSceneSequence* Sequence, int32 RowIndex);
+
+	/** Callback for ImportEDL. */
+	void ImportEDL();
+	
+	/** Callback for ExportEDL. */
+	void ExportEDL();
+
+	/** Callback for ImportFCPXML. */
+	void ImportFCPXML();
+
+	/** Callback for ExportFCPXML. */
+	void ExportFCPXML();
+
+private:
+
+	/** The Thumbnail pool which draws all the viewport thumbnails for the shot track. */
+	TSharedPtr<FTrackEditorThumbnailPool> ThumbnailPool;
+
+	/** The camera actor for the current cut. */
+	TWeakObjectPtr<AActor> CinematicShotCamera;
+
+	/** Delegate binding handle for ISequencer::OnCameraCut */
+	FDelegateHandle OnCameraCutHandle;
+};
diff --git a/Engine/Source/Editor/MovieSceneTools/Public/TrackEditors/SubTrackEditor.h b/Engine/Source/Editor/MovieSceneTools/Public/TrackEditors/SubTrackEditor.h
--- a/Engine/Source/Editor/MovieSceneTools/Public/TrackEditors/SubTrackEditor.h
+++ b/Engine/Source/Editor/MovieSceneTools/Public/TrackEditors/SubTrackEditor.h
@@ -59,8 +59,8 @@ public:
 	/** Switch the selected sub section's take */
 	void SwitchTake(UObject* TakeObject);
 
-protected:
-
+public:
+	
 	/**
 	 * Check whether the given sequence can be added as a sub-sequence.
 	 *
diff --git a/Engine/Source/Runtime/MovieSceneTracks/Public/Sections/MovieSceneCinematicShotSection.h b/Engine/Source/Runtime/MovieSceneTracks/Public/Sections/MovieSceneCinematicShotSection.h
--- a/Engine/Source/Runtime/MovieSceneTracks/Public/Sections/MovieSceneCinematicShotSection.h
+++ b/Engine/Source/Runtime/MovieSceneTracks/Public/Sections/MovieSceneCinematicShotSection.h
@@ -16,9 +16,13 @@ class MOVIESCENETRACKS_API UMovieSceneCinematicShotSection
 {
 	GENERATED_BODY()
 
+public:
+
 	/** Object constructor. */
 	UMovieSceneCinematicShotSection(const FObjectInitializer& ObjInitializer);
 
+private:
+
 	/** ~UObject interface */
 	virtual void PostLoad() override;
 
diff --git a/Engine/Source/Runtime/MovieSceneTracks/Public/Tracks/MovieSceneCinematicShotTrack.h b/Engine/Source/Runtime/MovieSceneTracks/Public/Tracks/MovieSceneCinematicShotTrack.h
--- a/Engine/Source/Runtime/MovieSceneTracks/Public/Tracks/MovieSceneCinematicShotTrack.h
+++ b/Engine/Source/Runtime/MovieSceneTracks/Public/Tracks/MovieSceneCinematicShotTrack.h
@@ -13,10 +13,10 @@ class UMovieSceneSequence;
 class UMovieSceneSubSection;
 
 /**
- * A track that holds consecutive sub sequences.
- */
-UCLASS(MinimalAPI)
-class UMovieSceneCinematicShotTrack
+* A track that holds consecutive sub sequences.
+*/
+UCLASS()
+class MOVIESCENETRACKS_API UMovieSceneCinematicShotTrack
 	: public UMovieSceneSubTrack
 {
 	GENERATED_BODY()
@@ -25,15 +25,14 @@ public:
 
 	UMovieSceneCinematicShotTrack(const FObjectInitializer& ObjectInitializer);
 
-	MOVIESCENETRACKS_API void SortSections();
+	void SortSections();
 
 	// UMovieSceneSubTrack interface
+	virtual UMovieSceneSubSection* AddSequence(UMovieSceneSequence* Sequence, FFrameNumber StartTime, int32 Duration) { return AddSequenceOnRow(Sequence, StartTime, Duration, INDEX_NONE); }
 
-	MOVIESCENETRACKS_API virtual UMovieSceneSubSection* AddSequence(UMovieSceneSequence* Sequence, FFrameNumber StartTime, int32 Duration) { return AddSequenceOnRow(Sequence, StartTime, Duration, INDEX_NONE); }
-	MOVIESCENETRACKS_API virtual UMovieSceneSubSection* AddSequenceOnRow(UMovieSceneSequence* Sequence, FFrameNumber StartTime, int32 Duration, int32 RowIndex);
-
+	virtual UMovieSceneSubSection* AddSequenceOnRow(UMovieSceneSequence* Sequence, FFrameNumber StartTime, int32 Duration, int32 RowIndex);
+	
 	// UMovieSceneTrack interface
-
 	virtual void AddSection(UMovieSceneSection& Section) override;
 	virtual bool SupportsType(TSubclassOf<UMovieSceneSection> SectionClass) const override;
 	virtual UMovieSceneSection* CreateNewSection() override;
