Sequencer: Remove Show Motion Trail Options from Transform property sections, no longer valid since the tab now lives in the animation mode, need to open that first and then open the options.
#trivial

diff --git a/Engine/Source/Editor/MovieSceneTools/Private/Sections/TransformPropertySection.cpp b/Engine/Source/Editor/MovieSceneTools/Private/Sections/TransformPropertySection.cpp
--- a/Engine/Source/Editor/MovieSceneTools/Private/Sections/TransformPropertySection.cpp
+++ b/Engine/Source/Editor/MovieSceneTools/Private/Sections/TransformPropertySection.cpp
@@ -62,26 +62,6 @@ void FTransformSection::BuildSectionContextMenu(FMenuBuilder& MenuBuilder, const
 			})
 		);
 	};
-	MenuBuilder.BeginSection(NAME_None, LOCTEXT("MotionTrailOptions", "Motion Trail Options"));
-	{
-
-		FUIAction Action_ConfigureTexturesForSprites(
-			FExecuteAction::CreateLambda([]
-				{
-					FName ControlRigMotionTrailTab("ControlRigMotionTrailTab");
-					FGlobalTabmanager::Get()->TryInvokeTab(ControlRigMotionTrailTab);
-				}
-		));
-
-		MenuBuilder.AddMenuEntry(
-			LOCTEXT("ShowMotionTrailOptions", "Show Motion Trail Options"),
-			LOCTEXT("ShowMotionTrailOptionsTooltip", "Show various motion trail options"),
-			FSlateIcon(),
-			Action_ConfigureTexturesForSprites,
-			NAME_None,
-			EUserInterfaceActionType::Button);
-	}
-	MenuBuilder.EndSection();
 
 	MenuBuilder.BeginSection(NAME_None, LOCTEXT("TransformChannelsText", "Active Channels"));
 	{
