Sequencer: Revert OnPostPropertyChanged for now

diff --git a/Engine/Source/Editor/MovieSceneTools/Private/TrackEditors/SkeletalAnimationTrackEditor.cpp b/Engine/Source/Editor/MovieSceneTools/Private/TrackEditors/SkeletalAnimationTrackEditor.cpp
--- a/Engine/Source/Editor/MovieSceneTools/Private/TrackEditors/SkeletalAnimationTrackEditor.cpp
+++ b/Engine/Source/Editor/MovieSceneTools/Private/TrackEditors/SkeletalAnimationTrackEditor.cpp
@@ -13,6 +13,7 @@
 #include "Framework/Application/SlateApplication.h"
 #include "Widgets/Layout/SBox.h"
 #include "SequencerSectionPainter.h"
+#include "Components/SceneComponent.h"
 #include "Components/SkeletalMeshComponent.h"
 #include "Editor/UnrealEdEngine.h"
 #include "UnrealEdGlobals.h"
@@ -1218,7 +1219,19 @@ void FSkeletalAnimationTrackEditor::OnSequencerSaved(ISequencer& )
 
 void FSkeletalAnimationTrackEditor::OnPostPropertyChanged(UObject* InObject, FPropertyChangedEvent& InPropertyChangedEvent)
 {
-	// If the object changed has any animation track, notify sequencer to update because animations tick on their own and sequencer needs to evaluate again
+	// This attempts to fix tposes when changing properties without evaluating Sequencer UE-101261, 
+	// but unfortunately causes other problems like the temporary unkeyed value getting lost UE-136405
+/*
+	if (InPropertyChangedEvent.ChangeType != EPropertyChangeType::ValueSet)
+	{
+		return;
+	}
+
+	// If the object changed has any animation track:
+	// 1. Store the current transform (which may be an unkeyed value),
+	// 2. Evaluate Sequencer so that the skeletal animation track will be evaluated, and then the skeletal mesh with tick
+	// 3. Restore the current transform
+	// Without this, changing a value on a skeletal mesh will tick but not necessarily evaluate Sequencer, resulting in a tpose.
 	const bool bCreateIfMissing = false;
 	FFindOrCreateHandleResult HandleResult = FindOrCreateHandleToObject(InObject, bCreateIfMissing );
 	FGuid ObjectHandle = HandleResult.Handle;
@@ -1227,9 +1240,30 @@ void FSkeletalAnimationTrackEditor::OnPostPropertyChanged(UObject* InObject, FPr
 		FFindOrCreateTrackResult TrackResult = FindOrCreateTrackForObject(ObjectHandle, UMovieSceneSkeletalAnimationTrack::StaticClass(), NAME_None, bCreateIfMissing);
 		if (TrackResult.Track)
 		{
-			GetSequencer()->NotifyMovieSceneDataChanged(EMovieSceneDataChangeType::TrackValueChanged);
+			USceneComponent* SceneComponent = Cast<USceneComponent>(InObject);
+			if (!SceneComponent)
+			{
+				if (AActor* Actor = Cast<AActor>(InObject))
+				{
+					SceneComponent = Actor->GetRootComponent();
+				}
+			}
+
+			FTransform RelativeTransform;
+			if (SceneComponent)
+			{
+				RelativeTransform = SceneComponent->GetRelativeTransform();
+			}
+
+			GetSequencer()->ForceEvaluate();
+
+			if (SceneComponent)
+			{
+				SceneComponent->SetRelativeTransform(RelativeTransform);
+			}
 		}
 	}
+*/
 }
 
 bool FSkeletalAnimationTrackEditor::CreateAnimationSequence(const TArray<UObject*> NewAssets, USkeletalMeshComponent* SkelMeshComp, FGuid Binding, bool bCreateSoftLink)
