Sequencer: Operate on all the float/double channels of the section when getting/setting extrapolation rather than just the typed channel. This fixes an issue where double and float channels are extended by the menu but only the last is operated on.

diff --git a/Engine/Source/Editor/MovieSceneTools/Private/Channels/BuiltInChannelEditors.cpp b/Engine/Source/Editor/MovieSceneTools/Private/Channels/BuiltInChannelEditors.cpp
--- a/Engine/Source/Editor/MovieSceneTools/Private/Channels/BuiltInChannelEditors.cpp
+++ b/Engine/Source/Editor/MovieSceneTools/Private/Channels/BuiltInChannelEditors.cpp
@@ -24,6 +24,8 @@
 #include "Framework/MultiBox/MultiBoxBuilder.h"
 #include "Channels/MovieSceneChannelProxy.h"
 #include "Channels/MovieSceneChannelEditorData.h"
+#include "Channels/MovieSceneFloatChannel.h"
+#include "Channels/MovieSceneDoubleChannel.h"
 #include "Channels/DoubleChannelCurveModel.h"
 #include "Channels/FloatChannelCurveModel.h"
 #include "Channels/IntegerChannelCurveModel.h"
@@ -1349,15 +1351,23 @@ struct TCurveChannelSectionMenuExtension : FExtender, TSharedFromThis<TCurveChan
 		}
 
 		// Apply to all channels
-		for (const TMovieSceneChannelHandle<ChannelType>& Handle : Channels)
+		for (TWeakObjectPtr<UMovieSceneSection> WeakSection : Sections)
 		{
-			ChannelType* Channel = Handle.Get();
-
-			if (Channel)
+			if (UMovieSceneSection* Section = WeakSection.Get())
 			{
-				TEnumAsByte<ERichCurveExtrapolation>& DestExtrap = bPreInfinity ? Channel->PreInfinityExtrap : Channel->PostInfinityExtrap;
-				DestExtrap = ExtrapMode;
-				bAnythingChanged = true;
+				FMovieSceneChannelProxy& ChannelProxy = Section->GetChannelProxy();
+				for (FMovieSceneFloatChannel* Channel : ChannelProxy.GetChannels<FMovieSceneFloatChannel>())
+				{
+					TEnumAsByte<ERichCurveExtrapolation>& DestExtrap = bPreInfinity ? Channel->PreInfinityExtrap : Channel->PostInfinityExtrap;
+					DestExtrap = ExtrapMode;
+					bAnythingChanged = true;
+				}
+				for (FMovieSceneDoubleChannel* Channel : ChannelProxy.GetChannels<FMovieSceneDoubleChannel>())
+				{
+					TEnumAsByte<ERichCurveExtrapolation>& DestExtrap = bPreInfinity ? Channel->PreInfinityExtrap : Channel->PostInfinityExtrap;
+					DestExtrap = ExtrapMode;
+					bAnythingChanged = true;
+				}
 			}
 		}
 
@@ -1377,16 +1387,26 @@ struct TCurveChannelSectionMenuExtension : FExtender, TSharedFromThis<TCurveChan
 
 	bool IsExtrapolationModeSelected(ERichCurveExtrapolation ExtrapMode, bool bPreInfinity) const
 	{
-		for (const TMovieSceneChannelHandle<ChannelType>& Handle : Channels)
+		for (TWeakObjectPtr<UMovieSceneSection> WeakSection : Sections)
 		{
-			ChannelType* Channel = Handle.Get();
-
-			if (Channel)
+			if (UMovieSceneSection* Section = WeakSection.Get())
 			{
-				ERichCurveExtrapolation SourceExtrap = bPreInfinity ? Channel->PreInfinityExtrap : Channel->PostInfinityExtrap;
-				if (SourceExtrap != ExtrapMode)
+				FMovieSceneChannelProxy& ChannelProxy = Section->GetChannelProxy();
+				for (FMovieSceneFloatChannel* Channel : ChannelProxy.GetChannels<FMovieSceneFloatChannel>())
+				{
+					ERichCurveExtrapolation SourceExtrap = bPreInfinity ? Channel->PreInfinityExtrap : Channel->PostInfinityExtrap;
+					if (SourceExtrap != ExtrapMode)
+					{
+						return false;
+					}
+				}
+				for (FMovieSceneDoubleChannel* Channel : ChannelProxy.GetChannels<FMovieSceneDoubleChannel>())
 				{
-					return false;
+					ERichCurveExtrapolation SourceExtrap = bPreInfinity ? Channel->PreInfinityExtrap : Channel->PostInfinityExtrap;
+					if (SourceExtrap != ExtrapMode)
+					{
+						return false;
+					}
 				}
 			}
 		}
