Add proper ViewExtensions setup to the sequencer thumbnail capture

diff --git a/Engine/Source/Editor/MovieSceneTools/Private/TrackEditorThumbnail/TrackEditorThumbnail.cpp b/Engine/Source/Editor/MovieSceneTools/Private/TrackEditorThumbnail/TrackEditorThumbnail.cpp
--- a/Engine/Source/Editor/MovieSceneTools/Private/TrackEditorThumbnail/TrackEditorThumbnail.cpp
+++ b/Engine/Source/Editor/MovieSceneTools/Private/TrackEditorThumbnail/TrackEditorThumbnail.cpp
@@ -20,6 +20,7 @@
 #include "LegacyScreenPercentageDriver.h"
 #include "CanvasTypes.h"
 #include "EngineModule.h"
+#include "SceneViewExtension.h"
 
 namespace TrackEditorThumbnailConstants
 {
@@ -462,6 +463,13 @@ void FTrackEditorThumbnailCache::DrawViewportThumbnail(FTrackEditorThumbnail& Tr
 	FCanvas Canvas(TrackEditorThumbnail.GetRenderTarget(), nullptr, FGameTime::GetTimeSinceAppStart(), World->Scene->GetFeatureLevel());
 	Canvas.Clear(FLinearColor::Transparent);
 
+	ViewFamily.ViewExtensions = GEngine->ViewExtensions->GatherActiveExtensions(FSceneViewExtensionContext(World->Scene));
+	for (const FSceneViewExtensionRef& Extension : ViewFamily.ViewExtensions)
+	{
+		Extension->SetupViewFamily(ViewFamily);
+		Extension->SetupView(ViewFamily, *NewView);
+	}
+
 	GetRendererModule().BeginRenderingViewFamily(&Canvas, &ViewFamily);
 }
 
