Sequencer: Attach compensation when the attach object is a SceneComponent. This previously only worked for Actors and would crash when attaching with a SceneComponent.

diff --git a/Engine/Source/Editor/MovieSceneTools/Private/TrackEditors/AttachTrackEditor.cpp b/Engine/Source/Editor/MovieSceneTools/Private/TrackEditors/AttachTrackEditor.cpp
--- a/Engine/Source/Editor/MovieSceneTools/Private/TrackEditors/AttachTrackEditor.cpp
+++ b/Engine/Source/Editor/MovieSceneTools/Private/TrackEditors/AttachTrackEditor.cpp
@@ -485,13 +485,13 @@ FTransform GetLocationAtTime(UMovieScene3DTransformTrack* TransformTrack, FFrame
 	return FTransform::Identity;
 }
 
-AActor* GetConstraintActor(TSharedPtr<ISequencer> InSequencer, const FMovieSceneObjectBindingID& InConstraintBindingID)
+UObject* GetConstraintObject(TSharedPtr<ISequencer> InSequencer, const FMovieSceneObjectBindingID& InConstraintBindingID)
 {
 	TArrayView<TWeakObjectPtr<UObject>> RuntimeObjects = InConstraintBindingID.ResolveBoundObjects(InSequencer->GetFocusedTemplateID(), *InSequencer);
 
 	if (RuntimeObjects.Num() >= 1 && RuntimeObjects[0].IsValid())
 	{
-		return Cast<AActor>(RuntimeObjects[0].Get());
+		return RuntimeObjects[0].Get();
 	}
 
 	return nullptr;
@@ -530,18 +530,28 @@ struct FLocalTransformEvaluator : ITransformEvaluator
 		}
 
 		UMovieScene* MovieScene = Sequencer->GetFocusedMovieSceneSequence()->GetMovieScene();
-		AActor* Actor = Cast<AActor>(InObject);
 
-		FTransform ActorTransform = Actor->GetActorTransform();
-		TransformEval.SetSubtype<FTransform>(ActorTransform);
+		USceneComponent* SceneComponent = Cast<USceneComponent>(InObject);
+		if (AActor* Actor = Cast<AActor>(InObject))
+		{
+			SceneComponent = Actor->GetRootComponent();
+		}
 
-		FGuid ActorHandle = Sequencer->GetHandleToObject(Actor, false);
-		if (ActorHandle.IsValid())
+		if (!SceneComponent)
 		{
-			UMovieScene3DTransformTrack* ActorTransformTrack = Cast<UMovieScene3DTransformTrack>(MovieScene->FindTrack<UMovieScene3DTransformTrack>(ActorHandle));
+			return;
+		}
+
+		FTransform ComponentTransform = SceneComponent->GetComponentTransform();
+		TransformEval.SetSubtype<FTransform>(ComponentTransform);
+
+		FGuid ObjectHandle = Sequencer->GetHandleToObject(InObject, false);
+		if (ObjectHandle.IsValid())
+		{
+			UMovieScene3DTransformTrack* ActorTransformTrack = Cast<UMovieScene3DTransformTrack>(MovieScene->FindTrack<UMovieScene3DTransformTrack>(ObjectHandle));
 			if (ActorTransformTrack)
 			{
-				TransformEval.SetSubtype<TTuple<UMovieScene3DTransformTrack*, UObject*>>(TTuple<UMovieScene3DTransformTrack*, UObject*>(ActorTransformTrack, Actor));
+				TransformEval.SetSubtype<TTuple<UMovieScene3DTransformTrack*, UObject*>>(TTuple<UMovieScene3DTransformTrack*, UObject*>(ActorTransformTrack, SceneComponent));
 			}
 		}
 	}
@@ -622,8 +632,14 @@ struct FWorldTransformEvaluator : ITransformEvaluator
 		UMovieScene* MovieScene = Sequencer->GetFocusedMovieSceneSequence()->GetMovieScene();
 
 		FName SocketName = InSocketName;
-		AActor* Actor = Cast<AActor>(InObject);
-		if (!Actor)
+
+		USceneComponent* SceneComponent = Cast<USceneComponent>(InObject);
+		if (AActor* Actor = Cast<AActor>(InObject))
+		{
+			SceneComponent = Actor->GetRootComponent();
+		}
+
+		if (!SceneComponent)
 		{
 			return;
 		}
@@ -633,36 +649,36 @@ struct FWorldTransformEvaluator : ITransformEvaluator
 		{
 			TUnion<FTransform, TTuple<UMovieScene3DTransformTrack*, UObject*>> ActorEval;
 			// If we find a socket, get the world transform of the socket and break out immediately
-			if (Actor->GetRootComponent()->DoesSocketExist(SocketName))
+			if (SceneComponent->DoesSocketExist(SocketName))
 			{
-				const FTransform SocketWorldSpace = Actor->GetRootComponent()->GetSocketTransform(SocketName);
+				const FTransform SocketWorldSpace = SceneComponent->GetSocketTransform(SocketName);
 				ActorEval.SetSubtype<FTransform>(SocketWorldSpace);
 				TransformEvals.Add(ActorEval);
 				return;
 			}
 			
-			FTransform ActorTransform = Actor->GetActorTransform();
-			ActorEval.SetSubtype<FTransform>(ActorTransform);
+			FTransform ComponentTransform = SceneComponent->GetComponentTransform();
+			ActorEval.SetSubtype<FTransform>(ComponentTransform);
 
-			FGuid ActorHandle = Sequencer->GetHandleToObject(Actor, false);
-			if (ActorHandle.IsValid())
+			FGuid ObjectHandle = Sequencer->GetHandleToObject(SceneComponent, false);
+			if (ObjectHandle.IsValid())
 			{
-				UMovieScene3DTransformTrack* ActorTransformTrack = MovieScene->FindTrack<UMovieScene3DTransformTrack>(ActorHandle);
-				if (ActorTransformTrack)
+				UMovieScene3DTransformTrack* TransformTrack = MovieScene->FindTrack<UMovieScene3DTransformTrack>(ObjectHandle);
+				if (TransformTrack)
 				{
-					ActorEval.SetSubtype<TTuple<UMovieScene3DTransformTrack*, UObject*>>(TTuple<UMovieScene3DTransformTrack*, UObject*>(ActorTransformTrack, Actor));
+					ActorEval.SetSubtype<TTuple<UMovieScene3DTransformTrack*, UObject*>>(TTuple<UMovieScene3DTransformTrack*, UObject*>(TransformTrack, SceneComponent));
 				}
 			}
 
 			TransformEvals.Add(ActorEval);
 
-			Actor = Actor->GetAttachParentActor();
-			if (Actor)
+			SceneComponent = SceneComponent->GetAttachParent();
+			if (SceneComponent)
 			{
-				SocketName = Actor->GetAttachParentSocketName();
+				SocketName = SceneComponent->GetAttachSocketName();
 			}
 		} 
-		while (Actor);
+		while (SceneComponent);
 	}
 
 	/**
@@ -761,9 +777,9 @@ struct FAttachRevertModifier
 		, RevertRange(InRevertRange)
 	{
 		FMovieSceneObjectBindingID ConstraintID = InAttachSection->GetConstraintBindingID();
-		AActor* ConstraintActor = GetConstraintActor(InWeakAttachTrackEditor->GetSequencer(), ConstraintID);
+		UObject* ConstraintObject = GetConstraintObject(InWeakAttachTrackEditor->GetSequencer(), ConstraintID);
 
-		TransformEvaluator = FWorldTransformEvaluator(InWeakAttachTrackEditor, ConstraintActor, InSocketName);
+		TransformEvaluator = FWorldTransformEvaluator(InWeakAttachTrackEditor, ConstraintObject, InSocketName);
 
 		BeginConstraintTransform = TransformEvaluator(InRevertRange.GetLowerBoundValue());
 
@@ -1149,9 +1165,9 @@ FKeyPropertyResult F3DAttachTrackEditor::AddKeyInternal( FFrameNumber KeyTime, c
 	UMovieScene* MovieScene = GetSequencer()->GetFocusedMovieSceneSequence()->GetMovieScene();
 
 	// It's possible that the objects bound to this parent binding ID are null, in which case there will be no compensation
-	AActor* ParentActor = GetConstraintActor(GetSequencer(), ConstraintBindingID);	
+	UObject* ParentObject = GetConstraintObject(GetSequencer(), ConstraintBindingID);	
 
-	FWorldTransformEvaluator ParentTransformEval(SharedThis(this), ParentActor, SocketName);
+	FWorldTransformEvaluator ParentTransformEval(SharedThis(this), ParentObject, SocketName);
 
 	TOptional<TArrayView<FMovieSceneDoubleChannel*>> ParentChannels;
 
@@ -1170,7 +1186,7 @@ FKeyPropertyResult F3DAttachTrackEditor::AddKeyInternal( FFrameNumber KeyTime, c
 		UObject* Object = Objects[ObjectIndex].Get();
 
 		// Disallow attaching an object to itself
-		if (Object == ParentActor)
+		if (Object == ParentObject)
 		{
 			continue;
 		}
@@ -1223,10 +1239,19 @@ FKeyPropertyResult F3DAttachTrackEditor::AddKeyInternal( FFrameNumber KeyTime, c
 
 		// Create a blank world transform evaluator, add parent evaluator if there is a parent
 		FWorldTransformEvaluator WorldChildTransformEval(SharedThis(this), nullptr);
-		AActor* Actor = Cast<AActor>(Object);
-		if (AActor* PrevParentActor = Actor->GetAttachParentActor())
+
+		USceneComponent* SceneComponent = Cast<USceneComponent>(Object);
+		if (AActor* Actor = Cast<AActor>(Object))
+		{
+			SceneComponent = Actor->GetRootComponent();
+		}
+
+		if (SceneComponent)
 		{
-			WorldChildTransformEval = FWorldTransformEvaluator(SharedThis(this), PrevParentActor);
+			if (USceneComponent* PrevAttachParent = SceneComponent->GetAttachParent())
+			{
+				WorldChildTransformEval = FWorldTransformEvaluator(SharedThis(this), PrevAttachParent);
+			}
 		}
 
 		// Create transform track for object
@@ -1239,9 +1264,9 @@ FKeyPropertyResult F3DAttachTrackEditor::AddKeyInternal( FFrameNumber KeyTime, c
 		{
 			WorldChildTransformEval.PrependTransformEval(Object, TransformTrack);
 		}
-		else
+		else if (SceneComponent)
 		{
-			WorldChildTransformEval.PrependTransformEval(Actor->GetTransform());
+			WorldChildTransformEval.PrependTransformEval(SceneComponent->GetComponentTransform());
 		}
 
 		if (!TransformSection || !TransformTrack)
@@ -1278,7 +1303,7 @@ FKeyPropertyResult F3DAttachTrackEditor::AddKeyInternal( FFrameNumber KeyTime, c
 		KeyPropertyResult.bKeyCreated = true;
 
 		TOptional<FAttachRevertModifier> RevertModifier;
-		AActor* ReAttachOnDetach = nullptr;
+		USceneComponent* ReAttachOnDetach = nullptr;
 
 		// If there are existing channels, revert the transform from the previous parent's transform before setting the new relative transform
 		// Currently don't handle objects with both other attach sections and are already attached to other objects because its hard to think about
@@ -1323,7 +1348,13 @@ FKeyPropertyResult F3DAttachTrackEditor::AddKeyInternal( FFrameNumber KeyTime, c
 		
 			RevertModifier = FAttachRevertModifier(SharedThis(this), RevertRange, PrevParentEvaluator, PreserveType == ETransformPreserveType::CurrentKey);
 
-			ReAttachOnDetach = Actor->GetAttachParentActor();
+			if (SceneComponent)
+			{
+				if (USceneComponent* PrevAttachParent = SceneComponent->GetAttachParent())
+				{
+					ReAttachOnDetach = PrevAttachParent;
+				}
+			}
 		}
 
 		if (RevertModifier.IsSet())
diff --git a/Engine/Source/Runtime/MovieSceneTracks/Public/Sections/MovieScene3DAttachSection.h b/Engine/Source/Runtime/MovieSceneTracks/Public/Sections/MovieScene3DAttachSection.h
--- a/Engine/Source/Runtime/MovieSceneTracks/Public/Sections/MovieScene3DAttachSection.h
+++ b/Engine/Source/Runtime/MovieSceneTracks/Public/Sections/MovieScene3DAttachSection.h
@@ -50,7 +50,7 @@ public:
 	bool bFullRevertOnDetach;
 
 	UPROPERTY(transient)
-	TSoftObjectPtr<AActor> ReAttachOnDetach;
+	TSoftObjectPtr<USceneComponent> ReAttachOnDetach;
 
 #endif
 
