Sequencer: Control Rig: FBX Import Fix, missed Y negation when we did LWC
#prefligh t61d89c396511bc498e60078c

diff --git a/Engine/Source/Editor/MovieSceneTools/Private/MovieSceneToolHelpers.cpp b/Engine/Source/Editor/MovieSceneTools/Private/MovieSceneToolHelpers.cpp
--- a/Engine/Source/Editor/MovieSceneTools/Private/MovieSceneToolHelpers.cpp
+++ b/Engine/Source/Editor/MovieSceneTools/Private/MovieSceneToolHelpers.cpp
@@ -1419,6 +1419,7 @@ void SetChannelValue(FMovieSceneDoubleChannel* DoubleChannel, FMovieSceneFloatCh
 				}
 				else if (FloatChannel)
 				{
+					bNegate = !bNegate;
 					float Default = bNegate ? -Location.Y : Location.Y;
 					FloatChannel->SetDefault(Default);
 					ImportTransformChannelToFloat(TranslationY, FloatChannel, FrameRate, false, false, StartFrame, bNegate);
