MovieSceneTools: use Sequencer display format and padding settings for subsection track frame time hints

diff --git a/Engine/Source/Editor/MovieSceneTools/Private/TrackEditors/SubTrackEditorBase.cpp b/Engine/Source/Editor/MovieSceneTools/Private/TrackEditors/SubTrackEditorBase.cpp
--- a/Engine/Source/Editor/MovieSceneTools/Private/TrackEditors/SubTrackEditorBase.cpp
+++ b/Engine/Source/Editor/MovieSceneTools/Private/TrackEditors/SubTrackEditorBase.cpp
@@ -1,6 +1,9 @@
 // Copyright Epic Games, Inc. All Rights Reserved.
 
 #include "TrackEditors/SubTrackEditorBase.h"
+#include "FrameNumberDisplayFormat.h"
+#include "FrameNumberNumericInterface.h"
+#include "SequencerSettings.h"
 #include "Tracks/MovieSceneCinematicShotTrack.h"
 #include "Tracks/MovieSceneSubTrack.h"
 
@@ -109,12 +112,27 @@ FSubSectionPainterResult FSubSectionPainterUtil::PaintSection(TSharedPtr<const I
         FFrameTime CurrentTime = Sequencer->GetLocalTime().Time;
         if (SectionRange.Contains(CurrentTime.FrameNumber))
         {
-            UMovieScene* SubSequenceMovieScene = SectionObject.GetSequence()->GetMovieScene();
+            const UMovieScene* SubSequenceMovieScene = SectionObject.GetSequence()->GetMovieScene();
             const FFrameRate DisplayRate = SubSequenceMovieScene->GetDisplayRate();
             const FFrameRate TickResolution = SubSequenceMovieScene->GetTickResolution();
-            const FFrameNumber CurrentFrameNumber = ConvertFrameTime(CurrentTime * SectionObject.OuterToInnerTransform(), TickResolution, DisplayRate).FloorToFrame();
+			const FFrameTime HintFrameTime = CurrentTime * SectionObject.OuterToInnerTransform();
 
-            DrawFrameTimeHint(InPainter, CurrentTime, CurrentFrameNumber);
+			// Get the desired frame display format and zero padding from
+			// the sequencer settings, if possible.
+			TAttribute<EFrameNumberDisplayFormats> DisplayFormatAttr(EFrameNumberDisplayFormats::Frames);
+			TAttribute<uint8> ZeroPadFrameNumbersAttr(0u);
+			if (const USequencerSettings* SequencerSettings = Sequencer->GetSequencerSettings())
+			{
+				DisplayFormatAttr.Set(SequencerSettings->GetTimeDisplayFormat());
+				ZeroPadFrameNumbersAttr.Set(SequencerSettings->GetZeroPadFrames());
+			}
+
+			const TAttribute<FFrameRate> TickResolutionAttr(TickResolution);
+			const TAttribute<FFrameRate> DisplayRateAttr(DisplayRate);
+
+			const FFrameNumberInterface FrameNumberInterface(DisplayFormatAttr, ZeroPadFrameNumbersAttr, TickResolutionAttr, DisplayRateAttr);
+
+			DrawFrameTimeHint(InPainter, CurrentTime, HintFrameTime, &FrameNumberInterface);
         }
     }
 
