Sequencer: Pivot Tool.  Don't evaluate if at same time for actors, similar to what we did for control rig.
#trivial

diff --git a/Engine/Source/Editor/MovieSceneTools/Private/MovieSceneToolHelpers.cpp b/Engine/Source/Editor/MovieSceneTools/Private/MovieSceneToolHelpers.cpp
--- a/Engine/Source/Editor/MovieSceneTools/Private/MovieSceneToolHelpers.cpp
+++ b/Engine/Source/Editor/MovieSceneTools/Private/MovieSceneToolHelpers.cpp
@@ -4371,7 +4371,7 @@ static void GetNonSequencerActorWorldTransforms(IMovieScenePlayer* Player, UMovi
 		do
 		{
 			FGuid ActorHandle = GetHandleToObject(Actor, InSequence, Player, Template,false);
-			if (ActorHandle.IsValid())
+			if (Frames.Num() > 0)
 			{
 				if (InSequence->GetMovieScene()->FindTrack<UMovieScene3DTransformTrack>(ActorHandle))
 				{
@@ -4402,7 +4402,8 @@ static void GetNonSequencerActorWorldTransforms(IMovieScenePlayer* Player, UMovi
 		} while (Actor);
 		//if we get to here world transform is what we want 
 	}
-	OutTransforms.SetNumUninitialized(Frames.Num());
+	int32 NumFrames = Frames.Num() > 0 ? Frames.Num() : 1; //if no frames passed in we actually have one
+	OutTransforms.SetNumUninitialized(NumFrames);
 	for (FTransform& OutTransform : OutTransforms)
 	{
 		OutTransform = WorldTransform;
@@ -4412,9 +4413,9 @@ static void GetNonSequencerActorWorldTransforms(IMovieScenePlayer* Player, UMovi
 void MovieSceneToolHelpers::GetActorWorldTransforms(ISequencer* Sequencer, const FActorForWorldTransforms& ActorSelection, const TArray<FFrameNumber>& Frames, TArray<FTransform>& OutWorldTransforms)
 {
 	FMovieSceneSequenceIDRef Template = Sequencer->GetFocusedTemplateID();
-
+	const bool bAvoidEvaluates = (Frames.Num() == 0) || (Frames.Num() == 1 && Frames[0] == Sequencer->GetLocalTime().Time.RoundToFrame());
 	FGuid ObjectHandle = Sequencer->GetHandleToObject(ActorSelection.Actor.Get(), false);
-	if (ObjectHandle.IsValid())
+	if (!bAvoidEvaluates && ObjectHandle.IsValid())
 	{
 		if (Sequencer->GetFocusedMovieSceneSequence()->GetMovieScene()->FindTrack<UMovieScene3DTransformTrack>(ObjectHandle))
 		{
@@ -4424,7 +4425,7 @@ void MovieSceneToolHelpers::GetActorWorldTransforms(ISequencer* Sequencer, const
 	}
 
 	ObjectHandle = Sequencer->GetHandleToObject(ActorSelection.Component.Get(), false);
-	if (ObjectHandle.IsValid())
+	if (!bAvoidEvaluates && ObjectHandle.IsValid())
 	{
 		if (Sequencer->GetFocusedMovieSceneSequence()->GetMovieScene()->FindTrack<UMovieScene3DTransformTrack>(ObjectHandle))
 		{
@@ -4432,7 +4433,15 @@ void MovieSceneToolHelpers::GetActorWorldTransforms(ISequencer* Sequencer, const
 			return;
 		}
 	}
-	GetNonSequencerActorWorldTransforms(Sequencer, Sequencer->GetFocusedMovieSceneSequence(), Template, ActorSelection, Frames, OutWorldTransforms);
+	if (bAvoidEvaluates)
+	{
+		TArray<FFrameNumber> NoFrame; //this will make sure we don't evaluate
+		GetNonSequencerActorWorldTransforms(Sequencer, Sequencer->GetFocusedMovieSceneSequence(), Template, ActorSelection, NoFrame, OutWorldTransforms);
+	}
+	else
+	{
+		GetNonSequencerActorWorldTransforms(Sequencer, Sequencer->GetFocusedMovieSceneSequence(), Template, ActorSelection, Frames, OutWorldTransforms);
+	}
 }
 
 void MovieSceneToolHelpers::SetOrAddKey(TMovieSceneChannelData<FMovieSceneFloatValue>& ChannelData, FFrameNumber Time, float Value)
