Sequencer: Apply transform origin inverse to transform keys

diff --git a/Engine/Source/Editor/MovieSceneTools/Private/TrackEditors/TransformTrackEditor.cpp b/Engine/Source/Editor/MovieSceneTools/Private/TrackEditors/TransformTrackEditor.cpp
--- a/Engine/Source/Editor/MovieSceneTools/Private/TrackEditors/TransformTrackEditor.cpp
+++ b/Engine/Source/Editor/MovieSceneTools/Private/TrackEditors/TransformTrackEditor.cpp
@@ -34,6 +34,9 @@
 #include "Systems/MovieScenePropertyInstantiator.h"
 #include "MovieSceneTracksComponentTypes.h"
 
+#include "Tracks/IMovieSceneTransformOrigin.h"
+#include "IMovieScenePlaybackClient.h"
+
 #define LOCTEXT_NAMESPACE "MovieScene_TransformTrack"
 
 void GetActorAndSceneComponentFromObject( UObject* Object, AActor*& OutActor, USceneComponent*& OutSceneComponent )
@@ -804,6 +807,24 @@ void F3DTransformTrackEditor::GetTransformKeys( const TOptional<FTransformData>&
 	}
 }
 
+FTransform F3DTransformTrackEditor::GetTransformOrigin() const
+{
+	FTransform TransformOrigin;
+
+	const IMovieScenePlaybackClient*  Client       = GetSequencer()->GetPlaybackClient();
+	const UObject*                    InstanceData = Client ? Client->GetInstanceData() : nullptr;
+	const IMovieSceneTransformOrigin* RawInterface = Cast<const IMovieSceneTransformOrigin>(InstanceData);
+
+	const bool bHasInterface = RawInterface || (InstanceData && InstanceData->GetClass()->ImplementsInterface(UMovieSceneTransformOrigin::StaticClass()));
+	if (bHasInterface)
+	{
+		// Retrieve the current origin
+		TransformOrigin = RawInterface ? RawInterface->GetTransformOrigin() : IMovieSceneTransformOrigin::Execute_BP_GetTransformOrigin(InstanceData);
+	}
+
+	return TransformOrigin;
+}
+
 void F3DTransformTrackEditor::AddTransformKeysForHandle(TArray<FGuid> ObjectHandles, EMovieSceneTransformChannel ChannelToKey, ESequencerKeyMode KeyMode)
 {
 	const FScopedTransaction Transaction(NSLOCTEXT("Sequencer", "AddTransformTrack", "Add Transform Track"));
@@ -881,7 +902,11 @@ FTransformData F3DTransformTrackEditor::RecomposeTransform(const FTransformData&
 			FIntermediate3DTransform CurrentValue(InTransformData.Translation, InTransformData.Rotation, InTransformData.Scale);
 
 			TRecompositionResult<FIntermediate3DTransform> TransformData = System->RecomposeBlendOperational(FMovieSceneTracksComponentTypes::Get()->ComponentTransform, Query, CurrentValue);
-			return FTransformData(TransformData.Values[0].GetTranslation(), TransformData.Values[0].GetRotation(), TransformData.Values[0].GetScale());
+
+			FTransform CurrentTransform(TransformData.Values[0].GetRotation(), TransformData.Values[0].GetTranslation(), TransformData.Values[0].GetScale());
+			CurrentTransform *= GetTransformOrigin().Inverse();
+
+			return FTransformData(CurrentTransform.GetLocation(), CurrentTransform.GetRotation().Rotator(), CurrentTransform.GetScale3D());
 		}
 	}
 
@@ -974,7 +999,10 @@ void F3DTransformTrackEditor::ProcessKeyOperation(UObject* ObjectToKey, TArrayVi
 		Query.bConvertFromSourceEntityIDs = false;
 		Query.Object   = Component;
 
-		FIntermediate3DTransform CurrentValue(Component->GetRelativeLocation(), Component->GetRelativeRotation(), Component->GetRelativeScale3D());
+		FTransform CurrentTransform(Component->GetRelativeRotation(), Component->GetRelativeLocation(), Component->GetRelativeScale3D());
+		CurrentTransform *= GetTransformOrigin().Inverse();
+
+		FIntermediate3DTransform CurrentValue(CurrentTransform.GetTranslation(), CurrentTransform.GetRotation().Rotator(), CurrentTransform.GetScale3D());
 		TRecompositionResult<FIntermediate3DTransform> TransformData = System->RecomposeBlendOperational(FMovieSceneTracksComponentTypes::Get()->ComponentTransform, Query, CurrentValue);
 
 		for (int32 Index = 0; Index < SectionsToKey.Num(); ++Index)
diff --git a/Engine/Source/Editor/MovieSceneTools/Private/TrackEditors/TransformTrackEditor.h b/Engine/Source/Editor/MovieSceneTools/Private/TrackEditors/TransformTrackEditor.h
--- a/Engine/Source/Editor/MovieSceneTools/Private/TrackEditors/TransformTrackEditor.h
+++ b/Engine/Source/Editor/MovieSceneTools/Private/TrackEditors/TransformTrackEditor.h
@@ -132,7 +132,8 @@ private:
 		One transform key is generated for each individual key to be added to the section. */
 	void GetTransformKeys( const TOptional<FTransformData>& LastTransform, const FTransformData& CurrentTransform, EMovieSceneTransformChannel ChannelsToKey, UObject* Object, UMovieSceneSection* Section, FGeneratedTrackKeys& OutGeneratedKeys );
 
-
+	/** Transform origin which may be set for the current level sequence */
+	FTransform GetTransformOrigin() const;
 
 	/** 
 	 * Adds transform keys to an object represented by a handle.
