Sequencer: Copy camera template object to spawnable when copying properties. This fixes a bug where importing properties onto a spawnable camera loses those value if no track is created.

diff --git a/Engine/Source/Editor/MovieSceneTools/Private/MovieSceneToolHelpers.cpp b/Engine/Source/Editor/MovieSceneTools/Private/MovieSceneToolHelpers.cpp
--- a/Engine/Source/Editor/MovieSceneTools/Private/MovieSceneToolHelpers.cpp
+++ b/Engine/Source/Editor/MovieSceneTools/Private/MovieSceneToolHelpers.cpp
@@ -3008,6 +3008,13 @@ void MovieSceneToolHelpers::ImportFBXCameraToExisting(UnFbx::FFbxImporter* FbxIm
 					continue;
 				}
 
+				// If copying properties to a spawnable object, the template object must be updated
+				FMovieSceneSpawnable* Spawnable = MovieScene->FindSpawnable(InObjectBinding.Key);
+				if (Spawnable)
+				{
+					Spawnable->CopyObjectTemplate(*FoundObject, *InSequence);
+				}
+
 				UMovieSceneFloatTrack* FloatTrack = MovieScene->FindTrack<UMovieSceneFloatTrack>(PropertyOwnerGuid, TrackName);
 				if (FloatTrack)
 				{
