Sequencer: Fix typo in DoubleChannels

diff --git a/Engine/Source/Editor/MovieSceneTools/Private/MovieSceneToolHelpers.cpp b/Engine/Source/Editor/MovieSceneTools/Private/MovieSceneToolHelpers.cpp
--- a/Engine/Source/Editor/MovieSceneTools/Private/MovieSceneToolHelpers.cpp
+++ b/Engine/Source/Editor/MovieSceneTools/Private/MovieSceneToolHelpers.cpp
@@ -1938,7 +1938,7 @@ static bool ImportFBXTransformToChannels(FString NodeName, const UMovieSceneUser
 		ImportTransformChannelToFloat(EulerRotation[1], FloatChannels[4], FrameRate, true, false, StartFrame);
 		ImportTransformChannelToFloat(EulerRotation[2], FloatChannels[5], FrameRate, true, false, StartFrame);
 
-		if (DoubleChannels.Num() > 6) //noscale
+		if (FloatChannels.Num() > 6) //noscale
 		{
 			ImportTransformChannelToFloat(Scale[0], FloatChannels[6], FrameRate, false, false, StartFrame);
 			ImportTransformChannelToFloat(Scale[1], FloatChannels[7], FrameRate, false, false, StartFrame);
