MovieSceneTools: look for timecode properties for use as TimecodeSource when importing FBX to Control Rig

These changes add handling for custom timecode properties coming from FBX in
MovieSceneToolHelpers::ImportFBXIntoControlRigChannels(). All nodes being imported
from FBX are searched for animated properties that match the timecode custom attribute
names (which default to "TCHour", "TCMinute", "TCSecond", and "TCFrame", but can be
customized in the animation settings). When any of those properties are found on a node,
an FTimecode is constructed from the component values and is used to populate the
"TimecodeSource" property of all of the movie scene sections that were modified by the
import operation. The first node encountered that has timecode properties is the one
whose values will be used.

diff --git a/Engine/Source/Editor/MovieSceneTools/Private/MovieSceneToolHelpers.cpp b/Engine/Source/Editor/MovieSceneTools/Private/MovieSceneToolHelpers.cpp
--- a/Engine/Source/Editor/MovieSceneTools/Private/MovieSceneToolHelpers.cpp
+++ b/Engine/Source/Editor/MovieSceneTools/Private/MovieSceneToolHelpers.cpp
@@ -2,6 +2,7 @@
 
 #include "MovieSceneToolHelpers.h"
 #include "ActorForWorldTransforms.h"
+#include "Animation/AnimationSettings.h"
 #include "MovieSceneToolsModule.h"
 #include "MovieScene.h"
 #include "Layout/Margin.h"
@@ -2517,7 +2518,33 @@ bool MovieSceneToolHelpers::ImportFBXIntoControlRigChannels(UMovieScene* MovieSc
 		FFrameNumber  StartFrame = ImportFBXControlRigSettings->StartTimeRange;
 		FFrameNumber  EndFrame = ImportFBXControlRigSettings->EndTimeRange;
 
-		FString RootNodeName = FbxImporter->Scene->GetRootNode()->GetName();
+		// We'll be looking for timecode properties on all of the nodes in the FBX being
+		// considered for import, and the first one found will be used to populate the
+		// timecode source of all modified movie scene sections.
+		FName TCHourAttrName(TEXT("TCHour"));
+		FName TCMinuteAttrName(TEXT("TCMinute"));
+		FName TCSecondAttrName(TEXT("TCSecond"));
+		FName TCFrameAttrName(TEXT("TCFrame"));
+
+		if (const UAnimationSettings* AnimationSettings = UAnimationSettings::Get())
+		{
+			TCHourAttrName = AnimationSettings->BoneTimecodeCustomAttributeNameSettings.HourAttributeName;
+			TCMinuteAttrName = AnimationSettings->BoneTimecodeCustomAttributeNameSettings.MinuteAttributeName;
+			TCSecondAttrName = AnimationSettings->BoneTimecodeCustomAttributeNameSettings.SecondAttributeName;
+			TCFrameAttrName = AnimationSettings->BoneTimecodeCustomAttributeNameSettings.FrameAttributeName;
+		}
+
+		const TArray<FString> TimecodePropertyNames = {
+			TCHourAttrName.ToString(),
+			TCMinuteAttrName.ToString(),
+			TCSecondAttrName.ToString(),
+			TCFrameAttrName.ToString()
+		};
+
+		FTimecode Timecode;
+		bool bFoundTimecode = false;
+
+		TSet<UMovieSceneSection*> AllModifiedSections;
 
 		for (int32 NodeIndex = 0; NodeIndex < AllNodeNames.Num(); ++NodeIndex)
 		{
@@ -2537,24 +2564,18 @@ bool MovieSceneToolHelpers::ImportFBXIntoControlRigChannels(UMovieScene* MovieSc
 				{
 					if (NodeAndChannel.MovieSceneTrack)
 					{
-						if (NodeAndChannel.MovieSceneTrack->GetSectionToKey())
+						UMovieSceneSection* SectionToModify = NodeAndChannel.MovieSceneTrack->GetSectionToKey();
+						if (!SectionToModify && NodeAndChannel.MovieSceneTrack->GetAllSections().Num() > 0)
 						{
-							if (!ModifiedSections.Contains(NodeAndChannel.MovieSceneTrack->GetSectionToKey()))
-							{
-								NodeAndChannel.MovieSceneTrack->GetSectionToKey()->SetFlags(RF_Transactional);
-								NodeAndChannel.MovieSceneTrack->GetSectionToKey()->Modify();
-								ModifiedSections.Add(NodeAndChannel.MovieSceneTrack->GetSectionToKey());
-							}
+							SectionToModify = NodeAndChannel.MovieSceneTrack->GetAllSections()[0];
 						}
-						else if (NodeAndChannel.MovieSceneTrack->GetAllSections().Num() > 0)
-						{
-							if (!ModifiedSections.Contains(NodeAndChannel.MovieSceneTrack->GetAllSections()[0]))
-							{
-								NodeAndChannel.MovieSceneTrack->GetAllSections()[0]->SetFlags(RF_Transactional);
-								NodeAndChannel.MovieSceneTrack->GetAllSections()[0]->Modify();
-								ModifiedSections.Add(NodeAndChannel.MovieSceneTrack->GetAllSections()[0]);
 
-							}
+						if (SectionToModify && !ModifiedSections.Contains(SectionToModify))
+						{
+							SectionToModify->SetFlags(RF_Transactional);
+							SectionToModify->Modify();
+							ModifiedSections.Add(SectionToModify);
+							AllModifiedSections.Add(SectionToModify);
 						}
 					}
 
@@ -2565,6 +2586,61 @@ bool MovieSceneToolHelpers::ImportFBXIntoControlRigChannels(UMovieScene* MovieSc
 					ImportFBXTransformToChannels(NodeName, CurrentImportFBXSettings, ImportFBXControlRigSettings, FrameToInsertOrReplace, FrameRate, NodeAndChannel, CurveAPI);
 				}
 			}
+
+			// We consider *all* nodes in the FBX when looking for authored timecode properties,
+			// not just the ones that map to a particular node with channels.
+			if (!bFoundTimecode)
+			{
+				TArray<FString> AnimatedPropertyNames;
+				CurveAPI.GetNodeAnimatedPropertyNameArray(NodeName, AnimatedPropertyNames);
+				for (const FString& AnimatedPropertyName : AnimatedPropertyNames)
+				{
+					if (!TimecodePropertyNames.Contains(AnimatedPropertyName))
+					{
+						continue;
+					}
+
+					bFoundTimecode = true;
+
+					const int32 ChannelIndex = 0;
+					const int32 CompositeIndex = 0;
+					FRichCurve PropertyCurve;
+					const bool bNegative = false;
+					CurveAPI.GetCurveDataForSequencer(NodeName, AnimatedPropertyName, ChannelIndex, CompositeIndex, PropertyCurve, bNegative);
+
+					const float EvalTime = 0.0f;
+					const float DefaultValue = 0.0f;
+					const float Value = PropertyCurve.Eval(EvalTime, DefaultValue);
+					const int32 IntValue = FMath::TruncToInt(Value);
+
+					const FName AnimatedPropertyNameAsFName(AnimatedPropertyName);
+
+					if (AnimatedPropertyNameAsFName.IsEqual(TCHourAttrName))
+					{
+						Timecode.Hours = IntValue;
+					}
+					else if (AnimatedPropertyNameAsFName.IsEqual(TCMinuteAttrName))
+					{
+						Timecode.Minutes = IntValue;
+					}
+					else if (AnimatedPropertyNameAsFName.IsEqual(TCSecondAttrName))
+					{
+						Timecode.Seconds = IntValue;
+					}
+					else if (AnimatedPropertyNameAsFName.IsEqual(TCFrameAttrName))
+					{
+						Timecode.Frames = IntValue;
+					}
+				}
+			}
+		}
+
+		if (bFoundTimecode)
+		{
+			for (UMovieSceneSection* Section : AllModifiedSections)
+			{
+				Section->TimecodeSource = FMovieSceneTimecodeSource(Timecode);
+			}
 		}
 
 		// restore
