Sequencer: Added support for FVector4d - Vector4f for FFloatVectorPropertyTrackEditor and Vector4 for FDoubleVectorPropertyTrackEditor

diff --git a/Engine/Source/Editor/MovieSceneTools/Private/TrackEditors/PropertyTrackEditors/VectorPropertyTrackEditor.cpp b/Engine/Source/Editor/MovieSceneTools/Private/TrackEditors/PropertyTrackEditors/VectorPropertyTrackEditor.cpp
--- a/Engine/Source/Editor/MovieSceneTools/Private/TrackEditors/PropertyTrackEditors/VectorPropertyTrackEditor.cpp
+++ b/Engine/Source/Editor/MovieSceneTools/Private/TrackEditors/PropertyTrackEditors/VectorPropertyTrackEditor.cpp
@@ -16,7 +16,6 @@ FName FFloatVectorPropertyTrackEditor::YName( "Y" );
 FName FFloatVectorPropertyTrackEditor::ZName( "Z" );
 FName FFloatVectorPropertyTrackEditor::WName( "W" );
 
-
 TSharedRef<ISequencerTrackEditor> FFloatVectorPropertyTrackEditor::CreateTrackEditor( TSharedRef<ISequencer> InSequencer )
 {
 	return MakeShareable( new FFloatVectorPropertyTrackEditor( InSequencer ) );
@@ -40,7 +39,7 @@ void FFloatVectorPropertyTrackEditor::GenerateKeysFromPropertyChanged( const FPr
 
 	bool bIsVector4 = (StructName == NAME_Vector4f
 #if UE_LARGE_WORLD_COORDINATES_DISABLED
-		|| StructName == NAME_Vector
+		|| StructName == NAME_Vector4
 #endif
 		);
 
@@ -206,9 +205,31 @@ void FDoubleVectorPropertyTrackEditor::GenerateKeysFromPropertyChanged( const FP
 			|| StructName == NAME_Vector
 #endif
 			);
-	ensure(bIsVector);
 
-	FVector3d VectorValues = PropertyChangedParams.GetPropertyValue<FVector3d>();
+	bool bIsVector4 = (StructName == NAME_Vector4d
+#if !UE_LARGE_WORLD_COORDINATES_DISABLED
+		|| StructName == NAME_Vector4
+#endif
+		);
+
+	ensure(bIsVector || bIsVector4);
+
+	FVector4d VectorValues;
+	int32 Channels;
+	
+	if (bIsVector)
+	{
+		FVector3d Vector3DValue = PropertyChangedParams.GetPropertyValue<FVector3d>();
+		VectorValues.X = Vector3DValue.X;
+		VectorValues.Y = Vector3DValue.Y;
+		VectorValues.Z = Vector3DValue.Z;
+		Channels = 3;
+	}
+	else // if (bIsVector4)
+	{
+		VectorValues = PropertyChangedParams.GetPropertyValue<FVector4d>();
+		Channels = 4;
+	}
 
 	FPropertyPath StructPath = PropertyChangedParams.StructPathToKey;
 	FName ChannelName = StructPath.GetNumProperties() != 0 ? StructPath.GetLeafMostProperty().Property->GetFName() : NAME_None;
@@ -220,6 +241,11 @@ void FDoubleVectorPropertyTrackEditor::GenerateKeysFromPropertyChanged( const FP
 	OutGeneratedKeys.Add(FMovieSceneChannelValueSetter::Create<FMovieSceneDoubleChannel>(0, VectorValues.X, bKeyX));
 	OutGeneratedKeys.Add(FMovieSceneChannelValueSetter::Create<FMovieSceneDoubleChannel>(1, VectorValues.Y, bKeyY));
 	OutGeneratedKeys.Add(FMovieSceneChannelValueSetter::Create<FMovieSceneDoubleChannel>(2, VectorValues.Z, bKeyZ));
+
+	if ( Channels >= 4)
+	{
+		OutGeneratedKeys.Add(FMovieSceneChannelValueSetter::Create<FMovieSceneDoubleChannel>(3, VectorValues.W, ChannelName == NAME_None || ChannelName == WName));
+	}
 }
 
 void FDoubleVectorPropertyTrackEditor::InitializeNewTrack( UMovieSceneDoubleVectorTrack* NewTrack, FPropertyChangedParams PropertyChangedParams )
@@ -235,6 +261,14 @@ void FDoubleVectorPropertyTrackEditor::InitializeNewTrack( UMovieSceneDoubleVect
 			)
 	{
 		NewTrack->SetNumChannelsUsed( 3 );
+	}	
+	if ( StructName == NAME_Vector4d
+#if !UE_LARGE_WORLD_COORDINATES_DISABLED
+		|| StructName == NAME_Vector4
+#endif
+	)
+	{
+		NewTrack->SetNumChannelsUsed( 4 );
 	}
 }
 
@@ -270,6 +304,16 @@ bool FDoubleVectorPropertyTrackEditor::ModifyGeneratedKeysByCurrentAndWeight(UOb
 				GeneratedTotalKeys[2]->ModifyByCurrentAndWeight(Proxy, KeyTime, (void *)&Val.Z, Weight);
 			}
 			break;
+		case 4:
+			{
+				FVector4d Val(InterrogatedValues[0].X, InterrogatedValues[0].Y, InterrogatedValues[0].Z, InterrogatedValues[0].W);
+				FMovieSceneChannelProxy& Proxy = SectionToKey->GetChannelProxy();
+				GeneratedTotalKeys[0]->ModifyByCurrentAndWeight(Proxy, KeyTime, (void *)&Val.X, Weight);
+				GeneratedTotalKeys[1]->ModifyByCurrentAndWeight(Proxy, KeyTime, (void *)&Val.Y, Weight);
+				GeneratedTotalKeys[2]->ModifyByCurrentAndWeight(Proxy, KeyTime, (void *)&Val.Z, Weight);
+				GeneratedTotalKeys[3]->ModifyByCurrentAndWeight(Proxy, KeyTime, (void *)&Val.W, Weight);
+			}
+			break;
 		}
 
 		return true;
@@ -277,4 +321,3 @@ bool FDoubleVectorPropertyTrackEditor::ModifyGeneratedKeysByCurrentAndWeight(UOb
 
 	return false;
 }
-
diff --git a/Engine/Source/Editor/MovieSceneTools/Private/TrackEditors/PropertyTrackEditors/VectorPropertyTrackEditor.h b/Engine/Source/Editor/MovieSceneTools/Private/TrackEditors/PropertyTrackEditors/VectorPropertyTrackEditor.h
--- a/Engine/Source/Editor/MovieSceneTools/Private/TrackEditors/PropertyTrackEditors/VectorPropertyTrackEditor.h
+++ b/Engine/Source/Editor/MovieSceneTools/Private/TrackEditors/PropertyTrackEditors/VectorPropertyTrackEditor.h
@@ -38,9 +38,10 @@ public:
 		return TArray<FAnimatedPropertyKey, TInlineAllocator<3>>({
 #if UE_LARGE_WORLD_COORDINATES_DISABLED
 			FAnimatedPropertyKey::FromStructType(NAME_Vector),
+			FAnimatedPropertyKey::FromStructType(NAME_Vector4),
 #endif
 			FAnimatedPropertyKey::FromStructType(NAME_Vector3f),
-			FAnimatedPropertyKey::FromStructType(NAME_Vector4),
+			FAnimatedPropertyKey::FromStructType(NAME_Vector4f),
 			FAnimatedPropertyKey::FromStructType(NAME_Vector2D)
 		});
 	}
@@ -94,8 +95,10 @@ public:
 		return TArray<FAnimatedPropertyKey, TInlineAllocator<1>>({
 #if !UE_LARGE_WORLD_COORDINATES_DISABLED
 			FAnimatedPropertyKey::FromStructType(NAME_Vector),
+			FAnimatedPropertyKey::FromStructType(NAME_Vector4),
 #endif
-			FAnimatedPropertyKey::FromStructType(NAME_Vector3d)
+			FAnimatedPropertyKey::FromStructType(NAME_Vector3d),
+			FAnimatedPropertyKey::FromStructType(NAME_Vector4d)
 		});
 	}
 
diff --git a/Engine/Source/Runtime/MovieSceneTracks/Private/MovieSceneTracksComponentTypes.cpp b/Engine/Source/Runtime/MovieSceneTracks/Private/MovieSceneTracksComponentTypes.cpp
--- a/Engine/Source/Runtime/MovieSceneTracks/Private/MovieSceneTracksComponentTypes.cpp
+++ b/Engine/Source/Runtime/MovieSceneTracks/Private/MovieSceneTracksComponentTypes.cpp
@@ -129,11 +129,21 @@ void ConvertOperationalProperty(const FDoubleIntermediateVector& InVector, FVect
 	Out = FVector3d(InVector.X, InVector.Y, InVector.Z);
 }
 
+void ConvertOperationalProperty(const FDoubleIntermediateVector& InVector, FVector4d& Out)
+{
+	Out = FVector4d(InVector.X, InVector.Y, InVector.Z, InVector.W);
+}
+
 void ConvertOperationalProperty(const FVector3d& In, FDoubleIntermediateVector& Out)
 {
 	Out = FDoubleIntermediateVector(In.X, In.Y, In.Z);
 }
 
+void ConvertOperationalProperty(const FVector4d& In, FDoubleIntermediateVector& Out)
+{
+	Out = FDoubleIntermediateVector(In.X, In.Y, In.Z, In.W);
+}
+
 
 FIntermediate3DTransform GetComponentTransform(const UObject* Object)
 {
@@ -448,12 +458,15 @@ struct FDoubleVectorHandler : TPropertyComponentHandler<FDoubleVectorPropertyTra
 			FStructProperty* BoundProperty = CastField<FStructProperty>(FTrackInstancePropertyBindings::FindProperty(Object, Binding.PropertyPath.ToString()));
 			if (ensure(BoundProperty && BoundProperty->Struct))
 			{
-#if UE_LARGE_WORLD_COORDINATES_DISABLED
-				ensure(BoundProperty->Struct->GetFName() == NAME_Vector3d);
-#else
-				ensure(BoundProperty->Struct->GetFName() == NAME_Vector3d || BoundProperty->Struct->GetFName() == NAME_Vector);
-#endif
-				OutMetaData.NumChannels = 3;
+				if (BoundProperty->Struct->GetFName() == NAME_Vector3d || BoundProperty->Struct->GetFName() == NAME_Vector)
+				{
+					OutMetaData.NumChannels = 3;
+				}
+				else
+				{
+					ensure(BoundProperty->Struct == TBaseStructure<FVector4>::Get() || (BoundProperty->Struct->GetFName() == NAME_Vector4d) || (BoundProperty->Struct->GetFName() == NAME_Vector4));
+					OutMetaData.NumChannels = 4;
+				}
 			}
 			else
 			{
diff --git a/Engine/Source/Runtime/MovieSceneTracks/Private/Sections/MovieSceneVectorSection.cpp b/Engine/Source/Runtime/MovieSceneTracks/Private/Sections/MovieSceneVectorSection.cpp
--- a/Engine/Source/Runtime/MovieSceneTracks/Private/Sections/MovieSceneVectorSection.cpp
+++ b/Engine/Source/Runtime/MovieSceneTracks/Private/Sections/MovieSceneVectorSection.cpp
@@ -58,9 +58,17 @@ struct TVectorSectionEditorDataTraits<double>
 
 	static FCompositeValue GetPropertyValue(UObject& InObject, FTrackInstancePropertyBindings& Bindings, int32 NumChannels)
 	{
-		ensure(NumChannels == 3);
-		FVector3d Vector = Bindings.GetCurrentValue<FVector3d>(InObject);
-		return FCompositeValue(Vector.X, Vector.Y, Vector.Z, 0.f);
+		if (NumChannels == 3)
+		{
+			FVector3d Vector = Bindings.GetCurrentValue<FVector3d>(InObject);
+			return FCompositeValue(Vector.X, Vector.Y, Vector.Z, 0.f);
+		}
+		else
+		{
+			ensure(NumChannels == 4);
+			FVector4d Vector = Bindings.GetCurrentValue<FVector4d>(InObject);
+			return FCompositeValue(Vector.X, Vector.Y, Vector.Z, Vector.W);
+		}
 	}
 };
 
@@ -201,7 +209,7 @@ TSharedPtr<FStructOnScope> UMovieSceneFloatVectorSection::GetKeyStruct(TArrayVie
 	}
 	else if (ChannelsUsed == 4)
 	{
-		KeyStruct = MakeShareable(new FStructOnScope(FMovieSceneVector4KeyStruct::StaticStruct()));
+		KeyStruct = MakeShareable(new FStructOnScope(FMovieSceneVector4fKeyStruct::StaticStruct()));
 	}
 
 	if (KeyStruct.IsValid())
@@ -310,6 +318,10 @@ TSharedPtr<FStructOnScope> UMovieSceneDoubleVectorSection::GetKeyStruct(TArrayVi
 	{
 		KeyStruct = MakeShareable(new FStructOnScope(FMovieSceneVector3dKeyStruct::StaticStruct()));
 	}
+	else if (ChannelsUsed == 4)
+	{
+		KeyStruct = MakeShareable(new FStructOnScope(FMovieSceneVector4dKeyStruct::StaticStruct()));
+	}
 
 	if (KeyStruct.IsValid())
 	{
diff --git a/Engine/Source/Runtime/MovieSceneTracks/Public/MovieSceneTracksComponentTypes.h b/Engine/Source/Runtime/MovieSceneTracks/Public/MovieSceneTracksComponentTypes.h
--- a/Engine/Source/Runtime/MovieSceneTracks/Public/MovieSceneTracksComponentTypes.h
+++ b/Engine/Source/Runtime/MovieSceneTracks/Public/MovieSceneTracksComponentTypes.h
@@ -304,6 +304,7 @@ struct FDoubleVectorPropertyTraits
 		switch (MetaData.NumChannels)
 		{
 			case 3: TIndirectPropertyTraits<FVector3d, FDoubleIntermediateVector>::GetObjectPropertyValue(InObject, PropertyOffset, OutValue); return;
+			case 4: TIndirectPropertyTraits<FVector4d, FDoubleIntermediateVector>::GetObjectPropertyValue(InObject, PropertyOffset, OutValue); return;
 		}
 	}
 	static void GetObjectPropertyValue(const UObject* InObject, FVectorPropertyMetaData MetaData, FTrackInstancePropertyBindings* PropertyBindings, FDoubleIntermediateVector& OutValue)
@@ -311,6 +312,7 @@ struct FDoubleVectorPropertyTraits
 		switch (MetaData.NumChannels)
 		{
 			case 3: TIndirectPropertyTraits<FVector3d, FDoubleIntermediateVector>::GetObjectPropertyValue(InObject, PropertyBindings, OutValue); return;
+			case 4: TIndirectPropertyTraits<FVector4d, FDoubleIntermediateVector>::GetObjectPropertyValue(InObject, PropertyBindings, OutValue); return;
 		}
 	}
 	static void GetObjectPropertyValue(const UObject* InObject, FVectorPropertyMetaData MetaData, const FName& PropertyPath, StorageType& OutValue)
@@ -318,6 +320,7 @@ struct FDoubleVectorPropertyTraits
 		switch (MetaData.NumChannels)
 		{
 			case 3: TIndirectPropertyTraits<FVector3d, FDoubleIntermediateVector>::GetObjectPropertyValue(InObject, PropertyPath, OutValue); return;
+			case 4: TIndirectPropertyTraits<FVector4d, FDoubleIntermediateVector>::GetObjectPropertyValue(InObject, PropertyPath, OutValue); return;
 		}
 	}
 
@@ -331,6 +334,7 @@ struct FDoubleVectorPropertyTraits
 		switch (MetaData.NumChannels)
 		{
 			case 3: TIndirectPropertyTraits<FVector3d, FDoubleIntermediateVector>::SetObjectPropertyValue(InObject, PropertyOffset, InValue); return;
+			case 4: TIndirectPropertyTraits<FVector4d, FDoubleIntermediateVector>::SetObjectPropertyValue(InObject, PropertyOffset, InValue); return;
 		}
 
 		checkf(false, TEXT("Invalid number of channels"));
@@ -340,6 +344,7 @@ struct FDoubleVectorPropertyTraits
 		switch (MetaData.NumChannels)
 		{
 			case 3: TIndirectPropertyTraits<FVector3d, FDoubleIntermediateVector>::SetObjectPropertyValue(InObject, PropertyBindings, InValue); return;
+			case 4: TIndirectPropertyTraits<FVector4d, FDoubleIntermediateVector>::SetObjectPropertyValue(InObject, PropertyBindings, InValue); return;
 		}
 
 		checkf(false, TEXT("Invalid number of channels"));
diff --git a/Engine/Source/Runtime/MovieSceneTracks/Public/MovieSceneTracksPropertyTypes.h b/Engine/Source/Runtime/MovieSceneTracks/Public/MovieSceneTracksPropertyTypes.h
--- a/Engine/Source/Runtime/MovieSceneTracks/Public/MovieSceneTracksPropertyTypes.h
+++ b/Engine/Source/Runtime/MovieSceneTracks/Public/MovieSceneTracksPropertyTypes.h
@@ -210,7 +210,9 @@ MOVIESCENETRACKS_API void ConvertOperationalProperty(const FVector3f& In, FFloat
 MOVIESCENETRACKS_API void ConvertOperationalProperty(const FVector4& In, FFloatIntermediateVector& Out);
 
 MOVIESCENETRACKS_API void ConvertOperationalProperty(const FDoubleIntermediateVector& InVector, FVector3d& Out);
+MOVIESCENETRACKS_API void ConvertOperationalProperty(const FDoubleIntermediateVector& InVector, FVector4d& Out);
 MOVIESCENETRACKS_API void ConvertOperationalProperty(const FVector3d& In, FDoubleIntermediateVector& Out);
+MOVIESCENETRACKS_API void ConvertOperationalProperty(const FVector4d& In, FDoubleIntermediateVector& Out);
 
 } // namespace MovieScene
 } // namespace UE
diff --git a/Engine/Source/Runtime/MovieSceneTracks/Public/Sections/MovieSceneVectorSection.h b/Engine/Source/Runtime/MovieSceneTracks/Public/Sections/MovieSceneVectorSection.h
--- a/Engine/Source/Runtime/MovieSceneTracks/Public/Sections/MovieSceneVectorSection.h
+++ b/Engine/Source/Runtime/MovieSceneTracks/Public/Sections/MovieSceneVectorSection.h
@@ -76,10 +76,10 @@ struct FMovieSceneVector3fKeyStruct
 template<> struct TStructOpsTypeTraits<FMovieSceneVector3fKeyStruct> : public TStructOpsTypeTraitsBase2<FMovieSceneVector3fKeyStruct> { enum { WithCopy = false }; };
 
 /**
-* Proxy structure for vector4 section key data.
+* Proxy structure for vector4f section key data.
 */
 USTRUCT()
-struct FMovieSceneVector4KeyStruct
+struct FMovieSceneVector4fKeyStruct
 	: public FMovieSceneFloatVectorKeyStructBase
 {
 	GENERATED_BODY()
@@ -91,7 +91,7 @@ struct FMovieSceneVector4KeyStruct
 	//~ FMovieSceneFloatVectorKeyStructBase interface
 	virtual float* GetPropertyChannelByIndex(int32 Index) override { return &Vector[Index]; }
 };
-template<> struct TStructOpsTypeTraits<FMovieSceneVector4KeyStruct> : public TStructOpsTypeTraitsBase2<FMovieSceneVector4KeyStruct> { enum { WithCopy = false }; };
+template<> struct TStructOpsTypeTraits<FMovieSceneVector4fKeyStruct> : public TStructOpsTypeTraitsBase2<FMovieSceneVector4fKeyStruct> { enum { WithCopy = false }; };
 
 
 /**
@@ -135,6 +135,25 @@ struct FMovieSceneVector3dKeyStruct
 template<> struct TStructOpsTypeTraits<FMovieSceneVector3dKeyStruct> : public TStructOpsTypeTraitsBase2<FMovieSceneVector3dKeyStruct> { enum { WithCopy = false }; };
 
 
+/**
+* Proxy structure for double vector section key data.
+*/
+USTRUCT()
+struct FMovieSceneVector4dKeyStruct
+	: public FMovieSceneDoubleVectorKeyStructBase
+{
+	GENERATED_BODY()
+
+	/** They key's vector value. */
+	UPROPERTY(EditAnywhere, Category = Key)
+	FVector4d Vector = FVector4d(FVector3d::ZeroVector);
+
+	//~ FMovieSceneDoubleVectorKeyStructBase interface
+	virtual double* GetPropertyChannelByIndex(int32 Index) override { return &Vector[Index]; }
+};
+template<> struct TStructOpsTypeTraits<FMovieSceneVector4dKeyStruct> : public TStructOpsTypeTraitsBase2<FMovieSceneVector4dKeyStruct> { enum { WithCopy = false }; };
+
+
 /**
  * A float vector section.
  */
