Sequencer: CVar track icon

diff --git a/Engine/Source/Editor/MovieSceneTools/Private/TrackEditors/CVarTrackEditor.cpp b/Engine/Source/Editor/MovieSceneTools/Private/TrackEditors/CVarTrackEditor.cpp
--- a/Engine/Source/Editor/MovieSceneTools/Private/TrackEditors/CVarTrackEditor.cpp
+++ b/Engine/Source/Editor/MovieSceneTools/Private/TrackEditors/CVarTrackEditor.cpp
@@ -34,7 +34,7 @@ void FCVarTrackEditor::BuildAddTrackMenu(FMenuBuilder& MenuBuilder)
 	MenuBuilder.AddMenuEntry(
 		LOCTEXT("AddCVarTrack", "Console Variable Track"),
 		LOCTEXT("AddCVarTooltip", "Adds a new Console Variable track which will allow you to change console variables while a sequence plays."),
-		FSlateIcon(FEditorStyle::GetStyleSetName(), "Sequencer.Tracks.CVar"),
+		FSlateIcon(FEditorStyle::GetStyleSetName(), "DebugConsole.Icon"),
 		FUIAction(
 			FExecuteAction::CreateRaw(this, &FCVarTrackEditor::HandleAddCVarTrackMenuEntryExecute)
 		)
@@ -84,7 +84,7 @@ bool FCVarTrackEditor::SupportsType(TSubclassOf<UMovieSceneTrack> Type) const
 
 const FSlateBrush* FCVarTrackEditor::GetIconBrush() const
 {
-	return FEditorStyle::GetBrush("Sequencer.Tracks.CVar");
+	return FEditorStyle::GetBrush("DebugConsole.Icon");
 }
 // ~SequencerTrackEditor interface
 
