Sequencer: Set a relative object binding when there is no hierarchy.

This fixes an issue where if you create a binding when there is no hierarchy, the binding stays fixed and can't be used in a hierarchy afterwards.

diff --git a/Engine/Source/Editor/MovieSceneTools/Private/MovieSceneObjectBindingIDPicker.cpp b/Engine/Source/Editor/MovieSceneTools/Private/MovieSceneObjectBindingIDPicker.cpp
--- a/Engine/Source/Editor/MovieSceneTools/Private/MovieSceneObjectBindingIDPicker.cpp
+++ b/Engine/Source/Editor/MovieSceneTools/Private/MovieSceneObjectBindingIDPicker.cpp
@@ -294,7 +294,7 @@ void FMovieSceneObjectBindingIDPicker::SetCurrentValueFromFixed(UE::MovieScene::
 	// If we don't know the local sequence ID, or we have no hierarchy, or we're resetting the binding; just set the ID directly
 	if (LocalSequenceID == MovieSceneSequenceID::Invalid || !InValue.Guid.IsValid() || Hierarchy == nullptr)
 	{
-		SetCurrentValue(InValue);
+		SetCurrentValue(UE::MovieScene::FRelativeObjectBindingID(InValue.Guid));
 	}
 	else
 	{
