Sequencer: Follow up to keying with transform origin. There's no need to account for the transform origin when there is a parent because the transform origin doesn't apply to child components (it's already being applied to the parent)

diff --git a/Engine/Source/Editor/MovieSceneTools/Private/TrackEditors/TransformTrackEditor.cpp b/Engine/Source/Editor/MovieSceneTools/Private/TrackEditors/TransformTrackEditor.cpp
--- a/Engine/Source/Editor/MovieSceneTools/Private/TrackEditors/TransformTrackEditor.cpp
+++ b/Engine/Source/Editor/MovieSceneTools/Private/TrackEditors/TransformTrackEditor.cpp
@@ -895,16 +895,23 @@ FTransformData F3DTransformTrackEditor::RecomposeTransform(const FTransformData&
 		UMovieScenePropertyInstantiatorSystem* System = EntityLinker->FindSystem<UMovieScenePropertyInstantiatorSystem>();
 		if (System)
 		{
+			USceneComponent* SceneComponent = MovieSceneHelpers::SceneComponentFromRuntimeObject(AnimatedObject);
+
 			FDecompositionQuery Query;
 			Query.Entities = MakeArrayView(&EntityID, 1);
-			Query.Object   = MovieSceneHelpers::SceneComponentFromRuntimeObject(AnimatedObject);
+			Query.Object   = SceneComponent;
 
 			FIntermediate3DTransform CurrentValue(InTransformData.Translation, InTransformData.Rotation, InTransformData.Scale);
 
 			TRecompositionResult<FIntermediate3DTransform> TransformData = System->RecomposeBlendOperational(FMovieSceneTracksComponentTypes::Get()->ComponentTransform, Query, CurrentValue);
 
 			FTransform CurrentTransform(TransformData.Values[0].GetRotation(), TransformData.Values[0].GetTranslation(), TransformData.Values[0].GetScale());
-			CurrentTransform *= GetTransformOrigin().Inverse();
+
+			// Account for the transform origin only if this is not parented because the transform origin is already being applied to the parent.
+			if (!SceneComponent->GetAttachParent())
+			{
+				CurrentTransform *= GetTransformOrigin().Inverse();
+			}
 
 			return FTransformData(CurrentTransform.GetLocation(), CurrentTransform.GetRotation().Rotator(), CurrentTransform.GetScale3D());
 		}
@@ -1000,7 +1007,12 @@ void F3DTransformTrackEditor::ProcessKeyOperation(UObject* ObjectToKey, TArrayVi
 		Query.Object   = Component;
 
 		FTransform CurrentTransform(Component->GetRelativeRotation(), Component->GetRelativeLocation(), Component->GetRelativeScale3D());
-		CurrentTransform *= GetTransformOrigin().Inverse();
+
+		// Account for the transform origin only if this is not parented because the transform origin is already being applied to the parent.
+		if (!Component->GetAttachParent())
+		{
+			CurrentTransform *= GetTransformOrigin().Inverse();
+		}
 
 		FIntermediate3DTransform CurrentValue(CurrentTransform.GetTranslation(), CurrentTransform.GetRotation().Rotator(), CurrentTransform.GetScale3D());
 		TRecompositionResult<FIntermediate3DTransform> TransformData = System->RecomposeBlendOperational(FMovieSceneTracksComponentTypes::Get()->ComponentTransform, Query, CurrentValue);
