MovieSceneTools: use the authored subframe value when present in skeletal animations as the suffix of their section frame time hints

diff --git a/Engine/Source/Editor/MovieSceneTools/Private/TrackEditors/SkeletalAnimationTrackEditor.cpp b/Engine/Source/Editor/MovieSceneTools/Private/TrackEditors/SkeletalAnimationTrackEditor.cpp
--- a/Engine/Source/Editor/MovieSceneTools/Private/TrackEditors/SkeletalAnimationTrackEditor.cpp
+++ b/Engine/Source/Editor/MovieSceneTools/Private/TrackEditors/SkeletalAnimationTrackEditor.cpp
@@ -661,7 +661,20 @@ int32 FSkeletalAnimationSection::OnPaintSection( FSequencerSectionPainter& Paint
 			const TAttribute<FFrameRate> TickResolutionAttr(HintFrameTime.Rate);
 			const TAttribute<FFrameRate> DisplayRateAttr(HintFrameTime.Rate);
 
-			const FFrameNumberInterface FrameNumberInterface(DisplayFormatAttr, ZeroPadFrameNumbersAttr, TickResolutionAttr, DisplayRateAttr);
+			FFrameNumberInterface FrameNumberInterface(DisplayFormatAttr, ZeroPadFrameNumbersAttr, TickResolutionAttr, DisplayRateAttr);
+
+			float Subframe = 0.0f;
+			if (UAnimationBlueprintLibrary::EvaluateRootBoneTimecodeSubframeAttributeAtTime(Section.Params.Animation, static_cast<float>(AnimTime), Subframe))
+			{
+				if (FMath::IsNearlyEqual(Subframe, FMath::RoundToFloat(Subframe)))
+				{
+					FrameNumberInterface.SetSubframeIndicator(FString::Printf(TEXT(" (%d)"), FMath::RoundToInt(Subframe)));
+				}
+				else
+				{
+					FrameNumberInterface.SetSubframeIndicator(FString::Printf(TEXT(" (%s)"), *LexToSanitizedString(Subframe)));
+				}
+			}
 
 			DrawFrameTimeHint(Painter, CurrentTime, HintFrameTime.Time, &FrameNumberInterface);
 		}
