Sequencer: When a color is picked for an existing keyframe, use the properties of the existing keyframe and just update the value. This fixes a bug where picking a color from the color chooser would replace an existing keyframe with default properties (ie. auto interpolation).

diff --git a/Engine/Source/Editor/MovieSceneTools/Private/Sections/MovieSceneSectionHelpers.cpp b/Engine/Source/Editor/MovieSceneTools/Private/Sections/MovieSceneSectionHelpers.cpp
--- a/Engine/Source/Editor/MovieSceneTools/Private/Sections/MovieSceneSectionHelpers.cpp
+++ b/Engine/Source/Editor/MovieSceneTools/Private/Sections/MovieSceneSectionHelpers.cpp
@@ -147,12 +147,25 @@ FMovieSceneKeyColorPicker::FMovieSceneKeyColorPicker(UMovieSceneSection* Section
 	OpenColorPicker(PickerArgs);
 }
 
+void UpdateOrAddKey(FMovieSceneFloatChannel* Channel, FFrameNumber KeyTime, float Value)
+{
+	int32 KeyHandleIndex = Channel->GetData().FindKey(KeyTime);
+	if (KeyHandleIndex == INDEX_NONE)
+	{
+		Channel->GetData().UpdateOrAddKey(KeyTime, FMovieSceneFloatValue(Value));
+	}
+	else
+	{
+		Channel->GetData().GetValues()[KeyHandleIndex].Value = Value;
+	}
+}
+
 void FMovieSceneKeyColorPicker::OnColorPickerPicked(FLinearColor NewColor, FMovieSceneFloatChannel* RChannel, FMovieSceneFloatChannel* GChannel, FMovieSceneFloatChannel* BChannel, FMovieSceneFloatChannel* AChannel)
 {
-	RChannel->GetData().UpdateOrAddKey(KeyTime, FMovieSceneFloatValue(NewColor.R));
-	GChannel->GetData().UpdateOrAddKey(KeyTime, FMovieSceneFloatValue(NewColor.G));
-	BChannel->GetData().UpdateOrAddKey(KeyTime, FMovieSceneFloatValue(NewColor.B));
-	AChannel->GetData().UpdateOrAddKey(KeyTime, FMovieSceneFloatValue(NewColor.A));
+	UpdateOrAddKey(RChannel, KeyTime, NewColor.R);
+	UpdateOrAddKey(GChannel, KeyTime, NewColor.G);
+	UpdateOrAddKey(BChannel, KeyTime, NewColor.B);
+	UpdateOrAddKey(AChannel, KeyTime, NewColor.A);
 }
 
 void FMovieSceneKeyColorPicker::OnColorPickerClosed(const TSharedRef<SWindow>& Window, UMovieSceneSection* Section, FMovieSceneFloatChannel* RChannel, FMovieSceneFloatChannel* GChannel, FMovieSceneFloatChannel* BChannel, FMovieSceneFloatChannel* AChannel)
@@ -170,17 +183,17 @@ void FMovieSceneKeyColorPicker::OnColorPickerClosed(const TSharedRef<SWindow>& W
 		BChannel->Evaluate(KeyTime, B);
 		AChannel->Evaluate(KeyTime, A);
 
-		RChannel->GetData().UpdateOrAddKey(KeyTime, FMovieSceneFloatValue(InitialColor.R));
-		GChannel->GetData().UpdateOrAddKey(KeyTime, FMovieSceneFloatValue(InitialColor.G));
-		BChannel->GetData().UpdateOrAddKey(KeyTime, FMovieSceneFloatValue(InitialColor.B));
-		AChannel->GetData().UpdateOrAddKey(KeyTime, FMovieSceneFloatValue(InitialColor.A));
+		UpdateOrAddKey(RChannel, KeyTime, InitialColor.R);
+		UpdateOrAddKey(GChannel, KeyTime, InitialColor.G);
+		UpdateOrAddKey(BChannel, KeyTime, InitialColor.B);
+		UpdateOrAddKey(AChannel, KeyTime, InitialColor.A);
 
 		Section->Modify();
 
-		RChannel->GetData().UpdateOrAddKey(KeyTime, FMovieSceneFloatValue(R));
-		GChannel->GetData().UpdateOrAddKey(KeyTime, FMovieSceneFloatValue(G));
-		BChannel->GetData().UpdateOrAddKey(KeyTime, FMovieSceneFloatValue(B));
-		AChannel->GetData().UpdateOrAddKey(KeyTime, FMovieSceneFloatValue(A));
+		UpdateOrAddKey(RChannel, KeyTime, R);
+		UpdateOrAddKey(GChannel, KeyTime, G);
+		UpdateOrAddKey(BChannel, KeyTime, B);
+		UpdateOrAddKey(AChannel, KeyTime, A);
 	}
 }
 
@@ -189,8 +202,8 @@ void FMovieSceneKeyColorPicker::OnColorPickerCancelled(FLinearColor NewColor, FM
 	bColorPickerWasCancelled = true;
 	
 	// Restore the original color. No transaction will be created when the OnColorPickerClosed callback is called.
-	RChannel->GetData().UpdateOrAddKey(KeyTime, FMovieSceneFloatValue(InitialColor.R));
-	GChannel->GetData().UpdateOrAddKey(KeyTime, FMovieSceneFloatValue(InitialColor.G));
-	BChannel->GetData().UpdateOrAddKey(KeyTime, FMovieSceneFloatValue(InitialColor.B));
-	AChannel->GetData().UpdateOrAddKey(KeyTime, FMovieSceneFloatValue(InitialColor.A));
+	UpdateOrAddKey(RChannel, KeyTime, InitialColor.R);
+	UpdateOrAddKey(GChannel, KeyTime, InitialColor.G);
+	UpdateOrAddKey(BChannel, KeyTime, InitialColor.B);
+	UpdateOrAddKey(AChannel, KeyTime, InitialColor.A);
 }
