Animation Mode: Can hit ensure if actor is not in Sequencer when using Snapping tool

diff --git a/Engine/Source/Editor/MovieSceneTools/Private/MovieSceneToolHelpers.cpp b/Engine/Source/Editor/MovieSceneTools/Private/MovieSceneToolHelpers.cpp
--- a/Engine/Source/Editor/MovieSceneTools/Private/MovieSceneToolHelpers.cpp
+++ b/Engine/Source/Editor/MovieSceneTools/Private/MovieSceneToolHelpers.cpp
@@ -4456,7 +4456,7 @@ static void GetNonSequencerActorWorldTransforms(IMovieScenePlayer* Player, UMovi
 		do
 		{
 			FGuid ActorHandle = GetHandleToObject(Actor, InSequence, Player, Template,false);
-			if (Frames.Num() > 0)
+			if (ActorHandle.IsValid() && Frames.Num() > 0)
 			{
 				if (InSequence->GetMovieScene()->FindTrack<UMovieScene3DTransformTrack>(ActorHandle))
 				{
