Sequencer: Activate the default mode in case FEditorModeTools::Tick isn't run before here.

This fixes a bug where the mode is set to "No Active Mode" when reopening the editor because the Sequencer modes are added before the default mode.  This can be removed once a general fix for UE-143791 has been implemented.

diff --git a/Engine/Source/Editor/MovieSceneTools/Private/TrackEditors/SkeletalAnimationTrackEditor.cpp b/Engine/Source/Editor/MovieSceneTools/Private/TrackEditors/SkeletalAnimationTrackEditor.cpp
--- a/Engine/Source/Editor/MovieSceneTools/Private/TrackEditors/SkeletalAnimationTrackEditor.cpp
+++ b/Engine/Source/Editor/MovieSceneTools/Private/TrackEditors/SkeletalAnimationTrackEditor.cpp
@@ -1096,6 +1096,10 @@ void FSkeletalAnimationTrackEditor::OnInitialize()
 {
 	++FSkeletalAnimationTrackEditor::NumberActive;
 
+	// Activate the default mode in case FEditorModeTools::Tick isn't run before here. 
+	// This can be removed once a general fix for UE-143791 has been implemented.
+	GLevelEditorModeTools().ActivateDefaultMode();
+
 	GLevelEditorModeTools().ActivateMode(FSkeletalAnimationTrackEditMode::ModeName);
 	FSkeletalAnimationTrackEditMode* EditMode = static_cast<FSkeletalAnimationTrackEditMode*>(GLevelEditorModeTools().GetActiveMode(FSkeletalAnimationTrackEditMode::ModeName));
 	if (EditMode)
diff --git a/Engine/Source/Editor/Sequencer/Private/LevelEditorSequencerIntegration.cpp b/Engine/Source/Editor/Sequencer/Private/LevelEditorSequencerIntegration.cpp
--- a/Engine/Source/Editor/Sequencer/Private/LevelEditorSequencerIntegration.cpp
+++ b/Engine/Source/Editor/Sequencer/Private/LevelEditorSequencerIntegration.cpp
@@ -569,6 +569,10 @@ void FLevelEditorSequencerIntegration::ActivateSequencerEditorMode()
 	FName ResourceName("SequencerMode");
 	AcquiredResources.Release(ResourceName);
 
+	// Activate the default mode in case FEditorModeTools::Tick isn't run before here. 
+	// This can be removed once a general fix for UE-143791 has been implemented.
+	GLevelEditorModeTools().ActivateDefaultMode();
+
 	FEditorModeID ModeID = TEXT("SequencerToolsEditMode");
 	GLevelEditorModeTools().ActivateMode(ModeID);
 
