Fix a filtering issue in sequencer, where it didn't show anim sequences of compatible skeletons.

diff --git a/Engine/Source/Editor/MovieSceneTools/Private/TrackEditors/SkeletalAnimationTrackEditor.cpp b/Engine/Source/Editor/MovieSceneTools/Private/TrackEditors/SkeletalAnimationTrackEditor.cpp
--- a/Engine/Source/Editor/MovieSceneTools/Private/TrackEditors/SkeletalAnimationTrackEditor.cpp
+++ b/Engine/Source/Editor/MovieSceneTools/Private/TrackEditors/SkeletalAnimationTrackEditor.cpp
@@ -1692,7 +1692,7 @@ void FSkeletalAnimationTrackEditor::AddAnimationSubMenu(FMenuBuilder& MenuBuilde
 		AssetPickerConfig.OnShouldFilterAsset = FOnShouldFilterAsset::CreateRaw(this, &FSkeletalAnimationTrackEditor::ShouldFilterAsset);
 		AssetPickerConfig.Filter.bRecursiveClasses = true;
 		AssetPickerConfig.Filter.ClassNames.Add(UAnimSequenceBase::StaticClass()->GetFName());
-		AssetPickerConfig.Filter.TagsAndValues.Add(TEXT("Skeleton"), FAssetData(Skeleton).GetExportTextName());
+		AssetPickerConfig.OnShouldFilterAsset.BindRaw(this, &FSkeletalAnimationTrackEditor::FilterAnimSequences, Skeleton);
 		AssetPickerConfig.SaveSettingsName = TEXT("SequencerAssetPicker");
 	}
 
@@ -1708,6 +1708,15 @@ void FSkeletalAnimationTrackEditor::AddAnimationSubMenu(FMenuBuilder& MenuBuilde
 	MenuBuilder.AddWidget(MenuEntry.ToSharedRef(), FText::GetEmpty(), true);
 }
 
+bool FSkeletalAnimationTrackEditor::FilterAnimSequences(const FAssetData& AssetData, USkeleton* Skeleton)
+{
+	if (Skeleton && Skeleton->IsCompatibleSkeletonByAssetData(AssetData))
+	{
+		return false;
+	}
+
+	return true;
+}
 
 void FSkeletalAnimationTrackEditor::OnAnimationAssetSelected(const FAssetData& AssetData, TArray<FGuid> ObjectBindings, UMovieSceneTrack* Track)
 {
diff --git a/Engine/Source/Editor/MovieSceneTools/Private/TrackEditors/SkeletalAnimationTrackEditor.h b/Engine/Source/Editor/MovieSceneTools/Private/TrackEditors/SkeletalAnimationTrackEditor.h
--- a/Engine/Source/Editor/MovieSceneTools/Private/TrackEditors/SkeletalAnimationTrackEditor.h
+++ b/Engine/Source/Editor/MovieSceneTools/Private/TrackEditors/SkeletalAnimationTrackEditor.h
@@ -78,6 +78,9 @@ private:
 	TSharedRef<SWidget> BuildAnimationSubMenu(FGuid ObjectBinding, USkeleton* Skeleton, UMovieSceneTrack* Track);
 	void AddAnimationSubMenu(FMenuBuilder& MenuBuilder, TArray<FGuid> ObjectBindings, USkeleton* Skeleton, UMovieSceneTrack* Track);
 
+	/** Filter only compatible skeletons */
+	bool FilterAnimSequences(const FAssetData& AssetData, USkeleton* Skeleton);
+
 	/** Animation sub menu filter function */
 	bool ShouldFilterAsset(const FAssetData& AssetData);
 
@@ -93,7 +96,7 @@ private:
 	/** Construct the binding menu*/
 	void ConstructObjectBindingTrackMenu(FMenuBuilder& MenuBuilder, TArray<FGuid> ObjectBindings);
 
-	/** Callback to Create the Animation Asset, pop open the dialolg*/
+	/** Callback to Create the Animation Asset, pop open the dialog */
 	void HandleCreateAnimationSequence(USkeletalMeshComponent* SkelMeshComp, USkeleton* Skeleton, FGuid Binding, bool bCeateSoftLink);
 
 	/** Callback to Creae the Animation Asset after getting the name*/
