Sequencer: Change to AssetPickerConfig to List for better layout/sizing

diff --git a/Engine/Source/Editor/MovieSceneTools/Private/TrackEditors/CinematicShotTrackEditor.cpp b/Engine/Source/Editor/MovieSceneTools/Private/TrackEditors/CinematicShotTrackEditor.cpp
--- a/Engine/Source/Editor/MovieSceneTools/Private/TrackEditors/CinematicShotTrackEditor.cpp
+++ b/Engine/Source/Editor/MovieSceneTools/Private/TrackEditors/CinematicShotTrackEditor.cpp
@@ -639,7 +639,7 @@ TSharedRef<SWidget> FCinematicShotTrackEditor::HandleAddCinematicShotComboButton
 		AssetPickerConfig.OnAssetSelected = FOnAssetSelected::CreateRaw( this, &FCinematicShotTrackEditor::HandleAddCinematicShotComboButtonMenuEntryExecute);
 		AssetPickerConfig.OnAssetEnterPressed = FOnAssetEnterPressed::CreateRaw( this, &FCinematicShotTrackEditor::HandleAddCinematicShotComboButtonMenuEntryEnterPressed);
 		AssetPickerConfig.bAllowNullSelection = false;
-		AssetPickerConfig.InitialAssetViewType = EAssetViewType::Tile;
+		AssetPickerConfig.InitialAssetViewType = EAssetViewType::List;
 		AssetPickerConfig.Filter.ClassNames.Add(TEXT("LevelSequence"));
 		AssetPickerConfig.SaveSettingsName = TEXT("SequencerAssetPicker");
 	}
diff --git a/Engine/Source/Editor/MovieSceneTools/Private/TrackEditors/SubTrackEditor.cpp b/Engine/Source/Editor/MovieSceneTools/Private/TrackEditors/SubTrackEditor.cpp
--- a/Engine/Source/Editor/MovieSceneTools/Private/TrackEditors/SubTrackEditor.cpp
+++ b/Engine/Source/Editor/MovieSceneTools/Private/TrackEditors/SubTrackEditor.cpp
@@ -458,7 +458,7 @@ TSharedRef<SWidget> FSubTrackEditor::HandleAddSubSequenceComboButtonGetMenuConte
 			AssetPickerConfig.OnAssetSelected = FOnAssetSelected::CreateRaw( this, &FSubTrackEditor::HandleAddSubSequenceComboButtonMenuEntryExecute, InTrack);
 			AssetPickerConfig.OnAssetEnterPressed = FOnAssetEnterPressed::CreateRaw( this, &FSubTrackEditor::HandleAddSubSequenceComboButtonMenuEntryEnterPressed, InTrack);
 			AssetPickerConfig.bAllowNullSelection = false;
-			AssetPickerConfig.InitialAssetViewType = EAssetViewType::Tile;
+			AssetPickerConfig.InitialAssetViewType = EAssetViewType::List;
 			AssetPickerConfig.Filter.ClassNames.Add(TEXT("LevelSequence"));
 			AssetPickerConfig.SaveSettingsName = TEXT("SequencerAssetPicker");
 		}
