fixed visual artifacts for space bars / keys

diff --git a/Engine/Plugins/Animation/ControlRig/Source/ControlRigEditor/Private/ControlRigSpaceChannelCurveModel.cpp b/Engine/Plugins/Animation/ControlRig/Source/ControlRigEditor/Private/ControlRigSpaceChannelCurveModel.cpp
--- a/Engine/Plugins/Animation/ControlRig/Source/ControlRigEditor/Private/ControlRigSpaceChannelCurveModel.cpp
+++ b/Engine/Plugins/Animation/ControlRig/Source/ControlRigEditor/Private/ControlRigSpaceChannelCurveModel.cpp
@@ -168,7 +168,8 @@ void FControlRigSpaceChannelCurveModel::GetKeys(const FCurveEditor& CurveEditor,
 void FControlRigSpaceChannelCurveModel::GetKeyDrawInfo(ECurvePointType PointType, const FKeyHandle InKeyHandle, FKeyDrawInfo& OutDrawInfo) const
 {
 	OutDrawInfo.Brush = FEditorStyle::Get().GetBrush("FilledBorder");
-	OutDrawInfo.ScreenSize = FVector2D(10, 100);
+	// NOTE Y is set to SCurveEditorKeyBarView::TrackHeight
+	OutDrawInfo.ScreenSize = FVector2D(10, 24);
 }
 
 void FControlRigSpaceChannelCurveModel::GetKeyPositions(TArrayView<const FKeyHandle> InKeys, TArrayView<FKeyPosition> OutKeyPositions) const
diff --git a/Engine/Source/Editor/MovieSceneTools/Private/Channels/SCurveEditorKeyBarView.cpp b/Engine/Source/Editor/MovieSceneTools/Private/Channels/SCurveEditorKeyBarView.cpp
--- a/Engine/Source/Editor/MovieSceneTools/Private/Channels/SCurveEditorKeyBarView.cpp
+++ b/Engine/Source/Editor/MovieSceneTools/Private/Channels/SCurveEditorKeyBarView.cpp
@@ -90,7 +90,10 @@ void SCurveEditorKeyBarView::DrawLabels(const FGeometry& AllottedGeometry, FSlat
 	// We have to measure the string so we can draw it centered on the window.
 	const TSharedRef<FSlateFontMeasure> FontMeasure = FSlateApplication::Get().GetRenderer()->GetFontMeasureService();
 
-	const FVector2D LocalSize = AllottedGeometry.GetLocalSize();
+	const FVector2D LocalSize(AllottedGeometry.GetLocalSize().X, TrackHeight);
+
+	double InputMin = 0, InputMax = 1;
+	GetInputBounds(InputMin, InputMax);//in Sequencer Seconds
 
 	const FCurveEditorScreenSpace ViewSpace = GetViewSpace();
 
@@ -104,14 +107,9 @@ void SCurveEditorKeyBarView::DrawLabels(const FGeometry& AllottedGeometry, FSlat
 
 		const int32 CurveIndex = static_cast<int32>(It->Value.ViewToCurveTransform.GetTranslation().Y);
 
-		FCurveEditorScreenSpace CurveSpace = GetCurveSpace(It.Key());
-
+		const FCurveEditorScreenSpace CurveSpace = GetCurveSpace(It.Key());
 		const float LaneTop = CurveSpace.ValueToScreen(0.0) - TrackHeight*.5f;
 
-
-		double InputMin = 0, InputMax = 1;
-		GetInputBounds(InputMin, InputMax);//in Sequencer Seconds
-
 		// Draw the curve label and background
 		FKeyBarCurveModel* KeyBarCurveModel = static_cast<FKeyBarCurveModel*>(Curve);
 		if (KeyBarCurveModel)
@@ -145,18 +143,20 @@ void SCurveEditorKeyBarView::DrawLabels(const FGeometry& AllottedGeometry, FSlat
 					}
 				}
 
+				// draw background box
 				if (CurveColor != FLinearColor::White)
 				{
-					const double LowerSecondsForBox =(Index == 0 && LowerSeconds > InputMin )? InputMin : LowerSeconds;
-					const double BoxPos = CurveSpace.SecondsToScreen(LowerSecondsForBox);
-
+					const double LowerSecondsForBox = Index == 0 ? InputMin : FMath::Max<double>(InputMin, LowerSeconds);
+					const double BoxPos = CurveSpace.SecondsToScreen(LowerSecondsForBox);	
 					const FPaintGeometry BoxGeometry = AllottedGeometry.ToPaintGeometry(
-						FVector2D(AllottedGeometry.GetLocalSize().X, TrackHeight),
+						LocalSize,
 						FSlateLayoutTransform(FVector2D(BoxPos, LaneTop))
 					);
-
+				
 					FSlateDrawElement::MakeBox(OutDrawElements, BaseLayerId + CurveViewConstants::ELayerOffset::Background, BoxGeometry, WhiteBrush, DrawEffects, CurveColor);
 				}
+
+				// draw label
 				const double LowerSecondsForLabel = (InputMin > LowerSeconds) ? InputMin : LowerSeconds;
 				double LabelPos = CurveSpace.SecondsToScreen(LowerSecondsForLabel) + 10;
 
