Sequencer: Drag and drop should gather all skeletal mesh components (not just the root if it's a skeletal mesh component) so it can perform a animation skeleton match with all the possible skeletal mesh components.

follow up to CL 18963239

diff --git a/Engine/Source/Editor/MovieSceneTools/Private/TrackEditors/SkeletalAnimationTrackEditor.cpp b/Engine/Source/Editor/MovieSceneTools/Private/TrackEditors/SkeletalAnimationTrackEditor.cpp
--- a/Engine/Source/Editor/MovieSceneTools/Private/TrackEditors/SkeletalAnimationTrackEditor.cpp
+++ b/Engine/Source/Editor/MovieSceneTools/Private/TrackEditors/SkeletalAnimationTrackEditor.cpp
@@ -336,7 +336,11 @@ USkeleton* GetSkeletonFromComponent(UActorComponent* InComponent)
 	return nullptr;
 }
 
-TArray<USkeletalMeshComponent*> AcquireSkeletalMeshComponentsFromObjectGuid(const FGuid& Guid, TSharedPtr<ISequencer> SequencerPtr)
+// Get the skeletal mesh components from the guid
+// If bGetSingleRootComponent - return only the root component if it is a skeletal mesh component. 
+// This allows the root object binding to have an animation track without needing a skeletal mesh component binding
+//
+TArray<USkeletalMeshComponent*> AcquireSkeletalMeshComponentsFromObjectGuid(const FGuid& Guid, TSharedPtr<ISequencer> SequencerPtr, const bool bGetSingleRootComponent = true)
 {
 	TArray<USkeletalMeshComponent*> SkeletalMeshComponents;
 
@@ -354,10 +358,13 @@ TArray<USkeletalMeshComponent*> AcquireSkeletalMeshComponentsFromObjectGuid(cons
 
 	if (Actor)
 	{
-		if (USkeletalMeshComponent* SkeletalMeshComponent = Cast<USkeletalMeshComponent>(Actor->GetRootComponent()))
+		if (bGetSingleRootComponent)
 		{
-			SkeletalMeshComponents.Add(SkeletalMeshComponent);
-			return SkeletalMeshComponents;
+			if (USkeletalMeshComponent* SkeletalMeshComponent = Cast<USkeletalMeshComponent>(Actor->GetRootComponent()))
+			{
+				SkeletalMeshComponents.Add(SkeletalMeshComponent);
+				return SkeletalMeshComponents;
+			}
 		}
 
 		Actor->GetComponents(SkeletalMeshComponents);
@@ -369,10 +376,13 @@ TArray<USkeletalMeshComponent*> AcquireSkeletalMeshComponentsFromObjectGuid(cons
 		AActor* ActorCDO = Cast<AActor>(Actor->GetClass()->GetDefaultObject());
 		if (ActorCDO)
 		{
-			if (USkeletalMeshComponent* SkeletalMeshComponent = Cast<USkeletalMeshComponent>(ActorCDO->GetRootComponent()))
+			if (bGetSingleRootComponent)
 			{
-				SkeletalMeshComponents.Add(SkeletalMeshComponent);
-				return SkeletalMeshComponents;
+				if (USkeletalMeshComponent* SkeletalMeshComponent = Cast<USkeletalMeshComponent>(ActorCDO->GetRootComponent()))
+				{
+					SkeletalMeshComponents.Add(SkeletalMeshComponent);
+					return SkeletalMeshComponents;
+				}
 			}
 
 			ActorCDO->GetComponents(SkeletalMeshComponents);
@@ -1902,7 +1912,7 @@ bool FSkeletalAnimationTrackEditor::OnAllowDrop(const FDragDropEvent& DragDropEv
 		return false;
 	}
 
-	TArray<USkeletalMeshComponent*> SkeletalMeshComponents = AcquireSkeletalMeshComponentsFromObjectGuid(DragDropParams.TargetObjectGuid, GetSequencer());
+	TArray<USkeletalMeshComponent*> SkeletalMeshComponents = AcquireSkeletalMeshComponentsFromObjectGuid(DragDropParams.TargetObjectGuid, GetSequencer(), false);
 
 	TSharedPtr<FAssetDragDropOp> DragDropOp = StaticCastSharedPtr<FAssetDragDropOp>( Operation );
 
@@ -1943,7 +1953,7 @@ FReply FSkeletalAnimationTrackEditor::OnDrop(const FDragDropEvent& DragDropEvent
 		return FReply::Unhandled();
 	}
 
-	TArray<USkeletalMeshComponent*> SkeletalMeshComponents = AcquireSkeletalMeshComponentsFromObjectGuid(DragDropParams.TargetObjectGuid, GetSequencer());
+	TArray<USkeletalMeshComponent*> SkeletalMeshComponents = AcquireSkeletalMeshComponentsFromObjectGuid(DragDropParams.TargetObjectGuid, GetSequencer(), false);
 
 	const FScopedTransaction Transaction(LOCTEXT("DropAssets", "Drop Assets"));
 
