Sequencer: Set the ID directly if there is no sequencer. Otherwise, if we don't know the local sequence ID, or we have no hierarchy, or we're resetting the binding; set a relative ID.

This fixes a regression introduced by 18938437 which fixed an issue where if you create a binding when there is no hierarchy, the binding stays fixed and can't be used in a hierarchy afterwards.

The fix here is that sequencer doesn't exist (when setting the value through the level blueprint picker), so it should just set the ID directly.

diff --git a/Engine/Source/Editor/MovieSceneTools/Private/MovieSceneObjectBindingIDPicker.cpp b/Engine/Source/Editor/MovieSceneTools/Private/MovieSceneObjectBindingIDPicker.cpp
--- a/Engine/Source/Editor/MovieSceneTools/Private/MovieSceneObjectBindingIDPicker.cpp
+++ b/Engine/Source/Editor/MovieSceneTools/Private/MovieSceneObjectBindingIDPicker.cpp
@@ -291,8 +291,13 @@ void FMovieSceneObjectBindingIDPicker::SetCurrentValueFromFixed(UE::MovieScene::
 	TSharedPtr<ISequencer>              Sequencer = WeakSequencer.Pin();
 	const FMovieSceneSequenceHierarchy* Hierarchy = Sequencer.IsValid() ? Sequencer->GetEvaluationTemplate().GetHierarchy() : nullptr;
 
-	// If we don't know the local sequence ID, or we have no hierarchy, or we're resetting the binding; just set the ID directly
-	if (LocalSequenceID == MovieSceneSequenceID::Invalid || !InValue.Guid.IsValid() || Hierarchy == nullptr)
+	// If there is no sequencer, just set the ID directly 
+	if (!Sequencer.IsValid())
+	{
+		SetCurrentValue(InValue);
+	}
+	// If we don't know the local sequence ID, or we have no hierarchy, or we're resetting the binding; set a relative ID
+	else if (LocalSequenceID == MovieSceneSequenceID::Invalid || !InValue.Guid.IsValid() || Hierarchy == nullptr)
 	{
 		SetCurrentValue(UE::MovieScene::FRelativeObjectBindingID(InValue.Guid));
 	}
