Sequencer: Fix actor reference track automatically creating keys

diff --git a/Engine/Source/Editor/MovieSceneTools/Private/TrackEditors/PropertyTrackEditors/ActorReferencePropertyTrackEditor.cpp b/Engine/Source/Editor/MovieSceneTools/Private/TrackEditors/PropertyTrackEditors/ActorReferencePropertyTrackEditor.cpp
--- a/Engine/Source/Editor/MovieSceneTools/Private/TrackEditors/PropertyTrackEditors/ActorReferencePropertyTrackEditor.cpp
+++ b/Engine/Source/Editor/MovieSceneTools/Private/TrackEditors/PropertyTrackEditors/ActorReferencePropertyTrackEditor.cpp
@@ -10,19 +10,6 @@ TSharedRef<ISequencerTrackEditor> FActorReferencePropertyTrackEditor::CreateTrac
 	return MakeShareable(new FActorReferencePropertyTrackEditor(OwningSequencer));
 }
 
-void FActorReferencePropertyTrackEditor::OnAnimatedPropertyChanged( const FPropertyChangedParams& PropertyChangedParams )
-{
-	// Override the params by always manually adding a key/track so that we don't reference other spawnables and their levels. Override only if it's set to autokey, so that if a key is manually forced, it will still be created
-	ESequencerKeyMode OverrideKeyMode = PropertyChangedParams.KeyMode;
-	if (OverrideKeyMode == ESequencerKeyMode::AutoKey)
-	{
-		OverrideKeyMode = ESequencerKeyMode::ManualKey;
-	}
-	FPropertyChangedParams OverridePropertyChangedParams(PropertyChangedParams.ObjectsThatChanged, PropertyChangedParams.PropertyPath, PropertyChangedParams.StructPathToKey, OverrideKeyMode);
-
-	FPropertyTrackEditor::OnAnimatedPropertyChanged( OverridePropertyChangedParams );
-}
-
 void FActorReferencePropertyTrackEditor::GenerateKeysFromPropertyChanged( const FPropertyChangedParams& PropertyChangedParams, UMovieSceneSection* SectionToKey, FGeneratedTrackKeys& OutGeneratedKeys )
 {
 	// Care is taken here to ensure that GetPropertyValue is templated on UObject* which causes it to use the correct instantiation of GetPropertyValueImpl
diff --git a/Engine/Source/Editor/MovieSceneTools/Private/TrackEditors/PropertyTrackEditors/ActorReferencePropertyTrackEditor.h b/Engine/Source/Editor/MovieSceneTools/Private/TrackEditors/PropertyTrackEditors/ActorReferencePropertyTrackEditor.h
--- a/Engine/Source/Editor/MovieSceneTools/Private/TrackEditors/PropertyTrackEditors/ActorReferencePropertyTrackEditor.h
+++ b/Engine/Source/Editor/MovieSceneTools/Private/TrackEditors/PropertyTrackEditors/ActorReferencePropertyTrackEditor.h
@@ -52,7 +52,5 @@ public:
 protected:
 
 	//~ FPropertyTrackEditor interface
-
-	virtual void OnAnimatedPropertyChanged( const FPropertyChangedParams& PropertyChangedParams ) override;
 	virtual void GenerateKeysFromPropertyChanged(const FPropertyChangedParams& PropertyChangedParams, UMovieSceneSection* SectionToKey, FGeneratedTrackKeys& OutGeneratedKeys) override;
 };
