Matinee: Fixed a bug in the Matinee to LevelSequence converter that caused audio sections to get added to their track twice. This has been an issue for a long time but only surfaced now that the MVVM UI has checks for it.

diff --git a/Engine/Source/Editor/MovieSceneTools/Private/MatineeImportTools.cpp b/Engine/Source/Editor/MovieSceneTools/Private/MatineeImportTools.cpp
--- a/Engine/Source/Editor/MovieSceneTools/Private/MatineeImportTools.cpp
+++ b/Engine/Source/Editor/MovieSceneTools/Private/MatineeImportTools.cpp
@@ -869,6 +869,7 @@ bool FMatineeImportTools::CopyInterpSoundTrack( UInterpTrackSound* MatineeSoundT
 
 	for ( const FSoundTrackKey& SoundTrackKey : MatineeSoundTrack->Sounds )
 	{
+		// AddNewSound automatically adds the new section to the track list, so no need to do it later.
 		AudioTrack->AddNewSound( SoundTrackKey.Sound, (SoundTrackKey.Time * FrameRate).RoundToFrame() );
 
 		UMovieSceneAudioSection* NewAudioSection = Cast<UMovieSceneAudioSection>(AudioTrack->GetAllSections().Last());
@@ -878,7 +879,6 @@ bool FMatineeImportTools::CopyInterpSoundTrack( UInterpTrackSound* MatineeSoundT
 		FloatChannels[0]->SetDefault(SoundTrackKey.Volume);
 		FloatChannels[1]->SetDefault(SoundTrackKey.Pitch);
 
-		AudioTrack->AddSection( *NewAudioSection );
 		bSectionCreated = true;
 	}
 
