Sequencer: Linked Anim Sequence, save anim sequence to disk when the level sequence is also saved to disk

diff --git a/Engine/Source/Editor/MovieSceneTools/Private/TrackEditors/SkeletalAnimationTrackEditor.cpp b/Engine/Source/Editor/MovieSceneTools/Private/TrackEditors/SkeletalAnimationTrackEditor.cpp
--- a/Engine/Source/Editor/MovieSceneTools/Private/TrackEditors/SkeletalAnimationTrackEditor.cpp
+++ b/Engine/Source/Editor/MovieSceneTools/Private/TrackEditors/SkeletalAnimationTrackEditor.cpp
@@ -81,7 +81,7 @@
 #include "LevelSequence.h"
 #include "LevelSequenceAnimSequenceLink.h"
 #include "AnimSequenceLevelSequenceLink.h"
-
+#include "UObject/SavePackage.h"
 
 int32 FSkeletalAnimationTrackEditor::NumberActive = 0;
 
@@ -1260,6 +1260,16 @@ void FSkeletalAnimationTrackEditor::OnSequencerSaved(ISequencer& )
 						AnimSeqExportOption->bExportMaterialCurves = bSavedExportMaterialCurves;
 						AnimSeqExportOption->bExportTransforms = bSavedExportTransforms;
 						AnimSeqExportOption->bRecordInWorldSpace = bSavedIncludeComponentTransform;
+
+						//save the anim sequence to disk to make sure they are in sync
+						UPackage* const Package = AnimSequence->GetOutermost();
+						FString const PackageName = Package->GetName();
+						FString const PackageFileName = FPackageName::LongPackageNameToFilename(PackageName, FPackageName::GetAssetPackageExtension());
+
+						FSavePackageArgs SaveArgs;
+						SaveArgs.TopLevelFlags = RF_Standalone;
+						SaveArgs.SaveFlags = SAVE_NoError;
+						UPackage::SavePackage(Package, NULL, *PackageFileName, SaveArgs);
 					}
 				}
 			}
