Sequencer: Fixed console variable track icon

diff --git a/Engine/Source/Editor/EditorStyle/Private/StarshipStyle.cpp b/Engine/Source/Editor/EditorStyle/Private/StarshipStyle.cpp
--- a/Engine/Source/Editor/EditorStyle/Private/StarshipStyle.cpp
+++ b/Engine/Source/Editor/EditorStyle/Private/StarshipStyle.cpp
@@ -2109,6 +2109,7 @@ void FStarshipEditorStyle::FStyle::SetupSequencerStyles()
 		Set("Sequencer.Tracks.Sub", new IMAGE_BRUSH("Sequencer/Dropdown_Icons/Icon_Sub_Track_16x", Icon16x16));
 		Set("Sequencer.Tracks.LevelVisibility", new IMAGE_BRUSH("Sequencer/Dropdown_Icons/Icon_Level_Visibility_Track_16x", Icon16x16));
 		Set("Sequencer.Tracks.DataLayer", new IMAGE_BRUSH_SVG("Starship/Common/DataLayers", Icon16x16));
+		Set("Sequencer.Tracks.CVar", new CORE_IMAGE_BRUSH_SVG("Starship/Common/Console", Icon16x16));
 
 		Set("Sequencer.CursorDecorator_MarqueeAdd", new IMAGE_BRUSH("Sequencer/CursorDecorator_MarqueeAdd", Icon16x16));
 		Set("Sequencer.CursorDecorator_MarqueeSubtract", new IMAGE_BRUSH("Sequencer/CursorDecorator_MarqueeSubtract", Icon16x16));
diff --git a/Engine/Source/Editor/MovieSceneTools/Private/TrackEditors/CVarTrackEditor.cpp b/Engine/Source/Editor/MovieSceneTools/Private/TrackEditors/CVarTrackEditor.cpp
--- a/Engine/Source/Editor/MovieSceneTools/Private/TrackEditors/CVarTrackEditor.cpp
+++ b/Engine/Source/Editor/MovieSceneTools/Private/TrackEditors/CVarTrackEditor.cpp
@@ -34,7 +34,7 @@ void FCVarTrackEditor::BuildAddTrackMenu(FMenuBuilder& MenuBuilder)
 	MenuBuilder.AddMenuEntry(
 		LOCTEXT("AddCVarTrack", "Console Variable Track"),
 		LOCTEXT("AddCVarTooltip", "Adds a new Console Variable track which will allow you to change console variables while a sequence plays."),
-		FSlateIcon(FAppStyle::GetAppStyleSetName(), "DebugConsole.Icon"),
+		FSlateIcon(FAppStyle::GetAppStyleSetName(), "Sequencer.Tracks.CVar"),
 		FUIAction(
 			FExecuteAction::CreateRaw(this, &FCVarTrackEditor::HandleAddCVarTrackMenuEntryExecute)
 		)
@@ -84,7 +84,7 @@ bool FCVarTrackEditor::SupportsType(TSubclassOf<UMovieSceneTrack> Type) const
 
 const FSlateBrush* FCVarTrackEditor::GetIconBrush() const
 {
-	return FAppStyle::GetBrush("DebugConsole.Icon");
+	return FAppStyle::GetBrush("Sequencer.Tracks.CVar");
 }
 // ~SequencerTrackEditor interface
 
