Sequencer: Linked Anim Sequences: Dirty linked anim sequence if the sequencer asset changes.
#trivial

diff --git a/Engine/Source/Editor/MovieSceneTools/Private/TrackEditors/SkeletalAnimationTrackEditor.cpp b/Engine/Source/Editor/MovieSceneTools/Private/TrackEditors/SkeletalAnimationTrackEditor.cpp
--- a/Engine/Source/Editor/MovieSceneTools/Private/TrackEditors/SkeletalAnimationTrackEditor.cpp
+++ b/Engine/Source/Editor/MovieSceneTools/Private/TrackEditors/SkeletalAnimationTrackEditor.cpp
@@ -1098,7 +1098,7 @@ FSkeletalAnimationTrackEditor::FSkeletalAnimationTrackEditor( TSharedRef<ISequen
 	AnimSeqExportOption = NewObject<UAnimSeqExportOption>();
 
 	SequencerSavedHandle = InSequencer->OnPostSave().AddRaw(this, &FSkeletalAnimationTrackEditor::OnSequencerSaved);
-
+	SequencerChangedHandle = InSequencer->OnMovieSceneDataChanged().AddRaw(this, &FSkeletalAnimationTrackEditor::OnSequencerDataChanged);
 	FCoreUObjectDelegates::OnObjectPropertyChanged.AddRaw(this, &FSkeletalAnimationTrackEditor::OnPostPropertyChanged);
 }
 
@@ -1124,9 +1124,18 @@ void FSkeletalAnimationTrackEditor::OnRelease()
 
 	--FSkeletalAnimationTrackEditor::NumberActive;
 
-	if (GetSequencer().IsValid() && SequencerSavedHandle.IsValid())
+	if (GetSequencer().IsValid())
 	{
-		GetSequencer()->OnPostSave().Remove(SequencerSavedHandle);
+		if (SequencerSavedHandle.IsValid())
+		{
+			GetSequencer()->OnPostSave().Remove(SequencerSavedHandle);
+			SequencerSavedHandle.Reset();
+		}
+		if (SequencerChangedHandle.IsValid())
+		{
+			GetSequencer()->OnMovieSceneDataChanged().Remove(SequencerChangedHandle);
+			SequencerChangedHandle.Reset();
+		}
 	}
 	if (FSkeletalAnimationTrackEditor::NumberActive == 0)
 	{
@@ -1277,6 +1286,36 @@ void FSkeletalAnimationTrackEditor::OnSequencerSaved(ISequencer& )
 	}
 }
 
+//dirty anim sequence when the sequencer changes, to make sure it get's checked out etc..
+void FSkeletalAnimationTrackEditor::OnSequencerDataChanged(EMovieSceneDataChangeType DataChangeType)
+{
+	TSharedPtr<ISequencer> SequencerPtr = GetSequencer();
+	if (!SequencerPtr.IsValid())
+	{
+		return;
+	}
+	ULevelSequence* LevelSequence = Cast<ULevelSequence>(SequencerPtr->GetFocusedMovieSceneSequence());
+	if (LevelSequence && LevelSequence->GetClass()->ImplementsInterface(UInterface_AssetUserData::StaticClass()))
+	{
+		if (IInterface_AssetUserData* AssetUserDataInterface = Cast< IInterface_AssetUserData >(LevelSequence))
+		{
+			ULevelSequenceAnimSequenceLink* LevelAnimLink = AssetUserDataInterface->GetAssetUserData< ULevelSequenceAnimSequenceLink >();
+			if (LevelAnimLink)
+			{
+				for (int32 Index = LevelAnimLink->AnimSequenceLinks.Num() - 1; Index >= 0; --Index)
+				{
+					FLevelSequenceAnimSequenceLinkItem& Item = LevelAnimLink->AnimSequenceLinks[Index];
+					UAnimSequence* AnimSequence = Item.ResolveAnimSequence();
+					if (AnimSequence)
+					{
+						AnimSequence->Modify();
+					}
+				}
+			}
+		}
+	}
+}
+
 void FSkeletalAnimationTrackEditor::OnPostPropertyChanged(UObject* InObject, FPropertyChangedEvent& InPropertyChangedEvent)
 {
 	// This attempts to fix tposes when changing properties without evaluating Sequencer UE-101261, 
diff --git a/Engine/Source/Editor/MovieSceneTools/Private/TrackEditors/SkeletalAnimationTrackEditor.h b/Engine/Source/Editor/MovieSceneTools/Private/TrackEditors/SkeletalAnimationTrackEditor.h
--- a/Engine/Source/Editor/MovieSceneTools/Private/TrackEditors/SkeletalAnimationTrackEditor.h
+++ b/Engine/Source/Editor/MovieSceneTools/Private/TrackEditors/SkeletalAnimationTrackEditor.h
@@ -124,7 +124,8 @@ private:
 	/* Delegate to handle sequencer changes for auto baking of anim sequences*/
 	FDelegateHandle SequencerSavedHandle;
 	void OnSequencerSaved(ISequencer& InSequence);
-
+	FDelegateHandle SequencerChangedHandle;
+	void OnSequencerDataChanged(EMovieSceneDataChangeType DataChangeType);
 	void OnPostPropertyChanged(UObject* InObject, struct FPropertyChangedEvent& InPropertyChangedEvent);
 };
 
