SubTrackEditor: Expose the sub track type for both the asset picker and track creation.

diff --git a/Engine/Source/Editor/MovieSceneTools/Private/TrackEditors/SubTrackEditor.cpp b/Engine/Source/Editor/MovieSceneTools/Private/TrackEditors/SubTrackEditor.cpp
--- a/Engine/Source/Editor/MovieSceneTools/Private/TrackEditors/SubTrackEditor.cpp
+++ b/Engine/Source/Editor/MovieSceneTools/Private/TrackEditors/SubTrackEditor.cpp
@@ -433,6 +433,15 @@ bool FSubTrackEditor::CanAddSubSequence(const UMovieSceneSequence& Sequence) con
 	return FSubTrackEditorUtil::CanAddSubSequence(FocusedSequence, Sequence);
 }
 
+UMovieSceneSubTrack* FSubTrackEditor::CreateNewTrack(UMovieScene* MovieScene) const
+{
+	return MovieScene->AddMasterTrack<UMovieSceneSubTrack>();
+}
+
+void FSubTrackEditor::GetSupportedSequenceClassPaths(TArray<FTopLevelAssetPath>& ClassPaths) const
+{
+	ClassPaths.Add(FTopLevelAssetPath(TEXT("/Script/LevelSequence"), TEXT("LevelSequence")));
+}
 
 /* FSubTrackEditor callbacks
  *****************************************************************************/
@@ -454,7 +463,7 @@ void FSubTrackEditor::HandleAddSubTrackMenuEntryExecute()
 	const FScopedTransaction Transaction(LOCTEXT("AddSubTrack_Transaction", "Add Sub Track"));
 	FocusedMovieScene->Modify();
 
-	auto NewTrack = FocusedMovieScene->AddMasterTrack<UMovieSceneSubTrack>();
+	UMovieSceneSubTrack* NewTrack = CreateNewTrack(FocusedMovieScene);
 	ensure(NewTrack);
 
 	if (GetSequencer().IsValid())
@@ -495,7 +504,7 @@ TSharedRef<SWidget> FSubTrackEditor::HandleAddSubSequenceComboButtonGetMenuConte
 			AssetPickerConfig.OnAssetEnterPressed = FOnAssetEnterPressed::CreateRaw( this, &FSubTrackEditor::HandleAddSubSequenceComboButtonMenuEntryEnterPressed, InTrack);
 			AssetPickerConfig.bAllowNullSelection = false;
 			AssetPickerConfig.InitialAssetViewType = EAssetViewType::List;
-			AssetPickerConfig.Filter.ClassPaths.Add(FTopLevelAssetPath(TEXT("/Script/LevelSequence"), TEXT("LevelSequence")));
+			GetSupportedSequenceClassPaths(AssetPickerConfig.Filter.ClassPaths);
 			AssetPickerConfig.SaveSettingsName = TEXT("SequencerAssetPicker");
 			AssetPickerConfig.AdditionalReferencingAssets.Add(FAssetData(Sequence));
 		}
diff --git a/Engine/Source/Editor/MovieSceneTools/Public/TrackEditors/SubTrackEditor.h b/Engine/Source/Editor/MovieSceneTools/Public/TrackEditors/SubTrackEditor.h
--- a/Engine/Source/Editor/MovieSceneTools/Public/TrackEditors/SubTrackEditor.h
+++ b/Engine/Source/Editor/MovieSceneTools/Public/TrackEditors/SubTrackEditor.h
@@ -15,6 +15,7 @@
 class AActor;
 struct FAssetData;
 class FMenuBuilder;
+class UMovieSceneSubTrack;
 
 /**
  * Tools for animatable property types such as floats ands vectors
@@ -72,6 +73,22 @@ public:
 	 */
 	bool CanAddSubSequence(const UMovieSceneSequence& Sequence) const;
 
+protected:
+	/**
+	 * Create a new UMovieSceneSubTrack.
+	 * 
+	 * @param MovieScene the MovieScene object to create a track.
+	 * @return the created MovieSceneSubTrack
+	 */
+	virtual UMovieSceneSubTrack* CreateNewTrack(UMovieScene* MovieScene) const;
+
+	/**
+	 * Return the list of supported sequence class paths.
+	 * 
+	 * @param OutClassPaths the list of supported sequence class paths.
+	 */
+	virtual void GetSupportedSequenceClassPaths(TArray<FTopLevelAssetPath>& OutClassPaths) const;
+
 private:
 
 	/** Callback for executing the "Add Event Track" menu entry. */
