Remove "Export to Camera Anim..." feature

It's not doing anything anymore since the Matinee asset factory was removed.

diff --git a/Engine/Source/Editor/MovieSceneTools/Private/MovieSceneToolHelpers.cpp b/Engine/Source/Editor/MovieSceneTools/Private/MovieSceneToolHelpers.cpp
--- a/Engine/Source/Editor/MovieSceneTools/Private/MovieSceneToolHelpers.cpp
+++ b/Engine/Source/Editor/MovieSceneTools/Private/MovieSceneToolHelpers.cpp
@@ -3682,87 +3682,6 @@ void MovieSceneToolHelpers::CopyKeyDataToMoveAxis(const TMovieSceneChannelData<F
 	}
 }
 
-UObject* MovieSceneToolHelpers::ExportToCameraAnim(UMovieScene* InMovieScene, FGuid& InObjectBinding)
-{
-	// Create a new camera anim
-	IAssetTools& AssetTools = FModuleManager::GetModuleChecked<FAssetToolsModule>("AssetTools").Get();
-
-	UObject* NewAsset = nullptr;
-
-	// Attempt to create a new asset
-	for (TObjectIterator<UClass> It ; It ; ++It)
-	{
-		UClass* CurrentClass = *It;
-		if (CurrentClass->IsChildOf(UFactory::StaticClass()) && !(CurrentClass->HasAnyClassFlags(CLASS_Abstract)))
-		{
-			UFactory* Factory = Cast<UFactory>(CurrentClass->GetDefaultObject());
-			if (Factory->CanCreateNew() && Factory->ImportPriority >= 0 && Factory->SupportedClass == UCameraAnim::StaticClass())
-			{
-				NewAsset = AssetTools.CreateAssetWithDialog(UCameraAnim::StaticClass(), Factory);
-				break;
-			}
-		}
-	}
-
-	if (!NewAsset)
-	{
-		return NewAsset;
-	}
-
-	static FName Transform("Transform");
-	UMovieScene3DTransformTrack* TransformTrack = InMovieScene->FindTrack<UMovieScene3DTransformTrack>(InObjectBinding, Transform); 
-	if (TransformTrack)
-	{
-		UCameraAnim* CameraAnim = CastChecked<UCameraAnim>(NewAsset);
-		UInterpGroup* CameraInterpGroup = CameraAnim->CameraInterpGroup;
-		CameraAnim->bRelativeToInitialTransform=false;
-
-		UInterpGroupInst* CameraInst = NewObject<UInterpGroupInst>(CameraAnim, NAME_None, RF_Transactional);
-		CameraInst->InitGroupInst(CameraInterpGroup, nullptr);
-
-		UInterpTrackMove* MovementTrack = NewObject<UInterpTrackMove>(CameraInterpGroup, NAME_None, RF_Transactional);
-		CameraInterpGroup->InterpTracks.Add(MovementTrack);
-		
-		UInterpTrackInstMove* MovementTrackInst = NewObject<UInterpTrackInstMove>(CameraInst, NAME_None, RF_Transactional);
-		CameraInst->TrackInst.Add(MovementTrackInst);
-		MovementTrackInst->InitTrackInst(MovementTrack);
-			
-		MovementTrack->CreateSubTracks(false);
-
-		UInterpTrackMoveAxis* MoveAxies[6];
-		for( int32 SubTrackIndex = 0; SubTrackIndex < 6; ++SubTrackIndex )
-		{
-			MoveAxies[ SubTrackIndex ] = Cast<UInterpTrackMoveAxis>( MovementTrack->SubTracks[ SubTrackIndex ] );
-		}
-
-		TArray<UMovieSceneSection*> Sections = TransformTrack->GetAllSections();
-
-		if (Sections.Num())
-		{
-			if (Sections.Num() > 1)
-			{
-				UE_LOG(LogMovieScene, Error, TEXT("Export to Camera Anim: Failed to export, multiple sections (%d) are not supported"), Sections.Num());
-			}
-			else
-			{
-				FFrameRate TickResolution = InMovieScene->GetTickResolution();
-				UMovieScene3DTransformSection* TransformSection = Cast<UMovieScene3DTransformSection>(Sections[0]);
-				TArrayView<FMovieSceneDoubleChannel*> DoubleChannels = TransformSection->GetChannelProxy().GetChannels<FMovieSceneDoubleChannel>();
-
-				CopyKeyDataToMoveAxis(DoubleChannels[0]->GetData(), MoveAxies[AXIS_TranslationX], TickResolution);
-				CopyKeyDataToMoveAxis(DoubleChannels[1]->GetData(), MoveAxies[AXIS_TranslationY], TickResolution);
-				CopyKeyDataToMoveAxis(DoubleChannels[2]->GetData(), MoveAxies[AXIS_TranslationZ], TickResolution);
-				CopyKeyDataToMoveAxis(DoubleChannels[3]->GetData(), MoveAxies[AXIS_RotationX],    TickResolution);
-				CopyKeyDataToMoveAxis(DoubleChannels[4]->GetData(), MoveAxies[AXIS_RotationY],    TickResolution);
-				CopyKeyDataToMoveAxis(DoubleChannels[5]->GetData(), MoveAxies[AXIS_RotationZ],    TickResolution);
-			}
-		}
-	}
-
-	return NewAsset;
-}
-
-
 bool MovieSceneToolHelpers::HasHiddenMobility(const UClass* ObjectClass)
 {
 	if (ObjectClass)
diff --git a/Engine/Source/Editor/MovieSceneTools/Public/MovieSceneToolHelpers.h b/Engine/Source/Editor/MovieSceneTools/Public/MovieSceneToolHelpers.h
--- a/Engine/Source/Editor/MovieSceneTools/Public/MovieSceneToolHelpers.h
+++ b/Engine/Source/Editor/MovieSceneTools/Public/MovieSceneToolHelpers.h
@@ -407,15 +407,6 @@ public:
 	 */
 	static void CopyKeyDataToMoveAxis(const TMovieSceneChannelData<FMovieSceneDoubleValue>& KeyData, UInterpTrackMoveAxis* MoveAxis, FFrameRate FrameRate);
 
-	/*
-	 * Export the object binding to a camera anim
-	 *
-	 * @param InMovieScene The movie scene to export the object binding from
-	 * @param InObjectBinding The object binding to export
-	 * @return The exported camera anim asset
-	 */
-	static UObject* ExportToCameraAnim(UMovieScene* InMovieScene, FGuid& InObjectBinding);
-
 
 	/*
 	 * Export the SkelMesh to an Anim Sequence for specified MovieScene and Player
diff --git a/Engine/Source/Editor/Sequencer/Private/MVVM/ViewModels/ObjectBindingModel.cpp b/Engine/Source/Editor/Sequencer/Private/MVVM/ViewModels/ObjectBindingModel.cpp
--- a/Engine/Source/Editor/Sequencer/Private/MVVM/ViewModels/ObjectBindingModel.cpp
+++ b/Engine/Source/Editor/Sequencer/Private/MVVM/ViewModels/ObjectBindingModel.cpp
@@ -969,14 +969,7 @@ void FObjectBindingModel::BuildContextMenu(FMenuBuilder& MenuBuilder)
 			FUIAction(
 				FExecuteAction::CreateLambda([=]{ Sequencer->ExportFBX(); })
 			));
-
-		MenuBuilder.AddMenuEntry(
-			LOCTEXT("ExportToCameraAnim", "Export to Camera Anim..."),
-			LOCTEXT("ExportToCameraAnimTooltip", "Exports the animation to a camera anim asset"),
-			FSlateIcon(),
-			FUIAction(
-				FExecuteAction::CreateLambda([=]{ Sequencer->ExportToCameraAnim(); })
-			));
+			
 		MenuBuilder.EndSection();
 	}
 
diff --git a/Engine/Source/Editor/Sequencer/Private/Sequencer.cpp b/Engine/Source/Editor/Sequencer/Private/Sequencer.cpp
--- a/Engine/Source/Editor/Sequencer/Private/Sequencer.cpp
+++ b/Engine/Source/Editor/Sequencer/Private/Sequencer.cpp
@@ -12355,22 +12355,6 @@ void FSequencer::ExportFBXInternal(const FString& ExportFilename, const TArray<F
 }
 
 
-void FSequencer::ExportToCameraAnim()
-{
-	using namespace UE::Sequencer;
-
-	for (const TWeakPtr<FViewModel>& Node : Selection.GetSelectedOutlinerItems())
-	{
-		if (FObjectBindingModel* ObjectBindingNode = ICastable::CastWeakPtr<FObjectBindingModel>(Node))
-		{
-			FGuid Guid = ObjectBindingNode->GetObjectGuid();
-			
-			MovieSceneToolHelpers::ExportToCameraAnim(GetFocusedMovieSceneSequence()->GetMovieScene(), Guid);
-		}
-	}
-}
-
-
 void FSequencer::GenericTextEntryModeless(const FText& DialogText, const FText& DefaultText, FOnTextCommitted OnTextComitted)
 {
 	TSharedRef<STextEntryPopup> TextEntryPopup = 
@@ -13127,11 +13111,6 @@ void FSequencer::BindCommands()
 		FExecuteAction::CreateSP( this, &FSequencer::ExportFBX ),
 		FCanExecuteAction::CreateLambda( [] { return true; } ) );
 
-	SequencerCommandBindings->MapAction(
-		Commands.ExportToCameraAnim,
-		FExecuteAction::CreateSP( this, &FSequencer::ExportToCameraAnim ),
-		FCanExecuteAction::CreateLambda( [] { return true; } ) );
-
 	SequencerCommandBindings->MapAction(
 		Commands.MoveToNewFolder,
 		FExecuteAction::CreateSP( this, &FSequencer::MoveSelectedNodesToNewFolder ),
diff --git a/Engine/Source/Editor/Sequencer/Private/Sequencer.h b/Engine/Source/Editor/Sequencer/Private/Sequencer.h
--- a/Engine/Source/Editor/Sequencer/Private/Sequencer.h
+++ b/Engine/Source/Editor/Sequencer/Private/Sequencer.h
@@ -674,9 +674,6 @@
 	/** Exports the animation to an fbx file. */
 	void ExportFBX();
 
-	/** Exports the animation to a camera anim asset. */
-	void ExportToCameraAnim();
-
 public:
 
 	/** Put the sequencer in a horizontally auto-scrolling state with the given rate */
diff --git a/Engine/Source/Editor/Sequencer/Private/SequencerCommands.cpp b/Engine/Source/Editor/Sequencer/Private/SequencerCommands.cpp
--- a/Engine/Source/Editor/Sequencer/Private/SequencerCommands.cpp
+++ b/Engine/Source/Editor/Sequencer/Private/SequencerCommands.cpp
@@ -162,8 +162,6 @@ void FSequencerCommands::RegisterCommands()
 	UI_COMMAND( ImportFBX, "Import...", "Import the animation from an FBX file.", EUserInterfaceActionType::Button, FInputChord() );
 	UI_COMMAND( ExportFBX, "Export...", "Export the selected objects (all if none selected) and animation to an FBX file. (Shots and sub-scenes not supported)", EUserInterfaceActionType::Button, FInputChord() );
 
-	UI_COMMAND( ExportToCameraAnim, "Export to Camera Anim...", "Exports the animation to a camera anim asset.", EUserInterfaceActionType::Button, FInputChord() );
-
 	UI_COMMAND( ToggleEvaluateSubSequencesInIsolation, "Evaluate Sub Sequences In Isolation", "When enabled, will only evaluate the currently focused sequence; otherwise evaluate from the master sequence.", EUserInterfaceActionType::ToggleButton, FInputChord() );
 
 	UI_COMMAND( QuickTreeSearch, "Quick Tree Search", "Jumps keyboard focus to the tree searchbox to allow searching for tracks in the current Sequence.", EUserInterfaceActionType::Button, FInputChord(EModifierKey::Control, EKeys::F));
diff --git a/Engine/Source/Editor/Sequencer/Public/SequencerCommands.h b/Engine/Source/Editor/Sequencer/Public/SequencerCommands.h
--- a/Engine/Source/Editor/Sequencer/Public/SequencerCommands.h
+++ b/Engine/Source/Editor/Sequencer/Public/SequencerCommands.h
@@ -406,9 +406,6 @@ public:
 	/** Exports animation to fbx. */
 	TSharedPtr< FUICommandInfo > ExportFBX;
 
-	/** Exports animation to camera anim. */
-	TSharedPtr< FUICommandInfo > ExportToCameraAnim;
-
 	/** Toggle whether we should evaluate sub sequences in isolation */
 	TSharedPtr< FUICommandInfo > ToggleEvaluateSubSequencesInIsolation;
 
