20666935 follow up (missing file)

diff --git a/Engine/Source/Editor/MovieSceneTools/Private/MovieSceneToolsModule.cpp b/Engine/Source/Editor/MovieSceneTools/Private/MovieSceneToolsModule.cpp
--- a/Engine/Source/Editor/MovieSceneTools/Private/MovieSceneToolsModule.cpp
+++ b/Engine/Source/Editor/MovieSceneTools/Private/MovieSceneToolsModule.cpp
@@ -77,6 +77,8 @@
 #include "Channels/EventChannelCurveModel.h"
 #include "Channels/SCurveEditorEventChannelView.h"
 #include "Channels/MovieSceneAudioTriggerChannel.h"
+#include "Channels/MovieSceneConstraintChannel.h"
+#include "Channels/ConstraintChannelEditor.h"
 #include "Sections/MovieSceneEventSection.h"
 
 #include "MovieSceneEventUtils.h"
@@ -188,6 +190,8 @@
 
 		SequencerModule.RegisterChannelInterface<FMovieSceneCameraShakeSourceTriggerChannel>();
 
+		SequencerModule.RegisterChannelInterface<FMovieSceneConstraintChannel>();
+
 		ICurveEditorModule& CurveEditorModule = FModuleManager::LoadModuleChecked<ICurveEditorModule>("CurveEditor");
 
 		FEventChannelCurveModel::EventView = CurveEditorModule.RegisterView(FOnCreateCurveEditorView::CreateStatic(
