Sequencer: Tick parent chain skeletal mesh components when gathering actor transforms for trails and pivots

PR #9301: Fix for Control Rig Snap Tool not taking full world transform into account (Contributed by [REDACTED_PERSON])

diff --git a/Engine/Source/Editor/MovieSceneTools/Private/MovieSceneToolHelpers.cpp b/Engine/Source/Editor/MovieSceneTools/Private/MovieSceneToolHelpers.cpp
--- a/Engine/Source/Editor/MovieSceneTools/Private/MovieSceneToolHelpers.cpp
+++ b/Engine/Source/Editor/MovieSceneTools/Private/MovieSceneToolHelpers.cpp
@@ -4317,20 +4317,26 @@ static void GetSequencerActorWorldTransforms(IMovieScenePlayer* Player, UMovieSc
 						}
 					}
 
-					TArray<USkeletalMeshComponent*> MeshComps;
-					ActorSelection.Actor->GetComponents<USkeletalMeshComponent>(MeshComps, true);
-
-					for (USkeletalMeshComponent* MeshComp : MeshComps)
+					AActor* Parent = ActorSelection.Actor.Get();
+					while (Parent)
 					{
-						MeshComp->TickAnimation(0.03f, false);
-						MeshComp->RefreshBoneTransforms();
-						MeshComp->RefreshSlaveComponents();
-						MeshComp->UpdateComponentToWorld();
-						MeshComp->FinalizeBoneTransform();
-						MeshComp->MarkRenderTransformDirty();
-						MeshComp->MarkRenderDynamicDataDirty();
+						TArray<USkeletalMeshComponent*> MeshComps;
+						ActorSelection.Actor->GetComponents<USkeletalMeshComponent>(MeshComps, true);
+
+						for (USkeletalMeshComponent* MeshComp : MeshComps)
+						{
+							MeshComp->TickAnimation(0.03f, false);
+							MeshComp->RefreshBoneTransforms();
+							MeshComp->RefreshSlaveComponents();
+							MeshComp->UpdateComponentToWorld();
+							MeshComp->FinalizeBoneTransform();
+							MeshComp->MarkRenderTransformDirty();
+							MeshComp->MarkRenderDynamicDataDirty();
+						}
+
+						Parent = Parent->GetAttachParentActor();
 					}
-					
+
 					OutTransforms[Index] = SkelMeshComp->GetSocketTransform(ActorSelection.SocketName);// GetSocketTransofrm is world space in theory*OutTransforms[Index];
 
 				}
