Fix AnimSequenceLink missing option (Evaluate All Skeletal Mesh Components)

diff --git a/Engine/Plugins/Animation/ControlRig/Source/ControlRigEditor/Private/ControlRigEditorModule.cpp b/Engine/Plugins/Animation/ControlRig/Source/ControlRigEditor/Private/ControlRigEditorModule.cpp
--- a/Engine/Plugins/Animation/ControlRig/Source/ControlRigEditor/Private/ControlRigEditorModule.cpp
+++ b/Engine/Plugins/Animation/ControlRig/Source/ControlRigEditor/Private/ControlRigEditorModule.cpp
@@ -769,6 +769,7 @@ void FControlRigEditorModule::BakeToControlRig(UClass* ControlRigClass, UAnimSeq
 						LevelAnimLinkItem.bExportMaterialCurves = true;
 						LevelAnimLinkItem.bExportTransforms = true;
 						LevelAnimLinkItem.bRecordInWorldSpace = false;
+						LevelAnimLinkItem.bEvaluateAllSkeletalMeshComponents = false;
 						LevelAnimLink->AnimSequenceLinks.Add(LevelAnimLinkItem);
 						AssetUserDataInterface->AddAssetUserData(LevelAnimLink);
 					}
diff --git a/Engine/Plugins/MovieScene/SequencerScripting/Source/SequencerScriptingEditor/Private/SequencerTools.cpp b/Engine/Plugins/MovieScene/SequencerScripting/Source/SequencerScriptingEditor/Private/SequencerTools.cpp
--- a/Engine/Plugins/MovieScene/SequencerScripting/Source/SequencerScriptingEditor/Private/SequencerTools.cpp
+++ b/Engine/Plugins/MovieScene/SequencerScripting/Source/SequencerScriptingEditor/Private/SequencerTools.cpp
@@ -390,6 +390,7 @@ bool USequencerToolsFunctionLibrary::LinkAnimSequence(ULevelSequence*  Sequence,
 						LevelAnimLinkItem.bExportMaterialCurves = ExportOptions->bExportMaterialCurves;
 						LevelAnimLinkItem.bExportTransforms = ExportOptions->bExportTransforms;
 						LevelAnimLinkItem.bRecordInWorldSpace = ExportOptions->bRecordInWorldSpace;
+						LevelAnimLinkItem.bEvaluateAllSkeletalMeshComponents = ExportOptions->bEvaluateAllSkeletalMeshComponents;
 
 						break;
 					}
@@ -409,6 +410,7 @@ bool USequencerToolsFunctionLibrary::LinkAnimSequence(ULevelSequence*  Sequence,
 				LevelAnimLinkItem.bExportMaterialCurves = ExportOptions->bExportMaterialCurves;
 				LevelAnimLinkItem.bExportTransforms = ExportOptions->bExportTransforms;
 				LevelAnimLinkItem.bRecordInWorldSpace = ExportOptions->bRecordInWorldSpace;
+				LevelAnimLinkItem.bEvaluateAllSkeletalMeshComponents = ExportOptions->bEvaluateAllSkeletalMeshComponents;
 
 				LevelAnimLink->AnimSequenceLinks.Add(LevelAnimLinkItem);
 				AssetUserDataInterface->AddAssetUserData(LevelAnimLink);
diff --git a/Engine/Source/Editor/MovieSceneTools/Private/TrackEditors/SkeletalAnimationTrackEditor.cpp b/Engine/Source/Editor/MovieSceneTools/Private/TrackEditors/SkeletalAnimationTrackEditor.cpp
--- a/Engine/Source/Editor/MovieSceneTools/Private/TrackEditors/SkeletalAnimationTrackEditor.cpp
+++ b/Engine/Source/Editor/MovieSceneTools/Private/TrackEditors/SkeletalAnimationTrackEditor.cpp
@@ -1256,11 +1256,13 @@ void FSkeletalAnimationTrackEditor::OnSequencerSaved(ISequencer& )
 						const bool bSavedExportMaterialCurves = AnimSeqExportOption->bExportMaterialCurves;
 						const bool bSavedExportTransforms = AnimSeqExportOption->bExportTransforms;
 						const bool bSavedIncludeComponentTransform = AnimSeqExportOption->bRecordInWorldSpace;
+						const bool bSavedEvaluateAllSkeletalMeshComponents = AnimSeqExportOption->bEvaluateAllSkeletalMeshComponents;
 						AnimSeqExportOption->bExportMorphTargets = Item.bExportMorphTargets;
 						AnimSeqExportOption->bExportAttributeCurves = Item.bExportAttributeCurves;
 						AnimSeqExportOption->bExportMaterialCurves = Item.bExportMaterialCurves;
 						AnimSeqExportOption->bExportTransforms = Item.bExportTransforms;
 						AnimSeqExportOption->bRecordInWorldSpace = Item.bRecordInWorldSpace;
+						AnimSeqExportOption->bEvaluateAllSkeletalMeshComponents = Item.bEvaluateAllSkeletalMeshComponents;
 
 						bool bResult = MovieSceneToolHelpers::ExportToAnimSequence(AnimSequence, AnimSeqExportOption, MovieScene, SequencerPtr.Get(), SkelMeshComp, Template, RootToLocalTransform);
 
@@ -1269,6 +1271,7 @@ void FSkeletalAnimationTrackEditor::OnSequencerSaved(ISequencer& )
 						AnimSeqExportOption->bExportMaterialCurves = bSavedExportMaterialCurves;
 						AnimSeqExportOption->bExportTransforms = bSavedExportTransforms;
 						AnimSeqExportOption->bRecordInWorldSpace = bSavedIncludeComponentTransform;
+						AnimSeqExportOption->bEvaluateAllSkeletalMeshComponents = bSavedEvaluateAllSkeletalMeshComponents;
 
 						//save the anim sequence to disk to make sure they are in sync
 						UPackage* const Package = AnimSequence->GetOutermost();
@@ -1460,6 +1463,7 @@ bool FSkeletalAnimationTrackEditor::CreateAnimationSequence(const TArray<UObject
 								LevelAnimLinkItem.bExportMaterialCurves = AnimSeqExportOption->bExportMaterialCurves;
 								LevelAnimLinkItem.bExportTransforms = AnimSeqExportOption->bExportTransforms;
 								LevelAnimLinkItem.bRecordInWorldSpace = AnimSeqExportOption->bRecordInWorldSpace;
+								LevelAnimLinkItem.bEvaluateAllSkeletalMeshComponents = AnimSeqExportOption->bEvaluateAllSkeletalMeshComponents;
 
 								break;
 							}
@@ -1480,6 +1484,7 @@ bool FSkeletalAnimationTrackEditor::CreateAnimationSequence(const TArray<UObject
 						LevelAnimLinkItem.bExportMaterialCurves = AnimSeqExportOption->bExportMaterialCurves;
 						LevelAnimLinkItem.bExportTransforms = AnimSeqExportOption->bExportTransforms;
 						LevelAnimLinkItem.bRecordInWorldSpace = AnimSeqExportOption->bRecordInWorldSpace;
+						LevelAnimLinkItem.bEvaluateAllSkeletalMeshComponents = AnimSeqExportOption->bEvaluateAllSkeletalMeshComponents;
 
 						LevelAnimLink->AnimSequenceLinks.Add(LevelAnimLinkItem);
 						AssetUserDataInterface->AddAssetUserData(LevelAnimLink);
diff --git a/Engine/Source/Runtime/LevelSequence/Public/LevelSequenceAnimSequenceLink.h b/Engine/Source/Runtime/LevelSequence/Public/LevelSequenceAnimSequenceLink.h
--- a/Engine/Source/Runtime/LevelSequence/Public/LevelSequenceAnimSequenceLink.h
+++ b/Engine/Source/Runtime/LevelSequence/Public/LevelSequenceAnimSequenceLink.h
@@ -31,6 +31,8 @@ public:
 	bool bExportMaterialCurves = true;
 	UPROPERTY(BlueprintReadWrite, Category = Property)
 	bool bRecordInWorldSpace = false;
+	UPROPERTY(BlueprintReadWrite, Category = Property)
+	bool bEvaluateAllSkeletalMeshComponents = false;
 
 	void SetAnimSequence(UAnimSequence* InAnimSequence);
 	UAnimSequence* ResolveAnimSequence();
