Sequencer: Add  Sequencer.DuplicateLinkedAnimSequence cvar to control whether or not we duplicate the anim sequence when the level sequence is duplicated. Also added a ClearLinkedAnimSequences to SequencerTools so the

diff --git a/Engine/Plugins/MovieScene/SequencerScripting/Source/SequencerScriptingEditor/Private/SequencerTools.cpp b/Engine/Plugins/MovieScene/SequencerScripting/Source/SequencerScriptingEditor/Private/SequencerTools.cpp
--- a/Engine/Plugins/MovieScene/SequencerScripting/Source/SequencerScriptingEditor/Private/SequencerTools.cpp
+++ b/Engine/Plugins/MovieScene/SequencerScripting/Source/SequencerScriptingEditor/Private/SequencerTools.cpp
@@ -333,6 +333,31 @@ bool USequencerToolsFunctionLibrary::ExportAnimSequence(UWorld* World, ULevelSeq
 	}
 	return bResult;
 }
+void USequencerToolsFunctionLibrary::ClearLinkedAnimSequences(ULevelSequence* LevelSequence)
+{
+	if (LevelSequence)
+	{
+		if (IInterface_AssetUserData* LevelSequenceUserDataInterface = Cast< IInterface_AssetUserData >(LevelSequence))
+		{
+			ULevelSequenceAnimSequenceLink* LevelAnimLink = LevelSequenceUserDataInterface->GetAssetUserData< ULevelSequenceAnimSequenceLink >();
+			if (LevelAnimLink)
+			{
+				for (int32 Index = 0; Index < LevelAnimLink->AnimSequenceLinks.Num(); ++Index)
+				{
+					FLevelSequenceAnimSequenceLinkItem& LevelAnimLinkItem = LevelAnimLink->AnimSequenceLinks[Index];
+					if (UAnimSequence* AnimSequence = LevelAnimLinkItem.ResolveAnimSequence())
+					{
+						if (IInterface_AssetUserData* AnimAssetUserData = Cast< IInterface_AssetUserData >(AnimSequence))
+						{
+							AnimAssetUserData->RemoveUserDataOfClass(UAnimSequenceLevelSequenceLink::StaticClass());
+						}
+					}
+				}
+				LevelSequenceUserDataInterface->RemoveUserDataOfClass(ULevelSequenceAnimSequenceLink::StaticClass());			
+			}
+		}
+	}
+}
 
 bool USequencerToolsFunctionLibrary::LinkAnimSequence(ULevelSequence*  Sequence,  UAnimSequence* AnimSequence, const UAnimSeqExportOption* ExportOptions,const FMovieSceneBindingProxy& Binding)
 {
diff --git a/Engine/Plugins/MovieScene/SequencerScripting/Source/SequencerScriptingEditor/Public/SequencerTools.h b/Engine/Plugins/MovieScene/SequencerScripting/Source/SequencerScriptingEditor/Public/SequencerTools.h
--- a/Engine/Plugins/MovieScene/SequencerScripting/Source/SequencerScriptingEditor/Public/SequencerTools.h
+++ b/Engine/Plugins/MovieScene/SequencerScripting/Source/SequencerScriptingEditor/Public/SequencerTools.h
@@ -157,6 +157,14 @@ public:
 	UFUNCTION(BlueprintCallable, Category = "Editor Scripting | Sequencer Tools | Animation")
 	static bool LinkAnimSequence(ULevelSequence* Sequence, UAnimSequence* AnimSequence, const UAnimSeqExportOption* ExportOptions, const FMovieSceneBindingProxy& Binding);
 
+	/*
+	 * Clear all linked anim sequences between this level sequence and any linked anim sequences
+	 *
+	 * @InLevelSequence LevelSequence to clean links for
+	*/
+	UFUNCTION(BlueprintCallable, Category = "Editor Scripting | Sequencer Tools | Animation")
+	static void ClearLinkedAnimSequences(ULevelSequence* InLevelSequence);
+
 	/*
 	 * Get the link to the level sequence if it exists on this anim sequence
 	 *
diff --git a/Engine/Source/Editor/MovieSceneTools/Private/MovieSceneToolsModule.cpp b/Engine/Source/Editor/MovieSceneTools/Private/MovieSceneToolsModule.cpp
--- a/Engine/Source/Editor/MovieSceneTools/Private/MovieSceneToolsModule.cpp
+++ b/Engine/Source/Editor/MovieSceneTools/Private/MovieSceneToolsModule.cpp
@@ -97,6 +97,8 @@
 
 #define LOCTEXT_NAMESPACE "FMovieSceneToolsModule"
 
+TAutoConsoleVariable<bool> CVarDuplicateLinkedAnimSequence(TEXT("Sequencer.DuplicateLinkedAnimSequence"), false, TEXT("When true when we duplicate a level sequence that has a linked anim sequence it will duplicate and link the anim sequencel, if false we leave any link alone."));
+
 #if !IS_MONOLITHIC
 	UE::MovieScene::FEntityManager*& GEntityManagerForDebugging = UE::MovieScene::GEntityManagerForDebuggingVisualizers;
 #endif
@@ -344,7 +346,7 @@ void FMovieSceneToolsModule::RemoveForCookEventSection(UMovieSceneEventSectionBa
 }
 
 //When we duplicate a ULevelSequence we check to see if there are any linked UAnimSequences in the asset user data,
-//if so we make a copy of the anim sequence, and then set the links to the copies.
+//if so we either make a copy of the anim sequence, or leave it alone since a rename can also be a duplicate, the user will need to clean up this link later.
 void FMovieSceneToolsModule::PostDuplicateEvent(ULevelSequence* LevelSequence)
 {
 	if (LevelSequence && LevelSequence->GetClass()->ImplementsInterface(UInterface_AssetUserData::StaticClass()))
@@ -354,27 +356,31 @@ void FMovieSceneToolsModule::PostDuplicateEvent(ULevelSequence* LevelSequence)
 			ULevelSequenceAnimSequenceLink* LevelAnimLink = AssetUserDataInterface->GetAssetUserData< ULevelSequenceAnimSequenceLink >();
 			if (LevelAnimLink)
 			{
-				for (FLevelSequenceAnimSequenceLinkItem& Item : LevelAnimLink->AnimSequenceLinks)
+				const bool bDuplicateAnimSequence = CVarDuplicateLinkedAnimSequence.GetValueOnGameThread();
+				if (bDuplicateAnimSequence)
 				{
-					if (UAnimSequence* AnimSequence = Item.ResolveAnimSequence())
+					for (FLevelSequenceAnimSequenceLinkItem& Item : LevelAnimLink->AnimSequenceLinks)
 					{
-						TArray<UAnimSequence*> AnimSequencesToDuplicate;
-						AnimSequencesToDuplicate.Add(AnimSequence);
-						UPackage* DestinationPackage = AnimSequence->GetPackage();
-						EditorAnimUtils::FNameDuplicationRule NameRule;
-						NameRule.FolderPath = FPackageName::GetLongPackagePath(AnimSequence->GetPathName()) / TEXT("");
-						TMap<UAnimSequence*, UAnimSequence*> DuplicatedAnimAssets = EditorAnimUtils::DuplicateAssets<UAnimSequence>(AnimSequencesToDuplicate, DestinationPackage, &NameRule);
-						for (TPair<UAnimSequence*, UAnimSequence*>& Duplicates : DuplicatedAnimAssets)
+						if (UAnimSequence* AnimSequence = Item.ResolveAnimSequence())
 						{
-							if (UAnimSequence* NewAnimSequence = Duplicates.Value)
+							TArray<UAnimSequence*> AnimSequencesToDuplicate;
+							AnimSequencesToDuplicate.Add(AnimSequence);
+							UPackage* DestinationPackage = AnimSequence->GetPackage();
+							EditorAnimUtils::FNameDuplicationRule NameRule;
+							NameRule.FolderPath = FPackageName::GetLongPackagePath(AnimSequence->GetPathName()) / TEXT("");
+							TMap<UAnimSequence*, UAnimSequence*> DuplicatedAnimAssets = EditorAnimUtils::DuplicateAssets<UAnimSequence>(AnimSequencesToDuplicate, DestinationPackage, &NameRule);
+							for (TPair<UAnimSequence*, UAnimSequence*>& Duplicates : DuplicatedAnimAssets)
 							{
-								Item.PathToAnimSequence = FSoftObjectPath(NewAnimSequence);
-								if (IInterface_AssetUserData* AnimAssetUserData = Cast< IInterface_AssetUserData >(NewAnimSequence))
+								if (UAnimSequence* NewAnimSequence = Duplicates.Value)
 								{
-									UAnimSequenceLevelSequenceLink* AnimLevelLink = AnimAssetUserData->GetAssetUserData< UAnimSequenceLevelSequenceLink >();
-									if (AnimLevelLink)
+									Item.PathToAnimSequence = FSoftObjectPath(NewAnimSequence);
+									if (IInterface_AssetUserData* AnimAssetUserData = Cast< IInterface_AssetUserData >(NewAnimSequence))
 									{
-										AnimLevelLink->SetLevelSequence(LevelSequence);
+										UAnimSequenceLevelSequenceLink* AnimLevelLink = AnimAssetUserData->GetAssetUserData< UAnimSequenceLevelSequenceLink >();
+										if (AnimLevelLink)
+										{
+											AnimLevelLink->SetLevelSequence(LevelSequence);
+										}
 									}
 								}
 							}
