- Label improvements to display constraint's type and parent in editor
- avoid creating handles for certain controls

diff --git a/Engine/Plugins/Animation/ControlRig/Source/ControlRig/Private/Constraints/ControlRigTransformableHandle.cpp b/Engine/Plugins/Animation/ControlRig/Source/ControlRig/Private/Constraints/ControlRigTransformableHandle.cpp
--- a/Engine/Plugins/Animation/ControlRig/Source/ControlRig/Private/Constraints/ControlRigTransformableHandle.cpp
+++ b/Engine/Plugins/Animation/ControlRig/Source/ControlRig/Private/Constraints/ControlRigTransformableHandle.cpp
@@ -262,16 +262,23 @@ void UTransformableControlHandle::OnControlModified(
 
 #if WITH_EDITOR
 
-FName UTransformableControlHandle::GetName() const
+FString UTransformableControlHandle::GetLabel() const
+{
+	return ControlName.ToString();
+}
+
+FString UTransformableControlHandle::GetFullLabel() const
 {
 	const USkeletalMeshComponent* SkeletalMesh = GetSkeletalMesh();
+	if (!SkeletalMesh)
+	{
+		static const FString DummyLabel;
+		return DummyLabel;
+	}
 	
-	const AActor* Actor = SkeletalMesh ? SkeletalMesh->GetOwner() : nullptr;
-	const FName ControlRigName = Actor ? FName(*Actor->GetActorLabel()) : SkeletalMesh ? SkeletalMesh->GetFName() : NAME_None; 
-
-	const FString FullName = FString::Printf(TEXT("%s/%s"), *ControlRigName.ToString(), *ControlName.ToString() );
-
-	return FName(*FullName);
+	const AActor* Actor = SkeletalMesh->GetOwner();
+	const FString ControlRigLabel = Actor ? Actor->GetActorLabel() : SkeletalMesh->GetName();
+	return FString::Printf(TEXT("%s/%s"), *ControlRigLabel, *ControlName.ToString() );
 }
 
 void UTransformableControlHandle::OnObjectsReplaced(const TMap<UObject*, UObject*>& InOldToNewInstances)
diff --git a/Engine/Plugins/Animation/ControlRig/Source/ControlRig/Private/ControlRig.cpp b/Engine/Plugins/Animation/ControlRig/Source/ControlRig/Private/ControlRig.cpp
--- a/Engine/Plugins/Animation/ControlRig/Source/ControlRig/Private/ControlRig.cpp
+++ b/Engine/Plugins/Animation/ControlRig/Source/ControlRig/Private/ControlRig.cpp
@@ -3373,8 +3373,34 @@ uint32 UControlRig::GetHashForInitializeVMSnapShot()
 	return Hash;
 }
 
-UTransformableControlHandle* UControlRig::CreateTransformableControlHandle(UObject* InOuter, const FName& InControlName)
+UTransformableControlHandle* UControlRig::CreateTransformableControlHandle(
+	UObject* InOuter,
+	const FName& InControlName) const
 {
+	auto IsConstrainable = [this](const FName& InControlName)
+	{
+		const FRigControlElement* ControlElement = FindControl(InControlName);
+		if (!ControlElement)
+		{
+			return false;
+		}
+		
+		const FRigControlSettings& ControlSettings = ControlElement->Settings;
+		if (ControlSettings.ControlType == ERigControlType::Bool ||
+			ControlSettings.ControlType == ERigControlType::Float ||
+			ControlSettings.ControlType == ERigControlType::Integer)
+		{
+			return false;
+		}
+
+		return true;
+	};
+
+	if (!IsConstrainable(InControlName))
+	{
+		return nullptr;
+	}
+	
 	UTransformableControlHandle* CtrlHandle = NewObject<UTransformableControlHandle>(InOuter);
 	ensure(CtrlHandle);
 	CtrlHandle->ControlRig = this;
diff --git a/Engine/Plugins/Animation/ControlRig/Source/ControlRig/Private/Sequencer/MovieSceneControlRigParameterSection.cpp b/Engine/Plugins/Animation/ControlRig/Source/ControlRig/Private/Sequencer/MovieSceneControlRigParameterSection.cpp
--- a/Engine/Plugins/Animation/ControlRig/Source/ControlRig/Private/Sequencer/MovieSceneControlRigParameterSection.cpp
+++ b/Engine/Plugins/Animation/ControlRig/Source/ControlRig/Private/Sequencer/MovieSceneControlRigParameterSection.cpp
@@ -1280,7 +1280,7 @@ void UMovieSceneControlRigParameterSection::AddSpaceChannel(FName InControlName,
 void UMovieSceneControlRigParameterSection::RegisterConstraintsHandles()
 {
 	if (OnConstraintRemovedHandle.IsValid())
-	{
+	{ 
 		return;
 	}
 
@@ -1289,6 +1289,11 @@ void UMovieSceneControlRigParameterSection::RegisterConstraintsHandles()
 		return;
 	}
 
+	if (ConstraintsChannels.IsEmpty())
+	{
+		return;
+	}
+	
 	// register constraint deletion handle
 	FConstraintsManagerController& Controller = FConstraintsManagerController::Get(ControlRig->GetWorld());
 	OnConstraintRemovedHandle = Controller.OnConstraintRemoved().AddLambda([this](FName InConstraintName)
@@ -1341,6 +1346,20 @@ void UMovieSceneControlRigParameterSection::AddConstraintChannel(UTickableConstr
 	}
 }
 
+void UMovieSceneControlRigParameterSection::RemoveConstraintChannel(const FName& InConstraintName)
+{
+	const int32 Index = ConstraintsChannels.IndexOfByPredicate([InConstraintName](const FConstraintAndActiveChannel& InChannel)
+	{
+		return InChannel.Constraint.IsValid() ? InChannel.Constraint->GetFName() == InConstraintName : false;
+	});
+
+	if (ConstraintsChannels.IsValidIndex(Index))
+	{
+		ConstraintsChannels.RemoveAt(Index);
+		ReconstructChannelProxy();
+	}
+}
+
 const TArray<FConstraintAndActiveChannel>& UMovieSceneControlRigParameterSection::GetConstraintsChannels() const
 {
 	return ConstraintsChannels;
@@ -1411,7 +1430,8 @@ void UMovieSceneControlRigParameterSection::ReconstructChannelProxy()
 		};
 
 
-		auto AddConstrainChannels = [this, GetConstraints, &ConstraintsChannelIndex, &TotalIndex, &Channels](const FName& InControlName, const FText& InGroup, const bool bEnabled)
+		auto AddConstrainChannels = [this, GetConstraints, &ConstraintsChannelIndex, &TotalIndex, &Channels](
+			const FName& InControlName, const FText& InGroup, const bool bEnabled)
 		{
 			TArray<TObjectPtr<UTickableConstraint>> Constraints = GetConstraints(InControlName);
 			for (const TObjectPtr<UTickableConstraint>& Constraint: Constraints)
@@ -1425,12 +1445,26 @@ void UMovieSceneControlRigParameterSection::ReconstructChannelProxy()
 					}
 
 #if WITH_EDITOR
-					FMovieSceneChannelMetaData MetaData(ConstraintName,  FText::FromName(ConstraintName), InGroup, bEnabled);
+					ConstraintChannel->ActiveChannel.ExtraLabel = [WeakConstraint = MakeWeakObjectPtr(Constraint)]
+					{
+						if (WeakConstraint.IsValid())
+						{
+							FString ParentStr; WeakConstraint->GetLabel().Split(TEXT("."), &ParentStr, nullptr);
+							if (!ParentStr.IsEmpty())
+							{
+								return ParentStr;
+							}		
+						}
+						static const FString DummyStr;
+						return DummyStr;
+					};
+					
+					const FText DisplayText = FText::FromString(Constraint->GetTypeLabel());
+					FMovieSceneChannelMetaData MetaData(ConstraintName, DisplayText,InGroup, bEnabled);
 					ConstraintsChannelIndex += 1;
 					MetaData.SortOrder = TotalIndex++;
 					MetaData.bCanCollapseToTrack = false;
-
-					// Channels.Add(ConstraintChannel->ActiveCurve.ParameterCurve, MetaData, TMovieSceneExternalValue<bool>());
+		
 					Channels.Add(ConstraintChannel->ActiveChannel, MetaData, TMovieSceneExternalValue<bool>());
 #else
 					Channels.Add(ConstraintChannel->ActiveChannel);
diff --git a/Engine/Plugins/Animation/ControlRig/Source/ControlRig/Public/Constraints/ControlRigTransformableHandle.h b/Engine/Plugins/Animation/ControlRig/Source/ControlRig/Public/Constraints/ControlRigTransformableHandle.h
--- a/Engine/Plugins/Animation/ControlRig/Source/ControlRig/Public/Constraints/ControlRigTransformableHandle.h
+++ b/Engine/Plugins/Animation/ControlRig/Source/ControlRig/Public/Constraints/ControlRigTransformableHandle.h
@@ -60,7 +60,8 @@ public:
 
 #if WITH_EDITOR
 	/** @todo document */
-	virtual FName GetName() const override;
+	virtual FString GetLabel() const override;
+	virtual FString GetFullLabel() const override;
 #endif
 
 	/** The ControlRig that this handle is pointing at. */
diff --git a/Engine/Plugins/Animation/ControlRig/Source/ControlRig/Public/ControlRig.h b/Engine/Plugins/Animation/ControlRig/Source/ControlRig/Public/ControlRig.h
--- a/Engine/Plugins/Animation/ControlRig/Source/ControlRig/Public/ControlRig.h
+++ b/Engine/Plugins/Animation/ControlRig/Source/ControlRig/Public/ControlRig.h
@@ -158,7 +158,7 @@ public:
 	/** Creates a transformable control handle for the specified control to be used by the constraints system. Should use the UObject from 
 	ConstraintsScriptingLibrary::GetManager(UWorld* InWorld)*/
 	UFUNCTION(BlueprintPure, Category = "Control Rig | Constraints")
-	UTransformableControlHandle* CreateTransformableControlHandle(UObject* Outer, const FName& ControlName);
+	UTransformableControlHandle* CreateTransformableControlHandle(UObject* Outer, const FName& ControlName) const;
 
 
 #if WITH_EDITOR
diff --git a/Engine/Plugins/Animation/ControlRig/Source/ControlRig/Public/Sequencer/MovieSceneControlRigParameterSection.h b/Engine/Plugins/Animation/ControlRig/Source/ControlRig/Public/Sequencer/MovieSceneControlRigParameterSection.h
--- a/Engine/Plugins/Animation/ControlRig/Source/ControlRig/Public/Sequencer/MovieSceneControlRigParameterSection.h
+++ b/Engine/Plugins/Animation/ControlRig/Source/ControlRig/Public/Sequencer/MovieSceneControlRigParameterSection.h
@@ -404,6 +404,7 @@ public:
 	bool HasConstraintChannel(const FName& InConstraintName) const;
 	FConstraintAndActiveChannel* GetConstraintChannel(const FName& InConstraintName);
 	void AddConstraintChannel(UTickableConstraint* InConstraint, bool bReconstructChannel);
+	void RemoveConstraintChannel(const FName& InConstraintName);
 
 	/** Clear Everything Out*/
 	void ClearAllParameters();
diff --git a/Engine/Plugins/Animation/ControlRig/Source/ControlRigEditor/Private/Editor/Constraints/SConstraintsWidget.cpp b/Engine/Plugins/Animation/ControlRig/Source/ControlRigEditor/Private/Editor/Constraints/SConstraintsWidget.cpp
--- a/Engine/Plugins/Animation/ControlRig/Source/ControlRigEditor/Private/Editor/Constraints/SConstraintsWidget.cpp
+++ b/Engine/Plugins/Animation/ControlRig/Source/ControlRigEditor/Private/Editor/Constraints/SConstraintsWidget.cpp
@@ -363,13 +363,30 @@ void SEditableConstraintItem::Construct(
 		}
 	});
 
+	// constraint
 	auto GetConstraint = [this]()
 	{
 		UWorld* World = GCurrentLevelEditingViewportClient->GetWorld();
 		const FConstraintsManagerController& Controller = FConstraintsManagerController::Get(World);
 		return Controller.GetConstraint(ConstraintItem->Name);
 	};
+	UTickableConstraint* Constraint = GetConstraint();
+
+	// labels
+	FString ParentLabel(TEXT("undefined")), ChildLabel(TEXT("undefined"));
+	if (!InItem->Label.IsEmpty())
+	{
+		ChildLabel = InItem->Label;
+		InItem->Label.Split(TEXT("."), &ParentLabel, &ChildLabel);
+	}
 	
+	FString ParentFullLabel = ParentLabel, ChildFullLabel = ChildLabel;
+	if (Constraint)
+	{
+		Constraint->GetFullLabel().Split(TEXT("."), &ParentFullLabel, &ChildFullLabel);
+	}
+
+	// widgets
 	ChildSlot
 	.Padding(FMargin(8.f, 2.f, 12.f, 2.f))
 	[
@@ -384,13 +401,13 @@ void SEditableConstraintItem::Construct(
 			SNew(SBorder)
 			.Padding(FMargin(5.0, 2.0, 5.0, 2.0))
 			.BorderImage(RoundedBoxBrush)
-			.BorderBackgroundColor_Lambda([GetConstraint]()
+			.BorderBackgroundColor_Lambda([Constraint]()
 			{
-				if (const UTickableConstraint* Constraint = GetConstraint())
+				if (!Constraint)
 				{
-					return Constraint->Active ? FStyleColors::Select : FStyleColors::Transparent;
+					return FStyleColors::Transparent;
 				}
-				return FStyleColors::Transparent;
+				return Constraint->Active ? FStyleColors::Select : FStyleColors::Transparent;
 			})
 			.Content()
 			[
@@ -416,25 +433,29 @@ void SEditableConstraintItem::Construct(
 				.Padding(0)
 				[
 					SNew( STextBlock )
-					.Text_Lambda( [InItem]()
+					.Text_Lambda( [ParentLabel]()
+					{
+						return FText::FromString(ParentLabel);
+					})
+					.Font_Lambda([Constraint]()
 					{
-						if(InItem->ParentName == NAME_None)
+						if (!Constraint)
 						{
-							return FText::FromName(InItem->Name);
+							return IDetailLayoutBuilder::GetDetailFont();
 						}
-
-						const FString TextStr = FString::Printf(TEXT("%s (%s)"), *InItem->Name.ToString(), *InItem->ParentName.ToString());
-						return FText::FromString(TextStr);
+						return Constraint->Active ? IDetailLayoutBuilder::GetDetailFont() : IDetailLayoutBuilder::GetDetailFontItalic();
 					})
-					.Font_Lambda([GetConstraint]()
+					.ToolTipText_Lambda( [Constraint, ParentFullLabel, ChildFullLabel]()
 					{
-						const UTickableConstraint* Constraint = GetConstraint();
 						if (!Constraint)
 						{
-							return IDetailLayoutBuilder::GetDetailFont();
+							return FText();
 						}
 
-						return Constraint->Active ? IDetailLayoutBuilder::GetDetailFont() : IDetailLayoutBuilder::GetDetailFontItalic();
+						static constexpr TCHAR ToolTipFormat[] = TEXT("%s constraint between parent '%s' and child '%s'.");
+						const FString TypeLabel = Constraint->GetTypeLabel();
+						const FString FullLabel = FString::Printf(ToolTipFormat, *TypeLabel, *ParentFullLabel, *ChildFullLabel);
+						return FText::FromString(FullLabel);
 					})
 				]
 			]
@@ -456,9 +477,8 @@ void SEditableConstraintItem::Construct(
 			SNew(SButton)
 			.ButtonStyle(FAppStyle::Get(), TEXT("SimpleButton"))
 			.ContentPadding(0)
-			.OnClicked_Lambda(	[GetConstraint]()
+			.OnClicked_Lambda(	[Constraint]()
 			{
-				UTickableConstraint* Constraint = GetConstraint();
 				if (UTickableTransformConstraint* TransformConstraint = Cast<UTickableTransformConstraint>(Constraint))
 				{
 					// FConstraintChannelHelper::AddConstraintKey(TransformConstraint);
diff --git a/Engine/Plugins/Animation/ControlRig/Source/ControlRigEditor/Private/Editor/Constraints/SConstraintsWidget.h b/Engine/Plugins/Animation/ControlRig/Source/ControlRigEditor/Private/Editor/Constraints/SConstraintsWidget.h
--- a/Engine/Plugins/Animation/ControlRig/Source/ControlRigEditor/Private/Editor/Constraints/SConstraintsWidget.h
+++ b/Engine/Plugins/Animation/ControlRig/Source/ControlRigEditor/Private/Editor/Constraints/SConstraintsWidget.h
@@ -148,20 +148,20 @@ public:
 	static TSharedRef<FEditableConstraintItem> Make(
 		const FName& InName,
 		ETransformConstraintType InType,
-		const FName& InParentName)
+		const FString& InLabel)
 	{
-		return MakeShareable(new FEditableConstraintItem(InName, InType, InParentName));
+		return MakeShareable(new FEditableConstraintItem(InName, InType, InLabel));
 	}
 	
 	FName Name = NAME_None;
 	ETransformConstraintType Type = ETransformConstraintType::Parent;
-	FName ParentName = NAME_None;
+	FString Label;
 
 private:
-	FEditableConstraintItem(const FName& InName, ETransformConstraintType InType, const FName& InParentName)
+	FEditableConstraintItem(const FName& InName, ETransformConstraintType InType, const FString& InLabel)
 		: Name(InName)
 		, Type(InType)
-		, ParentName(InParentName)
+		, Label(InLabel)
 	{}
 	FEditableConstraintItem() {}
 };
diff --git a/Engine/Source/Editor/MovieSceneTools/Private/Channels/ConstraintChannelEditor.cpp b/Engine/Source/Editor/MovieSceneTools/Private/Channels/ConstraintChannelEditor.cpp
--- a/Engine/Source/Editor/MovieSceneTools/Private/Channels/ConstraintChannelEditor.cpp
+++ b/Engine/Source/Editor/MovieSceneTools/Private/Channels/ConstraintChannelEditor.cpp
@@ -107,6 +107,64 @@ void DrawKeys(
 	}
 }
 
+namespace
+{
+
+struct FConstraintChannelLabelBuilder
+{
+public:
+	FConstraintChannelLabelBuilder(const UMovieSceneSection* InOwner, const FMovieSceneConstraintChannel* InChannel)
+		: Owner(InOwner)
+		, Channel(InChannel)
+	{
+		ConstructLabel();
+	}
+
+	FText Label;
+	bool bValid = false;
+	
+private:
+	void ConstructLabel() 
+	{
+		if (!Channel || !Owner)
+		{
+			return;
+		}
+
+		// get the stored extra label
+		const FString ExtraLabel = Channel->ExtraLabel ? Channel->ExtraLabel() : FString();
+		if (!ExtraLabel.IsEmpty())
+		{
+			Label = FText::FromString(ExtraLabel);
+			bValid = true;
+			return;
+		}
+
+		// fallback to meta data
+		const FMovieSceneChannelProxy& Proxy = Owner->GetChannelProxy();
+		const TArrayView<FMovieSceneConstraintChannel*> ConstraintChannels = Proxy.GetChannels<FMovieSceneConstraintChannel>();
+		const TArrayView<const FMovieSceneChannelMetaData> MetaData = Proxy.GetMetaData<FMovieSceneConstraintChannel>();
+		
+		const int32 ChannelIndex = ConstraintChannels.IndexOfByKey(Channel);
+		if (!ConstraintChannels.IsValidIndex(ChannelIndex) || !MetaData.IsValidIndex(ChannelIndex))
+		{
+			return;
+		}
+
+		const FName& ConstraintName = MetaData[ChannelIndex].Name;
+		if (ConstraintName != NAME_None)
+		{
+			Label = FText::FromName(ConstraintName);
+			bValid = true;
+		}
+	}
+
+	const UMovieSceneSection* Owner = nullptr;
+	const FMovieSceneConstraintChannel* Channel = nullptr;
+};
+
+}
+
 void DrawExtra(
 	FMovieSceneConstraintChannel* Channel,
 	const UMovieSceneSection* Owner,
@@ -165,26 +223,9 @@ void DrawExtra(
 	const double InputMax = TimeToPixel.PixelToSeconds(Painter.SectionGeometry.GetLocalSize().X);
 
 	// label data
-	FText Label;
-	bool bDrawLabel = false;
-	
-	const FMovieSceneChannelProxy& Proxy = Owner->GetChannelProxy();
-	const TArrayView<FMovieSceneConstraintChannel*> ConstraintChannels = Proxy.GetChannels<FMovieSceneConstraintChannel>();
-	const int32 ChannelIndex = ConstraintChannels.IndexOfByKey(Channel);
-	if (ConstraintChannels.IsValidIndex(ChannelIndex))
-	{
-		const TArrayView<const FMovieSceneChannelMetaData> MetaData = Proxy.GetMetaData<FMovieSceneConstraintChannel>();
-		if (MetaData.IsValidIndex(ChannelIndex))
-		{
-			const FName& ConstraintName = MetaData[ChannelIndex].Name;
-			if (ConstraintName != NAME_None)
-			{
-				Label = FText::FromName(ConstraintName);
-				bDrawLabel = true;
-			}
-		}
-	}
-	const FVector2D TextSize = bDrawLabel ? FontMeasure->Measure(Label, FontInfo) : FVector2D::ZeroVector;
+	const FConstraintChannelLabelBuilder LabelBuilder(Owner, Channel);
+	const FText& Label = LabelBuilder.Label;
+	const FVector2D TextSize = LabelBuilder.bValid ? FontMeasure->Measure(Label, FontInfo) : FVector2D::ZeroVector;
 	static constexpr double LabelPixelOffset = 10.0;  
 	
 	for (int32 Index = 0; Index < Ranges.Num(); ++Index)
@@ -210,7 +251,7 @@ void DrawExtra(
 		FSlateDrawElement::MakeBox(Painter.DrawElements, Painter.LayerId, BoxGeometry, WhiteBrush, DrawEffects, Range.Color);
 
 		// draw label
-		if (bDrawLabel)
+		if (LabelBuilder.bValid)
 		{
 			const double LabelPos = BoxStart + LabelPixelOffset;
 			const FVector2D Position(LabelPos, LaneTop + (LocalSize.Y - TextSize.Y) * .5f);
@@ -233,3 +274,29 @@ void DrawExtra(
 		}
 	}
 }
+
+FKeyHandle AddOrUpdateKey(
+	FMovieSceneConstraintChannel* Channel,
+	UMovieSceneSection* SectionToKey,
+	FFrameNumber Time,
+	ISequencer& Sequencer,
+	const FGuid& ObjectBindingID,
+	FTrackInstancePropertyBindings* PropertyBindings)
+{
+	return FKeyHandle::Invalid();
+}
+
+bool CanCreateKeyEditor(const FMovieSceneConstraintChannel* InChannel)
+{
+	return false;
+}
+
+TSharedRef<SWidget> CreateKeyEditor(
+	const TMovieSceneChannelHandle<FMovieSceneConstraintChannel>& InChannel,
+	UMovieSceneSection* InSection,
+	const FGuid& InObjectBindingID,
+	TWeakPtr<FTrackInstancePropertyBindings> PropertyBindings,
+	TWeakPtr<ISequencer> InSequencer)
+{
+	return SNullWidget::NullWidget;
+}
\ No newline at end of file
diff --git a/Engine/Source/Editor/MovieSceneTools/Private/Channels/ConstraintChannelEditor.h b/Engine/Source/Editor/MovieSceneTools/Private/Channels/ConstraintChannelEditor.h
--- a/Engine/Source/Editor/MovieSceneTools/Private/Channels/ConstraintChannelEditor.h
+++ b/Engine/Source/Editor/MovieSceneTools/Private/Channels/ConstraintChannelEditor.h
@@ -19,3 +19,22 @@ void DrawExtra(
 	const UMovieSceneSection* Owner,
 	const FGeometry& AllottedGeometry,
 	FSequencerSectionPainter& Painter);
+
+/** Overrides for adding or updating a key for non-standard channels */
+FKeyHandle AddOrUpdateKey(
+	FMovieSceneConstraintChannel* Channel,
+	UMovieSceneSection* SectionToKey,
+	FFrameNumber Time,
+	ISequencer& Sequencer,
+	const FGuid& ObjectBindingID,
+	FTrackInstancePropertyBindings* PropertyBindings);
+
+/** Key editor overrides */
+bool CanCreateKeyEditor(const FMovieSceneConstraintChannel* InChannel);
+
+TSharedRef<SWidget> CreateKeyEditor(
+	const TMovieSceneChannelHandle<FMovieSceneConstraintChannel>& InChannel,
+	UMovieSceneSection* InSection,
+	const FGuid& InObjectBindingID,
+	TWeakPtr<FTrackInstancePropertyBindings> PropertyBindings,
+	TWeakPtr<ISequencer> Sequencer);
\ No newline at end of file
diff --git a/Engine/Source/Runtime/Experimental/Animation/Constraints/Private/ConstraintsManager.cpp b/Engine/Source/Runtime/Experimental/Animation/Constraints/Private/ConstraintsManager.cpp
--- a/Engine/Source/Runtime/Experimental/Animation/Constraints/Private/ConstraintsManager.cpp
+++ b/Engine/Source/Runtime/Experimental/Animation/Constraints/Private/ConstraintsManager.cpp
@@ -55,7 +55,7 @@ FString FConstraintTickFunction::DiagnosticMessage()
 	}
 
 #if WITH_EDITOR
-	return FString::Printf(TEXT("FConstraintTickFunction::Tick[%p] (%s)"), this, *Constraint->GetLabel().ToString());
+	return FString::Printf(TEXT("FConstraintTickFunction::Tick[%p] (%s)"), this, *Constraint->GetLabel());
 #else
 	return FString::Printf(TEXT("FConstraintTickFunction::Tick[%p] (%s)"), this, *Constraint->GetName());
 #endif
@@ -78,9 +78,19 @@ void UTickableConstraint::Evaluate() const
 
 #if WITH_EDITOR
 
-FName UTickableConstraint::GetLabel() const
+FString UTickableConstraint::GetLabel() const
 {
-	return GetFName();
+	return UTickableConstraint::StaticClass()->GetName();
+}
+
+FString UTickableConstraint::GetFullLabel() const
+{
+	return GetLabel();
+}
+
+FString UTickableConstraint::GetTypeLabel() const
+{
+	return GetLabel();
 }
 
 void UTickableConstraint::PostEditChangeProperty(FPropertyChangedEvent& PropertyChangedEvent)
@@ -98,7 +108,6 @@ void UTickableConstraint::PostEditChangeProperty(FPropertyChangedEvent& Property
 	}
 }
 
-
 #endif
 
 /** 
diff --git a/Engine/Source/Runtime/Experimental/Animation/Constraints/Private/TransformConstraint.cpp b/Engine/Source/Runtime/Experimental/Animation/Constraints/Private/TransformConstraint.cpp
--- a/Engine/Source/Runtime/Experimental/Animation/Constraints/Private/TransformConstraint.cpp
+++ b/Engine/Source/Runtime/Experimental/Animation/Constraints/Private/TransformConstraint.cpp
@@ -10,16 +10,54 @@
  * UTickableTransformConstraint
  **/
 
-/** @todo remove to use something else. */
 int64 UTickableTransformConstraint::GetType() const
 {
 	return static_cast<int64>(Type);
 }
 
 #if WITH_EDITOR
-FName UTickableTransformConstraint::GetLabel() const
+
+FString UTickableTransformConstraint::GetLabel() const
+{
+	if (!ChildTRSHandle->IsValid())
+	{
+		static const FString DummyLabel;
+		return DummyLabel;
+	}
+	
+	if (ParentTRSHandle->IsValid())
+	{
+		return FString::Printf(TEXT("%s.%s"), *ParentTRSHandle->GetLabel(), *ChildTRSHandle->GetLabel() );		
+	}
+
+	return ChildTRSHandle->GetLabel();
+}
+
+FString UTickableTransformConstraint::GetFullLabel() const
+{
+	if (!ChildTRSHandle->IsValid())
+	{
+		static const FString DummyLabel;
+		return DummyLabel;
+	}
+	
+	if (ParentTRSHandle->IsValid())
+	{
+		return FString::Printf(TEXT("%s.%s"), *ParentTRSHandle->GetFullLabel(), *ChildTRSHandle->GetFullLabel() );		
+	}
+
+	return ChildTRSHandle->GetLabel();
+}
+
+FString UTickableTransformConstraint::GetTypeLabel() const
 {
-	return ParentTRSHandle->IsValid() ? ParentTRSHandle->GetName() : NAME_None;
+	static const UEnum* TypeEnum = StaticEnum<ETransformConstraintType>();
+	if (TypeEnum->IsValidEnumValue(GetType()))
+	{
+		return TypeEnum->GetNameStringByValue(GetType());
+	}
+
+	return Super::GetTypeLabel();
 }
 
 void UTickableTransformConstraint::PostEditChangeProperty(FPropertyChangedEvent& PropertyChangedEvent)
@@ -837,10 +875,13 @@ UTickableTransformConstraint* FTransformConstraintUtils::CreateAndAddFromActors(
 
 	UTransformableHandle* ParentHandle = GetHandle(InParent, ConstraintsManager);
 	UTransformableHandle* ChildHandle = GetHandle(InChild, ConstraintsManager);
+
+	if (!ParentHandle || !ChildHandle)
+	{
+		return nullptr;
+	}
 	
 	UTickableTransformConstraint* Constraint = FTransformConstraintUtils::CreateFromType(InWorld, InType);
-
-
 	if (Constraint && (ParentHandle->IsValid() && ChildHandle->IsValid()))
 	{
 		if (AddConstraint(InWorld, ParentHandle, ChildHandle, Constraint, bMaintainOffset) == false)
diff --git a/Engine/Source/Runtime/Experimental/Animation/Constraints/Private/TransformableHandle.cpp b/Engine/Source/Runtime/Experimental/Animation/Constraints/Private/TransformableHandle.cpp
--- a/Engine/Source/Runtime/Experimental/Animation/Constraints/Private/TransformableHandle.cpp
+++ b/Engine/Source/Runtime/Experimental/Animation/Constraints/Private/TransformableHandle.cpp
@@ -167,14 +167,21 @@ void UTransformableComponentHandle::OnPostPropertyChanged(
 }
 
 #if WITH_EDITOR
-FName UTransformableComponentHandle::GetName() const
+FString UTransformableComponentHandle::GetLabel() const
 {
 	if (!Component.IsValid())
 	{
-		return NAME_None;
+		static const FString DummyLabel;
+		return DummyLabel;
 	}
 
 	const AActor* Actor = Component->GetOwner();
-	return Actor ? FName(*Actor->GetActorLabel()) : Component->GetFName();
+	return Actor ? Actor->GetActorLabel() : Component->GetName();
 }
+
+FString UTransformableComponentHandle::GetFullLabel() const
+{
+	return GetLabel();
+};
+
 #endif
\ No newline at end of file
diff --git a/Engine/Source/Runtime/Experimental/Animation/Constraints/Public/ConstraintChannel.h b/Engine/Source/Runtime/Experimental/Animation/Constraints/Public/ConstraintChannel.h
--- a/Engine/Source/Runtime/Experimental/Animation/Constraints/Public/ConstraintChannel.h
+++ b/Engine/Source/Runtime/Experimental/Animation/Constraints/Public/ConstraintChannel.h
@@ -5,6 +5,7 @@
 #include "Channels/MovieSceneBoolChannel.h"
 #include "MovieSceneClipboard.h"
 #include "ConstraintsManager.h"
+#include "Templates/Function.h"
 #include "ConstraintChannel.generated.h"
 
 USTRUCT()
@@ -16,6 +17,11 @@ struct CONSTRAINTS_API FMovieSceneConstraintChannel : public FMovieSceneBoolChan
 
 	// @todo 
 	bool Evaluate(FFrameTime InTime, bool& OutValue) const;
+
+#if WITH_EDITOR
+	using ExtraLabelFunction = TFunction< FString() >;
+	ExtraLabelFunction ExtraLabel;
+#endif
 };
 
 template<>
diff --git a/Engine/Source/Runtime/Experimental/Animation/Constraints/Public/ConstraintsManager.h b/Engine/Source/Runtime/Experimental/Animation/Constraints/Public/ConstraintsManager.h
--- a/Engine/Source/Runtime/Experimental/Animation/Constraints/Public/ConstraintsManager.h
+++ b/Engine/Source/Runtime/Experimental/Animation/Constraints/Public/ConstraintsManager.h
@@ -93,8 +93,12 @@ public:
 	virtual bool ReferencesObject(TWeakObjectPtr<UObject> InObject) const PURE_VIRTUAL(ReferencesObject, return false;);
 
 #if WITH_EDITOR
-	/** @todo document */
-	virtual FName GetLabel() const;
+	/** Returns the constraint's label used for UI. */
+	virtual FString GetLabel() const;
+	virtual FString GetFullLabel() const;
+
+	/** Returns the constraint's type label used for UI. */
+	virtual FString GetTypeLabel() const;
 
 	// UObject interface
 	virtual void PostEditChangeProperty(FPropertyChangedEvent& PropertyChangedEvent) override;
diff --git a/Engine/Source/Runtime/Experimental/Animation/Constraints/Public/TransformConstraint.h b/Engine/Source/Runtime/Experimental/Animation/Constraints/Public/TransformConstraint.h
--- a/Engine/Source/Runtime/Experimental/Animation/Constraints/Public/TransformConstraint.h
+++ b/Engine/Source/Runtime/Experimental/Animation/Constraints/Public/TransformConstraint.h
@@ -113,7 +113,11 @@ protected:
 #if WITH_EDITOR
 public:
 	/** Returns the constraint's label used for UI. */
-	virtual FName GetLabel() const override;
+	virtual FString GetLabel() const override;
+	virtual FString GetFullLabel() const override;
+
+	/** Returns the constraint's type label used for UI. */
+	virtual FString GetTypeLabel() const override;
 
 	// UObject interface
 	virtual void PostEditChangeProperty(FPropertyChangedEvent& PropertyChangedEvent) override;
diff --git a/Engine/Source/Runtime/Experimental/Animation/Constraints/Public/TransformableHandle.h b/Engine/Source/Runtime/Experimental/Animation/Constraints/Public/TransformableHandle.h
--- a/Engine/Source/Runtime/Experimental/Animation/Constraints/Public/TransformableHandle.h
+++ b/Engine/Source/Runtime/Experimental/Animation/Constraints/Public/TransformableHandle.h
@@ -58,7 +58,8 @@ public:
 	FHandleModifiedEvent& HandleModified();
 
 #if WITH_EDITOR
-	virtual FName GetName() const PURE_VIRTUAL(GetName, return NAME_None;);
+	virtual FString GetLabel() const PURE_VIRTUAL(GetLabel, return FString(););
+	virtual FString GetFullLabel() const PURE_VIRTUAL(GetFullLabel, return FString(););
 #endif
 
 protected:
@@ -102,7 +103,8 @@ public:
 	virtual TWeakObjectPtr<UObject> GetTarget() const override;
 
 #if WITH_EDITOR
-	virtual FName GetName() const override;
+	virtual FString GetLabel() const override;
+	virtual FString GetFullLabel() const override;
 #endif
 	
 	/** The Component that this handle is pointing at. */
