Sequencer: Change shortcut to lock/unlock camera cuts to Shift-C so that it doesn't conflict with Ctrl-L in the viewport

diff --git a/Engine/Source/Editor/MovieSceneTools/Private/TrackEditors/CameraCutTrackEditor.cpp b/Engine/Source/Editor/MovieSceneTools/Private/TrackEditors/CameraCutTrackEditor.cpp
--- a/Engine/Source/Editor/MovieSceneTools/Private/TrackEditors/CameraCutTrackEditor.cpp
+++ b/Engine/Source/Editor/MovieSceneTools/Private/TrackEditors/CameraCutTrackEditor.cpp
@@ -58,7 +58,7 @@ public:
 
 void FCameraCutTrackCommands::RegisterCommands()
 {
-	UI_COMMAND( ToggleLockCamera, "Toggle Lock Camera", "Toggle locking the viewport to the camera cut track.", EUserInterfaceActionType::Button, FInputChord(EModifierKey::Control, EKeys::L) );
+	UI_COMMAND( ToggleLockCamera, "Toggle Lock Camera", "Toggle locking the viewport to the camera cut track.", EUserInterfaceActionType::Button, FInputChord(EModifierKey::Shift, EKeys::C) );
 }
 
 
