Fix crash when dragging an Anim Sequence onto an incorrect Sequencer Track Type.

diff --git a/Engine/Source/Editor/MovieSceneTools/Private/TrackEditors/SkeletalAnimationTrackEditor.cpp b/Engine/Source/Editor/MovieSceneTools/Private/TrackEditors/SkeletalAnimationTrackEditor.cpp
--- a/Engine/Source/Editor/MovieSceneTools/Private/TrackEditors/SkeletalAnimationTrackEditor.cpp
+++ b/Engine/Source/Editor/MovieSceneTools/Private/TrackEditors/SkeletalAnimationTrackEditor.cpp
@@ -1824,20 +1824,21 @@ FKeyPropertyResult FSkeletalAnimationTrackEditor::AddKeyInternal( FFrameNumber K
 	if (ObjectHandle.IsValid())
 	{
 		UMovieScene* MovieScene = GetSequencer()->GetFocusedMovieSceneSequence()->GetMovieScene();
+		UMovieSceneSkeletalAnimationTrack* SkelAnimTrack = Cast<UMovieSceneSkeletalAnimationTrack>(Track);
 		FMovieSceneBinding* Binding = MovieScene->FindBinding(ObjectHandle);
 
 		// Add a track if no track was specified or if the track specified doesn't belong to the tracks of the targeted guid
-		if (!Track || (Binding && !Binding->GetTracks().Contains(Track)))
+		if (!SkelAnimTrack || (Binding && !Binding->GetTracks().Contains(SkelAnimTrack)))
 		{
-			Track = AddTrack(MovieScene, ObjectHandle, UMovieSceneSkeletalAnimationTrack::StaticClass(), NAME_None);
+			SkelAnimTrack = CastChecked<UMovieSceneSkeletalAnimationTrack>(AddTrack(MovieScene, ObjectHandle, UMovieSceneSkeletalAnimationTrack::StaticClass(), NAME_None), ECastCheckedType::NullAllowed);
 			KeyPropertyResult.bTrackCreated = true;
 		}
 
-		if (ensure(Track))
+		if (ensure(SkelAnimTrack))
 		{
-			Track->Modify();
+			SkelAnimTrack->Modify();
 
-			UMovieSceneSection* NewSection = Cast<UMovieSceneSkeletalAnimationTrack>(Track)->AddNewAnimationOnRow( KeyTime, AnimSequence, RowIndex );
+			UMovieSceneSection* NewSection = SkelAnimTrack->AddNewAnimationOnRow(KeyTime, AnimSequence, RowIndex);
 			KeyPropertyResult.bTrackModified = true;
 			KeyPropertyResult.SectionsCreated.Add(NewSection);
 
