CIS: Missing copyright header.
#trivial

diff --git a/Engine/Source/Editor/MovieSceneTools/Private/MovieSceneConstraintChannelHelper.cpp b/Engine/Source/Editor/MovieSceneTools/Private/MovieSceneConstraintChannelHelper.cpp
--- a/Engine/Source/Editor/MovieSceneTools/Private/MovieSceneConstraintChannelHelper.cpp
+++ b/Engine/Source/Editor/MovieSceneTools/Private/MovieSceneConstraintChannelHelper.cpp
@@ -1,4 +1,6 @@
-﻿#include "MovieSceneConstraintChannelHelper.h"
+﻿// Copyright Epic Games, Inc. All Rights Reserved.
+
+#include "MovieSceneConstraintChannelHelper.h"
 #include "ConstraintChannelHelper.inl"
 
 #include "ISequencer.h"
