Sequencer: create a menu extender and pass it down to sequencer context menus

Previously, only float channels were adding an extension to context menus. When double channels were made, they had the same code but only the latest one to add an extension would "win", because menu builders only use the top extender. So now we create one extender with multiple extensions in it. However, this means that now we'll get duplicate sets of menu entries (one for double channels, one for float channels). Plus the menu entries from overriden channels. So we also now put channel menu entries inside sub-menus, in a different context menu section.

diff --git a/Engine/Source/Editor/MovieSceneTools/Private/Channels/BuiltInChannelEditors.cpp b/Engine/Source/Editor/MovieSceneTools/Private/Channels/BuiltInChannelEditors.cpp
--- a/Engine/Source/Editor/MovieSceneTools/Private/Channels/BuiltInChannelEditors.cpp
+++ b/Engine/Source/Editor/MovieSceneTools/Private/Channels/BuiltInChannelEditors.cpp
@@ -1087,7 +1087,7 @@ struct FDoubleChannelKeyMenuExtension : TCurveChannelKeyMenuExtension<FMovieScen
 };
 
 template<typename ChannelType>
-struct TCurveChannelSectionMenuExtension : FExtender, TSharedFromThis<TCurveChannelSectionMenuExtension<ChannelType>>
+struct TCurveChannelSectionMenuExtension : TSharedFromThis<TCurveChannelSectionMenuExtension<ChannelType>>
 {
 	using ChannelValueType = typename ChannelType::ChannelValueType;
 	using CurveValueType = typename ChannelType::CurveValueType;
@@ -1103,6 +1103,11 @@ struct TCurveChannelSectionMenuExtension : FExtender, TSharedFromThis<TCurveChan
 		}
 	}
 
+	virtual ~TCurveChannelSectionMenuExtension() {}
+
+	virtual FText GetSubMenuLabel() const { return FText(); }
+	virtual FText GetSubMenuToolTip() const { return FText(); }
+
 	void ExtendMenu(FMenuBuilder& MenuBuilder)
 	{
 		ISequencer* SequencerPtr = WeakSequencer.Pin().Get();
@@ -1114,22 +1119,28 @@ struct TCurveChannelSectionMenuExtension : FExtender, TSharedFromThis<TCurveChan
 		TSharedRef<TCurveChannelSectionMenuExtension<ChannelType>> SharedThis = this->AsShared();
 
 		MenuBuilder.AddSubMenu(
-			LOCTEXT("SetPreInfinityExtrap", "Pre-Infinity"),
-			LOCTEXT("SetPreInfinityExtrapTooltip", "Set pre-infinity extrapolation"),
-			FNewMenuDelegate::CreateLambda([SharedThis](FMenuBuilder& SubMenuBuilder){ SharedThis->AddExtrapolationMenu(SubMenuBuilder, true); })
-			);
+			GetSubMenuLabel(),
+			GetSubMenuToolTip(),
+			FNewMenuDelegate::CreateLambda([SharedThis](FMenuBuilder& SubMenuBuilder)
+			{
+				SubMenuBuilder.AddSubMenu(
+						LOCTEXT("SetPreInfinityExtrap", "Pre-Infinity"),
+						LOCTEXT("SetPreInfinityExtrapTooltip", "Set pre-infinity extrapolation"),
+						FNewMenuDelegate::CreateLambda([SharedThis](FMenuBuilder& SubMenuBuilder){ SharedThis->AddExtrapolationMenu(SubMenuBuilder, true); })
+						);
 
-		MenuBuilder.AddSubMenu(
-			LOCTEXT("SetPostInfinityExtrap", "Post-Infinity"),
-			LOCTEXT("SetPostInfinityExtrapTooltip", "Set post-infinity extrapolation"),
-			FNewMenuDelegate::CreateLambda([SharedThis](FMenuBuilder& SubMenuBuilder){ SharedThis->AddExtrapolationMenu(SubMenuBuilder, false); })
-			);
+				SubMenuBuilder.AddSubMenu(
+						LOCTEXT("SetPostInfinityExtrap", "Post-Infinity"),
+						LOCTEXT("SetPostInfinityExtrapTooltip", "Set post-infinity extrapolation"),
+						FNewMenuDelegate::CreateLambda([SharedThis](FMenuBuilder& SubMenuBuilder){ SharedThis->AddExtrapolationMenu(SubMenuBuilder, false); })
+						);
 
-		MenuBuilder.AddSubMenu(
-			LOCTEXT("DisplayOpyions", "Display"),
-			LOCTEXT("DisplayOptionsTooltip", "Display options"),
-			FNewMenuDelegate::CreateLambda([SharedThis](FMenuBuilder& SubMenuBuilder){ SharedThis->AddDisplayOptionsMenu(SubMenuBuilder); })
-			);
+				SubMenuBuilder.AddSubMenu(
+						LOCTEXT("DisplayOpyions", "Display"),
+						LOCTEXT("DisplayOptionsTooltip", "Display options"),
+						FNewMenuDelegate::CreateLambda([SharedThis](FMenuBuilder& SubMenuBuilder){ SharedThis->AddDisplayOptionsMenu(SubMenuBuilder); })
+						);
+			}));
 	}
 
 	void AddDisplayOptionsMenu(FMenuBuilder& MenuBuilder)
@@ -1486,6 +1497,9 @@ struct FFloatChannelSectionMenuExtension : TCurveChannelSectionMenuExtension<FMo
 	FFloatChannelSectionMenuExtension(TWeakPtr<ISequencer> InSequencer, TArray<TMovieSceneChannelHandle<FMovieSceneFloatChannel>>&& InChannels, TArrayView<UMovieSceneSection* const> InSections)
 		: TCurveChannelSectionMenuExtension<FMovieSceneFloatChannel>(InSequencer, MoveTemp(InChannels), InSections)
 	{}
+
+	virtual FText GetSubMenuLabel() const override { return LOCTEXT("FloatChannelsMenuLabel", "Float Channels"); }
+	virtual FText GetSubMenuToolTip() const override { return LOCTEXT("FloatChannelsMenuToolTip", "Edit parameters for float channels"); }
 };
 
 struct FDoubleChannelSectionMenuExtension : TCurveChannelSectionMenuExtension<FMovieSceneDoubleChannel>
@@ -1493,42 +1507,37 @@ struct FDoubleChannelSectionMenuExtension : TCurveChannelSectionMenuExtension<FM
 	FDoubleChannelSectionMenuExtension(TWeakPtr<ISequencer> InSequencer, TArray<TMovieSceneChannelHandle<FMovieSceneDoubleChannel>>&& InChannels, TArrayView<UMovieSceneSection* const> InSections)
 		: TCurveChannelSectionMenuExtension<FMovieSceneDoubleChannel>(InSequencer, MoveTemp(InChannels), InSections)
 	{}
+
+	virtual FText GetSubMenuLabel() const override { return LOCTEXT("DoubleChannelsMenuLabel", "Double Channels"); }
+	virtual FText GetSubMenuToolTip() const override { return LOCTEXT("DoubleChannelsMenuToolTip", "Edit parameters for double channels"); }
 };
 
-void ExtendSectionMenu(FMenuBuilder& OuterMenuBuilder, TArray<TMovieSceneChannelHandle<FMovieSceneFloatChannel>>&& Channels, TArrayView<UMovieSceneSection* const> Sections, TWeakPtr<ISequencer> InSequencer)
+void ExtendSectionMenu(FMenuBuilder& OuterMenuBuilder, TSharedPtr<FExtender> MenuExtender, TArray<TMovieSceneChannelHandle<FMovieSceneFloatChannel>>&& Channels, TArrayView<UMovieSceneSection* const> Sections, TWeakPtr<ISequencer> InSequencer)
 {
 	TSharedRef<FFloatChannelSectionMenuExtension> Extension = MakeShared<FFloatChannelSectionMenuExtension>(InSequencer, MoveTemp(Channels), Sections);
 
-	Extension->AddMenuExtension("SequencerSections", EExtensionHook::First, nullptr, FMenuExtensionDelegate::CreateLambda([Extension](FMenuBuilder& MenuBuilder) { Extension->ExtendMenu(MenuBuilder); }));
-
-	OuterMenuBuilder.PushExtender(Extension);
+	MenuExtender->AddMenuExtension("SequencerChannels", EExtensionHook::First, nullptr, FMenuExtensionDelegate::CreateLambda([Extension](FMenuBuilder& MenuBuilder) { Extension->ExtendMenu(MenuBuilder); }));
 }
 
-void ExtendSectionMenu(FMenuBuilder& OuterMenuBuilder, TArray<TMovieSceneChannelHandle<FMovieSceneDoubleChannel>>&& Channels, TArrayView<UMovieSceneSection* const> Sections, TWeakPtr<ISequencer> InSequencer)
+void ExtendSectionMenu(FMenuBuilder& OuterMenuBuilder, TSharedPtr<FExtender> MenuExtender, TArray<TMovieSceneChannelHandle<FMovieSceneDoubleChannel>>&& Channels, TArrayView<UMovieSceneSection* const> Sections, TWeakPtr<ISequencer> InSequencer)
 {
 	TSharedRef<FDoubleChannelSectionMenuExtension> Extension = MakeShared<FDoubleChannelSectionMenuExtension>(InSequencer, MoveTemp(Channels), Sections);
 
-	Extension->AddMenuExtension("SequencerSections", EExtensionHook::First, nullptr, FMenuExtensionDelegate::CreateLambda([Extension](FMenuBuilder& MenuBuilder) { Extension->ExtendMenu(MenuBuilder); }));
-
-	OuterMenuBuilder.PushExtender(Extension);
+	MenuExtender->AddMenuExtension("SequencerChannels", EExtensionHook::First, nullptr, FMenuExtensionDelegate::CreateLambda([Extension](FMenuBuilder& MenuBuilder) { Extension->ExtendMenu(MenuBuilder); }));
 }
 
-void ExtendKeyMenu(FMenuBuilder& OuterMenuBuilder, TArray<TExtendKeyMenuParams<FMovieSceneFloatChannel>>&& Channels, TWeakPtr<ISequencer> InSequencer)
+void ExtendKeyMenu(FMenuBuilder& OuterMenuBuilder, TSharedPtr<FExtender> MenuExtender, TArray<TExtendKeyMenuParams<FMovieSceneFloatChannel>>&& Channels, TWeakPtr<ISequencer> InSequencer)
 {
 	TSharedRef<FFloatChannelKeyMenuExtension> Extension = MakeShared<FFloatChannelKeyMenuExtension>(InSequencer, MoveTemp(Channels));
 
-	Extension->AddMenuExtension("SequencerKeyEdit", EExtensionHook::After, nullptr, FMenuExtensionDelegate::CreateLambda([Extension](FMenuBuilder& MenuBuilder) { Extension->ExtendMenu(MenuBuilder); }));
-
-	OuterMenuBuilder.PushExtender(Extension);
+	MenuExtender->AddMenuExtension("SequencerKeyEdit", EExtensionHook::After, nullptr, FMenuExtensionDelegate::CreateLambda([Extension](FMenuBuilder& MenuBuilder) { Extension->ExtendMenu(MenuBuilder); }));
 }
 
-void ExtendKeyMenu(FMenuBuilder& OuterMenuBuilder, TArray<TExtendKeyMenuParams<FMovieSceneDoubleChannel>>&& Channels, TWeakPtr<ISequencer> InSequencer)
+void ExtendKeyMenu(FMenuBuilder& OuterMenuBuilder, TSharedPtr<FExtender> MenuExtender, TArray<TExtendKeyMenuParams<FMovieSceneDoubleChannel>>&& Channels, TWeakPtr<ISequencer> InSequencer)
 {
 	TSharedRef<FDoubleChannelKeyMenuExtension> Extension = MakeShared<FDoubleChannelKeyMenuExtension>(InSequencer, MoveTemp(Channels));
 
-	Extension->AddMenuExtension("SequencerKeyEdit", EExtensionHook::After, nullptr, FMenuExtensionDelegate::CreateLambda([Extension](FMenuBuilder& MenuBuilder) { Extension->ExtendMenu(MenuBuilder); }));
-
-	OuterMenuBuilder.PushExtender(Extension);
+	MenuExtender->AddMenuExtension("SequencerKeyEdit", EExtensionHook::After, nullptr, FMenuExtensionDelegate::CreateLambda([Extension](FMenuBuilder& MenuBuilder) { Extension->ExtendMenu(MenuBuilder); }));
 }
 
 TUniquePtr<FCurveModel> CreateCurveEditorModel(const TMovieSceneChannelHandle<FMovieSceneFloatChannel>& FloatChannel, UMovieSceneSection* OwningSection, TSharedRef<ISequencer> InSequencer)
diff --git a/Engine/Source/Editor/MovieSceneTools/Private/Channels/BuiltInChannelEditors.h b/Engine/Source/Editor/MovieSceneTools/Private/Channels/BuiltInChannelEditors.h
--- a/Engine/Source/Editor/MovieSceneTools/Private/Channels/BuiltInChannelEditors.h
+++ b/Engine/Source/Editor/MovieSceneTools/Private/Channels/BuiltInChannelEditors.h
@@ -69,10 +69,10 @@ void DrawKeys(FMovieSceneParticleChannel* Channel, TArrayView<const FKeyHandle>
 void DrawKeys(FMovieSceneEventChannel*    Channel, TArrayView<const FKeyHandle> InKeyHandles, const UMovieSceneSection* InOwner, TArrayView<FKeyDrawParams> OutKeyDrawParams);
 
 /** Context menu overrides */
-void ExtendSectionMenu(FMenuBuilder& OuterMenuBuilder, TArray<TMovieSceneChannelHandle<FMovieSceneDoubleChannel>>&& Channels, TArrayView<UMovieSceneSection* const> Sections, TWeakPtr<ISequencer> InSequencer);
-void ExtendSectionMenu(FMenuBuilder& OuterMenuBuilder, TArray<TMovieSceneChannelHandle<FMovieSceneFloatChannel>>&& Channels, TArrayView<UMovieSceneSection* const> Sections, TWeakPtr<ISequencer> InSequencer);
-void ExtendKeyMenu(FMenuBuilder& OuterMenuBuilder, TArray<TExtendKeyMenuParams<FMovieSceneDoubleChannel>>&& Channels, TWeakPtr<ISequencer> InSequencer);
-void ExtendKeyMenu(FMenuBuilder& OuterMenuBuilder, TArray<TExtendKeyMenuParams<FMovieSceneFloatChannel>>&& Channels, TWeakPtr<ISequencer> InSequencer);
+void ExtendSectionMenu(FMenuBuilder& OuterMenuBuilder, TSharedPtr<FExtender> MenuExtender, TArray<TMovieSceneChannelHandle<FMovieSceneDoubleChannel>>&& Channels, TArrayView<UMovieSceneSection* const> Sections, TWeakPtr<ISequencer> InSequencer);
+void ExtendSectionMenu(FMenuBuilder& OuterMenuBuilder, TSharedPtr<FExtender> MenuExtender, TArray<TMovieSceneChannelHandle<FMovieSceneFloatChannel>>&& Channels, TArrayView<UMovieSceneSection* const> Sections, TWeakPtr<ISequencer> InSequencer);
+void ExtendKeyMenu(FMenuBuilder& OuterMenuBuilder, TSharedPtr<FExtender> MenuExtender, TArray<TExtendKeyMenuParams<FMovieSceneDoubleChannel>>&& Channels, TWeakPtr<ISequencer> InSequencer);
+void ExtendKeyMenu(FMenuBuilder& OuterMenuBuilder, TSharedPtr<FExtender> MenuExtender, TArray<TExtendKeyMenuParams<FMovieSceneFloatChannel>>&& Channels, TWeakPtr<ISequencer> InSequencer);
 
 /** Curve editor models */
 inline bool SupportsCurveEditorModels(const TMovieSceneChannelHandle<FMovieSceneDoubleChannel>& DoubleChannel) { return true; }
diff --git a/Engine/Source/Editor/Sequencer/Private/SSequencer.cpp b/Engine/Source/Editor/Sequencer/Private/SSequencer.cpp
--- a/Engine/Source/Editor/Sequencer/Private/SSequencer.cpp
+++ b/Engine/Source/Editor/Sequencer/Private/SSequencer.cpp
@@ -3826,9 +3826,10 @@ bool SSequencer::OpenPasteMenu()
 	}
 
 	const bool bShouldCloseWindowAfterMenuSelection = true;
-	FMenuBuilder MenuBuilder(bShouldCloseWindowAfterMenuSelection, SequencerPtr.Pin()->GetCommandBindings());
+	TSharedPtr<FExtender> MenuExtender = MakeShared<FExtender>();
+	FMenuBuilder MenuBuilder(bShouldCloseWindowAfterMenuSelection, SequencerPtr.Pin()->GetCommandBindings(), MenuExtender);
 
-	ContextMenu->PopulateMenu(MenuBuilder);
+	ContextMenu->PopulateMenu(MenuBuilder, MenuExtender);
 
 	FWidgetPath Path;
 	FSlateApplication::Get().FindPathToWidget(AsShared(), Path);
@@ -3858,9 +3859,10 @@ void SSequencer::PasteFromHistory()
 	if (ContextMenu.IsValid())
 	{
 		const bool bShouldCloseWindowAfterMenuSelection = true;
-		FMenuBuilder MenuBuilder(bShouldCloseWindowAfterMenuSelection, Sequencer->GetCommandBindings());
+		TSharedPtr<FExtender> MenuExtender = MakeShared<FExtender>();
+		FMenuBuilder MenuBuilder(bShouldCloseWindowAfterMenuSelection, Sequencer->GetCommandBindings(), MenuExtender);
 
-		ContextMenu->PopulateMenu(MenuBuilder);
+		ContextMenu->PopulateMenu(MenuBuilder, MenuExtender);
 
 		FWidgetPath Path;
 		FSlateApplication::Get().FindPathToWidget(AsShared(), Path);
diff --git a/Engine/Source/Editor/Sequencer/Private/SequencerCommonHelpers.cpp b/Engine/Source/Editor/Sequencer/Private/SequencerCommonHelpers.cpp
--- a/Engine/Source/Editor/Sequencer/Private/SequencerCommonHelpers.cpp
+++ b/Engine/Source/Editor/Sequencer/Private/SequencerCommonHelpers.cpp
@@ -238,11 +238,14 @@ TSharedPtr<SWidget> SequencerHelpers::SummonContextMenu(FSequencer& Sequencer, c
 	const bool bInRecursivelySearchable = false;
 
 	const bool bShouldCloseWindowAfterMenuSelection = true;
-	FMenuBuilder MenuBuilder(bShouldCloseWindowAfterMenuSelection, Sequencer.GetCommandBindings(), nullptr, false, &FCoreStyle::Get(), true, NAME_None, bInRecursivelySearchable);
+
+	TSharedPtr<FExtender> MenuExtender = MakeShared<FExtender>();
+
+	FMenuBuilder MenuBuilder(bShouldCloseWindowAfterMenuSelection, Sequencer.GetCommandBindings(), MenuExtender, false, &FCoreStyle::Get(), true, NAME_None, bInRecursivelySearchable);
 
 	TSharedPtr<ITrackAreaHotspot> Hotspot = Sequencer.GetViewModel()->GetTrackArea()->GetHotspot();
 
-	if (Hotspot.IsValid() && Hotspot->PopulateContextMenu(MenuBuilder, PasteAtTime))
+	if (Hotspot.IsValid() && Hotspot->PopulateContextMenu(MenuBuilder, MenuExtender, PasteAtTime))
 	{
 		return MenuBuilder.MakeWidget();
 	}
@@ -251,7 +254,7 @@ TSharedPtr<SWidget> SequencerHelpers::SummonContextMenu(FSequencer& Sequencer, c
 		TSharedPtr<FPasteContextMenu> PasteMenu = FPasteContextMenu::CreateMenu(Sequencer, SequencerWidget->GeneratePasteArgs(PasteAtTime));
 		if (PasteMenu.IsValid() && PasteMenu->IsValidPaste())
 		{
-			PasteMenu->PopulateMenu(MenuBuilder);
+			PasteMenu->PopulateMenu(MenuBuilder, MenuExtender);
 
 			return MenuBuilder.MakeWidget();
 		}
diff --git a/Engine/Source/Editor/Sequencer/Private/SequencerContextMenus.cpp b/Engine/Source/Editor/Sequencer/Private/SequencerContextMenus.cpp
--- a/Engine/Source/Editor/Sequencer/Private/SequencerContextMenus.cpp
+++ b/Engine/Source/Editor/Sequencer/Private/SequencerContextMenus.cpp
@@ -98,13 +98,13 @@ static void CreateKeyStructForSelection(TSharedPtr<ISequencer> InSequencer, TSha
 	}
 }
 
-void FKeyContextMenu::BuildMenu(FMenuBuilder& MenuBuilder, FSequencer& InSequencer)
+void FKeyContextMenu::BuildMenu(FMenuBuilder& MenuBuilder, TSharedPtr<FExtender> MenuExtender, FSequencer& InSequencer)
 {
 	TSharedRef<FKeyContextMenu> Menu = MakeShareable(new FKeyContextMenu(InSequencer));
-	Menu->PopulateMenu(MenuBuilder);
+	Menu->PopulateMenu(MenuBuilder, MenuExtender);
 }
 
-void FKeyContextMenu::PopulateMenu(FMenuBuilder& MenuBuilder)
+void FKeyContextMenu::PopulateMenu(FMenuBuilder& MenuBuilder, TSharedPtr<FExtender> MenuExtender)
 {
 	using namespace UE::Sequencer;
 
@@ -134,7 +134,7 @@ void FKeyContextMenu::PopulateMenu(FMenuBuilder& MenuBuilder)
 			ISequencerChannelInterface* ChannelInterface = SequencerModule.FindChannelEditorInterface(Pair.Key);
 			if (ChannelInterface)
 			{
-				ChannelInterface->ExtendKeyMenu_Raw(MenuBuilder, MoveTemp(Pair.Value), Sequencer);
+				ChannelInterface->ExtendKeyMenu_Raw(MenuBuilder, MenuExtender, MoveTemp(Pair.Value), Sequencer);
 			}
 		}
 	}
@@ -247,14 +247,14 @@ FSectionContextMenu::FSectionContextMenu(FSequencer& InSeqeuncer, FFrameTime InM
 	}
 }
 
-void FSectionContextMenu::BuildMenu(FMenuBuilder& MenuBuilder, FSequencer& InSequencer, FFrameTime InMouseDownTime)
+void FSectionContextMenu::BuildMenu(FMenuBuilder& MenuBuilder, TSharedPtr<FExtender> MenuExtender, FSequencer& InSequencer, FFrameTime InMouseDownTime)
 {
 	TSharedRef<FSectionContextMenu> Menu = MakeShareable(new FSectionContextMenu(InSequencer, InMouseDownTime));
-	Menu->PopulateMenu(MenuBuilder);
+	Menu->PopulateMenu(MenuBuilder, MenuExtender);
 }
 
 
-void FSectionContextMenu::PopulateMenu(FMenuBuilder& MenuBuilder)
+void FSectionContextMenu::PopulateMenu(FMenuBuilder& MenuBuilder, TSharedPtr<FExtender> MenuExtender)
 {
 	// Copy a reference to the context menu by value into each lambda handler to ensure the type stays alive until the menu is closed
 	TSharedRef<FSectionContextMenu> Shared = AsShared();
@@ -275,7 +275,7 @@ void FSectionContextMenu::PopulateMenu(FMenuBuilder& MenuBuilder)
 		ISequencerChannelInterface* ChannelInterface = SequencerModule.FindChannelEditorInterface(Pair.Key);
 		if (ChannelInterface)
 		{
-			ChannelInterface->ExtendSectionMenu_Raw(MenuBuilder, Pair.Value, Sections, Sequencer);
+			ChannelInterface->ExtendSectionMenu_Raw(MenuBuilder, MenuExtender, Pair.Value, Sections, Sequencer);
 		}
 	}
 
@@ -313,7 +313,7 @@ void FSectionContextMenu::PopulateMenu(FMenuBuilder& MenuBuilder)
 		MenuBuilder.AddSubMenu(
 			LOCTEXT("Paste", "Paste"),
 			FText(),
-			FNewMenuDelegate::CreateLambda([=](FMenuBuilder& SubMenuBuilder){ if (PasteMenu.IsValid()) { PasteMenu->PopulateMenu(SubMenuBuilder); } }),
+			FNewMenuDelegate::CreateLambda([=](FMenuBuilder& SubMenuBuilder){ if (PasteMenu.IsValid()) { PasteMenu->PopulateMenu(SubMenuBuilder, MenuExtender); } }),
 			FUIAction (
 				FExecuteAction(),
 				FCanExecuteAction::CreateLambda([=]{ return PasteMenu.IsValid() && PasteMenu->IsValidPaste(); })
@@ -325,7 +325,7 @@ void FSectionContextMenu::PopulateMenu(FMenuBuilder& MenuBuilder)
 		MenuBuilder.AddSubMenu(
 			LOCTEXT("PasteFromHistory", "Paste From History"),
 			FText(),
-			FNewMenuDelegate::CreateLambda([=](FMenuBuilder& SubMenuBuilder){ if (PasteFromHistoryMenu.IsValid()) { PasteFromHistoryMenu->PopulateMenu(SubMenuBuilder); } }),
+			FNewMenuDelegate::CreateLambda([=](FMenuBuilder& SubMenuBuilder){ if (PasteFromHistoryMenu.IsValid()) { PasteFromHistoryMenu->PopulateMenu(SubMenuBuilder, MenuExtender); } }),
 			FUIAction (
 				FExecuteAction(),
 				FCanExecuteAction::CreateLambda([=]{ return PasteFromHistoryMenu.IsValid(); })
@@ -336,6 +336,11 @@ void FSectionContextMenu::PopulateMenu(FMenuBuilder& MenuBuilder)
 	}
 	MenuBuilder.EndSection(); // SequencerKeyEdit
 
+	MenuBuilder.BeginSection("SequencerChannels", LOCTEXT("ChannelsMenu", "Channels"));
+	{
+	}
+	MenuBuilder.EndSection(); // SequencerChannels
+
 	MenuBuilder.BeginSection("SequencerSections", LOCTEXT("SectionsMenu", "Sections"));
 	{
 		if (CanSelectAllKeys())
@@ -1290,7 +1295,7 @@ void FSectionContextMenu::SendBackward()
 }
 
 
-bool FPasteContextMenu::BuildMenu(FMenuBuilder& MenuBuilder, FSequencer& InSequencer, const FPasteContextMenuArgs& Args)
+bool FPasteContextMenu::BuildMenu(FMenuBuilder& MenuBuilder, TSharedPtr<FExtender> MenuExtender, FSequencer& InSequencer, const FPasteContextMenuArgs& Args)
 {
 	TSharedRef<FPasteContextMenu> Menu = MakeShareable(new FPasteContextMenu(InSequencer, Args));
 	Menu->Setup();
@@ -1299,7 +1304,7 @@ bool FPasteContextMenu::BuildMenu(FMenuBuilder& MenuBuilder, FSequencer& InSeque
 		return false;
 	}
 
-	Menu->PopulateMenu(MenuBuilder);
+	Menu->PopulateMenu(MenuBuilder, MenuExtender);
 	return true;
 }
 
@@ -1525,7 +1530,7 @@ bool FPasteContextMenu::IsValidPaste() const
 }
 
 
-void FPasteContextMenu::PopulateMenu(FMenuBuilder& MenuBuilder)
+void FPasteContextMenu::PopulateMenu(FMenuBuilder& MenuBuilder, TSharedPtr<FExtender> MenuExtender)
 {
 	// Copy a reference to the context menu by value into each lambda handler to ensure the type stays alive until the menu is closed
 	TSharedRef<FPasteContextMenu> Shared = AsShared();
@@ -1652,7 +1657,7 @@ bool FPasteContextMenu::PasteInto(int32 DestinationIndex, FName KeyAreaName, TSe
 }
 
 
-bool FPasteFromHistoryContextMenu::BuildMenu(FMenuBuilder& MenuBuilder, FSequencer& InSequencer, const FPasteContextMenuArgs& Args)
+bool FPasteFromHistoryContextMenu::BuildMenu(FMenuBuilder& MenuBuilder, TSharedPtr<FExtender> MenuExtender, FSequencer& InSequencer, const FPasteContextMenuArgs& Args)
 {
 	if (InSequencer.GetClipboardStack().Num() == 0)
 	{
@@ -1660,7 +1665,7 @@ bool FPasteFromHistoryContextMenu::BuildMenu(FMenuBuilder& MenuBuilder, FSequenc
 	}
 
 	TSharedRef<FPasteFromHistoryContextMenu> Menu = MakeShareable(new FPasteFromHistoryContextMenu(InSequencer, Args));
-	Menu->PopulateMenu(MenuBuilder);
+	Menu->PopulateMenu(MenuBuilder, MenuExtender);
 	return true;
 }
 
@@ -1676,7 +1681,7 @@ TSharedPtr<FPasteFromHistoryContextMenu> FPasteFromHistoryContextMenu::CreateMen
 }
 
 
-void FPasteFromHistoryContextMenu::PopulateMenu(FMenuBuilder& MenuBuilder)
+void FPasteFromHistoryContextMenu::PopulateMenu(FMenuBuilder& MenuBuilder, TSharedPtr<FExtender> MenuExtender)
 {
 	// Copy a reference to the context menu by value into each lambda handler to ensure the type stays alive until the menu is closed
 	TSharedRef<FPasteFromHistoryContextMenu> Shared = AsShared();
@@ -1693,7 +1698,7 @@ void FPasteFromHistoryContextMenu::PopulateMenu(FMenuBuilder& MenuBuilder)
 		MenuBuilder.AddSubMenu(
 			ThisPasteArgs.Clipboard->GetDisplayText(),
 			FText(),
-			FNewMenuDelegate::CreateLambda([=](FMenuBuilder& SubMenuBuilder){ PasteMenu->PopulateMenu(SubMenuBuilder); }),
+			FNewMenuDelegate::CreateLambda([=](FMenuBuilder& SubMenuBuilder){ PasteMenu->PopulateMenu(SubMenuBuilder, MenuExtender); }),
 			FUIAction (
 				FExecuteAction(),
 				FCanExecuteAction::CreateLambda([=]{ return PasteMenu->IsValidPaste(); })
@@ -1706,17 +1711,17 @@ void FPasteFromHistoryContextMenu::PopulateMenu(FMenuBuilder& MenuBuilder)
 	MenuBuilder.EndSection();
 }
 
-void FEasingContextMenu::BuildMenu(FMenuBuilder& MenuBuilder, const TArray<UE::Sequencer::FEasingAreaHandle>& InEasings, FSequencer& Sequencer, FFrameTime InMouseDownTime)
+void FEasingContextMenu::BuildMenu(FMenuBuilder& MenuBuilder, TSharedPtr<FExtender> MenuExtender, const TArray<UE::Sequencer::FEasingAreaHandle>& InEasings, FSequencer& Sequencer, FFrameTime InMouseDownTime)
 {
 	TSharedRef<FEasingContextMenu> EasingMenu = MakeShareable(new FEasingContextMenu(InEasings, Sequencer));
-	EasingMenu->PopulateMenu(MenuBuilder);
+	EasingMenu->PopulateMenu(MenuBuilder, MenuExtender);
 
 	MenuBuilder.AddMenuSeparator();
 
-	FSectionContextMenu::BuildMenu(MenuBuilder, Sequencer, InMouseDownTime);
+	FSectionContextMenu::BuildMenu(MenuBuilder, MenuExtender, Sequencer, InMouseDownTime);
 }
 
-void FEasingContextMenu::PopulateMenu(FMenuBuilder& MenuBuilder)
+void FEasingContextMenu::PopulateMenu(FMenuBuilder& MenuBuilder, TSharedPtr<FExtender> MenuExtender)
 {
 	using namespace UE::Sequencer;
 
diff --git a/Engine/Source/Editor/Sequencer/Private/SequencerContextMenus.h b/Engine/Source/Editor/Sequencer/Private/SequencerContextMenus.h
--- a/Engine/Source/Editor/Sequencer/Private/SequencerContextMenus.h
+++ b/Engine/Source/Editor/Sequencer/Private/SequencerContextMenus.h
@@ -14,6 +14,7 @@
 
 struct FSequencerSelectedKey;
 class FMenuBuilder;
+class FExtender;
 class UMovieSceneSection;
 
 /**
@@ -24,7 +25,7 @@ class UMovieSceneSection;
  */
 struct FSectionContextMenu : TSharedFromThis<FSectionContextMenu>
 {
-	static void BuildMenu(FMenuBuilder& MenuBuilder, FSequencer& Sequencer, FFrameTime InMouseDownTime);
+	static void BuildMenu(FMenuBuilder& MenuBuilder, TSharedPtr<FExtender> MenuExtender, FSequencer& Sequencer, FFrameTime InMouseDownTime);
 
 private:
 
@@ -33,7 +34,7 @@ private:
 
 	FSectionContextMenu(FSequencer& InSeqeuncer, FFrameTime InMouseDownTime);
 
-	void PopulateMenu(FMenuBuilder& MenuBuilder);
+	void PopulateMenu(FMenuBuilder& MenuBuilder, TSharedPtr<FExtender> MenuExtender);
 
 	/** Add edit menu for trim and split */
 	void AddEditMenu(FMenuBuilder& MenuBuilder);
@@ -129,11 +130,11 @@ struct FPasteContextMenuArgs
 
 struct FPasteContextMenu : TSharedFromThis<FPasteContextMenu>
 {
-	static bool BuildMenu(FMenuBuilder& MenuBuilder, FSequencer& Sequencer, const FPasteContextMenuArgs& Args);
+	static bool BuildMenu(FMenuBuilder& MenuBuilder, TSharedPtr<FExtender> MenuExtender, FSequencer& Sequencer, const FPasteContextMenuArgs& Args);
 
 	static TSharedRef<FPasteContextMenu> CreateMenu(FSequencer& Sequencer, const FPasteContextMenuArgs& Args);
 
-	void PopulateMenu(FMenuBuilder& MenuBuilder);
+	void PopulateMenu(FMenuBuilder& MenuBuilder, TSharedPtr<FExtender> MenuExtender);
 
 	bool IsValidPaste() const;
 
@@ -178,11 +179,11 @@ private:
 
 struct FPasteFromHistoryContextMenu : TSharedFromThis<FPasteFromHistoryContextMenu>
 {
-	static bool BuildMenu(FMenuBuilder& MenuBuilder, FSequencer& Sequencer, const FPasteContextMenuArgs& Args);
+	static bool BuildMenu(FMenuBuilder& MenuBuilder, TSharedPtr<FExtender> MenuExtender, FSequencer& Sequencer, const FPasteContextMenuArgs& Args);
 
 	static TSharedPtr<FPasteFromHistoryContextMenu> CreateMenu(FSequencer& Sequencer, const FPasteContextMenuArgs& Args);
 
-	void PopulateMenu(FMenuBuilder& MenuBuilder);
+	void PopulateMenu(FMenuBuilder& MenuBuilder, TSharedPtr<FExtender> MenuExtender);
 
 private:
 
@@ -208,7 +209,7 @@ private:
  */
 struct FKeyContextMenu : TSharedFromThis<FKeyContextMenu>
 {
-	static void BuildMenu(FMenuBuilder& MenuBuilder, FSequencer& Sequencer);
+	static void BuildMenu(FMenuBuilder& MenuBuilder, TSharedPtr<FExtender> MenuExtender, FSequencer& Sequencer);
 
 private:
 
@@ -222,7 +223,7 @@ private:
 	/** Add the Properties sub-menu. */
 	void AddPropertiesMenu(FMenuBuilder& MenuBuilder);
 
-	void PopulateMenu(FMenuBuilder& MenuBuilder);
+	void PopulateMenu(FMenuBuilder& MenuBuilder, TSharedPtr<FExtender> MenuExtender);
 
 	/** The sequencer */
 	TSharedRef<FSequencer> Sequencer;
@@ -240,7 +241,7 @@ private:
  */
 struct FEasingContextMenu : TSharedFromThis<FEasingContextMenu>
 {
-	static void BuildMenu(FMenuBuilder& MenuBuilder, const TArray<UE::Sequencer::FEasingAreaHandle>& InEasings, FSequencer& Sequencer, FFrameTime InMouseDownTime);
+	static void BuildMenu(FMenuBuilder& MenuBuilder, TSharedPtr<FExtender> MenuExtender, const TArray<UE::Sequencer::FEasingAreaHandle>& InEasings, FSequencer& Sequencer, FFrameTime InMouseDownTime);
 
 private:
 
@@ -253,7 +254,7 @@ private:
 	/** Hidden AsShared() methods to discourage CreateSP delegate use. */
 	using TSharedFromThis::AsShared;
 
-	void PopulateMenu(FMenuBuilder& MenuBuilder);
+	void PopulateMenu(FMenuBuilder& MenuBuilder, TSharedPtr<FExtender> MenuExtender);
 
 	FText GetEasingTypeText() const;
 
diff --git a/Engine/Source/Editor/Sequencer/Private/SequencerHotspots.cpp b/Engine/Source/Editor/Sequencer/Private/SequencerHotspots.cpp
--- a/Engine/Source/Editor/Sequencer/Private/SequencerHotspots.cpp
+++ b/Engine/Source/Editor/Sequencer/Private/SequencerHotspots.cpp
@@ -172,11 +172,11 @@ TOptional<FFrameNumber> FKeyHotspot::GetTime() const
 	return Time;
 }
 
-bool FKeyHotspot::PopulateContextMenu(FMenuBuilder& MenuBuilder, FFrameTime MouseDownTime)
+bool FKeyHotspot::PopulateContextMenu(FMenuBuilder& MenuBuilder, TSharedPtr<FExtender> MenuExtender, FFrameTime MouseDownTime)
 {
 	if (TSharedPtr<FSequencer> Sequencer = WeakSequencer.Pin())
 	{
-		FKeyContextMenu::BuildMenu(MenuBuilder, *Sequencer);
+		FKeyContextMenu::BuildMenu(MenuBuilder, MenuExtender, *Sequencer);
 	}
 	return true;
 }
@@ -250,14 +250,14 @@ void FSectionHotspotBase::UpdateOnHover(FTrackAreaViewModel& InTrackArea) const
 	}
 }
 
-bool FSectionHotspotBase::PopulateContextMenu(FMenuBuilder& MenuBuilder, FFrameTime MouseDownTime)
+bool FSectionHotspotBase::PopulateContextMenu(FMenuBuilder& MenuBuilder, TSharedPtr<FExtender> MenuExtender, FFrameTime MouseDownTime)
 {
 	TSharedPtr<FSectionModel> SectionModel = WeakSectionModel.Pin();
 	UMovieSceneSection*       ThisSection  = SectionModel ? SectionModel->GetSection() : nullptr;
 	if (ThisSection)
 	{
 		TSharedPtr<FSequencer> Sequencer = WeakSequencer.Pin();
-		FSectionContextMenu::BuildMenu(MenuBuilder, *Sequencer, MouseDownTime);
+		FSectionContextMenu::BuildMenu(MenuBuilder, MenuExtender, *Sequencer, MouseDownTime);
 
 		TSharedPtr<IObjectBindingExtension> ObjectBinding = SectionModel->FindAncestorOfType<IObjectBindingExtension>();
 		SectionModel->GetSectionInterface()->BuildSectionContextMenu(MenuBuilder, ObjectBinding ? ObjectBinding->GetObjectGuid() : FGuid());
@@ -341,11 +341,11 @@ void FSectionEasingHandleHotspot::UpdateOnHover(FTrackAreaViewModel& InTrackArea
 	InTrackArea.AttemptToActivateTool(FSequencerEditTool_Movement::Identifier);
 }
 
-bool FSectionEasingHandleHotspot::PopulateContextMenu(FMenuBuilder& MenuBuilder, FFrameTime MouseDownTime)
+bool FSectionEasingHandleHotspot::PopulateContextMenu(FMenuBuilder& MenuBuilder, TSharedPtr<FExtender> MenuExtender, FFrameTime MouseDownTime)
 {
 	TSharedPtr<FSequencer> Sequencer = WeakSequencer.Pin();
 
-	FEasingContextMenu::BuildMenu(MenuBuilder, { FEasingAreaHandle{WeakSectionModel, HandleType} }, *Sequencer, MouseDownTime);
+	FEasingContextMenu::BuildMenu(MenuBuilder, MenuExtender, { FEasingAreaHandle{WeakSectionModel, HandleType} }, *Sequencer, MouseDownTime);
 	return true;
 }
 
@@ -361,12 +361,12 @@ const FSlateBrush* FSectionEasingHandleHotspot::GetCursorDecorator(const FGeomet
 	return FAppStyle::Get().GetBrush(TEXT("Sequencer.CursorDecorator_EasingHandle"));
 }
 
-bool FSectionEasingAreaHotspot::PopulateContextMenu(FMenuBuilder& MenuBuilder, FFrameTime MouseDownTime)
+bool FSectionEasingAreaHotspot::PopulateContextMenu(FMenuBuilder& MenuBuilder, TSharedPtr<FExtender> MenuExtender, FFrameTime MouseDownTime)
 {
 	using namespace UE::Sequencer;
 
 	TSharedPtr<FSequencer> Sequencer = WeakSequencer.Pin();
-	FEasingContextMenu::BuildMenu(MenuBuilder, Easings, *Sequencer, MouseDownTime);
+	FEasingContextMenu::BuildMenu(MenuBuilder, MenuExtender, Easings, *Sequencer, MouseDownTime);
 
 	TSharedPtr<FSectionModel> SectionModel = WeakSectionModel.Pin();
 	UMovieSceneSection*       ThisSection  = SectionModel ? SectionModel->GetSection() : nullptr;
diff --git a/Engine/Source/Editor/Sequencer/Private/SequencerHotspots.h b/Engine/Source/Editor/Sequencer/Private/SequencerHotspots.h
--- a/Engine/Source/Editor/Sequencer/Private/SequencerHotspots.h
+++ b/Engine/Source/Editor/Sequencer/Private/SequencerHotspots.h
@@ -70,7 +70,7 @@ struct FKeyHotspot
 
 	virtual void UpdateOnHover(FTrackAreaViewModel& InTrackArea) const override;
 	virtual TOptional<FFrameNumber> GetTime() const override;
-	virtual bool PopulateContextMenu(FMenuBuilder& MenuBuilder, FFrameTime MouseDownTime) override;
+	virtual bool PopulateContextMenu(FMenuBuilder& MenuBuilder, TSharedPtr<FExtender> MenuExtender, FFrameTime MouseDownTime) override;
 
 	virtual void HandleMouseSelection(FHotspotSelectionManager& SelectionManager) override;
 
@@ -94,7 +94,7 @@ struct FSectionHotspotBase
 	virtual TOptional<FFrameNumber> GetTime() const override;
 	virtual TOptional<FFrameTime> GetOffsetTime() const override;
 	virtual TSharedPtr<ISequencerEditToolDragOperation> InitiateDrag(const FPointerEvent& MouseEvent) override { return nullptr; }
-	virtual bool PopulateContextMenu(FMenuBuilder& MenuBuilder, FFrameTime MouseDownTime) override;
+	virtual bool PopulateContextMenu(FMenuBuilder& MenuBuilder, TSharedPtr<FExtender> MenuExtender, FFrameTime MouseDownTime) override;
 
 	/** IMouseHandlerHotspot */
 	virtual void HandleMouseSelection(FHotspotSelectionManager& SelectionManager) override;
@@ -157,7 +157,7 @@ struct FSectionEasingHandleHotspot : FSectionHotspotBase
 	{}
 
 	virtual void UpdateOnHover(FTrackAreaViewModel& InTrackArea) const override;
-	virtual bool PopulateContextMenu(FMenuBuilder& MenuBuilder, FFrameTime MouseDownTime) override;
+	virtual bool PopulateContextMenu(FMenuBuilder& MenuBuilder, TSharedPtr<FExtender> MenuExtender, FFrameTime MouseDownTime) override;
 	virtual TOptional<FFrameNumber> GetTime() const override;
 	virtual TSharedPtr<ISequencerEditToolDragOperation> InitiateDrag(const FPointerEvent& MouseEvent) override;
 	virtual FCursorReply GetCursor() const { return FCursorReply::Cursor( EMouseCursor::ResizeLeftRight ); }
@@ -183,7 +183,7 @@ struct FSectionEasingAreaHotspot : FSectionHotspotBase
 		, Easings(InEasings)
 	{}
 
-	virtual bool PopulateContextMenu(FMenuBuilder& MenuBuilder, FFrameTime MouseDownTime) override;
+	virtual bool PopulateContextMenu(FMenuBuilder& MenuBuilder, TSharedPtr<FExtender> MenuExtender, FFrameTime MouseDownTime) override;
 
 	/** IMouseHandlerHotspot */
 	virtual void HandleMouseSelection(FHotspotSelectionManager& SelectionManager) override;
diff --git a/Engine/Source/Editor/Sequencer/Private/Widgets/SSequencerLayerBar.cpp b/Engine/Source/Editor/Sequencer/Private/Widgets/SSequencerLayerBar.cpp
--- a/Engine/Source/Editor/Sequencer/Private/Widgets/SSequencerLayerBar.cpp
+++ b/Engine/Source/Editor/Sequencer/Private/Widgets/SSequencerLayerBar.cpp
@@ -75,7 +75,7 @@ struct FLayerBarHotspot
 
 		return nullptr;
 	}
-	bool PopulateContextMenu(FMenuBuilder& MenuBuilder, FFrameTime MouseDownTime) override
+	bool PopulateContextMenu(FMenuBuilder& MenuBuilder, TSharedPtr<FExtender> MenuExtender, FFrameTime MouseDownTime) override
 	{
 		return true;
 	}
@@ -129,7 +129,7 @@ struct FLayerBarStretchHotspot
 
 		return MakeShared<FEditToolDragOperation_Stretch>(Sequencer.Get(), StretchConstraint, GetTime().GetValue());
 	}
-	bool PopulateContextMenu(FMenuBuilder& MenuBuilder, FFrameTime MouseDownTime) override
+	bool PopulateContextMenu(FMenuBuilder& MenuBuilder, TSharedPtr<FExtender> MenuExtender, FFrameTime MouseDownTime) override
 	{
 		return true;
 	}
diff --git a/Engine/Source/Editor/Sequencer/Public/ISequencerChannelInterface.h b/Engine/Source/Editor/Sequencer/Public/ISequencerChannelInterface.h
--- a/Engine/Source/Editor/Sequencer/Public/ISequencerChannelInterface.h
+++ b/Engine/Source/Editor/Sequencer/Public/ISequencerChannelInterface.h
@@ -18,6 +18,7 @@ struct FSequencerPasteEnvironment;
 class SWidget;
 class ISequencer;
 class FCurveModel;
+class FExtender;
 class FMenuBuilder;
 class FStructOnScope;
 class UMovieSceneSection;
@@ -122,7 +123,7 @@ struct ISequencerChannelInterface
 	 * @param Channels              Array of channels and handles that are being shown in the context menu
 	 * @param InSequencer           The currently active sequencer
 	 */
-	virtual void ExtendKeyMenu_Raw(FMenuBuilder& MenuBuilder, TArrayView<const FExtendKeyMenuParams> Parameters, TWeakPtr<ISequencer> InSequencer) const = 0;
+	virtual void ExtendKeyMenu_Raw(FMenuBuilder& MenuBuilder, TSharedPtr<FExtender> MenuExtender, TArrayView<const FExtendKeyMenuParams> Parameters, TWeakPtr<ISequencer> InSequencer) const = 0;
 
 	/**
 	 * Extend the section context menu
@@ -132,7 +133,7 @@ struct ISequencerChannelInterface
 	 * @param Sections              Array of sections being shown on the context menu
 	 * @param InSequencer           The currently active sequencer
 	 */
-	virtual void ExtendSectionMenu_Raw(FMenuBuilder& MenuBuilder, TArrayView<const FMovieSceneChannelHandle> Channels, TArrayView<UMovieSceneSection* const> Sections, TWeakPtr<ISequencer> InSequencer) const = 0;
+	virtual void ExtendSectionMenu_Raw(FMenuBuilder& MenuBuilder, TSharedPtr<FExtender> MenuExtender, TArrayView<const FMovieSceneChannelHandle> Channels, TArrayView<UMovieSceneSection* const> Sections, TWeakPtr<ISequencer> InSequencer) const = 0;
 
 	/**
 	 * Gather information on how to draw the specified keys
@@ -166,4 +167,4 @@ struct ISequencerChannelInterface
 	 */
 	virtual void DrawExtra_Raw(FMovieSceneChannel* InChannel, const UMovieSceneSection* InOwner, const FGeometry& InKeyGeometry, FSequencerSectionPainter& Painter) const = 0;
 
-};
\ No newline at end of file
+};
diff --git a/Engine/Source/Editor/Sequencer/Public/SequencerChannelInterface.h b/Engine/Source/Editor/Sequencer/Public/SequencerChannelInterface.h
--- a/Engine/Source/Editor/Sequencer/Public/SequencerChannelInterface.h
+++ b/Engine/Source/Editor/Sequencer/Public/SequencerChannelInterface.h
@@ -85,7 +85,7 @@ struct TSequencerChannelInterfaceCommon : ISequencerChannelInterface
 	 * @param Channels              Array of channels and handles that are being shown in the context menu
 	 * @param InSequencer           The currently active sequencer
 	 */
-	virtual void ExtendKeyMenu_Raw(FMenuBuilder& MenuBuilder, TArrayView<const FExtendKeyMenuParams> ChannelsAndHandles, TWeakPtr<ISequencer> InSequencer) const override
+	virtual void ExtendKeyMenu_Raw(FMenuBuilder& MenuBuilder, TSharedPtr<FExtender> MenuExtender, TArrayView<const FExtendKeyMenuParams> ChannelsAndHandles, TWeakPtr<ISequencer> InSequencer) const override
 	{
 		using namespace Sequencer;
 		TArray<TExtendKeyMenuParams<ChannelType>> TypedChannels;
@@ -100,7 +100,7 @@ struct TSequencerChannelInterfaceCommon : ISequencerChannelInterface
 			TypedChannels.Add(MoveTemp(TypedChannelAndHandles));
 		}
 
-		ExtendKeyMenu(MenuBuilder, MoveTemp(TypedChannels), InSequencer);
+		ExtendKeyMenu(MenuBuilder, MenuExtender, MoveTemp(TypedChannels), InSequencer);
 	}
 
 	/**
@@ -111,7 +111,7 @@ struct TSequencerChannelInterfaceCommon : ISequencerChannelInterface
 	 * @param Sections              Array of sections being shown on the context menu
 	 * @param InSequencer           The currently active sequencer
 	 */
-	virtual void ExtendSectionMenu_Raw(FMenuBuilder& MenuBuilder, TArrayView<const FMovieSceneChannelHandle> Channels, TArrayView<UMovieSceneSection* const> Sections, TWeakPtr<ISequencer> InSequencer) const override
+	virtual void ExtendSectionMenu_Raw(FMenuBuilder& MenuBuilder, TSharedPtr<FExtender> MenuExtender, TArrayView<const FMovieSceneChannelHandle> Channels, TArrayView<UMovieSceneSection* const> Sections, TWeakPtr<ISequencer> InSequencer) const override
 	{
 		using namespace Sequencer;
 		TArray<TMovieSceneChannelHandle<ChannelType>> TypedChannels;
@@ -121,7 +121,7 @@ struct TSequencerChannelInterfaceCommon : ISequencerChannelInterface
 			TypedChannels.Add(RawHandle.Cast<ChannelType>());
 		}
 
-		ExtendSectionMenu(MenuBuilder, MoveTemp(TypedChannels), Sections, InSequencer);
+		ExtendSectionMenu(MenuBuilder, MenuExtender, MoveTemp(TypedChannels), Sections, InSequencer);
 	}
 
 	/**
@@ -256,4 +256,4 @@ struct TSequencerChannelInterfaceBase<ChannelType, true> : TSequencerChannelInte
 template<typename ChannelType>
 struct TSequencerChannelInterface : TSequencerChannelInterfaceBase<ChannelType, !TIsSame<typename TMovieSceneChannelTraits<ChannelType>::ExtendedEditorDataType, void>::Value>
 {
-};
\ No newline at end of file
+};
diff --git a/Engine/Source/Editor/Sequencer/Public/SequencerChannelTraits.h b/Engine/Source/Editor/Sequencer/Public/SequencerChannelTraits.h
--- a/Engine/Source/Editor/Sequencer/Public/SequencerChannelTraits.h
+++ b/Engine/Source/Editor/Sequencer/Public/SequencerChannelTraits.h
@@ -48,7 +48,7 @@ namespace Sequencer
 	 * @param InSequencer    The sequencer that is currently active
 	 */
 	template<typename ChannelType>
-	void ExtendSectionMenu(FMenuBuilder& MenuBuilder, TArray<TMovieSceneChannelHandle<ChannelType>>&& Channels, TArrayView<UMovieSceneSection* const> Sections, TWeakPtr<ISequencer> InSequencer)
+	void ExtendSectionMenu(FMenuBuilder& MenuBuilder, TSharedPtr<FExtender> MenuExtender, TArray<TMovieSceneChannelHandle<ChannelType>>&& Channels, TArrayView<UMovieSceneSection* const> Sections, TWeakPtr<ISequencer> InSequencer)
 	{}
 
 
@@ -61,7 +61,7 @@ namespace Sequencer
 	 * @param InSequencer    The sequencer that is currently active
 	 */
 	template<typename ChannelType>
-	void ExtendKeyMenu(FMenuBuilder& MenuBuilder, TArray<TExtendKeyMenuParams<ChannelType>>&& InChannels, TWeakPtr<ISequencer> InSequencer)
+	void ExtendKeyMenu(FMenuBuilder& MenuBuilder, TSharedPtr<FExtender> MenuExtender, TArray<TExtendKeyMenuParams<ChannelType>>&& InChannels, TWeakPtr<ISequencer> InSequencer)
 	{}
 
 
@@ -220,7 +220,7 @@ namespace Sequencer
 	template<typename ChannelType, typename ValueType>
 	FKeyHandle AddOrUpdateKey(
 		ChannelType*                               InChannel,
-		UMovieSceneSection*             SectionToKey,
+		UMovieSceneSection*                        SectionToKey,
 		const TMovieSceneExternalValue<ValueType>& InExternalValue,
 		FFrameNumber                               InTime,
 		ISequencer&                                InSequencer,
@@ -357,4 +357,4 @@ namespace Sequencer
 	 */
 	SEQUENCER_API TUniquePtr<FCurveModel> CreateCurveEditorModel(const FMovieSceneChannelHandle& ChannelHandle, UMovieSceneSection* OwningSection, TSharedRef<ISequencer> InSequencer);
 
-}	// namespace Sequencer
\ No newline at end of file
+}	// namespace Sequencer
diff --git a/Engine/Source/Editor/SequencerCore/Public/MVVM/Views/ITrackAreaHotspot.h b/Engine/Source/Editor/SequencerCore/Public/MVVM/Views/ITrackAreaHotspot.h
--- a/Engine/Source/Editor/SequencerCore/Public/MVVM/Views/ITrackAreaHotspot.h
+++ b/Engine/Source/Editor/SequencerCore/Public/MVVM/Views/ITrackAreaHotspot.h
@@ -12,6 +12,7 @@
 struct FPointerEvent;
 
 class FMenuBuilder;
+class FExtender;
 
 namespace UE
 {
@@ -32,7 +33,7 @@ struct SEQUENCERCORE_API ITrackAreaHotspot
 	virtual TOptional<FFrameNumber> GetTime() const { return TOptional<FFrameNumber>(); }
 	virtual TOptional<FFrameTime> GetOffsetTime() const { return TOptional<FFrameTime>(); }
 	virtual TSharedPtr<ISequencerEditToolDragOperation> InitiateDrag(const FPointerEvent& MouseEvent) { return nullptr; }
-	virtual bool PopulateContextMenu(FMenuBuilder& MenuBuilder, FFrameTime MouseDownTime){ return false; }
+	virtual bool PopulateContextMenu(FMenuBuilder& MenuBuilder, TSharedPtr<FExtender> MenuExtender, FFrameTime MouseDownTime){ return false; }
 	virtual FCursorReply GetCursor() const { return FCursorReply::Unhandled(); }
 	virtual const FSlateBrush* GetCursorDecorator(const FGeometry& MyGeometry, const FPointerEvent& CursorEvent) const { return nullptr; }
 
