Constraints: Fix Undo Issue with Component constraints and fix regression with compensation happening twice do to the wrong componensation guard being used.

diff --git a/Engine/Plugins/Animation/ControlRig/Source/ControlRigEditor/Private/Sequencer/ConstraintChannelHelper.cpp b/Engine/Plugins/Animation/ControlRig/Source/ControlRigEditor/Private/Sequencer/ConstraintChannelHelper.cpp
--- a/Engine/Plugins/Animation/ControlRig/Source/ControlRigEditor/Private/Sequencer/ConstraintChannelHelper.cpp
+++ b/Engine/Plugins/Animation/ControlRig/Source/ControlRigEditor/Private/Sequencer/ConstraintChannelHelper.cpp
@@ -18,8 +18,8 @@
 #include "ScopedTransaction.h"
 #include "Channels/MovieSceneFloatChannel.h"
 #include "Channels/MovieSceneDoubleChannel.h"
+#include "MovieSceneConstraintChannelHelper.h"
 
-bool FConstraintChannelHelper::bDoNotCompensate = false;
 
 #define LOCTEXT_NAMESPACE "Constraints"
 
@@ -167,7 +167,7 @@ void FConstraintChannelHelper::SmartControlConstraintKey(
 			if (CanAddKey(Channel->ActiveChannel, Time, ActiveValueToBeSet))
 			{
 
-				TGuardValue<bool> CompensateGuard(bDoNotCompensate, true);
+				TGuardValue<bool> CompensateGuard(FMovieSceneConstraintChannelHelper::bDoNotCompensate, true);
 				
 				// set constraint as dynamic
 				InConstraint->bDynamicOffset = true;
@@ -290,7 +290,7 @@ void FConstraintChannelHelper::SmartComponentConstraintKey(
 
 				//new for compensatio
 
-				TGuardValue<bool> CompensateGuard(bDoNotCompensate, true);
+				TGuardValue<bool> CompensateGuard(FMovieSceneConstraintChannelHelper::bDoNotCompensate, true);
 
 				// set constraint as dynamic
 				InConstraint->bDynamicOffset = true;
@@ -339,6 +339,17 @@ void FConstraintChannelHelper::SmartComponentConstraintKey(
 					// set tangents at Time
 					SetTangentsAtThisTime<FMovieSceneDoubleChannel>(ChannelIndex,NumChannels, Section, Time, Tangents);
 				}
+				// evaluate the constraint, this is needed so the global transform will be set up on the component 
+				//Todo do we need to evaluate all constraints?
+				InConstraint->SetActive(true); //will be false
+				InConstraint->Evaluate();
+				//need to fire this event so the transform values set by the constraint propragate to the transform section
+				FProperty* TransformProperty = FindFProperty<FProperty>(USceneComponent::StaticClass(), USceneComponent::GetRelativeLocationPropertyName());
+				FEditPropertyChain PropertyChain;
+				PropertyChain.AddHead(TransformProperty);
+				FCoreUObjectDelegates::OnPreObjectPropertyChanged.Broadcast(Actor, PropertyChain);
+				FPropertyChangedEvent PropertyChangedEvent(TransformProperty, EPropertyChangeType::ValueSet);
+				FCoreUObjectDelegates::OnObjectPropertyChanged.Broadcast(Actor, PropertyChangedEvent);
 			}
 		}
 	}
@@ -376,12 +387,12 @@ void FConstraintChannelHelper::CompensateIfNeeded(
 	const UMovieSceneControlRigParameterSection* Section,
 	const TOptional<FFrameNumber>& OptionalTime)
 {
-	if (bDoNotCompensate)
+	if (FMovieSceneConstraintChannelHelper::bDoNotCompensate)
 	{
 		return;
 	}
 
-	TGuardValue<bool> CompensateGuard(bDoNotCompensate, true);
+	TGuardValue<bool> CompensateGuard(FMovieSceneConstraintChannelHelper::bDoNotCompensate, true);
 
 	// Frames to compensate
 	TArray<FFrameNumber> OptionalTimeArray;
diff --git a/Engine/Plugins/Animation/ControlRig/Source/ControlRigEditor/Private/Sequencer/ControlRigParameterTrackEditor.cpp b/Engine/Plugins/Animation/ControlRig/Source/ControlRigEditor/Private/Sequencer/ControlRigParameterTrackEditor.cpp
--- a/Engine/Plugins/Animation/ControlRig/Source/ControlRigEditor/Private/Sequencer/ControlRigParameterTrackEditor.cpp
+++ b/Engine/Plugins/Animation/ControlRig/Source/ControlRigEditor/Private/Sequencer/ControlRigParameterTrackEditor.cpp
@@ -2332,16 +2332,22 @@ void FControlRigParameterTrackEditor::HandleOnSpaceAdded(UMovieSceneControlRigPa
 {
 	if (SpaceChannel)
 	{
-		SpaceChannel->OnKeyMovedEvent().AddLambda([this,Section](FMovieSceneChannel* Channel, const  TArray<FKeyMoveEventItem>& MovedItems)
-			{
-				FMovieSceneControlRigSpaceChannel* SpaceChannel = static_cast<FMovieSceneControlRigSpaceChannel*>(Channel);
-				HandleSpaceKeyMoved(Section, SpaceChannel,MovedItems);
-			});
-		SpaceChannel->OnKeyDeletedEvent().AddLambda([this, Section](FMovieSceneChannel* Channel, const  TArray<FKeyAddOrDeleteEventItem>& Items)
-			{
-				FMovieSceneControlRigSpaceChannel* SpaceChannel = static_cast<FMovieSceneControlRigSpaceChannel*>(Channel);
-				HandleSpaceKeyDeleted(Section, SpaceChannel,Items);
-			});
+		if (!SpaceChannel->OnKeyMovedEvent().IsBound())
+		{
+			SpaceChannel->OnKeyMovedEvent().AddLambda([this, Section](FMovieSceneChannel* Channel, const  TArray<FKeyMoveEventItem>& MovedItems)
+				{
+					FMovieSceneControlRigSpaceChannel* SpaceChannel = static_cast<FMovieSceneControlRigSpaceChannel*>(Channel);
+					HandleSpaceKeyMoved(Section, SpaceChannel, MovedItems);
+				});
+		}
+		if (!SpaceChannel->OnKeyDeletedEvent().IsBound())
+		{
+			SpaceChannel->OnKeyDeletedEvent().AddLambda([this, Section](FMovieSceneChannel* Channel, const  TArray<FKeyAddOrDeleteEventItem>& Items)
+				{
+					FMovieSceneControlRigSpaceChannel* SpaceChannel = static_cast<FMovieSceneControlRigSpaceChannel*>(Channel);
+					HandleSpaceKeyDeleted(Section, SpaceChannel, Items);
+				});
+		}
 	}
 	//todoo do we need to remove this or not mz
 }
@@ -2541,20 +2547,26 @@ void FControlRigParameterTrackEditor::HandleOnConstraintAdded(
 	}
 
 	// handle key moved
-	InConstraintChannel->OnKeyMovedEvent().AddLambda([this,InSection](
-		FMovieSceneChannel* InChannel, const TArray<FKeyMoveEventItem>& InMovedItems)
+	if (!InConstraintChannel->OnKeyMovedEvent().IsBound())
 	{
-		const FMovieSceneConstraintChannel* ConstraintChannel = static_cast<FMovieSceneConstraintChannel*>(InChannel);
-		HandleConstraintKeyMoved(InSection, ConstraintChannel, InMovedItems);
-	});
+		InConstraintChannel->OnKeyMovedEvent().AddLambda([this, InSection](
+			FMovieSceneChannel* InChannel, const TArray<FKeyMoveEventItem>& InMovedItems)
+			{
+				const FMovieSceneConstraintChannel* ConstraintChannel = static_cast<FMovieSceneConstraintChannel*>(InChannel);
+				HandleConstraintKeyMoved(InSection, ConstraintChannel, InMovedItems);
+			});
+	}
 
 	// handle key deleted
-	InConstraintChannel->OnKeyDeletedEvent().AddLambda([this, InSection](
-		FMovieSceneChannel* InChannel, const TArray<FKeyAddOrDeleteEventItem>& InDeletedItems)
+	if (!InConstraintChannel->OnKeyDeletedEvent().IsBound())
 	{
-		const FMovieSceneConstraintChannel* ConstraintChannel = static_cast<FMovieSceneConstraintChannel*>(InChannel);
-		HandleConstraintKeyDeleted(InSection, ConstraintChannel, InDeletedItems);
-	});
+		InConstraintChannel->OnKeyDeletedEvent().AddLambda([this, InSection](
+			FMovieSceneChannel* InChannel, const TArray<FKeyAddOrDeleteEventItem>& InDeletedItems)
+			{
+				const FMovieSceneConstraintChannel* ConstraintChannel = static_cast<FMovieSceneConstraintChannel*>(InChannel);
+				HandleConstraintKeyDeleted(InSection, ConstraintChannel, InDeletedItems);
+			});
+	}
 
 	// handle constraint deleted
 	if (InSection)
@@ -2621,7 +2633,7 @@ void FControlRigParameterTrackEditor::HandleConstraintRemoved(IMovieSceneConstra
 	if (const UControlRig* ControlRig = Section->GetControlRig())
 	{
 		FConstraintsManagerController& Controller = FConstraintsManagerController::Get(ControlRig->GetWorld());
-		if (!Controller.OnConstraintRemoved().IsBoundToObject(InSection))
+		if (!Controller.OnConstraintRemoved().IsBound())
 		{
 			InSection->OnConstraintRemovedHandle =
 			Controller.OnConstraintRemoved().AddLambda([InSection, Section, this](FName InConstraintName)
diff --git a/Engine/Plugins/Animation/ControlRig/Source/ControlRigEditor/Public/ConstraintChannelHelper.h b/Engine/Plugins/Animation/ControlRig/Source/ControlRigEditor/Public/ConstraintChannelHelper.h
--- a/Engine/Plugins/Animation/ControlRig/Source/ControlRigEditor/Public/ConstraintChannelHelper.h
+++ b/Engine/Plugins/Animation/ControlRig/Source/ControlRigEditor/Public/ConstraintChannelHelper.h
@@ -78,5 +78,4 @@ private:
 	/** END COMPONENT SECTION */
 
 	
-	static bool bDoNotCompensate;
 };
diff --git a/Engine/Source/Editor/MovieSceneTools/Private/MovieSceneConstraintChannelHelper.cpp b/Engine/Source/Editor/MovieSceneTools/Private/MovieSceneConstraintChannelHelper.cpp
--- a/Engine/Source/Editor/MovieSceneTools/Private/MovieSceneConstraintChannelHelper.cpp
+++ b/Engine/Source/Editor/MovieSceneTools/Private/MovieSceneConstraintChannelHelper.cpp
@@ -411,6 +411,13 @@ void FMovieSceneConstraintChannelHelper::HandleConstraintKeyDeleted(
 		UWorld* World = GCurrentLevelEditingViewportClient ? GCurrentLevelEditingViewportClient->GetWorld() : nullptr;
 		Evaluator.ComputeLocalTransforms(World, InSequencer, FramesToCompensate, PreviousValue);
 		TArray<FTransform>& ChildLocals = Evaluator.ChildLocals;
+		//turn off constraint, if we delete the key it may not evaluate to false
+		InConstraint->SetActive(false);
+
+		if (ChildLocals.Num() < 2)
+		{
+			return;
+		}
 		ChildLocals.RemoveAt(0);
 
 		const EMovieSceneTransformChannel ChannelsToKey = InConstraint->GetChannelsToKey();
diff --git a/Engine/Source/Editor/MovieSceneTools/Private/TrackEditors/TransformTrackEditor.cpp b/Engine/Source/Editor/MovieSceneTools/Private/TrackEditors/TransformTrackEditor.cpp
--- a/Engine/Source/Editor/MovieSceneTools/Private/TrackEditors/TransformTrackEditor.cpp
+++ b/Engine/Source/Editor/MovieSceneTools/Private/TrackEditors/TransformTrackEditor.cpp
@@ -70,11 +70,17 @@ F3DTransformTrackEditor::F3DTransformTrackEditor( TSharedRef<ISequencer> InSeque
 	// Listen for actor/component movement
 	FCoreUObjectDelegates::OnPreObjectPropertyChanged.AddRaw(this, &F3DTransformTrackEditor::OnPrePropertyChanged);
 	FCoreUObjectDelegates::OnObjectPropertyChanged.AddRaw(this, &F3DTransformTrackEditor::OnPostPropertyChanged);
+	if (GEditor != nullptr)
+	{
+		GEditor->RegisterForUndo(this);
+	}
 }
-
-
 F3DTransformTrackEditor::~F3DTransformTrackEditor()
 {
+	if (GEditor != nullptr)
+	{
+		GEditor->UnregisterForUndo(this);
+	}
 }
 
 void F3DTransformTrackEditor::OnRelease()
@@ -870,6 +876,58 @@ void F3DTransformTrackEditor::AddTransformKeys( UObject* ObjectToKey, const TOpt
 	AnimatablePropertyChanged( FOnKeyProperty::CreateLambda(OnKeyProperty) );
 }
 
+//todo move to external function so it can also be used by sdk/python
+static void UpdateTransformBasedOnConstraint(FTransform& CurrentTransform, USceneComponent* SceneComponent)
+{
+	TArray< TObjectPtr<UTickableConstraint> > Constraints;
+	AActor* ShapeActor = SceneComponent->GetTypedOuter<AActor>();
+	
+	if (ShapeActor)
+	{
+		FTransformConstraintUtils::GetParentConstraints(SceneComponent->GetWorld(), ShapeActor, Constraints);
+		const int32 LastActiveIndex = Constraints.FindLastByPredicate([](const TObjectPtr<UTickableConstraint>& InConstraint)
+			{
+				if (const UTickableTransformConstraint* TransformConstraint = Cast<UTickableTransformConstraint>(InConstraint.Get()))
+				{
+					return InConstraint->Active && TransformConstraint->bDynamicOffset;
+				}
+				return false;
+			});
+
+		if (Constraints.IsValidIndex(LastActiveIndex))
+		{
+			// switch to constraint space
+			const UTickableTransformConstraint* Constraint = Cast<UTickableTransformConstraint>(Constraints[LastActiveIndex]);
+			const FTransform ParentTransform = Constraint->GetParentGlobalTransform();
+			const FTransform OriginalParentTransform = SceneComponent->GetAttachParent() ? SceneComponent->GetAttachParent()->GetSocketTransform(SceneComponent->GetAttachSocketName()) : FTransform::Identity;
+			const FTransform DiffParents = OriginalParentTransform.GetRelativeTransform(ParentTransform);
+
+			const ETransformConstraintType ConstraintType = static_cast<ETransformConstraintType>(Constraint->GetType());
+			switch (ConstraintType)
+			{
+				case ETransformConstraintType::Translation:
+				{
+					CurrentTransform.SetLocation(CurrentTransform.GetLocation() +  OriginalParentTransform.GetLocation() - ParentTransform.GetLocation());
+				}
+				case ETransformConstraintType::Rotation:
+				{
+					FQuat RelativeRotation = ParentTransform.GetRotation().Inverse() * OriginalParentTransform.GetRotation();
+					RelativeRotation.Normalize();
+					CurrentTransform.SetRotation(CurrentTransform.GetRotation() * RelativeRotation);
+				}
+				case ETransformConstraintType::Scale:
+				{
+					CurrentTransform.SetScale3D(CurrentTransform.GetScale3D() * DiffParents.GetScale3D());
+				}
+				case ETransformConstraintType::Parent:
+					CurrentTransform = CurrentTransform * DiffParents;
+				case ETransformConstraintType::LookAt: //leave current alone
+				default:
+					break;
+			}
+		}
+	}
+}
 
 FTransformData F3DTransformTrackEditor::RecomposeTransform(const FTransformData& InTransformData, UObject* AnimatedObject, UMovieSceneSection* Section)
 {
@@ -910,6 +968,7 @@ FTransformData F3DTransformTrackEditor::RecomposeTransform(const FTransformData&
 				CurrentTransform *= GetTransformOrigin().Inverse();
 			}
 
+			UpdateTransformBasedOnConstraint(CurrentTransform,SceneComponent);
 			return FTransformData(CurrentTransform.GetLocation(), CurrentTransform.GetRotation().Rotator(), CurrentTransform.GetScale3D());
 		}
 	}
@@ -1011,6 +1070,8 @@ void F3DTransformTrackEditor::ProcessKeyOperation(UObject* ObjectToKey, TArrayVi
 			CurrentTransform *= GetTransformOrigin().Inverse();
 		}
 
+		UpdateTransformBasedOnConstraint(CurrentTransform, Component);
+
 		FIntermediate3DTransform CurrentValue(CurrentTransform.GetTranslation(), CurrentTransform.GetRotation().Rotator(), CurrentTransform.GetScale3D());
 		TRecompositionResult<FIntermediate3DTransform> TransformData = System->RecomposeBlendOperational(FMovieSceneTracksComponentTypes::Get()->ComponentTransform, Query, CurrentValue);
 
@@ -1288,7 +1349,6 @@ void F3DTransformTrackEditor::PostUndo(bool bSuccess)
 	}
 }
 
-
 void F3DTransformTrackEditor::HandleOnConstraintAdded(IMovieSceneConstrainedSection* InSection, FMovieSceneConstraintChannel* InConstraintChannel)
 {
 	if (!InConstraintChannel)
@@ -1298,7 +1358,7 @@ void F3DTransformTrackEditor::HandleOnConstraintAdded(IMovieSceneConstrainedSect
 	// handle scene component changes so we can update the gizmo location when it changes
 	UWorld* World = GCurrentLevelEditingViewportClient ? GCurrentLevelEditingViewportClient->GetWorld() : nullptr;
 	FConstraintsManagerController& Controller = FConstraintsManagerController::Get(World);
-	if (!Controller.OnSceneComponentConstrained().IsBoundToObject(this))
+	if (!Controller.OnSceneComponentConstrained().IsBound())
 	{
 		OnSceneComponentConstrainedHandle =
 			Controller.OnSceneComponentConstrained().AddLambda([](USceneComponent* InSceneComponent)
@@ -1324,20 +1384,26 @@ void F3DTransformTrackEditor::HandleOnConstraintAdded(IMovieSceneConstrainedSect
 				});
 	}
 	// handle key moved
-	InConstraintChannel->OnKeyMovedEvent().AddLambda([this, InSection](
-		FMovieSceneChannel* InChannel, const TArray<FKeyMoveEventItem>& InMovedItems)
-		{
-			const FMovieSceneConstraintChannel* ConstraintChannel = static_cast<FMovieSceneConstraintChannel*>(InChannel);
-			HandleConstraintKeyMoved(InSection, ConstraintChannel, InMovedItems);
-		});
+	if (!InConstraintChannel->OnKeyMovedEvent().IsBound())
+	{
+		InConstraintChannel->OnKeyMovedEvent().AddLambda([this, InSection](
+			FMovieSceneChannel* InChannel, const TArray<FKeyMoveEventItem>& InMovedItems)
+			{
+				const FMovieSceneConstraintChannel* ConstraintChannel = static_cast<FMovieSceneConstraintChannel*>(InChannel);
+				HandleConstraintKeyMoved(InSection, ConstraintChannel, InMovedItems);
+			});
+	}
 
 	// handle key deleted
-	InConstraintChannel->OnKeyDeletedEvent().AddLambda([this, InSection](
-		FMovieSceneChannel* InChannel, const TArray<FKeyAddOrDeleteEventItem>& InDeletedItems)
-		{
-			const FMovieSceneConstraintChannel* ConstraintChannel = static_cast<FMovieSceneConstraintChannel*>(InChannel);
-			HandleConstraintKeyDeleted(InSection, ConstraintChannel, InDeletedItems);
-		});
+	if (!InConstraintChannel->OnKeyDeletedEvent().IsBound())
+	{
+		InConstraintChannel->OnKeyDeletedEvent().AddLambda([this, InSection](
+			FMovieSceneChannel* InChannel, const TArray<FKeyAddOrDeleteEventItem>& InDeletedItems)
+			{
+				const FMovieSceneConstraintChannel* ConstraintChannel = static_cast<FMovieSceneConstraintChannel*>(InChannel);
+				HandleConstraintKeyDeleted(InSection, ConstraintChannel, InDeletedItems);
+			});
+	}
 
 	// handle constraint deleted
 	if (InSection)
@@ -1375,6 +1441,7 @@ void F3DTransformTrackEditor::HandleConstraintKeyDeleted(IMovieSceneConstrainedS
 	{
 		return;
 	}
+
     if(UTickableTransformConstraint* Constrain = GetTickableTransformConstraint(InSection,InConstraintChannel))
 	{
 		for (const FKeyAddOrDeleteEventItem& EventItem : InDeletedItems)
@@ -1391,7 +1458,6 @@ void F3DTransformTrackEditor::HandleConstraintKeyDeleted(IMovieSceneConstrainedS
 void F3DTransformTrackEditor::HandleConstraintKeyMoved(IMovieSceneConstrainedSection* InSection, const FMovieSceneConstraintChannel* InConstraintChannel,
 	const TArray<FKeyMoveEventItem>& InMovedItems)
 {
-
 	if (!InConstraintChannel)
 	{
 		return;
@@ -1401,6 +1467,7 @@ void F3DTransformTrackEditor::HandleConstraintKeyMoved(IMovieSceneConstrainedSec
 	{
 		return;
 	}
+
 	UMovieSceneSection* Section = Cast<UMovieSceneSection>(InSection);
 
 	for (const FKeyMoveEventItem& MoveEventItem : InMovedItems)
@@ -1417,7 +1484,7 @@ void F3DTransformTrackEditor::HandleConstraintRemoved(IMovieSceneConstrainedSect
 	UMovieSceneSection* Section = Cast<UMovieSceneSection>(InSection);
 
 	FConstraintsManagerController& Controller = FConstraintsManagerController::Get(World);
-	if (!Controller.OnConstraintRemoved().IsBoundToObject(InSection))
+	if (!Controller.OnConstraintRemoved().IsBound())
 	{
 		InSection->OnConstraintRemovedHandle =
 			Controller.OnConstraintRemoved().AddLambda([InSection,Section, this](FName InConstraintName)
diff --git a/Engine/Source/Editor/MovieSceneTools/Public/MovieSceneConstraintChannelHelper.h b/Engine/Source/Editor/MovieSceneTools/Public/MovieSceneConstraintChannelHelper.h
--- a/Engine/Source/Editor/MovieSceneTools/Public/MovieSceneConstraintChannelHelper.h
+++ b/Engine/Source/Editor/MovieSceneTools/Public/MovieSceneConstraintChannelHelper.h
@@ -105,7 +105,6 @@ public:
 		const TArrayView<ChannelType*>& InChannels,
 		const FFrameNumber& InTime);
 
-private:
 	static bool bDoNotCompensate;
 
 };
diff --git a/Engine/Source/Runtime/MovieSceneTracks/Private/Sections/MovieScene3DTransformSection.cpp b/Engine/Source/Runtime/MovieSceneTracks/Private/Sections/MovieScene3DTransformSection.cpp
--- a/Engine/Source/Runtime/MovieSceneTracks/Private/Sections/MovieScene3DTransformSection.cpp
+++ b/Engine/Source/Runtime/MovieSceneTracks/Private/Sections/MovieScene3DTransformSection.cpp
@@ -15,7 +15,6 @@
 #include "MovieSceneTracksComponentTypes.h"
 #include "Algo/AnyOf.h"
 #include "TransformConstraint.h"
-#include "TransformableHandle.h"
 
 #if WITH_EDITOR
 
@@ -847,16 +846,31 @@ TArray<FConstraintAndActiveChannel>& UMovieScene3DTransformSection::GetConstrain
 	}
 	return EmptyChannels;
 }
+#if WITH_EDITOR
+bool UMovieScene3DTransformSection::Modify(bool bAlwaysMarkDirty)
+{
+	using namespace UE::MovieScene;
 
+	bool bModified = Super::Modify(bAlwaysMarkDirty);
+	if (Constraints)
+	{
+		Constraints->SetFlags(RF_Transactional);
+		Constraints->Modify(bAlwaysMarkDirty);
+	}
+	return bModified;
+}
+#endif
 void UMovieScene3DTransformSection::AddConstraintChannel(UTickableConstraint* InConstraint)
 {
 	Modify();
 	if (!Constraints)
 	{
-		Constraints = NewObject<UMovieScene3dTransformSectionConstraints>(this, NAME_None, RF_Transactional);
+		Constraints = NewObject<UMovieScene3dTransformSectionConstraints>(this, NAME_None, RF_Public| RF_Transactional);
 	}
 	if (!HasConstraintChannel(InConstraint->GetFName()))
 	{
+		Constraints->SetFlags(RF_Transactional);
+		Constraints->Modify();
 		const int32 NewIndex = Constraints->ConstraintsChannels.Add(FConstraintAndActiveChannel(InConstraint));
 
 		FMovieSceneConstraintChannel* ExistingChannel = &Constraints->ConstraintsChannels[NewIndex].ActiveChannel;
diff --git a/Engine/Source/Runtime/MovieSceneTracks/Private/Systems/MovieSceneConstraintSystem.cpp b/Engine/Source/Runtime/MovieSceneTracks/Private/Systems/MovieSceneConstraintSystem.cpp
--- a/Engine/Source/Runtime/MovieSceneTracks/Private/Systems/MovieSceneConstraintSystem.cpp
+++ b/Engine/Source/Runtime/MovieSceneTracks/Private/Systems/MovieSceneConstraintSystem.cpp
@@ -13,6 +13,8 @@
 #include "EntitySystem/BuiltInComponentTypes.h"
 #include "MovieSceneTracksComponentTypes.h"
 #include "Systems/MovieSceneComponentTransformSystem.h"
+#include "TransformConstraint.h"
+#include "TransformableHandle.h"
 
 class UTickableConstraint;
 
@@ -71,7 +73,7 @@ UMovieSceneConstraintSystem::UMovieSceneConstraintSystem(const FObjectInitialize
 	RelevantComponent = TracksComponents->ConstraintChannel;
 
 	// Run constraints during instantiation or evaluation
-	Phase = ESystemPhase::Instantiation | ESystemPhase::Evaluation;
+	Phase = ESystemPhase::Instantiation | ESystemPhase::Evaluation | ESystemPhase::Finalization;
 
 	if (HasAnyFlags(RF_ClassDefaultObject) ) 
 	{
@@ -101,13 +103,14 @@ void UMovieSceneConstraintSystem::OnRun(FSystemTaskPrerequisites& InPrerequisite
 	}
 	else if (CurrentPhase == ESystemPhase::Evaluation)
 	{
+		DynamicOffsets.Reset();
 		FConstraintsManagerController& Controller = FConstraintsManagerController::Get(GetWorld());
 		FMovieSceneTracksComponentTypes* TracksComponents = FMovieSceneTracksComponentTypes::Get();
 		FBuiltInComponentTypes* BuiltInComponents = FBuiltInComponentTypes::Get();
 
 		struct FEvaluateConstraintChannels
 		{
-			FEvaluateConstraintChannels(FConstraintsManagerController* InController) :Controller(InController) { check(Controller); };
+			FEvaluateConstraintChannels(FConstraintsManagerController* InController, UMovieSceneConstraintSystem* InSystem) :Controller(InController), System(InSystem) { check(Controller); check(System) };
 			void ForEachEntity(UObject* BoundObject, const FConstraintComponentData& ConstraintChannel, FFrameTime FrameTime)
 			{
 				UTickableConstraint* Constraint = Controller->GetConstraint(ConstraintChannel.ConstraintName);
@@ -116,9 +119,20 @@ void UMovieSceneConstraintSystem::OnRun(FSystemTaskPrerequisites& InPrerequisite
 					bool Result = false;
 					ConstraintChannel.Channel->Evaluate(FrameTime, Result);
 					Constraint->SetActive(Result);
+					if (UTickableTransformConstraint* TransformConstraint = Cast< UTickableTransformConstraint>(Constraint))
+					{
+						if (UTransformableComponentHandle* ComponentHandle = Cast<UTransformableComponentHandle>(TransformConstraint->ChildTRSHandle))
+						{
+							FUpdateHandleForConstraint UpdateHandle;
+							UpdateHandle.Constraint = TransformConstraint;
+							UpdateHandle.TransformHandle = ComponentHandle;
+							System->DynamicOffsets.Add(UpdateHandle);
+						}
+					}
 				}
 			}
 			FConstraintsManagerController* Controller;
+			UMovieSceneConstraintSystem* System;
 		};
 
 		// Set up new constraints
@@ -127,7 +141,18 @@ void UMovieSceneConstraintSystem::OnRun(FSystemTaskPrerequisites& InPrerequisite
 			.Read(BuiltInComponents->BoundObject)
 			.Read(TracksComponents->ConstraintChannel)
 			.Read(BuiltInComponents->EvalTime)
-			.Dispatch_PerEntity<FEvaluateConstraintChannels>(&Linker->EntityManager, InPrerequisites, &Subsequents, &Controller);
+			.Dispatch_PerEntity<FEvaluateConstraintChannels>(&Linker->EntityManager, InPrerequisites, &Subsequents, &Controller, this);
+	}
+	else if (CurrentPhase == ESystemPhase::Finalization)
+	{
+		for (FUpdateHandleForConstraint& UpdateHandle : DynamicOffsets)
+		{
+			if (UpdateHandle.Constraint.IsValid() && UpdateHandle.TransformHandle.IsValid())
+			{
+				UpdateHandle.Constraint->OnHandleModified(UpdateHandle.TransformHandle.Get(),false);
+			}
+		}
+		DynamicOffsets.Reset();
 	}
 }
 
diff --git a/Engine/Source/Runtime/MovieSceneTracks/Public/Sections/MovieScene3DTransformSection.h b/Engine/Source/Runtime/MovieSceneTracks/Public/Sections/MovieScene3DTransformSection.h
--- a/Engine/Source/Runtime/MovieSceneTracks/Public/Sections/MovieScene3DTransformSection.h
+++ b/Engine/Source/Runtime/MovieSceneTracks/Public/Sections/MovieScene3DTransformSection.h
@@ -216,6 +216,10 @@ class UMovieScene3DTransformSection
 
 public:
 
+#if WITH_EDITOR
+	/* From UObject*/
+	virtual bool Modify(bool bAlwaysMarkDirty = true) override;
+#endif
 	/* From UMovieSection*/
 	
 	MOVIESCENETRACKS_API virtual bool ShowCurveForChannel(const void *Channel) const override;
diff --git a/Engine/Source/Runtime/MovieSceneTracks/Public/Systems/MovieSceneConstraintSystem.h b/Engine/Source/Runtime/MovieSceneTracks/Public/Systems/MovieSceneConstraintSystem.h
--- a/Engine/Source/Runtime/MovieSceneTracks/Public/Systems/MovieSceneConstraintSystem.h
+++ b/Engine/Source/Runtime/MovieSceneTracks/Public/Systems/MovieSceneConstraintSystem.h
@@ -5,6 +5,9 @@
 #include "EntitySystem/MovieSceneEntitySystem.h"
 #include "MovieSceneConstraintSystem.generated.h"
 
+class UTransformableComponentHandle;
+class UTickableTransformConstraint;
+
 /**
  * System that is responsible for propagating constraints to a bound object's FConstraintsManagerController.
  */
@@ -14,8 +17,17 @@ class MOVIESCENETRACKS_API UMovieSceneConstraintSystem : public UMovieSceneEntit
 public:
 
 	GENERATED_BODY()
+	//if constraint has a dynamic offset we need to make sure the constraint get's notified the handle is updated
+	struct FUpdateHandleForConstraint
+	{
+		TWeakObjectPtr<UTickableTransformConstraint> Constraint;
+		TWeakObjectPtr<UTransformableComponentHandle> TransformHandle;
+	};
 
 	UMovieSceneConstraintSystem(const FObjectInitializer& ObjInit);
-
 	virtual void OnRun(FSystemTaskPrerequisites& InPrerequisites, FSystemSubsequentTasks& Subsequents) override;
+
+protected:
+	TArray<FUpdateHandleForConstraint>  DynamicOffsets;
+
 };
