Sequencer: handle the case where multiple ECS entities are driving a single transform section

When auto-keying, we try to figure out the contribution of the section being keyed compared to the overall current transform value (for when that transform is computed from multiple sections, like an absolute plus an additive). This code was working under the assumption that a transform section's contribution is built from a single ECS entity. This isn't a valid assumption anymore with overridable channels.

If a transform section has some overriden channels, there will be one extra ECS entity for each overriden channel. So the recomposition code now needs to deal with multiple entities, which this change addresses.

Note that the index of channels can now vary: if Translation.X is overriden, the first "double" channel will be Translation.Y, with index 0 in the bucket of channels of type "double". But Translation.Y is really channel 1, not channel 0. So this change also makes use of GetChannelBySortOrder, instead of GetChannel, so that we get the correct channel to key.

diff --git a/Engine/Source/Editor/MovieSceneTools/Private/TrackEditors/TransformTrackEditor.cpp b/Engine/Source/Editor/MovieSceneTools/Private/TrackEditors/TransformTrackEditor.cpp
--- a/Engine/Source/Editor/MovieSceneTools/Private/TrackEditors/TransformTrackEditor.cpp
+++ b/Engine/Source/Editor/MovieSceneTools/Private/TrackEditors/TransformTrackEditor.cpp
@@ -862,14 +862,12 @@ void F3DTransformTrackEditor::AddTransformKeys( UObject* ObjectToKey, const TOpt
 		{
 			for (TWeakObjectPtr<UMovieSceneSection>& WeakSection : KeyPropertyResult.SectionsKeyed)
 			{
-				
 				if (UMovieScene3DTransformSection* Section = Cast< UMovieScene3DTransformSection>(WeakSection.Get()))
 				{
 					FMovieSceneConstraintChannelHelper::CompensateIfNeeded(GetSequencer(), Section, Time);
 				}
 			}
 		}
-		
 		return KeyPropertyResult;
 	};
 
@@ -943,24 +941,81 @@ FTransformData F3DTransformTrackEditor::RecomposeTransform(const FTransformData&
 
 	TGuardValue<FEntityManager*> DebugVizGuard(GEntityManagerForDebuggingVisualizers, &EntityLinker->EntityManager);
 
-	FMovieSceneEntityID EntityID = EvaluationTemplate.FindEntityFromOwner(Section, 0, GetSequencer()->GetFocusedTemplateID());
+	// We want the transform value contributed to by the given transform section.
+	//
+	// In most cases, the section only has one ECS entity for it specifiy all 9 channel curves and resulting values, but if
+	// some of those channels have been overriden (e.g. with some noise channel), then there will be one extra ECS entity
+	// for each overriden channel, and the "main" entity will be missing that channel. We therefore gather up all entities
+	// related to the given section and recompose their contribution, and then join them all together by matching which one
+	// handles which channel.
+	//
+	TArray<FMovieSceneEntityID> ImportedEntityIDs;
+	EvaluationTemplate.FindEntitiesFromOwner(Section, GetSequencer()->GetFocusedTemplateID(), ImportedEntityIDs);
 
-	if (EntityID)
+	if (ImportedEntityIDs.Num())
 	{
 		UMovieScenePropertyInstantiatorSystem* System = EntityLinker->FindSystem<UMovieScenePropertyInstantiatorSystem>();
 		if (System)
 		{
+			FBuiltInComponentTypes* BuiltInComponents = FBuiltInComponentTypes::Get();
+			FMovieSceneTracksComponentTypes* TrackComponents = FMovieSceneTracksComponentTypes::Get();
 			USceneComponent* SceneComponent = MovieSceneHelpers::SceneComponentFromRuntimeObject(AnimatedObject);
 
+			TArray<FMovieSceneEntityID> EntityIDs;
+			{
+				// In order to check for the result channels later, we need to look up the children entities that are
+				// bound to the given animated object. Imported entities generally don't have the result channels.
+				FEntityTaskBuilder()
+				.ReadEntityIDs()
+				.Read(BuiltInComponents->ParentEntity)
+				.Read(BuiltInComponents->BoundObject)
+				.FilterAll({ TrackComponents->ComponentTransform.PropertyTag })
+				.Iterate_PerEntity(
+					&EntityLinker->EntityManager, 
+					[SceneComponent, ImportedEntityIDs, &EntityIDs](FMovieSceneEntityID EntityID, FMovieSceneEntityID ParentEntityID, UObject* BoundObject)
+					{
+						if (SceneComponent == BoundObject && ImportedEntityIDs.Contains(ParentEntityID))
+						{
+							EntityIDs.Add(EntityID);
+						}
+					});
+			}
+
 			FDecompositionQuery Query;
-			Query.Entities = MakeArrayView(&EntityID, 1);
+			Query.Entities = MakeArrayView(EntityIDs);
 			Query.Object   = SceneComponent;
+			Query.bConvertFromSourceEntityIDs = false;  // We already pass the children entity IDs
 
 			FIntermediate3DTransform CurrentValue(InTransformData.Translation, InTransformData.Rotation, InTransformData.Scale);
 
-			TRecompositionResult<FIntermediate3DTransform> TransformData = System->RecomposeBlendOperational(FMovieSceneTracksComponentTypes::Get()->ComponentTransform, Query, CurrentValue);
+			TRecompositionResult<FIntermediate3DTransform> TransformData = System->RecomposeBlendOperational(TrackComponents->ComponentTransform, Query, CurrentValue);
+
+			double CurrentTransformChannels[9] = { 0, 0, 0, 0, 0, 0, 0, 0, 0 };
+			EMovieSceneTransformChannel ChannelsObtained(EMovieSceneTransformChannel::None);
+			check(EntityIDs.Num() == TransformData.Values.Num());
+			for (int32 EntityIndex = 0; EntityIndex < TransformData.Values.Num(); ++EntityIndex)
+			{
+				// For each entity, find which channel they contribute to by checking the result channel.
+				// We don't (yet) handle the case where two entities contribute to the same channel -- in theory the entities
+				// should be mutually exclusive in that regard, at least as far as overriden channels are concerned.
+				FMovieSceneEntityID EntityID = EntityIDs[EntityIndex];
+				FIntermediate3DTransform EntityTransformData = TransformData.Values[EntityIndex];
+				FComponentMask EntityType = EntityLinker->EntityManager.GetEntityType(EntityID);
+				for (int32 CompositeIndex = 0; CompositeIndex < 9; ++CompositeIndex)
+				{
+					EMovieSceneTransformChannel ChannelMask = (EMovieSceneTransformChannel)(1 << CompositeIndex);
+					if (!EnumHasAnyFlags(ChannelsObtained, ChannelMask) && EntityType.Contains(BuiltInComponents->DoubleResult[CompositeIndex]))
+					{
+						EnumAddFlags(ChannelsObtained, (EMovieSceneTransformChannel)(1 << CompositeIndex));
+						CurrentTransformChannels[CompositeIndex] = EntityTransformData[CompositeIndex];
+					}
+				}
+			}
 
-			FTransform CurrentTransform(TransformData.Values[0].GetRotation(), TransformData.Values[0].GetTranslation(), TransformData.Values[0].GetScale());
+			FTransform CurrentTransform(
+				FRotator(CurrentTransformChannels[3], CurrentTransformChannels[4], CurrentTransformChannels[5]),
+				FVector(CurrentTransformChannels[0], CurrentTransformChannels[1], CurrentTransformChannels[2]),
+				FVector(CurrentTransformChannels[6], CurrentTransformChannels[7], CurrentTransformChannels[8]));
 
 			// Account for the transform origin only if this is not parented because the transform origin is already being applied to the parent.
 			if (!SceneComponent->GetAttachParent())
diff --git a/Engine/Source/Editor/MovieSceneTools/Public/KeyframeTrackEditor.h b/Engine/Source/Editor/MovieSceneTools/Public/KeyframeTrackEditor.h
--- a/Engine/Source/Editor/MovieSceneTools/Public/KeyframeTrackEditor.h
+++ b/Engine/Source/Editor/MovieSceneTools/Public/KeyframeTrackEditor.h
@@ -55,7 +55,7 @@ struct TAddKeyImpl : IImpl
 		bool bKeyCreated = false;
 		using namespace UE::MovieScene;
 
-		ChannelType* Channel = Proxy.GetChannel<ChannelType>(ChannelIndex);
+		ChannelType* Channel = Proxy.GetChannelBySortOrder<ChannelType>(ChannelIndex);
 		if (bAddKey && Channel)
 		{
 			bool bShouldKeyChannel = bKeyEvenIfUnchanged;
@@ -84,7 +84,7 @@ struct TAddKeyImpl : IImpl
 	{
 		using namespace UE::MovieScene;
 
-		ChannelType* Channel = Proxy.GetChannel<ChannelType>(ChannelIndex);
+		ChannelType* Channel = Proxy.GetChannelBySortOrder<ChannelType>(ChannelIndex);
 		if (Channel && Channel->GetData().GetTimes().Num() == 0)
 		{
 			if (Section->TryModify())
@@ -115,7 +115,7 @@ struct TAddKeyImpl<FMovieSceneBoolChannel, bool> : IImpl
 		bool bKeyCreated = false;
 		using namespace UE::MovieScene;
 
-		FMovieSceneBoolChannel* Channel = Proxy.GetChannel<FMovieSceneBoolChannel>(ChannelIndex);
+		FMovieSceneBoolChannel* Channel = Proxy.GetChannelBySortOrder<FMovieSceneBoolChannel>(ChannelIndex);
 		if (bAddKey && Channel)
 		{
 			bool bShouldKeyChannel = bKeyEvenIfUnchanged;
@@ -144,7 +144,7 @@ struct TAddKeyImpl<FMovieSceneBoolChannel, bool> : IImpl
 	{
 		using namespace UE::MovieScene;
 
-		FMovieSceneBoolChannel* Channel = Proxy.GetChannel<FMovieSceneBoolChannel>(ChannelIndex);
+		FMovieSceneBoolChannel* Channel = Proxy.GetChannelBySortOrder<FMovieSceneBoolChannel>(ChannelIndex);
 		if (Channel && Channel->GetData().GetTimes().Num() == 0 && Channel->GetDefault() != ValueToSet)
 		{
 			if (Section->TryModify())
@@ -174,7 +174,7 @@ struct TAddKeyImpl<FMovieSceneDoubleChannel, double> : IImpl
 		bool bKeyCreated = false;
 		using namespace UE::MovieScene;
 
-		FMovieSceneDoubleChannel* Channel = Proxy.GetChannel<FMovieSceneDoubleChannel>(ChannelIndex);
+		FMovieSceneDoubleChannel* Channel = Proxy.GetChannelBySortOrder<FMovieSceneDoubleChannel>(ChannelIndex);
 		if (bAddKey && Channel)
 		{
 			bool bShouldKeyChannel = bKeyEvenIfUnchanged;
@@ -203,7 +203,7 @@ struct TAddKeyImpl<FMovieSceneDoubleChannel, double> : IImpl
 	{
 		using namespace UE::MovieScene;
 
-		FMovieSceneDoubleChannel* Channel = Proxy.GetChannel<FMovieSceneDoubleChannel>(ChannelIndex);
+		FMovieSceneDoubleChannel* Channel = Proxy.GetChannelBySortOrder<FMovieSceneDoubleChannel>(ChannelIndex);
 		if (Channel && Channel->GetData().GetTimes().Num() == 0  && Channel->GetDefault() != ValueToSet)
 		{
 			if (Section->TryModify())
@@ -218,7 +218,7 @@ struct TAddKeyImpl<FMovieSceneDoubleChannel, double> : IImpl
 	{
 		using namespace UE::MovieScene;
 		double CurrentValue = *(double*)(VCurrentValue);
-		FMovieSceneDoubleChannel* Channel = Proxy.GetChannel<FMovieSceneDoubleChannel>(ChannelIndex);
+		FMovieSceneDoubleChannel* Channel = Proxy.GetChannelBySortOrder<FMovieSceneDoubleChannel>(ChannelIndex);
 		if (Channel)
 		{
 			double LocalValue;
@@ -251,7 +251,7 @@ struct TAddKeyImpl<FMovieSceneFloatChannel, float> : IImpl
 		bool bKeyCreated = false;
 		using namespace UE::MovieScene;
 
-		FMovieSceneFloatChannel* Channel = Proxy.GetChannel<FMovieSceneFloatChannel>(ChannelIndex);
+		FMovieSceneFloatChannel* Channel = Proxy.GetChannelBySortOrder<FMovieSceneFloatChannel>(ChannelIndex);
 		if (bAddKey && Channel)
 		{
 			bool bShouldKeyChannel = bKeyEvenIfUnchanged;
@@ -280,7 +280,7 @@ struct TAddKeyImpl<FMovieSceneFloatChannel, float> : IImpl
 	{
 		using namespace UE::MovieScene;
 
-		FMovieSceneFloatChannel* Channel = Proxy.GetChannel<FMovieSceneFloatChannel>(ChannelIndex);
+		FMovieSceneFloatChannel* Channel = Proxy.GetChannelBySortOrder<FMovieSceneFloatChannel>(ChannelIndex);
 		if (Channel && Channel->GetData().GetTimes().Num() == 0  && Channel->GetDefault() != ValueToSet)
 		{
 			if (Section->TryModify())
@@ -295,7 +295,7 @@ struct TAddKeyImpl<FMovieSceneFloatChannel, float> : IImpl
 	{
 		using namespace UE::MovieScene;
 		float CurrentValue = *(float*)(VCurrentValue);
-		FMovieSceneFloatChannel* Channel = Proxy.GetChannel<FMovieSceneFloatChannel>(ChannelIndex);
+		FMovieSceneFloatChannel* Channel = Proxy.GetChannelBySortOrder<FMovieSceneFloatChannel>(ChannelIndex);
 		if (Channel)
 		{
 			float LocalValue;
@@ -328,7 +328,7 @@ struct TAddKeyImpl<FMovieSceneIntegerChannel, int32> : IImpl
 		bool bKeyCreated = false;
 		using namespace UE::MovieScene;
 
-		FMovieSceneIntegerChannel* Channel = Proxy.GetChannel<FMovieSceneIntegerChannel>(ChannelIndex);
+		FMovieSceneIntegerChannel* Channel = Proxy.GetChannelBySortOrder<FMovieSceneIntegerChannel>(ChannelIndex);
 		if (bAddKey && Channel)
 		{
 			bool bShouldKeyChannel = bKeyEvenIfUnchanged;
@@ -357,7 +357,7 @@ struct TAddKeyImpl<FMovieSceneIntegerChannel, int32> : IImpl
 	{
 		using namespace UE::MovieScene;
 
-		FMovieSceneIntegerChannel* Channel = Proxy.GetChannel<FMovieSceneIntegerChannel>(ChannelIndex);
+		FMovieSceneIntegerChannel* Channel = Proxy.GetChannelBySortOrder<FMovieSceneIntegerChannel>(ChannelIndex);
 		if (Channel && Channel->GetData().GetTimes().Num() == 0 && Channel->GetDefault() != ValueToSet)
 		{
 			if (Section->TryModify())
@@ -372,7 +372,7 @@ struct TAddKeyImpl<FMovieSceneIntegerChannel, int32> : IImpl
 	{
 		using namespace UE::MovieScene;
 		int32 CurrentValue = *(int32*)(VCurrentValue);
-		FMovieSceneIntegerChannel* Channel = Proxy.GetChannel<FMovieSceneIntegerChannel>(ChannelIndex);
+		FMovieSceneIntegerChannel* Channel = Proxy.GetChannelBySortOrder<FMovieSceneIntegerChannel>(ChannelIndex);
 		if (Channel)
 		{
 			int32 LocalValue;
