Sequencer: Use name lookups for channels instead of assuming their type and index.

diff --git a/Engine/Source/Editor/MovieSceneTools/Private/TrackEditors/TransformTrackEditor.cpp b/Engine/Source/Editor/MovieSceneTools/Private/TrackEditors/TransformTrackEditor.cpp
--- a/Engine/Source/Editor/MovieSceneTools/Private/TrackEditors/TransformTrackEditor.cpp
+++ b/Engine/Source/Editor/MovieSceneTools/Private/TrackEditors/TransformTrackEditor.cpp
@@ -34,6 +34,10 @@
 #include "Tracks/IMovieSceneTransformOrigin.h"
 #include "IMovieScenePlaybackClient.h"
 
+#include "MVVM/SectionModelStorageExtension.h"
+#include "MVVM/ViewModels/EditorViewModel.h"
+#include "MVVM/ViewModels/SectionModel.h"
+
 #include "Editor.h"
 #include "LevelEditorViewport.h"
 #include "TransformConstraint.h"
@@ -665,11 +669,13 @@ void F3DTransformTrackEditor::GetTransformKeys( const TOptional<FTransformData>&
 	EMovieSceneTransformChannel TransformMask = TransformSection->GetMask().GetChannels();
 
 	using namespace UE::MovieScene;
+	using namespace UE::Sequencer;
 
-	bool bLastVectorIsValid = LastTransform.IsSet();
+	TSharedPtr<ISequencer> SequencerPtr = GetSequencer();
 
 	// If key all is enabled, for a key on all the channels
-	if (GetSequencer()->GetKeyGroupMode() == EKeyGroupMode::KeyAll)
+	bool bLastVectorIsValid = LastTransform.IsSet();
+	if (SequencerPtr->GetKeyGroupMode() == EKeyGroupMode::KeyAll)
 	{
 		bLastVectorIsValid = false;
 		ChannelsToKey = EMovieSceneTransformChannel::All;
@@ -679,6 +685,34 @@ void F3DTransformTrackEditor::GetTransformKeys( const TOptional<FTransformData>&
 
 	FTransformData RecomposedTransform = RecomposeTransform(CurrentTransform, Object, Section);
 
+	// Get the channel indices for our curve channels.
+	// We will ignore any channels that are anything else than FMovieSceneDoubleChannel (such as when
+	// a channel has been overriden by a procedural channel) because we assume, possibly wrongly, that
+	// anything else won't be keyable with a double value.
+	FViewModelPtr RootViewModel = SequencerPtr->GetViewModel()->GetRootModel();
+	TSharedPtr<FSectionModelStorageExtension> SectionStorageExtension = RootViewModel.ImplicitCast();
+	TSharedPtr<FSectionModel> SectionModel = SectionStorageExtension->FindModelForSection(Section);
+	FCachedChannelModels CachedChannels(SectionModel);
+
+	FMovieSceneChannelHandle ChannelHandles[] = {
+		CachedChannels.GetChannel("Location.X"),
+		CachedChannels.GetChannel("Location.Y"),
+		CachedChannels.GetChannel("Location.Z"),
+		CachedChannels.GetChannel("Rotation.X"),
+		CachedChannels.GetChannel("Rotation.Y"),
+		CachedChannels.GetChannel("Rotation.Z"),
+		CachedChannels.GetChannel("Scale.X"),
+		CachedChannels.GetChannel("Scale.Y"),
+		CachedChannels.GetChannel("Scale.Z")
+	};
+
+	bool IsChannelKeyable[9];
+	const FName DoubleChannelTypeName = FMovieSceneDoubleChannel::StaticStruct()->GetFName();
+	for (int32 Index = 0; Index < sizeof(IsChannelKeyable); ++Index)
+	{
+		IsChannelKeyable[Index] = (ChannelHandles[Index].GetChannelTypeName() == DoubleChannelTypeName);
+	}
+
 	// Set translation keys/defaults
 	{
 		bool bKeyX = EnumHasAnyFlags(ChannelsToKey, EMovieSceneTransformChannel::TranslationX);
@@ -712,9 +746,9 @@ void F3DTransformTrackEditor::GetTransformKeys( const TOptional<FTransformData>&
 
 		FVector KeyVector = RecomposedTransform.Translation;
 
-		OutGeneratedKeys.Add(FMovieSceneChannelValueSetter::Create<FMovieSceneDoubleChannel>(0, (double)KeyVector.X, bKeyX));
-		OutGeneratedKeys.Add(FMovieSceneChannelValueSetter::Create<FMovieSceneDoubleChannel>(1, (double)KeyVector.Y, bKeyY));
-		OutGeneratedKeys.Add(FMovieSceneChannelValueSetter::Create<FMovieSceneDoubleChannel>(2, (double)KeyVector.Z, bKeyZ));
+		if (IsChannelKeyable[0]) { OutGeneratedKeys.Add(FMovieSceneChannelValueSetter::Create<FMovieSceneDoubleChannel>(ChannelHandles[0].GetChannelIndex(), (double)KeyVector.X, bKeyX)); }
+		if (IsChannelKeyable[1]) { OutGeneratedKeys.Add(FMovieSceneChannelValueSetter::Create<FMovieSceneDoubleChannel>(ChannelHandles[1].GetChannelIndex(), (double)KeyVector.Y, bKeyY)); }
+		if (IsChannelKeyable[2]) { OutGeneratedKeys.Add(FMovieSceneChannelValueSetter::Create<FMovieSceneDoubleChannel>(ChannelHandles[2].GetChannelIndex(), (double)KeyVector.Z, bKeyZ)); }
 	}
 
 	// Set rotation keys/defaults
@@ -753,10 +787,9 @@ void F3DTransformTrackEditor::GetTransformKeys( const TOptional<FTransformData>&
 
 		// Do we need to unwind re-composed rotations?
 		KeyRotator = UnwindRotator(CurrentTransform.Rotation, RecomposedTransform.Rotation);
-		OutGeneratedKeys.Add(FMovieSceneChannelValueSetter::Create<FMovieSceneDoubleChannel>(3, (double)KeyRotator.Roll, bKeyX));
-		OutGeneratedKeys.Add(FMovieSceneChannelValueSetter::Create<FMovieSceneDoubleChannel>(4, (double)KeyRotator.Pitch, bKeyY));
-		OutGeneratedKeys.Add(FMovieSceneChannelValueSetter::Create<FMovieSceneDoubleChannel>(5, (double)KeyRotator.Yaw, bKeyZ));
-
+		if (IsChannelKeyable[3]) { OutGeneratedKeys.Add(FMovieSceneChannelValueSetter::Create<FMovieSceneDoubleChannel>(ChannelHandles[3].GetChannelIndex(), (double)KeyRotator.Roll, bKeyX)); }
+		if (IsChannelKeyable[4]) { OutGeneratedKeys.Add(FMovieSceneChannelValueSetter::Create<FMovieSceneDoubleChannel>(ChannelHandles[4].GetChannelIndex(), (double)KeyRotator.Pitch, bKeyY)); }
+		if (IsChannelKeyable[5]) { OutGeneratedKeys.Add(FMovieSceneChannelValueSetter::Create<FMovieSceneDoubleChannel>(ChannelHandles[5].GetChannelIndex(), (double)KeyRotator.Yaw, bKeyZ)); }
 	}
 
 	// Set scale keys/defaults
@@ -791,9 +824,9 @@ void F3DTransformTrackEditor::GetTransformKeys( const TOptional<FTransformData>&
 		}
 
 		FVector KeyVector = RecomposedTransform.Scale;
-		OutGeneratedKeys.Add(FMovieSceneChannelValueSetter::Create<FMovieSceneDoubleChannel>(6, (double)KeyVector.X, bKeyX));
-		OutGeneratedKeys.Add(FMovieSceneChannelValueSetter::Create<FMovieSceneDoubleChannel>(7, (double)KeyVector.Y, bKeyY));
-		OutGeneratedKeys.Add(FMovieSceneChannelValueSetter::Create<FMovieSceneDoubleChannel>(8, (double)KeyVector.Z, bKeyZ));
+		if (IsChannelKeyable[6]) { OutGeneratedKeys.Add(FMovieSceneChannelValueSetter::Create<FMovieSceneDoubleChannel>(ChannelHandles[6].GetChannelIndex(), (double)KeyVector.X, bKeyX)); }
+		if (IsChannelKeyable[7]) { OutGeneratedKeys.Add(FMovieSceneChannelValueSetter::Create<FMovieSceneDoubleChannel>(ChannelHandles[7].GetChannelIndex(), (double)KeyVector.Y, bKeyY)); }
+		if (IsChannelKeyable[8]) { OutGeneratedKeys.Add(FMovieSceneChannelValueSetter::Create<FMovieSceneDoubleChannel>(ChannelHandles[8].GetChannelIndex(), (double)KeyVector.Z, bKeyZ)); }
 	}
 }
 
@@ -1143,7 +1176,7 @@ void F3DTransformTrackEditor::ProcessKeyOperation(UObject* ObjectToKey, TArrayVi
 			for (TSharedPtr<IKeyArea> KeyArea : SectionsToKey[Index].KeyAreas)
 			{
 				FMovieSceneChannelHandle Handle  = KeyArea->GetChannel();
-				if (Handle.GetChannelTypeName() == FMovieSceneDoubleChannel::StaticStruct()->GetFName() && Handle.GetChannelIndex() < 9)
+				if (Handle.GetChannelTypeName() == FMovieSceneDoubleChannel::StaticStruct()->GetFName())
 				{
 					FMovieSceneDoubleChannel* Channel = static_cast<FMovieSceneDoubleChannel*>(Handle.Get());
 
@@ -1224,7 +1257,7 @@ void F3DTransformTrackEditor::ImportAnimSequenceTransforms(const FAssetData& Ass
 		}
 	}
 
-	if(AnimSequence && AnimSequence->GetDataModel()->GetNumBoneTracks() > 0)
+	if (AnimSequence && AnimSequence->GetDataModel()->GetNumBoneTracks() > 0)
 	{
 		const FScopedTransaction Transaction( NSLOCTEXT( "Sequencer", "ImportAnimSequenceTransforms", "Import Anim Sequence Transforms" ) );
 
@@ -1237,6 +1270,8 @@ void F3DTransformTrackEditor::ImportAnimSequenceTransforms(const FAssetData& Ass
 		
 		FFrameRate TickResolution = Section->GetTypedOuter<UMovieScene>()->GetTickResolution();
 
+		// We know that there aren't any channel overrides (all 9 channels are double channels) because we're working
+		// on a brand new transform section.
 		TArrayView<FMovieSceneDoubleChannel*> DoubleChannels = Section->GetChannelProxy().GetChannels<FMovieSceneDoubleChannel>();
 
 		// Set default translation and rotation
diff --git a/Engine/Source/Editor/MovieSceneTools/Public/KeyframeTrackEditor.h b/Engine/Source/Editor/MovieSceneTools/Public/KeyframeTrackEditor.h
--- a/Engine/Source/Editor/MovieSceneTools/Public/KeyframeTrackEditor.h
+++ b/Engine/Source/Editor/MovieSceneTools/Public/KeyframeTrackEditor.h
@@ -55,7 +55,7 @@ struct TAddKeyImpl : IImpl
 		bool bKeyCreated = false;
 		using namespace UE::MovieScene;
 
-		ChannelType* Channel = Proxy.GetChannelBySortOrder<ChannelType>(ChannelIndex);
+		ChannelType* Channel = Proxy.GetChannel<ChannelType>(ChannelIndex);
 		if (bAddKey && Channel)
 		{
 			bool bShouldKeyChannel = bKeyEvenIfUnchanged;
@@ -84,7 +84,7 @@ struct TAddKeyImpl : IImpl
 	{
 		using namespace UE::MovieScene;
 
-		ChannelType* Channel = Proxy.GetChannelBySortOrder<ChannelType>(ChannelIndex);
+		ChannelType* Channel = Proxy.GetChannel<ChannelType>(ChannelIndex);
 		if (Channel && Channel->GetData().GetTimes().Num() == 0)
 		{
 			if (Section->TryModify())
@@ -115,7 +115,7 @@ struct TAddKeyImpl<FMovieSceneBoolChannel, bool> : IImpl
 		bool bKeyCreated = false;
 		using namespace UE::MovieScene;
 
-		FMovieSceneBoolChannel* Channel = Proxy.GetChannelBySortOrder<FMovieSceneBoolChannel>(ChannelIndex);
+		FMovieSceneBoolChannel* Channel = Proxy.GetChannel<FMovieSceneBoolChannel>(ChannelIndex);
 		if (bAddKey && Channel)
 		{
 			bool bShouldKeyChannel = bKeyEvenIfUnchanged;
@@ -144,7 +144,7 @@ struct TAddKeyImpl<FMovieSceneBoolChannel, bool> : IImpl
 	{
 		using namespace UE::MovieScene;
 
-		FMovieSceneBoolChannel* Channel = Proxy.GetChannelBySortOrder<FMovieSceneBoolChannel>(ChannelIndex);
+		FMovieSceneBoolChannel* Channel = Proxy.GetChannel<FMovieSceneBoolChannel>(ChannelIndex);
 		if (Channel && Channel->GetData().GetTimes().Num() == 0 && Channel->GetDefault() != ValueToSet)
 		{
 			if (Section->TryModify())
@@ -174,7 +174,7 @@ struct TAddKeyImpl<FMovieSceneDoubleChannel, double> : IImpl
 		bool bKeyCreated = false;
 		using namespace UE::MovieScene;
 
-		FMovieSceneDoubleChannel* Channel = Proxy.GetChannelBySortOrder<FMovieSceneDoubleChannel>(ChannelIndex);
+		FMovieSceneDoubleChannel* Channel = Proxy.GetChannel<FMovieSceneDoubleChannel>(ChannelIndex);
 		if (bAddKey && Channel)
 		{
 			bool bShouldKeyChannel = bKeyEvenIfUnchanged;
@@ -203,7 +203,7 @@ struct TAddKeyImpl<FMovieSceneDoubleChannel, double> : IImpl
 	{
 		using namespace UE::MovieScene;
 
-		FMovieSceneDoubleChannel* Channel = Proxy.GetChannelBySortOrder<FMovieSceneDoubleChannel>(ChannelIndex);
+		FMovieSceneDoubleChannel* Channel = Proxy.GetChannel<FMovieSceneDoubleChannel>(ChannelIndex);
 		if (Channel && Channel->GetData().GetTimes().Num() == 0  && Channel->GetDefault() != ValueToSet)
 		{
 			if (Section->TryModify())
@@ -218,7 +218,7 @@ struct TAddKeyImpl<FMovieSceneDoubleChannel, double> : IImpl
 	{
 		using namespace UE::MovieScene;
 		double CurrentValue = *(double*)(VCurrentValue);
-		FMovieSceneDoubleChannel* Channel = Proxy.GetChannelBySortOrder<FMovieSceneDoubleChannel>(ChannelIndex);
+		FMovieSceneDoubleChannel* Channel = Proxy.GetChannel<FMovieSceneDoubleChannel>(ChannelIndex);
 		if (Channel)
 		{
 			double LocalValue;
@@ -251,7 +251,7 @@ struct TAddKeyImpl<FMovieSceneFloatChannel, float> : IImpl
 		bool bKeyCreated = false;
 		using namespace UE::MovieScene;
 
-		FMovieSceneFloatChannel* Channel = Proxy.GetChannelBySortOrder<FMovieSceneFloatChannel>(ChannelIndex);
+		FMovieSceneFloatChannel* Channel = Proxy.GetChannel<FMovieSceneFloatChannel>(ChannelIndex);
 		if (bAddKey && Channel)
 		{
 			bool bShouldKeyChannel = bKeyEvenIfUnchanged;
@@ -280,7 +280,7 @@ struct TAddKeyImpl<FMovieSceneFloatChannel, float> : IImpl
 	{
 		using namespace UE::MovieScene;
 
-		FMovieSceneFloatChannel* Channel = Proxy.GetChannelBySortOrder<FMovieSceneFloatChannel>(ChannelIndex);
+		FMovieSceneFloatChannel* Channel = Proxy.GetChannel<FMovieSceneFloatChannel>(ChannelIndex);
 		if (Channel && Channel->GetData().GetTimes().Num() == 0  && Channel->GetDefault() != ValueToSet)
 		{
 			if (Section->TryModify())
@@ -295,7 +295,7 @@ struct TAddKeyImpl<FMovieSceneFloatChannel, float> : IImpl
 	{
 		using namespace UE::MovieScene;
 		float CurrentValue = *(float*)(VCurrentValue);
-		FMovieSceneFloatChannel* Channel = Proxy.GetChannelBySortOrder<FMovieSceneFloatChannel>(ChannelIndex);
+		FMovieSceneFloatChannel* Channel = Proxy.GetChannel<FMovieSceneFloatChannel>(ChannelIndex);
 		if (Channel)
 		{
 			float LocalValue;
@@ -328,7 +328,7 @@ struct TAddKeyImpl<FMovieSceneIntegerChannel, int32> : IImpl
 		bool bKeyCreated = false;
 		using namespace UE::MovieScene;
 
-		FMovieSceneIntegerChannel* Channel = Proxy.GetChannelBySortOrder<FMovieSceneIntegerChannel>(ChannelIndex);
+		FMovieSceneIntegerChannel* Channel = Proxy.GetChannel<FMovieSceneIntegerChannel>(ChannelIndex);
 		if (bAddKey && Channel)
 		{
 			bool bShouldKeyChannel = bKeyEvenIfUnchanged;
@@ -357,7 +357,7 @@ struct TAddKeyImpl<FMovieSceneIntegerChannel, int32> : IImpl
 	{
 		using namespace UE::MovieScene;
 
-		FMovieSceneIntegerChannel* Channel = Proxy.GetChannelBySortOrder<FMovieSceneIntegerChannel>(ChannelIndex);
+		FMovieSceneIntegerChannel* Channel = Proxy.GetChannel<FMovieSceneIntegerChannel>(ChannelIndex);
 		if (Channel && Channel->GetData().GetTimes().Num() == 0 && Channel->GetDefault() != ValueToSet)
 		{
 			if (Section->TryModify())
@@ -372,7 +372,7 @@ struct TAddKeyImpl<FMovieSceneIntegerChannel, int32> : IImpl
 	{
 		using namespace UE::MovieScene;
 		int32 CurrentValue = *(int32*)(VCurrentValue);
-		FMovieSceneIntegerChannel* Channel = Proxy.GetChannelBySortOrder<FMovieSceneIntegerChannel>(ChannelIndex);
+		FMovieSceneIntegerChannel* Channel = Proxy.GetChannel<FMovieSceneIntegerChannel>(ChannelIndex);
 		if (Channel)
 		{
 			int32 LocalValue;
diff --git a/Engine/Source/Editor/Sequencer/Private/MVVM/ViewModels/SectionModel.cpp b/Engine/Source/Editor/Sequencer/Private/MVVM/ViewModels/SectionModel.cpp
--- a/Engine/Source/Editor/Sequencer/Private/MVVM/ViewModels/SectionModel.cpp
+++ b/Engine/Source/Editor/Sequencer/Private/MVVM/ViewModels/SectionModel.cpp
@@ -5,6 +5,7 @@
 #include "MVVM/ViewModels/ViewModelIterators.h"
 #include "MVVM/ViewModels/SequencerEditorViewModel.h"
 #include "MVVM/Extensions/ITrackExtension.h"
+#include "IKeyArea.h"
 #include "Sequencer.h"
 #include "SSequencerSection.h"
 
@@ -540,6 +541,28 @@ TArray<FOverlappingSections> FSectionModel::GetEasingSegments()
 	return Result;
 }
 
+FCachedChannelModels::FCachedChannelModels(TSharedPtr<FSectionModel> InSection)
+{
+	for (TSharedPtr<FChannelModel> ChannelModel : InSection->GetDescendantsOfType<FChannelModel>())
+	{
+		if (TSharedPtr<IKeyArea> KeyArea = ChannelModel->GetKeyArea())
+		{
+			const FName ChannelName = KeyArea->GetName();
+			const FMovieSceneChannelHandle ChannelHandle = KeyArea->GetChannel();
+			ChannelsByName.Add(ChannelName, FCachedChannelModel{ ChannelHandle, ChannelModel });
+		}
+	}
+}
+
+FMovieSceneChannelHandle FCachedChannelModels::GetChannel(FName InChannelName)
+{
+	if (FCachedChannelModel* CachedChannelModel = ChannelsByName.Find(InChannelName))
+	{
+		return CachedChannelModel->Handle;
+	}
+	return FMovieSceneChannelHandle();
+}
+
 } // namespace Sequencer
 } // namespace UE
 
diff --git a/Engine/Source/Editor/Sequencer/Public/MVVM/SectionModelStorageExtension.h b/Engine/Source/Editor/Sequencer/Public/MVVM/SectionModelStorageExtension.h
--- a/Engine/Source/Editor/Sequencer/Public/MVVM/SectionModelStorageExtension.h
+++ b/Engine/Source/Editor/Sequencer/Public/MVVM/SectionModelStorageExtension.h
@@ -19,7 +19,7 @@ namespace Sequencer
 class FSectionModel;
 class FSequenceModel;
 
-class FSectionModelStorageExtension
+class SEQUENCER_API FSectionModelStorageExtension
 	: public IDynamicExtension
 {
 public:
diff --git a/Engine/Source/Editor/Sequencer/Public/MVVM/ViewModels/SectionModel.h b/Engine/Source/Editor/Sequencer/Public/MVVM/ViewModels/SectionModel.h
--- a/Engine/Source/Editor/Sequencer/Public/MVVM/ViewModels/SectionModel.h
+++ b/Engine/Source/Editor/Sequencer/Public/MVVM/ViewModels/SectionModel.h
@@ -2,6 +2,7 @@
 
 #pragma once
 
+#include "Channels/MovieSceneChannelHandle.h"
 #include "MVVM/ViewModels/ViewModel.h"
 #include "MVVM/ViewModels/ViewModelHierarchy.h"
 #include "MVVM/Extensions/ILayerBarExtension.h"
@@ -25,6 +26,7 @@ namespace Sequencer
 {
 
 class ITrackExtension;
+class FChannelModel;
 struct FOverlappingSections;
 
 /**
@@ -130,6 +132,25 @@ private:
 	TRange<FFrameNumber> LayerBarRange;
 };
 
+/**
+ * Stores a snapshot of the channels contained in a section
+ */
+class SEQUENCER_API FCachedChannelModels
+{
+public:
+	FCachedChannelModels(TSharedPtr<FSectionModel> InSection);
+
+	FMovieSceneChannelHandle GetChannel(FName InChannelName);
+
+private:
+	struct FCachedChannelModel
+	{
+		FMovieSceneChannelHandle Handle;
+		TWeakViewModelPtr<FChannelModel> Model;
+	};
+	TMap<FName, FCachedChannelModel> ChannelsByName;
+};
+
 struct FOverlappingSections
 {
 	/** The range for the overlap */
