Sequencer: add lookup table to channel proxies for finding channels by name

Looking up channels by name was done via the view-models in the transform track editor, but view-models might not be created if a keying operation auto-created the transform track or the section. So now we are doing this at the data-level.

Remove lookup by sort order that we don't need anymore, and remove lookup by name in view-models that we don't need anymore either.

diff --git a/Engine/Source/Editor/MovieSceneTools/Private/TrackEditors/TransformTrackEditor.cpp b/Engine/Source/Editor/MovieSceneTools/Private/TrackEditors/TransformTrackEditor.cpp
--- a/Engine/Source/Editor/MovieSceneTools/Private/TrackEditors/TransformTrackEditor.cpp
+++ b/Engine/Source/Editor/MovieSceneTools/Private/TrackEditors/TransformTrackEditor.cpp
@@ -34,10 +34,6 @@
 #include "Tracks/IMovieSceneTransformOrigin.h"
 #include "IMovieScenePlaybackClient.h"
 
-#include "MVVM/SectionModelStorageExtension.h"
-#include "MVVM/ViewModels/EditorViewModel.h"
-#include "MVVM/ViewModels/SectionModel.h"
-
 #include "Editor.h"
 #include "LevelEditorViewport.h"
 #include "TransformConstraint.h"
@@ -669,7 +665,6 @@ void F3DTransformTrackEditor::GetTransformKeys( const TOptional<FTransformData>&
 	EMovieSceneTransformChannel TransformMask = TransformSection->GetMask().GetChannels();
 
 	using namespace UE::MovieScene;
-	using namespace UE::Sequencer;
 
 	TSharedPtr<ISequencer> SequencerPtr = GetSequencer();
 
@@ -686,38 +681,23 @@ void F3DTransformTrackEditor::GetTransformKeys( const TOptional<FTransformData>&
 	FTransformData RecomposedTransform = RecomposeTransform(CurrentTransform, Object, Section);
 
 	// Get the channel indices for our curve channels.
-	// We will ignore any channels that are anything else than FMovieSceneDoubleChannel (such as when
-	// a channel has been overriden by a procedural channel) because we assume, possibly wrongly, that
-	// anything else won't be keyable with a double value.
-	FViewModelPtr RootViewModel = SequencerPtr->GetViewModel()->GetRootModel();
-	TSharedPtr<FSectionModelStorageExtension> SectionStorageExtension = RootViewModel.ImplicitCast();
-	TSharedPtr<FSectionModel> SectionModel = SectionStorageExtension->FindModelForSection(Section);
-	FCachedChannelModels CachedChannels(SectionModel);
-
-	FMovieSceneChannelHandle ChannelHandles[] = {
-		CachedChannels.GetChannel("Location.X"),
-		CachedChannels.GetChannel("Location.Y"),
-		CachedChannels.GetChannel("Location.Z"),
-		CachedChannels.GetChannel("Rotation.X"),
-		CachedChannels.GetChannel("Rotation.Y"),
-		CachedChannels.GetChannel("Rotation.Z"),
-		CachedChannels.GetChannel("Scale.X"),
-		CachedChannels.GetChannel("Scale.Y"),
-		CachedChannels.GetChannel("Scale.Z")
+	// We will get invalid handles for any channels that are anything else than FMovieSceneDoubleChannel (such as when
+	// a channel has been overriden by a procedural channel). This means that the value setters below will get null
+	// channels and will simply ignore those non-double channels. This is what we want, since we assume (perhaps
+	// incorrectly) that overriden channels are most probably not keyable.
+	FMovieSceneChannelProxy& SectionChannelProxy = Section->GetChannelProxy();
+	TMovieSceneChannelHandle<FMovieSceneDoubleChannel> ChannelHandles[] = {
+		SectionChannelProxy.GetChannelByName<FMovieSceneDoubleChannel>("Location.X"),
+		SectionChannelProxy.GetChannelByName<FMovieSceneDoubleChannel>("Location.Y"),
+		SectionChannelProxy.GetChannelByName<FMovieSceneDoubleChannel>("Location.Z"),
+		SectionChannelProxy.GetChannelByName<FMovieSceneDoubleChannel>("Rotation.X"),
+		SectionChannelProxy.GetChannelByName<FMovieSceneDoubleChannel>("Rotation.Y"),
+		SectionChannelProxy.GetChannelByName<FMovieSceneDoubleChannel>("Rotation.Z"),
+		SectionChannelProxy.GetChannelByName<FMovieSceneDoubleChannel>("Scale.X"),
+		SectionChannelProxy.GetChannelByName<FMovieSceneDoubleChannel>("Scale.Y"),
+		SectionChannelProxy.GetChannelByName<FMovieSceneDoubleChannel>("Scale.Z")
 	};
 
-	bool IsChannelKeyable[9];
-	const FName DoubleChannelTypeName = FMovieSceneDoubleChannel::StaticStruct()->GetFName();
-	TWeakPtr<FMovieSceneChannelProxy> WeakSectionProxy = Section->GetChannelProxy().AsWeak();
-	for (int32 Index = 0; Index < sizeof(IsChannelKeyable); ++Index)
-	{
-		if (ChannelHandles[Index] == FMovieSceneChannelHandle())
-		{
-			ChannelHandles[Index] = FMovieSceneChannelHandle(WeakSectionProxy, FMovieSceneDoubleChannel::StaticStruct()->GetFName(), Index);
-		}
-		IsChannelKeyable[Index] = (ChannelHandles[Index].GetChannelTypeName() == DoubleChannelTypeName);
-	}
-
 	// Set translation keys/defaults
 	{
 		bool bKeyX = EnumHasAnyFlags(ChannelsToKey, EMovieSceneTransformChannel::TranslationX);
@@ -751,9 +731,9 @@ void F3DTransformTrackEditor::GetTransformKeys( const TOptional<FTransformData>&
 
 		FVector KeyVector = RecomposedTransform.Translation;
 
-		if (IsChannelKeyable[0]) { OutGeneratedKeys.Add(FMovieSceneChannelValueSetter::Create<FMovieSceneDoubleChannel>(ChannelHandles[0].GetChannelIndex(), (double)KeyVector.X, bKeyX)); }
-		if (IsChannelKeyable[1]) { OutGeneratedKeys.Add(FMovieSceneChannelValueSetter::Create<FMovieSceneDoubleChannel>(ChannelHandles[1].GetChannelIndex(), (double)KeyVector.Y, bKeyY)); }
-		if (IsChannelKeyable[2]) { OutGeneratedKeys.Add(FMovieSceneChannelValueSetter::Create<FMovieSceneDoubleChannel>(ChannelHandles[2].GetChannelIndex(), (double)KeyVector.Z, bKeyZ)); }
+		OutGeneratedKeys.Add(FMovieSceneChannelValueSetter::Create<FMovieSceneDoubleChannel>(ChannelHandles[0].GetChannelIndex(), (double)KeyVector.X, bKeyX));
+		OutGeneratedKeys.Add(FMovieSceneChannelValueSetter::Create<FMovieSceneDoubleChannel>(ChannelHandles[1].GetChannelIndex(), (double)KeyVector.Y, bKeyY));
+		OutGeneratedKeys.Add(FMovieSceneChannelValueSetter::Create<FMovieSceneDoubleChannel>(ChannelHandles[2].GetChannelIndex(), (double)KeyVector.Z, bKeyZ));
 	}
 
 	// Set rotation keys/defaults
@@ -792,9 +772,9 @@ void F3DTransformTrackEditor::GetTransformKeys( const TOptional<FTransformData>&
 
 		// Do we need to unwind re-composed rotations?
 		KeyRotator = UnwindRotator(CurrentTransform.Rotation, RecomposedTransform.Rotation);
-		if (IsChannelKeyable[3]) { OutGeneratedKeys.Add(FMovieSceneChannelValueSetter::Create<FMovieSceneDoubleChannel>(ChannelHandles[3].GetChannelIndex(), (double)KeyRotator.Roll, bKeyX)); }
-		if (IsChannelKeyable[4]) { OutGeneratedKeys.Add(FMovieSceneChannelValueSetter::Create<FMovieSceneDoubleChannel>(ChannelHandles[4].GetChannelIndex(), (double)KeyRotator.Pitch, bKeyY)); }
-		if (IsChannelKeyable[5]) { OutGeneratedKeys.Add(FMovieSceneChannelValueSetter::Create<FMovieSceneDoubleChannel>(ChannelHandles[5].GetChannelIndex(), (double)KeyRotator.Yaw, bKeyZ)); }
+		OutGeneratedKeys.Add(FMovieSceneChannelValueSetter::Create<FMovieSceneDoubleChannel>(ChannelHandles[3].GetChannelIndex(), (double)KeyRotator.Roll, bKeyX));
+		OutGeneratedKeys.Add(FMovieSceneChannelValueSetter::Create<FMovieSceneDoubleChannel>(ChannelHandles[4].GetChannelIndex(), (double)KeyRotator.Pitch, bKeyY));
+		OutGeneratedKeys.Add(FMovieSceneChannelValueSetter::Create<FMovieSceneDoubleChannel>(ChannelHandles[5].GetChannelIndex(), (double)KeyRotator.Yaw, bKeyZ));
 	}
 
 	// Set scale keys/defaults
@@ -829,9 +809,9 @@ void F3DTransformTrackEditor::GetTransformKeys( const TOptional<FTransformData>&
 		}
 
 		FVector KeyVector = RecomposedTransform.Scale;
-		if (IsChannelKeyable[6]) { OutGeneratedKeys.Add(FMovieSceneChannelValueSetter::Create<FMovieSceneDoubleChannel>(ChannelHandles[6].GetChannelIndex(), (double)KeyVector.X, bKeyX)); }
-		if (IsChannelKeyable[7]) { OutGeneratedKeys.Add(FMovieSceneChannelValueSetter::Create<FMovieSceneDoubleChannel>(ChannelHandles[7].GetChannelIndex(), (double)KeyVector.Y, bKeyY)); }
-		if (IsChannelKeyable[8]) { OutGeneratedKeys.Add(FMovieSceneChannelValueSetter::Create<FMovieSceneDoubleChannel>(ChannelHandles[8].GetChannelIndex(), (double)KeyVector.Z, bKeyZ)); }
+		OutGeneratedKeys.Add(FMovieSceneChannelValueSetter::Create<FMovieSceneDoubleChannel>(ChannelHandles[6].GetChannelIndex(), (double)KeyVector.X, bKeyX));
+		OutGeneratedKeys.Add(FMovieSceneChannelValueSetter::Create<FMovieSceneDoubleChannel>(ChannelHandles[7].GetChannelIndex(), (double)KeyVector.Y, bKeyY));
+		OutGeneratedKeys.Add(FMovieSceneChannelValueSetter::Create<FMovieSceneDoubleChannel>(ChannelHandles[8].GetChannelIndex(), (double)KeyVector.Z, bKeyZ));
 	}
 }
 
diff --git a/Engine/Source/Editor/Sequencer/Private/MVVM/ViewModels/SectionModel.cpp b/Engine/Source/Editor/Sequencer/Private/MVVM/ViewModels/SectionModel.cpp
--- a/Engine/Source/Editor/Sequencer/Private/MVVM/ViewModels/SectionModel.cpp
+++ b/Engine/Source/Editor/Sequencer/Private/MVVM/ViewModels/SectionModel.cpp
@@ -541,31 +541,6 @@ TArray<FOverlappingSections> FSectionModel::GetEasingSegments()
 	return Result;
 }
 
-FCachedChannelModels::FCachedChannelModels(TSharedPtr<FSectionModel> InSection)
-{
-	if (InSection)
-	{
-		for (TSharedPtr<FChannelModel> ChannelModel : InSection->GetDescendantsOfType<FChannelModel>())
-		{
-			if (TSharedPtr<IKeyArea> KeyArea = ChannelModel->GetKeyArea())
-			{
-				const FName ChannelName = KeyArea->GetName();
-				const FMovieSceneChannelHandle ChannelHandle = KeyArea->GetChannel();
-				ChannelsByName.Add(ChannelName, FCachedChannelModel{ ChannelHandle, ChannelModel });
-			}
-		}
-	}
-}
-
-FMovieSceneChannelHandle FCachedChannelModels::GetChannel(FName InChannelName)
-{
-	if (FCachedChannelModel* CachedChannelModel = ChannelsByName.Find(InChannelName))
-	{
-		return CachedChannelModel->Handle;
-	}
-	return FMovieSceneChannelHandle();
-}
-
 } // namespace Sequencer
 } // namespace UE
 
diff --git a/Engine/Source/Editor/Sequencer/Public/MVVM/ViewModels/SectionModel.h b/Engine/Source/Editor/Sequencer/Public/MVVM/ViewModels/SectionModel.h
--- a/Engine/Source/Editor/Sequencer/Public/MVVM/ViewModels/SectionModel.h
+++ b/Engine/Source/Editor/Sequencer/Public/MVVM/ViewModels/SectionModel.h
@@ -132,25 +132,6 @@ private:
 	TRange<FFrameNumber> LayerBarRange;
 };
 
-/**
- * Stores a snapshot of the channels contained in a section
- */
-class SEQUENCER_API FCachedChannelModels
-{
-public:
-	FCachedChannelModels(TSharedPtr<FSectionModel> InSection);
-
-	FMovieSceneChannelHandle GetChannel(FName InChannelName);
-
-private:
-	struct FCachedChannelModel
-	{
-		FMovieSceneChannelHandle Handle;
-		TWeakViewModelPtr<FChannelModel> Model;
-	};
-	TMap<FName, FCachedChannelModel> ChannelsByName;
-};
-
 struct FOverlappingSections
 {
 	/** The range for the overlap */
diff --git a/Engine/Source/Runtime/MovieScene/Private/Channels/MovieSceneChannelProxy.cpp b/Engine/Source/Runtime/MovieScene/Private/Channels/MovieSceneChannelProxy.cpp
--- a/Engine/Source/Runtime/MovieScene/Private/Channels/MovieSceneChannelProxy.cpp
+++ b/Engine/Source/Runtime/MovieScene/Private/Channels/MovieSceneChannelProxy.cpp
@@ -6,7 +6,7 @@
 
 FMovieSceneChannelHandle FMovieSceneChannelProxy::MakeHandle(FName ChannelTypeName, int32 Index)
 {
-	TWeakPtr<FMovieSceneChannelProxy> WeakProxy = TSharedPtr<FMovieSceneChannelProxy>(AsShared());
+	TWeakPtr<FMovieSceneChannelProxy> WeakProxy = AsShared();
 	return FMovieSceneChannelHandle(WeakProxy, ChannelTypeName, Index);
 }
 
@@ -45,18 +45,38 @@ FMovieSceneChannel* FMovieSceneChannelProxy::GetChannel(FName ChannelTypeName, i
 
 #if WITH_EDITOR
 
-FMovieSceneChannel* FMovieSceneChannelProxy::GetChannelBySortOrder(FName ChannelTypeName, int32 ChannelSortOrder) const
+void FMovieSceneChannelProxy::EnsureHandlesByNamePopulated() const
 {
-	if (const FMovieSceneChannelEntry* Entry = FindEntry(ChannelTypeName))
+	if (bHandlesByNamePopulated)
 	{
-		TArrayView<FMovieSceneChannel* const> Channels = Entry->GetChannels();
-		const int32 ChannelIndex = Entry->GetChannelIndexBySortOrder(ChannelSortOrder);
-		if (ChannelIndex != INDEX_NONE)
+		return;
+	}
+
+	bHandlesByNamePopulated = true;
+	HandlesByName.Empty();
+
+	TWeakPtr<FMovieSceneChannelProxy> WeakNonConstThis = const_cast<FMovieSceneChannelProxy*>(this)->AsShared();
+
+	for (const FMovieSceneChannelEntry& Entry : Entries)
+	{
+		const FName ChannelTypeName = Entry.GetChannelTypeName();
+		TArrayView<const FMovieSceneChannelMetaData> ChannelMetaDatas = Entry.GetMetaData();
+		for (int32 Index = 0; Index < ChannelMetaDatas.Num(); ++Index)
 		{
-			return Channels[ChannelIndex];
+			const FMovieSceneChannelMetaData& ChannelMetaData = ChannelMetaDatas[Index];
+			HandlesByName.Add(ChannelMetaData.Name, FMovieSceneChannelHandle(WeakNonConstThis, ChannelTypeName, Index));
 		}
 	}
-	return nullptr;
+}
+
+FMovieSceneChannelHandle FMovieSceneChannelProxy::GetChannelByName(FName ChannelName) const
+{
+	EnsureHandlesByNamePopulated();
+	if (const FMovieSceneChannelHandle* Handle = HandlesByName.Find(ChannelName))
+	{
+		return *Handle;
+	}
+	return FMovieSceneChannelHandle();
 }
 
 #endif // WITH_EDITOR
diff --git a/Engine/Source/Runtime/MovieScene/Public/Channels/MovieSceneChannelEditorDataEntry.h b/Engine/Source/Runtime/MovieScene/Public/Channels/MovieSceneChannelEditorDataEntry.h
--- a/Engine/Source/Runtime/MovieScene/Public/Channels/MovieSceneChannelEditorDataEntry.h
+++ b/Engine/Source/Runtime/MovieScene/Public/Channels/MovieSceneChannelEditorDataEntry.h
@@ -4,7 +4,6 @@
 
 #include "UObject/ObjectMacros.h"
 #include "Misc/InlineValue.h"
-#include "Algo/BinarySearch.h"
 #include "Algo/Sort.h"
 #include "Misc/Attribute.h"
 #include "Misc/Optional.h"
@@ -37,20 +36,6 @@ struct FMovieSceneChannelEditorDataEntry
 		return MetaDataArray;
 	}
 
-	/**
-	 * Returns the index of the first channel with the given sort order
-	 * Note that there can be multiple channels with the same sort order
-	 */
-	int32 GetChannelIndexBySortOrder(uint32 SortOrder) const
-	{
-		const int32 Index = Algo::BinarySearchBy(SortOrderedArray, SortOrder, &FChannelSortOrderIndex::SortOrder);
-		if (Index != INDEX_NONE)
-		{
-			return SortOrderedArray[Index].ChannelIndex;
-		}
-		return INDEX_NONE;
-	}
-
 	/**
 	 * Access the extended editor data for a specific channel
 	 */
@@ -73,9 +58,6 @@ protected:
 	{
 		static_assert(TIsSame<typename TMovieSceneChannelTraits<ChannelType>::ExtendedEditorDataType, void>::Value, "Must supply extended editor data according to the channel's traits.");
 
-		// Keep track of channels in sort order
-		InsertSortOrderEntry(FChannelSortOrderIndex { MetaData.SortOrder, MetaDataArray.Num() }, MetaData.Name);
-
 		// Add the editor meta-data
 		MetaDataArray.Add(MetaData);
 	}
@@ -86,9 +68,6 @@ protected:
 	template<typename ChannelType, typename ExtendedEditorDataType>
 	void AddMetaData(const FMovieSceneChannelMetaData& MetaData, ExtendedEditorDataType&& InExtendedEditorData)
 	{
-		// Keep track of channels in sort order
-		InsertSortOrderEntry(FChannelSortOrderIndex { MetaData.SortOrder, MetaDataArray.Num() }, MetaData.Name);
-
 		// Add the editor meta-data
 		MetaDataArray.Add(MetaData);
 
@@ -122,12 +101,6 @@ protected:
 
 private:
 
-	struct FChannelSortOrderIndex
-	{
-		uint32 SortOrder;
-		int32 ChannelIndex;
-	};
-
 	/** Construct the extended editor data container for channel types that require it */
 	template<typename ChannelType, typename ExtendedEditorDataType>
 	void ConstructExtendedEditorDataArray(ExtendedEditorDataType*)
@@ -140,33 +113,11 @@ private:
 	{
 	}
 
-	void InsertSortOrderEntry(FChannelSortOrderIndex InEntry, FName InChannelName)
-	{
-		check(SortOrderedArray.Num() == MetaDataArray.Num());
-
-		// Get the first element where SortOrder >= InEntry.SortOrder.
-		int32 Index = Algo::LowerBoundBy(SortOrderedArray, InEntry.SortOrder, &FChannelSortOrderIndex::SortOrder);
-		// Keep going if there are more entries with the same sort order, but lesser names.
-		while (Index < SortOrderedArray.Num() - 1)
-		{
-			const FMovieSceneChannelMetaData& NextEntry(MetaDataArray[Index + 1]);
-			if (NextEntry.SortOrder != InEntry.SortOrder || NextEntry.Name.Compare(InChannelName) > 0)
-			{
-				break;
-			}
-			++Index;
-		}
-		SortOrderedArray.Insert(InEntry, Index);
-	}
-
 private:
 
 	/** Base editor data, one per channel */
 	TArray<FMovieSceneChannelMetaData, TInlineAllocator<1>> MetaDataArray;
 
-	/** Indirection between channel index and sort order */
-	TArray<FChannelSortOrderIndex> SortOrderedArray;
-
 	/**
 	 * We store the array behind an interface whose access is via void*
 	 * Typed access is only permitted using the original ChannelType templated methods
diff --git a/Engine/Source/Runtime/MovieScene/Public/Channels/MovieSceneChannelProxy.h b/Engine/Source/Runtime/MovieScene/Public/Channels/MovieSceneChannelProxy.h
--- a/Engine/Source/Runtime/MovieScene/Public/Channels/MovieSceneChannelProxy.h
+++ b/Engine/Source/Runtime/MovieScene/Public/Channels/MovieSceneChannelProxy.h
@@ -327,19 +327,19 @@ public:
 	TArrayView<const FMovieSceneChannelMetaData> GetMetaData() const;
 
 	/**
-	 * Get the channel for the specified sort index of a particular type.
+	 * Get the channel with the specified name, assuming it is of a given type
 	 *
-	 * @return A pointer to the channel, or nullptr if the index was invalid, or the type was not present
+	 * @return A typed handle to the channel, or an invalid handle if no channel of that name was found, or it wasn't of the specified type
 	 */
 	template<typename ChannelType>
-	ChannelType* GetChannelBySortOrder(int32 ChannelSortOrder) const;
+	TMovieSceneChannelHandle<ChannelType> GetChannelByName(FName ChannelName) const;
 
 	/**
-	 * Get the channel for the specified sort index of a particular type.
+	 * Get the channel with the specified name
 	 *
-	 * @return A pointer to the channel, or nullptr if the index was invalid, or the type was not present
+	 * @return A handle to the channel, or an invalid handle if no channel of that name was found
 	 */
-	FMovieSceneChannel* GetChannelBySortOrder(FName ChannelTypeName, int32 ChannelSortOrder) const;
+	FMovieSceneChannelHandle GetChannelByName(FName ChannelName) const;
 
 	/**
 	 * Access all the extended data for the templated channel type
@@ -349,7 +349,7 @@ public:
 	template<typename ChannelType>
 	TArrayView<const typename TMovieSceneChannelTraits<ChannelType>::ExtendedEditorDataType> GetAllExtendedEditorData() const;
 
-#endif
+#endif  // !WITH_EDITOR
 
 private:
 
@@ -358,6 +358,18 @@ private:
 
 	/** Array of channel entries, one per channel type. Should never be changed or reallocated after construction to keep pointers alive. */
 	TArray<FMovieSceneChannelEntry, TInlineAllocator<1>> Entries;
+
+#if WITH_EDITOR
+
+	/** Populate the named channel table */
+	void EnsureHandlesByNamePopulated() const;
+
+	/** Lazy-created lookup table between a channel name and a channel handle */
+	mutable TMap<FName, FMovieSceneChannelHandle> HandlesByName;
+	/** Whether the named channel table has been populated */
+	mutable bool bHandlesByNamePopulated = false;
+
+#endif // WITH_EDITOR
 };
 
 
@@ -519,10 +531,16 @@ TArrayView<const FMovieSceneChannelMetaData> FMovieSceneChannelProxy::GetMetaDat
  * @return A pointer to the channel, or nullptr if the index was invalid, or the type was not present
  */
 template<typename ChannelType>
-ChannelType* FMovieSceneChannelProxy::GetChannelBySortOrder(int32 ChannelSortOrder) const
+TMovieSceneChannelHandle<ChannelType> FMovieSceneChannelProxy::GetChannelByName(FName ChannelName) const
 {
-	FName ChannelTypeName = ChannelType::StaticStruct()->GetFName();
-	return static_cast<ChannelType*>(GetChannelBySortOrder(ChannelTypeName, ChannelSortOrder));
+	const FMovieSceneChannelHandle UntypedHandle = GetChannelByName(ChannelName);
+	const FName ChannelTypeName = ChannelType::StaticStruct()->GetFName();
+	// Invalid handles will have a type name of 'None', so we will always fail this test and retrun a null typed handle
+	if (UntypedHandle.GetChannelTypeName() == ChannelTypeName)
+	{
+		return UntypedHandle.Cast<ChannelType>();
+	}
+	return TMovieSceneChannelHandle<ChannelType>();
 }
 
 
