Sequencer: Reinstated parameter section color indicators

diff --git a/Engine/Source/Editor/MovieSceneTools/Private/Sections/MovieSceneSectionHelpers.cpp b/Engine/Source/Editor/MovieSceneTools/Private/Sections/MovieSceneSectionHelpers.cpp
--- a/Engine/Source/Editor/MovieSceneTools/Private/Sections/MovieSceneSectionHelpers.cpp
+++ b/Engine/Source/Editor/MovieSceneTools/Private/Sections/MovieSceneSectionHelpers.cpp
@@ -13,7 +13,7 @@
 #include "Framework/Application/SlateApplication.h"
 #include "Engine/Engine.h"
 
-void MovieSceneSectionHelpers::ConsolidateColorCurves( TArray< TTuple<float, FLinearColor> >& OutColorKeys, const FLinearColor& DefaultColor, const TArray<FMovieSceneFloatChannel*>& ColorChannels, const FTimeToPixel& TimeConverter )
+void MovieSceneSectionHelpers::ConsolidateColorCurves( TArray< TTuple<float, FLinearColor> >& OutColorKeys, const FLinearColor& DefaultColor, TArrayView<const FMovieSceneFloatChannel* const> ColorChannels, const FTimeToPixel& TimeConverter )
 {
 	// @todo Sequencer Optimize - This could all get cached, instead of recalculating everything every OnPaint
 
diff --git a/Engine/Source/Editor/MovieSceneTools/Private/Sections/MovieSceneSectionHelpers.h b/Engine/Source/Editor/MovieSceneTools/Private/Sections/MovieSceneSectionHelpers.h
--- a/Engine/Source/Editor/MovieSceneTools/Private/Sections/MovieSceneSectionHelpers.h
+++ b/Engine/Source/Editor/MovieSceneTools/Private/Sections/MovieSceneSectionHelpers.h
@@ -3,6 +3,7 @@
 #pragma once
 
 #include "Widgets/SWindow.h"
+#include "Containers/ArrayView.h"
 
 struct FKeyHandle;
 struct FTimeToPixel;
@@ -14,7 +15,7 @@ class MOVIESCENETOOLS_API MovieSceneSectionHelpers
 public:
 
 	/** Consolidate color curves for all track sections. */
-	static void ConsolidateColorCurves(TArray< TTuple<float, FLinearColor> >& OutColorKeys, const FLinearColor& DefaultColor, const TArray<FMovieSceneFloatChannel*>& ColorChannels, const FTimeToPixel& TimeConverter);
+	static void ConsolidateColorCurves(TArray< TTuple<float, FLinearColor> >& OutColorKeys, const FLinearColor& DefaultColor, TArrayView<const FMovieSceneFloatChannel* const> ColorChannels, const FTimeToPixel& TimeConverter);
 };
 
 class MOVIESCENETOOLS_API FMovieSceneKeyColorPicker
diff --git a/Engine/Source/Editor/MovieSceneTools/Private/Sections/ParameterSection.cpp b/Engine/Source/Editor/MovieSceneTools/Private/Sections/ParameterSection.cpp
--- a/Engine/Source/Editor/MovieSceneTools/Private/Sections/ParameterSection.cpp
+++ b/Engine/Source/Editor/MovieSceneTools/Private/Sections/ParameterSection.cpp
@@ -17,6 +17,10 @@
 #include "MovieSceneSection.h"
 #include "MovieSceneSectionHelpers.h"
 #include "MVVM/ViewModels/CategoryModel.h"
+#include "MVVM/ViewModels/ChannelModel.h"
+#include "MVVM/Views/SChannelView.h"
+#include "MVVM/Views/STrackLane.h"
+#include "MVVM/Views/STrackAreaView.h"
 #include "Rendering/DrawElementPayloads.h"
 #include "Rendering/DrawElements.h"
 #include "Rendering/RenderingCommon.h"
@@ -38,99 +42,103 @@ struct FKeyHandle;
 
 namespace UE::Sequencer
 {
-	class SColorStripView : public SLeafWidget, public ITrackLaneWidget
+	class SColorStripView : public SChannelView
 	{
 		SLATE_BEGIN_ARGS(SColorStripView){}
 		SLATE_END_ARGS()
 
-		void Construct(const FArguments& InArgs, TViewModelPtr<ITrackLaneExtension> InModel)
+		void Construct(const FArguments& InArgs, const FViewModelPtr& InViewModel, TSharedPtr<STrackAreaView> InTrackAreaView)
 		{
-			WeakModel = InModel;
-		}
+			WeakModel = InViewModel;
+			WeakTrackAreaView = InTrackAreaView;
 
-		TSharedRef<const SWidget> AsWidget() const override
-		{
-			return AsShared();
+			SChannelView::Construct(SChannelView::FArguments(), InViewModel, InTrackAreaView);
 		}
 
-		FTrackLaneScreenAlignment GetAlignment(const FTimeToPixel& InTimeToPixel, const FGeometry& InParentGeometry) const override
+		int32 OnPaint( const FPaintArgs& Args, const FGeometry& AllottedGeometry, const FSlateRect& MyCullingRect, FSlateWindowElementList& OutDrawElements, int32 LayerId, const FWidgetStyle& InWidgetStyle, bool bParentEnabled ) const override
 		{
-			TSharedPtr<ITrackLaneExtension> TrackLaneExtension = WeakModel.ImplicitPin();
-			if (TrackLaneExtension)
+			FViewModelPtr OwningModel = WeakModel.Pin();
+			if (!OwningModel)
 			{
-				FTrackLaneVirtualAlignment VirtualAlignment = TrackLaneExtension->ArrangeVirtualTrackLaneView();
-				FTrackLaneScreenAlignment  ScreenAlignment  = VirtualAlignment.ToScreen(InTimeToPixel, InParentGeometry);
+				return LayerId;
+			}
 
-				return ScreenAlignment;
+			bool bIsExpanded = false;
+			if (TSharedPtr<FLinkedOutlinerExtension> LinkedOutlinerExtension = OwningModel.ImplicitCast())
+			{
+				TSharedPtr<IOutlinerExtension> OutlinerItem = LinkedOutlinerExtension ? LinkedOutlinerExtension->GetLinkedOutlinerItem() : nullptr;
+				if (OutlinerItem)
+				{
+					bIsExpanded = OutlinerItem->IsExpanded();
+				}
 			}
-			return FTrackLaneScreenAlignment();
-		}
 
-		int32 OnPaint( const FPaintArgs& Args, const FGeometry& AllottedGeometry, const FSlateRect& MyCullingRect, FSlateWindowElementList& OutDrawElements, int32 LayerId, const FWidgetStyle& InWidgetStyle, bool bParentEnabled ) const override
-		{
-			#if 0
-			const FMovieSceneChannelMetaData* MetaData = KeyArea->GetChannel().GetMetaData();
-			if (MetaData)
+			const FMovieSceneFloatChannel* ColorChannels[4] = { };
+			for (TViewModelPtr<FChannelModel> ChildChannel : OwningModel->GetChildrenOfType<FChannelModel>())
 			{
-				if (MetaData->DisplayText.EqualTo(FCommonChannelData::ChannelR))
+				FMovieSceneChannelHandle ChannelHandle = ChildChannel->GetKeyArea()->GetChannel();
+
+				FText DisplayText = ChannelHandle.GetMetaData()->DisplayText;
+
+				if (ChannelHandle.GetMetaData()->DisplayText.EqualTo(FCommonChannelData::ChannelR))
 				{
-					RChannel = static_cast<FMovieSceneFloatChannel*>(KeyArea->ResolveChannel());
+					ColorChannels[0] = ChannelHandle.Cast<FMovieSceneFloatChannel>().Get();
 				}
-				if (MetaData->DisplayText.EqualTo(FCommonChannelData::ChannelG))
+				if (ChannelHandle.GetMetaData()->DisplayText.EqualTo(FCommonChannelData::ChannelG))
 				{
-					GChannel = static_cast<FMovieSceneFloatChannel*>(KeyArea->ResolveChannel());
+					ColorChannels[1] = ChannelHandle.Cast<FMovieSceneFloatChannel>().Get();
 				}
-				if (MetaData->DisplayText.EqualTo(FCommonChannelData::ChannelB))
+				if (ChannelHandle.GetMetaData()->DisplayText.EqualTo(FCommonChannelData::ChannelB))
 				{
-					BChannel = static_cast<FMovieSceneFloatChannel*>(KeyArea->ResolveChannel());
+					ColorChannels[2] = ChannelHandle.Cast<FMovieSceneFloatChannel>().Get();
 				}
-				if (MetaData->DisplayText.EqualTo(FCommonChannelData::ChannelA))
+				if (ChannelHandle.GetMetaData()->DisplayText.EqualTo(FCommonChannelData::ChannelA))
 				{
-					AChannel = static_cast<FMovieSceneFloatChannel*>(KeyArea->ResolveChannel());
+					ColorChannels[3] = ChannelHandle.Cast<FMovieSceneFloatChannel>().Get();
 				}
 			}
 
-			if (RChannel && GChannel && BChannel && AChannel)
+			if (ColorChannels[0] && ColorChannels[1] && ColorChannels[2] && ColorChannels[3])
 			{
-				const ESlateDrawEffect DrawEffects = Painter.bParentEnabled ? ESlateDrawEffect::None : ESlateDrawEffect::DisabledEffect;
+				const ESlateDrawEffect DrawEffects = bParentEnabled ? ESlateDrawEffect::None : ESlateDrawEffect::DisabledEffect;
 
-				const FTimeToPixel& TimeConverter = Painter.GetTimeConverter();
+				FTimeToPixel RelativeTimeToPixel = *WeakTrackAreaView->GetTimeToPixel();
 
-				FGeometry KeyAreaGeometry = KeyAreaElement.KeyAreaGeometry;
+				TSharedPtr<ITrackLaneExtension> TrackLaneExtension = WeakModel.ImplicitPin();
+				if (TrackLaneExtension)
+				{
+					FTrackLaneVirtualAlignment VirtualAlignment = TrackLaneExtension->ArrangeVirtualTrackLaneView();
+					if (VirtualAlignment.Range.GetLowerBound().IsClosed())
+					{
+						RelativeTimeToPixel = RelativeTimeToPixel.RelativeTo(VirtualAlignment.Range.GetLowerBoundValue());
+					}
+				}
 
-				const float StartTime       = TimeConverter.PixelToSeconds(0.f);
-				const float EndTime         = TimeConverter.PixelToSeconds(KeyAreaGeometry.GetLocalSize().X);
+				const float StartTime       = RelativeTimeToPixel.PixelToSeconds(0.f);
+				const float EndTime         = RelativeTimeToPixel.PixelToSeconds(AllottedGeometry.GetLocalSize().X);
 				const float SectionDuration = EndTime - StartTime;
 
-				const float KeyAreaHeight = KeyAreaElement.Type == FKeyAreaElement::Group ? KeyAreaGeometry.Size.Y / 3.f - 4.0f : KeyAreaGeometry.Size.Y - 3.0f;
-
-				FVector2D GradientSize = FVector2D( KeyAreaGeometry.Size.X - 2.f, KeyAreaHeight );
+				FVector2D GradientSize = AllottedGeometry.GetLocalSize() - FVector2D(2.f, 0.f);
 				if ( GradientSize.X >= 1.f )
 				{
-					FPaintGeometry PaintGeometry = KeyAreaGeometry.ToPaintGeometry( FVector2D( 1.f, 1.f ), GradientSize );
+					FPaintGeometry PaintGeometry = AllottedGeometry.ToPaintGeometry( FVector2D( 1.f, 1.f ), GradientSize );
 
 					// If we are showing a background pattern and the colors is transparent, draw a checker pattern
 					FSlateDrawElement::MakeBox(
-						Painter.DrawElements,
+						OutDrawElements,
 						LayerId,
 						PaintGeometry,
 						FAppStyle::GetBrush( "Checker" ),
 						DrawEffects);
 
 					FLinearColor DefaultColor = FLinearColor::Black;
-					DefaultColor.R = RChannel->GetDefault().IsSet() ? RChannel->GetDefault().GetValue() : 0.f;
-					DefaultColor.G = GChannel->GetDefault().IsSet() ? GChannel->GetDefault().GetValue() : 0.f;
-					DefaultColor.B = BChannel->GetDefault().IsSet() ? BChannel->GetDefault().GetValue() : 0.f;
-					DefaultColor.A = AChannel->GetDefault().IsSet() ? AChannel->GetDefault().GetValue() : 0.f;
-						
-					TArray<FMovieSceneFloatChannel*> ColorChannels;
-					ColorChannels.Add(RChannel);
-					ColorChannels.Add(GChannel);
-					ColorChannels.Add(BChannel);
-					ColorChannels.Add(AChannel);
+					DefaultColor.R = ColorChannels[0]->GetDefault().IsSet() ? ColorChannels[0]->GetDefault().GetValue() : 0.f;
+					DefaultColor.G = ColorChannels[1]->GetDefault().IsSet() ? ColorChannels[1]->GetDefault().GetValue() : 0.f;
+					DefaultColor.B = ColorChannels[2]->GetDefault().IsSet() ? ColorChannels[2]->GetDefault().GetValue() : 0.f;
+					DefaultColor.A = ColorChannels[3]->GetDefault().IsSet() ? ColorChannels[3]->GetDefault().GetValue() : 0.f;
 
 					TArray< TTuple<float, FLinearColor> > ColorKeys;
-					MovieSceneSectionHelpers::ConsolidateColorCurves( ColorKeys, DefaultColor, ColorChannels, TimeConverter );
+					MovieSceneSectionHelpers::ConsolidateColorCurves( ColorKeys, DefaultColor, ColorChannels, RelativeTimeToPixel );
 
 					TArray<FSlateGradientStop> GradientStops;
 
@@ -143,14 +151,14 @@ namespace UE::Sequencer
 						const FLinearColor Color = ColorStop.Get<1>().ToFColor( true ).ReinterpretAsLinear();
 
 						float TimeFraction = (Time - StartTime) / SectionDuration;
-						GradientStops.Add( FSlateGradientStop( FVector2D( TimeFraction * KeyAreaGeometry.Size.X, 0 ), Color ) );
+						GradientStops.Add( FSlateGradientStop( FVector2D( TimeFraction * GradientSize.X, 0 ), Color ) );
 					}
 
 					if ( GradientStops.Num() > 0 )
 					{
 						FSlateDrawElement::MakeGradient(
-							Painter.DrawElements,
-							Painter.LayerId + 1,
+							OutDrawElements,
+							LayerId + 1,
 							PaintGeometry,
 							GradientStops,
 							Orient_Vertical,
@@ -159,16 +167,18 @@ namespace UE::Sequencer
 					}
 				}
 			}
-			#endif
-			return LayerId;
+
+			LayerId += 2;
+			return SChannelView::OnPaint(Args, AllottedGeometry, MyCullingRect, OutDrawElements, LayerId, InWidgetStyle, bParentEnabled);
 		}
-		virtual FVector2D ComputeDesiredSize(float) const override
+
+		FVector2D ComputeDesiredSize(float) const override
 		{
 			return FVector2D();
 		}
 
-	private:
-		TWeakViewModelPtr<ITrackLaneExtension> WeakModel;
+		FWeakViewModelPtr WeakModel;
+		TSharedPtr<STrackAreaView> WeakTrackAreaView;
 	};
 }
 
@@ -226,7 +236,7 @@ TSharedPtr<UE::Sequencer::FCategoryModel> FParameterSection::ConstructCategoryMo
 
 			TSharedPtr<ITrackLaneWidget> CreateTrackLaneView(const FCreateTrackLaneViewParams& InParams) override
 			{
-				return SNew(SColorStripView, SharedThis(this));
+				return SNew(SColorStripView, SharedThis(this), InParams.OwningTrackLane->GetTrackAreaView());
 			}
 		};
 
diff --git a/Engine/Source/Editor/Sequencer/Private/MVVM/Views/SChannelView.cpp b/Engine/Source/Editor/Sequencer/Private/MVVM/Views/SChannelView.cpp
--- a/Engine/Source/Editor/Sequencer/Private/MVVM/Views/SChannelView.cpp
+++ b/Engine/Source/Editor/Sequencer/Private/MVVM/Views/SChannelView.cpp
@@ -491,7 +491,7 @@ int32 SChannelView::OnPaint(const FPaintArgs& Args, const FGeometry& AllottedGeo
 
 	KeyRenderer.Update(Params, AllottedGeometry);
 	LayerId = KeyRenderer.DrawCurve(Params, AllottedGeometry, MyCullingRect, PaintArgs, LayerId);
-	LayerId = KeyRenderer.Draw(Params, AllottedGeometry, MyCullingRect, PaintArgs, LayerId);
+	LayerId = KeyRenderer.Draw(Params, AllottedGeometry, MyCullingRect, PaintArgs, LayerId + 10);
 
 	KeyRendererCache = NewCachedState;
 
diff --git a/Engine/Source/Editor/Sequencer/Public/MVVM/ViewModels/CategoryModel.h b/Engine/Source/Editor/Sequencer/Public/MVVM/ViewModels/CategoryModel.h
--- a/Engine/Source/Editor/Sequencer/Public/MVVM/ViewModels/CategoryModel.h
+++ b/Engine/Source/Editor/Sequencer/Public/MVVM/ViewModels/CategoryModel.h
@@ -55,7 +55,7 @@ public:
 	FName GetCategoryName() const { return CategoryName; }
 
 	/** Returns the desired sizing for the track area row */
-	FOutlinerSizing GetDesiredSizing() const;
+	virtual FOutlinerSizing GetDesiredSizing() const;
 
 	/*~ ITrackLaneExtension */
 	TSharedPtr<ITrackLaneWidget> CreateTrackLaneView(const FCreateTrackLaneViewParams& InParams) override;
