Constraints: When adding a constraint to a transform need to make sure to bind delegates if the channel already exists . Also make sure all deletion lambda's are bound for the sections, not just the first one.

diff --git a/Engine/Plugins/Animation/ControlRig/Source/ControlRigEditor/Private/Sequencer/ControlRigParameterTrackEditor.cpp b/Engine/Plugins/Animation/ControlRig/Source/ControlRigEditor/Private/Sequencer/ControlRigParameterTrackEditor.cpp
--- a/Engine/Plugins/Animation/ControlRig/Source/ControlRigEditor/Private/Sequencer/ControlRigParameterTrackEditor.cpp
+++ b/Engine/Plugins/Animation/ControlRig/Source/ControlRigEditor/Private/Sequencer/ControlRigParameterTrackEditor.cpp
@@ -2702,7 +2702,7 @@ void FControlRigParameterTrackEditor::HandleConstraintRemoved(IMovieSceneConstra
 	if (const UControlRig* ControlRig = Section->GetControlRig())
 	{
 		FConstraintsManagerController& Controller = FConstraintsManagerController::Get(ControlRig->GetWorld());
-		if (!Controller.OnConstraintRemoved().IsBound())
+		if (!InSection->OnConstraintRemovedHandle.IsValid())
 		{
 			InSection->OnConstraintRemovedHandle =
 			Controller.OnConstraintRemoved().AddLambda([InSection, Section, this](FName InConstraintName)
diff --git a/Engine/Source/Editor/MovieSceneTools/Private/TrackEditors/TransformTrackEditor.cpp b/Engine/Source/Editor/MovieSceneTools/Private/TrackEditors/TransformTrackEditor.cpp
--- a/Engine/Source/Editor/MovieSceneTools/Private/TrackEditors/TransformTrackEditor.cpp
+++ b/Engine/Source/Editor/MovieSceneTools/Private/TrackEditors/TransformTrackEditor.cpp
@@ -167,6 +167,12 @@ TSharedRef<ISequencerSection> F3DTransformTrackEditor::MakeSectionInterface(UMov
 	{
 		Section->ConstraintChannelAdded().AddRaw(this, &F3DTransformTrackEditor::HandleOnConstraintAdded);
 	}
+	//if there are channels already we need to act like they were added
+	TArray<FConstraintAndActiveChannel>& Channels = Section->GetConstraintsChannels();
+	for (FConstraintAndActiveChannel& Channel : Channels)
+	{
+		HandleOnConstraintAdded(Section, &Channel.ActiveChannel);
+	}
 	return MakeShared<FTransformSection>(SectionObject, GetSequencer());
 }
 
@@ -1559,7 +1565,7 @@ void F3DTransformTrackEditor::HandleConstraintRemoved(IMovieSceneConstrainedSect
 	UMovieSceneSection* Section = Cast<UMovieSceneSection>(InSection);
 
 	FConstraintsManagerController& Controller = FConstraintsManagerController::Get(World);
-	if (!Controller.OnConstraintRemoved().IsBound())
+	if (!InSection->OnConstraintRemovedHandle.IsValid())
 	{
 		InSection->OnConstraintRemovedHandle =
 			Controller.OnConstraintRemoved().AddLambda([InSection,Section, this](FName InConstraintName)
