Constraints: Sequencer: Use scaling option when compensating parent constraints on components.

diff --git a/Engine/Source/Editor/MovieSceneTools/Private/TrackEditors/TransformTrackEditor.cpp b/Engine/Source/Editor/MovieSceneTools/Private/TrackEditors/TransformTrackEditor.cpp
--- a/Engine/Source/Editor/MovieSceneTools/Private/TrackEditors/TransformTrackEditor.cpp
+++ b/Engine/Source/Editor/MovieSceneTools/Private/TrackEditors/TransformTrackEditor.cpp
@@ -942,7 +942,21 @@ static void UpdateTransformBasedOnConstraint(FTransform& CurrentTransform, UScen
 					CurrentTransform.SetScale3D(CurrentTransform.GetScale3D() * DiffParents.GetScale3D());
 				}
 				case ETransformConstraintType::Parent:
+				{
+					TOptional<FVector> Scale3D;
+					if (const UTickableParentConstraint* ParentConstraint = Cast<const UTickableParentConstraint>(Constraint))
+					{
+						if (ParentConstraint->IsScalingEnabled() == false)
+						{
+							Scale3D = CurrentTransform.GetScale3D();
+						}
+					}
 					CurrentTransform = CurrentTransform * DiffParents;
+					if (Scale3D.IsSet())
+					{
+						CurrentTransform.SetScale3D(Scale3D.GetValue());
+					}
+				}
 				case ETransformConstraintType::LookAt: //leave current alone
 				default:
 					break;
