Constraints: Fix for Position Constraint with components, missing break !?!
#trivial

diff --git a/Engine/Source/Editor/MovieSceneTools/Private/TrackEditors/TransformTrackEditor.cpp b/Engine/Source/Editor/MovieSceneTools/Private/TrackEditors/TransformTrackEditor.cpp
--- a/Engine/Source/Editor/MovieSceneTools/Private/TrackEditors/TransformTrackEditor.cpp
+++ b/Engine/Source/Editor/MovieSceneTools/Private/TrackEditors/TransformTrackEditor.cpp
@@ -929,18 +929,21 @@ static void UpdateTransformBasedOnConstraint(FTransform& CurrentTransform, UScen
 			{
 				case ETransformConstraintType::Translation:
 				{
-					CurrentTransform.SetLocation(CurrentTransform.GetLocation() +  OriginalParentTransform.GetLocation() - ParentTransform.GetLocation());
+					CurrentTransform.SetLocation(CurrentTransform.GetLocation() + OriginalParentTransform.GetLocation() - ParentTransform.GetLocation());
 				}
+				break;
 				case ETransformConstraintType::Rotation:
 				{
 					FQuat RelativeRotation = ParentTransform.GetRotation().Inverse() * OriginalParentTransform.GetRotation();
 					RelativeRotation.Normalize();
 					CurrentTransform.SetRotation(CurrentTransform.GetRotation() * RelativeRotation);
 				}
+				break;
 				case ETransformConstraintType::Scale:
 				{
 					CurrentTransform.SetScale3D(CurrentTransform.GetScale3D() * DiffParents.GetScale3D());
 				}
+				break;
 				case ETransformConstraintType::Parent:
 				{
 					TOptional<FVector> Scale3D;
@@ -957,6 +960,7 @@ static void UpdateTransformBasedOnConstraint(FTransform& CurrentTransform, UScen
 						CurrentTransform.SetScale3D(Scale3D.GetValue());
 					}
 				}
+				break;
 				case ETransformConstraintType::LookAt: //leave current alone
 				default:
 					break;
