Sequencer: change the colour of perlin noise curve since the background colour of sequencer sections has changed

diff --git a/Engine/Source/Editor/MovieSceneTools/Public/Channels/PerlinNoiseChannelInterface.h b/Engine/Source/Editor/MovieSceneTools/Public/Channels/PerlinNoiseChannelInterface.h
--- a/Engine/Source/Editor/MovieSceneTools/Public/Channels/PerlinNoiseChannelInterface.h
+++ b/Engine/Source/Editor/MovieSceneTools/Public/Channels/PerlinNoiseChannelInterface.h
@@ -198,7 +198,7 @@ struct TPerlinNoiseChannelInterface : ISequencerChannelInterface
 	{
 		using namespace UE::Sequencer;
 
-		FColor FillColor(0, 0, 0, 51);
+		FLinearColor FillColor(1, 1, 1, 0.334f);
 		ChannelType* TypedChannel = static_cast<ChannelType*>(InChannel);
 
 		TArray<FVector2D> CurvePoints;
