Constraints: Componensate All works with components, unify evaluation code all into compensation evaluation code to also be used by baker, make sure everyhing ticks while baking.

diff --git a/Engine/Plugins/Animation/ControlRig/Source/ControlRigEditor/Private/Sequencer/ConstraintChannelHelper.cpp b/Engine/Plugins/Animation/ControlRig/Source/ControlRigEditor/Private/Sequencer/ConstraintChannelHelper.cpp
--- a/Engine/Plugins/Animation/ControlRig/Source/ControlRigEditor/Private/Sequencer/ConstraintChannelHelper.cpp
+++ b/Engine/Plugins/Animation/ControlRig/Source/ControlRigEditor/Private/Sequencer/ConstraintChannelHelper.cpp
@@ -374,29 +374,58 @@ void FConstraintChannelHelper::Compensate(UTickableTransformConstraint* InConstr
 	{
 		return;
 	}
+	const TSharedPtr<ISequencer> Sequencer = WeakSequencer.Pin();
 
-	const UTransformableControlHandle* ControlHandle = Cast<UTransformableControlHandle>(InConstraint->ChildTRSHandle);
-	if (!ControlHandle)
+	IMovieSceneConstrainedSection* Section = nullptr;
+	UWorld* World = nullptr;
+	if (const UTransformableComponentHandle* ComponentHandle = Cast<UTransformableComponentHandle>(InConstraint->ChildTRSHandle))
 	{
-		return;
+
+		AActor* Actor = ComponentHandle->Component->GetOwner();
+		if (!Actor)
+		{
+			return;
+		}
+		World = Actor->GetWorld();
+		const FTransform LocalTransform = ComponentHandle->GetLocalTransform();
+		const FGuid Guid = Sequencer->GetHandleToObject(Actor, true);
+		if (!Guid.IsValid())
+		{
+			return;
+		}
+
+		if (UMovieScene3DTransformSection* TransformSection = MovieSceneToolHelpers::GetTransformSection(Sequencer.Get(), Guid, LocalTransform))
+		{
+			Section = TransformSection;
+		}
 	}
 
-	const TSharedPtr<ISequencer> Sequencer = WeakSequencer.Pin();
-	if (const UMovieSceneControlRigParameterSection* Section = GetControlSection(ControlHandle, Sequencer))
+	if (const UTransformableControlHandle* ControlHandle = Cast<UTransformableControlHandle>(InConstraint->ChildTRSHandle))
 	{
-		const FFrameRate TickResolution = Sequencer->GetFocusedTickResolution();
-		const FFrameTime FrameTime = Sequencer->GetLocalTime().ConvertTo(TickResolution);
-		const FFrameNumber Time = FrameTime.GetFrame();
-
-		const TOptional<FFrameNumber> OptTime = bAllTimes ? TOptional<FFrameNumber>() : TOptional<FFrameNumber>(Time);
-		CompensateIfNeeded(ControlHandle->ControlRig.Get(), Sequencer, Section, OptTime);
+		UControlRig* ControlRig = ControlHandle->ControlRig.LoadSynchronous();
+		if (!ControlRig)
+		{
+			return;
+		}
+		World = ControlRig->GetWorld();
+		if (UMovieSceneControlRigParameterSection* ControlSection = GetControlSection(ControlHandle,Sequencer))
+		{
+			Section = ControlSection;
+		}
 	}
+
+	const FFrameRate TickResolution = Sequencer->GetFocusedTickResolution();
+	const FFrameTime FrameTime = Sequencer->GetLocalTime().ConvertTo(TickResolution);
+	const FFrameNumber Time = FrameTime.GetFrame();
+
+	const TOptional<FFrameNumber> OptTime = bAllTimes ? TOptional<FFrameNumber>() : TOptional<FFrameNumber>(Time);
+	CompensateIfNeeded(World, Sequencer, Section, OptTime);
 }
 
 void FConstraintChannelHelper::CompensateIfNeeded(
-	const UControlRig* ControlRig,
+	UWorld* InWorld,
 	const TSharedPtr<ISequencer>& InSequencer,
-	const UMovieSceneControlRigParameterSection* Section,
+	IMovieSceneConstrainedSection* Section,
 	const TOptional<FFrameNumber>& OptionalTime)
 {
 	if (FMovieSceneConstraintChannelHelper::bDoNotCompensate)
@@ -422,16 +451,12 @@ void FConstraintChannelHelper::CompensateIfNeeded(
 		return OptionalTimeArray;
 	};
 
-	// keyframe context
-	FRigControlModifiedContext KeyframeContext;
-	KeyframeContext.SetKey = EControlRigSetKey::Always;
-	const FFrameRate TickResolution = InSequencer->GetFocusedTickResolution();
-
 	bool bNeedsEvaluation = false;
 
 	// gather all transform constraints
+	TArray<FConstraintAndActiveChannel>& ConstraintChannels = Section->GetConstraintsChannels();
 	TArray<FConstraintAndActiveChannel> TransformConstraintsChannels;
-	Algo::CopyIf(Section->GetConstraintsChannels(), TransformConstraintsChannels,
+	Algo::CopyIf(ConstraintChannels, TransformConstraintsChannels,
 		[](const FConstraintAndActiveChannel& InChannel)
 		{
 			return InChannel.Constraint.IsValid() && InChannel.Constraint->IsA<UTickableTransformConstraint>();
@@ -458,7 +483,7 @@ void FConstraintChannelHelper::CompensateIfNeeded(
 				// if switching from active to inactive then we must add a key at T-1 in the constraint space
 				// if switching from inactive to active then we must add a key at T-1 in the previous constraint or parent space
 				FCompensationEvaluator Evaluator(Constraint);
-				Evaluator.ComputeCompensation(ControlRig->GetWorld(), InSequencer, Time);
+				Evaluator.ComputeCompensation(InWorld, InSequencer, Time);
 				const TArray<FTransform>& LocalTransforms = Evaluator.ChildLocals;
 
 				const EMovieSceneTransformChannel ChannelsToKey = Constraint->GetChannelsToKey();
diff --git a/Engine/Plugins/Animation/ControlRig/Source/ControlRigEditor/Private/Sequencer/ControlRigParameterTrackEditor.cpp b/Engine/Plugins/Animation/ControlRig/Source/ControlRigEditor/Private/Sequencer/ControlRigParameterTrackEditor.cpp
--- a/Engine/Plugins/Animation/ControlRig/Source/ControlRigEditor/Private/Sequencer/ControlRigParameterTrackEditor.cpp
+++ b/Engine/Plugins/Animation/ControlRig/Source/ControlRigEditor/Private/Sequencer/ControlRigParameterTrackEditor.cpp
@@ -3352,7 +3352,7 @@ FKeyPropertyResult FControlRigParameterTrackEditor::AddKeysToControlRigHandle(US
 
 				// compensate constraints
 				FConstraintChannelHelper::CompensateIfNeeded(
-					SectionControlRig, GetSequencer(), ParamSection, OptionalKeyTime);
+					SectionControlRig->GetWorld(), GetSequencer(), ParamSection, OptionalKeyTime);
 			}
 		}
 	}
diff --git a/Engine/Plugins/Animation/ControlRig/Source/ControlRigEditor/Private/Tools/ConstraintBaker.cpp b/Engine/Plugins/Animation/ControlRig/Source/ControlRigEditor/Private/Tools/ConstraintBaker.cpp
--- a/Engine/Plugins/Animation/ControlRig/Source/ControlRigEditor/Private/Tools/ConstraintBaker.cpp
+++ b/Engine/Plugins/Animation/ControlRig/Source/ControlRigEditor/Private/Tools/ConstraintBaker.cpp
@@ -94,71 +94,6 @@ namespace
 	}
 }
 
-void FConstraintBaker::GetHandleTransforms(
-	UWorld* InWorld,
-	const TSharedPtr<ISequencer>& InSequencer,
-	const UTransformableHandle* InHandle,
-	const TArray<FFrameNumber>& InFrames,
-	const bool bLocal,
-	TArray<FTransform>& OutTransforms)
-{
-	const FConstraintsManagerController& Controller = FConstraintsManagerController::Get(InWorld);
-	static constexpr bool bSorted = true;
-	TArray<TObjectPtr<UTickableConstraint>> Constraints = Controller.GetParentConstraints(InHandle->GetHash(), bSorted);
-
-	TArray<UTickableTransformConstraint*> TransformConstraints;
-	for (const TObjectPtr<UTickableConstraint>& Constraint: Constraints)
-	{
-		if (UTickableTransformConstraint* TransformConstraint = Cast<UTickableTransformConstraint>(Constraint.Get()))
-		{
-			TransformConstraints.Add(TransformConstraint);
-		}
-	}
-	
-	return GetHandleTransforms(InSequencer, InHandle, TransformConstraints, InFrames, bLocal, OutTransforms);
-}
-
-void FConstraintBaker::GetHandleTransforms(
-	const TSharedPtr<ISequencer>& InSequencer,
-	const UTransformableHandle* InHandle,
-	const TArray<UTickableTransformConstraint*>& InConstraintsToEvaluate,
-	const TArray<FFrameNumber>& InFrames,
-	const bool bLocal,
-	TArray<FTransform>& OutTransforms)
-{
-	// if (InConstraintsToEvaluate.IsEmpty())
-	// {
-	// 	return;
-	// }
-	
-	const UMovieScene* MovieScene = InSequencer->GetFocusedMovieSceneSequence()->GetMovieScene();
-	
-	const FFrameRate TickResolution = MovieScene->GetTickResolution();
-	const EMovieScenePlayerStatus::Type PlaybackStatus = InSequencer->GetPlaybackStatus();
-
-	OutTransforms.SetNum(InFrames.Num());
-	for (int32 Index = 0; Index < InFrames.Num(); ++Index)
-	{
-		const FFrameNumber& FrameNumber = InFrames[Index];
-	
-		// evaluate animation
-		const FMovieSceneEvaluationRange EvaluationRange = FMovieSceneEvaluationRange(FFrameTime (FrameNumber), TickResolution);
-		const FMovieSceneContext Context = FMovieSceneContext(EvaluationRange, PlaybackStatus).SetHasJumped(true);
-
-		InSequencer->GetEvaluationTemplate().Evaluate(Context, *InSequencer);
-	
-		// evaluate constraints
-		for (const UTickableTransformConstraint* Constraint: InConstraintsToEvaluate)
-		{
-			Constraint->Evaluate();
-		}
-	
-		// evaluate ControlRig?
-		// ControlRig->Evaluate_AnyThread();
-		
-		OutTransforms[Index] = bLocal ? InHandle->GetLocalTransform() : InHandle->GetGlobalTransform();
-	}
-}
 
 void FConstraintBaker::AddTransformKeys(
 	const TSharedPtr<ISequencer>& InSequencer,
@@ -167,6 +102,7 @@ void FConstraintBaker::AddTransformKeys(
 	const TArray<FTransform>& InTransforms,
 	const EMovieSceneTransformChannel& InChannels)
 {
+
 	if (const UTransformableComponentHandle* ComponentHandle = Cast<UTransformableComponentHandle>(InHandle))
 	{
 		return BakeComponent(InSequencer, ComponentHandle, InFrames, InTransforms, InChannels); 
@@ -190,8 +126,9 @@ void FConstraintBaker::Bake(UWorld* InWorld,
 	}
 	
 	// compute transforms
-	TArray<FTransform> Transforms;
-	GetHandleTransforms(InSequencer, InConstraint->ChildTRSHandle, { InConstraint }, InFrames, true, Transforms);
+	FCompensationEvaluator Evaluator(InConstraint);
+	Evaluator.ComputeLocalTransformsForBaking(InWorld, InSequencer, InFrames);
+	TArray<FTransform>& Transforms = Evaluator.ChildLocals;
 
 	if (InFrames.Num() != Transforms.Num())
 	{
diff --git a/Engine/Plugins/Animation/ControlRig/Source/ControlRigEditor/Public/ConstraintChannelHelper.h b/Engine/Plugins/Animation/ControlRig/Source/ControlRigEditor/Public/ConstraintChannelHelper.h
--- a/Engine/Plugins/Animation/ControlRig/Source/ControlRigEditor/Public/ConstraintChannelHelper.h
+++ b/Engine/Plugins/Animation/ControlRig/Source/ControlRigEditor/Public/ConstraintChannelHelper.h
@@ -10,6 +10,7 @@ class UControlRig;
 class ISequencer;
 class UMovieScene3DTransformSection;
 class UMovieSceneControlRigParameterSection;
+class IMovieSceneConstrainedSection;
 class UMovieSceneSection;
 class UTransformableHandle;
 class UTickableTransformConstraint;
@@ -39,9 +40,9 @@ public:
 		UTickableTransformConstraint* InConstraint,
 		const bool bAllTimes = false);
 	static void CompensateIfNeeded(
-		const UControlRig* ControlRig,
+		UWorld* InWorld,
 		const TSharedPtr<ISequencer>& InSequencer,
-		const UMovieSceneControlRigParameterSection* Section,
+		IMovieSceneConstrainedSection* Section,
 		const TOptional<FFrameNumber>& OptionalTime);
 
 	/** @todo documentation. */
diff --git a/Engine/Plugins/Animation/ControlRig/Source/ControlRigEditor/Public/Tools/ConstraintBaker.h b/Engine/Plugins/Animation/ControlRig/Source/ControlRigEditor/Public/Tools/ConstraintBaker.h
--- a/Engine/Plugins/Animation/ControlRig/Source/ControlRigEditor/Public/Tools/ConstraintBaker.h
+++ b/Engine/Plugins/Animation/ControlRig/Source/ControlRigEditor/Public/Tools/ConstraintBaker.h
@@ -18,14 +18,6 @@ public:
 	/** Bake constraint over specified frames, frames must be in order*/
 	static void Bake(UWorld* InWorld, UTickableTransformConstraint* InConstraint, const TSharedPtr<ISequencer>& InSequencer, const TArray<FFrameNumber>& InFrames);
 
-	/** Stores InHandle local (or global) transforms at InFrames. */
-	static void GetHandleTransforms(
-		UWorld* InWorld,
-		const TSharedPtr<ISequencer>& InSequencer,
-		const UTransformableHandle* InHandle,
-		const TArray<FFrameNumber>& InFrames,
-		const bool bLocal,
-		TArray<FTransform>& OutTransforms);
 	
 	/** Add InTransforms keys at InFrames into the InHandle transform animation channels. */
 	static void AddTransformKeys(
@@ -37,13 +29,5 @@ public:
 	
 private:
 
-	/** Evaluates the constraint at each frames and stores the resulting child transforms. */
-	static void GetHandleTransforms(
-		const TSharedPtr<ISequencer>& InSequencer,
-		const UTransformableHandle* InHandle,
-		const TArray<UTickableTransformConstraint*>& InConstraintsToEvaluate,
-		const TArray<FFrameNumber>& InFrames,
-		const bool bLocal,
-		TArray<FTransform>& OutTransforms);
 };
 
diff --git a/Engine/Source/Editor/MovieSceneTools/Private/MovieSceneConstraintChannelHelper.cpp b/Engine/Source/Editor/MovieSceneTools/Private/MovieSceneConstraintChannelHelper.cpp
--- a/Engine/Source/Editor/MovieSceneTools/Private/MovieSceneConstraintChannelHelper.cpp
+++ b/Engine/Source/Editor/MovieSceneTools/Private/MovieSceneConstraintChannelHelper.cpp
@@ -19,6 +19,7 @@
 #include "Channels/MovieSceneDoubleChannel.h"
 #include "LevelEditorViewport.h"
 #include "Sections/MovieSceneConstrainedSection.h"
+#include "MovieSceneToolsModule.h"
 
 
 /*
@@ -59,7 +60,7 @@ void FCompensationEvaluator::ComputeLocalTransforms(
 		return LastActiveIndex > INDEX_NONE ? Constraints[LastActiveIndex] : nullptr;
 	};
 
-	const UMovieScene* MovieScene = InSequencer->GetFocusedMovieSceneSequence()->GetMovieScene();
+	UMovieScene* MovieScene = InSequencer->GetFocusedMovieSceneSequence()->GetMovieScene();
 	const FFrameRate TickResolution = MovieScene->GetTickResolution();
 	const EMovieScenePlayerStatus::Type PlaybackStatus = InSequencer->GetPlaybackStatus();
 
@@ -69,7 +70,14 @@ void FCompensationEvaluator::ComputeLocalTransforms(
 	ChildLocals.SetNum(NumFrames + 1);
 	ChildGlobals.SetNum(NumFrames + 1);
 	SpaceGlobals.SetNum(NumFrames + 1);
-
+	const TArray<IMovieSceneToolsAnimationBakeHelper*>& BakeHelpers = FMovieSceneToolsModule::Get().GetAnimationBakeHelpers();
+	for (IMovieSceneToolsAnimationBakeHelper* BakeHelper : BakeHelpers)
+	{
+		if (BakeHelper)
+		{
+			BakeHelper->StartBaking(MovieScene);
+		}
+	}
 	for (int32 Index = 0; Index < NumFrames + 1; ++Index)
 	{
 		const FFrameNumber FrameNumber = (Index == 0) ? InFrames[0] - 1 : InFrames[Index - 1];
@@ -78,14 +86,31 @@ void FCompensationEvaluator::ComputeLocalTransforms(
 		const FMovieSceneEvaluationRange EvaluationRange = FMovieSceneEvaluationRange(FFrameTime(FrameNumber), TickResolution);
 		const FMovieSceneContext Context = FMovieSceneContext(EvaluationRange, PlaybackStatus).SetHasJumped(true);
 
+		for (IMovieSceneToolsAnimationBakeHelper* BakeHelper : BakeHelpers)
+		{
+			if (BakeHelper)
+			{
+				BakeHelper->PreEvaluation(MovieScene, FrameNumber);
+			}
+		}
 		InSequencer->GetEvaluationTemplate().Evaluate(Context, *InSequencer);
 
 		// evaluate constraints
 		for (const UTickableTransformConstraint* InConstraint : Constraints)
 		{
-			InConstraint->Evaluate();
+			if (InConstraint)
+			{
+				InConstraint->Evaluate(true);
+			}
 		}
 
+		for (IMovieSceneToolsAnimationBakeHelper* BakeHelper : BakeHelpers)
+		{
+			if (BakeHelper)
+			{
+				BakeHelper->PostEvaluation(MovieScene, FrameNumber);
+			}
+		}
 		// evaluate ControlRig?
 		// ControlRig->Evaluate_AnyThread();
 
@@ -141,8 +166,91 @@ void FCompensationEvaluator::ComputeLocalTransforms(
 			}
 		}
 	}
+	for (IMovieSceneToolsAnimationBakeHelper* BakeHelper : BakeHelpers)
+	{
+		if (BakeHelper)
+		{
+			BakeHelper->StopBaking(MovieScene);
+		}
+	}
 }
+void FCompensationEvaluator::ComputeLocalTransformsForBaking(UWorld* InWorld, const TSharedPtr<ISequencer>& InSequencer, const TArray<FFrameNumber>& InFrames)
+{
+	if (InFrames.IsEmpty())
+	{
+		return;
+	}
+
+	const TArray<UTickableTransformConstraint*> Constraints = GetHandleTransformConstraints(InWorld);
+
+	UMovieScene* MovieScene = InSequencer->GetFocusedMovieSceneSequence()->GetMovieScene();
+	const FFrameRate TickResolution = MovieScene->GetTickResolution();
+	const EMovieScenePlayerStatus::Type PlaybackStatus = InSequencer->GetPlaybackStatus();
+
+	const int32 NumFrames = InFrames.Num();
+
+	ChildLocals.SetNum(NumFrames);
+	ChildGlobals.SetNum(NumFrames);
+	SpaceGlobals.SetNum(NumFrames);
+	const TArray<IMovieSceneToolsAnimationBakeHelper*>& BakeHelpers = FMovieSceneToolsModule::Get().GetAnimationBakeHelpers();
+	for (IMovieSceneToolsAnimationBakeHelper* BakeHelper : BakeHelpers)
+	{
+		if (BakeHelper)
+		{
+			BakeHelper->StartBaking(MovieScene);
+		}
+	}
+	for (int32 Index = 0; Index < NumFrames; ++Index)
+	{
+		const FFrameNumber& FrameNumber = InFrames[Index];
+
+		// evaluate animation
+		const FMovieSceneEvaluationRange EvaluationRange = FMovieSceneEvaluationRange(FFrameTime(FrameNumber), TickResolution);
+		const FMovieSceneContext Context = FMovieSceneContext(EvaluationRange, PlaybackStatus).SetHasJumped(true);
+
+		for (IMovieSceneToolsAnimationBakeHelper* BakeHelper : BakeHelpers)
+		{
+			if (BakeHelper)
+			{
+				BakeHelper->PreEvaluation(MovieScene, FrameNumber);
+			}
+		}
+		InSequencer->GetEvaluationTemplate().Evaluate(Context, *InSequencer);
+
+		// evaluate constraints
+		for (const UTickableTransformConstraint* InConstraint : Constraints)
+		{
+			InConstraint->Evaluate(true);
+		}
+
+		for (IMovieSceneToolsAnimationBakeHelper* BakeHelper : BakeHelpers)
+		{
+			if (BakeHelper)
+			{
+				BakeHelper->PostEvaluation(MovieScene, FrameNumber);
+			}
+		}
+
+		// evaluate ControlRig?
+		// ControlRig->Evaluate_AnyThread();
+
+		FTransform& ChildLocal = ChildLocals[Index];
+		FTransform& ChildGlobal = ChildGlobals[Index];
+		FTransform& SpaceGlobal = SpaceGlobals[Index];
+
+		// store child transforms        	
+		ChildLocal = Handle->GetLocalTransform();
+		ChildGlobal = Handle->GetGlobalTransform();
+	}
+	for (IMovieSceneToolsAnimationBakeHelper* BakeHelper : BakeHelpers)
+	{
+		if (BakeHelper)
+		{
+			BakeHelper->StopBaking(MovieScene);
+		}
+	}
 
+}
 void FCompensationEvaluator::ComputeLocalTransformsBeforeDeletion(
 	UWorld* InWorld,
 	const TSharedPtr<ISequencer>& InSequencer,
@@ -172,7 +280,7 @@ void FCompensationEvaluator::ComputeLocalTransformsBeforeDeletion(
 		return LastActiveIndex > INDEX_NONE ? Constraints[LastActiveIndex] : nullptr;
 	};
 
-	const UMovieScene* MovieScene = InSequencer->GetFocusedMovieSceneSequence()->GetMovieScene();
+	UMovieScene* MovieScene = InSequencer->GetFocusedMovieSceneSequence()->GetMovieScene();
 	const FFrameRate TickResolution = MovieScene->GetTickResolution();
 	const EMovieScenePlayerStatus::Type PlaybackStatus = InSequencer->GetPlaybackStatus();
 
@@ -181,7 +289,14 @@ void FCompensationEvaluator::ComputeLocalTransformsBeforeDeletion(
 	ChildLocals.SetNum(NumFrames);
 	ChildGlobals.SetNum(NumFrames);
 	SpaceGlobals.SetNum(NumFrames);
-
+	const TArray<IMovieSceneToolsAnimationBakeHelper*>& BakeHelpers = FMovieSceneToolsModule::Get().GetAnimationBakeHelpers();
+	for (IMovieSceneToolsAnimationBakeHelper* BakeHelper : BakeHelpers)
+	{
+		if (BakeHelper)
+		{
+			BakeHelper->StartBaking(MovieScene);
+		}
+	}
 	for (int32 Index = 0; Index < NumFrames; ++Index)
 	{
 		const FFrameNumber& FrameNumber = InFrames[Index];
@@ -190,14 +305,27 @@ void FCompensationEvaluator::ComputeLocalTransformsBeforeDeletion(
 		const FMovieSceneEvaluationRange EvaluationRange = FMovieSceneEvaluationRange(FFrameTime(FrameNumber), TickResolution);
 		const FMovieSceneContext Context = FMovieSceneContext(EvaluationRange, PlaybackStatus).SetHasJumped(true);
 
+		for (IMovieSceneToolsAnimationBakeHelper* BakeHelper : BakeHelpers)
+		{
+			if (BakeHelper)
+			{
+				BakeHelper->PreEvaluation(MovieScene, FrameNumber);
+			}
+		}
 		InSequencer->GetEvaluationTemplate().Evaluate(Context, *InSequencer);
 
 		// evaluate constraints
 		for (const UTickableTransformConstraint* InConstraint : Constraints)
 		{
-			InConstraint->Evaluate();
+			InConstraint->Evaluate(true);
+		}
+		for (IMovieSceneToolsAnimationBakeHelper* BakeHelper : BakeHelpers)
+		{
+			if (BakeHelper)
+			{
+				BakeHelper->PostEvaluation(MovieScene, FrameNumber);
+			}
 		}
-
 		// evaluate ControlRig?
 		// ControlRig->Evaluate_AnyThread();
 
@@ -217,6 +345,13 @@ void FCompensationEvaluator::ComputeLocalTransformsBeforeDeletion(
 				ChildLocal, ChildGlobal, SpaceGlobal, LastConstraint);
 		}
 	}
+	for (IMovieSceneToolsAnimationBakeHelper* BakeHelper : BakeHelpers)
+	{
+		if (BakeHelper)
+		{
+			BakeHelper->StopBaking(MovieScene);
+		}
+	}
 }
 
 void FCompensationEvaluator::ComputeCompensation(UWorld* InWorld, const TSharedPtr<ISequencer>& InSequencer, const FFrameNumber& InTime)
@@ -240,30 +375,58 @@ void FCompensationEvaluator::ComputeCompensation(UWorld* InWorld, const TSharedP
 		return LastActiveIndex > INDEX_NONE ? Constraints[LastActiveIndex] : nullptr;
 	};
 
-	auto EvaluateAt = [InSequencer, &Constraints](const FFrameNumber InFrame)
+	UMovieScene* MovieScene = InSequencer->GetFocusedMovieSceneSequence()->GetMovieScene();
+	const TArray<IMovieSceneToolsAnimationBakeHelper*>& BakeHelpers = FMovieSceneToolsModule::Get().GetAnimationBakeHelpers();
+	for (IMovieSceneToolsAnimationBakeHelper* BakeHelper : BakeHelpers)
 	{
-		const UMovieScene* MovieScene = InSequencer->GetFocusedMovieSceneSequence()->GetMovieScene();
+		if (BakeHelper)
+		{
+			BakeHelper->StartBaking(MovieScene);
+		}
+	}
+	auto EvaluateAt = [InSequencer, &Constraints, &BakeHelpers](const FFrameNumber InFrame)
+	{
+
+		UMovieScene* MovieScene = InSequencer->GetFocusedMovieSceneSequence()->GetMovieScene();
 		const FFrameRate TickResolution = MovieScene->GetTickResolution();
 		const EMovieScenePlayerStatus::Type PlaybackStatus = InSequencer->GetPlaybackStatus();
 
 		const FMovieSceneEvaluationRange EvaluationRange0 = FMovieSceneEvaluationRange(FFrameTime(InFrame), TickResolution);
 		const FMovieSceneContext Context0 = FMovieSceneContext(EvaluationRange0, PlaybackStatus).SetHasJumped(true);
 
+		for (IMovieSceneToolsAnimationBakeHelper* BakeHelper : BakeHelpers)
+		{
+			if (BakeHelper)
+			{
+				BakeHelper->PreEvaluation(MovieScene, InFrame);
+			}
+		}
 		InSequencer->GetEvaluationTemplate().Evaluate(Context0, *InSequencer);
 
 		for (const UTickableTransformConstraint* InConstraint : Constraints)
 		{
-			InConstraint->Evaluate();
+			InConstraint->Evaluate(true);
+		}
+
+		for (IMovieSceneToolsAnimationBakeHelper* BakeHelper : BakeHelpers)
+		{
+			if (BakeHelper)
+			{
+				BakeHelper->PostEvaluation(MovieScene, InFrame);
+			}
 		}
 
 		// ControlRig->Evaluate_AnyThread();
 	};
 
+
 	// allocate
 	ChildLocals.SetNum(1);
 	ChildGlobals.SetNum(1);
 	SpaceGlobals.SetNum(1);
 
+	
+
 	// evaluate at InTime and store global
 	EvaluateAt(InTime);
 	ChildGlobals[0] = Handle->GetGlobalTransform();
@@ -272,6 +435,14 @@ void FCompensationEvaluator::ComputeCompensation(UWorld* InWorld, const TSharedP
 	EvaluateAt(InTime - 1);
 	ChildLocals[0] = Handle->GetLocalTransform();
 
+
+	for (IMovieSceneToolsAnimationBakeHelper* BakeHelper : BakeHelpers)
+	{
+		if (BakeHelper)
+		{
+			BakeHelper->StopBaking(MovieScene);
+		}
+	}
 	// if constraint at T-1 then switch to its space
 	if (const UTickableTransformConstraint* LastConstraint = GetLastActiveConstraint())
 	{
diff --git a/Engine/Source/Editor/MovieSceneTools/Private/MovieSceneToolHelpers.cpp b/Engine/Source/Editor/MovieSceneTools/Private/MovieSceneToolHelpers.cpp
--- a/Engine/Source/Editor/MovieSceneTools/Private/MovieSceneToolHelpers.cpp
+++ b/Engine/Source/Editor/MovieSceneTools/Private/MovieSceneToolHelpers.cpp
@@ -78,6 +78,7 @@
 #include "Channels/MovieSceneByteChannel.h"
 #include "Framework/MultiBox/MultiBoxBuilder.h"
 #include "EntitySystem/Interrogation/MovieSceneInterrogationLinker.h"
+#include "ConstraintsManager.h"
 
 /* FSkelMeshRecorder
  ***********/
@@ -3766,6 +3767,14 @@ static void TickFrame(const FFrameNumber& FrameNumber,float DeltaTime, UMovieSce
 	// This will call UpdateSkelPose on the skeletal mesh component to move bones based on animations in the sequence
 	int32 Index = FrameNumber.Value;
 	AnimTrackAdapter.UpdateAnimation(Index);
+
+	UWorld* World = GCurrentLevelEditingViewportClient ? GCurrentLevelEditingViewportClient->GetWorld() : nullptr;
+	if (World)
+	{
+		const FConstraintsManagerController& Controller = FConstraintsManagerController::Get(World);
+		Controller.EvaluateAllConstraints();
+	}
+
 	for (IMovieSceneToolsAnimationBakeHelper* BakeHelper : BakeHelpers)
 	{
 		if (BakeHelper)
@@ -4252,6 +4261,10 @@ static void GetSequencerActorWorldTransforms(IMovieScenePlayer* Player, UMovieSc
 					}
 					Player->GetEvaluationTemplate().Evaluate(Context, *Player);
 
+
+					const FConstraintsManagerController& Controller = FConstraintsManagerController::Get(Actor->GetWorld());
+					Controller.EvaluateAllConstraints();
+					
 					for (IMovieSceneToolsAnimationBakeHelper* BakeHelper : BakeHelpers)
 					{
 						if (BakeHelper)
@@ -4264,7 +4277,7 @@ static void GetSequencerActorWorldTransforms(IMovieScenePlayer* Player, UMovieSc
 					while (Parent)
 					{
 						TArray<USkeletalMeshComponent*> MeshComps;
-						ActorSelection.Actor->GetComponents(MeshComps, true);
+						Parent->GetComponents(MeshComps, true);
 
 						for (USkeletalMeshComponent* MeshComp : MeshComps)
 						{
diff --git a/Engine/Source/Editor/MovieSceneTools/Public/ConstraintChannelHelper.inl b/Engine/Source/Editor/MovieSceneTools/Public/ConstraintChannelHelper.inl
--- a/Engine/Source/Editor/MovieSceneTools/Public/ConstraintChannelHelper.inl
+++ b/Engine/Source/Editor/MovieSceneTools/Public/ConstraintChannelHelper.inl
@@ -253,8 +253,8 @@ void FMovieSceneConstraintChannelHelper::DeleteTransformKeys(
 		TMovieSceneChannelData<typename ChannelType::ChannelValueType> Data = Channel->GetData();
 		const TArrayView<const FFrameNumber> Times = Data.GetTimes();
 		
-		const int32 KeyIndex = Algo::UpperBound(Times, InTime);
-		if (Times.IsValidIndex(KeyIndex))
+		const int32 KeyIndex = Algo::LowerBound(Times, InTime);
+		if (Times.IsValidIndex(KeyIndex) && Times[KeyIndex] == InTime)
 		{
 			Data.RemoveKey(KeyIndex);
 		}
diff --git a/Engine/Source/Editor/MovieSceneTools/Public/MovieSceneConstraintChannelHelper.h b/Engine/Source/Editor/MovieSceneTools/Public/MovieSceneConstraintChannelHelper.h
--- a/Engine/Source/Editor/MovieSceneTools/Public/MovieSceneConstraintChannelHelper.h
+++ b/Engine/Source/Editor/MovieSceneTools/Public/MovieSceneConstraintChannelHelper.h
@@ -26,7 +26,7 @@ public:
 	void ComputeLocalTransforms(UWorld* InWorld, const TSharedPtr<ISequencer>& InSequencer, const TArray<FFrameNumber>& InFrames, const bool bToActive);
 	void ComputeLocalTransformsBeforeDeletion(UWorld* InWorld, const TSharedPtr<ISequencer>& InSequencer, const TArray<FFrameNumber>& InFrames);
 	void ComputeCompensation(UWorld* InWorld, const TSharedPtr<ISequencer>& InSequencer, const FFrameNumber& InTime);
-
+	void ComputeLocalTransformsForBaking(UWorld* InWorld, const TSharedPtr<ISequencer>& InSequencer, const TArray<FFrameNumber>& InFrames);
 private:
 
 	TArray<UTickableTransformConstraint*> GetHandleTransformConstraints(UWorld* InWorld) const;
diff --git a/Engine/Source/Runtime/Experimental/Animation/Constraints/Private/ConstraintsManager.cpp b/Engine/Source/Runtime/Experimental/Animation/Constraints/Private/ConstraintsManager.cpp
--- a/Engine/Source/Runtime/Experimental/Animation/Constraints/Private/ConstraintsManager.cpp
+++ b/Engine/Source/Runtime/Experimental/Animation/Constraints/Private/ConstraintsManager.cpp
@@ -71,7 +71,7 @@ void UTickableConstraint::SetActive(const bool bIsActive)
 	ConstraintTick.SetTickFunctionEnable(bIsActive);
 }
 
-void UTickableConstraint::Evaluate() const
+void UTickableConstraint::Evaluate(bool bTickHandlesAlso) const
 {
 	ConstraintTick.EvaluateFunctions();
 }
@@ -388,7 +388,6 @@ bool FConstraintsManagerController::AddConstraint(UTickableConstraint* InConstra
 	}
 
 	Manager->Modify();
-
 	Manager->Constraints.Emplace(InConstraint);
 
 	InConstraint->ConstraintTick.RegisterFunction(InConstraint->GetFunction());
@@ -599,4 +598,13 @@ TArray< TObjectPtr<UTickableConstraint> > FConstraintsManagerController::GetAllC
 	});
 	
 	return SortedConstraints;
+}
+
+void FConstraintsManagerController::EvaluateAllConstraints() const
+{
+	TArray< TObjectPtr<UTickableConstraint>>Constraints = GetAllConstraints(true);
+	for (const UTickableConstraint* InConstraint : Constraints)
+	{
+		InConstraint->Evaluate(true);
+	}
 }
\ No newline at end of file
diff --git a/Engine/Source/Runtime/Experimental/Animation/Constraints/Private/TransformableHandle.cpp b/Engine/Source/Runtime/Experimental/Animation/Constraints/Private/TransformableHandle.cpp
--- a/Engine/Source/Runtime/Experimental/Animation/Constraints/Private/TransformableHandle.cpp
+++ b/Engine/Source/Runtime/Experimental/Animation/Constraints/Private/TransformableHandle.cpp
@@ -11,6 +11,8 @@
 #include "Channels/MovieSceneDoubleChannel.h"
 #include "Channels/MovieSceneChannelProxy.h"
 #include "Channels/MovieSceneDoubleChannel.h"
+#include "Components/SkeletalMeshComponent.h"
+
 /**
  * UTransformableHandle
  */
@@ -42,6 +44,34 @@ bool UTransformableComponentHandle::IsValid() const
 	return Component.IsValid();
 }
 
+//need to tick any skelmesh component, sibling or parent
+void UTransformableComponentHandle::TickForBaking()
+{
+	if (!Component.IsValid())
+	{
+		return;
+	}
+
+	const AActor* Parent = Component->GetOwner();
+	while (Parent)
+	{
+		TArray<USkeletalMeshComponent*> MeshComps;
+		Parent->GetComponents(MeshComps, true);
+
+		for (USkeletalMeshComponent* MeshComp : MeshComps)
+		{
+			MeshComp->TickAnimation(0.03f, false);
+			MeshComp->RefreshBoneTransforms();
+			MeshComp->RefreshFollowerComponents();
+			MeshComp->UpdateComponentToWorld();
+			MeshComp->FinalizeBoneTransform();
+			MeshComp->MarkRenderTransformDirty();
+			MeshComp->MarkRenderDynamicDataDirty();
+		}
+
+		Parent = Parent->GetAttachParentActor();
+	}
+}
 void UTransformableComponentHandle::SetGlobalTransform(const FTransform& InGlobal) const
 {
 	if (Component.IsValid())
diff --git a/Engine/Source/Runtime/Experimental/Animation/Constraints/Public/ConstraintsManager.h b/Engine/Source/Runtime/Experimental/Animation/Constraints/Public/ConstraintsManager.h
--- a/Engine/Source/Runtime/Experimental/Animation/Constraints/Public/ConstraintsManager.h
+++ b/Engine/Source/Runtime/Experimental/Animation/Constraints/Public/ConstraintsManager.h
@@ -77,6 +77,9 @@ public:
 	/** Returns the actual function that the tick function needs to evaluate. */
 	virtual FConstraintTickFunction::ConstraintFunction GetFunction() const PURE_VIRTUAL(GetFunction, return {};);
 
+	/** Evaluates the constraint in a context where it's mot done thru the ConstraintTick's tick function. */
+	virtual void Evaluate(bool bTickHandlesAlso = false) const;
+
 	/** Tick function that will be registered and evaluated. */
 	UPROPERTY()
 	FConstraintTickFunction ConstraintTick;
@@ -84,9 +87,6 @@ public:
 	/** Sets the Active value and enable/disable the tick function. */
 	void SetActive(const bool bIsActive);
 	
-	/** Evaluates the constraint in a context where it's mot done thru the ConstraintTick's tick function. */
-	void Evaluate() const;
-	
 	/** @todo document */
 	virtual uint32 GetTargetHash() const PURE_VIRTUAL(GetTargetHash, return 0;);
 	/** @todo document */
@@ -204,7 +204,8 @@ public:
 		// unique name (we may want to use another approach here to manage uniqueness)
 		const FName Name = MakeUniqueObjectName(Manager, TConstraint::StaticClass(), InBaseName);
 
-		TConstraint* NewConstraint = NewObject<TConstraint>(Manager, Name, RF_Transactional);		
+		TConstraint* NewConstraint = NewObject<TConstraint>(Manager, Name, RF_Transactional);
+		NewConstraint->Modify();
 		return NewConstraint;
 	}
 
@@ -241,6 +242,9 @@ public:
 	void SetConstraintsDependencies(
 		const FName& InNameToTickBefore,
 		const FName& InNameToTickAfter) const;
+
+	/** Go through each constraint in order and evaluate and tick them*/
+	void EvaluateAllConstraints() const;
 	
 private:
 	/** Delegeate that's fired when a scene component is constrained, this is needed to make sure things like gizmo's get updated after the constraint tick happens*/
diff --git a/Engine/Source/Runtime/Experimental/Animation/Constraints/Public/TransformableHandle.h b/Engine/Source/Runtime/Experimental/Animation/Constraints/Public/TransformableHandle.h
--- a/Engine/Source/Runtime/Experimental/Animation/Constraints/Public/TransformableHandle.h
+++ b/Engine/Source/Runtime/Experimental/Animation/Constraints/Public/TransformableHandle.h
@@ -56,6 +56,9 @@ public:
 	/** Gets the local transform of the underlying transformable object in it's parent space. */
 	virtual FTransform GetLocalTransform() const PURE_VIRTUAL(GetLocalTransform, return FTransform::Identity;);
 
+	/** Performan any special ticking needed for this handle, by default it does nothing
+	todo need to see if we need to tick control rig also*/
+	virtual void TickForBaking() {};
 	/** Get the array of float channels for the specified section*/
 	virtual TArrayView<FMovieSceneFloatChannel*>  GetFloatChannels(const UMovieSceneSection* InSection) const PURE_VIRTUAL(GetFloatChannels, return TArrayView<FMovieSceneFloatChannel*>(); );
 	/** Get the array of double channels for the specified section*/
@@ -126,7 +129,8 @@ public:
 	virtual FTransform GetGlobalTransform() const override;
 	/** Gets the local transform of Component in it's attachment. */
 	virtual FTransform GetLocalTransform() const override;
-
+	/** Tick the component*/
+	virtual void TickForBaking() override;
 	/** Returns the target object containing the tick function (e.i. Component). */
 	virtual UObject* GetPrerequisiteObject() const override;
 	/** Returns Component's tick function. */
