Fixup deprecated function calls

diff --git a/Engine/Source/Editor/MovieSceneTools/Private/MovieSceneConstraintChannelHelper.cpp b/Engine/Source/Editor/MovieSceneTools/Private/MovieSceneConstraintChannelHelper.cpp
--- a/Engine/Source/Editor/MovieSceneTools/Private/MovieSceneConstraintChannelHelper.cpp
+++ b/Engine/Source/Editor/MovieSceneTools/Private/MovieSceneConstraintChannelHelper.cpp
@@ -93,7 +93,7 @@ void FCompensationEvaluator::ComputeLocalTransforms(
 				BakeHelper->PreEvaluation(MovieScene, FrameNumber);
 			}
 		}
-		InSequencer->GetEvaluationTemplate().Evaluate(Context, *InSequencer);
+		InSequencer->GetEvaluationTemplate().EvaluateSynchronousBlocking(Context, *InSequencer);
 
 		// evaluate constraints
 		for (const UTickableTransformConstraint* InConstraint : Constraints)
@@ -215,7 +215,7 @@ void FCompensationEvaluator::ComputeLocalTransformsForBaking(UWorld* InWorld, co
 				BakeHelper->PreEvaluation(MovieScene, FrameNumber);
 			}
 		}
-		InSequencer->GetEvaluationTemplate().Evaluate(Context, *InSequencer);
+		InSequencer->GetEvaluationTemplate().EvaluateSynchronousBlocking(Context, *InSequencer);
 
 		// evaluate constraints
 		for (const UTickableTransformConstraint* InConstraint : Constraints)
@@ -312,7 +312,7 @@ void FCompensationEvaluator::ComputeLocalTransformsBeforeDeletion(
 				BakeHelper->PreEvaluation(MovieScene, FrameNumber);
 			}
 		}
-		InSequencer->GetEvaluationTemplate().Evaluate(Context, *InSequencer);
+		InSequencer->GetEvaluationTemplate().EvaluateSynchronousBlocking(Context, *InSequencer);
 
 		// evaluate constraints
 		for (const UTickableTransformConstraint* InConstraint : Constraints)
