Fix for sequencer crash when adding a montage to an animation track, which is currently not supported (the filter was broken, as a second filter was added which invalidated the original one)

diff --git a/Engine/Source/Editor/MovieSceneTools/Private/TrackEditors/SkeletalAnimationTrackEditor.cpp b/Engine/Source/Editor/MovieSceneTools/Private/TrackEditors/SkeletalAnimationTrackEditor.cpp
--- a/Engine/Source/Editor/MovieSceneTools/Private/TrackEditors/SkeletalAnimationTrackEditor.cpp
+++ b/Engine/Source/Editor/MovieSceneTools/Private/TrackEditors/SkeletalAnimationTrackEditor.cpp
@@ -1756,7 +1756,6 @@ void FSkeletalAnimationTrackEditor::AddAnimationSubMenu(FMenuBuilder& MenuBuilde
 		AssetPickerConfig.OnAssetEnterPressed = FOnAssetEnterPressed::CreateRaw( this, &FSkeletalAnimationTrackEditor::OnAnimationAssetEnterPressed, ObjectBindings, Track);
 		AssetPickerConfig.bAllowNullSelection = false;
 		AssetPickerConfig.InitialAssetViewType = EAssetViewType::List;
-		AssetPickerConfig.OnShouldFilterAsset = FOnShouldFilterAsset::CreateRaw(this, &FSkeletalAnimationTrackEditor::ShouldFilterAsset);
 		AssetPickerConfig.Filter.bRecursiveClasses = true;
 		AssetPickerConfig.Filter.ClassPaths.Add(UAnimSequenceBase::StaticClass()->GetClassPathName());
 		AssetPickerConfig.OnShouldFilterAsset.BindRaw(this, &FSkeletalAnimationTrackEditor::FilterAnimSequences, Skeleton);
@@ -1778,12 +1777,17 @@ void FSkeletalAnimationTrackEditor::AddAnimationSubMenu(FMenuBuilder& MenuBuilde
 
 bool FSkeletalAnimationTrackEditor::FilterAnimSequences(const FAssetData& AssetData, USkeleton* Skeleton)
 {
-	if (Skeleton && Skeleton->IsCompatibleSkeletonByAssetData(AssetData))
+	if (ShouldFilterAsset(AssetData))
 	{
-		return false;
+		return true;
 	}
 
-	return true;
+	if (Skeleton && Skeleton->IsCompatibleSkeletonByAssetData(AssetData) == false)
+	{
+		return true;
+	}
+
+	return false;
 }
 
 void FSkeletalAnimationTrackEditor::OnAnimationAssetSelected(const FAssetData& AssetData, TArray<FGuid> ObjectBindings, UMovieSceneTrack* Track)
