Constraints improvements:
    - register component handle delegates on post load
    - evaluate all constraints when compensating to make sure things are up to date

diff --git a/Engine/Plugins/Animation/ControlRig/Source/ControlRigEditor/Private/EditMode/ControlRigEditMode.cpp b/Engine/Plugins/Animation/ControlRig/Source/ControlRigEditor/Private/EditMode/ControlRigEditMode.cpp
--- a/Engine/Plugins/Animation/ControlRig/Source/ControlRigEditor/Private/EditMode/ControlRigEditMode.cpp
+++ b/Engine/Plugins/Animation/ControlRig/Source/ControlRigEditor/Private/EditMode/ControlRigEditMode.cpp
@@ -3776,13 +3776,12 @@ void FControlRigEditMode::SetControlShapeTransform(
 		return false;
 	});
 
-	// set the global space
-	// assumes it's attached to actor
-	ControlRig->SetControlGlobalTransform(
-		InShapeActor->ControlName, InGlobalTransform, bNotify, InContext, bUndo, bPrintPython, bFixEuler);
-
 	if (!Constraints.IsValidIndex(LastActiveIndex))
 	{
+		// set the global space
+		// assumes it's attached to actor
+		ControlRig->SetControlGlobalTransform(
+			InShapeActor->ControlName, InGlobalTransform, bNotify, InContext, bUndo, bPrintPython, bFixEuler);
 		return;
 	}
 	
@@ -3799,8 +3798,9 @@ void FControlRigEditMode::SetControlShapeTransform(
 	Context.bConstraintUpdate = false;
 	
 	ControlRig->SetControlLocalTransform(InShapeActor->ControlName, LocalTransform, bNotify, Context, bUndo, bFixEuler);
-	
-	Constraint->Evaluate();
+
+	const FConstraintsManagerController& Controller = FConstraintsManagerController::Get(ControlRig->GetWorld());
+	Controller.EvaluateAllConstraints();
 }
 
 FTransform FControlRigEditMode::GetControlShapeTransform(const AControlRigShapeActor* ShapeActor)
diff --git a/Engine/Source/Editor/MovieSceneTools/Private/MovieSceneConstraintChannelHelper.cpp b/Engine/Source/Editor/MovieSceneTools/Private/MovieSceneConstraintChannelHelper.cpp
--- a/Engine/Source/Editor/MovieSceneTools/Private/MovieSceneConstraintChannelHelper.cpp
+++ b/Engine/Source/Editor/MovieSceneTools/Private/MovieSceneConstraintChannelHelper.cpp
@@ -32,6 +32,7 @@ FCompensationEvaluator::FCompensationEvaluator(UTickableTransformConstraint* InC
 	: Constraint(InConstraint)
 	, Handle(InConstraint ? InConstraint->ChildTRSHandle : nullptr)
 {}
+
 void FCompensationEvaluator::ComputeLocalTransforms(
 	UWorld* InWorld, const TSharedPtr<ISequencer>& InSequencer,
 	const TArray<FFrameNumber>& InFrames, const bool bToActive)
@@ -78,6 +79,12 @@ void FCompensationEvaluator::ComputeLocalTransforms(
 			BakeHelper->StartBaking(MovieScene);
 		}
 	}
+
+	// get all constraints to evaluate
+	const FConstraintsManagerController& Controller = FConstraintsManagerController::Get(InWorld);
+	static constexpr bool bSorted = true;
+	TArray<TObjectPtr<UTickableConstraint>> AllConstraints = Controller.GetAllConstraints(bSorted);
+	
 	for (int32 Index = 0; Index < NumFrames + 1; ++Index)
 	{
 		const FFrameNumber FrameNumber = (Index == 0) ? InFrames[0] - 1 : InFrames[Index - 1];
@@ -96,12 +103,9 @@ void FCompensationEvaluator::ComputeLocalTransforms(
 		InSequencer->GetEvaluationTemplate().EvaluateSynchronousBlocking(Context, *InSequencer);
 
 		// evaluate constraints
-		for (const UTickableTransformConstraint* InConstraint : Constraints)
+		for (const UTickableConstraint* InConstraint : AllConstraints)
 		{
-			if (InConstraint)
-			{
-				InConstraint->Evaluate(true);
-			}
+			InConstraint->Evaluate();
 		}
 
 		for (IMovieSceneToolsAnimationBakeHelper* BakeHelper : BakeHelpers)
@@ -192,6 +196,12 @@ void FCompensationEvaluator::ComputeLocalTransformsForBaking(UWorld* InWorld, co
 	ChildLocals.SetNum(NumFrames);
 	ChildGlobals.SetNum(NumFrames);
 	SpaceGlobals.SetNum(NumFrames);
+
+	// get all constraints for evaluation
+	const FConstraintsManagerController& Controller = FConstraintsManagerController::Get(InWorld);
+	static constexpr bool bSorted = true;
+	TArray<TObjectPtr<UTickableConstraint>> AllConstraints = Controller.GetAllConstraints(bSorted);
+	
 	const TArray<IMovieSceneToolsAnimationBakeHelper*>& BakeHelpers = FMovieSceneToolsModule::Get().GetAnimationBakeHelpers();
 	for (IMovieSceneToolsAnimationBakeHelper* BakeHelper : BakeHelpers)
 	{
@@ -218,7 +228,7 @@ void FCompensationEvaluator::ComputeLocalTransformsForBaking(UWorld* InWorld, co
 		InSequencer->GetEvaluationTemplate().EvaluateSynchronousBlocking(Context, *InSequencer);
 
 		// evaluate constraints
-		for (const UTickableTransformConstraint* InConstraint : Constraints)
+		for (UTickableConstraint* InConstraint : AllConstraints)
 		{
 			InConstraint->Evaluate(true);
 		}
@@ -280,6 +290,11 @@ void FCompensationEvaluator::ComputeLocalTransformsBeforeDeletion(
 		return LastActiveIndex > INDEX_NONE ? Constraints[LastActiveIndex] : nullptr;
 	};
 
+	// get all constraints for evaluation
+	const FConstraintsManagerController& Controller = FConstraintsManagerController::Get(InWorld);
+	static constexpr bool bSorted = true;
+	TArray<TObjectPtr<UTickableConstraint>> AllConstraints = Controller.GetAllConstraints(bSorted);
+	
 	UMovieScene* MovieScene = InSequencer->GetFocusedMovieSceneSequence()->GetMovieScene();
 	const FFrameRate TickResolution = MovieScene->GetTickResolution();
 	const EMovieScenePlayerStatus::Type PlaybackStatus = InSequencer->GetPlaybackStatus();
@@ -315,10 +330,11 @@ void FCompensationEvaluator::ComputeLocalTransformsBeforeDeletion(
 		InSequencer->GetEvaluationTemplate().EvaluateSynchronousBlocking(Context, *InSequencer);
 
 		// evaluate constraints
-		for (const UTickableTransformConstraint* InConstraint : Constraints)
+		for (const UTickableConstraint* InConstraint : AllConstraints)
 		{
 			InConstraint->Evaluate(true);
 		}
+		
 		for (IMovieSceneToolsAnimationBakeHelper* BakeHelper : BakeHelpers)
 		{
 			if (BakeHelper)
@@ -375,6 +391,11 @@ void FCompensationEvaluator::ComputeCompensation(UWorld* InWorld, const TSharedP
 		return LastActiveIndex > INDEX_NONE ? Constraints[LastActiveIndex] : nullptr;
 	};
 
+	// get all constraints for evaluation
+	const FConstraintsManagerController& Controller = FConstraintsManagerController::Get(InWorld);
+	static constexpr bool bSorted = true;
+	TArray<TObjectPtr<UTickableConstraint>> AllConstraints = Controller.GetAllConstraints(bSorted);
+	
 	UMovieScene* MovieScene = InSequencer->GetFocusedMovieSceneSequence()->GetMovieScene();
 	const TArray<IMovieSceneToolsAnimationBakeHelper*>& BakeHelpers = FMovieSceneToolsModule::Get().GetAnimationBakeHelpers();
 	for (IMovieSceneToolsAnimationBakeHelper* BakeHelper : BakeHelpers)
@@ -384,7 +405,8 @@ void FCompensationEvaluator::ComputeCompensation(UWorld* InWorld, const TSharedP
 			BakeHelper->StartBaking(MovieScene);
 		}
 	}
-	auto EvaluateAt = [InSequencer, &Constraints, &BakeHelpers](const FFrameNumber InFrame)
+	
+	auto EvaluateAt = [InSequencer, &AllConstraints, &BakeHelpers](const FFrameNumber InFrame)
 	{
 
 		UMovieScene* MovieScene = InSequencer->GetFocusedMovieSceneSequence()->GetMovieScene();
@@ -403,7 +425,7 @@ void FCompensationEvaluator::ComputeCompensation(UWorld* InWorld, const TSharedP
 		}
 		InSequencer->GetEvaluationTemplate().EvaluateSynchronousBlocking(Context0, *InSequencer);
 
-		for (const UTickableTransformConstraint* InConstraint : Constraints)
+		for (const UTickableConstraint* InConstraint : AllConstraints)
 		{
 			InConstraint->Evaluate(true);
 		}
@@ -426,7 +448,6 @@ void FCompensationEvaluator::ComputeCompensation(UWorld* InWorld, const TSharedP
 	SpaceGlobals.SetNum(1);
 
 	
-
 	// evaluate at InTime and store global
 	EvaluateAt(InTime);
 	ChildGlobals[0] = Handle->GetGlobalTransform();
@@ -435,7 +456,6 @@ void FCompensationEvaluator::ComputeCompensation(UWorld* InWorld, const TSharedP
 	EvaluateAt(InTime - 1);
 	ChildLocals[0] = Handle->GetLocalTransform();
 
-
 	for (IMovieSceneToolsAnimationBakeHelper* BakeHelper : BakeHelpers)
 	{
 		if (BakeHelper)
@@ -443,6 +463,7 @@ void FCompensationEvaluator::ComputeCompensation(UWorld* InWorld, const TSharedP
 			BakeHelper->StopBaking(MovieScene);
 		}
 	}
+	
 	// if constraint at T-1 then switch to its space
 	if (const UTickableTransformConstraint* LastConstraint = GetLastActiveConstraint())
 	{
diff --git a/Engine/Source/Runtime/Experimental/Animation/Constraints/Private/TransformableHandle.cpp b/Engine/Source/Runtime/Experimental/Animation/Constraints/Private/TransformableHandle.cpp
--- a/Engine/Source/Runtime/Experimental/Animation/Constraints/Private/TransformableHandle.cpp
+++ b/Engine/Source/Runtime/Experimental/Animation/Constraints/Private/TransformableHandle.cpp
@@ -37,7 +37,15 @@ UTransformableHandle::FHandleModifiedEvent& UTransformableHandle::HandleModified
  */
 
 UTransformableComponentHandle::~UTransformableComponentHandle()
-{}
+{
+	UnregisterDelegates();
+}
+
+void UTransformableComponentHandle::PostLoad()
+{
+	Super::PostLoad();
+	RegisterDelegates();
+}
 
 bool UTransformableComponentHandle::IsValid() const
 {
@@ -236,7 +244,10 @@ void UTransformableComponentHandle::UnregisterDelegates() const
 {
 #if WITH_EDITOR
 	FCoreUObjectDelegates::OnObjectPropertyChanged.RemoveAll(this);
-	GEngine->OnActorMoving().RemoveAll(this);
+	if (GEngine)
+	{
+		GEngine->OnActorMoving().RemoveAll(this);
+	}
 #endif
 }
 
diff --git a/Engine/Source/Runtime/Experimental/Animation/Constraints/Public/TransformableHandle.h b/Engine/Source/Runtime/Experimental/Animation/Constraints/Public/TransformableHandle.h
--- a/Engine/Source/Runtime/Experimental/Animation/Constraints/Public/TransformableHandle.h
+++ b/Engine/Source/Runtime/Experimental/Animation/Constraints/Public/TransformableHandle.h
@@ -117,6 +117,8 @@ class CONSTRAINTS_API UTransformableComponentHandle : public UTransformableHandl
 public:
 	
 	virtual ~UTransformableComponentHandle();
+
+	virtual void PostLoad() override;
 	
 	/** Sanity check to ensure that Component. */
 	virtual bool IsValid() const override;
