remove constraint removal delegate when appropriate

diff --git a/Engine/Plugins/Animation/ControlRig/Source/ControlRigEditor/Private/Sequencer/ControlRigParameterTrackEditor.cpp b/Engine/Plugins/Animation/ControlRig/Source/ControlRigEditor/Private/Sequencer/ControlRigParameterTrackEditor.cpp
--- a/Engine/Plugins/Animation/ControlRig/Source/ControlRigEditor/Private/Sequencer/ControlRigParameterTrackEditor.cpp
+++ b/Engine/Plugins/Animation/ControlRig/Source/ControlRigEditor/Private/Sequencer/ControlRigParameterTrackEditor.cpp
@@ -2556,6 +2556,16 @@ void FControlRigParameterTrackEditor::ClearOutAllSpaceAndConstraintDelegates(con
 						Channel.ActiveChannel.OnKeyMovedEvent().Clear();
 						Channel.ActiveChannel.OnKeyDeletedEvent().Clear();							
 					}
+
+					if (CRSection->OnConstraintRemovedHandle.IsValid())
+					{
+						if (const UControlRig* ControlRig = CRSection->GetControlRig())
+						{
+							FConstraintsManagerController& Controller = FConstraintsManagerController::Get(ControlRig->GetWorld());
+							Controller.OnConstraintRemoved().Remove(CRSection->OnConstraintRemovedHandle);
+							CRSection->OnConstraintRemovedHandle.Reset();
+						}
+					}
 				}
 			}
 		}
diff --git a/Engine/Source/Editor/MovieSceneTools/Private/TrackEditors/TransformTrackEditor.cpp b/Engine/Source/Editor/MovieSceneTools/Private/TrackEditors/TransformTrackEditor.cpp
--- a/Engine/Source/Editor/MovieSceneTools/Private/TrackEditors/TransformTrackEditor.cpp
+++ b/Engine/Source/Editor/MovieSceneTools/Private/TrackEditors/TransformTrackEditor.cpp
@@ -1617,6 +1617,9 @@ void F3DTransformTrackEditor::ClearOutConstraintDelegates() const
 		return;
 	}
 
+	UWorld* World = GCurrentLevelEditingViewportClient ? GCurrentLevelEditingViewportClient->GetWorld() : nullptr;
+	FConstraintsManagerController& Controller = FConstraintsManagerController::Get(World);
+	
 	const TArray<FMovieSceneBinding>& Bindings = MovieScene->GetBindings();
 	for (const FMovieSceneBinding& Binding : Bindings)
 	{
@@ -1635,6 +1638,12 @@ void F3DTransformTrackEditor::ClearOutConstraintDelegates() const
 						Channel.ActiveChannel.OnKeyMovedEvent().Clear();
 						Channel.ActiveChannel.OnKeyDeletedEvent().Clear();
 					}
+
+					if (CRSection->OnConstraintRemovedHandle.IsValid())
+					{
+						Controller.OnConstraintRemoved().Remove(CRSection->OnConstraintRemovedHandle);
+						CRSection->OnConstraintRemovedHandle.Reset();
+					}
 				}
 			}
 		}
