Sequencer: Fix cases of TransformSection->GetChannelProxy().GetChannels()

diff --git a/Engine/Source/Editor/MovieSceneTools/Private/MovieSceneToolHelpers.cpp b/Engine/Source/Editor/MovieSceneTools/Private/MovieSceneToolHelpers.cpp
--- a/Engine/Source/Editor/MovieSceneTools/Private/MovieSceneToolHelpers.cpp
+++ b/Engine/Source/Editor/MovieSceneTools/Private/MovieSceneToolHelpers.cpp
@@ -2789,32 +2789,68 @@ bool ImportFBXTransform(FString NodeName, FGuid ObjectBinding, UnFbx::FFbxCurves
 	}
 
 	FVector Location = DefaultTransform.GetLocation(), Rotation = DefaultTransform.GetRotation().Euler(), Scale3D = DefaultTransform.GetScale3D();
+	
+	FMovieSceneChannelProxy& SectionChannelProxy = TransformSection->GetChannelProxy();
+	TMovieSceneChannelHandle<FMovieSceneDoubleChannel> DoubleChannels[] = {
+		SectionChannelProxy.GetChannelByName<FMovieSceneDoubleChannel>("Location.X"),
+		SectionChannelProxy.GetChannelByName<FMovieSceneDoubleChannel>("Location.Y"),
+		SectionChannelProxy.GetChannelByName<FMovieSceneDoubleChannel>("Location.Z"),
+		SectionChannelProxy.GetChannelByName<FMovieSceneDoubleChannel>("Rotation.X"),
+		SectionChannelProxy.GetChannelByName<FMovieSceneDoubleChannel>("Rotation.Y"),
+		SectionChannelProxy.GetChannelByName<FMovieSceneDoubleChannel>("Rotation.Z"),
+		SectionChannelProxy.GetChannelByName<FMovieSceneDoubleChannel>("Scale.X"),
+		SectionChannelProxy.GetChannelByName<FMovieSceneDoubleChannel>("Scale.Y"),
+		SectionChannelProxy.GetChannelByName<FMovieSceneDoubleChannel>("Scale.Z")
+	};
 
-	TArrayView<FMovieSceneDoubleChannel*> Channels = TransformSection->GetChannelProxy().GetChannels<FMovieSceneDoubleChannel>();
-
-	Channels[0]->SetDefault(Location.X);
-	Channels[1]->SetDefault(Location.Y);
-	Channels[2]->SetDefault(Location.Z);
-
-	Channels[3]->SetDefault(Rotation.X);
-	Channels[4]->SetDefault(Rotation.Y);
-	Channels[5]->SetDefault(Rotation.Z);
-
-	Channels[6]->SetDefault(Scale3D.X);
-	Channels[7]->SetDefault(Scale3D.Y);
-	Channels[8]->SetDefault(Scale3D.Z);
+	if (DoubleChannels[0].Get())
+	{
+		DoubleChannels[0].Get()->SetDefault(Location.X);
+		ImportTransformChannelToDouble(Translation[0], DoubleChannels[0].Get(), FrameRate, false, true);
+	}
 
-	ImportTransformChannelToDouble(Translation[0],   Channels[0], FrameRate, false, true);
-	ImportTransformChannelToDouble(Translation[1],   Channels[1], FrameRate, false, true);
-	ImportTransformChannelToDouble(Translation[2],   Channels[2], FrameRate, false, true);
+	if (DoubleChannels[1].Get())
+	{
+		DoubleChannels[1].Get()->SetDefault(Location.Y);
+		ImportTransformChannelToDouble(Translation[1], DoubleChannels[1].Get(), FrameRate, false, true);
+	}
+	if (DoubleChannels[2].Get())
+	{
+		DoubleChannels[2].Get()->SetDefault(Location.Z);
+		ImportTransformChannelToDouble(Translation[2], DoubleChannels[2].Get(), FrameRate, false, true);
+	}
 
-	ImportTransformChannelToDouble(EulerRotation[0], Channels[3], FrameRate, false, true);
-	ImportTransformChannelToDouble(EulerRotation[1], Channels[4], FrameRate, false, true);
-	ImportTransformChannelToDouble(EulerRotation[2], Channels[5], FrameRate, false, true);
+	if (DoubleChannels[3].Get())
+	{
+		DoubleChannels[3].Get()->SetDefault(Rotation.X);
+		ImportTransformChannelToDouble(EulerRotation[0], DoubleChannels[3].Get(), FrameRate, false, true);
+	}
+	if (DoubleChannels[4].Get())
+	{
+		DoubleChannels[4].Get()->SetDefault(Rotation.Y);
+		ImportTransformChannelToDouble(EulerRotation[1], DoubleChannels[4].Get(), FrameRate, false, true);
+	}
+	if (DoubleChannels[5].Get())
+	{
+		DoubleChannels[5].Get()->SetDefault(Rotation.Z);
+		ImportTransformChannelToDouble(EulerRotation[2], DoubleChannels[5].Get(), FrameRate, false, true);
+	}
 
-	ImportTransformChannelToDouble(Scale[0],         Channels[6], FrameRate, false, true);
-	ImportTransformChannelToDouble(Scale[1],         Channels[7], FrameRate, false, true);
-	ImportTransformChannelToDouble(Scale[2],         Channels[8], FrameRate, false, true);
+	if (DoubleChannels[6].Get())
+	{
+		DoubleChannels[6].Get()->SetDefault(Scale3D.X);
+		ImportTransformChannelToDouble(Scale[0], DoubleChannels[6].Get(), FrameRate, false, true);
+	}
+	if (DoubleChannels[7].Get())
+	{
+		DoubleChannels[7].Get()->SetDefault(Scale3D.Y);
+		ImportTransformChannelToDouble(Scale[1], DoubleChannels[7].Get(), FrameRate, false, true);
+	}
+	if (DoubleChannels[8].Get())
+	{
+		DoubleChannels[8].Get()->SetDefault(Scale3D.Z);
+		ImportTransformChannelToDouble(Scale[2], DoubleChannels[8].Get(), FrameRate, false, true);
+	}
 
 	return true;
 }
@@ -4624,17 +4660,55 @@ UMovieScene3DTransformSection* MovieSceneToolHelpers::GetTransformSection(
 		const FRotator Rotation0 = InDefaultTransform.GetRotation().Rotator();
 		const FVector Scale3D0 = InDefaultTransform.GetScale3D();
 
-		const TArrayView<FMovieSceneDoubleChannel*> Channels =
-			TransformSection->GetChannelProxy().GetChannels<FMovieSceneDoubleChannel>();
-		Channels[0]->SetDefault(Location0.X);
-		Channels[1]->SetDefault(Location0.Y);
-		Channels[2]->SetDefault(Location0.Z);
-		Channels[3]->SetDefault(Rotation0.Roll);
-		Channels[4]->SetDefault(Rotation0.Pitch);
-		Channels[5]->SetDefault(Rotation0.Yaw);
-		Channels[6]->SetDefault(Scale3D0.X);
-		Channels[7]->SetDefault(Scale3D0.Y);
-		Channels[8]->SetDefault(Scale3D0.Z);
+		FMovieSceneChannelProxy& SectionChannelProxy = TransformSection->GetChannelProxy();
+		TMovieSceneChannelHandle<FMovieSceneDoubleChannel> DoubleChannels[] = {
+			SectionChannelProxy.GetChannelByName<FMovieSceneDoubleChannel>("Location.X"),
+			SectionChannelProxy.GetChannelByName<FMovieSceneDoubleChannel>("Location.Y"),
+			SectionChannelProxy.GetChannelByName<FMovieSceneDoubleChannel>("Location.Z"),
+			SectionChannelProxy.GetChannelByName<FMovieSceneDoubleChannel>("Rotation.X"),
+			SectionChannelProxy.GetChannelByName<FMovieSceneDoubleChannel>("Rotation.Y"),
+			SectionChannelProxy.GetChannelByName<FMovieSceneDoubleChannel>("Rotation.Z"),
+			SectionChannelProxy.GetChannelByName<FMovieSceneDoubleChannel>("Scale.X"),
+			SectionChannelProxy.GetChannelByName<FMovieSceneDoubleChannel>("Scale.Y"),
+			SectionChannelProxy.GetChannelByName<FMovieSceneDoubleChannel>("Scale.Z")
+		};
+
+		if (DoubleChannels[0].Get())
+		{
+			DoubleChannels[0].Get()->SetDefault(Location0.X);
+		}
+		if (DoubleChannels[1].Get())
+		{
+			DoubleChannels[1].Get()->SetDefault(Location0.Y);
+		}
+		if (DoubleChannels[2].Get())
+		{
+			DoubleChannels[2].Get()->SetDefault(Location0.Z);
+		}
+		if (DoubleChannels[3].Get())
+		{
+			DoubleChannels[3].Get()->SetDefault(Rotation0.Roll);
+		}
+		if (DoubleChannels[4].Get())
+		{
+			DoubleChannels[4].Get()->SetDefault(Rotation0.Pitch);
+		}
+		if (DoubleChannels[5].Get())
+		{
+			DoubleChannels[5].Get()->SetDefault(Rotation0.Yaw);
+		}
+		if (DoubleChannels[6].Get())
+		{
+			DoubleChannels[6].Get()->SetDefault(Scale3D0.X);
+		}
+		if (DoubleChannels[7].Get())
+		{
+			DoubleChannels[7].Get()->SetDefault(Scale3D0.Y);
+		}
+		if (DoubleChannels[8].Get())
+		{
+			DoubleChannels[8].Get()->SetDefault(Scale3D0.Z);
+		}
 	}
 
 	return TransformSection;
@@ -4703,8 +4777,18 @@ bool MovieSceneToolHelpers::AddTransformKeys(
 		ChannelsIndexToKey.Append({ 6,7,8 });
 	}
 
-	const TArrayView<FMovieSceneDoubleChannel*> Channels =
-		InTransformSection->GetChannelProxy().GetChannels<FMovieSceneDoubleChannel>();
+	FMovieSceneChannelProxy& SectionChannelProxy = InTransformSection->GetChannelProxy();
+	TMovieSceneChannelHandle<FMovieSceneDoubleChannel> DoubleChannels[] = {
+		SectionChannelProxy.GetChannelByName<FMovieSceneDoubleChannel>("Location.X"),
+		SectionChannelProxy.GetChannelByName<FMovieSceneDoubleChannel>("Location.Y"),
+		SectionChannelProxy.GetChannelByName<FMovieSceneDoubleChannel>("Location.Z"),
+		SectionChannelProxy.GetChannelByName<FMovieSceneDoubleChannel>("Rotation.X"),
+		SectionChannelProxy.GetChannelByName<FMovieSceneDoubleChannel>("Rotation.Y"),
+		SectionChannelProxy.GetChannelByName<FMovieSceneDoubleChannel>("Rotation.Z"),
+		SectionChannelProxy.GetChannelByName<FMovieSceneDoubleChannel>("Scale.X"),
+		SectionChannelProxy.GetChannelByName<FMovieSceneDoubleChannel>("Scale.Y"),
+		SectionChannelProxy.GetChannelByName<FMovieSceneDoubleChannel>("Scale.Z")
+	};
 
 	// set default
 	const FTransform& LocalTransform0 = InLocalTransforms[0];
@@ -4714,10 +4798,13 @@ bool MovieSceneToolHelpers::AddTransformKeys(
 
 	for (int32 ChannelIndex = 0; ChannelIndex < 9; ChannelIndex++)
 	{
-		if (!Channels[ChannelIndex]->GetDefault().IsSet())
+		if (DoubleChannels[ChannelIndex].Get())
 		{
-			const double Value = GetValue(ChannelIndex, Location0, Rotation0, Scale3D0);
-			Channels[ChannelIndex]->SetDefault(Value);
+			if (!DoubleChannels[ChannelIndex].Get()->GetDefault().IsSet())
+			{
+				const double Value = GetValue(ChannelIndex, Location0, Rotation0, Scale3D0);
+				DoubleChannels[ChannelIndex].Get()->SetDefault(Value);
+			}
 		}
 	}
 
@@ -4733,16 +4820,22 @@ bool MovieSceneToolHelpers::AddTransformKeys(
 
 		for (const int32 ChannelIndex : ChannelsIndexToKey)
 		{
-			const double Value = GetValue(ChannelIndex, Location, Rotation, Scale3D);
-			TMovieSceneChannelData<FMovieSceneDoubleValue> ChannelData = Channels[ChannelIndex]->GetData();
-			MovieSceneToolHelpers::SetOrAddKey(ChannelData, Frame, Value);
+			if (DoubleChannels[ChannelIndex].Get())
+			{
+				const double Value = GetValue(ChannelIndex, Location, Rotation, Scale3D);
+				TMovieSceneChannelData<FMovieSceneDoubleValue> ChannelData = DoubleChannels[ChannelIndex].Get()->GetData();
+				MovieSceneToolHelpers::SetOrAddKey(ChannelData, Frame, Value);
+			}
 		}
 	}
 
 	//now we need to set auto tangents
 	for (const int32 ChannelIndex : ChannelsIndexToKey)
 	{
-		Channels[ChannelIndex]->AutoSetTangents();
+		if (DoubleChannels[ChannelIndex].Get())
+		{
+			DoubleChannels[ChannelIndex].Get()->AutoSetTangents();
+		}
 	}
 
 	return true;
diff --git a/Engine/Source/Editor/Sequencer/Private/SequencerKeyActor.cpp b/Engine/Source/Editor/Sequencer/Private/SequencerKeyActor.cpp
--- a/Engine/Source/Editor/Sequencer/Private/SequencerKeyActor.cpp
+++ b/Engine/Source/Editor/Sequencer/Private/SequencerKeyActor.cpp
@@ -85,12 +85,26 @@ void ASequencerKeyActor::PropagateKeyChange()
 		FFrameRate   TickResolution  = TrackSection->GetTypedOuter<UMovieScene>()->GetTickResolution();
 		FFrameNumber FrameNumber     = (KeyTime * TickResolution).RoundToFrame();
 
-		TArrayView<FMovieSceneDoubleChannel*> DoubleChannels = TrackSection->GetChannelProxy().GetChannels<FMovieSceneDoubleChannel>();
+		FMovieSceneChannelProxy& SectionChannelProxy = TrackSection->GetChannelProxy();
+		TMovieSceneChannelHandle<FMovieSceneDoubleChannel> DoubleChannels[] = {
+			SectionChannelProxy.GetChannelByName<FMovieSceneDoubleChannel>("Location.X"),
+			SectionChannelProxy.GetChannelByName<FMovieSceneDoubleChannel>("Location.Y"),
+			SectionChannelProxy.GetChannelByName<FMovieSceneDoubleChannel>("Location.Z")
+		};
 
 		const FVector Translation = GetActorTransform().GetLocation();
-		DoubleChannels[0]->GetData().UpdateOrAddKey(FrameNumber, FMovieSceneDoubleValue(Translation.X));
-		DoubleChannels[1]->GetData().UpdateOrAddKey(FrameNumber, FMovieSceneDoubleValue(Translation.Y));
-		DoubleChannels[2]->GetData().UpdateOrAddKey(FrameNumber, FMovieSceneDoubleValue(Translation.Z));
+		if (DoubleChannels[0].Get())
+		{
+			DoubleChannels[0].Get()->GetData().UpdateOrAddKey(FrameNumber, FMovieSceneDoubleValue(Translation.X));
+		}
+		if (DoubleChannels[1].Get())
+		{
+			DoubleChannels[1].Get()->GetData().UpdateOrAddKey(FrameNumber, FMovieSceneDoubleValue(Translation.Y));
+		}
+		if (DoubleChannels[2].Get())
+		{
+			DoubleChannels[2].Get()->GetData().UpdateOrAddKey(FrameNumber, FMovieSceneDoubleValue(Translation.Z));
+		}
 
 		// Draw a single transform track based on the data from this key
 		FEditorViewportClient* ViewportClient = StaticCast<FEditorViewportClient*>(GEditor->GetActiveViewport()->GetClient());
diff --git a/Engine/Source/Editor/UnrealEd/Private/Fbx/FbxMainExport.cpp b/Engine/Source/Editor/UnrealEd/Private/Fbx/FbxMainExport.cpp
--- a/Engine/Source/Editor/UnrealEd/Private/Fbx/FbxMainExport.cpp
+++ b/Engine/Source/Editor/UnrealEd/Private/Fbx/FbxMainExport.cpp
@@ -2472,27 +2472,65 @@ void FFbxExporter::ExportLevelSequence3DTransformTrack(FbxNode* FbxNode, IMovieS
 	FbxAnimCurve* FbxCurveScaleY = FbxNode->LclScaling.GetCurve(BaseLayer, FBXSDK_CURVENODE_COMPONENT_Y, true);
 	FbxAnimCurve* FbxCurveScaleZ = FbxNode->LclScaling.GetCurve(BaseLayer, FBXSDK_CURVENODE_COMPONENT_Z, true);
 
-	TArrayView<FMovieSceneDoubleChannel*> DoubleChannels = TransformSection->GetChannelProxy().GetChannels<FMovieSceneDoubleChannel>();
+	FMovieSceneChannelProxy& SectionChannelProxy = TransformSection->GetChannelProxy();
+	TMovieSceneChannelHandle<FMovieSceneDoubleChannel> DoubleChannels[] = {
+		SectionChannelProxy.GetChannelByName<FMovieSceneDoubleChannel>("Location.X"),
+		SectionChannelProxy.GetChannelByName<FMovieSceneDoubleChannel>("Location.Y"),
+		SectionChannelProxy.GetChannelByName<FMovieSceneDoubleChannel>("Location.Z"),
+		SectionChannelProxy.GetChannelByName<FMovieSceneDoubleChannel>("Rotation.X"),
+		SectionChannelProxy.GetChannelByName<FMovieSceneDoubleChannel>("Rotation.Y"),
+		SectionChannelProxy.GetChannelByName<FMovieSceneDoubleChannel>("Rotation.Z"),
+		SectionChannelProxy.GetChannelByName<FMovieSceneDoubleChannel>("Scale.X"),
+		SectionChannelProxy.GetChannelByName<FMovieSceneDoubleChannel>("Scale.Y"),
+		SectionChannelProxy.GetChannelByName<FMovieSceneDoubleChannel>("Scale.Z")
+	};
 
 	// Translation
-	ExportChannelToFbxCurve(*FbxCurveTransX, *DoubleChannels[0], TickResolution, ERichCurveValueMode::Default, false, RootToLocalTransform);
-	ExportChannelToFbxCurve(*FbxCurveTransY, *DoubleChannels[1], TickResolution, ERichCurveValueMode::Default, true, RootToLocalTransform);
-	ExportChannelToFbxCurve(*FbxCurveTransZ, *DoubleChannels[2], TickResolution, ERichCurveValueMode::Default, false, RootToLocalTransform);
+	if (DoubleChannels[0].Get())
+	{
+		ExportChannelToFbxCurve(*FbxCurveTransX, *DoubleChannels[0].Get(), TickResolution, ERichCurveValueMode::Default, false, RootToLocalTransform);
+	}
+	if (DoubleChannels[1].Get())
+	{
+		ExportChannelToFbxCurve(*FbxCurveTransY, *DoubleChannels[1].Get(), TickResolution, ERichCurveValueMode::Default, true, RootToLocalTransform);
+	}
+	if (DoubleChannels[2].Get())
+	{
+		ExportChannelToFbxCurve(*FbxCurveTransZ, *DoubleChannels[2].Get(), TickResolution, ERichCurveValueMode::Default, false, RootToLocalTransform);
+	}
 
 	// Scale - don't generate scale keys for cameras
 	if (!bIsCameraActor)
 	{
-		ExportChannelToFbxCurve(*FbxCurveScaleX, *DoubleChannels[6], TickResolution, ERichCurveValueMode::Default, false, RootToLocalTransform);
-		ExportChannelToFbxCurve(*FbxCurveScaleY, *DoubleChannels[7], TickResolution, ERichCurveValueMode::Default, false, RootToLocalTransform);
-		ExportChannelToFbxCurve(*FbxCurveScaleZ, *DoubleChannels[8], TickResolution, ERichCurveValueMode::Default, false, RootToLocalTransform);
+		if (DoubleChannels[6].Get())
+		{
+			ExportChannelToFbxCurve(*FbxCurveScaleX, *DoubleChannels[6].Get(), TickResolution, ERichCurveValueMode::Default, false, RootToLocalTransform);
+		}
+		if (DoubleChannels[7].Get())
+		{
+			ExportChannelToFbxCurve(*FbxCurveScaleY, *DoubleChannels[7].Get(), TickResolution, ERichCurveValueMode::Default, false, RootToLocalTransform);
+		}
+		if (DoubleChannels[8].Get())
+		{
+			ExportChannelToFbxCurve(*FbxCurveScaleZ, *DoubleChannels[8].Get(), TickResolution, ERichCurveValueMode::Default, false, RootToLocalTransform);
+		}
 	}
 
 	// Rotation - bake rotation for cameras and lights
 	if (!bBakeRotations)
 	{
-		ExportChannelToFbxCurve(*FbxCurveRotX, *DoubleChannels[3], TickResolution, ERichCurveValueMode::Default, false, RootToLocalTransform);
-		ExportChannelToFbxCurve(*FbxCurveRotY, *DoubleChannels[4], TickResolution, ERichCurveValueMode::Default, true, RootToLocalTransform);
-		ExportChannelToFbxCurve(*FbxCurveRotZ, *DoubleChannels[5], TickResolution, ERichCurveValueMode::Default, true, RootToLocalTransform);
+		if (DoubleChannels[3].Get())
+		{
+			ExportChannelToFbxCurve(*FbxCurveRotX, *DoubleChannels[3].Get(), TickResolution, ERichCurveValueMode::Default, false, RootToLocalTransform);
+		}
+		if (DoubleChannels[4].Get())
+		{
+			ExportChannelToFbxCurve(*FbxCurveRotY, *DoubleChannels[4].Get(), TickResolution, ERichCurveValueMode::Default, true, RootToLocalTransform);
+		}
+		if (DoubleChannels[5].Get())
+		{
+			ExportChannelToFbxCurve(*FbxCurveRotZ, *DoubleChannels[5].Get(), TickResolution, ERichCurveValueMode::Default, true, RootToLocalTransform);
+		}
 	}
 	else
 	{
@@ -2523,19 +2561,46 @@ void FFbxExporter::ExportLevelSequence3DTransformTrack(FbxNode* FbxNode, IMovieS
 			FFrameTime LocalTime = FFrameRate::TransformTime(FFrameTime(LocalFrame), DisplayRate, TickResolution);
 
 			FVector3f Trans = FVector3f::ZeroVector;
-			DoubleChannels[0]->Evaluate(LocalTime, Trans.X);
-			DoubleChannels[1]->Evaluate(LocalTime, Trans.Y);
-			DoubleChannels[2]->Evaluate(LocalTime, Trans.Z);
+			if (DoubleChannels[0].Get())
+			{
+				DoubleChannels[0].Get()->Evaluate(LocalTime, Trans.X);
+			}
+			if (DoubleChannels[1].Get())
+			{
+				DoubleChannels[1].Get()->Evaluate(LocalTime, Trans.Y);
+			}
+			if (DoubleChannels[2].Get())
+			{
+				DoubleChannels[2].Get()->Evaluate(LocalTime, Trans.Z);
+			}
 
 			FRotator Rotator;
-			DoubleChannels[3]->Evaluate(LocalTime, Rotator.Roll);
-			DoubleChannels[4]->Evaluate(LocalTime, Rotator.Pitch);
-			DoubleChannels[5]->Evaluate(LocalTime, Rotator.Yaw);
+			if (DoubleChannels[3].Get())
+			{
+				DoubleChannels[3].Get()->Evaluate(LocalTime, Rotator.Roll);
+			}
+			if (DoubleChannels[4].Get())
+			{
+				DoubleChannels[4].Get()->Evaluate(LocalTime, Rotator.Pitch);
+			}
+			if (DoubleChannels[5].Get())
+			{
+				DoubleChannels[5].Get()->Evaluate(LocalTime, Rotator.Yaw);
+			}
 
 			FVector3f Scale;
-			DoubleChannels[6]->Evaluate(LocalTime, Scale.X);
-			DoubleChannels[7]->Evaluate(LocalTime, Scale.Y);
-			DoubleChannels[8]->Evaluate(LocalTime, Scale.Z);
+			if (DoubleChannels[6].Get())
+			{
+				DoubleChannels[6].Get()->Evaluate(LocalTime, Scale.X);
+			}
+			if (DoubleChannels[7].Get())
+			{
+				DoubleChannels[7].Get()->Evaluate(LocalTime, Scale.Y);
+			}
+			if (DoubleChannels[8].Get())
+			{
+				DoubleChannels[8].Get()->Evaluate(LocalTime, Scale.Z);
+			}
 
 			FTransform RelativeTransform;
 			RelativeTransform.SetTranslation((FVector)Trans);
diff --git a/Engine/Source/Runtime/MovieSceneTracks/Private/Tracks/MovieScene3DTransformTrack.cpp b/Engine/Source/Runtime/MovieSceneTracks/Private/Tracks/MovieScene3DTransformTrack.cpp
--- a/Engine/Source/Runtime/MovieSceneTracks/Private/Tracks/MovieScene3DTransformTrack.cpp
+++ b/Engine/Source/Runtime/MovieSceneTracks/Private/Tracks/MovieScene3DTransformTrack.cpp
@@ -178,39 +178,48 @@ TArray<FTrajectoryKey> UMovieScene3DTransformTrack::GetTrajectoryData(FFrameNumb
 		UMovieScene3DTransformSection* TransformSection = Cast<UMovieScene3DTransformSection>(Section);
 		if (TransformSection && MovieSceneHelpers::IsSectionKeyable(TransformSection))
 		{
-			TArrayView<FMovieSceneDoubleChannel*>        DoubleChannels = TransformSection->GetChannelProxy().GetChannels<FMovieSceneDoubleChannel>();
-			TArrayView<const FMovieSceneChannelMetaData> MetaData      = TransformSection->GetChannelProxy().GetMetaData<FMovieSceneDoubleChannel>();
-
-			EMovieSceneTransformChannel Mask = TransformSection->GetMask().GetChannels();
-			if (EnumHasAnyFlags(Mask, EMovieSceneTransformChannel::TranslationX))
+			FMovieSceneChannelProxy& SectionChannelProxy = TransformSection->GetChannelProxy();
+			TMovieSceneChannelHandle<FMovieSceneDoubleChannel> ChannelHandles[] = {
+				SectionChannelProxy.GetChannelByName<FMovieSceneDoubleChannel>("Location.X"),
+				SectionChannelProxy.GetChannelByName<FMovieSceneDoubleChannel>("Location.Y"),
+				SectionChannelProxy.GetChannelByName<FMovieSceneDoubleChannel>("Location.Z"),
+				SectionChannelProxy.GetChannelByName<FMovieSceneDoubleChannel>("Rotation.X"),
+				SectionChannelProxy.GetChannelByName<FMovieSceneDoubleChannel>("Rotation.Y"),
+				SectionChannelProxy.GetChannelByName<FMovieSceneDoubleChannel>("Rotation.Z"),
+				SectionChannelProxy.GetChannelByName<FMovieSceneDoubleChannel>("Scale.X"),
+				SectionChannelProxy.GetChannelByName<FMovieSceneDoubleChannel>("Scale.Y"),
+				SectionChannelProxy.GetChannelByName<FMovieSceneDoubleChannel>("Scale.Z")
+			};
+
+			if (ChannelHandles[0].Get())
 			{
-				ForwardIters.Emplace(TransformSection, DoubleChannels[0], MetaData[0].Name, Time, ViewRange);
-				BackwardIters.Emplace(TransformSection, DoubleChannels[0], MetaData[0].Name, Time, ViewRange);
+				ForwardIters.Emplace(TransformSection, ChannelHandles[0].Get(), TEXT("Location.X"), Time, ViewRange);
+				BackwardIters.Emplace(TransformSection, ChannelHandles[0].Get(), TEXT("Location.X"), Time, ViewRange);
 			}
-			if (EnumHasAnyFlags(Mask, EMovieSceneTransformChannel::TranslationY))
+			if (ChannelHandles[1].Get())
 			{
-				ForwardIters.Emplace(TransformSection, DoubleChannels[1], MetaData[1].Name, Time, ViewRange);
-				BackwardIters.Emplace(TransformSection, DoubleChannels[1], MetaData[1].Name, Time, ViewRange);
+				ForwardIters.Emplace(TransformSection, ChannelHandles[1].Get(), TEXT("Location.Y"), Time, ViewRange);
+				BackwardIters.Emplace(TransformSection, ChannelHandles[1].Get(), TEXT("Location.Y"), Time, ViewRange);
 			}
-			if (EnumHasAnyFlags(Mask, EMovieSceneTransformChannel::TranslationZ))
+			if (ChannelHandles[2].Get())
 			{
-				ForwardIters.Emplace(TransformSection, DoubleChannels[2], MetaData[2].Name, Time, ViewRange);
-				BackwardIters.Emplace(TransformSection, DoubleChannels[2], MetaData[2].Name, Time, ViewRange);
+				ForwardIters.Emplace(TransformSection, ChannelHandles[2].Get(), TEXT("Location.Z"), Time, ViewRange);
+				BackwardIters.Emplace(TransformSection, ChannelHandles[2].Get(), TEXT("Location.Z"), Time, ViewRange);
 			}
-			if (EnumHasAnyFlags(Mask, EMovieSceneTransformChannel::RotationX))
+			if (ChannelHandles[3].Get())
 			{
-				ForwardIters.Emplace(TransformSection, DoubleChannels[3], MetaData[3].Name, Time, ViewRange);
-				BackwardIters.Emplace(TransformSection, DoubleChannels[3], MetaData[3].Name, Time, ViewRange);
+				ForwardIters.Emplace(TransformSection, ChannelHandles[3].Get(), TEXT("Rotation.X"), Time, ViewRange);
+				BackwardIters.Emplace(TransformSection, ChannelHandles[3].Get(), TEXT("Rotation.X"), Time, ViewRange);
 			}
-			if (EnumHasAnyFlags(Mask, EMovieSceneTransformChannel::RotationY))
+			if (ChannelHandles[4].Get())
 			{
-				ForwardIters.Emplace(TransformSection, DoubleChannels[4], MetaData[4].Name, Time, ViewRange);
-				BackwardIters.Emplace(TransformSection, DoubleChannels[4], MetaData[4].Name, Time, ViewRange);
+				ForwardIters.Emplace(TransformSection, ChannelHandles[4].Get(), TEXT("Rotation.Y"), Time, ViewRange);
+				BackwardIters.Emplace(TransformSection, ChannelHandles[4].Get(), TEXT("Rotation.Y"), Time, ViewRange);
 			}
-			if (EnumHasAnyFlags(Mask, EMovieSceneTransformChannel::RotationZ))
+			if (ChannelHandles[5].Get())
 			{
-				ForwardIters.Emplace(TransformSection, DoubleChannels[5], MetaData[5].Name, Time, ViewRange);
-				BackwardIters.Emplace(TransformSection, DoubleChannels[5], MetaData[5].Name, Time, ViewRange);
+				ForwardIters.Emplace(TransformSection, ChannelHandles[5].Get(), TEXT("Rotation.Z"), Time, ViewRange);
+				BackwardIters.Emplace(TransformSection, ChannelHandles[5].Get(), TEXT("Rotation.Z"), Time, ViewRange);
 			}
 		}
 	}
