Sequencer: Improvements to perlin noise channels' parameters UI

- The old code was specifically looking at perlin noise channels in override containers, but there could arguably be a section that uses a perlin noise channel by default. So the new code uses the channel proxy.
- Because of this, we can't rely on the details view customization for override channel containers. So we just show each channel individually in a different sub-menu. There isn't much support for multi-editing in the rest of Sequencer anyway.

diff --git a/Engine/Source/Editor/MovieSceneTools/Private/Channels/PerlinNoiseChannelInterface.cpp b/Engine/Source/Editor/MovieSceneTools/Private/Channels/PerlinNoiseChannelInterface.cpp
new file mode 100644
--- /dev/null
+++ b/Engine/Source/Editor/MovieSceneTools/Private/Channels/PerlinNoiseChannelInterface.cpp
@@ -0,0 +1,105 @@
+// Copyright Epic Games, Inc. All Rights Reserved.
+
+#include "Channels/PerlinNoiseChannelInterface.h"
+#include "Channels/MovieSceneChannelHandle.h"
+#include "Channels/MovieSceneDoublePerlinNoiseChannel.h"
+#include "Channels/MovieSceneFloatPerlinNoiseChannel.h"
+#include "IStructureDetailsView.h"
+#include "Framework/MultiBox/MultiBoxBuilder.h"
+#include "Framework/MultiBox/MultiBoxExtender.h"
+#include "Modules/ModuleManager.h"
+#include "PropertyEditorModule.h"
+
+#define LOCTEXT_NAMESPACE "PerlinNoiseChannelInterface"
+
+FPerlinNoiseChannelSectionMenuExtension::FPerlinNoiseChannelSectionMenuExtension(TArrayView<const FMovieSceneChannelHandle> InChannelHandles)
+	: ChannelHandles(InChannelHandles)
+{
+}
+
+void FPerlinNoiseChannelSectionMenuExtension::ExtendMenu(FMenuBuilder& MenuBuilder)
+{
+	TSharedRef<FPerlinNoiseChannelSectionMenuExtension> SharedThis = this->AsShared();
+
+	if (ChannelHandles.Num() > 1)
+	{
+		MenuBuilder.AddSubMenu(
+			LOCTEXT("PerlinNoiseChannelsMenu", "Perlin Noise Channels"),
+			LOCTEXT("PerlinNoiseChannelsMenuToolTip", "Edit parameters for Perlin Noise channels"),
+			FNewMenuDelegate::CreateLambda([SharedThis](FMenuBuilder& InnerMenuBuilder) { SharedThis->BuildChannelsMenu(InnerMenuBuilder); })
+		);
+	}
+	else if (ChannelHandles.Num() == 1)
+	{
+		MenuBuilder.AddSubMenu(
+			LOCTEXT("PerlinNoiseChannelsMenu", "Perlin Noise Channels"),
+			LOCTEXT("PerlinNoiseChannelsMenuToolTip", "Edit parameters for Perlin Noise channels"),
+			FNewMenuDelegate::CreateLambda([SharedThis](FMenuBuilder& InnerMenuBuilder) { SharedThis->BuildParametersMenu(InnerMenuBuilder, 0); })
+		);
+	}
+}
+
+void FPerlinNoiseChannelSectionMenuExtension::BuildChannelsMenu(FMenuBuilder& MenuBuilder)
+{
+	TSharedRef<FPerlinNoiseChannelSectionMenuExtension> SharedThis = this->AsShared();
+
+	for (int32 Index = 0; Index < ChannelHandles.Num(); ++Index)
+	{
+		const FMovieSceneChannelHandle ChannelHandle(ChannelHandles[Index]);
+		MenuBuilder.AddSubMenu(
+			FText::Format(LOCTEXT("PerlinNoiseChannelSelectMenu", "Channel {0}"), Index + 1),
+			LOCTEXT("PerlinNoiseChannelSelectMenuToolTip", "Edits parameters for this Perlin Noise channel"),
+			FNewMenuDelegate::CreateLambda([SharedThis, Index](FMenuBuilder& InnerMenuBuilder) { SharedThis->BuildParametersMenu(InnerMenuBuilder, Index); })
+		);
+	}
+}
+
+void FPerlinNoiseChannelSectionMenuExtension::BuildParametersMenu(FMenuBuilder& MenuBuilder, int32 ChannelHandleIndex)
+{
+	if (!ensure(ChannelHandles.IsValidIndex(ChannelHandleIndex)))
+	{
+		return;
+	}
+
+	FPerlinNoiseParams* PerlinNoiseParams = nullptr;
+	FMovieSceneChannelHandle ChannelHandle(ChannelHandles[ChannelHandleIndex]);
+	if (ChannelHandle.GetChannelTypeName() == FMovieSceneFloatPerlinNoiseChannel::StaticStruct()->GetFName())
+	{
+		FMovieSceneFloatPerlinNoiseChannel* FloatChannel = ChannelHandle.Cast<FMovieSceneFloatPerlinNoiseChannel>().Get();
+		PerlinNoiseParams = &FloatChannel->PerlinNoiseParams;
+	}
+	else if (ChannelHandle.GetChannelTypeName() == FMovieSceneDoublePerlinNoiseChannel::StaticStruct()->GetFName())
+	{
+		FMovieSceneDoublePerlinNoiseChannel* DoubleChannel = ChannelHandle.Cast<FMovieSceneDoublePerlinNoiseChannel>().Get();
+		PerlinNoiseParams = &DoubleChannel->PerlinNoiseParams;
+	}
+	else
+	{
+		ensureMsgf(false, TEXT("Unknown perlin noise channel type: %s"), *ChannelHandle.GetChannelTypeName().ToString());
+		return;
+	}
+
+	FPropertyEditorModule& EditModule = FModuleManager::Get().GetModuleChecked<FPropertyEditorModule>("PropertyEditor");
+
+	FDetailsViewArgs DetailsViewArgs;
+	DetailsViewArgs.bAllowSearch = false;
+	DetailsViewArgs.NameAreaSettings = FDetailsViewArgs::HideNameArea;
+	DetailsViewArgs.bHideSelectionTip = true;
+	DetailsViewArgs.bShowOptions = false;
+	DetailsViewArgs.bShowScrollBar = false;
+
+	FStructureDetailsViewArgs StructureDetailsViewArgs;
+	StructureDetailsViewArgs.bShowObjects = true;
+	StructureDetailsViewArgs.bShowAssets = true;
+	StructureDetailsViewArgs.bShowClasses = true;
+	StructureDetailsViewArgs.bShowInterfaces = true;
+
+	TSharedRef<FStructOnScope> StructData = MakeShareable(new FStructOnScope(FPerlinNoiseParams::StaticStruct(), (uint8*)PerlinNoiseParams));
+
+	TSharedRef<IStructureDetailsView> DetailsView = EditModule.CreateStructureDetailView(DetailsViewArgs, StructureDetailsViewArgs, StructData);
+
+	MenuBuilder.AddWidget(DetailsView->GetWidget().ToSharedRef(), FText(), true, false);
+}
+
+#undef LOCTEXT_NAMESPACE
+
diff --git a/Engine/Source/Editor/MovieSceneTools/Public/Channels/PerlinNoiseChannelInterface.h b/Engine/Source/Editor/MovieSceneTools/Public/Channels/PerlinNoiseChannelInterface.h
--- a/Engine/Source/Editor/MovieSceneTools/Public/Channels/PerlinNoiseChannelInterface.h
+++ b/Engine/Source/Editor/MovieSceneTools/Public/Channels/PerlinNoiseChannelInterface.h
@@ -2,20 +2,16 @@
 
 #pragma once
 
-#include "ISequencerChannelInterface.h"
-#include "Channels/IMovieSceneChannelOverrideProvider.h"
-#include "Channels/MovieSceneSectionChannelOverrideRegistry.h"
 #include "CurveKeyEditors/SequencerKeyEditor.h"
-#include "Framework/MultiBox/MultiBoxBuilder.h"
 #include "Framework/MultiBox/MultiBoxExtender.h"
-#include "Modules/ModuleManager.h"
+#include "ISequencerChannelInterface.h"
 #include "MovieSceneSection.h"
-#include "PropertyEditorModule.h"
 #include "TimeToPixel.h"
-#include "ToolMenus.h"
 #include "Widgets/Input/NumericTypeInterface.h"
 #include "Widgets/Text/STextBlock.h"
 
+class FMenuBuilder;
+
 template<typename ChannelType, typename NumericType>
 class SNumericTextBlockKeyEditor : public SCompoundWidget
 {
@@ -46,76 +42,20 @@ class SNumericTextBlockKeyEditor : public SCompoundWidget
 	}
 };
 
-template<typename ChannelContainerType>
-struct TPerlinNoiseChannelSectionMenuExtension : TSharedFromThis<TPerlinNoiseChannelSectionMenuExtension<ChannelContainerType>>
+struct FPerlinNoiseChannelSectionMenuExtension : TSharedFromThis<FPerlinNoiseChannelSectionMenuExtension>
 {
-	using ChannelType = typename ChannelContainerType::ChannelType;
+	FPerlinNoiseChannelSectionMenuExtension(TArrayView<const FMovieSceneChannelHandle> InChannelHandles);
 
-	TPerlinNoiseChannelSectionMenuExtension(TArrayView<UMovieSceneSection* const> InSections)
-		: Sections(InSections)
-	{}
-
-	void ExtendMenu(FMenuBuilder& MenuBuilder)
-	{
-		TSharedRef<TPerlinNoiseChannelSectionMenuExtension<ChannelContainerType>> SharedThis = this->AsShared();
-
-		MenuBuilder.AddSubMenu(
-			NSLOCTEXT("PerlinNoiseChannelInterface", "PerlinNoiseChannelsMenu", "Perlin Noise Channels"),
-			NSLOCTEXT("PerlinNoiseChannelInterface", "PerlinNoiseChannelsMenuToolTip", "Edit parameters for Perlin Noise channels"),
-			FNewMenuDelegate::CreateLambda([SharedThis](FMenuBuilder& MenuBuilder) { SharedThis->BuildChannelOverrideParametersMenu(MenuBuilder); })
-		);
-	}
+	void ExtendMenu(FMenuBuilder& MenuBuilder);
 
 private:
 
-	void BuildChannelOverrideParametersMenu(FMenuBuilder& MenuBuilder)
-	{
-		FPropertyEditorModule& EditModule = FModuleManager::Get().GetModuleChecked<FPropertyEditorModule>("PropertyEditor");
-
-		FDetailsViewArgs DetailsViewArgs;
-		DetailsViewArgs.bAllowSearch = false;
-		DetailsViewArgs.NameAreaSettings = FDetailsViewArgs::HideNameArea;
-		DetailsViewArgs.bHideSelectionTip = true;
-		DetailsViewArgs.bShowOptions = false;
-		DetailsViewArgs.bShowScrollBar = false;
-
-		TSharedRef<IDetailsView> DetailsView = EditModule.CreateDetailView(DetailsViewArgs);
-
-		TArray<UObject*> Objects;
-		for (TWeakObjectPtr<UMovieSceneSection> WeakSection : Sections)
-		{
-			UMovieSceneSection* Section = WeakSection.Get();
-			if (!Section)
-			{
-				continue;
-			}
-
-			IMovieSceneChannelOverrideProvider* OverrideProvider = Cast<IMovieSceneChannelOverrideProvider>(Section);
-			if (!OverrideProvider)
-			{
-				continue;
-			}
-			
-			UMovieSceneSectionChannelOverrideRegistry* OverrideRegistry = OverrideProvider->GetChannelOverrideRegistry(false);
-			if (!OverrideRegistry)
-			{
-				continue;
-			}
-	
-			TArray<ChannelContainerType*> ChannelContainers;
-			OverrideRegistry->GetChannels<ChannelContainerType>(ChannelContainers);
-			Objects.Append(ChannelContainers);
-		}
-
-		DetailsView->SetObjects(Objects, true);
-
-		MenuBuilder.AddWidget(DetailsView, FText(), true, false);
-	}
-
+	void BuildChannelsMenu(FMenuBuilder& MenuBuilder);
+	void BuildParametersMenu(FMenuBuilder& MenuBuilder, int32 ChannelHandleIndex);
 
 private:
 
-	TArray<TWeakObjectPtr<UMovieSceneSection>> Sections;
+	TArray<FMovieSceneChannelHandle> ChannelHandles;
 };
 
 template<typename ChannelContainerType>
@@ -169,8 +109,7 @@ struct TPerlinNoiseChannelInterface : ISequencerChannelInterface
 
 	virtual void ExtendSectionMenu_Raw(FMenuBuilder& MenuBuilder, TSharedPtr<FExtender> MenuExtender, TArrayView<const FMovieSceneChannelHandle> Channels, TArrayView<UMovieSceneSection* const> Sections, TWeakPtr<ISequencer> InSequencer) const override
 	{
-		using MenuExtenderType = TPerlinNoiseChannelSectionMenuExtension<ChannelContainerType>;
-		TSharedRef<MenuExtenderType> Extension = MakeShared<MenuExtenderType>(Sections);
+		TSharedRef<FPerlinNoiseChannelSectionMenuExtension> Extension = MakeShared<FPerlinNoiseChannelSectionMenuExtension>(Channels);
 
 		MenuExtender->AddMenuExtension("SequencerChannels", EExtensionHook::First, nullptr, FMenuExtensionDelegate::CreateLambda([Extension](FMenuBuilder& MenuBuilder) { Extension->ExtendMenu(MenuBuilder); }));
 	}
