Sequencer: Merge the float and double channel context menus into a single "curve channels" context menu

We can't do this easily within the framework of Sequencer right now, so the workaround is to create an extension class on the first call, and re-use the same instance on subsequent calls. The extension pointer is kept weakly across calls, and should be naturally deleted when the menu is closed, by virtual of who's keeping the ref-count up (only the menu itself and any capturing callbacks inside).

diff --git a/Engine/Source/Editor/MovieSceneTools/Private/Channels/BuiltInChannelEditors.cpp b/Engine/Source/Editor/MovieSceneTools/Private/Channels/BuiltInChannelEditors.cpp
--- a/Engine/Source/Editor/MovieSceneTools/Private/Channels/BuiltInChannelEditors.cpp
+++ b/Engine/Source/Editor/MovieSceneTools/Private/Channels/BuiltInChannelEditors.cpp
@@ -918,7 +918,7 @@ void DrawKeys(FMovieSceneEventChannel* Channel, TArrayView<const FKeyHandle> InK
 }
 
 template<typename ChannelType>
-struct TCurveChannelKeyMenuExtension : FExtender, TSharedFromThis<TCurveChannelKeyMenuExtension<ChannelType>>
+struct TCurveChannelKeyMenuExtension : TSharedFromThis<TCurveChannelKeyMenuExtension<ChannelType>>
 {
 	using ChannelValueType = typename ChannelType::ChannelValueType;
 
@@ -1089,41 +1089,63 @@ struct FDoubleChannelKeyMenuExtension : TCurveChannelKeyMenuExtension<FMovieScen
 	{}
 };
 
-template<typename ChannelType>
-struct TCurveChannelSectionMenuExtension : TSharedFromThis<TCurveChannelSectionMenuExtension<ChannelType>>
+struct FCurveChannelSectionMenuExtension : TSharedFromThis<FCurveChannelSectionMenuExtension>
 {
-	using ChannelValueType = typename ChannelType::ChannelValueType;
-	using CurveValueType = typename ChannelType::CurveValueType;
+	static TSharedRef<FCurveChannelSectionMenuExtension> GetOrCreate(TWeakPtr<ISequencer> InSequencer)
+	{
+		TSharedPtr<FCurveChannelSectionMenuExtension> CurrentExtension = WeakCurrentExtension.Pin();
+		if (!CurrentExtension)
+		{
+			CurrentExtension = MakeShared<FCurveChannelSectionMenuExtension>(InSequencer);
+			WeakCurrentExtension = CurrentExtension;
+		}
+		else
+		{
+			ensure(CurrentExtension->NumCurveChannelTypes > 0);
+			ensure(CurrentExtension->WeakSequencer == InSequencer);
+		}
+		return CurrentExtension.ToSharedRef();
+	}
 
-	TCurveChannelSectionMenuExtension(TWeakPtr<ISequencer> InSequencer, TArray<TMovieSceneChannelHandle<ChannelType>>&& InChannels, TArrayView<UMovieSceneSection* const> InSections)
+	FCurveChannelSectionMenuExtension(TWeakPtr<ISequencer> InSequencer)
 		: WeakSequencer(InSequencer)
-		, Channels(MoveTemp(InChannels))
+		, NumCurveChannelTypes(0)
+		, bMenusAdded(false)
+	{
+	}
+
+	void AddSections(TArrayView<UMovieSceneSection* const> InSections)
 	{
-		Sections.Reserve(InSections.Num());
 		for (UMovieSceneSection* Section : InSections)
 		{
 			Sections.Add(Section);
 		}
+		++NumCurveChannelTypes;
 	}
 
-	virtual ~TCurveChannelSectionMenuExtension() {}
-
-	virtual FText GetSubMenuLabel() const { return FText(); }
-	virtual FText GetSubMenuToolTip() const { return FText(); }
-
 	void ExtendMenu(FMenuBuilder& MenuBuilder)
 	{
+		--NumCurveChannelTypes;
+
+		if (bMenusAdded)
+		{
+			// Only add menus once -- not once per curve channel type (float, double, etc)
+			return;
+		}
+
+		bMenusAdded = true;
+
 		ISequencer* SequencerPtr = WeakSequencer.Pin().Get();
 		if (!SequencerPtr)
 		{
 			return;
 		}
 
-		TSharedRef<TCurveChannelSectionMenuExtension<ChannelType>> SharedThis = this->AsShared();
+		TSharedRef<FCurveChannelSectionMenuExtension> SharedThis = this->AsShared();
 
 		MenuBuilder.AddSubMenu(
-			GetSubMenuLabel(),
-			GetSubMenuToolTip(),
+			LOCTEXT("CurveChannelsMenuLabel", "Curve Channels"),
+			LOCTEXT("CurveChannelsMenuToolTip", "Edit parameters for curve channels"),
 			FNewMenuDelegate::CreateLambda([SharedThis](FMenuBuilder& SubMenuBuilder)
 			{
 				SubMenuBuilder.AddSubMenu(
@@ -1148,7 +1170,7 @@ struct TCurveChannelSectionMenuExtension : TSharedFromThis<TCurveChannelSectionM
 
 	void AddDisplayOptionsMenu(FMenuBuilder& MenuBuilder)
 	{
-		TSharedRef<TCurveChannelSectionMenuExtension<ChannelType>> SharedThis = this->AsShared();
+		TSharedRef<FCurveChannelSectionMenuExtension> SharedThis = this->AsShared();
 
 		ISequencer* Sequencer = WeakSequencer.Pin().Get();
 		if (!Sequencer)
@@ -1193,7 +1215,7 @@ struct TCurveChannelSectionMenuExtension : TSharedFromThis<TCurveChannelSectionM
 			FUIAction(
 				FExecuteAction::CreateLambda([SharedThis]{ SharedThis->ToggleShowCurve(); }),
 				FCanExecuteAction(),
-				FIsActionChecked::CreateLambda([SharedThis]{ return SharedThis->IsShowCurve(); })
+				FGetActionCheckState::CreateLambda([SharedThis]{ return SharedThis->IsShowCurve(); })
 			),
 			NAME_None,
 			EUserInterfaceActionType::ToggleButton
@@ -1217,7 +1239,7 @@ struct TCurveChannelSectionMenuExtension : TSharedFromThis<TCurveChannelSectionM
 			FSlateIcon(),
 			FUIAction(
 				FExecuteAction::CreateLambda([=] { OnKeyAreaCurveNormalized(KeyAreaName); }),
-				FCanExecuteAction(FCanExecuteAction::CreateLambda([SharedThis]{ return SharedThis->IsShowCurve(); })),
+				FCanExecuteAction(FCanExecuteAction::CreateLambda([SharedThis]{ return SharedThis->IsAnyShowCurve(); })),
 				FIsActionChecked::CreateLambda([=] { return GetKeyAreaCurveNormalized(KeyAreaName); })
 			),
 			NAME_None,
@@ -1235,7 +1257,7 @@ struct TCurveChannelSectionMenuExtension : TSharedFromThis<TCurveChannelSectionM
 				[
 					SNew(SBox)
 					.WidthOverride(50.f)
-					.IsEnabled_Lambda([=]() { return SharedThis->IsShowCurve() && Settings->HasKeyAreaCurveExtents(KeyAreaName); })
+					.IsEnabled_Lambda([=]() { return SharedThis->IsAnyShowCurve() && Settings->HasKeyAreaCurveExtents(KeyAreaName); })
 					[
 						SNew(SSpinBox<double>)
 						.Style(&FAppStyle::GetWidgetStyle<FSpinBoxStyle>("Sequencer.HyperlinkSpinBox"))
@@ -1249,7 +1271,7 @@ struct TCurveChannelSectionMenuExtension : TSharedFromThis<TCurveChannelSectionM
 				[
 					SNew(SBox)
 					.WidthOverride(50.f)
-					.IsEnabled_Lambda([=]() { return SharedThis->IsShowCurve() && Settings->HasKeyAreaCurveExtents(KeyAreaName); })
+					.IsEnabled_Lambda([=]() { return SharedThis->IsAnyShowCurve() && Settings->HasKeyAreaCurveExtents(KeyAreaName); })
 					[
 						SNew(SSpinBox<double>)
 						.Style(&FAppStyle::GetWidgetStyle<FSpinBoxStyle>("Sequencer.HyperlinkSpinBox"))
@@ -1286,10 +1308,9 @@ struct TCurveChannelSectionMenuExtension : TSharedFromThis<TCurveChannelSectionM
 		);
 	}
 
-
 	void AddExtrapolationMenu(FMenuBuilder& MenuBuilder, bool bPreInfinity)
 	{
-		TSharedRef<TCurveChannelSectionMenuExtension<ChannelType>> SharedThis = this->AsShared();
+		TSharedRef<FCurveChannelSectionMenuExtension> SharedThis = this->AsShared();
 
 		MenuBuilder.AddMenuEntry(
 			LOCTEXT("SetExtrapCycle", "Cycle"),
@@ -1406,7 +1427,6 @@ struct TCurveChannelSectionMenuExtension : TSharedFromThis<TCurveChannelSectionM
 		}
 	}
 
-
 	bool IsExtrapolationModeSelected(ERichCurveExtrapolation ExtrapMode, bool bPreInfinity) const
 	{
 		for (TWeakObjectPtr<UMovieSceneSection> WeakSection : Sections)
@@ -1438,12 +1458,15 @@ struct TCurveChannelSectionMenuExtension : TSharedFromThis<TCurveChannelSectionM
 
 	void ToggleShowCurve()
 	{
-		bool bShowCurve = !IsShowCurve();
+		const ECheckBoxState CurrentState = IsShowCurve();
+		const bool bShowCurve = (CurrentState != ECheckBoxState::Checked); // If unchecked or mixed, check it
+
 		FScopedTransaction Transaction(LOCTEXT("ToggleShowCurve_Transaction", "Toggle Show Curve"));
 
 		bool bAnythingChanged = false;
 
 		// Modify all sections
+
 		for (TWeakObjectPtr<UMovieSceneSection> WeakSection : Sections)
 		{
 			if (UMovieSceneSection* Section = WeakSection.Get())
@@ -1453,14 +1476,27 @@ struct TCurveChannelSectionMenuExtension : TSharedFromThis<TCurveChannelSectionM
 		}
 
 		// Apply to all channels
-		for (const TMovieSceneChannelHandle<ChannelType>& Handle : Channels)
+		for (TWeakObjectPtr<UMovieSceneSection> WeakSection : Sections)
 		{
-			ChannelType* Channel = Handle.Get();
-
-			if (Channel)
+			if (UMovieSceneSection* Section = WeakSection.Get())
 			{
-				Channel->SetShowCurve(bShowCurve);
-				bAnythingChanged = true;
+				FMovieSceneChannelProxy& ChannelProxy = Section->GetChannelProxy();
+				for (FMovieSceneFloatChannel* Channel : ChannelProxy.GetChannels<FMovieSceneFloatChannel>())
+				{
+					if (Channel)
+					{
+						Channel->SetShowCurve(bShowCurve);
+						bAnythingChanged = true;
+					}
+				}
+				for (FMovieSceneDoubleChannel* Channel : ChannelProxy.GetChannels<FMovieSceneDoubleChannel>())
+				{
+					if (Channel)
+					{
+						Channel->SetShowCurve(bShowCurve);
+						bAnythingChanged = true;
+					}
+				}
 			}
 		}
 
@@ -1470,61 +1506,96 @@ struct TCurveChannelSectionMenuExtension : TSharedFromThis<TCurveChannelSectionM
 		}
 	}
 
-	bool IsShowCurve() const
+	ECheckBoxState IsShowCurve() const
 	{
-		for (const TMovieSceneChannelHandle<ChannelType>& Handle : Channels)
+		int32 NumShowedAndHidden[2] = { 0, 0 };
+		for (TWeakObjectPtr<UMovieSceneSection> WeakSection : Sections)
 		{
-			ChannelType* Channel = Handle.Get();
-
-			if (Channel && !Channel->GetShowCurve())
+			if (UMovieSceneSection* Section = WeakSection.Get())
 			{
-				return false;
+				FMovieSceneChannelProxy& ChannelProxy = Section->GetChannelProxy();
+				for (FMovieSceneFloatChannel* Channel : ChannelProxy.GetChannels<FMovieSceneFloatChannel>())
+				{
+					if (Channel)
+					{
+						NumShowedAndHidden[Channel->GetShowCurve() ? 0 : 1]++;
+					}
+				}
+				for (FMovieSceneDoubleChannel* Channel : ChannelProxy.GetChannels<FMovieSceneDoubleChannel>())
+				{
+					if (Channel)
+					{
+						NumShowedAndHidden[Channel->GetShowCurve() ? 0 : 1]++;
+					}
+				}
 			}
 		}
 
-		return true;
+		if (NumShowedAndHidden[0] == 0 && NumShowedAndHidden[1] > 0)  // No curve showed, some hidden
+		{
+			return ECheckBoxState::Unchecked;
+		}
+		else if (NumShowedAndHidden[0] > 0 && NumShowedAndHidden[1] == 0) // Some curves showed, none hidden
+		{
+			return ECheckBoxState::Checked;
+		}
+		return ECheckBoxState::Undetermined;  // Mixed states, or no curves
+	}
+
+	bool IsAnyShowCurve() const
+	{
+		for (TWeakObjectPtr<UMovieSceneSection> WeakSection : Sections)
+		{
+			if (UMovieSceneSection* Section = WeakSection.Get())
+			{
+				FMovieSceneChannelProxy& ChannelProxy = Section->GetChannelProxy();
+				for (FMovieSceneFloatChannel* Channel : ChannelProxy.GetChannels<FMovieSceneFloatChannel>())
+				{
+					if (Channel && Channel->GetShowCurve())
+					{
+						return true;
+					}
+				}
+				for (FMovieSceneDoubleChannel* Channel : ChannelProxy.GetChannels<FMovieSceneDoubleChannel>())
+				{
+					if (Channel && Channel->GetShowCurve())
+					{
+						return true;
+					}
+				}
+			}
+		}
+		return false;
 	}
 
 private:
 
 	/** Hidden AsShared() methods to prevent CreateSP delegate use since this extender disappears with its menu. */
-	using TSharedFromThis<TCurveChannelSectionMenuExtension<ChannelType>>::AsShared;
-
-	TWeakPtr<ISequencer> WeakSequencer;
-	TArray<TMovieSceneChannelHandle<ChannelType>> Channels;
-	TArray<TWeakObjectPtr<UMovieSceneSection>> Sections;
-};
+	using TSharedFromThis<FCurveChannelSectionMenuExtension>::AsShared;
 
-struct FFloatChannelSectionMenuExtension : TCurveChannelSectionMenuExtension<FMovieSceneFloatChannel>
-{
-	FFloatChannelSectionMenuExtension(TWeakPtr<ISequencer> InSequencer, TArray<TMovieSceneChannelHandle<FMovieSceneFloatChannel>>&& InChannels, TArrayView<UMovieSceneSection* const> InSections)
-		: TCurveChannelSectionMenuExtension<FMovieSceneFloatChannel>(InSequencer, MoveTemp(InChannels), InSections)
-	{}
+	/** Held weekly so that only the context menu owns the instance, and it gets naturally deleted when the menu closes */
+	static TWeakPtr<FCurveChannelSectionMenuExtension> WeakCurrentExtension;
 
-	virtual FText GetSubMenuLabel() const override { return LOCTEXT("FloatChannelsMenuLabel", "Float Channels"); }
-	virtual FText GetSubMenuToolTip() const override { return LOCTEXT("FloatChannelsMenuToolTip", "Edit parameters for float channels"); }
+	TWeakPtr<ISequencer> WeakSequencer;
+	TSet<TWeakObjectPtr<UMovieSceneSection>> Sections;
+	int32 NumCurveChannelTypes;
+	bool bMenusAdded;
 };
 
-struct FDoubleChannelSectionMenuExtension : TCurveChannelSectionMenuExtension<FMovieSceneDoubleChannel>
-{
-	FDoubleChannelSectionMenuExtension(TWeakPtr<ISequencer> InSequencer, TArray<TMovieSceneChannelHandle<FMovieSceneDoubleChannel>>&& InChannels, TArrayView<UMovieSceneSection* const> InSections)
-		: TCurveChannelSectionMenuExtension<FMovieSceneDoubleChannel>(InSequencer, MoveTemp(InChannels), InSections)
-	{}
-
-	virtual FText GetSubMenuLabel() const override { return LOCTEXT("DoubleChannelsMenuLabel", "Double Channels"); }
-	virtual FText GetSubMenuToolTip() const override { return LOCTEXT("DoubleChannelsMenuToolTip", "Edit parameters for double channels"); }
-};
+TWeakPtr<FCurveChannelSectionMenuExtension> FCurveChannelSectionMenuExtension::WeakCurrentExtension;
 
 void ExtendSectionMenu(FMenuBuilder& OuterMenuBuilder, TSharedPtr<FExtender> MenuExtender, TArray<TMovieSceneChannelHandle<FMovieSceneFloatChannel>>&& Channels, TArrayView<UMovieSceneSection* const> Sections, TWeakPtr<ISequencer> InSequencer)
 {
-	TSharedRef<FFloatChannelSectionMenuExtension> Extension = MakeShared<FFloatChannelSectionMenuExtension>(InSequencer, MoveTemp(Channels), Sections);
+	TSharedRef<FCurveChannelSectionMenuExtension> Extension = FCurveChannelSectionMenuExtension::GetOrCreate(InSequencer);
+	Extension->AddSections(Sections);
 
 	MenuExtender->AddMenuExtension("SequencerChannels", EExtensionHook::First, nullptr, FMenuExtensionDelegate::CreateLambda([Extension](FMenuBuilder& MenuBuilder) { Extension->ExtendMenu(MenuBuilder); }));
 }
 
 void ExtendSectionMenu(FMenuBuilder& OuterMenuBuilder, TSharedPtr<FExtender> MenuExtender, TArray<TMovieSceneChannelHandle<FMovieSceneDoubleChannel>>&& Channels, TArrayView<UMovieSceneSection* const> Sections, TWeakPtr<ISequencer> InSequencer)
 {
-	TSharedRef<FDoubleChannelSectionMenuExtension> Extension = MakeShared<FDoubleChannelSectionMenuExtension>(InSequencer, MoveTemp(Channels), Sections);
+	TSharedRef<FCurveChannelSectionMenuExtension> Extension = FCurveChannelSectionMenuExtension::GetOrCreate(InSequencer);
+	Extension->AddSections(Sections);
 
 	MenuExtender->AddMenuExtension("SequencerChannels", EExtensionHook::First, nullptr, FMenuExtensionDelegate::CreateLambda([Extension](FMenuBuilder& MenuBuilder) { Extension->ExtendMenu(MenuBuilder); }));
 }
