Fixed LOD issues when baking an animation in Sequencer

Ensure that we update the predicted LOD level before attempting to iterate the animation recorder

diff --git a/Engine/Source/Editor/MovieSceneTools/Private/MovieSceneToolHelpers.cpp b/Engine/Source/Editor/MovieSceneTools/Private/MovieSceneToolHelpers.cpp
--- a/Engine/Source/Editor/MovieSceneTools/Private/MovieSceneToolHelpers.cpp
+++ b/Engine/Source/Editor/MovieSceneTools/Private/MovieSceneToolHelpers.cpp
@@ -3985,9 +3985,10 @@ bool MovieSceneToolHelpers::ExportToAnimSequence(UAnimSequence* AnimSequence, UA
 		});
 
 	
-	FStartAnimationCB StartCallback = FStartAnimationCB::CreateLambda([&AnimationRecorder]
+	FStartAnimationCB StartCallback = FStartAnimationCB::CreateLambda([SkelMeshComp, &AnimationRecorder]
 	{
 		AnimationRecorder.BeginRecording();
+		SkelMeshComp->UpdateLODStatus();
 	});
 
 	FTickAnimationCB TickCallback = FTickAnimationCB::CreateLambda([&AnimationRecorder](float DeltaTime)
