Sequencer: Display names of perlin noise channels in the section context menu

diff --git a/Engine/Source/Editor/MovieSceneTools/Private/Channels/PerlinNoiseChannelInterface.cpp b/Engine/Source/Editor/MovieSceneTools/Private/Channels/PerlinNoiseChannelInterface.cpp
--- a/Engine/Source/Editor/MovieSceneTools/Private/Channels/PerlinNoiseChannelInterface.cpp
+++ b/Engine/Source/Editor/MovieSceneTools/Private/Channels/PerlinNoiseChannelInterface.cpp
@@ -12,8 +12,9 @@
 
 #define LOCTEXT_NAMESPACE "PerlinNoiseChannelInterface"
 
-FPerlinNoiseChannelSectionMenuExtension::FPerlinNoiseChannelSectionMenuExtension(TArrayView<const FMovieSceneChannelHandle> InChannelHandles)
+FPerlinNoiseChannelSectionMenuExtension::FPerlinNoiseChannelSectionMenuExtension(TArrayView<const FMovieSceneChannelHandle> InChannelHandles, TArrayView<UMovieSceneSection* const> InSections)
 	: ChannelHandles(InChannelHandles)
+	, Sections(InSections)
 {
 }
 
@@ -41,16 +42,43 @@ void FPerlinNoiseChannelSectionMenuExtension::ExtendMenu(FMenuBuilder& MenuBuild
 
 void FPerlinNoiseChannelSectionMenuExtension::BuildChannelsMenu(FMenuBuilder& MenuBuilder)
 {
+	TArray<FMovieSceneChannelProxy*> ChannelProxies;
+	for (UMovieSceneSection* Section : Sections)
+	{
+		ChannelProxies.Add(&Section->GetChannelProxy());
+	}
+
+	TArray<int32> ChannelHandleSectionIndexes;
+	for (const FMovieSceneChannelHandle& ChannelHandle : ChannelHandles)
+	{
+		int32 SectionIndex = ChannelProxies.Find(ChannelHandle.GetChannelProxy());
+		ChannelHandleSectionIndexes.Add(SectionIndex);
+	}
+
+	const bool bMultipleSections = Sections.Num() > 1;
 	TSharedRef<FPerlinNoiseChannelSectionMenuExtension> SharedThis = this->AsShared();
 
 	for (int32 Index = 0; Index < ChannelHandles.Num(); ++Index)
 	{
+		const int32 SectionIndex(ChannelHandleSectionIndexes[Index]);
 		const FMovieSceneChannelHandle ChannelHandle(ChannelHandles[Index]);
-		MenuBuilder.AddSubMenu(
-			FText::Format(LOCTEXT("PerlinNoiseChannelSelectMenu", "Channel {0}"), Index + 1),
-			LOCTEXT("PerlinNoiseChannelSelectMenuToolTip", "Edits parameters for this Perlin Noise channel"),
-			FNewMenuDelegate::CreateLambda([SharedThis, Index](FMenuBuilder& InnerMenuBuilder) { SharedThis->BuildParametersMenu(InnerMenuBuilder, Index); })
-		);
+
+		if (bMultipleSections)
+		{
+			MenuBuilder.AddSubMenu(
+				FText::Format(LOCTEXT("PerlinNoiseChannelAndSectionSelectMenu", "Section{0}.{1}"), SectionIndex + 1, FText::FromName(ChannelHandle.GetMetaData()->Name)),
+				LOCTEXT("PerlinNoiseChannelAndSectionSelectMenuToolTip", "Edit parameters for this Perlin Noise channel"),
+				FNewMenuDelegate::CreateLambda([SharedThis, Index](FMenuBuilder& InnerMenuBuilder) { SharedThis->BuildParametersMenu(InnerMenuBuilder, Index); })
+			);
+		}
+		else
+		{
+			MenuBuilder.AddSubMenu(
+				FText::FromName(ChannelHandle.GetMetaData()->Name),
+				LOCTEXT("PerlinNoiseChannelSelectMenuToolTip", "Edit parameters for this Perlin Noise channel"),
+				FNewMenuDelegate::CreateLambda([SharedThis, Index](FMenuBuilder& InnerMenuBuilder) { SharedThis->BuildParametersMenu(InnerMenuBuilder, Index); })
+			);
+		}
 	}
 }
 
diff --git a/Engine/Source/Editor/MovieSceneTools/Public/Channels/PerlinNoiseChannelInterface.h b/Engine/Source/Editor/MovieSceneTools/Public/Channels/PerlinNoiseChannelInterface.h
--- a/Engine/Source/Editor/MovieSceneTools/Public/Channels/PerlinNoiseChannelInterface.h
+++ b/Engine/Source/Editor/MovieSceneTools/Public/Channels/PerlinNoiseChannelInterface.h
@@ -44,7 +44,7 @@ class SNumericTextBlockKeyEditor : public SCompoundWidget
 
 struct FPerlinNoiseChannelSectionMenuExtension : TSharedFromThis<FPerlinNoiseChannelSectionMenuExtension>
 {
-	FPerlinNoiseChannelSectionMenuExtension(TArrayView<const FMovieSceneChannelHandle> InChannelHandles);
+	FPerlinNoiseChannelSectionMenuExtension(TArrayView<const FMovieSceneChannelHandle> InChannelHandles, TArrayView<UMovieSceneSection* const> InSections);
 
 	void ExtendMenu(FMenuBuilder& MenuBuilder);
 
@@ -56,6 +56,7 @@ private:
 private:
 
 	TArray<FMovieSceneChannelHandle> ChannelHandles;
+	TArray<UMovieSceneSection*> Sections;
 };
 
 template<typename ChannelContainerType>
@@ -109,7 +110,7 @@ struct TPerlinNoiseChannelInterface : ISequencerChannelInterface
 
 	virtual void ExtendSectionMenu_Raw(FMenuBuilder& MenuBuilder, TSharedPtr<FExtender> MenuExtender, TArrayView<const FMovieSceneChannelHandle> Channels, TArrayView<UMovieSceneSection* const> Sections, TWeakPtr<ISequencer> InSequencer) const override
 	{
-		TSharedRef<FPerlinNoiseChannelSectionMenuExtension> Extension = MakeShared<FPerlinNoiseChannelSectionMenuExtension>(Channels);
+		TSharedRef<FPerlinNoiseChannelSectionMenuExtension> Extension = MakeShared<FPerlinNoiseChannelSectionMenuExtension>(Channels, Sections);
 
 		MenuExtender->AddMenuExtension("SequencerChannels", EExtensionHook::First, nullptr, FMenuExtensionDelegate::CreateLambda([Extension](FMenuBuilder& MenuBuilder) { Extension->ExtendMenu(MenuBuilder); }));
 	}
