Sequencer: Allow snapping in reverse by calculating the frames in between in reverse if EndFrame > StartFrame

diff --git a/Engine/Source/Editor/MovieSceneTools/Private/MovieSceneToolHelpers.cpp b/Engine/Source/Editor/MovieSceneTools/Private/MovieSceneToolHelpers.cpp
--- a/Engine/Source/Editor/MovieSceneTools/Private/MovieSceneToolHelpers.cpp
+++ b/Engine/Source/Editor/MovieSceneTools/Private/MovieSceneToolHelpers.cpp
@@ -4596,10 +4596,12 @@ void MovieSceneToolHelpers::CalculateFramesBetween(
 	FFrameNumber EndFrame,
 	TArray<FFrameNumber>& OutFrames)
 {
-	if (StartFrame > EndFrame)
+	const bool bReverse = StartFrame > EndFrame;
+	if (bReverse)
 	{
 		Swap(StartFrame, EndFrame);
 	}
+
 	//if the end frame is the last frame, we move it back one tick otherwise it's possible when evaluating it won't evaluate.
 	if (EndFrame == MovieScene->GetPlaybackRange().GetUpperBoundValue())
 	{
@@ -4612,10 +4614,22 @@ void MovieSceneToolHelpers::CalculateFramesBetween(
 	const FFrameNumber EndTimeInDisplay = FFrameRate::TransformTime(FFrameTime(EndFrame), TickResolution, DisplayResolution).CeilToFrame();
 
 	OutFrames.Reserve(EndTimeInDisplay.Value - StartTimeInDisplay.Value + 1);
-	for (FFrameNumber DisplayFrameNumber = StartTimeInDisplay; DisplayFrameNumber <= EndTimeInDisplay; ++DisplayFrameNumber)
+
+	if (bReverse)
 	{
-		FFrameNumber TickFrameNumber = FFrameRate::TransformTime(FFrameTime(DisplayFrameNumber), DisplayResolution, TickResolution).FrameNumber;
-		OutFrames.Add(TickFrameNumber);
+		for (FFrameNumber DisplayFrameNumber = EndTimeInDisplay; DisplayFrameNumber >= StartTimeInDisplay; --DisplayFrameNumber)
+		{
+			FFrameNumber TickFrameNumber = FFrameRate::TransformTime(FFrameTime(DisplayFrameNumber), DisplayResolution, TickResolution).FrameNumber;
+			OutFrames.Add(TickFrameNumber);
+		}
+	}
+	else
+	{
+		for (FFrameNumber DisplayFrameNumber = StartTimeInDisplay; DisplayFrameNumber <= EndTimeInDisplay; ++DisplayFrameNumber)
+		{
+			FFrameNumber TickFrameNumber = FFrameRate::TransformTime(FFrameTime(DisplayFrameNumber), DisplayResolution, TickResolution).FrameNumber;
+			OutFrames.Add(TickFrameNumber);
+		}
 	}
 }
 
