disable constraint keyframing thru the key navigation system

diff --git a/Engine/Source/Editor/MovieSceneTools/Private/Channels/ConstraintChannelEditor.cpp b/Engine/Source/Editor/MovieSceneTools/Private/Channels/ConstraintChannelEditor.cpp
--- a/Engine/Source/Editor/MovieSceneTools/Private/Channels/ConstraintChannelEditor.cpp
+++ b/Engine/Source/Editor/MovieSceneTools/Private/Channels/ConstraintChannelEditor.cpp
@@ -334,6 +334,18 @@ FKeyHandle AddOrUpdateKey(
 	return FKeyHandle::Invalid();
 }
 
+FKeyHandle AddOrUpdateKey(
+	FMovieSceneConstraintChannel* Channel,
+	UMovieSceneSection* SectionToKey,
+	const TMovieSceneExternalValue<bool>& EditorData,
+	FFrameNumber InTime,
+	ISequencer& Sequencer,
+	const FGuid& InObjectBindingID,
+	FTrackInstancePropertyBindings* PropertyBindings)
+{
+	return FKeyHandle::Invalid();
+}
+
 bool CanCreateKeyEditor(const FMovieSceneConstraintChannel* InChannel)
 {
 	return false;
diff --git a/Engine/Source/Editor/MovieSceneTools/Private/Channels/ConstraintChannelEditor.h b/Engine/Source/Editor/MovieSceneTools/Private/Channels/ConstraintChannelEditor.h
--- a/Engine/Source/Editor/MovieSceneTools/Private/Channels/ConstraintChannelEditor.h
+++ b/Engine/Source/Editor/MovieSceneTools/Private/Channels/ConstraintChannelEditor.h
@@ -51,6 +51,15 @@ FKeyHandle AddOrUpdateKey(
 	const FGuid& ObjectBindingID,
 	FTrackInstancePropertyBindings* PropertyBindings);
 
+FKeyHandle AddOrUpdateKey(
+	FMovieSceneConstraintChannel* Channel,
+	UMovieSceneSection* SectionToKey,
+	const TMovieSceneExternalValue<bool>& EditorData,
+	FFrameNumber InTime,
+	ISequencer& Sequencer,
+	const FGuid& InObjectBindingID,
+	FTrackInstancePropertyBindings* PropertyBindings);
+
 /** Key editor overrides */
 bool CanCreateKeyEditor(const FMovieSceneConstraintChannel* InChannel);
 
