Sequencer: In AddKeysToSection, call ApplyDefault before keys are added.

The default value is a value for the channel when there are no keyframes. For example, if you add keys and then delete them all, the default value is the value of the channel. In the implementation of ApplyDefault, all the setters check that the default value is only set when there are NO keyframes. So, ApplyDefault needs to be called here in AddKeysToSection BEFORE any keys are added.

diff --git a/Engine/Source/Editor/MovieSceneTools/Public/KeyframeTrackEditor.h b/Engine/Source/Editor/MovieSceneTools/Public/KeyframeTrackEditor.h
--- a/Engine/Source/Editor/MovieSceneTools/Public/KeyframeTrackEditor.h
+++ b/Engine/Source/Editor/MovieSceneTools/Public/KeyframeTrackEditor.h
@@ -647,7 +647,19 @@ protected:
 		FMovieSceneChannelProxy& Proxy = Section->GetChannelProxy();
 			
 		const bool bSetDefaults = GetSequencer()->GetAutoSetTrackDefaults();
-		
+
+		// The default value is a value for the channel when there are no keyframes. For example, if you add keys and 
+		// then delete them all, the default value is the value of the channel. In the implementation of ApplyDefault, 
+		// all the setters check that the default value is only set when there are NO keyframes. So, ApplyDefault needs 
+		// to be called here in AddKeysToSection BEFORE any keys are added.
+		if (bSetDefaults)
+		{
+			for (const FMovieSceneChannelValueSetter& GeneratedKey : Keys)
+			{
+				GeneratedKey->ApplyDefault(Section, Proxy);
+			}
+		}
+
 		if ( KeyMode != ESequencerKeyMode::AutoKey || AutoChangeMode == EAutoChangeMode::AutoKey || AutoChangeMode == EAutoChangeMode::All)
 		{
 			EMovieSceneKeyInterpolation InterpolationMode = GetSequencer()->GetKeyInterpolation();
@@ -671,14 +683,6 @@ protected:
 			}
 		}
 			
-		if (bSetDefaults)
-		{
-			for (const FMovieSceneChannelValueSetter& GeneratedKey : Keys)
-			{
-				GeneratedKey->ApplyDefault(Section, Proxy);
-			}
-		}
-
 		return KeyPropertyResult;
 	}
 
