Sequencer: Fix delete parameter channels.

FChannelGroupOutlinerModel::Delete should be calling RequestDeleteKeyArea
Note, we can't yet handle deleting individual channels of a vector/color parameter because masks and intial values aren't implemented for parameter sections.

diff --git a/Engine/Source/Editor/MovieSceneTools/Private/Sections/ParameterSection.cpp b/Engine/Source/Editor/MovieSceneTools/Private/Sections/ParameterSection.cpp
--- a/Engine/Source/Editor/MovieSceneTools/Private/Sections/ParameterSection.cpp
+++ b/Engine/Source/Editor/MovieSceneTools/Private/Sections/ParameterSection.cpp
@@ -248,13 +248,17 @@ TSharedPtr<UE::Sequencer::FCategoryModel> FParameterSection::ConstructCategoryMo
 
 bool FParameterSection::RequestDeleteCategory( const TArray<FName>& CategoryNamePath )
 {
-	const FScopedTransaction Transaction( LOCTEXT( "DeleteVectorOrColorParameter", "Delete vector or color parameter" ) );
+	const FScopedTransaction Transaction( LOCTEXT( "DeleteParameter", "Delete parameter" ) );
 	UMovieSceneParameterSection* ParameterSection = Cast<UMovieSceneParameterSection>( WeakSection.Get() );
 	if( ParameterSection->Modify() )
 	{
-		bool bVectorParameterDeleted = ParameterSection->RemoveVectorParameter( CategoryNamePath[0] );
-		bool bColorParameterDeleted = ParameterSection->RemoveColorParameter( CategoryNamePath[0] );
-		return bVectorParameterDeleted || bColorParameterDeleted;
+		const bool bVectorParameterDeleted = ParameterSection->RemoveVectorParameter(CategoryNamePath[0]);
+		const bool bColorParameterDeleted = ParameterSection->RemoveColorParameter(CategoryNamePath[0]);
+		const bool bScalarParameterDeleted = ParameterSection->RemoveScalarParameter(CategoryNamePath[0]);
+		const bool bBoolParameterDeleted = ParameterSection->RemoveBoolParameter(CategoryNamePath[0]);
+		const bool bVector2DParameterDeleted = ParameterSection->RemoveVector2DParameter(CategoryNamePath[0]);
+		const bool bTransformParameterDeleted = ParameterSection->RemoveTransformParameter(CategoryNamePath[0]);
+		return bVectorParameterDeleted || bColorParameterDeleted || bScalarParameterDeleted || bBoolParameterDeleted || bVector2DParameterDeleted || bTransformParameterDeleted;
 	}
 	return false;
 }
@@ -263,13 +267,19 @@ bool FParameterSection::RequestDeleteCategory( const TArray<FName>& CategoryName
 bool FParameterSection::RequestDeleteKeyArea( const TArray<FName>& KeyAreaNamePath )
 {
 	// Only handle paths with a single name, in all other cases the user is deleting a component of a vector parameter.
-	if ( KeyAreaNamePath.Num() == 1)
+	if (KeyAreaNamePath.Num() == 1)
 	{
-		const FScopedTransaction Transaction( LOCTEXT( "DeleteScalarParameter", "Delete scalar parameter" ) );
+		const FScopedTransaction Transaction(LOCTEXT("DeleteParameter", "Delete parameter"));
 		UMovieSceneParameterSection* ParameterSection = Cast<UMovieSceneParameterSection>( WeakSection.Get() );
 		if (ParameterSection->TryModify())
 		{
-			return ParameterSection->RemoveScalarParameter( KeyAreaNamePath[0] );
+			const bool bVectorParameterDeleted = ParameterSection->RemoveVectorParameter(KeyAreaNamePath[0]);
+			const bool bColorParameterDeleted = ParameterSection->RemoveColorParameter(KeyAreaNamePath[0]);
+			const bool bScalarParameterDeleted = ParameterSection->RemoveScalarParameter(KeyAreaNamePath[0]);
+			const bool bBoolParameterDeleted = ParameterSection->RemoveBoolParameter(KeyAreaNamePath[0]);
+			const bool bVector2DParameterDeleted = ParameterSection->RemoveVector2DParameter(KeyAreaNamePath[0]);
+			const bool bTransformParameterDeleted = ParameterSection->RemoveTransformParameter(KeyAreaNamePath[0]);
+			return bVectorParameterDeleted || bColorParameterDeleted || bScalarParameterDeleted || bBoolParameterDeleted || bVector2DParameterDeleted || bTransformParameterDeleted;
 		}
 	}
 	return false;
diff --git a/Engine/Source/Editor/Sequencer/Private/MVVM/ViewModels/ChannelModel.cpp b/Engine/Source/Editor/Sequencer/Private/MVVM/ViewModels/ChannelModel.cpp
--- a/Engine/Source/Editor/Sequencer/Private/MVVM/ViewModels/ChannelModel.cpp
+++ b/Engine/Source/Editor/Sequencer/Private/MVVM/ViewModels/ChannelModel.cpp
@@ -754,7 +754,7 @@ void FChannelGroupOutlinerModel::Delete()
 	{
 		if (TViewModelPtr<FSectionModel> Section = Channel->FindAncestorOfType<FSectionModel>())
 		{
-			Section->GetSectionInterface()->RequestDeleteCategory(PathFromTrack);
+			Section->GetSectionInterface()->RequestDeleteKeyArea(PathFromTrack);
 		}
 	}
 }
