Sequencer: Constraints: Added function to clear proxy on constrained sections due to the constraint channels not being set up correctly. This is then by the curve model to lazily get the handle.

diff --git a/Engine/Source/Editor/MovieSceneTools/Private/Channels/ConstraintChannelCurveModel.cpp b/Engine/Source/Editor/MovieSceneTools/Private/Channels/ConstraintChannelCurveModel.cpp
--- a/Engine/Source/Editor/MovieSceneTools/Private/Channels/ConstraintChannelCurveModel.cpp
+++ b/Engine/Source/Editor/MovieSceneTools/Private/Channels/ConstraintChannelCurveModel.cpp
@@ -30,6 +30,7 @@
 #include "Styling/ISlateStyle.h"
 #include "Templates/UnrealTemplate.h"
 #include "UObject/WeakObjectPtr.h"
+#include "Sections/MovieSceneConstrainedSection.h"
 
 class FCurveEditor;
 class UObject;
@@ -45,7 +46,7 @@ FConstraintChannelCurveModel::FConstraintChannelCurveModel(TMovieSceneChannelHan
 	ChannelHandle = NewChannelProxy->MakeHandle<FMovieSceneConstraintChannel>(InChannel.GetChannelIndex());
 	if (FMovieSceneChannelProxy* ChannelProxy = InChannel.GetChannelProxy())
 	{
-		OnDestroyHandle = NewChannelProxy->OnDestroy.AddRaw(this, &FConstraintChannelCurveModel::FixupCurve);
+		OnDestroyHandle = NewChannelProxy->OnDestroy.AddRaw(this, &FConstraintChannelCurveModel::NeedNewProxy);
 	}
 
 	WeakSection = OwningSection;
@@ -62,19 +63,32 @@ FConstraintChannelCurveModel::~FConstraintChannelCurveModel()
 		ChannelProxy->OnDestroy.Remove(OnDestroyHandle);
 	}
 }
+void FConstraintChannelCurveModel::NeedNewProxy() const
+{	
+	bNeedNewProxy = true; 
+}
 
-void FConstraintChannelCurveModel::FixupCurve()
+void FConstraintChannelCurveModel::FixupCurve() const
 {
 	if (UMovieSceneSection* Section = WeakSection.Get())
 	{
-		FMovieSceneChannelProxy* NewChannelProxy = &Section->GetChannelProxy();
-		ChannelHandle = NewChannelProxy->MakeHandle<FMovieSceneConstraintChannel>(ChannelHandle.GetChannelIndex());
-		OnDestroyHandle = NewChannelProxy->OnDestroy.AddRaw(this, &FConstraintChannelCurveModel::FixupCurve);
+		if (bNeedNewProxy)
+		{
+			if (IMovieSceneConstrainedSection* ConstrainedSection = Cast<IMovieSceneConstrainedSection>(Section))
+			{
+				ConstrainedSection->OnConstraintsChanged();
+			}
+			FMovieSceneChannelProxy* NewChannelProxy = &Section->GetChannelProxy();
+			ChannelHandle = NewChannelProxy->MakeHandle<FMovieSceneConstraintChannel>(ChannelHandle.GetChannelIndex());
+			OnDestroyHandle = NewChannelProxy->OnDestroy.AddRaw(this, &FConstraintChannelCurveModel::NeedNewProxy);
+			bNeedNewProxy = false;
+		}
 	}
 }
+
 TArray<FKeyBarCurveModel::FBarRange> FConstraintChannelCurveModel::FindRanges()
 {
-	FMovieSceneConstraintChannel* Channel = ChannelHandle.Get();
+	FMovieSceneConstraintChannel* Channel = GetChannel();
 	UMovieSceneSection* Section = WeakSection.Get();
 	TArray<FKeyBarCurveModel::FBarRange> Ranges = FConstraintChannelEditor::GetBarRanges(Channel, Section);
 	return Ranges;
@@ -82,9 +96,17 @@ TArray<FKeyBarCurveModel::FBarRange> FConstraintChannelCurveModel::FindRanges()
 
 const void* FConstraintChannelCurveModel::GetCurve() const
 {
-	return ChannelHandle.Get();
+	return GetChannel();
 }
 
+FMovieSceneConstraintChannel* FConstraintChannelCurveModel::GetChannel() const
+{
+	FixupCurve();
+	FMovieSceneConstraintChannel* Channel = ChannelHandle.Get();
+	return Channel;
+}
+
+
 void FConstraintChannelCurveModel::Modify()
 {
 	if (UMovieSceneSection* Section = WeakSection.Get())
@@ -97,7 +119,7 @@ void FConstraintChannelCurveModel::AddKeys(TArrayView<const FKeyPosition> InKeyP
 {
 	check(InKeyPositions.Num() == InKeyAttributes.Num() && (!OutKeyHandles || OutKeyHandles->Num() == InKeyPositions.Num()));
 
-	FMovieSceneConstraintChannel* Channel = ChannelHandle.Get();
+	FMovieSceneConstraintChannel* Channel = GetChannel();
 	UMovieSceneSection* Section = WeakSection.Get();
 	if (Channel && Section)
 	{
@@ -135,7 +157,7 @@ bool FConstraintChannelCurveModel::Evaluate(double Time, double& OutValue) const
 
 void FConstraintChannelCurveModel::RemoveKeys(TArrayView<const FKeyHandle> InKeys)
 {
-	FMovieSceneConstraintChannel* Channel = ChannelHandle.Get();
+	FMovieSceneConstraintChannel* Channel = GetChannel();
 	UMovieSceneSection* Section = WeakSection.Get();
 	if (Channel && Section)
 	{
@@ -161,7 +183,7 @@ void FConstraintChannelCurveModel::DrawCurve(const FCurveEditor& CurveEditor, co
 
 void FConstraintChannelCurveModel::GetKeys(const FCurveEditor& CurveEditor, double MinTime, double MaxTime, double MinValue, double MaxValue, TArray<FKeyHandle>& OutKeyHandles) const
 {
-	FMovieSceneConstraintChannel* Channel = ChannelHandle.Get();
+	FMovieSceneConstraintChannel* Channel = GetChannel();
 	UMovieSceneSection* Section = WeakSection.Get();
 
 	if (Channel && Section)
@@ -195,7 +217,7 @@ void FConstraintChannelCurveModel::GetKeyDrawInfo(ECurvePointType PointType, con
 
 void FConstraintChannelCurveModel::GetKeyPositions(TArrayView<const FKeyHandle> InKeys, TArrayView<FKeyPosition> OutKeyPositions) const
 {
-	FMovieSceneConstraintChannel* Channel = ChannelHandle.Get();
+	FMovieSceneConstraintChannel* Channel = GetChannel();
 	UMovieSceneSection* Section = WeakSection.Get();
 
 	if (Channel && Section)
@@ -220,7 +242,7 @@ void FConstraintChannelCurveModel::GetKeyPositions(TArrayView<const FKeyHandle>
 
 void FConstraintChannelCurveModel::SetKeyPositions(TArrayView<const FKeyHandle> InKeys, TArrayView<const FKeyPosition> InKeyPositions, EPropertyChangeType::Type ChangeType)
 {
-	FMovieSceneConstraintChannel* Channel = ChannelHandle.Get();
+	FMovieSceneConstraintChannel* Channel = GetChannel();
 	UMovieSceneSection* Section = WeakSection.Get();
 
 	if (Channel && Section)
@@ -269,7 +291,7 @@ void FConstraintChannelCurveModel::CreateKeyProxies(TArrayView<const FKeyHandle>
 
 void FConstraintChannelCurveModel::GetTimeRange(double& MinTime, double& MaxTime) const
 {
-	FMovieSceneConstraintChannel* Channel = ChannelHandle.Get();
+	FMovieSceneConstraintChannel* Channel = GetChannel();
 	UMovieSceneSection* Section = WeakSection.Get();
 
 	if (Channel && Section)
@@ -298,7 +320,7 @@ void FConstraintChannelCurveModel::GetValueRange(double& MinValue, double& MaxVa
 
 int32 FConstraintChannelCurveModel::GetNumKeys() const
 {
-	FMovieSceneConstraintChannel* Channel = ChannelHandle.Get();
+	FMovieSceneConstraintChannel* Channel = GetChannel();
 
 	if (Channel)
 	{
@@ -318,7 +340,7 @@ void FConstraintChannelCurveModel::BuildContextMenu(const FCurveEditor& CurveEdi
 	{
 		MenuBuilder.BeginSection("SpaceContextMenuSection", NSLOCTEXT("CurveEditor", "Constraint", "Constraint"));
 		{
-			FMovieSceneConstraintChannel* Channel = ChannelHandle.Get();
+			FMovieSceneConstraintChannel* Channel = GetChannel();
 			const FKeyHandle KeyHandle = ClickedPoint.GetValue().KeyHandle;
 			if (Channel && KeyHandle != FKeyHandle::Invalid())
 			{
diff --git a/Engine/Source/Editor/MovieSceneTools/Public/Channels/ConstraintChannelCurveModel.h b/Engine/Source/Editor/MovieSceneTools/Public/Channels/ConstraintChannelCurveModel.h
--- a/Engine/Source/Editor/MovieSceneTools/Public/Channels/ConstraintChannelCurveModel.h
+++ b/Engine/Source/Editor/MovieSceneTools/Public/Channels/ConstraintChannelCurveModel.h
@@ -72,12 +72,18 @@ public:
 
 
 private:
-	void FixupCurve();
+	//Fixup the curve and get the right proxy later when we access it due to the fact the constraints channels may not be setup correctly
+	//These functions all need to be const since they can be called by const functions so we need to specify the things they change as mutable
+	void FixupCurve() const;
+	void NeedNewProxy() const;
+	FMovieSceneConstraintChannel* GetChannel() const;
+
 private:
-	TMovieSceneChannelHandle<FMovieSceneConstraintChannel> ChannelHandle;
+	mutable TMovieSceneChannelHandle<FMovieSceneConstraintChannel> ChannelHandle;
 	TWeakObjectPtr<UMovieSceneSection> WeakSection;
 	TWeakPtr<ISequencer> WeakSequencer;
 	TWeakObjectPtr<UTickableConstraint> Constraint;
 
-	FDelegateHandle OnDestroyHandle;
+	mutable FDelegateHandle OnDestroyHandle;
+	mutable bool bNeedNewProxy = false; 
 };
\ No newline at end of file
diff --git a/Engine/Source/Runtime/MovieSceneTracks/Private/Sections/MovieScene3DTransformSection.cpp b/Engine/Source/Runtime/MovieSceneTracks/Private/Sections/MovieScene3DTransformSection.cpp
--- a/Engine/Source/Runtime/MovieSceneTracks/Private/Sections/MovieScene3DTransformSection.cpp
+++ b/Engine/Source/Runtime/MovieSceneTracks/Private/Sections/MovieScene3DTransformSection.cpp
@@ -980,7 +980,7 @@ void UMovieScene3DTransformSection::AddConstraintChannel(UTickableConstraint* In
 	Modify();
 	if (!Constraints)
 	{
-		Constraints = NewObject<UMovieScene3dTransformSectionConstraints>(this, NAME_None, RF_Public| RF_Transactional);
+		Constraints = NewObject<UMovieScene3DTransformSectionConstraints>(this, NAME_None, RF_Public| RF_Transactional);
 	}
 	if (InConstraint && !HasConstraintChannel(InConstraint->GetFName()))
 	{
@@ -1078,8 +1078,29 @@ UE::MovieScene::FChannelOverrideProviderTraitsHandle UMovieScene3DTransformSecti
 	return UE::MovieScene::FChannelOverrideProviderTraitsHandle(Traits);
 }
 
+void UMovieScene3DTransformSection::OnConstraintsChanged()
+{
+	ChannelProxy = nullptr;
+}
+
 void UMovieScene3DTransformSection::OnChannelOverridesChanged()
 {
 	ChannelProxy = nullptr;
 }
 
+void UMovieScene3DTransformSectionConstraints::Serialize(FArchive& Ar)
+{
+	Super::Serialize(Ar);
+	//if we are loading during an undo/redo set the cache to be empty
+	if (Ar.IsLoading() && Ar.IsTransacting())
+	{
+		if (UMovieScene3DTransformSection* Section = GetTypedOuter<UMovieScene3DTransformSection>())
+		{
+			if (IMovieSceneConstrainedSection* ConstrainedSection = Cast<IMovieSceneConstrainedSection>(Section))
+			{
+				ConstrainedSection->OnConstraintsChanged();
+			}
+		}
+	}
+}
+
diff --git a/Engine/Source/Runtime/MovieSceneTracks/Public/Sections/MovieScene3DTransformSection.h b/Engine/Source/Runtime/MovieSceneTracks/Public/Sections/MovieScene3DTransformSection.h
--- a/Engine/Source/Runtime/MovieSceneTracks/Public/Sections/MovieScene3DTransformSection.h
+++ b/Engine/Source/Runtime/MovieSceneTracks/Public/Sections/MovieScene3DTransformSection.h
@@ -192,12 +192,14 @@ private:
 * This object contains information needed for constraint channels on the transform section
 */
 UCLASS()
-class MOVIESCENETRACKS_API UMovieScene3dTransformSectionConstraints : public UObject
+class MOVIESCENETRACKS_API UMovieScene3DTransformSectionConstraints : public UObject
 {
 	GENERATED_BODY()
 
 public:
 
+	virtual void Serialize(FArchive& Ar) override;
+
 	/** Constraint Channels*/
 	UPROPERTY()
 	TArray<FConstraintAndActiveChannel> ConstraintsChannels;
@@ -309,7 +311,7 @@ private:
 
 	/** Optional pointer to constraint channels*/
 	UPROPERTY()
-	TObjectPtr<UMovieScene3dTransformSectionConstraints> Constraints;
+	TObjectPtr<UMovieScene3DTransformSectionConstraints> Constraints;
 
 	/** Whether to use a quaternion linear interpolation between keys. This finds the 'shortest' rotation between keyed orientations. */
 	UPROPERTY(EditAnywhere, DisplayName = "Use Quaternion Interpolation", Category = "Rotation")
@@ -347,6 +349,12 @@ public:
 	*  Replace the constraint with the specified name with the new one
 	*/
 	virtual void ReplaceConstraint(const FName InName, UTickableConstraint* InConstraint)  override;
+
+	/*
+	* Clear proxy if changed
+	*/
+	virtual void OnConstraintsChanged() override;
+
 private:
 
 #if WITH_EDITORONLY_DATA
diff --git a/Engine/Source/Runtime/MovieSceneTracks/Public/Sections/MovieSceneConstrainedSection.h b/Engine/Source/Runtime/MovieSceneTracks/Public/Sections/MovieSceneConstrainedSection.h
--- a/Engine/Source/Runtime/MovieSceneTracks/Public/Sections/MovieSceneConstrainedSection.h
+++ b/Engine/Source/Runtime/MovieSceneTracks/Public/Sections/MovieSceneConstrainedSection.h
@@ -72,6 +72,11 @@ public:
 	*/
 	 FConstraintChannelAddedEvent& ConstraintChannelAdded() { return OnConstraintChannelAdded; }
 
+	 /*
+	 *  What to do if the constraint object has been changed, for example by an undo or redo. By default nothing to be overriden if needed.
+	 */
+	 virtual void OnConstraintsChanged() {};
+
 	/*
 	*  Removed delegate that get's added by the track editor
 	*/
