make sure we compensate thru all constraints (as a mask) if the last one doesn't constrain the entire child's transform

diff --git a/Engine/Plugins/Animation/ControlRig/Source/ControlRigEditor/Private/EditMode/ControlRigEditMode.cpp b/Engine/Plugins/Animation/ControlRig/Source/ControlRigEditor/Private/EditMode/ControlRigEditMode.cpp
--- a/Engine/Plugins/Animation/ControlRig/Source/ControlRigEditor/Private/EditMode/ControlRigEditMode.cpp
+++ b/Engine/Plugins/Animation/ControlRig/Source/ControlRigEditor/Private/EditMode/ControlRigEditMode.cpp
@@ -3827,15 +3827,8 @@ void FControlRigEditMode::SetControlShapeTransform(
 	// find the last constraint in the stack (this could be cached on mouse press)
 	TArray< TObjectPtr<UTickableConstraint> > Constraints;
 	FTransformConstraintUtils::GetParentConstraints(ControlRig->GetWorld(), InShapeActor, Constraints);
-	const int32 LastActiveIndex = Constraints.FindLastByPredicate([](const TObjectPtr<UTickableConstraint>& InConstraint)
-	{
-		if (const UTickableTransformConstraint* TransformConstraint = Cast<UTickableTransformConstraint>(InConstraint.Get()))
-		{
-			return InConstraint->Active && TransformConstraint->bDynamicOffset;
-		}
-		return false;
-	});
 
+	const int32 LastActiveIndex = FTransformConstraintUtils::GetLastActiveConstraintIndex(Constraints);
 	const bool bNeedsConstraintPostProcess = Constraints.IsValidIndex(LastActiveIndex);
 	
 	// set the global space, assumes it's attached to actor
@@ -3852,13 +3845,15 @@ void FControlRigEditMode::SetControlShapeTransform(
 	}
 	
 	// switch to constraint space
-	const UTickableTransformConstraint* Constraint = Cast<UTickableTransformConstraint>(Constraints[LastActiveIndex]);
-	const FTransform ParentTransform = Constraint->GetParentGlobalTransform();
-	const FTransform ChildTransform = InGlobalTransform * InToWorldTransform;
+	const FTransform WorldTransform = InGlobalTransform * InToWorldTransform;
 	FTransform LocalTransform = ControlRig->GetControlLocalTransform(InShapeActor->ControlName);
-	
-	LocalTransform = FTransformConstraintUtils::ComputeRelativeTransform(LocalTransform,
-		ChildTransform, ParentTransform, Constraint);
+
+	const TOptional<FTransform> RelativeTransform =
+		FTransformConstraintUtils::GetConstraintsRelativeTransform(Constraints, LocalTransform, WorldTransform);
+	if (RelativeTransform)
+	{
+		LocalTransform = *RelativeTransform; 
+	}
 
 	FRigControlModifiedContext Context = InContext;
 	Context.bConstraintUpdate = false;
diff --git a/Engine/Plugins/Animation/ControlRig/Source/ControlRigEditor/Private/Sequencer/ControlRigParameterTrackEditor.cpp b/Engine/Plugins/Animation/ControlRig/Source/ControlRigEditor/Private/Sequencer/ControlRigParameterTrackEditor.cpp
--- a/Engine/Plugins/Animation/ControlRig/Source/ControlRigEditor/Private/Sequencer/ControlRigParameterTrackEditor.cpp
+++ b/Engine/Plugins/Animation/ControlRig/Source/ControlRigEditor/Private/Sequencer/ControlRigParameterTrackEditor.cpp
@@ -3224,7 +3224,7 @@ void FControlRigParameterTrackEditor::GetControlRigKeys(
 				if (bInConstraintSpace)
 				{
 					const uint32 ControlHash = UTransformableControlHandle::ComputeHash(InControlRig, ControlElement->GetName());
-					TOptional<FTransform> Transform = FTransformConstraintUtils::GetConstraintRelativeTransform(InControlRig->GetWorld(), ControlHash);
+					TOptional<FTransform> Transform = FTransformConstraintUtils::GetRelativeTransform(InControlRig->GetWorld(), ControlHash);
 					if (Transform)
 					{
 						Translation = Transform->GetTranslation();
diff --git a/Engine/Source/Editor/MovieSceneTools/Private/MovieSceneConstraintChannelHelper.cpp b/Engine/Source/Editor/MovieSceneTools/Private/MovieSceneConstraintChannelHelper.cpp
--- a/Engine/Source/Editor/MovieSceneTools/Private/MovieSceneConstraintChannelHelper.cpp
+++ b/Engine/Source/Editor/MovieSceneTools/Private/MovieSceneConstraintChannelHelper.cpp
@@ -42,23 +42,27 @@ void FCompensationEvaluator::ComputeLocalTransforms(
 		return;
 	}
 
-	const TArray<UTickableTransformConstraint*> Constraints = GetHandleTransformConstraints(InWorld);
+	using ConstraintPtr = TObjectPtr<UTickableConstraint>;
+	const TArray< ConstraintPtr > Constraints = GetHandleTransformConstraints(InWorld);
 	if (Constraints.IsEmpty())
 	{
 		return;
 	}
+	
+	const TArray<ConstraintPtr> ConstraintsMinusThis =
+		Constraints.FilterByPredicate([this](const ConstraintPtr& InConstraint)
+		{
+			return InConstraint != Constraint;
+		});
 
 	// find last active constraint in the list that is different than the on we want to compensate for
-	auto GetLastActiveConstraint = [this, Constraints]()
+	auto GetLastActiveConstraint = [ConstraintsMinusThis]()
 	{
-		// find last active constraint in the list that is different than the on we want to compensate for
-		const int32 LastActiveIndex = Constraints.FindLastByPredicate([this](const UTickableTransformConstraint* InConstraint)
-			{
-				return InConstraint->Active && InConstraint->bDynamicOffset && InConstraint != Constraint;
-			});
+		// find last active constraint in the list that is different than the one we want to compensate for
+		const int32 LastActiveIndex = FTransformConstraintUtils::GetLastActiveConstraintIndex(ConstraintsMinusThis);
 
 		// if found, return its parent global transform
-		return LastActiveIndex > INDEX_NONE ? Constraints[LastActiveIndex] : nullptr;
+		return LastActiveIndex > INDEX_NONE ? Cast<UTickableTransformConstraint>(ConstraintsMinusThis[LastActiveIndex]) : nullptr;
 	};
 
 	UMovieScene* MovieScene = InSequencer->GetFocusedMovieSceneSequence()->GetMovieScene();
@@ -83,7 +87,7 @@ void FCompensationEvaluator::ComputeLocalTransforms(
 	// get all constraints to evaluate
 	const FConstraintsManagerController& Controller = FConstraintsManagerController::Get(InWorld);
 	static constexpr bool bSorted = true;
-	TArray<TObjectPtr<UTickableConstraint>> AllConstraints = Controller.GetAllConstraints(bSorted);
+	const TArray<ConstraintPtr> AllConstraints = Controller.GetAllConstraints(bSorted);
 	
 	for (int32 Index = 0; Index < NumFrames + 1; ++Index)
 	{
@@ -138,8 +142,12 @@ void FCompensationEvaluator::ComputeLocalTransforms(
 				if (LastConstraint)
 				{
 					SpaceGlobal = LastConstraint->GetParentGlobalTransform();
-					ChildLocal = FTransformConstraintUtils::ComputeRelativeTransform(
-						ChildLocal, ChildGlobal, SpaceGlobal, LastConstraint);
+					TOptional<FTransform> Relative =
+						FTransformConstraintUtils::GetConstraintsRelativeTransform(ConstraintsMinusThis, ChildLocal, ChildGlobal);
+					if (Relative)
+					{
+						ChildLocal = *Relative;
+					}
 				}
 			}
 			else
@@ -164,8 +172,12 @@ void FCompensationEvaluator::ComputeLocalTransforms(
 				if (LastConstraint)
 				{
 					SpaceGlobal = LastConstraint->GetParentGlobalTransform();
-					ChildLocal = FTransformConstraintUtils::ComputeRelativeTransform(
-						ChildLocal, ChildGlobal, SpaceGlobal, LastConstraint);
+					TOptional<FTransform> Relative =
+						FTransformConstraintUtils::GetConstraintsRelativeTransform(ConstraintsMinusThis, ChildLocal, ChildGlobal);
+					if (Relative)
+					{
+						ChildLocal = *Relative;
+					}
 				}
 			}
 		}
@@ -185,18 +197,22 @@ void FCompensationEvaluator::ComputeLocalTransformsForBaking(UWorld* InWorld, co
 		return;
 	}
 
-	const TArray<UTickableTransformConstraint*> Constraints = GetHandleTransformConstraints(InWorld);
-
-	auto GetLastActiveConstraint = [this, Constraints]()
-	{
-		// find last active constraint in the list that is different than the on we want to compensate for
-		const int32 LastActiveIndex = Constraints.FindLastByPredicate([this](const UTickableTransformConstraint* InConstraint)
+	using ConstraintPtr = TObjectPtr<UTickableConstraint>;
+	const TArray< ConstraintPtr > Constraints = GetHandleTransformConstraints(InWorld);
+	
+	const TArray< ConstraintPtr > ConstraintsMinusThis =
+		Constraints.FilterByPredicate([this](const ConstraintPtr& InConstraint)
 		{
-			return InConstraint->Active && InConstraint->NeedsCompensation() && InConstraint->bDynamicOffset && InConstraint != Constraint;
+			return InConstraint != Constraint;
 		});
 
+	auto GetLastActiveConstraint = [ConstraintsMinusThis]()
+	{
+		// find last active constraint in the list that is different than the one we want to compensate for
+		const int32 LastActiveIndex = FTransformConstraintUtils::GetLastActiveConstraintIndex(ConstraintsMinusThis);
+
 		// if found, return its parent global transform
-		return LastActiveIndex > INDEX_NONE ? Constraints[LastActiveIndex] : nullptr;
+		return LastActiveIndex > INDEX_NONE ? Cast<UTickableTransformConstraint>(ConstraintsMinusThis[LastActiveIndex]) : nullptr;
 	};
 
 	UMovieScene* MovieScene = InSequencer->GetFocusedMovieSceneSequence()->GetMovieScene();
@@ -212,7 +228,7 @@ void FCompensationEvaluator::ComputeLocalTransformsForBaking(UWorld* InWorld, co
 	// get all constraints for evaluation
 	const FConstraintsManagerController& Controller = FConstraintsManagerController::Get(InWorld);
 	static constexpr bool bSorted = true;
-	TArray<TObjectPtr<UTickableConstraint>> AllConstraints = Controller.GetAllConstraints(bSorted);
+	const TArray<ConstraintPtr> AllConstraints = Controller.GetAllConstraints(bSorted);
 	
 	const TArray<IMovieSceneToolsAnimationBakeHelper*>& BakeHelpers = FMovieSceneToolsModule::Get().GetAnimationBakeHelpers();
 	for (IMovieSceneToolsAnimationBakeHelper* BakeHelper : BakeHelpers)
@@ -268,8 +284,12 @@ void FCompensationEvaluator::ComputeLocalTransformsForBaking(UWorld* InWorld, co
 		if (const UTickableTransformConstraint* LastConstraint = GetLastActiveConstraint())
 		{
 			SpaceGlobal = LastConstraint->GetParentGlobalTransform();
-			ChildLocal = FTransformConstraintUtils::ComputeRelativeTransform(
-				ChildLocal, ChildGlobal, SpaceGlobal, LastConstraint);
+			TOptional<FTransform> Relative =
+				FTransformConstraintUtils::GetConstraintsRelativeTransform(ConstraintsMinusThis, ChildLocal, ChildGlobal);
+			if (Relative)
+			{
+				ChildLocal = *Relative;
+			}
 		}
 	}
 	for (IMovieSceneToolsAnimationBakeHelper* BakeHelper : BakeHelpers)
@@ -291,29 +311,28 @@ void FCompensationEvaluator::ComputeLocalTransformsBeforeDeletion(
 		return;
 	}
 
-	const TArray<UTickableTransformConstraint*> Constraints = GetHandleTransformConstraints(InWorld);
-	// if (Constraints.IsEmpty())
-	// {
-	// 	return;
-	// }
+	using ConstraintPtr = TObjectPtr<UTickableConstraint>;
+	const TArray<ConstraintPtr> Constraints = GetHandleTransformConstraints(InWorld);
+	const TArray<ConstraintPtr> ConstraintsMinusThis = Constraints.FilterByPredicate(
+		[this](const ConstraintPtr& InConstraint)
+		{
+			return InConstraint != Constraint;
+		});
 
 	// find last active constraint in the list that is different than the on we want to compensate for
-	auto GetLastActiveConstraint = [this, Constraints]()
+	auto GetLastActiveConstraint = [this, ConstraintsMinusThis]()
 	{
 		// find last active constraint in the list that is different than the on we want to compensate for
-		const int32 LastActiveIndex = Constraints.FindLastByPredicate([this](const UTickableTransformConstraint* InConstraint)
-			{
-				return InConstraint->Active && InConstraint->bDynamicOffset && InConstraint != Constraint;
-			});
+		const int32 LastActiveIndex = FTransformConstraintUtils::GetLastActiveConstraintIndex(ConstraintsMinusThis);
 
 		// if found, return its parent global transform
-		return LastActiveIndex > INDEX_NONE ? Constraints[LastActiveIndex] : nullptr;
+		return LastActiveIndex > INDEX_NONE ? Cast<UTickableTransformConstraint>(ConstraintsMinusThis[LastActiveIndex]) : nullptr;
 	};
 
 	// get all constraints for evaluation
 	const FConstraintsManagerController& Controller = FConstraintsManagerController::Get(InWorld);
 	static constexpr bool bSorted = true;
-	TArray<TObjectPtr<UTickableConstraint>> AllConstraints = Controller.GetAllConstraints(bSorted);
+	const TArray<ConstraintPtr> AllConstraints = Controller.GetAllConstraints(bSorted);
 	
 	UMovieScene* MovieScene = InSequencer->GetFocusedMovieSceneSequence()->GetMovieScene();
 	const FFrameRate TickResolution = MovieScene->GetTickResolution();
@@ -377,8 +396,12 @@ void FCompensationEvaluator::ComputeLocalTransformsBeforeDeletion(
 		if (const UTickableTransformConstraint* LastConstraint = GetLastActiveConstraint())
 		{
 			SpaceGlobal = LastConstraint->GetParentGlobalTransform();
-			ChildLocal = FTransformConstraintUtils::ComputeRelativeTransform(
-				ChildLocal, ChildGlobal, SpaceGlobal, LastConstraint);
+			TOptional<FTransform> Relative =
+				FTransformConstraintUtils::GetConstraintsRelativeTransform(Constraints, ChildLocal, ChildGlobal);
+			if(Relative)
+			{
+				ChildLocal = *Relative;
+			}
 		}
 	}
 	for (IMovieSceneToolsAnimationBakeHelper* BakeHelper : BakeHelpers)
@@ -392,7 +415,8 @@ void FCompensationEvaluator::ComputeLocalTransformsBeforeDeletion(
 
 void FCompensationEvaluator::ComputeCompensation(UWorld* InWorld, const TSharedPtr<ISequencer>& InSequencer, const FFrameNumber& InTime)
 {
-	const TArray<UTickableTransformConstraint*> Constraints = GetHandleTransformConstraints(InWorld);
+	using ConstraintPtr = TObjectPtr<UTickableConstraint>;
+	const TArray<ConstraintPtr> Constraints = GetHandleTransformConstraints(InWorld);
 	if (Constraints.IsEmpty())
 	{
 		return;
@@ -402,19 +426,16 @@ void FCompensationEvaluator::ComputeCompensation(UWorld* InWorld, const TSharedP
 	auto GetLastActiveConstraint = [this, Constraints]()
 	{
 		// find last active constraint in the list that is different than the on we want to compensate for
-		const int32 LastActiveIndex = Constraints.FindLastByPredicate([this](const UTickableTransformConstraint* InConstraint)
-			{
-				return InConstraint->Active && InConstraint->bDynamicOffset;
-			});
+		const int32 LastActiveIndex = FTransformConstraintUtils::GetLastActiveConstraintIndex(Constraints);
 
 		// if found, return its parent global transform
-		return LastActiveIndex > INDEX_NONE ? Constraints[LastActiveIndex] : nullptr;
+		return LastActiveIndex > INDEX_NONE ? Cast<UTickableTransformConstraint>(Constraints[LastActiveIndex]) : nullptr;
 	};
 
 	// get all constraints for evaluation
 	const FConstraintsManagerController& Controller = FConstraintsManagerController::Get(InWorld);
 	static constexpr bool bSorted = true;
-	TArray<TObjectPtr<UTickableConstraint>> AllConstraints = Controller.GetAllConstraints(bSorted);
+	const TArray<ConstraintPtr> AllConstraints = Controller.GetAllConstraints(bSorted);
 	
 	UMovieScene* MovieScene = InSequencer->GetFocusedMovieSceneSequence()->GetMovieScene();
 	const TArray<IMovieSceneToolsAnimationBakeHelper*>& BakeHelpers = FMovieSceneToolsModule::Get().GetAnimationBakeHelpers();
@@ -488,8 +509,12 @@ void FCompensationEvaluator::ComputeCompensation(UWorld* InWorld, const TSharedP
 	if (const UTickableTransformConstraint* LastConstraint = GetLastActiveConstraint())
 	{
 		SpaceGlobals[0] = LastConstraint->GetParentGlobalTransform();
-		ChildLocals[0] = FTransformConstraintUtils::ComputeRelativeTransform(ChildLocals[0],
-			ChildGlobals[0], SpaceGlobals[0], LastConstraint);
+		TOptional<FTransform> Relative =
+			FTransformConstraintUtils::GetConstraintsRelativeTransform(Constraints, ChildLocals[0], ChildGlobals[0]);
+		if(Relative)
+		{
+			ChildLocals[0] = *Relative;
+		}
 	}
 	else // switch to parent space
 	{
@@ -500,24 +525,24 @@ void FCompensationEvaluator::ComputeCompensation(UWorld* InWorld, const TSharedP
 	}
 }
 
-TArray<UTickableTransformConstraint*> FCompensationEvaluator::GetHandleTransformConstraints(UWorld* InWorld) const
+TArray< TObjectPtr<UTickableConstraint> > FCompensationEvaluator::GetHandleTransformConstraints(UWorld* InWorld) const
 {
-	TArray<UTickableTransformConstraint*> TransformConstraints;
+	using ConstraintPtr = TObjectPtr<UTickableConstraint>;
+	
 	if (Handle)
 	{
 		// get sorted transform constraints
 		const FConstraintsManagerController& Controller = FConstraintsManagerController::Get(InWorld);
 		static constexpr bool bSorted = true;
-		TArray<TObjectPtr<UTickableConstraint>> Constraints = Controller.GetParentConstraints(Handle->GetHash(), bSorted);
-		for (const TObjectPtr<UTickableConstraint>& InConstraint : Constraints)
+		const TArray< ConstraintPtr > Constraints = Controller.GetParentConstraints(Handle->GetHash(), bSorted);
+		return Constraints.FilterByPredicate([](const ConstraintPtr& InConstraint)
 		{
-			if (UTickableTransformConstraint* TransformConstraint = Cast<UTickableTransformConstraint>(InConstraint.Get()))
-			{
-				TransformConstraints.Add(TransformConstraint);
-			}
-		}
+			return IsValid(InConstraint) && InConstraint.IsA<UTickableTransformConstraint>();  
+		});
 	}
-	return TransformConstraints;
+
+	static const TArray< ConstraintPtr > DummyArray;
+	return DummyArray;
 }
 
 bool FMovieSceneConstraintChannelHelper::bDoNotCompensate = false;
diff --git a/Engine/Source/Editor/MovieSceneTools/Private/TrackEditors/TransformTrackEditor.cpp b/Engine/Source/Editor/MovieSceneTools/Private/TrackEditors/TransformTrackEditor.cpp
--- a/Engine/Source/Editor/MovieSceneTools/Private/TrackEditors/TransformTrackEditor.cpp
+++ b/Engine/Source/Editor/MovieSceneTools/Private/TrackEditors/TransformTrackEditor.cpp
@@ -907,63 +907,17 @@ static void UpdateTransformBasedOnConstraint(FTransform& CurrentTransform, UScen
 	if (ShapeActor)
 	{
 		FTransformConstraintUtils::GetParentConstraints(SceneComponent->GetWorld(), ShapeActor, Constraints);
-		const int32 LastActiveIndex = Constraints.FindLastByPredicate([](const TObjectPtr<UTickableConstraint>& InConstraint)
-			{
-				if (const UTickableTransformConstraint* TransformConstraint = Cast<UTickableTransformConstraint>(InConstraint.Get()))
-				{
-					return InConstraint->Active && TransformConstraint->bDynamicOffset;
-				}
-				return false;
-			});
 
+		const int32 LastActiveIndex = FTransformConstraintUtils::GetLastActiveConstraintIndex(Constraints);
 		if (Constraints.IsValidIndex(LastActiveIndex))
 		{
 			// switch to constraint space
-			const UTickableTransformConstraint* Constraint = Cast<UTickableTransformConstraint>(Constraints[LastActiveIndex]);
-			const FTransform ParentTransform = Constraint->GetParentGlobalTransform();
-			const FTransform OriginalParentTransform = SceneComponent->GetAttachParent() ? SceneComponent->GetAttachParent()->GetSocketTransform(SceneComponent->GetAttachSocketName()) : FTransform::Identity;
-			const FTransform DiffParents = OriginalParentTransform.GetRelativeTransform(ParentTransform);
-
-			const ETransformConstraintType ConstraintType = static_cast<ETransformConstraintType>(Constraint->GetType());
-			switch (ConstraintType)
+			const FTransform WorldTransform = SceneComponent->GetSocketTransform(SceneComponent->GetAttachSocketName());
+			const TOptional<FTransform> RelativeTransform =
+				FTransformConstraintUtils::GetConstraintsRelativeTransform(Constraints, CurrentTransform, WorldTransform);
+			if (RelativeTransform)
 			{
-				case ETransformConstraintType::Translation:
-				{
-					CurrentTransform.SetLocation(CurrentTransform.GetLocation() + OriginalParentTransform.GetLocation() - ParentTransform.GetLocation());
-				}
-				break;
-				case ETransformConstraintType::Rotation:
-				{
-					FQuat RelativeRotation = ParentTransform.GetRotation().Inverse() * OriginalParentTransform.GetRotation();
-					RelativeRotation.Normalize();
-					CurrentTransform.SetRotation(CurrentTransform.GetRotation() * RelativeRotation);
-				}
-				break;
-				case ETransformConstraintType::Scale:
-				{
-					CurrentTransform.SetScale3D(CurrentTransform.GetScale3D() * DiffParents.GetScale3D());
-				}
-				break;
-				case ETransformConstraintType::Parent:
-				{
-					TOptional<FVector> Scale3D;
-					if (const UTickableParentConstraint* ParentConstraint = Cast<const UTickableParentConstraint>(Constraint))
-					{
-						if (ParentConstraint->IsScalingEnabled() == false)
-						{
-							Scale3D = CurrentTransform.GetScale3D();
-						}
-					}
-					CurrentTransform = CurrentTransform * DiffParents;
-					if (Scale3D.IsSet())
-					{
-						CurrentTransform.SetScale3D(Scale3D.GetValue());
-					}
-				}
-				break;
-				case ETransformConstraintType::LookAt: //leave current alone
-				default:
-					break;
+				CurrentTransform = *RelativeTransform; 
 			}
 		}
 	}
diff --git a/Engine/Source/Editor/MovieSceneTools/Public/MovieSceneConstraintChannelHelper.h b/Engine/Source/Editor/MovieSceneTools/Public/MovieSceneConstraintChannelHelper.h
--- a/Engine/Source/Editor/MovieSceneTools/Public/MovieSceneConstraintChannelHelper.h
+++ b/Engine/Source/Editor/MovieSceneTools/Public/MovieSceneConstraintChannelHelper.h
@@ -3,6 +3,7 @@
 #pragma once
 
 #include "Templates/SharedPointer.h"
+#include "UObject/ObjectPtr.h"
 
 class UWorld;
 class UTickableConstraint;
@@ -29,7 +30,7 @@ public:
 	void ComputeLocalTransformsForBaking(UWorld* InWorld, const TSharedPtr<ISequencer>& InSequencer, const TArray<FFrameNumber>& InFrames);
 private:
 
-	TArray<UTickableTransformConstraint*> GetHandleTransformConstraints(UWorld* InWorld) const;
+	TArray< TObjectPtr<UTickableConstraint> > GetHandleTransformConstraints(UWorld* InWorld) const;
 
 	UTickableTransformConstraint* Constraint = nullptr;
 	UTransformableHandle* Handle = nullptr;
diff --git a/Engine/Source/Runtime/Experimental/Animation/Constraints/Private/TransformConstraint.cpp b/Engine/Source/Runtime/Experimental/Animation/Constraints/Private/TransformConstraint.cpp
--- a/Engine/Source/Runtime/Experimental/Animation/Constraints/Private/TransformConstraint.cpp
+++ b/Engine/Source/Runtime/Experimental/Animation/Constraints/Private/TransformConstraint.cpp
@@ -1446,7 +1446,7 @@ FTransform FTransformConstraintUtils::ComputeRelativeTransform(
 	return InChildWorld.GetRelativeTransform(InSpaceWorld);
 }
 
-TOptional<FTransform> FTransformConstraintUtils::GetConstraintRelativeTransform(UWorld* InWorld, const uint32 InHandleHash)
+TOptional<FTransform> FTransformConstraintUtils::GetRelativeTransform(UWorld* InWorld, const uint32 InHandleHash)
 {
 	if (!InWorld || InHandleHash <= 0)
 	{
@@ -1455,23 +1455,15 @@ TOptional<FTransform> FTransformConstraintUtils::GetConstraintRelativeTransform(
 
 	static constexpr bool bSorted = true;
 	const FConstraintsManagerController& Controller = FConstraintsManagerController::Get(InWorld);
-	const TArray<TObjectPtr<UTickableConstraint>> Constraints = Controller.GetParentConstraints(InHandleHash, bSorted);
 
+	const TArray< TObjectPtr<UTickableConstraint> > Constraints = Controller.GetParentConstraints(InHandleHash, bSorted);
 	if (Constraints.IsEmpty())
 	{
 		return TOptional<FTransform>();
 	}
 
 	// get current active transform constraint 	
-	const int32 LastActiveIndex = Constraints.FindLastByPredicate([](const TObjectPtr<UTickableConstraint>& InConstraint)
-    {
-    	if (const UTickableTransformConstraint* TransformConstraint = Cast<UTickableTransformConstraint>(InConstraint.Get()))
-    	{
-    		return InConstraint->IsFullyActive() && TransformConstraint->bDynamicOffset;
-    	}
-    	return false;
-    });
-
+	const int32 LastActiveIndex = GetLastActiveConstraintIndex(Constraints);
 	if (!Constraints.IsValidIndex(LastActiveIndex))
 	{
 		return TOptional<FTransform>();
@@ -1481,9 +1473,92 @@ TOptional<FTransform> FTransformConstraintUtils::GetConstraintRelativeTransform(
 	const UTickableTransformConstraint* Constraint = Cast<UTickableTransformConstraint>(Constraints[LastActiveIndex]);
 	const FTransform ChildLocal = Constraint->GetChildLocalTransform();
 	const FTransform ChildGlobal = Constraint->GetChildGlobalTransform();
-	const FTransform ParentGlobal = Constraint->GetParentGlobalTransform();
+	
+	return GetConstraintsRelativeTransform(Constraints, ChildLocal, ChildGlobal);
+}
+
+TOptional<FTransform> FTransformConstraintUtils::GetConstraintsRelativeTransform(
+	const TArray< TObjectPtr<UTickableConstraint> >& InConstraints,
+	const FTransform& InChildLocal, const FTransform& InChildWorld)
+{
+	if (InConstraints.IsEmpty())
+	{
+		return TOptional<FTransform>();
+	}
+
+	// get current active transform constraint
+	const int32 LastActiveIndex = GetLastActiveConstraintIndex(InConstraints);
+	if (!InConstraints.IsValidIndex(LastActiveIndex))
+	{
+		return TOptional<FTransform>();
+	}
 
-	return ComputeRelativeTransform(ChildLocal, ChildGlobal, ParentGlobal, Constraint);
+	// get relative transform
+	// if that constraint handles the entire transform then return the relative transform directly
+	const UTickableTransformConstraint* Constraint = Cast<UTickableTransformConstraint>(InConstraints[LastActiveIndex]);
+	if (EnumHasAllFlags(Constraint->GetChannelsToKey(), EMovieSceneTransformChannel::AllTransform))
+	{
+		const FTransform ParentGlobal = Constraint->GetParentGlobalTransform();
+		return ComputeRelativeTransform(InChildLocal, InChildWorld, ParentGlobal, Constraint);
+	}
+
+	// otherwise, we need to look for constraints on a sub-transform basis so we compute the relative transform for each of them
+	using ConstraintPtr = TObjectPtr<UTickableConstraint>;
+	auto GetLastSubTransformIndex = [InConstraints](const EMovieSceneTransformChannel& InChannel)
+	{
+		return InConstraints.FindLastByPredicate([InChannel](const ConstraintPtr& InConstraint)
+		{
+			if (const UTickableTransformConstraint* TransformConstraint = Cast<UTickableTransformConstraint>(InConstraint.Get()))
+			{
+				const bool bHasChannelFlag = EnumHasAllFlags(TransformConstraint->GetChannelsToKey(), InChannel);
+				return InConstraint->IsFullyActive() && TransformConstraint->bDynamicOffset && bHasChannelFlag;
+			}
+			return false;
+		});
+	};
+
+	// look for last constraint index for each channel 
+	static constexpr EMovieSceneTransformChannel SubChannels[3] = {	EMovieSceneTransformChannel::Translation,
+																	EMovieSceneTransformChannel::Rotation,
+																	EMovieSceneTransformChannel::Scale };
+	TArray<int32> SubTransformIndices;
+	for (int32 Index = 0; Index < 3; Index++)
+	{
+		const int32 LastSubIndex = GetLastSubTransformIndex(SubChannels[Index]);
+    	if (InConstraints.IsValidIndex(LastSubIndex))  
+    	{
+    		SubTransformIndices.AddUnique(LastSubIndex);
+    	}
+	}
+	SubTransformIndices.Sort();
+
+	// if none then return 
+	if (SubTransformIndices.IsEmpty())
+	{
+		return TOptional<FTransform>();
+	}
+
+	// iterate thru constraints to compute the relative transform in each of them 
+	FTransform ChildLocal = InChildLocal;
+	for (const int32 SubConstraintIndex: SubTransformIndices)
+	{
+		const UTickableTransformConstraint* SubConstraint = Cast<UTickableTransformConstraint>(InConstraints[SubConstraintIndex]);
+		const FTransform ParentGlobal = SubConstraint->GetParentGlobalTransform();
+		ChildLocal = ComputeRelativeTransform(ChildLocal, InChildWorld, ParentGlobal, SubConstraint);
+	}
+	return ChildLocal;
+}
+
+int32 FTransformConstraintUtils::GetLastActiveConstraintIndex(const TArray< TObjectPtr<UTickableConstraint> >& InConstraints)
+{
+	return InConstraints.FindLastByPredicate([](const TObjectPtr<UTickableConstraint>& InConstraint)
+	{
+	   if (const UTickableTransformConstraint* TransformConstraint = Cast<UTickableTransformConstraint>(InConstraint.Get()))
+	   {
+		   return InConstraint->Active && TransformConstraint->bDynamicOffset;
+	   }
+	   return false;
+	});
 }
 
 void FTransformConstraintUtils::GetChildrenConstraints(
diff --git a/Engine/Source/Runtime/Experimental/Animation/Constraints/Public/TransformConstraint.h b/Engine/Source/Runtime/Experimental/Animation/Constraints/Public/TransformConstraint.h
--- a/Engine/Source/Runtime/Experimental/Animation/Constraints/Public/TransformConstraint.h
+++ b/Engine/Source/Runtime/Experimental/Animation/Constraints/Public/TransformConstraint.h
@@ -417,7 +417,13 @@ struct CONSTRAINTS_API FTransformConstraintUtils
 		const UTickableTransformConstraint* InConstraint);
 
 	/** Computes the current constraint space local transform. */
-	static TOptional<FTransform> GetConstraintRelativeTransform(UWorld* InWorld, const uint32 InHandleHash);
+	static TOptional<FTransform> GetRelativeTransform(UWorld* InWorld, const uint32 InHandleHash);
+	static TOptional<FTransform> GetConstraintsRelativeTransform(
+		const TArray< TObjectPtr<UTickableConstraint> >& InConstraints,
+		const FTransform& InChildLocal, const FTransform& InChildWorld);
+
+	/** Get the last active constraint that has dynamic offset. */
+	static int32 GetLastActiveConstraintIndex(const TArray< TObjectPtr<UTickableConstraint> >& InConstraints);
 
 	/** Fills a constraint array that InParentHandle is the parent of. */
 	static void GetChildrenConstraints(
