Sequencer: Fix drag drop onto skeletal animation track so that it respects the track being dropped onto and doesn't create another dedicated animation track for the skeletal mesh component.

diff --git a/Engine/Source/Editor/MovieSceneTools/Private/TrackEditors/SkeletalAnimationTrackEditor.cpp b/Engine/Source/Editor/MovieSceneTools/Private/TrackEditors/SkeletalAnimationTrackEditor.cpp
--- a/Engine/Source/Editor/MovieSceneTools/Private/TrackEditors/SkeletalAnimationTrackEditor.cpp
+++ b/Engine/Source/Editor/MovieSceneTools/Private/TrackEditors/SkeletalAnimationTrackEditor.cpp
@@ -2060,7 +2060,9 @@ FReply FSkeletalAnimationTrackEditor::OnDrop(const FDragDropEvent& DragDropEvent
 
 			if (bValidAnimSequence && Skeleton && Skeleton->IsCompatible(AnimSequence->GetSkeleton()))
 			{
-				AnimatablePropertyChanged( FOnKeyProperty::CreateRaw(this, &FSkeletalAnimationTrackEditor::AddKeyInternal, (UObject*)SkeletalMeshComponent, AnimSequence, DragDropParams.Track.Get(), DragDropParams.RowIndex));
+				UObject* BoundObject = SequencerPtr.IsValid() ? SequencerPtr->FindSpawnedObjectOrTemplate(DragDropParams.TargetObjectGuid) : nullptr;
+
+				AnimatablePropertyChanged( FOnKeyProperty::CreateRaw(this, &FSkeletalAnimationTrackEditor::AddKeyInternal, BoundObject, AnimSequence, DragDropParams.Track.Get(), DragDropParams.RowIndex));
 
 				bAnyDropped = true;
 			}
