Sequencer: correctly modify the owning UObject of a perlin noise channel when tweaking its parameters
This fixes the problem of the sequence not updating with the new values right away, and also fixes undo/redo.

diff --git a/Engine/Source/Editor/MovieSceneTools/Private/Channels/PerlinNoiseChannelInterface.cpp b/Engine/Source/Editor/MovieSceneTools/Private/Channels/PerlinNoiseChannelInterface.cpp
--- a/Engine/Source/Editor/MovieSceneTools/Private/Channels/PerlinNoiseChannelInterface.cpp
+++ b/Engine/Source/Editor/MovieSceneTools/Private/Channels/PerlinNoiseChannelInterface.cpp
@@ -1,9 +1,12 @@
 // Copyright Epic Games, Inc. All Rights Reserved.
 
 #include "Channels/PerlinNoiseChannelInterface.h"
+#include "Channels/IMovieSceneChannelOverrideProvider.h"
+#include "Channels/MovieSceneSectionChannelOverrideRegistry.h"
 #include "Channels/MovieSceneChannelHandle.h"
 #include "Channels/MovieSceneDoublePerlinNoiseChannel.h"
 #include "Channels/MovieSceneFloatPerlinNoiseChannel.h"
+#include "Editor.h"
 #include "IStructureDetailsView.h"
 #include "Framework/MultiBox/MultiBoxBuilder.h"
 #include "Framework/MultiBox/MultiBoxExtender.h"
@@ -16,6 +19,7 @@ FPerlinNoiseChannelSectionMenuExtension::FPerlinNoiseChannelSectionMenuExtension
 	: ChannelHandles(InChannelHandles)
 	, Sections(InSections)
 {
+	Initialize();
 }
 
 void FPerlinNoiseChannelSectionMenuExtension::ExtendMenu(FMenuBuilder& MenuBuilder)
@@ -40,21 +44,47 @@ void FPerlinNoiseChannelSectionMenuExtension::ExtendMenu(FMenuBuilder& MenuBuild
 	}
 }
 
-void FPerlinNoiseChannelSectionMenuExtension::BuildChannelsMenu(FMenuBuilder& MenuBuilder)
+void FPerlinNoiseChannelSectionMenuExtension::Initialize()
 {
+	// Figure out which channels belong to which section by building the index indirections.
+	// Also, create the notify hooks. Normal channels need to modify the section, but overriden channels
+	// need to modify their override channel container.
 	TArray<FMovieSceneChannelProxy*> ChannelProxies;
 	for (UMovieSceneSection* Section : Sections)
 	{
 		ChannelProxies.Add(&Section->GetChannelProxy());
 	}
 
-	TArray<int32> ChannelHandleSectionIndexes;
 	for (const FMovieSceneChannelHandle& ChannelHandle : ChannelHandles)
 	{
 		int32 SectionIndex = ChannelProxies.Find(ChannelHandle.GetChannelProxy());
 		ChannelHandleSectionIndexes.Add(SectionIndex);
+
+		UMovieSceneSection* Section = Sections[SectionIndex];
+
+		UObject* ObjectToModify = Section;
+
+		IMovieSceneChannelOverrideProvider* ChannelOverrideProvider = Cast<IMovieSceneChannelOverrideProvider>(Section);
+		UMovieSceneSectionChannelOverrideRegistry* ChannelOverrideRegistry = ChannelOverrideProvider ?
+			ChannelOverrideProvider->GetChannelOverrideRegistry(false) : nullptr;
+		if (ChannelOverrideRegistry)
+		{
+			const FName ChannelName = ChannelHandle.GetMetaData()->Name;
+			UMovieSceneChannelOverrideContainer* ChannelOverrideContainer = ChannelOverrideRegistry->GetChannel(ChannelName);
+			if (ChannelOverrideContainer)
+			{
+				ensureMsgf(ChannelOverrideContainer->GetChannel() == ChannelHandle.Get(),
+						TEXT("Mismatched channel override!"));
+				ObjectToModify = ChannelOverrideContainer;
+			}
+		}
+
+		NotifyHooks.Add(FChannelNotifyHook(ObjectToModify));
 	}
+}
 
+void FPerlinNoiseChannelSectionMenuExtension::BuildChannelsMenu(FMenuBuilder& MenuBuilder)
+{
 	const bool bMultipleSections = Sections.Num() > 1;
 	TSharedRef<FPerlinNoiseChannelSectionMenuExtension> SharedThis = this->AsShared();
 
@@ -115,6 +145,9 @@ void FPerlinNoiseChannelSectionMenuExtension::BuildParametersMenu(FMenuBuilder&
 	DetailsViewArgs.bHideSelectionTip = true;
 	DetailsViewArgs.bShowOptions = false;
 	DetailsViewArgs.bShowScrollBar = false;
+	// The notify hook is owned by us. We will live as long as the menu is active, so as long as the NotifyHooks array isn't
+	// modified, the address of the hooks should be valid.
+	DetailsViewArgs.NotifyHook = &NotifyHooks[ChannelHandleIndex];
 
 	FStructureDetailsViewArgs StructureDetailsViewArgs;
 	StructureDetailsViewArgs.bShowObjects = true;
@@ -129,5 +162,16 @@ void FPerlinNoiseChannelSectionMenuExtension::BuildParametersMenu(FMenuBuilder&
 	MenuBuilder.AddWidget(DetailsView->GetWidget().ToSharedRef(), FText(), true, false);
 }
 
-#undef LOCTEXT_NAMESPACE
+void FPerlinNoiseChannelSectionMenuExtension::FChannelNotifyHook::NotifyPreChange(FProperty* PropertyAboutToChange)
+{
+	GEditor->BeginTransaction(FText::Format(LOCTEXT("EditProperty", "Edit {0}"), PropertyAboutToChange->GetDisplayNameText()));
+
+	ObjectToModify->Modify();
+}
 
+void FPerlinNoiseChannelSectionMenuExtension::FChannelNotifyHook::NotifyPostChange(const FPropertyChangedEvent& PropertyChangedEvent, FProperty* PropertyThatChanged)
+{
+	GEditor->EndTransaction();
+}
+
+#undef LOCTEXT_NAMESPACE
diff --git a/Engine/Source/Editor/MovieSceneTools/Public/Channels/PerlinNoiseChannelInterface.h b/Engine/Source/Editor/MovieSceneTools/Public/Channels/PerlinNoiseChannelInterface.h
--- a/Engine/Source/Editor/MovieSceneTools/Public/Channels/PerlinNoiseChannelInterface.h
+++ b/Engine/Source/Editor/MovieSceneTools/Public/Channels/PerlinNoiseChannelInterface.h
@@ -5,6 +5,7 @@
 #include "CurveKeyEditors/SequencerKeyEditor.h"
 #include "Framework/MultiBox/MultiBoxExtender.h"
 #include "ISequencerChannelInterface.h"
+#include "Misc/NotifyHook.h"
 #include "MovieSceneSection.h"
 #include "TimeToPixel.h"
 #include "Widgets/Input/NumericTypeInterface.h"
@@ -50,13 +51,27 @@ struct FPerlinNoiseChannelSectionMenuExtension : TSharedFromThis<FPerlinNoiseCha
 
 private:
 
+	void Initialize();
+
 	void BuildChannelsMenu(FMenuBuilder& MenuBuilder);
 	void BuildParametersMenu(FMenuBuilder& MenuBuilder, int32 ChannelHandleIndex);
 
 private:
 
+	struct FChannelNotifyHook : FNotifyHook
+	{
+		UObject* ObjectToModify = nullptr;
+
+		FChannelNotifyHook(UObject* InObjectToModify) : ObjectToModify(InObjectToModify) {}
+
+		virtual void NotifyPreChange(FProperty* PropertyAboutToChange) override;
+		virtual void NotifyPostChange(const FPropertyChangedEvent& PropertyChangedEvent, FProperty* PropertyThatChanged) override;
+	};
+
 	TArray<FMovieSceneChannelHandle> ChannelHandles;
+	TArray<int32> ChannelHandleSectionIndexes;
 	TArray<UMovieSceneSection*> Sections;
+	TArray<FChannelNotifyHook> NotifyHooks;
 };
 
 template<typename ChannelContainerType>
