Sequencer: Add InsertSequence to subsequences just like there is InsertShot for cinematic shot tracks

Refactor/rename CreateShotInternal to MovieSceneToolHelpers::CreateSubSequence

diff --git a/Engine/Source/Editor/MovieSceneTools/Private/MovieSceneToolHelpers.cpp b/Engine/Source/Editor/MovieSceneTools/Private/MovieSceneToolHelpers.cpp
--- a/Engine/Source/Editor/MovieSceneTools/Private/MovieSceneToolHelpers.cpp
+++ b/Engine/Source/Editor/MovieSceneTools/Private/MovieSceneToolHelpers.cpp
@@ -69,6 +69,8 @@
 #include "Features/IModularFeatures.h"
 #include "Tracks/MovieSceneSpawnTrack.h"
 #include "Sections/MovieSceneSpawnSection.h"
+#include "Tracks/MovieSceneSubTrack.h"
+#include "Sections/MovieSceneSubSection.h"
 #include "Exporters/AnimSeqExportOption.h"
 #include "Widgets/Input/NumericTypeInterface.h"
 #include "FrameNumberDetailsCustomization.h"
@@ -500,6 +502,46 @@ FString MovieSceneToolHelpers::GenerateNewShotName(const TArray<UMovieSceneSecti
 	return ComposeShotName(ProjectSettings->ShotPrefix, ProjectSettings->FirstShotNumber, ProjectSettings->FirstTakeNumber);
 }
 
+UMovieSceneSubSection* MovieSceneToolHelpers::CreateSubSequence(FString& NewSequenceName, FString& NewSequencePath, FFrameNumber NewSequenceStartTime, UMovieSceneSubTrack* SubTrack, UMovieSceneSubSection* SectionToDuplicate)
+{
+	// Create a new level sequence asset with the appropriate name
+	IAssetTools& AssetTools = FModuleManager::GetModuleChecked<FAssetToolsModule>("AssetTools").Get();
+
+	UObject* NewAsset = nullptr;
+	for (TObjectIterator<UClass> It; It; ++It)
+	{
+		UClass* CurrentClass = *It;
+		if (CurrentClass->IsChildOf(UFactory::StaticClass()) && !(CurrentClass->HasAnyClassFlags(CLASS_Abstract)))
+		{
+			UFactory* Factory = Cast<UFactory>(CurrentClass->GetDefaultObject());
+			if (Factory->CanCreateNew() && Factory->ImportPriority >= 0 && Factory->SupportedClass == ULevelSequence::StaticClass())
+			{
+				if (SectionToDuplicate != nullptr)
+				{
+					NewAsset = AssetTools.DuplicateAssetWithDialog(NewSequenceName, NewSequencePath, SectionToDuplicate->GetSequence());
+				}
+				else
+				{
+					NewAsset = AssetTools.CreateAssetWithDialog(NewSequenceName, NewSequencePath, ULevelSequence::StaticClass(), Factory);
+				}
+				break;
+			}
+		}
+	}
+
+	if (NewAsset == nullptr)
+	{
+		return nullptr;
+	}
+
+	UMovieSceneSequence* NewSequence = Cast<UMovieSceneSequence>(NewAsset);
+
+	int32 Duration = UE::MovieScene::DiscreteSize(SectionToDuplicate ? SectionToDuplicate->GetRange() : NewSequence->GetMovieScene()->GetPlaybackRange());
+
+	UMovieSceneSubSection* NewSection = SubTrack->AddSequence(NewSequence, NewSequenceStartTime, Duration);
+	return NewSection;
+}
+
 void MovieSceneToolHelpers::GatherTakes(const UMovieSceneSection* Section, TArray<FAssetData>& AssetData, uint32& OutCurrentTakeNumber)
 {
 	const UMovieSceneSubSection* SubSection = Cast<const UMovieSceneSubSection>(Section);
diff --git a/Engine/Source/Editor/MovieSceneTools/Private/TrackEditors/CinematicShotTrackEditor.cpp b/Engine/Source/Editor/MovieSceneTools/Private/TrackEditors/CinematicShotTrackEditor.cpp
--- a/Engine/Source/Editor/MovieSceneTools/Private/TrackEditors/CinematicShotTrackEditor.cpp
+++ b/Engine/Source/Editor/MovieSceneTools/Private/TrackEditors/CinematicShotTrackEditor.cpp
@@ -367,61 +367,6 @@ FReply FCinematicShotTrackEditor::OnDrop(const FDragDropEvent& DragDropEvent, co
 	return bAnyDropped ? FReply::Handled() : FReply::Unhandled();
 }
 
-UMovieSceneSubSection* FCinematicShotTrackEditor::CreateShotInternal(FString& NewShotName, FFrameNumber NewShotStartTime, UMovieSceneCinematicShotSection* ShotToDuplicate)
-{
-	FString NewShotPath;
-	
-	if (ShotToDuplicate != nullptr)
-	{
-		// If duplicating a shot, use that shot's path
-		NewShotPath = FPaths::GetPath(ShotToDuplicate->GetSequence()->GetPathName());
-	}
-	else
-	{
-		NewShotPath = MovieSceneToolHelpers::GenerateNewShotPath(GetSequencer()->GetFocusedMovieSceneSequence()->GetMovieScene(), NewShotName);
-	}
-
-	// Create a new level sequence asset with the appropriate name
-	IAssetTools& AssetTools = FModuleManager::GetModuleChecked<FAssetToolsModule>("AssetTools").Get();
-
-	UObject* NewAsset = nullptr;
-	for (TObjectIterator<UClass> It ; It ; ++It)
-	{
-		UClass* CurrentClass = *It;
-		if (CurrentClass->IsChildOf(UFactory::StaticClass()) && !(CurrentClass->HasAnyClassFlags(CLASS_Abstract)))
-		{
-			UFactory* Factory = Cast<UFactory>(CurrentClass->GetDefaultObject());
-			if (Factory->CanCreateNew() && Factory->ImportPriority >= 0 && Factory->SupportedClass == ULevelSequence::StaticClass())
-			{
-				if (ShotToDuplicate != nullptr)
-				{
-					NewAsset = AssetTools.DuplicateAssetWithDialog(NewShotName, NewShotPath, ShotToDuplicate->GetSequence());
-				}
-				else
-				{
-					NewAsset = AssetTools.CreateAssetWithDialog(NewShotName, NewShotPath, ULevelSequence::StaticClass(), Factory);
-				}
-				break;
-			}
-		}
-	}
-
-	if (NewAsset == nullptr)
-	{
-		return nullptr;
-	}
-
-	UMovieSceneSequence* NewSequence = Cast<UMovieSceneSequence>(NewAsset);
-
-	int32 Duration = UE::MovieScene::DiscreteSize(ShotToDuplicate ? ShotToDuplicate->GetRange() : NewSequence->GetMovieScene()->GetPlaybackRange());
-
-	UMovieSceneCinematicShotTrack* CinematicShotTrack = FindOrCreateCinematicShotTrack();
-
-	// Create a cinematic shot section. 
-	UMovieSceneSubSection* NewSection = CinematicShotTrack->AddSequence(NewSequence, NewShotStartTime, Duration);
-	return NewSection;
-}
-
 void FCinematicShotTrackEditor::InsertShot()
 {
 	const FScopedTransaction Transaction(LOCTEXT("InsertShot_Transaction", "Insert Shot"));
@@ -430,8 +375,9 @@ void FCinematicShotTrackEditor::InsertShot()
 
 	UMovieSceneCinematicShotTrack* CinematicShotTrack = FindOrCreateCinematicShotTrack();
 	FString NewShotName = MovieSceneToolHelpers::GenerateNewShotName(CinematicShotTrack->GetAllSections(), NewShotStartTime.FrameNumber);
+	FString NewSequencePath = MovieSceneToolHelpers::GenerateNewShotPath(GetSequencer()->GetFocusedMovieSceneSequence()->GetMovieScene(), NewShotName);
 
-	UMovieSceneSubSection* NewShot = CreateShotInternal(NewShotName, NewShotStartTime.FrameNumber);
+	UMovieSceneSubSection* NewShot = MovieSceneToolHelpers::CreateSubSequence(NewShotName, NewSequencePath, NewShotStartTime.FrameNumber, CinematicShotTrack);
 	if (NewShot)
 	{
 		NewShot->SetRowIndex(MovieSceneToolHelpers::FindAvailableRowIndex(CinematicShotTrack, NewShot));
@@ -479,9 +425,10 @@ void FCinematicShotTrackEditor::DuplicateShot(UMovieSceneCinematicShotSection* S
 
 	FFrameNumber StartTime = Section->HasStartFrame() ? Section->GetInclusiveStartFrame() : 0;
 	FString NewShotName = MovieSceneToolHelpers::GenerateNewShotName(CinematicShotTrack->GetAllSections(), StartTime);
+	FString NewSequencePath = FPaths::GetPath(Section->GetSequence()->GetPathName());
 
 	// Duplicate the shot and put it on the next available row
-	UMovieSceneSubSection* NewShot = CreateShotInternal(NewShotName, StartTime, Section);
+	UMovieSceneSubSection* NewShot = MovieSceneToolHelpers::CreateSubSequence(NewShotName, NewSequencePath, StartTime, CinematicShotTrack, Section);
 	if (NewShot)
 	{
 		NewShot->SetRange(Section->GetRange());
@@ -545,11 +492,13 @@ void FCinematicShotTrackEditor::NewTake(UMovieSceneCinematicShotSection* Section
 		int32                NewRowIndex          = Section->GetRowIndex();
 		FFrameNumber         NewShotStartTime     = NewShotRange.GetLowerBound().IsClosed() ? UE::MovieScene::DiscreteInclusiveLower(NewShotRange) : 0;
 
-		UMovieSceneSubSection* NewShot = CreateShotInternal(NewShotName, NewShotStartTime, Section);
+		UMovieSceneCinematicShotTrack* CinematicShotTrack = FindOrCreateCinematicShotTrack();
+		FString NewSequencePath = FPaths::GetPath(Section->GetSequence()->GetPathName());
+
+		UMovieSceneSubSection* NewShot = MovieSceneToolHelpers::CreateSubSequence(NewShotName, NewSequencePath, NewShotStartTime, CinematicShotTrack, Section);
 
 		if (NewShot)
 		{
-			UMovieSceneCinematicShotTrack* CinematicShotTrack = FindOrCreateCinematicShotTrack();
 			CinematicShotTrack->RemoveSection(*Section);
 
 			NewShot->SetRange(NewShotRange);
diff --git a/Engine/Source/Editor/MovieSceneTools/Private/TrackEditors/SubTrackEditor.cpp b/Engine/Source/Editor/MovieSceneTools/Private/TrackEditors/SubTrackEditor.cpp
--- a/Engine/Source/Editor/MovieSceneTools/Private/TrackEditors/SubTrackEditor.cpp
+++ b/Engine/Source/Editor/MovieSceneTools/Private/TrackEditors/SubTrackEditor.cpp
@@ -28,6 +28,7 @@
 #include "EngineAnalytics.h"
 #include "Interfaces/IAnalyticsProvider.h"
 #include "Algo/Accumulate.h"
+#include "AssetToolsModule.h"
 
 #include "CommonMovieSceneTools.h"
 
@@ -494,6 +495,13 @@ TSharedRef<SWidget> FSubTrackEditor::HandleAddSubSequenceComboButtonGetMenuConte
 {
 	FMenuBuilder MenuBuilder(true, nullptr);
 
+	MenuBuilder.AddMenuEntry(
+		LOCTEXT("InsertSequence", "Insert Sequence"),
+		LOCTEXT("InsertSequenceTooltip", "Insert new sequence at current time"),
+		FSlateIcon(),
+		FUIAction(FExecuteAction::CreateSP(this, &FSubTrackEditor::InsertSequence, InTrack))
+	);
+
 	MenuBuilder.BeginSection(TEXT("ChooseSequence"), LOCTEXT("ChooseSequence", "Choose Sequence"));
 	{
 		UMovieSceneSequence* Sequence = GetSequencer() ? GetSequencer()->GetFocusedMovieSceneSequence() : nullptr;
@@ -548,6 +556,36 @@ void FSubTrackEditor::HandleAddSubSequenceComboButtonMenuEntryEnterPressed(const
 	}
 }
 
+void FSubTrackEditor::InsertSequence(UMovieSceneTrack* Track)
+{
+	const FScopedTransaction Transaction(LOCTEXT("InsertSequence_Transaction", "Insert Sequence"));
+
+	FFrameTime NewSectionStartTime = GetSequencer()->GetLocalTime().Time;
+
+	UMovieSceneSubTrack* SubTrack = Cast<UMovieSceneSubTrack>(Track);
+	if (!SubTrack)
+	{
+		SubTrack = FindOrCreateMasterTrack<UMovieSceneSubTrack>().Track;
+	}
+
+	FString NewSequencePath = FPaths::GetPath(GetSequencer()->GetFocusedMovieSceneSequence()->GetPathName());
+	FString NewSequenceName = TEXT("NewSubSequence");
+
+	FAssetToolsModule& AssetToolsModule = FModuleManager::GetModuleChecked<FAssetToolsModule>("AssetTools");
+	AssetToolsModule.Get().CreateUniqueAssetName(NewSequencePath + TEXT("/") + NewSequenceName, TEXT(""), NewSequencePath, NewSequenceName);
+
+	UMovieSceneSubSection* NewSection = MovieSceneToolHelpers::CreateSubSequence(NewSequenceName, NewSequencePath, NewSectionStartTime.FrameNumber, SubTrack);
+	if (NewSection)
+	{
+		NewSection->SetRowIndex(MovieSceneToolHelpers::FindAvailableRowIndex(Track, NewSection));
+	}
+
+	GetSequencer()->NotifyMovieSceneDataChanged(EMovieSceneDataChangeType::MovieSceneStructureItemAdded);
+	GetSequencer()->EmptySelection();
+	GetSequencer()->SelectSection(NewSection);
+	GetSequencer()->ThrobSectionSelection();
+}
+
 FKeyPropertyResult FSubTrackEditor::AddKeyInternal(FFrameNumber KeyTime, UMovieSceneSequence* InMovieSceneSequence, UMovieSceneTrack* InTrack, int32 RowIndex)
 {	
 	FKeyPropertyResult KeyPropertyResult;
diff --git a/Engine/Source/Editor/MovieSceneTools/Public/MovieSceneToolHelpers.h b/Engine/Source/Editor/MovieSceneTools/Public/MovieSceneToolHelpers.h
--- a/Engine/Source/Editor/MovieSceneTools/Public/MovieSceneToolHelpers.h
+++ b/Engine/Source/Editor/MovieSceneTools/Public/MovieSceneToolHelpers.h
@@ -38,6 +38,8 @@ enum class EVisibilityBasedAnimTickOption : uint8;
 class ACameraActor;
 struct FActorForWorldTransforms;
 class UMovieScene3DTransformSection;
+class UMovieSceneSubTrack;
+class UMovieSceneSubSection;
 enum class EMovieSceneTransformChannel : uint32;
 
 namespace fbxsdk
@@ -163,6 +165,18 @@ public:
 	 */
 	static FString GenerateNewShotName(const TArray<UMovieSceneSection*>& AllSections, FFrameNumber Time);
 
+	/*
+	 * Create sub sequence
+	 *
+	 * @param NewSequenceName The new sequence name.
+	 * @param NewSequencePath The new sequence path. 
+	 * @param NewSequenceStartTime The time to start the new sequence at.
+	 * @param SubTrack The track to put the sequence onto.
+	 * @param SectionToDuplicate The section to duplicate.
+	 * @return The new subsequence.
+	 */
+	static UMovieSceneSubSection* CreateSubSequence(FString& NewSequenceName, FString& NewSequencePath, FFrameNumber NewSequenceStartTime, UMovieSceneSubTrack* SubTrack, UMovieSceneSubSection* SectionToDuplicate = nullptr);
+
 	/**
 	 * Gather takes - level sequence assets that have the same shot prefix and shot number in the same asset path (directory)
 	 * 
@@ -172,7 +186,6 @@ public:
 	 */
 	static void GatherTakes(const UMovieSceneSection* Section, TArray<FAssetData>& AssetData, uint32& OutCurrentTakeNumber);
 
-
 	/**
 	 * Get the take number for the given asset
 	 *
diff --git a/Engine/Source/Editor/MovieSceneTools/Public/TrackEditors/CinematicShotTrackEditor.h b/Engine/Source/Editor/MovieSceneTools/Public/TrackEditors/CinematicShotTrackEditor.h
--- a/Engine/Source/Editor/MovieSceneTools/Public/TrackEditors/CinematicShotTrackEditor.h
+++ b/Engine/Source/Editor/MovieSceneTools/Public/TrackEditors/CinematicShotTrackEditor.h
@@ -111,16 +111,6 @@ public:
 
 private:
 
-	/*
-	 * Create shot
-	 *
-	 * @param NewShotName The new shot name.
-	 * @param NewShotStartTime The time to start the new shot at.
-	 * @param ShotToDuplicate The shot to duplicate.
-	 * @return The new shot.
-	 */
-	UMovieSceneSubSection* CreateShotInternal(FString& NewShotName, FFrameNumber NewShotStartTime, UMovieSceneCinematicShotSection* ShotToDuplicate = nullptr);
-	
 	/** Callback for determining whether the "Add Shot" menu entry can execute. */
 	bool HandleAddCinematicShotTrackMenuEntryCanExecute() const;
 
diff --git a/Engine/Source/Editor/MovieSceneTools/Public/TrackEditors/SubTrackEditor.h b/Engine/Source/Editor/MovieSceneTools/Public/TrackEditors/SubTrackEditor.h
--- a/Engine/Source/Editor/MovieSceneTools/Public/TrackEditors/SubTrackEditor.h
+++ b/Engine/Source/Editor/MovieSceneTools/Public/TrackEditors/SubTrackEditor.h
@@ -103,6 +103,8 @@ private:
 	/** Callback for executing a menu entry in the "Add Sequence" combo button when enter pressed. */
 	void HandleAddSubSequenceComboButtonMenuEntryEnterPressed(const TArray<FAssetData>& AssetData, UMovieSceneTrack* InTrack);
 
+	void InsertSequence(UMovieSceneTrack* Track);
+
 	/** Delegate for AnimatablePropertyChanged in AddKey */
 	FKeyPropertyResult AddKeyInternal(FFrameNumber KeyTime, UMovieSceneSequence* InMovieSceneSequence, UMovieSceneTrack* InTrack, int32 RowIndex);
 
