Sequencer: Drag drop support for data layer tracks

Allow drag and drop from the data layer window or data layer assets from the content browser onto existing sections.

diff --git a/Engine/Source/Editor/DataLayerEditor/Private/DataLayer/DataLayerDragDropOp.h b/Engine/Source/Editor/DataLayerEditor/Private/DataLayer/DataLayerDragDropOp.h
deleted file mode 100644
diff --git a/Engine/Source/Editor/DataLayerEditor/Private/DataLayer/DataLayerMode.h b/Engine/Source/Editor/DataLayerEditor/Private/DataLayer/DataLayerMode.h
--- a/Engine/Source/Editor/DataLayerEditor/Private/DataLayer/DataLayerMode.h
+++ b/Engine/Source/Editor/DataLayerEditor/Private/DataLayer/DataLayerMode.h
@@ -6,7 +6,7 @@
 #include "Containers/Map.h"
 #include "Containers/Set.h"
 #include "Containers/UnrealString.h"
-#include "DataLayerDragDropOp.h"
+#include "DataLayer/DataLayerDragDropOp.h"
 #include "Delegates/Delegate.h"
 #include "Folder.h"
 #include "Framework/Views/ITypedTableView.h"
diff --git a/Engine/Source/Editor/DataLayerEditor/Public/DataLayer/DataLayerDragDropOp.h b/Engine/Source/Editor/DataLayerEditor/Public/DataLayer/DataLayerDragDropOp.h
new file mode 100644
--- /dev/null
+++ b/Engine/Source/Editor/DataLayerEditor/Public/DataLayer/DataLayerDragDropOp.h
@@ -0,0 +1,51 @@
+// Copyright Epic Games, Inc. All Rights Reserved.
+
+#pragma once
+
+#include "Containers/Array.h"
+#include "Containers/Map.h"
+#include "Containers/UnrealString.h"
+#include "DragAndDrop/DecoratedDragDropOp.h"
+#include "Input/DragAndDrop.h"
+#include "UObject/NameTypes.h"
+#include "UObject/WeakObjectPtrTemplates.h"
+#include "WorldPartition/DataLayer/DataLayerInstance.h"
+
+class AActor;
+
+typedef TPair< TWeakObjectPtr<AActor>, TWeakObjectPtr<UDataLayerInstance>> FDataLayerActorMoveElement;
+
+/** Drag/drop operation for dragging layers in the editor */
+class FDataLayerDragDropOp : public FDecoratedDragDropOp
+{
+public:
+	DRAG_DROP_OPERATOR_TYPE(FDataLayerDragDropOp, FDecoratedDragDropOp)
+
+	struct FDragDropInfo
+	{
+		FDragDropInfo(const UDataLayerInstance* DataLayerInstance)
+			:DataLayerInstanceName(DataLayerInstance->GetDataLayerFName()),
+			DataLayerShortName(DataLayerInstance->GetDataLayerShortName())
+		{}
+
+		FName DataLayerInstanceName;
+		FString DataLayerShortName;
+	};
+
+	/** The labels of the layers being dragged */
+	TArray<FDragDropInfo> DataLayerDragDropInfos;
+
+	virtual void Construct() override;
+};
+
+/** Drag/drop operation for moving actor(s) in a data layer in the editor */
+class FDataLayerActorMoveOp : public FDecoratedDragDropOp
+{
+public:
+	DRAG_DROP_OPERATOR_TYPE(FDataLayerActorMoveOp, FDecoratedDragDropOp)
+
+	/** Actor that we are dragging */
+	TArray<FDataLayerActorMoveElement> DataLayerActorMoveElements;
+
+	virtual void Construct() override;
+};
\ No newline at end of file
diff --git a/Engine/Source/Editor/MovieSceneTools/Private/TrackEditors/DataLayerTrackEditor.cpp b/Engine/Source/Editor/MovieSceneTools/Private/TrackEditors/DataLayerTrackEditor.cpp
--- a/Engine/Source/Editor/MovieSceneTools/Private/TrackEditors/DataLayerTrackEditor.cpp
+++ b/Engine/Source/Editor/MovieSceneTools/Private/TrackEditors/DataLayerTrackEditor.cpp
@@ -12,18 +12,30 @@
 #include "SequencerUtilities.h"
 #include "LevelUtils.h"
 #include "MovieSceneTimeHelpers.h"
+#include "MovieSceneToolHelpers.h"
+#include "ISequencer.h"
 #include "ISequencerSection.h"
 #include "SequencerSectionPainter.h"
 #include "DataLayer/DataLayerEditorSubsystem.h"
 
+#include "WorldPartition/DataLayer/DataLayerAsset.h"
+#include "WorldPartition/DataLayer/DataLayerInstance.h"
+#include "WorldPartition/DataLayer/DataLayerInstanceWithAsset.h"
+
+#include "DragAndDrop/AssetDragDropOp.h"
+#include "Input/DragAndDrop.h"
+#include "DataLayer/DataLayerDragDropOp.h"
+#include "SDropTarget.h"
+
 #define LOCTEXT_NAMESPACE "DataLayerTrackEditor"
 
 struct FDataLayerSection
 	: public ISequencerSection
 	, public TSharedFromThis<FDataLayerSection>
 {
-	FDataLayerSection(UMovieSceneDataLayerSection* InSection)
-		: WeakSection(InSection)
+	FDataLayerSection(TSharedPtr<ISequencer> InSequencer, UMovieSceneDataLayerSection* InSection)
+		: SequencerPtr(InSequencer)
+		, WeakSection(InSection)
 	{}
 
 	/*~ ISequencerSection */
@@ -41,40 +53,47 @@ struct FDataLayerSection
 	}
 	virtual TSharedRef<SWidget> GenerateSectionWidget() override
 	{
-		return SNew(SBox)
-		.Padding(FMargin(4.f))
-		[
-			SNew(SVerticalBox)
-
-			+ SVerticalBox::Slot()
-			.AutoHeight()
+		return 
+			SNew(SDropTarget)
+			.OnAllowDrop(this, &FDataLayerSection::OnAllowDrop)
+			.OnDropped(this, &FDataLayerSection::OnDrop)
+			.Content()
 			[
-				SNew(SHorizontalBox)
-
-				+ SHorizontalBox::Slot()
-				.AutoWidth()
+				SNew(SBox)
+				.Padding(FMargin(4.f))
 				[
-					SNew(STextBlock)
-					.Text(this, &FDataLayerSection::GetVisibilityText)
-					.ColorAndOpacity(this, &FDataLayerSection::GetTextColor)
-					.TextStyle(FAppStyle::Get(), "NormalText.Important")
+					SNew(SVerticalBox)
+
+					+ SVerticalBox::Slot()
+					.AutoHeight()
+					[
+						SNew(SHorizontalBox)
+
+						+ SHorizontalBox::Slot()
+						.AutoWidth()
+						[
+							SNew(STextBlock)
+							.Text(this, &FDataLayerSection::GetVisibilityText)
+							.ColorAndOpacity(this, &FDataLayerSection::GetTextColor)
+							.TextStyle(FAppStyle::Get(), "NormalText.Important")
+						]
+
+						+ SHorizontalBox::Slot()
+						.AutoWidth()
+						[
+							SNew(STextBlock)
+							.Text(this, &FDataLayerSection::GetPrerollText)
+						]
+					]
+
+					+ SVerticalBox::Slot()
+					[
+						SNew(STextBlock)
+						.Text(this, &FDataLayerSection::GetLayerBarText)
+						.AutoWrapText(true)
+					]
 				]
-
-				+ SHorizontalBox::Slot()
-				.AutoWidth()
-				[
-					SNew(STextBlock)
-					.Text(this, &FDataLayerSection::GetPrerollText)
-				]
-			]
-
-			+ SVerticalBox::Slot()
-			[
-				SNew(STextBlock)
-				.Text(this, &FDataLayerSection::GetLayerBarText)
-				.AutoWrapText(true)
-			]
-		];
+			];
 	}
 
 	FText GetVisibilityText() const
@@ -157,8 +176,84 @@ struct FDataLayerSection
 		}
 		return FStyleColors::Foreground;
 	}
+
+	bool OnAllowDrop(TSharedPtr<FDragDropOperation> DragDropOperation)
+	{
+		if (DragDropOperation->IsOfType<FDataLayerDragDropOp>() && StaticCastSharedPtr<FDataLayerDragDropOp>(DragDropOperation)->DataLayerDragDropInfos.Num() > 0)
+		{
+			return true;
+		}
+		else if (DragDropOperation->IsOfType<FAssetDragDropOp>())
+		{
+			return true;
+		}
+		return false;
+	}
+
+	FReply OnDrop(const FGeometry& InGeometry, const FDragDropEvent& InDragDropEvent)
+	{
+		UMovieSceneDataLayerSection* Section = WeakSection.Get();
+		if (!Section)
+		{
+			return FReply::Handled();
+		}
+
+		UDataLayerEditorSubsystem* SubSystem = UDataLayerEditorSubsystem::Get();
+		if (!SubSystem)
+		{
+			return FReply::Handled();
+		}
+
+		TArray<UDataLayerAsset*> DataLayerAssets = Section->GetDataLayerAssets();
+		if (TSharedPtr<FDataLayerDragDropOp> DataLayerDragDropOp = InDragDropEvent.GetOperationAs<FDataLayerDragDropOp>())
+		{
+			if (DataLayerDragDropOp->DataLayerDragDropInfos.Num() > 0)
+			{
+				for (const FDataLayerDragDropOp::FDragDropInfo& DragDropInfo : DataLayerDragDropOp->DataLayerDragDropInfos)
+				{
+					if (UDataLayerInstance* DataLayerInstance = SubSystem->GetDataLayerInstance(DragDropInfo.DataLayerInstanceName))
+					{
+						UDataLayerInstanceWithAsset* DataLayerWithAsset = Cast<UDataLayerInstanceWithAsset>(DataLayerInstance);
+						UDataLayerAsset* DataLayerAsset = DataLayerWithAsset ? const_cast<UDataLayerAsset*>(DataLayerWithAsset->GetAsset()) : nullptr;
+
+						if (DataLayerAsset)
+						{
+							DataLayerAssets.AddUnique(DataLayerAsset);
+						}
+					}
+				}
+			}
+		}
+		else if (TSharedPtr<FAssetDragDropOp> AssetDragDropOp = InDragDropEvent.GetOperationAs<FAssetDragDropOp>())
+		{
+			UMovieSceneSequence* FocusedSequence = SequencerPtr.IsValid() ? SequencerPtr.Pin()->GetFocusedMovieSceneSequence() : nullptr;
+
+			for (const FAssetData& AssetData : AssetDragDropOp->GetAssets())
+			{
+				if (UDataLayerAsset* DataLayerAsset = Cast<UDataLayerAsset>(AssetData.GetAsset()))
+				{	
+					if (MovieSceneToolHelpers::IsValidAsset(FocusedSequence, AssetData))
+					{
+						DataLayerAssets.AddUnique(DataLayerAsset);
+					}
+				}
+			}
+		}
+		else
+		{	
+			return FReply::Handled();
+		}
+
+		FScopedTransaction Transaction(LOCTEXT("TransactionText", "Add Data Layer(s) to Data Layer Section"));
+		Section->Modify();
+		Section->SetDataLayerAssets(DataLayerAssets);
+
+		return FReply::Handled();
+	}
+
 private:
 
+	TWeakPtr<ISequencer> SequencerPtr;
 	TWeakObjectPtr<UMovieSceneDataLayerSection> WeakSection;
 };
 
@@ -192,7 +287,7 @@ TSharedRef<ISequencerSection> FDataLayerTrackEditor::MakeSectionInterface(UMovie
 	UMovieSceneDataLayerSection* DataLayerSection = Cast<UMovieSceneDataLayerSection>(&SectionObject);
 	check(SupportsType(SectionObject.GetOuter()->GetClass()) && DataLayerSection != nullptr);
 
-	return MakeShared<FDataLayerSection>(DataLayerSection);
+	return MakeShared<FDataLayerSection>(GetSequencer(), DataLayerSection);
 }
 
 void FDataLayerTrackEditor::BuildAddTrackMenu(FMenuBuilder& MenuBuilder)
