Sequencer: Fix shot and take number parsing so that it takes into account padding

diff --git a/Engine/Source/Editor/MovieSceneTools/Private/MovieSceneToolHelpers.cpp b/Engine/Source/Editor/MovieSceneTools/Private/MovieSceneToolHelpers.cpp
--- a/Engine/Source/Editor/MovieSceneTools/Private/MovieSceneToolHelpers.cpp
+++ b/Engine/Source/Editor/MovieSceneTools/Private/MovieSceneToolHelpers.cpp
@@ -225,7 +225,7 @@ void MovieSceneToolHelpers::SplitSection(const TSet<TWeakObjectPtr<UMovieSceneSe
 	}
 }
 
-bool MovieSceneToolHelpers::ParseShotName(const FString& ShotName, FString& ShotPrefix, uint32& ShotNumber, uint32& TakeNumber)
+bool MovieSceneToolHelpers::ParseShotName(const FString& ShotName, FString& ShotPrefix, uint32& ShotNumber, uint32& TakeNumber, uint32& ShotNumberDigits, uint32& TakeNumberDigits)
 {
 	// Parse a shot name
 	// 
@@ -291,6 +291,7 @@ bool MovieSceneToolHelpers::ParseShotName(const FString& ShotName, FString& Shot
 			if (ShotName[CharIndex] == ProjectSettings->TakeSeparator[0])
 			{
 				bFoundTakeSeparator = true;
+				ShotNumberDigits = LastShotNumberIndex - FirstShotNumberIndex + 1;
 			}
 		}
 	}
@@ -307,6 +308,7 @@ bool MovieSceneToolHelpers::ParseShotName(const FString& ShotName, FString& Shot
 		if (TakeStr.IsNumeric())
 		{
 			ParsedTakeNumber = FCString::Atoi(*TakeStr);
+			TakeNumberDigits = TakeStr.Len();
 		}
 	}
 
@@ -320,6 +322,7 @@ bool MovieSceneToolHelpers::ParseShotName(const FString& ShotName, FString& Shot
 			ShotPrefix = ShotName.Left(LastSlashPos);
 			ShotNumber = INDEX_NONE; // Nullify the shot number since we only have a shot prefix
 			TakeNumber = FCString::Atoi(*ShotName.RightChop(LastSlashPos+1));
+			TakeNumberDigits = LastSlashPos+1;
 			return true;
 		}
 	}
@@ -333,7 +336,7 @@ bool MovieSceneToolHelpers::ParseShotName(const FString& ShotName, FString& Shot
 }
 
 
-FString MovieSceneToolHelpers::ComposeShotName(const FString& ShotPrefix, uint32 ShotNumber, uint32 TakeNumber)
+FString MovieSceneToolHelpers::ComposeShotName(const FString& ShotPrefix, uint32 ShotNumber, uint32 TakeNumber, uint32 ShotNumberDigits, uint32 TakeNumberDigits)
 {
 	const UMovieSceneToolsProjectSettings* ProjectSettings = GetDefault<UMovieSceneToolsProjectSettings>();
 
@@ -341,15 +344,15 @@ FString MovieSceneToolHelpers::ComposeShotName(const FString& ShotPrefix, uint32
 
 	if (ShotNumber != INDEX_NONE)
 	{
-		ShotName += FString::Printf(TEXT("%0*d"), ProjectSettings->ShotNumDigits, ShotNumber);
+		ShotName += FString::Printf(TEXT("%0*d"), ShotNumberDigits, ShotNumber);
 	}
 
 	if (TakeNumber != INDEX_NONE)
 	{
-		FString TakeFormat = TEXT("%0") + FString::Printf(TEXT("%d"), ProjectSettings->TakeNumDigits) + TEXT("d");
+		FString TakeFormat = TEXT("%0") + FString::Printf(TEXT("%d"), TakeNumberDigits) + TEXT("d");
 		
 		ShotName += ProjectSettings->TakeSeparator;
-		ShotName += FString::Printf(TEXT("%0*d"), ProjectSettings->TakeNumDigits, TakeNumber);
+		ShotName += FString::Printf(TEXT("%0*d"), TakeNumberDigits, TakeNumber);
 	}
 	return ShotName;
 }
@@ -384,9 +387,11 @@ FString MovieSceneToolHelpers::GenerateNewShotPath(UMovieScene* SequenceMovieSce
 	FString NewShotPrefix;
 	uint32 NewShotNumber = INDEX_NONE;
 	uint32 NewTakeNumber = INDEX_NONE;
-	ParseShotName(NewShotName, NewShotPrefix, NewShotNumber, NewTakeNumber);
+	uint32 ShotNumberDigits = ProjectSettings->ShotNumDigits;
+	uint32 TakeNumberDigits = ProjectSettings->TakeNumDigits;
+	ParseShotName(NewShotName, NewShotPrefix, NewShotNumber, NewTakeNumber, ShotNumberDigits, TakeNumberDigits);
 
-	FString NewShotDirectory = ComposeShotName(NewShotPrefix, NewShotNumber, INDEX_NONE);
+	FString NewShotDirectory = ComposeShotName(NewShotPrefix, NewShotNumber, INDEX_NONE, ShotNumberDigits, TakeNumberDigits);
 	FString NewShotPath = SequencePath;
 
 	FString ShotDirectory = ProjectSettings->ShotDirectory;
@@ -404,8 +409,8 @@ FString MovieSceneToolHelpers::GenerateNewShotPath(UMovieScene* SequenceMovieSce
 		while (1)
 		{
 			NewShotNumber += ProjectSettings->ShotIncrement;
-			NewShotName = ComposeShotName(NewShotPrefix, NewShotNumber, NewTakeNumber);
-			NewShotDirectory = ComposeShotName(NewShotPrefix, NewShotNumber, INDEX_NONE);
+			NewShotName = ComposeShotName(NewShotPrefix, NewShotNumber, NewTakeNumber, ShotNumberDigits, TakeNumberDigits);
+			NewShotDirectory = ComposeShotName(NewShotPrefix, NewShotNumber, INDEX_NONE, ShotNumberDigits, TakeNumberDigits);
 			NewShotPath = SequencePath;
 			if (!ShotDirectory.IsEmpty())
 			{
@@ -469,11 +474,13 @@ FString MovieSceneToolHelpers::GenerateNewShotName(const TArray<UMovieSceneSecti
 		FString NextShotPrefix = ProjectSettings->ShotPrefix;
 		uint32 NextShotNumber = ProjectSettings->FirstShotNumber;
 		uint32 NextTakeNumber = ProjectSettings->FirstTakeNumber;
+		uint32 ShotNumberDigits = ProjectSettings->ShotNumDigits;
+		uint32 TakeNumberDigits = ProjectSettings->TakeNumDigits;
 
-		if (ParseShotName(NextShot->GetShotDisplayName(), NextShotPrefix, NextShotNumber, NextTakeNumber))
+		if (ParseShotName(NextShot->GetShotDisplayName(), NextShotPrefix, NextShotNumber, NextTakeNumber, ShotNumberDigits, TakeNumberDigits))
 		{
 			uint32 NewShotNumber = NextShotNumber + ProjectSettings->ShotIncrement;
-			return ComposeShotName(NextShotPrefix, NewShotNumber, ProjectSettings->FirstTakeNumber);
+			return ComposeShotName(NextShotPrefix, NewShotNumber, ProjectSettings->FirstTakeNumber, ShotNumberDigits, TakeNumberDigits);
 		}
 	}
 	// This is in between two shots
@@ -482,24 +489,28 @@ FString MovieSceneToolHelpers::GenerateNewShotName(const TArray<UMovieSceneSecti
 		FString BeforeShotPrefix = ProjectSettings->ShotPrefix;
 		uint32 BeforeShotNumber = ProjectSettings->FirstShotNumber;
 		uint32 BeforeTakeNumber = ProjectSettings->FirstTakeNumber;
+		uint32 BeforeShotNumberDigits = ProjectSettings->ShotNumDigits;
+		uint32 BeforeTakeNumberDigits = ProjectSettings->TakeNumDigits;
 
 		FString NextShotPrefix = ProjectSettings->ShotPrefix;
 		uint32 NextShotNumber = ProjectSettings->FirstShotNumber;
 		uint32 NextTakeNumber = ProjectSettings->FirstTakeNumber;
+		uint32 NextShotNumberDigits = ProjectSettings->ShotNumDigits;
+		uint32 NextTakeNumberDigits = ProjectSettings->TakeNumDigits;
 
-		if (ParseShotName(BeforeShot->GetShotDisplayName(), BeforeShotPrefix, BeforeShotNumber, BeforeTakeNumber) &&
-			ParseShotName(NextShot->GetShotDisplayName(), NextShotPrefix, NextShotNumber, NextTakeNumber))
+		if (ParseShotName(BeforeShot->GetShotDisplayName(), BeforeShotPrefix, BeforeShotNumber, BeforeTakeNumber, BeforeShotNumberDigits, BeforeTakeNumberDigits) &&
+			ParseShotName(NextShot->GetShotDisplayName(), NextShotPrefix, NextShotNumber, NextTakeNumber, NextShotNumberDigits, NextTakeNumberDigits))
 		{
 			if (BeforeShotNumber < NextShotNumber)
 			{
 				uint32 NewShotNumber = BeforeShotNumber + ( (NextShotNumber - BeforeShotNumber) / 2); // what if we can't find one? or conflicts with another?
-				return ComposeShotName(BeforeShotPrefix, NewShotNumber, ProjectSettings->FirstTakeNumber);
+				return ComposeShotName(BeforeShotPrefix, NewShotNumber, ProjectSettings->FirstTakeNumber, BeforeShotNumberDigits, BeforeTakeNumberDigits); // use before or next shot?
 			}
 		}
 	}
 
 	// Default case
-	return ComposeShotName(ProjectSettings->ShotPrefix, ProjectSettings->FirstShotNumber, ProjectSettings->FirstTakeNumber);
+	return ComposeShotName(ProjectSettings->ShotPrefix, ProjectSettings->FirstShotNumber, ProjectSettings->FirstTakeNumber, ProjectSettings->ShotNumDigits, ProjectSettings->TakeNumDigits);
 }
 
 UMovieSceneSubSection* MovieSceneToolHelpers::CreateSubSequence(FString& NewSequenceName, FString& NewSequencePath, FFrameNumber NewSequenceStartTime, UMovieSceneSubTrack* SubTrack, UMovieSceneSubSection* SectionToDuplicate)
@@ -563,6 +574,8 @@ void MovieSceneToolHelpers::GatherTakes(const UMovieSceneSection* Section, TArra
 	FString ShotPrefix;
 	uint32 ShotNumber = INDEX_NONE;
 	OutCurrentTakeNumber = INDEX_NONE;
+	uint32 ShotNumberDigits = 0;
+	uint32 TakeNumberDigits = 0;
 
 	FString SubSectionName = SubSection->GetSequence()->GetName();
 	if (SubSection->IsA<UMovieSceneCinematicShotSection>())
@@ -571,7 +584,7 @@ void MovieSceneToolHelpers::GatherTakes(const UMovieSceneSection* Section, TArra
 		SubSectionName = ShotSection->GetShotDisplayName();
 	}
 
-	if (ParseShotName(SubSectionName, ShotPrefix, ShotNumber, OutCurrentTakeNumber))
+	if (ParseShotName(SubSectionName, ShotPrefix, ShotNumber, OutCurrentTakeNumber, ShotNumberDigits, TakeNumberDigits))
 	{
 		// Gather up all level sequence assets
 		FAssetRegistryModule& AssetRegistryModule = FModuleManager::LoadModuleChecked<FAssetRegistryModule>(TEXT("AssetRegistry"));
@@ -588,7 +601,7 @@ void MovieSceneToolHelpers::GatherTakes(const UMovieSceneSection* Section, TArra
 				uint32 AssetShotNumber = INDEX_NONE;
 				uint32 AssetTakeNumber = INDEX_NONE;
 
-				if (ParseShotName(AssetObject.AssetName.ToString(), AssetShotPrefix, AssetShotNumber, AssetTakeNumber))
+				if (ParseShotName(AssetObject.AssetName.ToString(), AssetShotPrefix, AssetShotNumber, AssetTakeNumber, ShotNumberDigits, TakeNumberDigits))
 				{
 					if (AssetShotPrefix == ShotPrefix && AssetShotNumber == ShotNumber)
 					{
@@ -619,6 +632,8 @@ bool MovieSceneToolHelpers::GetTakeNumber(const UMovieSceneSection* Section, FAs
 	FString ShotPrefix;
 	uint32 ShotNumber = INDEX_NONE;
 	uint32 TakeNumberDummy = INDEX_NONE;
+	uint32 ShotNumberDigits = 0;
+	uint32 TakeNumberDigits = 0;
 
 	FString SubSectionName = SubSection->GetSequence()->GetName();
 	if (SubSection->IsA<UMovieSceneCinematicShotSection>())
@@ -627,7 +642,7 @@ bool MovieSceneToolHelpers::GetTakeNumber(const UMovieSceneSection* Section, FAs
 		SubSectionName = ShotSection->GetShotDisplayName();
 	}
 
-	if (ParseShotName(SubSectionName, ShotPrefix, ShotNumber, TakeNumberDummy))
+	if (ParseShotName(SubSectionName, ShotPrefix, ShotNumber, TakeNumberDummy, ShotNumberDigits, TakeNumberDigits))
 	{
 		// Gather up all level sequence assets
 		FAssetRegistryModule& AssetRegistryModule = FModuleManager::LoadModuleChecked<FAssetRegistryModule>(TEXT("AssetRegistry"));
@@ -649,7 +664,7 @@ bool MovieSceneToolHelpers::GetTakeNumber(const UMovieSceneSection* Section, FAs
 					uint32 AssetShotNumber = INDEX_NONE;
 					uint32 AssetTakeNumber = INDEX_NONE;
 
-					if (ParseShotName(AssetObject.AssetName.ToString(), AssetShotPrefix, AssetShotNumber, AssetTakeNumber))
+					if (ParseShotName(AssetObject.AssetName.ToString(), AssetShotPrefix, AssetShotNumber, AssetTakeNumber, ShotNumberDigits, TakeNumberDigits))
 					{
 						if (AssetShotPrefix == ShotPrefix && AssetShotNumber == ShotNumber)
 						{
diff --git a/Engine/Source/Editor/MovieSceneTools/Private/TrackEditors/CinematicShotTrackEditor.cpp b/Engine/Source/Editor/MovieSceneTools/Private/TrackEditors/CinematicShotTrackEditor.cpp
--- a/Engine/Source/Editor/MovieSceneTools/Private/TrackEditors/CinematicShotTrackEditor.cpp
+++ b/Engine/Source/Editor/MovieSceneTools/Private/TrackEditors/CinematicShotTrackEditor.cpp
@@ -464,7 +464,9 @@ void FCinematicShotTrackEditor::NewTake(UMovieSceneCinematicShotSection* Section
 	FString ShotPrefix;
 	uint32 ShotNumber = INDEX_NONE;
 	uint32 TakeNumber = INDEX_NONE;
-	if (MovieSceneToolHelpers::ParseShotName(Section->GetShotDisplayName(), ShotPrefix, ShotNumber, TakeNumber))
+	uint32 ShotNumberDigits = 0;
+	uint32 TakeNumberDigits = 0;
+	if (MovieSceneToolHelpers::ParseShotName(Section->GetShotDisplayName(), ShotPrefix, ShotNumber, TakeNumber, ShotNumberDigits, TakeNumberDigits))
 	{
 		TArray<FAssetData> AssetData;
 		uint32 CurrentTakeNumber = INDEX_NONE;
@@ -483,7 +485,7 @@ void FCinematicShotTrackEditor::NewTake(UMovieSceneCinematicShotSection* Section
 			}
 		}
 
-		FString NewShotName = MovieSceneToolHelpers::ComposeShotName(ShotPrefix, ShotNumber, NewTakeNumber);
+		FString NewShotName = MovieSceneToolHelpers::ComposeShotName(ShotPrefix, ShotNumber, NewTakeNumber, ShotNumberDigits, TakeNumberDigits);
 
 		TRange<FFrameNumber> NewShotRange         = Section->GetRange();
 		FFrameNumber         NewShotStartOffset   = Section->Parameters.StartFrameOffset;
diff --git a/Engine/Source/Editor/MovieSceneTools/Public/MovieSceneToolHelpers.h b/Engine/Source/Editor/MovieSceneTools/Public/MovieSceneToolHelpers.h
--- a/Engine/Source/Editor/MovieSceneTools/Public/MovieSceneToolHelpers.h
+++ b/Engine/Source/Editor/MovieSceneTools/Public/MovieSceneToolHelpers.h
@@ -133,9 +133,11 @@ public:
 	 * @param ShotPrefix The parsed shot prefix
 	 * @param ShotNumber The parsed shot number
 	 * @param TakeNumber The parsed take number
+	 * @param ShotNumberDigits The number of digits to pad for the shot number
+	 * @param TakeNumberDigits The number of digits to pad for the take number
 	 * @return Whether the shot name was parsed successfully
 	 */
-	static bool ParseShotName(const FString& ShotName, FString& ShotPrefix, uint32& ShotNumber, uint32& TakeNumber);
+	static bool ParseShotName(const FString& ShotName, FString& ShotPrefix, uint32& ShotNumber, uint32& TakeNumber, uint32& ShotNumberDigits, uint32& TakeNumberDigits);
 
 	/**
 	 * Compose a shot name given its components.
@@ -143,9 +145,11 @@ public:
 	 * @param ShotPrefix The shot prefix to use
 	 * @param ShotNumber The shot number to use
 	 * @param TakeNumber The take number to use
+	 * @param ShotNumberDigits The number of digits to pad for the shot number
+	 * @param TakeNumberDigits The number of digits to pad for the take number
 	 * @return The composed shot name
 	 */
-	static FString ComposeShotName(const FString& ShotPrefix, uint32 ShotNumber, uint32 TakeNumber);
+	static FString ComposeShotName(const FString& ShotPrefix, uint32 ShotNumber, uint32 TakeNumber, uint32 ShotNumberDigits, uint32 TakeNumberDigits);
 
 	/**
 	 * Generate a new shot package
