Sequencer: Constraints: Keep track of all sections, that we have added delegates to, and clear them out. Also make undo sections into weak pointers just in case.

diff --git a/Engine/Source/Editor/MovieSceneTools/Private/TrackEditors/TransformTrackEditor.cpp b/Engine/Source/Editor/MovieSceneTools/Private/TrackEditors/TransformTrackEditor.cpp
--- a/Engine/Source/Editor/MovieSceneTools/Private/TrackEditors/TransformTrackEditor.cpp
+++ b/Engine/Source/Editor/MovieSceneTools/Private/TrackEditors/TransformTrackEditor.cpp
@@ -1389,14 +1389,14 @@ bool F3DTransformTrackEditor::MatchesContext(const FTransactionContext& InContex
 
 void F3DTransformTrackEditor::PostUndo(bool bSuccess)
 {
-	for (UMovieScene3DTransformSection* Section : SectionsGettingUndone)
+	for (TWeakObjectPtr<UMovieScene3DTransformSection> &Section : SectionsGettingUndone)
 	{
-		if (Section)
+		if (Section.IsValid())
 		{
 			TArray<FConstraintAndActiveChannel>& ConstraintChannels = Section->GetConstraintsChannels();
 			for (FConstraintAndActiveChannel& Channel : ConstraintChannels)
 			{
-				HandleOnConstraintAdded(Section, &(Channel.ActiveChannel));
+				HandleOnConstraintAdded(Section.Get(), &(Channel.ActiveChannel));
 			}
 		}
 	}
@@ -1532,14 +1532,15 @@ void F3DTransformTrackEditor::HandleConstraintKeyMoved(IMovieSceneConstrainedSec
 	}
 }
 
-void F3DTransformTrackEditor::HandleConstraintRemoved(IMovieSceneConstrainedSection* InSection) const
+void F3DTransformTrackEditor::HandleConstraintRemoved(IMovieSceneConstrainedSection* InSection) 
 {
 	UWorld* World = GCurrentLevelEditingViewportClient ? GCurrentLevelEditingViewportClient->GetWorld() : nullptr;
-	UMovieSceneSection* Section = Cast<UMovieSceneSection>(InSection);
+	UMovieScene3DTransformSection* Section = Cast<UMovieScene3DTransformSection>(InSection);
 
 	FConstraintsManagerController& Controller = FConstraintsManagerController::Get(World);
 	if (!InSection->OnConstraintRemovedHandle.IsValid())
 	{
+		SectionsToClear.Add(Section);
 		InSection->OnConstraintRemovedHandle =
 			Controller.GetNotifyDelegate().AddLambda([InSection,Section, this](EConstraintsManagerNotifyType InNotifyType, UObject *InObject)
 				{
@@ -1584,45 +1585,29 @@ void F3DTransformTrackEditor::HandleConstraintRemoved(IMovieSceneConstrainedSect
 	}
 }
 
-void F3DTransformTrackEditor::ClearOutConstraintDelegates() const
+void F3DTransformTrackEditor::ClearOutConstraintDelegates() 
 {
-	const UMovieScene* MovieScene = GetSequencer().IsValid() && GetSequencer()->GetFocusedMovieSceneSequence() ? GetSequencer()->GetFocusedMovieSceneSequence()->GetMovieScene() : nullptr;
-	if (!MovieScene)
-	{
-		return;
-	}
-
 	UWorld* World = GCurrentLevelEditingViewportClient ? GCurrentLevelEditingViewportClient->GetWorld() : nullptr;
 	FConstraintsManagerController& Controller = FConstraintsManagerController::Get(World);
-	
-	const TArray<FMovieSceneBinding>& Bindings = MovieScene->GetBindings();
-	for (const FMovieSceneBinding& Binding : Bindings)
+	for (TWeakObjectPtr<UMovieScene3DTransformSection>& Section : SectionsToClear)
 	{
-		const UMovieSceneTrack* Track = MovieScene->FindTrack(
-			UMovieScene3DTransformTrack::StaticClass(), Binding.GetObjectGuid(), NAME_None);
-		if (const UMovieScene3DTransformTrack* CRTrack = Cast<UMovieScene3DTransformTrack>(Track))
+		if (IMovieSceneConstrainedSection* CRSection = Section.Get())
 		{
-			for (UMovieSceneSection* Section : Track->GetAllSections())
+			// clear constraint channels
+			TArray<FConstraintAndActiveChannel>& ConstraintChannels = CRSection->GetConstraintsChannels();
+			for (FConstraintAndActiveChannel& Channel : ConstraintChannels)
 			{
-				if (UMovieScene3DTransformSection* CRSection = Cast<UMovieScene3DTransformSection>(Section))
-				{
-					// clear constraint channels
-					TArray<FConstraintAndActiveChannel>& ConstraintChannels = CRSection->GetConstraintsChannels();
-					for (FConstraintAndActiveChannel& Channel : ConstraintChannels)
-					{
-						Channel.ActiveChannel.OnKeyMovedEvent().Clear();
-						Channel.ActiveChannel.OnKeyDeletedEvent().Clear();
-					}
+				Channel.ActiveChannel.OnKeyMovedEvent().Clear();
+				Channel.ActiveChannel.OnKeyDeletedEvent().Clear();
+			}
 
-					if (CRSection->OnConstraintRemovedHandle.IsValid())
-					{
-						Controller.GetNotifyDelegate().Remove(CRSection->OnConstraintRemovedHandle);
-						CRSection->OnConstraintRemovedHandle.Reset();
-					}
-				}
+			if (CRSection->OnConstraintRemovedHandle.IsValid())
+			{
+				Controller.GetNotifyDelegate().Remove(CRSection->OnConstraintRemovedHandle);
+				CRSection->OnConstraintRemovedHandle.Reset();
 			}
 		}
 	}
+	SectionsToClear.Reset();
 }
-
 #undef LOCTEXT_NAMESPACE
diff --git a/Engine/Source/Editor/MovieSceneTools/Private/TrackEditors/TransformTrackEditor.h b/Engine/Source/Editor/MovieSceneTools/Private/TrackEditors/TransformTrackEditor.h
--- a/Engine/Source/Editor/MovieSceneTools/Private/TrackEditors/TransformTrackEditor.h
+++ b/Engine/Source/Editor/MovieSceneTools/Private/TrackEditors/TransformTrackEditor.h
@@ -186,8 +186,8 @@ private:
 		const TArray<FKeyAddOrDeleteEventItem>& InDeletedItems) const;
 	void HandleConstraintKeyMoved(IMovieSceneConstrainedSection* InSection, const FMovieSceneConstraintChannel* InConstraintChannel,
 		const TArray<FKeyMoveEventItem>& InMovedItems);
-	void HandleConstraintRemoved(IMovieSceneConstrainedSection* InSection) const;
-	void ClearOutConstraintDelegates() const;
+	void HandleConstraintRemoved(IMovieSceneConstrainedSection* InSection);
+	void ClearOutConstraintDelegates();
 private:
 
 	static FName TransformPropertyName;
@@ -202,5 +202,8 @@ private:
 	TArray<FGuid> LockedCameraBindings;
 
 	/** Array of sections that are getting undone, we need to recreate any constraint channel add, move key delegates to them*/
-	mutable TArray<UMovieScene3DTransformSection*> SectionsGettingUndone;
+	mutable TArray<TWeakObjectPtr<UMovieScene3DTransformSection>> SectionsGettingUndone;
+
+	/** Set of sections we have added delegate's too, need to clear them*/
+	TSet<TWeakObjectPtr<UMovieScene3DTransformSection>> SectionsToClear;
 };
