Sequencer: Constraints: Need to keep track of the handle to the constraint deletion, since the section may be deleted already, not we don't need the channel delegates's cleared.

diff --git a/Engine/Source/Editor/MovieSceneTools/Private/TrackEditors/TransformTrackEditor.cpp b/Engine/Source/Editor/MovieSceneTools/Private/TrackEditors/TransformTrackEditor.cpp
--- a/Engine/Source/Editor/MovieSceneTools/Private/TrackEditors/TransformTrackEditor.cpp
+++ b/Engine/Source/Editor/MovieSceneTools/Private/TrackEditors/TransformTrackEditor.cpp
@@ -1540,7 +1540,6 @@ void F3DTransformTrackEditor::HandleConstraintRemoved(IMovieSceneConstrainedSect
 	FConstraintsManagerController& Controller = FConstraintsManagerController::Get(World);
 	if (!InSection->OnConstraintRemovedHandle.IsValid())
 	{
-		SectionsToClear.Add(Section);
 		InSection->OnConstraintRemovedHandle =
 			Controller.GetNotifyDelegate().AddLambda([InSection,Section, this](EConstraintsManagerNotifyType InNotifyType, UObject *InObject)
 				{
@@ -1582,6 +1581,8 @@ void F3DTransformTrackEditor::HandleConstraintRemoved(IMovieSceneConstrainedSect
 							break;		
 					}
 				});
+		ConstraintHandlesToClear.Add(InSection->OnConstraintRemovedHandle);
+
 	}
 }
 
@@ -1589,25 +1590,13 @@ void F3DTransformTrackEditor::ClearOutConstraintDelegates()
 {
 	UWorld* World = GCurrentLevelEditingViewportClient ? GCurrentLevelEditingViewportClient->GetWorld() : nullptr;
 	FConstraintsManagerController& Controller = FConstraintsManagerController::Get(World);
-	for (TWeakObjectPtr<UMovieScene3DTransformSection>& Section : SectionsToClear)
+	for (FDelegateHandle& Handle : ConstraintHandlesToClear)
 	{
-		if (IMovieSceneConstrainedSection* CRSection = Section.Get())
+		if (Handle.IsValid())
 		{
-			// clear constraint channels
-			TArray<FConstraintAndActiveChannel>& ConstraintChannels = CRSection->GetConstraintsChannels();
-			for (FConstraintAndActiveChannel& Channel : ConstraintChannels)
-			{
-				Channel.ActiveChannel.OnKeyMovedEvent().Clear();
-				Channel.ActiveChannel.OnKeyDeletedEvent().Clear();
-			}
-
-			if (CRSection->OnConstraintRemovedHandle.IsValid())
-			{
-				Controller.GetNotifyDelegate().Remove(CRSection->OnConstraintRemovedHandle);
-				CRSection->OnConstraintRemovedHandle.Reset();
-			}
+			Controller.GetNotifyDelegate().Remove(Handle);
 		}
 	}
-	SectionsToClear.Reset();
+	ConstraintHandlesToClear.Reset();
 }
 #undef LOCTEXT_NAMESPACE
diff --git a/Engine/Source/Editor/MovieSceneTools/Private/TrackEditors/TransformTrackEditor.h b/Engine/Source/Editor/MovieSceneTools/Private/TrackEditors/TransformTrackEditor.h
--- a/Engine/Source/Editor/MovieSceneTools/Private/TrackEditors/TransformTrackEditor.h
+++ b/Engine/Source/Editor/MovieSceneTools/Private/TrackEditors/TransformTrackEditor.h
@@ -204,6 +204,6 @@ private:
 	/** Array of sections that are getting undone, we need to recreate any constraint channel add, move key delegates to them*/
 	mutable TArray<TWeakObjectPtr<UMovieScene3DTransformSection>> SectionsGettingUndone;
 
-	/** Set of sections we have added delegate's too, need to clear them*/
-	TSet<TWeakObjectPtr<UMovieScene3DTransformSection>> SectionsToClear;
+	/** Set of delegate handles we have added delegate's too, need to clear them*/
+	TSet<FDelegateHandle> ConstraintHandlesToClear;
 };
