Sequencer: Control Rig:  Don't reset spaces to default on undo and don't set defaults when setting keys with control rig since that will cause a modify on the section which will trigger a sequencer eval on redo making control rig manipulations not-redoable as the sequencer eval will wipe out the redone value.  Control rig has it's own default system set up.

diff --git a/Engine/Plugins/Animation/ControlRig/Source/ControlRig/Private/Sequencer/MovieSceneControlRigParameterSection.cpp b/Engine/Plugins/Animation/ControlRig/Source/ControlRig/Private/Sequencer/MovieSceneControlRigParameterSection.cpp
--- a/Engine/Plugins/Animation/ControlRig/Source/ControlRig/Private/Sequencer/MovieSceneControlRigParameterSection.cpp
+++ b/Engine/Plugins/Animation/ControlRig/Source/ControlRig/Private/Sequencer/MovieSceneControlRigParameterSection.cpp
@@ -1540,19 +1540,6 @@ void UMovieSceneControlRigParameterSection::ReconstructChannelProxy()
 		ControlRig->GetControlsInOrder(SortedControls);
 		StoreLastControlsUsedToReconstruct(SortedControls);
 		bool bIsInUndo = false;
-#if WITH_EDITORONLY_DATA
-		bIsInUndo = GIsTransacting;
-		if (bIsInUndo) //if undo reset spaces @helge to do need to reset to some default not parent since rig may be set up to something else
-		{
-			if (URigHierarchy* RigHierarchy = ControlRig->GetHierarchy())
-			{
-				for (FRigControlElement* ControlElement : SortedControls)
-				{
-					RigHierarchy->SwitchToDefaultParent(ControlElement->GetKey());
-				}
-			}
-		}	
-#endif
 		if (ControlsMask.Num() != SortedControls.Num())
 		{
 			TArray<bool> OnArray;
diff --git a/Engine/Plugins/Animation/ControlRig/Source/ControlRigEditor/Private/Sequencer/ControlRigParameterTrackEditor.cpp b/Engine/Plugins/Animation/ControlRig/Source/ControlRigEditor/Private/Sequencer/ControlRigParameterTrackEditor.cpp
--- a/Engine/Plugins/Animation/ControlRig/Source/ControlRigEditor/Private/Sequencer/ControlRigParameterTrackEditor.cpp
+++ b/Engine/Plugins/Animation/ControlRig/Source/ControlRigEditor/Private/Sequencer/ControlRigParameterTrackEditor.cpp
@@ -3404,7 +3404,7 @@ FKeyPropertyResult FControlRigParameterTrackEditor::AddKeysToControlRigHandle(US
 			UMovieSceneControlRigParameterSection* ParamSection = Cast<UMovieSceneControlRigParameterSection>(SectionToKey);
 			if (!ParamSection->GetDoNotKey())
 			{
-				KeyPropertyResult |= AddKeysToSection(SectionToKey, KeyTime, GeneratedKeys, KeyMode);
+				KeyPropertyResult |= AddKeysToSection(SectionToKey, KeyTime, GeneratedKeys, KeyMode, FKeyframeTrackEditor::ESetDefault::DoNotSetDefault);
 			}
 		}
 
diff --git a/Engine/Source/Editor/MovieSceneTools/Public/KeyframeTrackEditor.h b/Engine/Source/Editor/MovieSceneTools/Public/KeyframeTrackEditor.h
--- a/Engine/Source/Editor/MovieSceneTools/Public/KeyframeTrackEditor.h
+++ b/Engine/Source/Editor/MovieSceneTools/Public/KeyframeTrackEditor.h
@@ -637,8 +637,15 @@ private:
 	}
 
 protected:
+	enum class ESetDefault
+	{
+		/** Do not set default. */
+		DoNotSetDefault,
+		/** Set default */
+		SetDefault,
+	};
 	/* Returns whether a section was added */
-	FKeyPropertyResult AddKeysToSection(UMovieSceneSection* Section, FFrameNumber KeyTime, const FGeneratedTrackKeys& Keys, ESequencerKeyMode KeyMode)
+	FKeyPropertyResult AddKeysToSection(UMovieSceneSection* Section, FFrameNumber KeyTime, const FGeneratedTrackKeys& Keys, ESequencerKeyMode KeyMode, ESetDefault SetDefault = ESetDefault::SetDefault)
 	{
 		FKeyPropertyResult KeyPropertyResult;
 
@@ -646,7 +653,7 @@ protected:
 
 		FMovieSceneChannelProxy& Proxy = Section->GetChannelProxy();
 			
-		const bool bSetDefaults = GetSequencer()->GetAutoSetTrackDefaults();
+		const bool bSetDefaults = GetSequencer()->GetAutoSetTrackDefaults() && (SetDefault == ESetDefault::SetDefault);
 
 		// The default value is a value for the channel when there are no keyframes. For example, if you add keys and 
 		// then delete them all, the default value is the value of the channel. In the implementation of ApplyDefault, 
