Sequencer: Remove Master from Sequencer

diff --git a/Engine/Plugins/Runtime/nDisplay/Source/DisplayClusterMoviePipeline/Private/DisplayClusterMoviePipelineViewportPass.cpp b/Engine/Plugins/Runtime/nDisplay/Source/DisplayClusterMoviePipeline/Private/DisplayClusterMoviePipelineViewportPass.cpp
--- a/Engine/Plugins/Runtime/nDisplay/Source/DisplayClusterMoviePipeline/Private/DisplayClusterMoviePipelineViewportPass.cpp
+++ b/Engine/Plugins/Runtime/nDisplay/Source/DisplayClusterMoviePipeline/Private/DisplayClusterMoviePipelineViewportPass.cpp
@@ -130,9 +130,9 @@ FIntPoint UDisplayClusterMoviePipelineViewportPassBase::GetEffectiveOutputResolu
 	}
 
 	// Use one size from OutputSettings
-	UMoviePipelineMasterConfig* MasterConfig = GetPipeline()->GetPipelineMasterConfig();
+	UMoviePipelinePrimaryConfig* PrimaryConfig = GetPipeline()->GetPipelinePrimaryConfig();
 	UMoviePipelineExecutorShot* CurrentShot = GetPipeline()->GetActiveShotList()[GetPipeline()->GetCurrentShotIndex()];
-	const FIntPoint OutputResolution = UMoviePipelineBlueprintLibrary::GetEffectiveOutputResolution(MasterConfig, CurrentShot);
+	const FIntPoint OutputResolution = UMoviePipelineBlueprintLibrary::GetEffectiveOutputResolution(PrimaryConfig, CurrentShot);
 
 	return OutputResolution;
 }
@@ -395,7 +395,7 @@ void UDisplayClusterMoviePipelineViewportPassBase::GetViewportCutOffset(const FM
 		const FMinimalViewInfo CameraCache = PlayerCameraManager->GetCameraCacheView();
 
 		// Taking overscan into account.
-		FIntPoint FullOutputSize = UMoviePipelineBlueprintLibrary::GetEffectiveOutputResolution(GetPipeline()->GetPipelineMasterConfig(), GetPipeline()->GetActiveShotList()[GetPipeline()->GetCurrentShotIndex()]);
+		FIntPoint FullOutputSize = UMoviePipelineBlueprintLibrary::GetEffectiveOutputResolution(GetPipeline()->GetPipelinePrimaryConfig(), GetPipeline()->GetActiveShotList()[GetPipeline()->GetCurrentShotIndex()]);
 
 		float OutputSizeAspectRatio = FullOutputSize.X / (float)FullOutputSize.Y;
 		const FIntPoint ConstrainedFullSize = CameraCache.AspectRatio > OutputSizeAspectRatio ?
diff --git a/Engine/Plugins/VirtualProduction/DMX/DMXEngine/Source/DMXEditor/Private/Sequencer/DMXLibraryTrackEditor.cpp b/Engine/Plugins/VirtualProduction/DMX/DMXEngine/Source/DMXEditor/Private/Sequencer/DMXLibraryTrackEditor.cpp
--- a/Engine/Plugins/VirtualProduction/DMX/DMXEngine/Source/DMXEditor/Private/Sequencer/DMXLibraryTrackEditor.cpp
+++ b/Engine/Plugins/VirtualProduction/DMX/DMXEngine/Source/DMXEditor/Private/Sequencer/DMXLibraryTrackEditor.cpp
@@ -381,7 +381,7 @@ void FDMXLibraryTrackEditor::AddDMXLibraryTrackToSequence(const FAssetData& InAs
 	}
 
 	// Attempt to find an existing DMX track that animates this object
-	for (UMovieSceneTrack* Track : MovieScene->GetMasterTracks())
+	for (UMovieSceneTrack* Track : MovieScene->GetTracks())
 	{
 		if (auto* DMXTrack = Cast<UMovieSceneDMXLibraryTrack>(Track))
 		{
@@ -395,7 +395,7 @@ void FDMXLibraryTrackEditor::AddDMXLibraryTrackToSequence(const FAssetData& InAs
 	const FScopedTransaction Transaction(LOCTEXT("AddDMXTrackTransaction", "Add DMX Library Track"));
 
 	MovieScene->Modify();
-	UMovieSceneDMXLibraryTrack* Track = MovieScene->AddMasterTrack<UMovieSceneDMXLibraryTrack>();
+	UMovieSceneDMXLibraryTrack* Track = MovieScene->AddTrack<UMovieSceneDMXLibraryTrack>();
 	check(Track);
 
 	Track->SetDMXLibrary(Library);
diff --git a/Engine/Plugins/VirtualProduction/DMX/DMXEngine/Source/DMXEditor/Private/Sequencer/MovieSceneDMXLibraryTrackRecorder.cpp b/Engine/Plugins/VirtualProduction/DMX/DMXEngine/Source/DMXEditor/Private/Sequencer/MovieSceneDMXLibraryTrackRecorder.cpp
--- a/Engine/Plugins/VirtualProduction/DMX/DMXEngine/Source/DMXEditor/Private/Sequencer/MovieSceneDMXLibraryTrackRecorder.cpp
+++ b/Engine/Plugins/VirtualProduction/DMX/DMXEngine/Source/DMXEditor/Private/Sequencer/MovieSceneDMXLibraryTrackRecorder.cpp
@@ -50,7 +50,7 @@ TWeakObjectPtr<UMovieSceneDMXLibraryTrack> UMovieSceneDMXLibraryTrackRecorder::C
 	MovieScene = InMovieScene;
 	bDiscardSamplesBeforeStart = bInDiscardSamplesBeforeStart;
 
-	DMXLibraryTrack = MovieScene->AddMasterTrack<UMovieSceneDMXLibraryTrack>();
+	DMXLibraryTrack = MovieScene->AddTrack<UMovieSceneDMXLibraryTrack>();
 	DMXLibraryTrack->SetDMXLibrary(Library);
 
 	DMXLibrarySection = Cast<UMovieSceneDMXLibrarySection>(DMXLibraryTrack->CreateNewSection());
diff --git a/Engine/Plugins/VirtualProduction/DMX/DMXEngine/Source/DMXEditor/Private/Sequencer/TakeRecorderDMXLibrarySource.cpp b/Engine/Plugins/VirtualProduction/DMX/DMXEngine/Source/DMXEditor/Private/Sequencer/TakeRecorderDMXLibrarySource.cpp
--- a/Engine/Plugins/VirtualProduction/DMX/DMXEngine/Source/DMXEditor/Private/Sequencer/TakeRecorderDMXLibrarySource.cpp
+++ b/Engine/Plugins/VirtualProduction/DMX/DMXEngine/Source/DMXEditor/Private/Sequencer/TakeRecorderDMXLibrarySource.cpp
@@ -60,11 +60,11 @@ void UTakeRecorderDMXLibrarySource::AddAllPatches()
 	}
 }
 
-TArray<UTakeRecorderSource*> UTakeRecorderDMXLibrarySource::PreRecording(ULevelSequence* InSequence, FMovieSceneSequenceID InSequenceID, ULevelSequence* InMasterSequence, FManifestSerializer* InManifestSerializer)
+TArray<UTakeRecorderSource*> UTakeRecorderDMXLibrarySource::PreRecording(ULevelSequence* InSequence, FMovieSceneSequenceID InSequenceID, ULevelSequence* InRootSequence, FManifestSerializer* InManifestSerializer)
 {
 	if (DMXLibrary)
 	{
-		UMovieScene* MovieScene = InMasterSequence->GetMovieScene();
+		UMovieScene* MovieScene = InRootSequence->GetMovieScene();
 		TrackRecorder = NewObject<UMovieSceneDMXLibraryTrackRecorder>();
 		CachedDMXLibraryTrack = TrackRecorder->CreateTrack(MovieScene, DMXLibrary, FixturePatchRefs, bDiscardSamplesBeforeStart, bRecordNormalizedValues);
 	}
@@ -100,7 +100,7 @@ void UTakeRecorderDMXLibrarySource::StopRecording(class ULevelSequence* InSequen
 	}
 }
 
-TArray<UTakeRecorderSource*> UTakeRecorderDMXLibrarySource::PostRecording(class ULevelSequence* InSequence, ULevelSequence* InMasterSequence, const bool bCancelled)
+TArray<UTakeRecorderSource*> UTakeRecorderDMXLibrarySource::PostRecording(class ULevelSequence* InSequence, ULevelSequence* InRootSequence, const bool bCancelled)
 {
 	if (TrackRecorder)
 	{
@@ -115,7 +115,7 @@ void UTakeRecorderDMXLibrarySource::AddContentsToFolder(UMovieSceneFolder* InFol
 {
 	if (CachedDMXLibraryTrack.IsValid())
 	{
-		InFolder->AddChildMasterTrack(CachedDMXLibraryTrack.Get());
+		InFolder->AddChildTrack(CachedDMXLibraryTrack.Get());
 	}
 }
 
diff --git a/Engine/Plugins/VirtualProduction/DMX/DMXEngine/Source/DMXEditor/Private/Sequencer/TakeRecorderDMXLibrarySource.h b/Engine/Plugins/VirtualProduction/DMX/DMXEngine/Source/DMXEditor/Private/Sequencer/TakeRecorderDMXLibrarySource.h
--- a/Engine/Plugins/VirtualProduction/DMX/DMXEngine/Source/DMXEditor/Private/Sequencer/TakeRecorderDMXLibrarySource.h
+++ b/Engine/Plugins/VirtualProduction/DMX/DMXEngine/Source/DMXEditor/Private/Sequencer/TakeRecorderDMXLibrarySource.h
@@ -73,11 +73,11 @@ public:
 
 private:
 	// ~Begin UTakeRecorderSource Interface
-	virtual TArray<UTakeRecorderSource*> PreRecording(ULevelSequence* InSequence, FMovieSceneSequenceID InSequenceID, ULevelSequence* InMasterSequence, FManifestSerializer* InManifestSerializer) override;
+	virtual TArray<UTakeRecorderSource*> PreRecording(ULevelSequence* InSequence, FMovieSceneSequenceID InSequenceID, ULevelSequence* InRootSequence, FManifestSerializer* InManifestSerializer) override;
 	virtual void StartRecording(const FTimecode& InSectionStartTimecode, const FFrameNumber& InSectionFirstFrame, class ULevelSequence* InSequence) override;
 	virtual void StopRecording(class ULevelSequence* InSequence) override;
 	virtual void TickRecording(const FQualifiedFrameTime& CurrentTime) override;
-	virtual TArray<UTakeRecorderSource*> PostRecording(class ULevelSequence* InSequence, ULevelSequence* InMasterSequence, const bool bCancelled) override;
+	virtual TArray<UTakeRecorderSource*> PostRecording(class ULevelSequence* InSequence, ULevelSequence* InRootSequence, const bool bCancelled) override;
 	virtual void AddContentsToFolder(UMovieSceneFolder* InFolder) override;
 	virtual FText GetDisplayTextImpl() const override;
 	// ~End UTakeRecorderSource
diff --git a/Engine/Plugins/VirtualProduction/SequencerPlaylists/Source/SequencerPlaylists/Private/SequencerPlaylistItem_Sequence.cpp b/Engine/Plugins/VirtualProduction/SequencerPlaylists/Source/SequencerPlaylists/Private/SequencerPlaylistItem_Sequence.cpp
--- a/Engine/Plugins/VirtualProduction/SequencerPlaylists/Source/SequencerPlaylists/Private/SequencerPlaylistItem_Sequence.cpp
+++ b/Engine/Plugins/VirtualProduction/SequencerPlaylists/Source/SequencerPlaylists/Private/SequencerPlaylistItem_Sequence.cpp
@@ -309,7 +309,7 @@ UMovieSceneSubTrack* FSequencerPlaylistItemPlayer_Sequence::GetOrCreateWorkingTr
 		return ItemState.WeakTrack.Get();
 	}
 
-	UMovieSceneSubTrack* NewWorkingTrack = RootScene->AddMasterTrack<UMovieSceneSubTrack>();
+	UMovieSceneSubTrack* NewWorkingTrack = RootScene->AddTrack<UMovieSceneSubTrack>();
 	NewWorkingTrack->SetDisplayName(FText::Format(LOCTEXT("SequenceItemTrackName", "Item - {0}"), SequenceItem->Sequence->GetDisplayName()));
 
 	// Find or create folder named for our playlist, and organize our track beneath it.
@@ -339,7 +339,7 @@ UMovieSceneSubTrack* FSequencerPlaylistItemPlayer_Sequence::GetOrCreateWorkingTr
 		//RootScene->GetRootFolders().Add(FolderToUse);
 	}
 
-	FolderToUse->AddChildMasterTrack(NewWorkingTrack);
+	FolderToUse->AddChildTrack(NewWorkingTrack);
 
 	ItemState.WeakTrack = NewWorkingTrack;
 	return NewWorkingTrack;
diff --git a/Engine/Plugins/VirtualProduction/Takes/Source/TakeRecorder/Private/SerializedRecorder.cpp b/Engine/Plugins/VirtualProduction/Takes/Source/TakeRecorder/Private/SerializedRecorder.cpp
--- a/Engine/Plugins/VirtualProduction/Takes/Source/TakeRecorder/Private/SerializedRecorder.cpp
+++ b/Engine/Plugins/VirtualProduction/Takes/Source/TakeRecorder/Private/SerializedRecorder.cpp
@@ -232,11 +232,11 @@ bool FSerializedRecorder::LoadSubSequenceFile(UMovieSceneSequence* InMovieSceneS
 	{
 		if (Serializer->OpenForRead(InFileName, Header, Error))
 		{
-			ULevelSequence* InMasterSequence = CastChecked<ULevelSequence>(InMovieSceneSequence);
+			ULevelSequence* InRootSequence = CastChecked<ULevelSequence>(InMovieSceneSequence);
 
-			ULevelSequence* TargetSequence = InMasterSequence;
+			ULevelSequence* TargetSequence = InRootSequence;
 			const FString& SubSequenceName = Header.Name;
-			TargetSequence = UTakeRecorderSources::CreateSubSequenceForSource(InMasterSequence, SubSequenceName, SubSequenceName);
+			TargetSequence = UTakeRecorderSources::CreateSubSequenceForSource(InRootSequence, SubSequenceName, SubSequenceName);
 			if (TargetSequence)
 			{
 				// If there's already a Subscene Track for our sub-sequence we need to remove that track before create a new one. No data is lost in this process as the
@@ -245,7 +245,7 @@ bool FSerializedRecorder::LoadSubSequenceFile(UMovieSceneSequence* InMovieSceneS
 				const FString SequenceName = FPaths::GetBaseFilename(TargetSequence->GetPathName());
 				UMovieSceneSubTrack* SubsceneTrack = nullptr;
 
-				for (UMovieSceneTrack* Track : InMasterSequence->GetMovieScene()->GetMasterTracks())
+				for (UMovieSceneTrack* Track : InRootSequence->GetMovieScene()->GetTracks())
 				{
 					if (Track->IsA<UMovieSceneSubTrack>())
 					{
@@ -267,11 +267,11 @@ bool FSerializedRecorder::LoadSubSequenceFile(UMovieSceneSequence* InMovieSceneS
 					}
 				}
 
-				// We need to add the new subsequence to the master sequence immediately so that it shows up in the UI and you can tell that things
+				// We need to add the new subsequence to the root sequence immediately so that it shows up in the UI and you can tell that things
 				// are being recorded, otherwise they don't show up until recording stops and then it magically pops in.
 				if (!SubsceneTrack)
 				{
-					SubsceneTrack = CastChecked<UMovieSceneSubTrack>(InMasterSequence->GetMovieScene()->AddMasterTrack(UMovieSceneSubTrack::StaticClass()));
+					SubsceneTrack = CastChecked<UMovieSceneSubTrack>(InRootSequence->GetMovieScene()->AddTrack(UMovieSceneSubTrack::StaticClass()));
 				}
 
 				// We create a new sub track for every Source so that we can name the Subtrack after the Source instead of just the sections within it.
@@ -281,8 +281,8 @@ bool FSerializedRecorder::LoadSubSequenceFile(UMovieSceneSequence* InMovieSceneS
 				/*
 				if (bCreateSequencerFolders)
 				{
-					UMovieSceneFolder* Folder = AddFolderForSource(Source, InMasterSequence->GetMovieScene());
-					Folder->AddChildMasterTrack(SubsceneTrack);
+					UMovieSceneFolder* Folder = AddFolderForSource(Source, InRootSequence->GetMovieScene());
+					Folder->AddChildTrack(SubsceneTrack);
 				}
 				*/
 
diff --git a/Engine/Plugins/VirtualProduction/Takes/Source/TakeRecorder/Public/Recorder/TakeRecorderParameters.h b/Engine/Plugins/VirtualProduction/Takes/Source/TakeRecorder/Public/Recorder/TakeRecorderParameters.h
--- a/Engine/Plugins/VirtualProduction/Takes/Source/TakeRecorder/Public/Recorder/TakeRecorderParameters.h
+++ b/Engine/Plugins/VirtualProduction/Takes/Source/TakeRecorder/Public/Recorder/TakeRecorderParameters.h
@@ -109,8 +109,8 @@ struct FTakeRecorderProjectParameters
 	bool bRecordTimecode;
 
 	/**
-	* If enabled, each Source will be recorded into a separate Sequence and embedded in the Master Sequence will link to them via Subscenes track.
-	* If disabled, all Sources will be recorded into the Master Sequence, and you will not be able to swap between various takes of specific source
+	* If enabled, each Source will be recorded into a separate Sequence and embedded in the Root Sequence will link to them via Subscenes track.
+	* If disabled, all Sources will be recorded into the Root Sequence, and you will not be able to swap between various takes of specific source
 	* using the Sequencer Take ui. This can still be done via copying and pasting between sequences if needed.
 	*/
 	UPROPERTY(config, EditAnywhere, BlueprintReadWrite, Category = "Take Recorder")
diff --git a/Engine/Plugins/VirtualProduction/Takes/Source/TakeRecorderSources/Private/TakeRecorderActorSource.cpp b/Engine/Plugins/VirtualProduction/Takes/Source/TakeRecorderSources/Private/TakeRecorderActorSource.cpp
--- a/Engine/Plugins/VirtualProduction/Takes/Source/TakeRecorderSources/Private/TakeRecorderActorSource.cpp
+++ b/Engine/Plugins/VirtualProduction/Takes/Source/TakeRecorderSources/Private/TakeRecorderActorSource.cpp
@@ -200,7 +200,7 @@ bool UTakeRecorderActorSource::IsValid() const
 }
 
 
-TArray<UTakeRecorderSource*> UTakeRecorderActorSource::PreRecording(ULevelSequence* InSequence, FMovieSceneSequenceID InSequenceID, ULevelSequence* InMasterSequence, FManifestSerializer* InManifestSerializer)
+TArray<UTakeRecorderSource*> UTakeRecorderActorSource::PreRecording(ULevelSequence* InSequence, FMovieSceneSequenceID InSequenceID, ULevelSequence* InRootSequence, FManifestSerializer* InManifestSerializer)
 {
 	// Don't bother doing anything if we don't have a valid actor to record.
 	if (!Target.IsValid())
@@ -213,7 +213,7 @@ TArray<UTakeRecorderSource*> UTakeRecorderActorSource::PreRecording(ULevelSequen
 	// Resolve which actor we wish to record 
 	AActor* ActorToRecord = Target.Get();
 	TargetLevelSequence = InSequence;
-	MasterLevelSequence = InMasterSequence;
+	RootLevelSequence = InRootSequence;
 	TargetSequenceID = InSequenceID;
 
 	FString ObjectBindingName = ActorToRecord->GetName();
@@ -772,7 +772,7 @@ void UTakeRecorderActorSource::ProcessRecordedTimes(ULevelSequence* InSequence)
 	}
 }
 
-TArray<UTakeRecorderSource*> UTakeRecorderActorSource::PostRecording(ULevelSequence* InSequence, class ULevelSequence* InMasterSequence, const bool bCancelled)
+TArray<UTakeRecorderSource*> UTakeRecorderActorSource::PostRecording(ULevelSequence* InSequence, class ULevelSequence* InRootSequence, const bool bCancelled)
 {
 	FTakeRecorderParameters Parameters;
 	Parameters.User = GetDefault<UTakeRecorderUserSettings>()->Settings;
@@ -825,7 +825,7 @@ TArray<UTakeRecorderSource*> UTakeRecorderActorSource::PostRecording(ULevelSeque
 			{
 				FGuid RecordedCameraGuid = GetRecordedActorGuid(Target.Get());
 				FMovieSceneSequenceID RecordedCameraSequenceID = GetLevelSequenceID(Target.Get());
-				TakesUtils::CreateCameraCutTrack(InMasterSequence, RecordedCameraGuid, RecordedCameraSequenceID, InSequence->GetMovieScene()->GetPlaybackRange());
+				TakesUtils::CreateCameraCutTrack(InRootSequence, RecordedCameraGuid, RecordedCameraSequenceID, InSequence->GetMovieScene()->GetPlaybackRange());
 			}
 
 			// Swap our target actor to the Editor actor (in case the recording was added while in PIE)
@@ -855,7 +855,7 @@ void UTakeRecorderActorSource::FinalizeRecording()
 {
 	// Null these out there and NOT in PostRecording because they are used for cross sequence object binding via GetLevelSequenceID in PostRecording
 	TargetLevelSequence = nullptr;
-	MasterLevelSequence = nullptr;
+	RootLevelSequence = nullptr;
 
 	ParentSource = nullptr;
 }
@@ -1686,7 +1686,7 @@ FTransform UTakeRecorderActorSource::GetRecordedActorAnimationInitialRootTransfo
 
 FMovieSceneSequenceID UTakeRecorderActorSource::GetLevelSequenceID(class AActor* OtherActor)
 {
-	return TakeRecorderSourcesUtils::GetLevelSequenceID(this, OtherActor, MasterLevelSequence);
+	return TakeRecorderSourcesUtils::GetLevelSequenceID(this, OtherActor, RootLevelSequence);
 }
 
 FTrackRecorderSettings UTakeRecorderActorSource::GetTrackRecorderSettings() const
diff --git a/Engine/Plugins/VirtualProduction/Takes/Source/TakeRecorderSources/Private/TakeRecorderActorSource.h b/Engine/Plugins/VirtualProduction/Takes/Source/TakeRecorderSources/Private/TakeRecorderActorSource.h
--- a/Engine/Plugins/VirtualProduction/Takes/Source/TakeRecorderSources/Private/TakeRecorderActorSource.h
+++ b/Engine/Plugins/VirtualProduction/Takes/Source/TakeRecorderSources/Private/TakeRecorderActorSource.h
@@ -80,9 +80,9 @@ public:
 	UPROPERTY()
 	TObjectPtr<ULevelSequence> TargetLevelSequence;
 
-	/** The master or uppermost level sequence that this source is being recorded into. Set during PreRecording, null after PostRecording. */
+	/** The root or uppermost level sequence that this source is being recorded into. Set during PreRecording, null after PostRecording. */
 	UPROPERTY()
-	TObjectPtr<ULevelSequence> MasterLevelSequence;
+	TObjectPtr<ULevelSequence> RootLevelSequence;
 
 	/**
 	* Dynamically created list of settings objects for the different factories that are recording something 
@@ -135,11 +135,11 @@ public:
 	
 	// UTakeRecorderSource Interface
 	virtual bool IsValid()const override;
-	virtual TArray<UTakeRecorderSource*> PreRecording(ULevelSequence* InSequence, FMovieSceneSequenceID InSequenceID, ULevelSequence* InMasterSequence, FManifestSerializer* InManifestSerializer) override;
+	virtual TArray<UTakeRecorderSource*> PreRecording(ULevelSequence* InSequence, FMovieSceneSequenceID InSequenceID, ULevelSequence* InRootSequence, FManifestSerializer* InManifestSerializer) override;
 	virtual void StartRecording(const FTimecode& InSectionStartTimecode, const FFrameNumber& InSectionFirstFrame, class ULevelSequence* InSequence) override;
 	virtual void TickRecording(const FQualifiedFrameTime& CurrentSequenceTime) override;
 	virtual void StopRecording(class ULevelSequence* InSequence) override;
-	virtual TArray<UTakeRecorderSource*> PostRecording(class ULevelSequence* InSequence, class ULevelSequence* InMasterSequence, const bool bCancelled) override;
+	virtual TArray<UTakeRecorderSource*> PostRecording(class ULevelSequence* InSequence, class ULevelSequence* InRootSequence, const bool bCancelled) override;
 	virtual void FinalizeRecording() override;
 	virtual TArray<UObject*> GetAdditionalSettingsObjects() const { return TArray<UObject*>(FactorySettings); }
 	virtual FString GetSubsceneTrackName(ULevelSequence* InSequence) const override;
@@ -264,8 +264,8 @@ protected:
 	FTrackRecorderSettings GetTrackRecorderSettings() const override;
 	/** Returns offset that may get set when recording a skeletal mesh animation. Needed to correctly transform attached children*/
 	FTransform GetRecordedActorAnimationInitialRootTransform(class AActor* OtherActor) const override;
-	/** Returns the master level sequence being recorded into */
-	ULevelSequence* GetMasterLevelSequence() const override { return MasterLevelSequence; };
+	/** Returns the root level sequence being recorded into */
+	ULevelSequence* GetRootLevelSequence() const override { return RootLevelSequence; };
 	// ~IMovieSceneTrackRecorderHost Interface
 	
 	/** Initializes an instance of the specified class if we don't already have it in our Settings array. */
diff --git a/Engine/Plugins/VirtualProduction/Takes/Source/TakeRecorderSources/Private/TakeRecorderCameraCutSource.cpp b/Engine/Plugins/VirtualProduction/Takes/Source/TakeRecorderSources/Private/TakeRecorderCameraCutSource.cpp
--- a/Engine/Plugins/VirtualProduction/Takes/Source/TakeRecorderSources/Private/TakeRecorderCameraCutSource.cpp
+++ b/Engine/Plugins/VirtualProduction/Takes/Source/TakeRecorderSources/Private/TakeRecorderCameraCutSource.cpp
@@ -28,10 +28,10 @@ UTakeRecorderCameraCutSource::UTakeRecorderCameraCutSource(const FObjectInitiali
 	TrackTint = FColor(160, 160, 160);
 }
 
-TArray<UTakeRecorderSource*> UTakeRecorderCameraCutSource::PreRecording(ULevelSequence* InSequence, FMovieSceneSequenceID InSequenceID, ULevelSequence* InMasterSequence, FManifestSerializer* InManifestSerializer)
+TArray<UTakeRecorderSource*> UTakeRecorderCameraCutSource::PreRecording(ULevelSequence* InSequence, FMovieSceneSequenceID InSequenceID, ULevelSequence* InRootSequence, FManifestSerializer* InManifestSerializer)
 {
 	World = TakeRecorderSourcesUtils::GetSourceWorld(InSequence);
-	MasterLevelSequence = InMasterSequence;
+	RootLevelSequence = InRootSequence;
 
 	return TArray<UTakeRecorderSource*>();
 }
@@ -90,7 +90,7 @@ void UTakeRecorderCameraCutSource::TickRecording(const FQualifiedFrameTime& Curr
 	}
 
 	FGuid RecordedCameraGuid = TakeRecorderSourcesUtils::GetRecordedActorGuid(this, Target);
-	FMovieSceneSequenceID RecordedCameraSequenceID = TakeRecorderSourcesUtils::GetLevelSequenceID(this, Target, MasterLevelSequence);
+	FMovieSceneSequenceID RecordedCameraSequenceID = TakeRecorderSourcesUtils::GetLevelSequenceID(this, Target, RootLevelSequence);
 
 	// If this camera is already noted, skip it until it changes
 	if (CameraCutData.Num() && CameraCutData.Last().Guid == RecordedCameraGuid && CameraCutData.Last().SequenceID == RecordedCameraSequenceID)
@@ -102,22 +102,22 @@ void UTakeRecorderCameraCutSource::TickRecording(const FQualifiedFrameTime& Curr
 }
 
 
-TArray<UTakeRecorderSource*> UTakeRecorderCameraCutSource::PostRecording(class ULevelSequence* InSequence, class ULevelSequence* InMasterSequence, const bool bCancelled)
+TArray<UTakeRecorderSource*> UTakeRecorderCameraCutSource::PostRecording(class ULevelSequence* InSequence, class ULevelSequence* InRootSequence, const bool bCancelled)
 {
 	// Build the camera cut track
 	if (!bCancelled && CameraCutData.Num())
 	{
-		UMovieSceneTrack* CameraCutTrack = InMasterSequence->GetMovieScene()->GetCameraCutTrack();
+		UMovieSceneTrack* CameraCutTrack = InRootSequence->GetMovieScene()->GetCameraCutTrack();
 		if (!CameraCutTrack)
 		{
-			CameraCutTrack = InMasterSequence->GetMovieScene()->AddCameraCutTrack(UMovieSceneCameraCutTrack::StaticClass());
+			CameraCutTrack = InRootSequence->GetMovieScene()->AddCameraCutTrack(UMovieSceneCameraCutTrack::StaticClass());
 		}
 		else
 		{
 			CameraCutTrack->RemoveAllAnimationData();
 		}
 
- 		FFrameRate TickResolution = InMasterSequence->GetMovieScene()->GetTickResolution();
+ 		FFrameRate TickResolution = InRootSequence->GetMovieScene()->GetTickResolution();
 
 		for (int32 CameraCutIndex = 0; CameraCutIndex < CameraCutData.Num(); ++CameraCutIndex)
 		{
@@ -129,12 +129,12 @@ TArray<UTakeRecorderSource*> UTakeRecorderCameraCutSource::PostRecording(class U
 				FFrameNumber NextCameraCutTime = 
 					CameraCutIndex+1 < CameraCutData.Num() ? 
 					CameraCutData[CameraCutIndex+1].Time.ConvertTo(TickResolution).FloorToFrame() : 
-					InMasterSequence->GetMovieScene()->GetPlaybackRange().GetUpperBoundValue();
-				Range = TRange<FFrameNumber>(InMasterSequence->GetMovieScene()->GetPlaybackRange().GetLowerBoundValue(), NextCameraCutTime);
+					InRootSequence->GetMovieScene()->GetPlaybackRange().GetUpperBoundValue();
+				Range = TRange<FFrameNumber>(InRootSequence->GetMovieScene()->GetPlaybackRange().GetLowerBoundValue(), NextCameraCutTime);
 			}
 			else if (CameraCutIndex == CameraCutData.Num()-1)
 			{
-				Range = TRange<FFrameNumber>(CameraCutTime, InMasterSequence->GetMovieScene()->GetPlaybackRange().GetUpperBoundValue());
+				Range = TRange<FFrameNumber>(CameraCutTime, InRootSequence->GetMovieScene()->GetPlaybackRange().GetUpperBoundValue());
 			}
 			else
 			{
@@ -160,7 +160,7 @@ TArray<UTakeRecorderSource*> UTakeRecorderCameraCutSource::PostRecording(class U
 	}
 
 	World = nullptr;
-	MasterLevelSequence = nullptr;
+	RootLevelSequence = nullptr;
 
 	return SourcesToRemove;
 }
diff --git a/Engine/Plugins/VirtualProduction/Takes/Source/TakeRecorderSources/Private/TakeRecorderCameraCutSource.h b/Engine/Plugins/VirtualProduction/Takes/Source/TakeRecorderSources/Private/TakeRecorderCameraCutSource.h
--- a/Engine/Plugins/VirtualProduction/Takes/Source/TakeRecorderSources/Private/TakeRecorderCameraCutSource.h
+++ b/Engine/Plugins/VirtualProduction/Takes/Source/TakeRecorderSources/Private/TakeRecorderCameraCutSource.h
@@ -25,9 +25,9 @@ public:
 private:
 
 	// UTakeRecorderSource
-	virtual TArray<UTakeRecorderSource*> PreRecording(ULevelSequence* InSequence, FMovieSceneSequenceID InSequenceID, ULevelSequence* InMasterSequence, FManifestSerializer* InManifestSerializer) override;
+	virtual TArray<UTakeRecorderSource*> PreRecording(ULevelSequence* InSequence, FMovieSceneSequenceID InSequenceID, ULevelSequence* InRootSequence, FManifestSerializer* InManifestSerializer) override;
 	virtual void TickRecording(const FQualifiedFrameTime& CurrentTime) override;
-	virtual TArray<UTakeRecorderSource*> PostRecording(class ULevelSequence* InSequence, class ULevelSequence* InMasterSequence, const bool bCancelled) override;
+	virtual TArray<UTakeRecorderSource*> PostRecording(class ULevelSequence* InSequence, class ULevelSequence* InRootSequence, const bool bCancelled) override;
 	virtual FText GetDisplayTextImpl() const override;
 
 	// This source does not support subscenes
@@ -38,9 +38,9 @@ private:
 	UPROPERTY()
 	TObjectPtr<UWorld> World;
 
-	/** The master or uppermost level sequence that this source is being recorded into. Set during PreRecording, null after PostRecording. */
+	/** The root or uppermost level sequence that this source is being recorded into. Set during PreRecording, null after PostRecording. */
 	UPROPERTY()
-	TObjectPtr<ULevelSequence> MasterLevelSequence;
+	TObjectPtr<ULevelSequence> RootLevelSequence;
 
 
 	struct FCameraCutData
diff --git a/Engine/Plugins/VirtualProduction/Takes/Source/TakeRecorderSources/Private/TakeRecorderLevelSequenceSource.cpp b/Engine/Plugins/VirtualProduction/Takes/Source/TakeRecorderSources/Private/TakeRecorderLevelSequenceSource.cpp
--- a/Engine/Plugins/VirtualProduction/Takes/Source/TakeRecorderSources/Private/TakeRecorderLevelSequenceSource.cpp
+++ b/Engine/Plugins/VirtualProduction/Takes/Source/TakeRecorderSources/Private/TakeRecorderLevelSequenceSource.cpp
@@ -18,7 +18,7 @@ UTakeRecorderLevelSequenceSource::UTakeRecorderLevelSequenceSource(const FObject
 	TrackTint = FColor(164, 74, 74);
 }
 
-TArray<UTakeRecorderSource*> UTakeRecorderLevelSequenceSource::PreRecording(ULevelSequence* InSequence, FMovieSceneSequenceID InSequenceID, ULevelSequence* InMasterSequence, FManifestSerializer* InManifestSerializer)
+TArray<UTakeRecorderSource*> UTakeRecorderLevelSequenceSource::PreRecording(ULevelSequence* InSequence, FMovieSceneSequenceID InSequenceID, ULevelSequence* InRootSequence, FManifestSerializer* InManifestSerializer)
 {
 	UWorld* World = TakesUtils::GetFirstPIEWorld();
 	if (!World)
diff --git a/Engine/Plugins/VirtualProduction/Takes/Source/TakeRecorderSources/Private/TakeRecorderLevelSequenceSource.h b/Engine/Plugins/VirtualProduction/Takes/Source/TakeRecorderSources/Private/TakeRecorderLevelSequenceSource.h
--- a/Engine/Plugins/VirtualProduction/Takes/Source/TakeRecorderSources/Private/TakeRecorderLevelSequenceSource.h
+++ b/Engine/Plugins/VirtualProduction/Takes/Source/TakeRecorderSources/Private/TakeRecorderLevelSequenceSource.h
@@ -25,7 +25,7 @@ public:
 private:
 
 	// UTakeRecorderSource
-	virtual TArray<UTakeRecorderSource*> PreRecording(ULevelSequence* InSequence, FMovieSceneSequenceID InSequenceID, ULevelSequence* InMasterSequence, FManifestSerializer* InManifestSerializer) override;
+	virtual TArray<UTakeRecorderSource*> PreRecording(ULevelSequence* InSequence, FMovieSceneSequenceID InSequenceID, ULevelSequence* InRootSequence, FManifestSerializer* InManifestSerializer) override;
 	virtual void StartRecording(const FTimecode& InSectionStartTimecode, const FFrameNumber& InSectionFirstFrame, class ULevelSequence* InSequence) override;
 	virtual void StopRecording(class ULevelSequence* InSequence) override;
 	virtual FText GetDisplayTextImpl() const override;
diff --git a/Engine/Plugins/VirtualProduction/Takes/Source/TakeRecorderSources/Private/TakeRecorderLevelVisibilitySource.cpp b/Engine/Plugins/VirtualProduction/Takes/Source/TakeRecorderSources/Private/TakeRecorderLevelVisibilitySource.cpp
--- a/Engine/Plugins/VirtualProduction/Takes/Source/TakeRecorderSources/Private/TakeRecorderLevelVisibilitySource.cpp
+++ b/Engine/Plugins/VirtualProduction/Takes/Source/TakeRecorderSources/Private/TakeRecorderLevelVisibilitySource.cpp
@@ -59,7 +59,7 @@ UTakeRecorderLevelVisibilitySource::UTakeRecorderLevelVisibilitySource(const FOb
 	: Super(ObjInit)
 {}
 
-TArray<UTakeRecorderSource*> UTakeRecorderLevelVisibilitySource::PreRecording(ULevelSequence* InSequence, FMovieSceneSequenceID InSequenceID, ULevelSequence* InMasterSequence, FManifestSerializer* InManifestSerializer)
+TArray<UTakeRecorderSource*> UTakeRecorderLevelVisibilitySource::PreRecording(ULevelSequence* InSequence, FMovieSceneSequenceID InSequenceID, ULevelSequence* InRootSequence, FManifestSerializer* InManifestSerializer)
 {
 	UWorld* World = TakeRecorderSourcesUtils::GetSourceWorld(InSequence);
 
@@ -91,17 +91,17 @@ TArray<UTakeRecorderSource*> UTakeRecorderLevelVisibilitySource::PreRecording(UL
 	}
 
 	UMovieScene* MovieScene = InSequence->GetMovieScene();
-	for (auto MasterTrack : MovieScene->GetMasterTracks())
+	for (auto Track : MovieScene->GetTracks())
 	{
-		if (MasterTrack->IsA(UMovieSceneLevelVisibilityTrack::StaticClass()) && MasterTrack->GetDisplayName().EqualTo(LevelVisibilityTrackName))
+		if (Track->IsA(UMovieSceneLevelVisibilityTrack::StaticClass()) && Track->GetDisplayName().EqualTo(LevelVisibilityTrackName))
 		{
-			CachedLevelVisibilityTrack = Cast<UMovieSceneLevelVisibilityTrack>(MasterTrack);
+			CachedLevelVisibilityTrack = Cast<UMovieSceneLevelVisibilityTrack>(Track);
 		}
 	}
 
 	if (!CachedLevelVisibilityTrack.IsValid())
 	{
-		CachedLevelVisibilityTrack = MovieScene->AddMasterTrack<UMovieSceneLevelVisibilityTrack>();
+		CachedLevelVisibilityTrack = MovieScene->AddTrack<UMovieSceneLevelVisibilityTrack>();
 		CachedLevelVisibilityTrack->SetDisplayName(LevelVisibilityTrackName);
 	}
 
@@ -141,7 +141,7 @@ void UTakeRecorderLevelVisibilitySource::AddContentsToFolder(UMovieSceneFolder*
 {
 	if (CachedLevelVisibilityTrack.IsValid())
 	{
-		InFolder->AddChildMasterTrack(CachedLevelVisibilityTrack.Get());
+		InFolder->AddChildTrack(CachedLevelVisibilityTrack.Get());
 	}
 }
 
diff --git a/Engine/Plugins/VirtualProduction/Takes/Source/TakeRecorderSources/Private/TakeRecorderLevelVisibilitySource.h b/Engine/Plugins/VirtualProduction/Takes/Source/TakeRecorderSources/Private/TakeRecorderLevelVisibilitySource.h
--- a/Engine/Plugins/VirtualProduction/Takes/Source/TakeRecorderSources/Private/TakeRecorderLevelVisibilitySource.h
+++ b/Engine/Plugins/VirtualProduction/Takes/Source/TakeRecorderSources/Private/TakeRecorderLevelVisibilitySource.h
@@ -40,7 +40,7 @@ public:
 private:
 
 	// UTakeRecorderSource
-	virtual TArray<UTakeRecorderSource*> PreRecording(ULevelSequence* InSequence, FMovieSceneSequenceID InSequenceID, ULevelSequence* InMasterSequence, FManifestSerializer* InManifestSerializer) override;
+	virtual TArray<UTakeRecorderSource*> PreRecording(ULevelSequence* InSequence, FMovieSceneSequenceID InSequenceID, ULevelSequence* InRootSequence, FManifestSerializer* InManifestSerializer) override;
 	virtual void TickRecording(const FQualifiedFrameTime& CurrentTime) override;
 	virtual void AddContentsToFolder(class UMovieSceneFolder* InFolder) override;
 	virtual FText GetDisplayTextImpl() const override;
diff --git a/Engine/Plugins/VirtualProduction/Takes/Source/TakeRecorderSources/Private/TakeRecorderMicrophoneAudioSource.cpp b/Engine/Plugins/VirtualProduction/Takes/Source/TakeRecorderSources/Private/TakeRecorderMicrophoneAudioSource.cpp
--- a/Engine/Plugins/VirtualProduction/Takes/Source/TakeRecorderSources/Private/TakeRecorderMicrophoneAudioSource.cpp
+++ b/Engine/Plugins/VirtualProduction/Takes/Source/TakeRecorderSources/Private/TakeRecorderMicrophoneAudioSource.cpp
@@ -96,20 +96,20 @@ static FString MakeNewAssetName(const FString& BaseAssetPath, const FString& Bas
 	return AssetName;
 }
 
-TArray<UTakeRecorderSource*> UTakeRecorderMicrophoneAudioSource::PreRecording(ULevelSequence* InSequence, FMovieSceneSequenceID InSequenceID, ULevelSequence* InMasterSequence, FManifestSerializer* InManifestSerializer)
+TArray<UTakeRecorderSource*> UTakeRecorderMicrophoneAudioSource::PreRecording(ULevelSequence* InSequence, FMovieSceneSequenceID InSequenceID, ULevelSequence* InRootSequence, FManifestSerializer* InManifestSerializer)
 {
 	UMovieScene* MovieScene = InSequence->GetMovieScene();
-	for (auto MasterTrack : MovieScene->GetMasterTracks())
+	for (auto Track : MovieScene->GetTracks())
 	{
-		if (MasterTrack->IsA(UMovieSceneAudioTrack::StaticClass()) && MasterTrack->GetDisplayName().EqualTo(AudioTrackName))
+		if (Track->IsA(UMovieSceneAudioTrack::StaticClass()) && Track->GetDisplayName().EqualTo(AudioTrackName))
 		{
-			CachedAudioTrack = Cast<UMovieSceneAudioTrack>(MasterTrack);
+			CachedAudioTrack = Cast<UMovieSceneAudioTrack>(Track);
 		}
 	}
 
 	if (!CachedAudioTrack.IsValid())
 	{
-		CachedAudioTrack = MovieScene->AddMasterTrack<UMovieSceneAudioTrack>();
+		CachedAudioTrack = MovieScene->AddTrack<UMovieSceneAudioTrack>();
 		CachedAudioTrack->SetDisplayName(AudioTrackName);
 	}
 
@@ -131,7 +131,7 @@ void UTakeRecorderMicrophoneAudioSource::AddContentsToFolder(UMovieSceneFolder*
 {
 	if (CachedAudioTrack.IsValid())
 	{
-		InFolder->AddChildMasterTrack(CachedAudioTrack.Get());
+		InFolder->AddChildTrack(CachedAudioTrack.Get());
 	}
 }
 
@@ -179,7 +179,7 @@ void UTakeRecorderMicrophoneAudioSource::StopRecording(class ULevelSequence* InS
 	AudioRecorder.Reset();
 }
 
-TArray<UTakeRecorderSource*> UTakeRecorderMicrophoneAudioSource::PostRecording(ULevelSequence* InSequence, class ULevelSequence* InMasterSequence, const bool bCancelled)
+TArray<UTakeRecorderSource*> UTakeRecorderMicrophoneAudioSource::PostRecording(ULevelSequence* InSequence, class ULevelSequence* InRootSequence, const bool bCancelled)
 {
 	if (!RecordedSoundWaves.Num())
 	{
diff --git a/Engine/Plugins/VirtualProduction/Takes/Source/TakeRecorderSources/Private/TakeRecorderMicrophoneAudioSource.h b/Engine/Plugins/VirtualProduction/Takes/Source/TakeRecorderSources/Private/TakeRecorderMicrophoneAudioSource.h
--- a/Engine/Plugins/VirtualProduction/Takes/Source/TakeRecorderSources/Private/TakeRecorderMicrophoneAudioSource.h
+++ b/Engine/Plugins/VirtualProduction/Takes/Source/TakeRecorderSources/Private/TakeRecorderMicrophoneAudioSource.h
@@ -59,10 +59,10 @@ public:
 private:
 
 	// UTakeRecorderSource
-	virtual TArray<UTakeRecorderSource*> PreRecording(ULevelSequence* InSequence, FMovieSceneSequenceID InSequenceID, ULevelSequence* InMasterSequence, FManifestSerializer* InManifestSerializer) override;
+	virtual TArray<UTakeRecorderSource*> PreRecording(ULevelSequence* InSequence, FMovieSceneSequenceID InSequenceID, ULevelSequence* InRootSequence, FManifestSerializer* InManifestSerializer) override;
 	virtual void StartRecording(const FTimecode& InSectionStartTimecode, const FFrameNumber& InSectionFirstFrame, class ULevelSequence* InSequence) override;
 	virtual void StopRecording(class ULevelSequence* InSequence) override;
-	virtual TArray<UTakeRecorderSource*> PostRecording(ULevelSequence* InSequence, class ULevelSequence* InMasterSequence, const bool bCancelled) override;
+	virtual TArray<UTakeRecorderSource*> PostRecording(ULevelSequence* InSequence, class ULevelSequence* InRootSequence, const bool bCancelled) override;
 	virtual FText GetDisplayTextImpl() const override;
 	virtual void AddContentsToFolder(class UMovieSceneFolder* InFolder) override;
 	virtual bool CanAddSource(UTakeRecorderSources* InSources) const override;
diff --git a/Engine/Plugins/VirtualProduction/Takes/Source/TakeRecorderSources/Private/TakeRecorderNearbySpawnedActorSource.cpp b/Engine/Plugins/VirtualProduction/Takes/Source/TakeRecorderSources/Private/TakeRecorderNearbySpawnedActorSource.cpp
--- a/Engine/Plugins/VirtualProduction/Takes/Source/TakeRecorderSources/Private/TakeRecorderNearbySpawnedActorSource.cpp
+++ b/Engine/Plugins/VirtualProduction/Takes/Source/TakeRecorderSources/Private/TakeRecorderNearbySpawnedActorSource.cpp
@@ -28,7 +28,7 @@ UTakeRecorderNearbySpawnedActorSource::UTakeRecorderNearbySpawnedActorSource(con
 	TrackTint = FColor(167, 167, 59);
 }
 
-TArray<UTakeRecorderSource*> UTakeRecorderNearbySpawnedActorSource::PreRecording(ULevelSequence* InSequence, FMovieSceneSequenceID InSequenceID, ULevelSequence* InMasterSequence, FManifestSerializer* InManifestSerializer)
+TArray<UTakeRecorderSource*> UTakeRecorderNearbySpawnedActorSource::PreRecording(ULevelSequence* InSequence, FMovieSceneSequenceID InSequenceID, ULevelSequence* InRootSequence, FManifestSerializer* InManifestSerializer)
 {
 	// Get the sources, add callbacks for new spawned
 	UWorld* World = TakeRecorderSourcesUtils::GetSourceWorld(InSequence);
@@ -46,7 +46,7 @@ TArray<UTakeRecorderSource*> UTakeRecorderNearbySpawnedActorSource::PreRecording
 	return TArray<UTakeRecorderSource*>();
 }
 
-TArray<UTakeRecorderSource*> UTakeRecorderNearbySpawnedActorSource::PostRecording(class ULevelSequence* InSequence, class ULevelSequence* InMasterSequence, const bool bCancelled)
+TArray<UTakeRecorderSource*> UTakeRecorderNearbySpawnedActorSource::PostRecording(class ULevelSequence* InSequence, class ULevelSequence* InRootSequence, const bool bCancelled)
 {
 	// Remove spawn delegates
 	for (auto It = ActorSpawningDelegateHandles.CreateConstIterator(); It; ++It)
diff --git a/Engine/Plugins/VirtualProduction/Takes/Source/TakeRecorderSources/Private/TakeRecorderNearbySpawnedActorSource.h b/Engine/Plugins/VirtualProduction/Takes/Source/TakeRecorderSources/Private/TakeRecorderNearbySpawnedActorSource.h
--- a/Engine/Plugins/VirtualProduction/Takes/Source/TakeRecorderSources/Private/TakeRecorderNearbySpawnedActorSource.h
+++ b/Engine/Plugins/VirtualProduction/Takes/Source/TakeRecorderSources/Private/TakeRecorderNearbySpawnedActorSource.h
@@ -35,8 +35,8 @@ public:
 private:
 
 	// UTakeRecorderSource
-	virtual TArray<UTakeRecorderSource*> PreRecording(ULevelSequence* InSequence, FMovieSceneSequenceID InSequenceID, ULevelSequence* InMasterSequence, FManifestSerializer* InManifestSerializer) override;
-	virtual TArray<UTakeRecorderSource*> PostRecording(class ULevelSequence* InSequence, class ULevelSequence* InMasterSequence, const bool bCancelled) override;
+	virtual TArray<UTakeRecorderSource*> PreRecording(ULevelSequence* InSequence, FMovieSceneSequenceID InSequenceID, ULevelSequence* InRootSequence, FManifestSerializer* InManifestSerializer) override;
+	virtual TArray<UTakeRecorderSource*> PostRecording(class ULevelSequence* InSequence, class ULevelSequence* InRootSequence, const bool bCancelled) override;
 	virtual FText GetDisplayTextImpl() const override;
 	virtual FText GetDescriptionTextImpl() const override;
 
diff --git a/Engine/Plugins/VirtualProduction/Takes/Source/TakeRecorderSources/Private/TakeRecorderPlayerSource.cpp b/Engine/Plugins/VirtualProduction/Takes/Source/TakeRecorderSources/Private/TakeRecorderPlayerSource.cpp
--- a/Engine/Plugins/VirtualProduction/Takes/Source/TakeRecorderSources/Private/TakeRecorderPlayerSource.cpp
+++ b/Engine/Plugins/VirtualProduction/Takes/Source/TakeRecorderSources/Private/TakeRecorderPlayerSource.cpp
@@ -21,7 +21,7 @@ UTakeRecorderPlayerSource::UTakeRecorderPlayerSource(const FObjectInitializer& O
 }
 
 
-TArray<UTakeRecorderSource*> UTakeRecorderPlayerSource::PreRecording(ULevelSequence* InSequence, FMovieSceneSequenceID InSequenceID, ULevelSequence* InMasterSequence, FManifestSerializer* InManifestSerializer)
+TArray<UTakeRecorderSource*> UTakeRecorderPlayerSource::PreRecording(ULevelSequence* InSequence, FMovieSceneSequenceID InSequenceID, ULevelSequence* InRootSequence, FManifestSerializer* InManifestSerializer)
 {
 	TArray<UTakeRecorderSource*> NewSources;
 
@@ -71,7 +71,7 @@ TArray<UTakeRecorderSource*> UTakeRecorderPlayerSource::PreRecording(ULevelSeque
 	return NewSources;
 }
 
-TArray<UTakeRecorderSource*> UTakeRecorderPlayerSource::PostRecording(class ULevelSequence* InSequence, class ULevelSequence* InMasterSequence, const bool bCancelled)
+TArray<UTakeRecorderSource*> UTakeRecorderPlayerSource::PostRecording(class ULevelSequence* InSequence, class ULevelSequence* InRootSequence, const bool bCancelled)
 {
 	TArray<UTakeRecorderSource*> SourcesToRemove;
 
diff --git a/Engine/Plugins/VirtualProduction/Takes/Source/TakeRecorderSources/Private/TakeRecorderPlayerSource.h b/Engine/Plugins/VirtualProduction/Takes/Source/TakeRecorderSources/Private/TakeRecorderPlayerSource.h
--- a/Engine/Plugins/VirtualProduction/Takes/Source/TakeRecorderSources/Private/TakeRecorderPlayerSource.h
+++ b/Engine/Plugins/VirtualProduction/Takes/Source/TakeRecorderSources/Private/TakeRecorderPlayerSource.h
@@ -21,8 +21,8 @@ public:
 private:
 
 	// UTakeRecorderSource
-	virtual TArray<UTakeRecorderSource*> PreRecording(ULevelSequence* InSequence, FMovieSceneSequenceID InSequenceID, ULevelSequence* InMasterSequence, FManifestSerializer* InManifestSerializer) override;
-	virtual TArray<UTakeRecorderSource*> PostRecording(class ULevelSequence* InSequence, class ULevelSequence* InMasterSequence, const bool bCancelled) override;
+	virtual TArray<UTakeRecorderSource*> PreRecording(ULevelSequence* InSequence, FMovieSceneSequenceID InSequenceID, ULevelSequence* InRootSequence, FManifestSerializer* InManifestSerializer) override;
+	virtual TArray<UTakeRecorderSource*> PostRecording(class ULevelSequence* InSequence, class ULevelSequence* InRootSequence, const bool bCancelled) override;
 	virtual FText GetDisplayTextImpl() const override;
 	virtual bool CanAddSource(UTakeRecorderSources* InSources) const override;
 
diff --git a/Engine/Plugins/VirtualProduction/Takes/Source/TakeRecorderSources/Private/TakeRecorderSourcesUtils.h b/Engine/Plugins/VirtualProduction/Takes/Source/TakeRecorderSources/Private/TakeRecorderSourcesUtils.h
--- a/Engine/Plugins/VirtualProduction/Takes/Source/TakeRecorderSources/Private/TakeRecorderSourcesUtils.h
+++ b/Engine/Plugins/VirtualProduction/Takes/Source/TakeRecorderSources/Private/TakeRecorderSourcesUtils.h
@@ -161,7 +161,7 @@ namespace TakeRecorderSourcesUtils
 	* @param Actor Actor to look for.
 	* @return The Sequence ID
 	*/
-	static FMovieSceneSequenceID GetLevelSequenceID(const UTakeRecorderSource* InSource, class AActor* Actor, ULevelSequence* MasterLevelSequence)
+	static FMovieSceneSequenceID GetLevelSequenceID(const UTakeRecorderSource* InSource, class AActor* Actor, ULevelSequence* RootLevelSequence)
 	{
 		UTakeRecorderSources* OwningSources = CastChecked<UTakeRecorderSources>(InSource->GetOuter());
 
@@ -172,7 +172,7 @@ namespace TakeRecorderSourcesUtils
 				AActor* OtherTarget = ActorSource->Target.Get();
 				if (OtherTarget && OtherTarget->GetName() == Actor->GetName()) //at the end the target's may have changed.
 				{
-					if (ActorSource->TargetLevelSequence != MasterLevelSequence)
+					if (ActorSource->TargetLevelSequence != RootLevelSequence)
 					{
 						return ActorSource->GetSequenceID();
 					}
diff --git a/Engine/Plugins/VirtualProduction/Takes/Source/TakeRecorderSources/Private/TakeRecorderWorldSource.cpp b/Engine/Plugins/VirtualProduction/Takes/Source/TakeRecorderSources/Private/TakeRecorderWorldSource.cpp
--- a/Engine/Plugins/VirtualProduction/Takes/Source/TakeRecorderSources/Private/TakeRecorderWorldSource.cpp
+++ b/Engine/Plugins/VirtualProduction/Takes/Source/TakeRecorderSources/Private/TakeRecorderWorldSource.cpp
@@ -42,7 +42,7 @@ UTakeRecorderWorldSource::UTakeRecorderWorldSource(const FObjectInitializer& Obj
 	TrackTint = FColor(129, 129, 129);
 }
 
-TArray<UTakeRecorderSource*> UTakeRecorderWorldSource::PreRecording(ULevelSequence* InSequence, FMovieSceneSequenceID InSequenceID, ULevelSequence* InMasterSequence, FManifestSerializer* InManifestSerializer)
+TArray<UTakeRecorderSource*> UTakeRecorderWorldSource::PreRecording(ULevelSequence* InSequence, FMovieSceneSequenceID InSequenceID, ULevelSequence* InRootSequence, FManifestSerializer* InManifestSerializer)
 {
 	TArray<UTakeRecorderSource*> NewSources;
 
@@ -100,7 +100,7 @@ TArray<UTakeRecorderSource*> UTakeRecorderWorldSource::PreRecording(ULevelSequen
 	return NewSources;
 }
 
-TArray<UTakeRecorderSource*> UTakeRecorderWorldSource::PostRecording(class ULevelSequence* InSequence, class ULevelSequence* InMasterSequence, const bool bCancelled)
+TArray<UTakeRecorderSource*> UTakeRecorderWorldSource::PostRecording(class ULevelSequence* InSequence, class ULevelSequence* InRootSequence, const bool bCancelled)
 {
 	TArray<UTakeRecorderSource*> SourcesToRemove;
 
diff --git a/Engine/Plugins/VirtualProduction/Takes/Source/TakeTrackRecorders/Private/MovieScene3DAttachTrackRecorder.cpp b/Engine/Plugins/VirtualProduction/Takes/Source/TakeTrackRecorders/Private/MovieScene3DAttachTrackRecorder.cpp
--- a/Engine/Plugins/VirtualProduction/Takes/Source/TakeTrackRecorders/Private/MovieScene3DAttachTrackRecorder.cpp
+++ b/Engine/Plugins/VirtualProduction/Takes/Source/TakeTrackRecorders/Private/MovieScene3DAttachTrackRecorder.cpp
@@ -63,7 +63,7 @@ void UMovieScene3DAttachTrackRecorder::RecordSampleImpl(const FQualifiedFrameTim
 				FMovieSceneSequenceID TargetSequenceID = OwningTakeRecorderSource->GetLevelSequenceID(AttachedToActor);
 				FMovieSceneSequenceID ThisSequenceID   = OwningTakeRecorderSource->GetSequenceID();
 
-				FMovieSceneObjectBindingID NewBinding = UE::MovieScene::FRelativeObjectBindingID(ThisSequenceID, TargetSequenceID, Guid, OwningTakeRecorderSource->GetMasterLevelSequence());
+				FMovieSceneObjectBindingID NewBinding = UE::MovieScene::FRelativeObjectBindingID(ThisSequenceID, TargetSequenceID, Guid, OwningTakeRecorderSource->GetRootLevelSequence());
 				MovieSceneSection->SetConstraintBindingID(NewBinding);
 			}
 
@@ -98,7 +98,7 @@ void UMovieScene3DAttachTrackRecorder::FinalizeTrackImpl()
 			FMovieSceneSequenceID TargetSequenceID = OwningTakeRecorderSource->GetLevelSequenceID(AttachedToActor);
 			FMovieSceneSequenceID ThisSequenceID   = OwningTakeRecorderSource->GetSequenceID();
 
-			FMovieSceneObjectBindingID NewBinding = UE::MovieScene::FRelativeObjectBindingID(ThisSequenceID, TargetSequenceID, Guid, OwningTakeRecorderSource->GetMasterLevelSequence());
+			FMovieSceneObjectBindingID NewBinding = UE::MovieScene::FRelativeObjectBindingID(ThisSequenceID, TargetSequenceID, Guid, OwningTakeRecorderSource->GetRootLevelSequence());
 			MovieSceneSection->SetConstraintBindingID(NewBinding);
 		}
 	}
diff --git a/Engine/Plugins/VirtualProduction/Takes/Source/TakeTrackRecorders/Private/MovieSceneAnimationTrackRecorder.cpp b/Engine/Plugins/VirtualProduction/Takes/Source/TakeTrackRecorders/Private/MovieSceneAnimationTrackRecorder.cpp
--- a/Engine/Plugins/VirtualProduction/Takes/Source/TakeTrackRecorders/Private/MovieSceneAnimationTrackRecorder.cpp
+++ b/Engine/Plugins/VirtualProduction/Takes/Source/TakeTrackRecorders/Private/MovieSceneAnimationTrackRecorder.cpp
@@ -54,7 +54,7 @@ void UMovieSceneAnimationTrackRecorder::CreateAnimationAssetAndSequence(const AA
 		ComponentTransform = SkeletalMeshComponent->GetComponentToWorld().GetRelativeTransform(Actor->GetTransform());
 		FString AnimationAssetName = Actor->GetActorLabel();
 
-		if (ULevelSequence* LeaderLevelSequence = OwningTakeRecorderSource->GetMasterLevelSequence())
+		if (ULevelSequence* LeaderLevelSequence = OwningTakeRecorderSource->GetRootLevelSequence())
 		{
 			UTakeMetaData* AssetMetaData = LeaderLevelSequence->FindMetaData<UTakeMetaData>();
 
diff --git a/Engine/Plugins/VirtualProduction/Takes/Source/TakeTrackRecorders/Public/TrackRecorders/IMovieSceneTrackRecorderHost.h b/Engine/Plugins/VirtualProduction/Takes/Source/TakeTrackRecorders/Public/TrackRecorders/IMovieSceneTrackRecorderHost.h
--- a/Engine/Plugins/VirtualProduction/Takes/Source/TakeTrackRecorders/Public/TrackRecorders/IMovieSceneTrackRecorderHost.h
+++ b/Engine/Plugins/VirtualProduction/Takes/Source/TakeTrackRecorders/Public/TrackRecorders/IMovieSceneTrackRecorderHost.h
@@ -185,8 +185,8 @@ public:
 	*/
 	virtual FTransform GetRecordedActorAnimationInitialRootTransform(class AActor* OtherActor) const = 0;
 
-	/** Get the master level sequence being recorded into */
-	virtual ULevelSequence* GetMasterLevelSequence() const = 0;
+	/** Get the root level sequence being recorded into */
+	virtual ULevelSequence* GetRootLevelSequence() const = 0;
 
 	/**/
 	virtual FMovieSceneSequenceID GetSequenceID() const = 0;
diff --git a/Engine/Plugins/VirtualProduction/Takes/Source/TakesCore/Private/TakeRecorderSource.cpp b/Engine/Plugins/VirtualProduction/Takes/Source/TakesCore/Private/TakeRecorderSource.cpp
--- a/Engine/Plugins/VirtualProduction/Takes/Source/TakesCore/Private/TakeRecorderSource.cpp
+++ b/Engine/Plugins/VirtualProduction/Takes/Source/TakesCore/Private/TakeRecorderSource.cpp
@@ -19,7 +19,7 @@ const FSlateBrush* UTakeRecorderSource::GetDisplayIconImpl() const
 	return FSlateIconFinder::FindCustomIconBrushForClass(GetClass(), TEXT("ClassThumbnail"));
 }
 
-TArray<UTakeRecorderSource*> UTakeRecorderSource::PreRecording(ULevelSequence* InSequence, FMovieSceneSequenceID InSequenceID, ULevelSequence* InMasterSequence, FManifestSerializer* InManifestSerializer)
+TArray<UTakeRecorderSource*> UTakeRecorderSource::PreRecording(ULevelSequence* InSequence, FMovieSceneSequenceID InSequenceID, ULevelSequence* InRootSequence, FManifestSerializer* InManifestSerializer)
 {
 	return TArray<UTakeRecorderSource*>();
 }
diff --git a/Engine/Plugins/VirtualProduction/Takes/Source/TakesCore/Private/TakeRecorderSources.cpp b/Engine/Plugins/VirtualProduction/Takes/Source/TakesCore/Private/TakeRecorderSources.cpp
--- a/Engine/Plugins/VirtualProduction/Takes/Source/TakesCore/Private/TakeRecorderSources.cpp
+++ b/Engine/Plugins/VirtualProduction/Takes/Source/TakesCore/Private/TakeRecorderSources.cpp
@@ -84,14 +84,14 @@ void UTakeRecorderSources::PostEditChangeProperty(FPropertyChangedEvent& Propert
 	}
 }
 
-void UTakeRecorderSources::PreRecordingRecursive(TArray<UTakeRecorderSource*> InSources, ULevelSequence* InMasterSequence, TArray<UTakeRecorderSource*>& NewSourcesOut, FManifestSerializer* InManifestSerializer)
+void UTakeRecorderSources::PreRecordingRecursive(TArray<UTakeRecorderSource*> InSources, ULevelSequence* InRootSequence, TArray<UTakeRecorderSource*>& NewSourcesOut, FManifestSerializer* InManifestSerializer)
 {
 
 	TArray<UTakeRecorderSource*> NewSources;
 
 	// Optionally create a folder in the Sequencer UI that will contain this source. We don't want sub-sequences to have folders
 	// created for their sources as you would end up with a Subscene with one item in it hidden inside of a folder, so instead
-	// only the master sequence gets folders created.
+	// only the root sequence gets folders created.
 	const bool bCreateSequencerFolders = true;
 	NewSourcesOut.Append(InSources);
 
@@ -99,25 +99,25 @@ void UTakeRecorderSources::PreRecordingRecursive(TArray<UTakeRecorderSource*> In
 	{
 		if (Source->bEnabled)
 		{
-			ULevelSequence* TargetSequence = InMasterSequence;
+			ULevelSequence* TargetSequence = InRootSequence;
 			FMovieSceneSequenceID TargetSequenceID = MovieSceneSequenceID::Root;
 
 			// The Sequencer Take system is built around swapping out sub-sequences. If they want to use this system, we create a sub-sequence
-			// for the Source and tell it to write into this sub-sequence instead of the master sequence. We then keep track of which Source
+			// for the Source and tell it to write into this sub-sequence instead of the root sequence. We then keep track of which Source
 			// is using which sub-sequence so that we can push the correct sequence for all points of the Source's recording lifecycle.
 			if (Settings.bRecordSourcesIntoSubSequences && Source->SupportsSubscenes())
 			{
-				const FString& SubSequenceTrackName = ObjectTools::SanitizeObjectName(Source->GetSubsceneTrackName(InMasterSequence));
-				const FString& SubSequenceAssetName = ObjectTools::SanitizeObjectName(Source->GetSubsceneAssetName(InMasterSequence));
+				const FString& SubSequenceTrackName = ObjectTools::SanitizeObjectName(Source->GetSubsceneTrackName(InRootSequence));
+				const FString& SubSequenceAssetName = ObjectTools::SanitizeObjectName(Source->GetSubsceneAssetName(InRootSequence));
 
-				TargetSequence = CreateSubSequenceForSource(InMasterSequence, SubSequenceTrackName, SubSequenceAssetName);
+				TargetSequence = CreateSubSequenceForSource(InRootSequence, SubSequenceTrackName, SubSequenceAssetName);
 
 				// If there's already a Subscene Track for our sub-sequence we need to remove that track before create a new one. No data is lost in this process as the
 				// sequence that the subscene points to has been copied by CreateSubSequenceForSource so a new track pointed to the new subsequence includes all the old data.
 				const FString SequenceName = FPaths::GetBaseFilename(TargetSequence->GetPathName());
 				UMovieSceneSubTrack* SubsceneTrack = nullptr;
 
-				for (UMovieSceneTrack* Track : InMasterSequence->GetMovieScene()->GetMasterTracks())
+				for (UMovieSceneTrack* Track : InRootSequence->GetMovieScene()->GetTracks())
 				{
 					if (Track->IsA<UMovieSceneSubTrack>())
 					{
@@ -129,31 +129,31 @@ void UTakeRecorderSources::PreRecordingRecursive(TArray<UTakeRecorderSource*> In
 					}
 				}
 
-				// We need to add the new subsequence to the master sequence immediately so that it shows up in the UI and you can tell that things
+				// We need to add the new subsequence to the root sequence immediately so that it shows up in the UI and you can tell that things
 				// are being recorded, otherwise they don't show up until recording stops and then it magically pops in.
 				if (!SubsceneTrack)
 				{
-					SubsceneTrack = CastChecked<UMovieSceneSubTrack>(InMasterSequence->GetMovieScene()->AddMasterTrack(UMovieSceneSubTrack::StaticClass()));
+					SubsceneTrack = CastChecked<UMovieSceneSubTrack>(InRootSequence->GetMovieScene()->AddTrack(UMovieSceneSubTrack::StaticClass()));
 				}
 
 				// Track should not be transactional during the recording process
 				SubsceneTrack->ClearFlags(RF_Transactional);
 
 				// We create a new sub track for every Source so that we can name the Subtrack after the Source instead of just the sections within it.
-				SubsceneTrack->SetDisplayName(FText::FromString(Source->GetSubsceneTrackName(InMasterSequence)));
+				SubsceneTrack->SetDisplayName(FText::FromString(Source->GetSubsceneTrackName(InRootSequence)));
 				SubsceneTrack->SetColorTint(Source->TrackTint);
 
 				// When we create the Subscene Track we'll make sure a folder is created for it to sort into and add the new Subscene Track as a child of it.
 				if (bCreateSequencerFolders)
 				{
-					UMovieSceneFolder* Folder = AddFolderForSource(Source, InMasterSequence->GetMovieScene());
-					Folder->AddChildMasterTrack(SubsceneTrack);
+					UMovieSceneFolder* Folder = AddFolderForSource(Source, InRootSequence->GetMovieScene());
+					Folder->AddChildTrack(SubsceneTrack);
 				}
 
 				// We initialize the sequence to start at zero and be a 0 frame length section as there is no data in the sections yet.
 				// We'll have to update these sections each frame as the recording progresses so they appear to get longer like normal
 				// tracks do as we record into them.
-				FFrameNumber RecordStartFrame = Settings.bStartAtCurrentTimecode ? FFrameRate::TransformTime(FFrameTime(FApp::GetTimecode().ToFrameNumber(TargetLevelSequenceDisplayRate)), TargetLevelSequenceDisplayRate, TargetLevelSequenceTickResolution).FloorToFrame() : InMasterSequence->GetMovieScene()->GetPlaybackRange().GetLowerBoundValue();
+				FFrameNumber RecordStartFrame = Settings.bStartAtCurrentTimecode ? FFrameRate::TransformTime(FFrameTime(FApp::GetTimecode().ToFrameNumber(TargetLevelSequenceDisplayRate)), TargetLevelSequenceDisplayRate, TargetLevelSequenceTickResolution).FloorToFrame() : InRootSequence->GetMovieScene()->GetPlaybackRange().GetLowerBoundValue();
 				UMovieSceneSubSection* NewSubSection = SubsceneTrack->AddSequence(TargetSequence, RecordStartFrame, 0);
 
 				// Section should not be transactional during the recording process
@@ -205,7 +205,7 @@ void UTakeRecorderSources::PreRecordingRecursive(TArray<UTakeRecorderSource*> In
 			// Update our mappings of which sources use which sub-sequence.
 			SourceSubSequenceMap.FindOrAdd(Source) = TargetSequence;
 
-			for (UTakeRecorderSource* NewlyAddedSource : Source->PreRecording(TargetSequence, TargetSequenceID, InMasterSequence, InManifestSerializer))
+			for (UTakeRecorderSource* NewlyAddedSource : Source->PreRecording(TargetSequence, TargetSequenceID, InRootSequence, InManifestSerializer))
 			{
 				// Add it to our classes list of sources 
 				Sources.Add(NewlyAddedSource);
@@ -218,7 +218,7 @@ void UTakeRecorderSources::PreRecordingRecursive(TArray<UTakeRecorderSource*> In
 			// so that the Source has had a chance to create any required sections that will go in the folder.
 			if (!Settings.bRecordSourcesIntoSubSequences && bCreateSequencerFolders)
 			{
-				UMovieSceneFolder* Folder = AddFolderForSource(Source, InMasterSequence->GetMovieScene());
+				UMovieSceneFolder* Folder = AddFolderForSource(Source, InRootSequence->GetMovieScene());
 
 				// Different sources can create different kinds of tracks so we allow each source to decide how it gets
 				// represented inside the folder.
@@ -229,9 +229,9 @@ void UTakeRecorderSources::PreRecordingRecursive(TArray<UTakeRecorderSource*> In
 
 	if (NewSources.Num())
 	{
-		// We don't want to nestle sub-sequences recursively so we always pass the Master Sequence and not the sequence
+		// We don't want to nestle sub-sequences recursively so we always pass the Root Sequence and not the sequence
 		// created for a new source.
-		PreRecordingRecursive(NewSources, InMasterSequence, NewSourcesOut, InManifestSerializer);
+		PreRecordingRecursive(NewSources, InRootSequence, NewSourcesOut, InManifestSerializer);
 		SourcesSerialNumber++;
 	}
 }
@@ -503,7 +503,7 @@ void UTakeRecorderSources::StopRecording(class ULevelSequence* InSequence, const
 	TArray<UObject*> AssetsToCleanUp;
 	for (auto SourceSubSequence : SourceSubSequenceMap)
 	{
-		// Only save subsequences but not the master sequence, it is already saved in UTakeRecorder::StopInternal
+		// Only save subsequences but not the root sequence, it is already saved in UTakeRecorder::StopInternal
 		const bool bIsValidSubSequence = SourceSubSequence.Value && SourceSubSequence.Value.Get() != InSequence;
 		if (bIsValidSubSequence)
 		{
@@ -530,20 +530,20 @@ void UTakeRecorderSources::StopRecording(class ULevelSequence* InSequence, const
 	}
 }
 
-ULevelSequence* UTakeRecorderSources::CreateSubSequenceForSource(ULevelSequence* InMasterSequence, const FString& SubSequenceTrackName, const FString& SubSequenceAssetName)
+ULevelSequence* UTakeRecorderSources::CreateSubSequenceForSource(ULevelSequence* InRootSequence, const FString& SubSequenceTrackName, const FString& SubSequenceAssetName)
 {
 	IAssetRegistry& AssetRegistry = FModuleManager::LoadModuleChecked<FAssetRegistryModule>("AssetRegistry").Get();
 
 	// We want to sanitize the object names because sometimes they come from names with spaces and other invalid characters in them.
-	const FString& SequenceDirectory = FPaths::GetPath(InMasterSequence->GetPathName());
-	const FString& SequenceName = FPaths::GetBaseFilename(InMasterSequence->GetPathName());
+	const FString& SequenceDirectory = FPaths::GetPath(InRootSequence->GetPathName());
+	const FString& SequenceName = FPaths::GetBaseFilename(InRootSequence->GetPathName());
 
-	// We need to check the Master Sequence to see if they already have a sub-sequence with this name so that we duplicate the right
+	// We need to check the Root Sequence to see if they already have a sub-sequence with this name so that we duplicate the right
 	// sequence and re-use that, instead of just making a new blank sequence every time. This will help in cases where they've done a recording, 
 	// modified a sub-sequence and want to record from that setup. Each source will individually remove any old data inside the Sub Sequence
 	// so we don't have to worry about any data the user added via Sequencer unrelated to what they recorded.
 	ULevelSequence* ExistingSubSequence = nullptr;
-	UMovieSceneSubTrack* SubTrack = InMasterSequence->GetMovieScene()->FindMasterTrack<UMovieSceneSubTrack>();
+	UMovieSceneSubTrack* SubTrack = InRootSequence->GetMovieScene()->FindTrack<UMovieSceneSubTrack>();
 	if (SubTrack)
 	{
 		// Look at each section in the track to see if it has the same name as our new SubSequence name.
@@ -575,16 +575,16 @@ ULevelSequence* UTakeRecorderSources::CreateSubSequenceForSource(ULevelSequence*
 			// Movie scene should not be transactional during the recording process
 			OutAsset->GetMovieScene()->ClearFlags(RF_Transactional);
 
-			OutAsset->GetMovieScene()->SetTickResolutionDirectly(InMasterSequence->GetMovieScene()->GetTickResolution());
-			OutAsset->GetMovieScene()->SetDisplayRate(InMasterSequence->GetMovieScene()->GetDisplayRate());
+			OutAsset->GetMovieScene()->SetTickResolutionDirectly(InRootSequence->GetMovieScene()->GetTickResolution());
+			OutAsset->GetMovieScene()->SetDisplayRate(InRootSequence->GetMovieScene()->GetDisplayRate());
 		}
 
-		UTakeMetaData* TakeMetaData = InMasterSequence->FindMetaData<UTakeMetaData>();
+		UTakeMetaData* TakeMetaData = InRootSequence->FindMetaData<UTakeMetaData>();
 		if (TakeMetaData)
 		{
 			UTakeMetaData* OutTakeMetaData = OutAsset->CopyMetaData(TakeMetaData);
 
-			// Tack on the sub sequence name so that it's unique from the master sequence
+			// Tack on the sub sequence name so that it's unique from the root sequence
 			OutTakeMetaData->SetSlate(TakeMetaData->GetSlate() + TEXT("_") + SubSequenceTrackName, false);
 		}
 
diff --git a/Engine/Plugins/VirtualProduction/Takes/Source/TakesCore/Public/TakeRecorderSource.h b/Engine/Plugins/VirtualProduction/Takes/Source/TakesCore/Public/TakeRecorderSource.h
--- a/Engine/Plugins/VirtualProduction/Takes/Source/TakesCore/Public/TakeRecorderSource.h
+++ b/Engine/Plugins/VirtualProduction/Takes/Source/TakesCore/Public/TakeRecorderSource.h
@@ -63,13 +63,13 @@ public:
 	* Will not be called if this recording source is not enabled.
 	*
 	* @param InSequence - The Level Sequence the take is being recorded into
-	* @param InSequenceID - The sequence ID for InSequence, or MovieSceneSequenceID::Root where InSequence == InMasterSequence (ie, not sub-sequence recording)
-	* @param InMasterSequence - The Master Level Sequence that may contain the InSequence as a child or if no subsequences, is the same as InSequence.
+	* @param InSequenceID - The sequence ID for InSequence, or MovieSceneSequenceID::Root where InSequence == InRootSequence (ie, not sub-sequence recording)
+	* @param InRootSequence - The Root Level Sequence that may contain the InSequence as a child or if no subsequences, is the same as InSequence.
 	* @param InManifestSerializer - Manifest Serializer that we may write into.
 	* @return An array of newly created take recorder sources. Can be an empty list if no additional sources needed to be
 	* created by this source.
 	*/
-	virtual TArray<UTakeRecorderSource*> PreRecording(ULevelSequence* InSequence, FMovieSceneSequenceID InSequenceID, ULevelSequence* InMasterSequence, FManifestSerializer* InManifestSerializer);
+	virtual TArray<UTakeRecorderSource*> PreRecording(ULevelSequence* InSequence, FMovieSceneSequenceID InSequenceID, ULevelSequence* InRootSequence, FManifestSerializer* InManifestSerializer);
 
 	/**
 	* This is called when the UTakeRecorderSources starts a recording, after all sources have had PreRecording called on them.
@@ -117,7 +117,7 @@ public:
 	* @return An array of take recorder sources to be removed. Likely to match the same list as PreRecording. Can be an empty list
 	* if no additional sources are required by this source.
 	*/
-	virtual TArray<UTakeRecorderSource*> PostRecording(class ULevelSequence* InSequence, class ULevelSequence* InMasterSequence, const bool bCancelled) { return TArray<UTakeRecorderSource*>(); }
+	virtual TArray<UTakeRecorderSource*> PostRecording(class ULevelSequence* InSequence, class ULevelSequence* InRootSequence, const bool bCancelled) { return TArray<UTakeRecorderSource*>(); }
 
 	/*
 	 * This is called on all sources after post recording. There should not be any dependencies between sources at this point.
diff --git a/Engine/Plugins/VirtualProduction/Takes/Source/TakesCore/Public/TakeRecorderSources.h b/Engine/Plugins/VirtualProduction/Takes/Source/TakesCore/Public/TakeRecorderSources.h
--- a/Engine/Plugins/VirtualProduction/Takes/Source/TakesCore/Public/TakeRecorderSources.h
+++ b/Engine/Plugins/VirtualProduction/Takes/Source/TakesCore/Public/TakeRecorderSources.h
@@ -165,8 +165,8 @@ public:
 	*  Static functions used by other parts of the take system
 	*/
 
-	/** Creates a sub-sequence asset for the specified sub sequence name based on the given master sequence. */
-	static ULevelSequence* CreateSubSequenceForSource(ULevelSequence* InMasterSequence, const FString& SubSequenceTrackName, const FString& SubSequenceAssetName);
+	/** Creates a sub-sequence asset for the specified sub sequence name based on the given root sequence. */
+	static ULevelSequence* CreateSubSequenceForSource(ULevelSequence* InRootSequence, const FString& SubSequenceTrackName, const FString& SubSequenceAssetName);
 
 	/**
 	 * Array of pairs - key time and the corresponding timecode
@@ -187,7 +187,7 @@ private:
 	FOnSourcesChanged OnSourcesChangedEvent;
 
 	/** Calls PreRecording on sources recursively allowing them to create other sources which properly get PreRecording called on them as well. */
-	void PreRecordingRecursive(TArray<UTakeRecorderSource*> InSources, ULevelSequence* InMasterSequence, TArray<UTakeRecorderSource*>& NewSourcesOut, FManifestSerializer* InManifestSerializer);
+	void PreRecordingRecursive(TArray<UTakeRecorderSource*> InSources, ULevelSequence* InRootSequence, TArray<UTakeRecorderSource*>& NewSourcesOut, FManifestSerializer* InManifestSerializer);
 
 	/** Finds the folder that the given Source should be created in, creating it if necessary. */
 	class UMovieSceneFolder* AddFolderForSource(const UTakeRecorderSource* InSource, class UMovieScene* InMovieScene);
@@ -209,7 +209,7 @@ private:
 	UPROPERTY(Instanced)
 	TArray<TObjectPtr<UTakeRecorderSource>> Sources;
 
-	/** Maps each source to the level sequence that was created for that source, or to the master source if a subsequence was not created. */
+	/** Maps each source to the level sequence that was created for that source, or to the root source if a subsequence was not created. */
 	UPROPERTY(Transient)
 	TMap<TObjectPtr<UTakeRecorderSource>, TObjectPtr<ULevelSequence>> SourceSubSequenceMap;
 
diff --git a/Engine/Source/Editor/MovieSceneTools/Private/AutomatedLevelSequenceCapture.cpp b/Engine/Source/Editor/MovieSceneTools/Private/AutomatedLevelSequenceCapture.cpp
--- a/Engine/Source/Editor/MovieSceneTools/Private/AutomatedLevelSequenceCapture.cpp
+++ b/Engine/Source/Editor/MovieSceneTools/Private/AutomatedLevelSequenceCapture.cpp
@@ -100,7 +100,7 @@ UMovieSceneCinematicShotTrack* GetCinematicShotTrack(TWeakObjectPtr<ALevelSequen
 		return nullptr;
 	}
 
-	return MovieScene->FindMasterTrack<UMovieSceneCinematicShotTrack>();
+	return MovieScene->FindTrack<UMovieSceneCinematicShotTrack>();
 }
 
 UMovieSceneMotionVectorSimulationSystem* FindMotionVectorSimulation(ALevelSequenceActor* LevelSequenceActor)
@@ -140,7 +140,7 @@ UAutomatedLevelSequenceCapture::UAutomatedLevelSequenceCapture(const FObjectInit
 #if WITH_EDITORONLY_DATA
 void UAutomatedLevelSequenceCapture::AddFormatMappings(TMap<FString, FStringFormatArg>& OutFormatMappings, const FFrameMetrics& FrameMetrics) const
 {
-	OutFormatMappings.Add(TEXT("sequence"), CachedState.MasterName);
+	OutFormatMappings.Add(TEXT("sequence"), CachedState.RootName);
 
 	OutFormatMappings.Add(TEXT("shot"), CachedState.CurrentShotName);
 	OutFormatMappings.Add(TEXT("shot_frame"), FString::Printf(TEXT("%0*d"), Settings.ZeroPadFrameNumbers, CachedState.CurrentShotLocalTime.Time.FrameNumber.Value));
@@ -522,7 +522,7 @@ bool UAutomatedLevelSequenceCapture::SetupShot(FFrameNumber& StartTime, FFrameNu
 		if (SectionIndex == ShotIndex)
 		{
 			// We intersect with the CachedPlaybackRange instead of copying the playback range from the shot to handle the case where
-			// the playback range intersected the middle of the shot before we started manipulating ranges. We manually expand the master
+			// the playback range intersected the middle of the shot before we started manipulating ranges. We manually expand the root
 			// Movie Sequence's playback range by the number of handle frames to allow handle frames to work as expected on first/last shot.
 			FFrameNumber HandleFramesResolutionSpace = ConvertFrameTime(Settings.HandleFrames, Settings.GetFrameRate(), MovieScene->GetTickResolution()).FloorToFrame();
 			TRange<FFrameNumber> ExtendedCachedPlaybackRange = UE::MovieScene::ExpandRange(CachedPlaybackRange, HandleFramesResolutionSpace);
@@ -830,7 +830,7 @@ void UAutomatedLevelSequenceCapture::SequenceUpdated(const UMovieSceneSequencePl
 		if (Actor && Actor->SequencePlayer)
 		{
 			// If this is a new shot, set the state to shot warm up and pause on this frame until warmed up			
-			const bool bHasMultipleShots = PreviousState.CurrentShotName != PreviousState.MasterName;
+			const bool bHasMultipleShots = PreviousState.CurrentShotName != PreviousState.RootName;
 			const bool bNewShot = bHasMultipleShots && PreviousState.ShotID != CachedState.ShotID;
 			const bool bNewFrame = PreviousTime != CurrentTime;
 
@@ -1069,7 +1069,7 @@ void UAutomatedLevelSequenceCapture::ExportEDL()
 		return;
 	}
 
-	UMovieSceneCinematicShotTrack* ShotTrack = MovieScene->FindMasterTrack<UMovieSceneCinematicShotTrack>();
+	UMovieSceneCinematicShotTrack* ShotTrack = MovieScene->FindTrack<UMovieSceneCinematicShotTrack>();
 	if (!ShotTrack)
 	{
 		return;
@@ -1110,7 +1110,7 @@ void UAutomatedLevelSequenceCapture::ExportFCPXML()
 		return;
 	}
 
-	UMovieSceneCinematicShotTrack* ShotTrack = MovieScene->FindMasterTrack<UMovieSceneCinematicShotTrack>();
+	UMovieSceneCinematicShotTrack* ShotTrack = MovieScene->FindTrack<UMovieSceneCinematicShotTrack>();
 	if (!ShotTrack)
 	{
 		return;
diff --git a/Engine/Source/Editor/MovieSceneTools/Private/FCPXML/FCPXMLExport.cpp b/Engine/Source/Editor/MovieSceneTools/Private/FCPXML/FCPXMLExport.cpp
--- a/Engine/Source/Editor/MovieSceneTools/Private/FCPXML/FCPXMLExport.cpp
+++ b/Engine/Source/Editor/MovieSceneTools/Private/FCPXML/FCPXMLExport.cpp
@@ -134,18 +134,18 @@ bool FFCPXMLExportVisitor::ConstructProjectNode(TSharedRef<FFCPXMLNode> InParent
 
 bool FFCPXMLExportVisitor::ConstructMasterVideoClipNodes(TSharedRef<FFCPXMLNode> InParentNode)
 {
-	if (!ExportData->IsExportDataValid() || !ExportData->MovieSceneData.IsValid() || !ExportData->MovieSceneData->CinematicMasterTrack.IsValid())
+	if (!ExportData->IsExportDataValid() || !ExportData->MovieSceneData.IsValid() || !ExportData->MovieSceneData->CinematicData.IsValid())
 	{
 		return false;
 	}
 
-	TSharedPtr<FMovieSceneExportCinematicMasterTrackData> CinematicMasterTrackData = ExportData->MovieSceneData->CinematicMasterTrack;
-	if (!CinematicMasterTrackData.IsValid())
+	TSharedPtr<FMovieSceneExportCinematicData> CinematicData = ExportData->MovieSceneData->CinematicData;
+	if (!CinematicData.IsValid())
 	{
 		return false;
 	}
 
-	for (TSharedPtr<FMovieSceneExportCinematicSectionData> CinematicSection : CinematicMasterTrackData->CinematicSections)
+	for (TSharedPtr<FMovieSceneExportCinematicSectionData> CinematicSection : CinematicData->CinematicSections)
 	{
 		if (!CinematicSection.IsValid())
 		{
@@ -159,7 +159,7 @@ bool FFCPXMLExportVisitor::ConstructMasterVideoClipNodes(TSharedRef<FFCPXMLNode>
 			continue;
 		}
 
-		if (!ConstructMasterClipNode(InParentNode, CinematicSection, CinematicMasterTrackData))
+		if (!ConstructMasterClipNode(InParentNode, CinematicSection, CinematicData))
 		{
 			return false;
 		}
@@ -170,19 +170,19 @@ bool FFCPXMLExportVisitor::ConstructMasterVideoClipNodes(TSharedRef<FFCPXMLNode>
 
 bool FFCPXMLExportVisitor::ConstructMasterAudioClipNodes(TSharedRef<FFCPXMLNode> InParentNode)
 {
-	if (!ExportData->IsExportDataValid() || !ExportData->MovieSceneData.IsValid() || !ExportData->MovieSceneData->CinematicMasterTrack.IsValid())
+	if (!ExportData->IsExportDataValid() || !ExportData->MovieSceneData.IsValid())
 	{
 		return false;
 	}
 
-	for (TSharedPtr<FMovieSceneExportAudioMasterTrackData> AudioMasterTrack : ExportData->MovieSceneData->AudioMasterTracks)
+	for (TSharedPtr<FMovieSceneExportAudioData> AudioData : ExportData->MovieSceneData->AudioData)
 	{
-		if (!AudioMasterTrack.IsValid())
+		if (!AudioData.IsValid())
 		{
 			return false;
 		}
 
-		for (TSharedPtr<FMovieSceneExportAudioTrackData> AudioTrack : AudioMasterTrack->AudioTracks)
+		for (TSharedPtr<FMovieSceneExportAudioTrackData> AudioTrack : AudioData->AudioTracks)
 		{
 			if (!AudioTrack.IsValid())
 			{
@@ -197,7 +197,7 @@ bool FFCPXMLExportVisitor::ConstructMasterAudioClipNodes(TSharedRef<FFCPXMLNode>
 
 				if (!bMasterClipExists)
 				{
-					ConstructMasterClipNode(InParentNode, AudioSection, AudioMasterTrack);
+					ConstructMasterClipNode(InParentNode, AudioSection, AudioData);
 				}
 			}
 		}
@@ -205,7 +205,7 @@ bool FFCPXMLExportVisitor::ConstructMasterAudioClipNodes(TSharedRef<FFCPXMLNode>
 	return true;
 }
 
-bool FFCPXMLExportVisitor::ConstructMasterClipNode(TSharedRef<FFCPXMLNode> InParentNode, const TSharedPtr<FMovieSceneExportCinematicSectionData> InCinematicSectionData, const TSharedPtr<FMovieSceneExportCinematicMasterTrackData> InCinematicMasterTrackData)
+bool FFCPXMLExportVisitor::ConstructMasterClipNode(TSharedRef<FFCPXMLNode> InParentNode, const TSharedPtr<FMovieSceneExportCinematicSectionData> InCinematicSectionData, const TSharedPtr<FMovieSceneExportCinematicData> InCinematicData)
 {
 	if (!InCinematicSectionData.IsValid())
 	{
@@ -245,7 +245,7 @@ bool FFCPXMLExportVisitor::ConstructMasterClipNode(TSharedRef<FFCPXMLNode> InPar
 	TSharedRef<FFCPXMLNode> VideoNode = MediaNode->CreateChildNode(TEXT("video"));
 	TSharedRef<FFCPXMLNode> TrackNode = VideoNode->CreateChildNode(TEXT("track"));
 
-	if (!ConstructVideoClipItemNode(TrackNode, InCinematicSectionData, InCinematicMasterTrackData, true))
+	if (!ConstructVideoClipItemNode(TrackNode, InCinematicSectionData, InCinematicData, true))
 	{
 		return false;
 	}
@@ -263,7 +263,7 @@ bool FFCPXMLExportVisitor::ConstructMasterClipNode(TSharedRef<FFCPXMLNode> InPar
 	return true;
 }
 
-bool FFCPXMLExportVisitor::ConstructMasterClipNode(TSharedRef<FFCPXMLNode> InParentNode, const TSharedPtr<FMovieSceneExportAudioSectionData> InAudioSectionData, const TSharedPtr<FMovieSceneExportAudioMasterTrackData> InAudioMasterTrackData)
+bool FFCPXMLExportVisitor::ConstructMasterClipNode(TSharedRef<FFCPXMLNode> InParentNode, const TSharedPtr<FMovieSceneExportAudioSectionData> InAudioSectionData, const TSharedPtr<FMovieSceneExportAudioData> InAudioData)
 {
 	if (!InAudioSectionData.IsValid())
 	{
@@ -313,7 +313,7 @@ bool FFCPXMLExportVisitor::ConstructMasterClipNode(TSharedRef<FFCPXMLNode> InPar
 	{
 		GetNextClipItemIdName(ClipItemIdName2);
 	}
-	if (!ConstructAudioClipItemNode(TrackNode, InAudioSectionData, InAudioMasterTrackData, 1, true, ClipItemIdName1, ClipItemIdName2, 1, 1, 1, 2))
+	if (!ConstructAudioClipItemNode(TrackNode, InAudioSectionData, InAudioData, 1, true, ClipItemIdName1, ClipItemIdName2, 1, 1, 1, 2))
 	{
 		return false;
 	}
@@ -323,7 +323,7 @@ bool FFCPXMLExportVisitor::ConstructMasterClipNode(TSharedRef<FFCPXMLNode> InPar
 	{
 		TrackNode = AudioNode->CreateChildNode(TEXT("track"));
 
-		if (!ConstructAudioClipItemNode(TrackNode, InAudioSectionData, InAudioMasterTrackData, 2, true, ClipItemIdName1, ClipItemIdName2, 1, 1, 1, 2))
+		if (!ConstructAudioClipItemNode(TrackNode, InAudioSectionData, InAudioData, 2, true, ClipItemIdName1, ClipItemIdName2, 1, 1, 1, 2))
 		{
 			return false;
 		}
@@ -534,13 +534,13 @@ bool FFCPXMLExportVisitor::ConstructSequenceNode(TSharedRef<FFCPXMLNode> InParen
 
 bool FFCPXMLExportVisitor::ConstructVideoNode(TSharedRef<FFCPXMLNode> InParentNode)
 {
-	if (!ExportData->IsExportDataValid() || !ExportData->MovieSceneData.IsValid() || !ExportData->MovieSceneData->CinematicMasterTrack.IsValid())
+	if (!ExportData->IsExportDataValid() || !ExportData->MovieSceneData.IsValid() || !ExportData->MovieSceneData->CinematicData.IsValid())
 	{ 
 		return false; 
 	}
 
-	TSharedPtr<FMovieSceneExportCinematicMasterTrackData> CinematicMasterTrackData = ExportData->MovieSceneData->CinematicMasterTrack;
-	if (!CinematicMasterTrackData.IsValid())
+	TSharedPtr<FMovieSceneExportCinematicData> CinematicData = ExportData->MovieSceneData->CinematicData;
+	if (!CinematicData.IsValid())
 	{
 		return false;
 	}
@@ -555,9 +555,9 @@ bool FFCPXMLExportVisitor::ConstructVideoNode(TSharedRef<FFCPXMLNode> InParentNo
 	}
  
 	// Add in reverse order
-	for (int32 RowIndex = CinematicMasterTrackData->CinematicTracks.Num()-1; RowIndex >= 0; --RowIndex)
+	for (int32 RowIndex = CinematicData->CinematicTracks.Num()-1; RowIndex >= 0; --RowIndex)
 	{
-		if (!ConstructVideoTrackNode(VideoNode, CinematicMasterTrackData->CinematicTracks[RowIndex], CinematicMasterTrackData)) 
+		if (!ConstructVideoTrackNode(VideoNode, CinematicData->CinematicTracks[RowIndex], CinematicData))
 		{ 
 			return false;
 		}
@@ -576,9 +576,9 @@ bool FFCPXMLExportVisitor::ConstructAudioNode(TSharedRef<FFCPXMLNode> InParentNo
 	TSharedRef<FFCPXMLNode> AudioNode = InParentNode->CreateChildNode(TEXT("audio"));
 
 	int32 NumChannels = 1;
-	for (TSharedPtr<FMovieSceneExportAudioMasterTrackData> AudioMasterTrack : ExportData->MovieSceneData->AudioMasterTracks)
+	for (TSharedPtr<FMovieSceneExportAudioData> AudioData : ExportData->MovieSceneData->AudioData)
 	{
-		if (HasStereoAudioSections(AudioMasterTrack->AudioSections))
+		if (HasStereoAudioSections(AudioData->AudioSections))
 		{
 			NumChannels = 2;
 			break;
@@ -616,17 +616,17 @@ bool FFCPXMLExportVisitor::ConstructAudioNode(TSharedRef<FFCPXMLNode> InParentNo
 	uint32 TrackIndex = 1;
 
 	// Add in reverse order
-	for (int32 RowIndex = ExportData->MovieSceneData->AudioMasterTracks.Num() - 1; RowIndex >= 0; --RowIndex)
+	for (int32 RowIndex = ExportData->MovieSceneData->AudioData.Num() - 1; RowIndex >= 0; --RowIndex)
 	{
-		if (!ExportData->MovieSceneData->AudioMasterTracks[RowIndex].IsValid())
+		if (!ExportData->MovieSceneData->AudioData[RowIndex].IsValid())
 		{
 			return false;
 		}
 
-		for (TSharedPtr<FMovieSceneExportAudioTrackData> AudioTrack : ExportData->MovieSceneData->AudioMasterTracks[RowIndex]->AudioTracks)
+		for (TSharedPtr<FMovieSceneExportAudioTrackData> AudioTrack : ExportData->MovieSceneData->AudioData[RowIndex]->AudioTracks)
 		{
 			uint32 OutNumTracks{ 0 };
-			if (!ConstructAudioTrackNode(AudioNode, AudioTrack, ExportData->MovieSceneData->AudioMasterTracks[RowIndex], TrackIndex, OutNumTracks))
+			if (!ConstructAudioTrackNode(AudioNode, AudioTrack, ExportData->MovieSceneData->AudioData[RowIndex], TrackIndex, OutNumTracks))
 			{
 				return false;
 			}
@@ -637,7 +637,7 @@ bool FFCPXMLExportVisitor::ConstructAudioNode(TSharedRef<FFCPXMLNode> InParentNo
 	return true;
 }
 
-bool FFCPXMLExportVisitor::ConstructVideoTrackNode(TSharedRef<FFCPXMLNode> InParentNode, const TSharedPtr<FMovieSceneExportCinematicTrackData> InCinematicTrackData, const TSharedPtr<FMovieSceneExportCinematicMasterTrackData> InCinematicMasterTrackData)
+bool FFCPXMLExportVisitor::ConstructVideoTrackNode(TSharedRef<FFCPXMLNode> InParentNode, const TSharedPtr<FMovieSceneExportCinematicTrackData> InCinematicTrackData, const TSharedPtr<FMovieSceneExportCinematicData> InCinematicData)
 {
 	if (!ExportData->IsExportDataValid() || !InCinematicTrackData.IsValid())
 	{
@@ -654,7 +654,7 @@ bool FFCPXMLExportVisitor::ConstructVideoTrackNode(TSharedRef<FFCPXMLNode> InPar
 			continue;
 		}
 		
-		if (!ConstructVideoClipItemNode(TrackNode, CinematicSection, InCinematicMasterTrackData, false))
+		if (!ConstructVideoClipItemNode(TrackNode, CinematicSection, InCinematicData, false))
 		{
 			return false;
 		}
@@ -681,7 +681,7 @@ bool FFCPXMLExportVisitor::HasStereoAudioSections(const TArray<TSharedPtr<FMovie
 	return false;
 }
 
-bool FFCPXMLExportVisitor::ConstructAudioTrackNode(TSharedRef<FFCPXMLNode> InParentNode, const TSharedPtr<FMovieSceneExportAudioTrackData> InAudioTrackData, const TSharedPtr<FMovieSceneExportAudioMasterTrackData> InAudioMasterTrackData, uint32 InTrackIndex, uint32 OutNumTracks)
+bool FFCPXMLExportVisitor::ConstructAudioTrackNode(TSharedRef<FFCPXMLNode> InParentNode, const TSharedPtr<FMovieSceneExportAudioTrackData> InAudioTrackData, const TSharedPtr<FMovieSceneExportAudioData> InAudioData, uint32 InTrackIndex, uint32 OutNumTracks)
 {
 	if (!ExportData->IsExportDataValid() || !InAudioTrackData.IsValid())
 	{
@@ -769,7 +769,7 @@ bool FFCPXMLExportVisitor::ConstructAudioTrackNode(TSharedRef<FFCPXMLNode> InPar
 			continue;
 		}
 
-		if (!ConstructAudioClipItemNode(TrackNode, AudioSection, InAudioMasterTrackData, 1, false, ClipItem1[Index], ClipItem2[Index], ClipIndex1[Index], ClipIndex2[Index], TrackIndex1, TrackIndex2))
+		if (!ConstructAudioClipItemNode(TrackNode, AudioSection, InAudioData, 1, false, ClipItem1[Index], ClipItem2[Index], ClipIndex1[Index], ClipIndex2[Index], TrackIndex1, TrackIndex2))
 		{
 			return false;
 		}
@@ -801,7 +801,7 @@ bool FFCPXMLExportVisitor::ConstructAudioTrackNode(TSharedRef<FFCPXMLNode> InPar
 
 			if (AudioSection->NumChannels == 2)
 			{
-				if (!ConstructAudioClipItemNode(TrackNode, AudioSection, InAudioMasterTrackData, 2, false, ClipItem1[Index], ClipItem2[Index], ClipIndex1[Index], ClipIndex2[Index], TrackIndex1, TrackIndex2))
+				if (!ConstructAudioClipItemNode(TrackNode, AudioSection, InAudioData, 2, false, ClipItem1[Index], ClipItem2[Index], ClipIndex1[Index], ClipIndex2[Index], TrackIndex1, TrackIndex2))
 				{
 					return false;
 				}
@@ -820,7 +820,7 @@ bool FFCPXMLExportVisitor::ConstructAudioTrackNode(TSharedRef<FFCPXMLNode> InPar
 	return true;
 }
 
-bool FFCPXMLExportVisitor::ConstructVideoClipItemNode(TSharedRef<FFCPXMLNode> InParentNode, const TSharedPtr<FMovieSceneExportCinematicSectionData> InCinematicSectionData, const TSharedPtr<FMovieSceneExportCinematicMasterTrackData> InCinematicMasterTrackData, bool bInMasterClip)
+bool FFCPXMLExportVisitor::ConstructVideoClipItemNode(TSharedRef<FFCPXMLNode> InParentNode, const TSharedPtr<FMovieSceneExportCinematicSectionData> InCinematicSectionData, const TSharedPtr<FMovieSceneExportCinematicData> InCinematicData, bool bInMasterClip)
 {
 	if (!ExportData->IsExportDataValid() || !InCinematicSectionData.IsValid())
 	{
@@ -882,7 +882,7 @@ bool FFCPXMLExportVisitor::ConstructVideoClipItemNode(TSharedRef<FFCPXMLNode> In
 }
 
 bool FFCPXMLExportVisitor::ConstructAudioClipItemNode(TSharedRef<FFCPXMLNode> InParentNode, const TSharedPtr<FMovieSceneExportAudioSectionData> InAudioSectionData, 
-	const TSharedPtr<FMovieSceneExportAudioMasterTrackData> InAudioMasterTrackData, int32 InChannel, bool bInMasterClip, 
+	const TSharedPtr<FMovieSceneExportAudioData> InAudioData, int32 InChannel, bool bInMasterClip, 
 	const FString& InClipItemIdName1, const FString& InClipItemIdName2, int32 InClipIndex1, int32 InClipIndex2, int32 InTrackIndex1, int32 InTrackIndex2)
 {
 	if (!ExportData->IsExportDataValid() || !InAudioSectionData.IsValid())
diff --git a/Engine/Source/Editor/MovieSceneTools/Private/FCPXML/FCPXMLExport.h b/Engine/Source/Editor/MovieSceneTools/Private/FCPXML/FCPXMLExport.h
--- a/Engine/Source/Editor/MovieSceneTools/Private/FCPXML/FCPXMLExport.h
+++ b/Engine/Source/Editor/MovieSceneTools/Private/FCPXML/FCPXMLExport.h
@@ -55,9 +55,9 @@ public:
 	/** Creates master audio clip nodes. */
 	bool ConstructMasterAudioClipNodes(TSharedRef<FFCPXMLNode> InParentNode);
 	/** Creates master video clip node. */
-	bool ConstructMasterClipNode(TSharedRef<FFCPXMLNode> InParentNode, const TSharedPtr<FMovieSceneExportCinematicSectionData> InCinematicSectionData, const TSharedPtr<FMovieSceneExportCinematicMasterTrackData> InCinematicMasterTrackData);
+	bool ConstructMasterClipNode(TSharedRef<FFCPXMLNode> InParentNode, const TSharedPtr<FMovieSceneExportCinematicSectionData> InCinematicSectionData, const TSharedPtr<FMovieSceneExportCinematicData> InCinematicData);
 	/** Creates master audio clip node. */
-	bool ConstructMasterClipNode(TSharedRef<FFCPXMLNode> InParentNode, const TSharedPtr<FMovieSceneExportAudioSectionData> InAudioSectionData, const TSharedPtr<FMovieSceneExportAudioMasterTrackData> InCinematicMasterTrackData);
+	bool ConstructMasterClipNode(TSharedRef<FFCPXMLNode> InParentNode, const TSharedPtr<FMovieSceneExportAudioSectionData> InAudioSectionData, const TSharedPtr<FMovieSceneExportAudioData> InCinematicData);
 	/** Creates colorinfo node. */
 	bool ConstructColorInfoNode(TSharedRef<FFCPXMLNode> InParentNode);
 	/** Creates logginginfo node. */
@@ -71,14 +71,14 @@ public:
 	/** Creates audio node. */
 	bool ConstructAudioNode(TSharedRef<FFCPXMLNode> InParentNode);
 	/** Creates video track node. */
-	virtual bool ConstructVideoTrackNode(TSharedRef<FFCPXMLNode> InParentNode, const TSharedPtr<FMovieSceneExportCinematicTrackData> InCinematicTrackData, const TSharedPtr<FMovieSceneExportCinematicMasterTrackData> InCinematicMasterTrackData);
+	virtual bool ConstructVideoTrackNode(TSharedRef<FFCPXMLNode> InParentNode, const TSharedPtr<FMovieSceneExportCinematicTrackData> InCinematicTrackData, const TSharedPtr<FMovieSceneExportCinematicData> InCinematicData);
 	/** Creates audio track node. */
-	bool ConstructAudioTrackNode(TSharedRef<FFCPXMLNode> InParentNode, const TSharedPtr<FMovieSceneExportAudioTrackData> InAudioTrackData, const TSharedPtr<FMovieSceneExportAudioMasterTrackData> InAudioMasterTrackData, uint32 InTrackIndex, uint32 OutNumTracks);
+	bool ConstructAudioTrackNode(TSharedRef<FFCPXMLNode> InParentNode, const TSharedPtr<FMovieSceneExportAudioTrackData> InAudioTrackData, const TSharedPtr<FMovieSceneExportAudioData> InAudioData, uint32 InTrackIndex, uint32 OutNumTracks);
 	/** Creates video clip item node. */
-	bool ConstructVideoClipItemNode(TSharedRef<FFCPXMLNode> InParentNode, const TSharedPtr<FMovieSceneExportCinematicSectionData> InCinematicSectionData, const TSharedPtr<FMovieSceneExportCinematicMasterTrackData> InCinematicMasterTrackData, bool bInMasterClip);
+	bool ConstructVideoClipItemNode(TSharedRef<FFCPXMLNode> InParentNode, const TSharedPtr<FMovieSceneExportCinematicSectionData> InCinematicSectionData, const TSharedPtr<FMovieSceneExportCinematicData> InCinematicData, bool bInMasterClip);
 	/** Creates audio clip item node. */
 	bool ConstructAudioClipItemNode(TSharedRef<FFCPXMLNode> InParentNode, const TSharedPtr<FMovieSceneExportAudioSectionData> InAudioSectionData, 
-		const TSharedPtr<FMovieSceneExportAudioMasterTrackData> InAudioMasterTrackData, int32 InChannel, bool bInMasterClip, 
+		const TSharedPtr<FMovieSceneExportAudioData> InAudioData, int32 InChannel, bool bInMasterClip, 
 		const FString& InClipItemIdName1, const FString& InClipItemIdName2, int32 InClipIndex1, int32 InClipIndex2, int32 InTrackIndex1, int32 InTrackIndex2);
 	/** Creates video file node. */
 	bool ConstructVideoFileNode(TSharedRef<FFCPXMLNode> InParentNode, const TSharedPtr<FMovieSceneExportCinematicSectionData> InCinematicSectionData, int32 InDuration, bool bInMasterClip);
diff --git a/Engine/Source/Editor/MovieSceneTools/Private/FCPXML/FCPXMLImport.cpp b/Engine/Source/Editor/MovieSceneTools/Private/FCPXML/FCPXMLImport.cpp
--- a/Engine/Source/Editor/MovieSceneTools/Private/FCPXML/FCPXMLImport.cpp
+++ b/Engine/Source/Editor/MovieSceneTools/Private/FCPXML/FCPXMLImport.cpp
@@ -59,7 +59,7 @@ FFCPXMLImportVisitor::FFCPXMLImportVisitor(TSharedRef<FMovieSceneImportData> InI
 , bInAudioTrackNode(false)
 , CurrVideoTrackRowIndex(0)
 , CurrAudioTrackListIndex(0)
-, CurrAudioMasterTrack(nullptr)
+, CurrAudioData(nullptr)
 , CurrAudioTrackRowIndex(0)
 , bCurrImportAudioTrackIsStereoChannel(false)
 , MaxVideoTrackRowIndex(0)
@@ -70,7 +70,7 @@ FFCPXMLImportVisitor::FFCPXMLImportVisitor(TSharedRef<FMovieSceneImportData> InI
 	{
 		if (AudioTrackList[0].IsValid())
 		{
-			CurrAudioMasterTrack = AudioTrackList[0]->AudioTrackData;
+			CurrAudioData = AudioTrackList[0]->AudioTrackData;
 			CurrAudioTrackRowIndex = AudioTrackList[0]->RowIndex;
 		}
 	}
@@ -168,7 +168,7 @@ bool FFCPXMLImportVisitor::VisitNode(TSharedRef<FFCPXMLTrackNode> InTrackNode)
 				{
 					return false;
 				}
-				CurrAudioMasterTrack = AudioTrackList[CurrAudioTrackListIndex]->AudioTrackData;
+				CurrAudioData = AudioTrackList[CurrAudioTrackListIndex]->AudioTrackData;
 				CurrAudioTrackRowIndex = AudioTrackList[CurrAudioTrackListIndex]->RowIndex;
 			}
 			else
@@ -384,10 +384,10 @@ bool FFCPXMLImportVisitor::VisitAudioClipItemNode(TSharedRef<FFCPXMLClipItemNode
 
 	// Get next audio section based on the audio metadata
 	TSharedPtr<FMovieSceneImportAudioSectionData> AudioSectionData = nullptr;
-	TSharedPtr<FMovieSceneImportAudioMasterTrackData> AudioMasterTrackData = nullptr;
+	TSharedPtr<FMovieSceneImportAudioData> AudioData = nullptr;
 	if (AudioMetadata.IsValid())
 	{
-		if (!GetNextAudioSection(AudioMetadata, AudioMasterTrackData, AudioSectionData))
+		if (!GetNextAudioSection(AudioMetadata, AudioData, AudioSectionData))
 		{
 			return false;
 		}
@@ -395,12 +395,12 @@ bool FFCPXMLImportVisitor::VisitAudioClipItemNode(TSharedRef<FFCPXMLClipItemNode
 
 	if (AudioSectionData.IsValid())
 	{
-		if (CurrAudioMasterTrack.IsValid() && CurrAudioMasterTrack->MovieSceneTrack != nullptr &&
-			AudioMasterTrackData.IsValid() && AudioMasterTrackData->MovieSceneTrack != nullptr &&
-			CurrAudioMasterTrack->MovieSceneTrack->GetFullName() != AudioMasterTrackData->MovieSceneTrack->GetFullName())
+		if (CurrAudioData.IsValid() && CurrAudioData->MovieSceneTrack != nullptr &&
+			AudioData.IsValid() && AudioData->MovieSceneTrack != nullptr &&
+			CurrAudioData->MovieSceneTrack->GetFullName() != CurrAudioData->MovieSceneTrack->GetFullName())
 		{
 			// Move audio section
-			if (!ImportData->MoveAudioSection(AudioSectionData, AudioMasterTrackData, CurrAudioMasterTrack, CurrAudioTrackRowIndex))
+			if (!ImportData->MoveAudioSection(AudioSectionData, AudioData, CurrAudioData, CurrAudioTrackRowIndex))
 			{
 				return false;
 			}
@@ -418,7 +418,7 @@ bool FFCPXMLImportVisitor::VisitAudioClipItemNode(TSharedRef<FFCPXMLClipItemNode
 		if (!SoundWaveName.IsEmpty())
 		{
 			// Add new audio section
-			AudioSectionData = ImportData->CreateAudioSection(SoundWaveName, bUseSoundPathName, CurrAudioMasterTrack, CurrAudioTrackRowIndex, FrameRate, Start, End, StartOffset);
+			AudioSectionData = ImportData->CreateAudioSection(SoundWaveName, bUseSoundPathName, CurrAudioData, CurrAudioTrackRowIndex, FrameRate, Start, End, StartOffset);
 			if (!AudioSectionData.IsValid())
 			{
 				return false;
@@ -562,13 +562,13 @@ bool FFCPXMLImportVisitor::ConstructAudioTrackList()
 		return false;
 	}
 
-	for (TSharedPtr<FMovieSceneImportAudioMasterTrackData> MasterTrackData : ImportData->MovieSceneData->AudioMasterTracks)
+	for (TSharedPtr<FMovieSceneImportAudioData> AudioData : ImportData->MovieSceneData->AudioData)
 	{
-		for (TSharedPtr<FMovieSceneImportAudioTrackData> TrackData : MasterTrackData->AudioTracks)
+		for (TSharedPtr<FMovieSceneImportAudioTrackData> AudioTrackData : AudioData->AudioTracks)
 		{
-			if (TrackData.IsValid())
+			if (AudioTrackData.IsValid())
 			{
-				TSharedPtr<FFCPXMLImportAudioTrackListItem> ListItem = MakeShared<FFCPXMLImportAudioTrackListItem>(MasterTrackData, TrackData->RowIndex);
+				TSharedPtr<FFCPXMLImportAudioTrackListItem> ListItem = MakeShared<FFCPXMLImportAudioTrackListItem>(AudioData, AudioTrackData->RowIndex);
 				AudioTrackList.Add(ListItem);
 			}
 		}
@@ -685,7 +685,7 @@ TSharedPtr<FFCPXMLImportAudioMetadata> FFCPXMLImportVisitor::GetAudioMetadataObj
 }
 
 /** Get audio section path name based on node metadata and masterclip id */
-bool FFCPXMLImportVisitor::GetNextAudioSection(TSharedPtr<FFCPXMLImportAudioMetadata> InAudioMetadata, TSharedPtr<FMovieSceneImportAudioMasterTrackData>& OutAudioMasterTrackData, TSharedPtr<FMovieSceneImportAudioSectionData>& OutAudioSectionData)
+bool FFCPXMLImportVisitor::GetNextAudioSection(TSharedPtr<FFCPXMLImportAudioMetadata> InAudioMetadata, TSharedPtr<FMovieSceneImportAudioData>& OutAudioData, TSharedPtr<FMovieSceneImportAudioSectionData>& OutAudioSectionData)
 {
 	if (!InAudioMetadata.IsValid())
 	{
@@ -707,7 +707,7 @@ bool FFCPXMLImportVisitor::GetNextAudioSection(TSharedPtr<FFCPXMLImportAudioMeta
 	// Find actual audio section
 	if (!AudioSectionPathName.IsEmpty())
 	{
-		OutAudioSectionData = ImportData->FindAudioSection(AudioSectionPathName, OutAudioMasterTrackData);
+		OutAudioSectionData = ImportData->FindAudioSection(AudioSectionPathName, OutAudioData);
 	}
 
 	return true;
diff --git a/Engine/Source/Editor/MovieSceneTools/Private/FCPXML/FCPXMLImport.h b/Engine/Source/Editor/MovieSceneTools/Private/FCPXML/FCPXMLImport.h
--- a/Engine/Source/Editor/MovieSceneTools/Private/FCPXML/FCPXMLImport.h
+++ b/Engine/Source/Editor/MovieSceneTools/Private/FCPXML/FCPXMLImport.h
@@ -11,9 +11,9 @@
 
 struct FFCPXMLImportAudioTrackListItem
 {
-	FFCPXMLImportAudioTrackListItem(TSharedPtr<FMovieSceneImportAudioMasterTrackData> InAudioTrackData, int InRowIndex) : AudioTrackData(InAudioTrackData), RowIndex(InRowIndex) {}
+	FFCPXMLImportAudioTrackListItem(TSharedPtr<FMovieSceneImportAudioData> InAudioTrackData, int InRowIndex) : AudioTrackData(InAudioTrackData), RowIndex(InRowIndex) {}
 
-	TSharedPtr<FMovieSceneImportAudioMasterTrackData> AudioTrackData;
+	TSharedPtr<FMovieSceneImportAudioData> AudioTrackData;
 	int32 RowIndex;
 };
 
@@ -75,7 +75,7 @@ private:
 	/** Query which channel audio clipitem node represents */
 	int32 GetAudioClipItemNodeChannel(TSharedRef<FFCPXMLClipItemNode> InClipItemNode, const FString& InClipItemId);
 
-	/** Construct list of master tracks and row indexes */
+	/** Construct list of audio tracks and row indexes */
 	bool ConstructAudioTrackList();
 
 	/** Add entry to master cinematic clip section name map */
@@ -92,7 +92,7 @@ private:
 	/** Get audio metadata object based on log note and master clip id */
 	TSharedPtr<FFCPXMLImportAudioMetadata> GetAudioMetadataObject(const FString& InLogNote, const FString& InMasterClipId);
 	/** Get next audio section based on audio metadata */
-	bool GetNextAudioSection(TSharedPtr<FFCPXMLImportAudioMetadata> InAudioMetadata, TSharedPtr<FMovieSceneImportAudioMasterTrackData>& OutAudioMasterTrackData, TSharedPtr<FMovieSceneImportAudioSectionData>& OutAudioSectionData);
+	bool GetNextAudioSection(TSharedPtr<FFCPXMLImportAudioMetadata> InAudioMetadata, TSharedPtr<FMovieSceneImportAudioData>& OutAudioData, TSharedPtr<FMovieSceneImportAudioSectionData>& OutAudioSectionData);
 
 	/** Add entry to metadata map */
 	bool AddMasterClipLoggingNode(const FString& InMasterClipIdName, TSharedPtr<FFCPXMLNode> InLoggingInfoNode);
@@ -135,7 +135,7 @@ private:
 	/** Map of masterclip names to logging info nodes */
 	TMap<FString, TSharedPtr<FFCPXMLNode>> MasterClipLoggingNodeMap;
 
-	/** List of existing master tracks including row index */
+	/** List of existing audio tracks including row index */
 	TArray< TSharedPtr<FFCPXMLImportAudioTrackListItem> > AudioTrackList;
 
 	/** Traversal state variables */
@@ -147,7 +147,7 @@ private:
 
 	int32 CurrVideoTrackRowIndex;
 	int32 CurrAudioTrackListIndex;
-	TSharedPtr<FMovieSceneImportAudioMasterTrackData> CurrAudioMasterTrack;
+	TSharedPtr<FMovieSceneImportAudioData> CurrAudioData;
 	int32 CurrAudioTrackRowIndex;
 	bool bCurrImportAudioTrackIsStereoChannel;
 
diff --git a/Engine/Source/Editor/MovieSceneTools/Private/FCPXML/FCPXMLMetadataExport.cpp b/Engine/Source/Editor/MovieSceneTools/Private/FCPXML/FCPXMLMetadataExport.cpp
--- a/Engine/Source/Editor/MovieSceneTools/Private/FCPXML/FCPXMLMetadataExport.cpp
+++ b/Engine/Source/Editor/MovieSceneTools/Private/FCPXML/FCPXMLMetadataExport.cpp
@@ -40,7 +40,7 @@ FString GetPreferredFormat(const TArray<FString>& Extensions)
 
 bool FFCPXMLMetadataExportVisitor::ConstructMasterVideoClipNodes(TSharedRef<FFCPXMLNode> InParentNode)
 {
-	if (!ExportData->IsExportDataValid() || !ExportData->MovieSceneData.IsValid() || !ExportData->MovieSceneData->CinematicMasterTrack.IsValid())
+	if (!ExportData->IsExportDataValid() || !ExportData->MovieSceneData.IsValid() || !ExportData->MovieSceneData->CinematicData.IsValid())
 	{
 		return false;
 	}
@@ -243,7 +243,7 @@ bool FFCPXMLMetadataExportVisitor::GetFileIdName(const FString& InFileName, FStr
 	return true;
 }
 
-bool FFCPXMLMetadataExportVisitor::ConstructVideoTrackNode(TSharedRef<FFCPXMLNode> InParentNode, const TSharedPtr<FMovieSceneExportCinematicTrackData> InCinematicTrackData, const TSharedPtr<FMovieSceneExportCinematicMasterTrackData> InCinematicMasterTrackData)
+bool FFCPXMLMetadataExportVisitor::ConstructVideoTrackNode(TSharedRef<FFCPXMLNode> InParentNode, const TSharedPtr<FMovieSceneExportCinematicTrackData> InCinematicTrackData, const TSharedPtr<FMovieSceneExportCinematicData> InCinematicData)
 {
 	if (!ExportData->IsExportDataValid() || !InCinematicTrackData.IsValid())
 	{
