Sequencer: Remove Master from Sequencer

diff --git a/Engine/Source/Editor/MovieSceneTools/Private/FCPXML/FCPXMLMetadataExport.h b/Engine/Source/Editor/MovieSceneTools/Private/FCPXML/FCPXMLMetadataExport.h
--- a/Engine/Source/Editor/MovieSceneTools/Private/FCPXML/FCPXMLMetadataExport.h
+++ b/Engine/Source/Editor/MovieSceneTools/Private/FCPXML/FCPXMLMetadataExport.h
@@ -18,7 +18,7 @@ public:
 
 	/** Creates master video clip nodes. */
 	virtual bool ConstructMasterVideoClipNodes(TSharedRef<FFCPXMLNode> InParentNode) override;
-	virtual bool ConstructVideoTrackNode(TSharedRef<FFCPXMLNode> InParentNode, const TSharedPtr<FMovieSceneExportCinematicTrackData> InCinematicTrackData, const TSharedPtr<FMovieSceneExportCinematicMasterTrackData> InCinematicMasterTrackData) override;
+	virtual bool ConstructVideoTrackNode(TSharedRef<FFCPXMLNode> InParentNode, const TSharedPtr<FMovieSceneExportCinematicTrackData> InCinematicTrackData, const TSharedPtr<FMovieSceneExportCinematicData> InCinematicData) override;
 
 	bool ConstructMasterClipNode(TSharedRef<FFCPXMLNode> InParentNode, const FString& InClipName, const FMovieSceneExportMetadataShot& InShotMetadata, const FMovieSceneExportMetadataClip& InClipMetadata, const int32 HandleFrames);
 	bool ConstructVideoClipItemNode(TSharedRef<FFCPXMLNode> InParentNode, const FString& InClipName, const FMovieSceneExportMetadataClip& InClipMetadata, const int32 HandleFrames, bool bInMasterClip);
diff --git a/Engine/Source/Editor/MovieSceneTools/Private/FCPXML/FCPXMLMovieSceneTranslator.cpp b/Engine/Source/Editor/MovieSceneTools/Private/FCPXML/FCPXMLMovieSceneTranslator.cpp
--- a/Engine/Source/Editor/MovieSceneTools/Private/FCPXML/FCPXMLMovieSceneTranslator.cpp
+++ b/Engine/Source/Editor/MovieSceneTools/Private/FCPXML/FCPXMLMovieSceneTranslator.cpp
@@ -84,10 +84,10 @@ bool FFCPXMLImporter::Import(UMovieScene* InMovieScene, FFrameRate InFrameRate,
 		bSuccess = FCPXMLFile->Accept(ImportVisitor);
 
 		// Flip rows
-		if (ImportData->GetCinematicMasterTrackData(false).IsValid())
+		if (ImportData->GetCinematicData(false).IsValid())
 		{
 			int32 MaxVideoTrackRowIndex = ImportVisitor.GetMaxVideoTrackRowIndex();
-			for (TSharedPtr<FMovieSceneImportCinematicSectionData> SectionData : ImportData->GetCinematicMasterTrackData(false)->CinematicSections)
+			for (TSharedPtr<FMovieSceneImportCinematicSectionData> SectionData : ImportData->GetCinematicData(false)->CinematicSections)
 			{
 				if (SectionData.IsValid())
 				{
@@ -99,10 +99,10 @@ bool FFCPXMLImporter::Import(UMovieScene* InMovieScene, FFrameRate InFrameRate,
 				}
 			}
 		}
-		if (ImportData->GetAudioMasterTrackData().IsValid())
+		if (ImportData->GetAudioData().IsValid())
 		{
 			int32 MaxAudioTrackRowIndex = ImportVisitor.GetMaxAudioTrackRowIndex();
-			for (TSharedPtr<FMovieSceneImportAudioSectionData> SectionData : ImportData->GetAudioMasterTrackData()->AudioSections)
+			for (TSharedPtr<FMovieSceneImportAudioSectionData> SectionData : ImportData->GetAudioData()->AudioSections)
 			{
 				if (SectionData.IsValid())
 				{
diff --git a/Engine/Source/Editor/MovieSceneTools/Private/MovieSceneObjectBindingIDPicker.cpp b/Engine/Source/Editor/MovieSceneTools/Private/MovieSceneObjectBindingIDPicker.cpp
--- a/Engine/Source/Editor/MovieSceneTools/Private/MovieSceneObjectBindingIDPicker.cpp
+++ b/Engine/Source/Editor/MovieSceneTools/Private/MovieSceneObjectBindingIDPicker.cpp
@@ -187,7 +187,7 @@ TSharedRef<SWidget> FMovieSceneObjectBindingIDPicker::GetWarningWidget()
 		.VAlign(VAlign_Center)
 		.ContentPadding(FMargin(0))
 		.ButtonStyle(FAppStyle::Get(), "HoverHintOnly")
-		.ToolTipText(LOCTEXT("FixedBindingWarningText", "This binding is fixed to the current Master Sequence hierarchy, so will break if evaluated in a different hierarchy.\nClick here to fix this problem."))
+		.ToolTipText(LOCTEXT("FixedBindingWarningText", "This binding is fixed to the current Root Sequence hierarchy, so will break if evaluated in a different hierarchy.\nClick here to fix this problem."))
 		.Visibility_Raw(this, &FMovieSceneObjectBindingIDPicker::GetFixedWarningVisibility)
 		.OnClicked_Raw(this, &FMovieSceneObjectBindingIDPicker::AttemptBindingFixup)
 		[
@@ -318,9 +318,9 @@ void FMovieSceneObjectBindingIDPicker::SetCurrentValueFromFixed(UE::MovieScene::
 	else
 	{
 		// Attempt to remap the desired binding to the current local sequence by either making it local to this sequence
-		// or specifying a parent index so that this binding is still able to resolve correctly if the master sequence is added
+		// or specifying a parent index so that this binding is still able to resolve correctly if the root sequence is added
 		// as a subsequence elsewhere
-		// This ensures that you can work on sub sequences on their own, or within a master sequence and the binding will resolve correctly.
+		// This ensures that you can work on sub sequences on their own, or within a root sequence and the binding will resolve correctly.
 		SetCurrentValue(InValue.ConvertToRelative(LocalSequenceID, Hierarchy));
 	}
 }
diff --git a/Engine/Source/Editor/MovieSceneTools/Private/MovieSceneToolHelpers.cpp b/Engine/Source/Editor/MovieSceneTools/Private/MovieSceneToolHelpers.cpp
--- a/Engine/Source/Editor/MovieSceneTools/Private/MovieSceneToolHelpers.cpp
+++ b/Engine/Source/Editor/MovieSceneTools/Private/MovieSceneToolHelpers.cpp
@@ -380,7 +380,7 @@ FString MovieSceneToolHelpers::GenerateNewShotPath(UMovieScene* SequenceMovieSce
 
 	UObject* SequenceAsset = SequenceMovieScene->GetOuter();
 	UPackage* SequencePackage = SequenceAsset->GetOutermost();
-	FString SequencePackageName = SequencePackage->GetName(); // ie. /Game/cine/max/master
+	FString SequencePackageName = SequencePackage->GetName(); // ie. /Game/cine/max/root
 	int32 LastSlashPos = SequencePackageName.Find(TEXT("/"), ESearchCase::IgnoreCase, ESearchDir::FromEnd);
 	FString SequencePath = SequencePackageName.Left(LastSlashPos);
 
@@ -1215,7 +1215,7 @@ void MovieSceneToolHelpers::LockCameraActorToViewport(const TSharedPtr<ISequence
 void MovieSceneToolHelpers::CreateCameraCutSectionForCamera(UMovieScene* OwnerMovieScene, FGuid CameraGuid, FFrameNumber FrameNumber)
 {
 	// If there's a cinematic shot track, no need to set this camera to a shot
-	UMovieSceneTrack* CinematicShotTrack = OwnerMovieScene->FindMasterTrack(UMovieSceneCinematicShotTrack::StaticClass());
+	UMovieSceneTrack* CinematicShotTrack = OwnerMovieScene->FindTrack(UMovieSceneCinematicShotTrack::StaticClass());
 	if (CinematicShotTrack)
 	{
 		return;
@@ -3779,7 +3779,7 @@ void ExportLevelMesh(UnFbx::FFbxExporter* Exporter, ULevel* Level, IMovieScenePl
 	Exporter->ExportLevelMesh(Level, !bSelectedOnly, ActorToExport, NodeNameAdapter, bSaveAnimSeq);
 }
 
-bool MovieSceneToolHelpers::ExportFBX(UWorld* World, UMovieScene* MovieScene, IMovieScenePlayer* Player, const TArray<FGuid>& Bindings, const TArray<UMovieSceneTrack*>& MasterTracks, INodeNameAdapter& NodeNameAdapter, FMovieSceneSequenceIDRef& Template, const FString& InFBXFileName, FMovieSceneSequenceTransform& RootToLocalTransform)
+bool MovieSceneToolHelpers::ExportFBX(UWorld* World, UMovieScene* MovieScene, IMovieScenePlayer* Player, const TArray<FGuid>& Bindings, const TArray<UMovieSceneTrack*>& Tracks, INodeNameAdapter& NodeNameAdapter, FMovieSceneSequenceIDRef& Template, const FString& InFBXFileName, FMovieSceneSequenceTransform& RootToLocalTransform)
 {
 	UnFbx::FFbxExporter* Exporter = UnFbx::FFbxExporter::GetInstance();
 
@@ -3803,14 +3803,9 @@ bool MovieSceneToolHelpers::ExportFBX(UWorld* World, UMovieScene* MovieScene, IM
 
 	Exporter->ExportLevelSequence(MovieScene, Bindings, Player, NodeNameAdapter, Template, RootToLocalTransform);
 
-	//Export given master tracks
-
-	for (UMovieSceneTrack* MasterTrack : MasterTracks)
-	{
-		TArray<UMovieSceneTrack*> Tracks;
-		Tracks.Add(MasterTrack);
-		Exporter->ExportLevelSequenceTracks(MovieScene, Player, Template, nullptr, nullptr, Tracks, RootToLocalTransform);
-	}
+	//Export given tracks
+	Exporter->ExportLevelSequenceTracks(MovieScene, Player, Template, nullptr, nullptr, Tracks, RootToLocalTransform);
+	
 	// Save to disk
 	Exporter->WriteToFile(*InFBXFileName);
 
@@ -4103,7 +4098,7 @@ FSpawnableRestoreState::FSpawnableRestoreState(UMovieScene* MovieScene)
 			
 			// Spawnable could be in a subscene, so temporarily override it to persist throughout
 			SpawnOwnershipMap.Add(Spawnable.GetGuid(), Spawnable.GetSpawnOwnership());
-			Spawnable.SetSpawnOwnership(ESpawnOwnership::MasterSequence);
+			Spawnable.SetSpawnOwnership(ESpawnOwnership::RootSequence);
 
 			UMovieSceneSpawnSection* SpawnSection = Cast<UMovieSceneSpawnSection>(SpawnTrack->GetAllSections()[0]);
 			SpawnSection->Modify();
diff --git a/Engine/Source/Editor/MovieSceneTools/Private/MovieSceneTranslator.cpp b/Engine/Source/Editor/MovieSceneTools/Private/MovieSceneTranslator.cpp
--- a/Engine/Source/Editor/MovieSceneTools/Private/MovieSceneTranslator.cpp
+++ b/Engine/Source/Editor/MovieSceneTools/Private/MovieSceneTranslator.cpp
@@ -21,7 +21,7 @@
 
 namespace
 {
-	bool SupportMultipleAudioMasterTracks()
+	bool SupportMultipleAudioTracks()
 	{
 		return true;
 	}
@@ -130,34 +130,34 @@ bool FMovieSceneExportData::ConstructMovieSceneData(const UMovieScene* InMovieSc
 	MovieSceneData->TickResolution = TickResolution;
 	MovieSceneData->Duration = ConvertFrameTime(UE::MovieScene::DiscreteSize(PlaybackRange), TickResolution, FrameRate).FrameNumber.Value;
 
-	bool bFoundCinematicMasterTrack = false;
+	bool bFoundCinematicTrack = false;
 
-	// sort audio master tracks
-	TMap<int32, TSharedPtr<FMovieSceneExportAudioMasterTrackData>> AudioTrackMap;
+	// sort audio tracks
+	TMap<int32, TSharedPtr<FMovieSceneExportAudioData>> AudioTrackMap;
 
-	const TArray<UMovieSceneTrack*> MasterTracks = InMovieScene->GetMasterTracks();
-	for (UMovieSceneTrack* MasterTrack : MasterTracks)
+	const TArray<UMovieSceneTrack*> Tracks = InMovieScene->GetTracks();
+	for (UMovieSceneTrack* Track : Tracks)
 	{
-		if (!bFoundCinematicMasterTrack && MasterTrack->IsA(UMovieSceneCinematicShotTrack::StaticClass()))
+		if (!bFoundCinematicTrack && Track->IsA(UMovieSceneCinematicShotTrack::StaticClass()))
 		{
-			const UMovieSceneCinematicShotTrack* CinematicTrack = Cast<UMovieSceneCinematicShotTrack>(MasterTrack);
-			if (CinematicTrack == nullptr || !ConstructCinematicMasterTrackData(InMovieScene, CinematicTrack))
+			const UMovieSceneCinematicShotTrack* CinematicTrack = Cast<UMovieSceneCinematicShotTrack>(Track);
+			if (CinematicTrack == nullptr || !ConstructCinematicData(InMovieScene, CinematicTrack))
 			{ 
 				return false; 
 			}
-			bFoundCinematicMasterTrack = true;
+			bFoundCinematicTrack = true;
 		}
-		else if (MasterTrack->IsA(UMovieSceneAudioTrack::StaticClass()))
+		else if (Track->IsA(UMovieSceneAudioTrack::StaticClass()))
 		{
-			const UMovieSceneAudioTrack* AudioTrack = Cast<UMovieSceneAudioTrack>(MasterTrack);
-			if (AudioTrack == nullptr || !ConstructAudioMasterTrackData(InMovieScene, AudioTrack, AudioTrackMap))
+			const UMovieSceneAudioTrack* AudioTrack = Cast<UMovieSceneAudioTrack>(Track);
+			if (AudioTrack == nullptr || !ConstructAudioData(InMovieScene, AudioTrack, AudioTrackMap))
 			{ 
 				return false; 
 			}
 		}
 	}
 
-	// sort the audio master tracks by their sorting index and add to the AudioMasterTracks array
+	// sort the audio tracks by their sorting index and add to the AudioTracks array
 	if (AudioTrackMap.Num() > 0)
 	{
 		AudioTrackMap.KeySort([](int32 A, int32 B) {
@@ -168,9 +168,9 @@ bool FMovieSceneExportData::ConstructMovieSceneData(const UMovieScene* InMovieSc
 		{
 			if (Elem.Value.IsValid())
 			{
-				MovieSceneData->AudioMasterTracks.Add(Elem.Value);
+				MovieSceneData->AudioData.Add(Elem.Value);
 
-				if (!SupportMultipleAudioMasterTracks())
+				if (!SupportMultipleAudioTracks())
 				{
 					break;
 				}
@@ -181,19 +181,19 @@ bool FMovieSceneExportData::ConstructMovieSceneData(const UMovieScene* InMovieSc
 	return true;
 }
 
-bool FMovieSceneExportData::ConstructCinematicMasterTrackData(const UMovieScene* InMovieScene, const UMovieSceneCinematicShotTrack* InCinematicMasterTrack)
+bool FMovieSceneExportData::ConstructCinematicData(const UMovieScene* InMovieScene, const UMovieSceneCinematicShotTrack* InCinematicTrack)
 {
 	if (InMovieScene == nullptr || !MovieSceneData.IsValid())
 	{
 		return false;
 	}
 
-	TSharedPtr<FMovieSceneExportCinematicMasterTrackData> MasterTrackData = MakeShared<FMovieSceneExportCinematicMasterTrackData>();
-	MasterTrackData->MovieSceneTrack = InCinematicMasterTrack;
-	MovieSceneData->CinematicMasterTrack = MasterTrackData;
+	TSharedPtr<FMovieSceneExportCinematicData> CinematicData = MakeShared<FMovieSceneExportCinematicData>();
+	CinematicData->MovieSceneTrack = InCinematicTrack;
+	MovieSceneData->CinematicData = CinematicData;
 
 	// Make a list of shotnames and how many duplicates
-	for (UMovieSceneSection* Section : InCinematicMasterTrack->GetAllSections())
+	for (UMovieSceneSection* Section : InCinematicTrack->GetAllSections())
 	{
 		const UMovieSceneCinematicShotSection* CinematicSection = Cast<UMovieSceneCinematicShotSection>(Section);
 
@@ -224,13 +224,13 @@ bool FMovieSceneExportData::ConstructCinematicMasterTrackData(const UMovieScene*
 	// Construct sections & create track row index array
 	TArray<int32> CinematicTrackRowIndices;
 
-	for (UMovieSceneSection* Section : InCinematicMasterTrack->GetAllSections())
+	for (UMovieSceneSection* Section : InCinematicTrack->GetAllSections())
 	{
 		const UMovieSceneCinematicShotSection* CinematicSection = Cast<UMovieSceneCinematicShotSection>(Section);
 
 		if (CinematicSection != nullptr && CinematicSection->GetSequence() != nullptr)
 		{
-			if (!ConstructCinematicSectionData(InMovieScene, MasterTrackData, CinematicSection))
+			if (!ConstructCinematicSectionData(InMovieScene, CinematicData, CinematicSection))
 			{
 				return false;
 			}
@@ -248,7 +248,7 @@ bool FMovieSceneExportData::ConstructCinematicMasterTrackData(const UMovieScene*
 
 	for (int32 CinematicTrackRowIndex : CinematicTrackRowIndices)
 	{	
-		if (!ConstructCinematicTrackData(InMovieScene, MasterTrackData, CinematicTrackRowIndex))
+		if (!ConstructCinematicTrackData(InMovieScene, CinematicData, CinematicTrackRowIndex))
 		{
 			return false;
 		}		
@@ -257,18 +257,18 @@ bool FMovieSceneExportData::ConstructCinematicMasterTrackData(const UMovieScene*
 	return true;
 }
 
-bool FMovieSceneExportData::ConstructCinematicTrackData(const UMovieScene* InMovieScene, TSharedPtr<FMovieSceneExportCinematicMasterTrackData> InCinematicMasterTrackData, int32 InRowIndex)
+bool FMovieSceneExportData::ConstructCinematicTrackData(const UMovieScene* InMovieScene, TSharedPtr<FMovieSceneExportCinematicData> InCinematicData, int32 InRowIndex)
 {
-	if (InMovieScene == nullptr || !InCinematicMasterTrackData.IsValid() || !MovieSceneData.IsValid() || !MovieSceneData->CinematicMasterTrack.IsValid())
+	if (InMovieScene == nullptr || !InCinematicData.IsValid() || !MovieSceneData.IsValid() || !MovieSceneData->CinematicData.IsValid())
 	{
 		return false;
 	}
 
 	TSharedPtr<FMovieSceneExportCinematicTrackData> TrackData = MakeShared<FMovieSceneExportCinematicTrackData>();
 	TrackData->RowIndex = InRowIndex;
-	MovieSceneData->CinematicMasterTrack->CinematicTracks.Add(TrackData);
+	MovieSceneData->CinematicData->CinematicTracks.Add(TrackData);
 
-	for (TSharedPtr<FMovieSceneExportCinematicSectionData> Section : InCinematicMasterTrackData->CinematicSections)
+	for (TSharedPtr<FMovieSceneExportCinematicSectionData> Section : InCinematicData->CinematicSections)
 	{
 		if (Section.IsValid() && Section->RowIndex == InRowIndex)
 		{
@@ -279,23 +279,23 @@ bool FMovieSceneExportData::ConstructCinematicTrackData(const UMovieScene* InMov
 	return true;
 }
 
-bool FMovieSceneExportData::ConstructAudioMasterTrackData(const UMovieScene* InMovieScene, const UMovieSceneAudioTrack* InAudioMasterTrack, TMap<int32, TSharedPtr<FMovieSceneExportAudioMasterTrackData>>& InAudioTrackMap)
+bool FMovieSceneExportData::ConstructAudioData(const UMovieScene* InMovieScene, const UMovieSceneAudioTrack* InAudioTrack, TMap<int32, TSharedPtr<FMovieSceneExportAudioData>>& InAudioTrackMap)
 {
 	if (InMovieScene == nullptr || !MovieSceneData.IsValid())
 	{
 		return false;
 	}
 
-	TSharedPtr<FMovieSceneExportAudioMasterTrackData> MasterTrackData = MakeShared<FMovieSceneExportAudioMasterTrackData>();
-	MasterTrackData->MovieSceneTrack = InAudioMasterTrack;
-	InAudioTrackMap.Add(InAudioMasterTrack->GetSortingOrder(), MasterTrackData);
+	TSharedPtr<FMovieSceneExportAudioData> TrackData = MakeShared<FMovieSceneExportAudioData>();
+	TrackData->MovieSceneTrack = InAudioTrack;
+	InAudioTrackMap.Add(InAudioTrack->GetSortingOrder(), TrackData);
 
 	// Construct sections & create track row index array
 	TArray<int32> AudioTrackRowIndices;
 
 	TArray<FString> SectionPathNames;
 
-	for (UMovieSceneSection* Section : InAudioMasterTrack->GetAudioSections())
+	for (UMovieSceneSection* Section : InAudioTrack->GetAudioSections())
 	{
 		const UMovieSceneAudioSection* AudioSection = Cast<UMovieSceneAudioSection>(Section);
 		if (AudioSection == nullptr)
@@ -311,7 +311,7 @@ bool FMovieSceneExportData::ConstructAudioMasterTrackData(const UMovieScene* InM
 
 		if (AudioSectionIsSoundWave(AudioSection))
 		{
-			if (!ConstructAudioSectionData(InMovieScene, MasterTrackData, AudioSection))
+			if (!ConstructAudioSectionData(InMovieScene, TrackData, AudioSection))
 			{
 				return false;
 			}
@@ -330,7 +330,7 @@ bool FMovieSceneExportData::ConstructAudioMasterTrackData(const UMovieScene* InM
 
 	for (int32 AudioTrackRowIndex : AudioTrackRowIndices)
 	{
-		if (!ConstructAudioTrackData(InMovieScene, MasterTrackData, AudioTrackRowIndex))
+		if (!ConstructAudioData(InMovieScene, TrackData, AudioTrackRowIndex))
 		{
 			return false;
 		}
@@ -339,9 +339,9 @@ bool FMovieSceneExportData::ConstructAudioMasterTrackData(const UMovieScene* InM
 	return true;
 }
 
-bool FMovieSceneExportData::ConstructAudioTrackData(const UMovieScene* InMovieScene, TSharedPtr<FMovieSceneExportAudioMasterTrackData> InAudioMasterTrackData, int32 InRowIndex)
+bool FMovieSceneExportData::ConstructAudioData(const UMovieScene* InMovieScene, TSharedPtr<FMovieSceneExportAudioData> InAudioData, int32 InRowIndex)
 {
-	if (InMovieScene == nullptr || !InAudioMasterTrackData.IsValid())
+	if (InMovieScene == nullptr || !InAudioData.IsValid())
 	{
 		return false;
 	}
@@ -349,9 +349,9 @@ bool FMovieSceneExportData::ConstructAudioTrackData(const UMovieScene* InMovieSc
 	TSharedPtr<FMovieSceneExportAudioTrackData> TrackData = MakeShared<FMovieSceneExportAudioTrackData>();
 	TrackData->SampleRate = DefaultAudioSampleRate;
 	TrackData->RowIndex = InRowIndex;
-	InAudioMasterTrackData->AudioTracks.Add(TrackData);
+	InAudioData->AudioTracks.Add(TrackData);
 
-	for (TSharedPtr<FMovieSceneExportAudioSectionData> Section : InAudioMasterTrackData->AudioSections)
+	for (TSharedPtr<FMovieSceneExportAudioSectionData> Section : InAudioData->AudioSections)
 	{
 		if (Section.IsValid() && Section->RowIndex == InRowIndex)
 		{
@@ -362,15 +362,15 @@ bool FMovieSceneExportData::ConstructAudioTrackData(const UMovieScene* InMovieSc
 	return true;
 }
 
-bool FMovieSceneExportData::ConstructCinematicSectionData(const UMovieScene* InMovieScene, TSharedPtr<FMovieSceneExportCinematicMasterTrackData> InMasterTrackData, const UMovieSceneCinematicShotSection* InCinematicSection)
+bool FMovieSceneExportData::ConstructCinematicSectionData(const UMovieScene* InMovieScene, TSharedPtr<FMovieSceneExportCinematicData> InCinematicData, const UMovieSceneCinematicShotSection* InCinematicSection)
 {
-	if (InMovieScene == nullptr || !InMasterTrackData.IsValid() || InCinematicSection == nullptr)
+	if (InMovieScene == nullptr || !InCinematicData.IsValid() || InCinematicSection == nullptr)
 	{
 		return false;
 	}
 
 	TSharedPtr<FMovieSceneExportCinematicSectionData> SectionData = MakeShared<FMovieSceneExportCinematicSectionData>();
-	InMasterTrackData->CinematicSections.Add(SectionData);
+	InCinematicData->CinematicSections.Add(SectionData);
 
 	SectionData->DisplayName = InCinematicSection->GetShotDisplayName();
 	if (DuplicateShotNameCounts.Contains(SectionData->DisplayName))
@@ -387,9 +387,9 @@ bool FMovieSceneExportData::ConstructCinematicSectionData(const UMovieScene* InM
 	return true;
 }
 
-bool FMovieSceneExportData::ConstructAudioSectionData(const UMovieScene* InMovieScene, TSharedPtr<FMovieSceneExportAudioMasterTrackData> InTrackData, const UMovieSceneAudioSection* InAudioSection)
+bool FMovieSceneExportData::ConstructAudioSectionData(const UMovieScene* InMovieScene, TSharedPtr<FMovieSceneExportAudioData> InAudioData, const UMovieSceneAudioSection* InAudioSection)
 {
-	if (InMovieScene == nullptr || !InTrackData.IsValid() || InAudioSection == nullptr)
+	if (InMovieScene == nullptr || !InAudioData.IsValid() || InAudioSection == nullptr)
 	{
 		return false;
 	}
@@ -408,7 +408,7 @@ bool FMovieSceneExportData::ConstructAudioSectionData(const UMovieScene* InMovie
 	}
 
 	TSharedPtr<FMovieSceneExportAudioSectionData> SectionData = MakeShared<FMovieSceneExportAudioSectionData>();
-	InTrackData->AudioSections.Add(SectionData);
+	InAudioData->AudioSections.Add(SectionData);
 
 	TArray<FString> Filenames = SoundWave->AssetImportData->ExtractFilenames();
 	if (Filenames.Num() < 1)
@@ -510,9 +510,9 @@ bool FMovieSceneExportData::FindAudioSections(const FString& InSoundPathName, TA
 		return false;
 	}
 
-	for (TSharedPtr<FMovieSceneExportAudioMasterTrackData> AudioMasterTrack : MovieSceneData->AudioMasterTracks)
+	for (TSharedPtr<FMovieSceneExportAudioData> AudioData : MovieSceneData->AudioData)
 	{
-		for (TSharedPtr<FMovieSceneExportAudioSectionData> AudioSectionData : AudioMasterTrack->AudioSections)
+		for (TSharedPtr<FMovieSceneExportAudioSectionData> AudioSectionData : AudioData->AudioSections)
 		{
 			if (!AudioSectionData.IsValid() || AudioSectionData->MovieSceneSection == nullptr)
 			{
@@ -628,13 +628,13 @@ bool FMovieSceneImportData::IsImportDataValid() const
 
 TSharedPtr<FMovieSceneImportCinematicSectionData> FMovieSceneImportData::FindCinematicSection(const FString& InSectionPathName)
 {
-	TSharedPtr<FMovieSceneImportCinematicMasterTrackData> MasterTrackData = GetCinematicMasterTrackData(false);
-	if (!MasterTrackData.IsValid())
+	TSharedPtr<FMovieSceneImportCinematicData> CinematicData = GetCinematicData(false);
+	if (!CinematicData.IsValid())
 	{
 		return nullptr;
 	}
 
-	for (TSharedPtr<FMovieSceneImportCinematicSectionData> CinematicSection : MasterTrackData->CinematicSections)
+	for (TSharedPtr<FMovieSceneImportCinematicSectionData> CinematicSection : CinematicData->CinematicSections)
 	{
 		UMovieSceneCinematicShotSection* CinematicShotSection = CinematicSection->CinematicSection;
 		if (CinematicShotSection != nullptr)
@@ -659,14 +659,14 @@ TSharedPtr<FMovieSceneImportCinematicSectionData> FMovieSceneImportData::CreateC
 		return nullptr;
 	}
 
-	TSharedPtr<FMovieSceneImportCinematicMasterTrackData> MasterTrackData = GetCinematicMasterTrackData(true);
-	if (!MasterTrackData.IsValid() || MasterTrackData->MovieSceneTrack == nullptr)
+	TSharedPtr<FMovieSceneImportCinematicData> CinematicData = GetCinematicData(true);
+	if (!CinematicData.IsValid() || CinematicData->MovieSceneTrack == nullptr)
 	{
 		return nullptr;
 	}
 
-	UMovieSceneCinematicShotTrack* MasterTrack = Cast<UMovieSceneCinematicShotTrack>(MasterTrackData->MovieSceneTrack);
-	if (MasterTrack == nullptr)
+	UMovieSceneCinematicShotTrack* CinematicTrack = Cast<UMovieSceneCinematicShotTrack>(CinematicData->MovieSceneTrack);
+	if (CinematicTrack == nullptr)
 	{
 		return nullptr;
 	}
@@ -713,8 +713,8 @@ TSharedPtr<FMovieSceneImportCinematicSectionData> FMovieSceneImportData::CreateC
 	FFrameNumber EndFrame = ConvertFrameTime(InEndFrame, InFrameRate, TickResolution).RoundToFrame();
 	int32 Duration = (EndFrame - StartFrame).Value;
 	
-	MasterTrack->Modify();
-	UMovieSceneCinematicShotSection* Section = Cast<UMovieSceneCinematicShotSection>(MasterTrack->AddSequence(SequenceToAdd, StartFrame, Duration));
+	CinematicTrack->Modify();
+	UMovieSceneCinematicShotSection* Section = Cast<UMovieSceneCinematicShotSection>(CinematicTrack->AddSequence(SequenceToAdd, StartFrame, Duration));
 	if (Section == nullptr)
 	{
 		return nullptr;
@@ -725,8 +725,8 @@ TSharedPtr<FMovieSceneImportCinematicSectionData> FMovieSceneImportData::CreateC
 	Section->SetRange(TRange<FFrameNumber>(StartFrame, EndFrame));
 
 	TSharedPtr<FMovieSceneImportCinematicSectionData> SectionData = ConstructCinematicSectionData(Section);
-	MasterTrackData->CinematicSections.Add(SectionData);
-	for (TSharedPtr<FMovieSceneImportCinematicTrackData> TrackData : MasterTrackData->CinematicTracks)
+	CinematicData->CinematicSections.Add(SectionData);
+	for (TSharedPtr<FMovieSceneImportCinematicTrackData> TrackData : CinematicData->CinematicTracks)
 	{
 		if (InRow == TrackData->RowIndex)
 		{
@@ -762,49 +762,49 @@ bool FMovieSceneImportData::SetCinematicSection(TSharedPtr<FMovieSceneImportCine
 	return true;
 }
 
-TSharedPtr<FMovieSceneImportAudioSectionData> FMovieSceneImportData::FindAudioSection(const FString& InSectionPathName, TSharedPtr<FMovieSceneImportAudioMasterTrackData>& OutMasterTrackData)
+TSharedPtr<FMovieSceneImportAudioSectionData> FMovieSceneImportData::FindAudioSection(const FString& InSectionPathName, TSharedPtr<FMovieSceneImportAudioData>& OutAudioData)
 {
 	if (!MovieSceneData.IsValid())
 	{
-		OutMasterTrackData = nullptr;
+		OutAudioData = nullptr;
 		return nullptr;
 	}
 
-	for (TSharedPtr<FMovieSceneImportAudioMasterTrackData> MasterTrackData : MovieSceneData->AudioMasterTracks)
+	for (TSharedPtr<FMovieSceneImportAudioData> AudioData : MovieSceneData->AudioData)
 	{
-		if (!MasterTrackData.IsValid())
+		if (!AudioData.IsValid())
 		{
 			continue;
 		}
 
-		for (TSharedPtr<FMovieSceneImportAudioSectionData> AudioSectionData : MasterTrackData->AudioSections)
+		for (TSharedPtr<FMovieSceneImportAudioSectionData> AudioSectionData : AudioData->AudioSections)
 		{
 			if (AudioSectionData.IsValid() && AudioSectionData->AudioSection != nullptr)
 			{
 				FString SectionName = AudioSectionData->AudioSection->GetPathName();
 				if (SectionName == InSectionPathName)
 				{
-					OutMasterTrackData = MasterTrackData;
+					OutAudioData = AudioData;
 					return AudioSectionData;
 				}
 			}
 		}
 	}
 
-	OutMasterTrackData = nullptr;
+	OutAudioData = nullptr;
 	return nullptr;
 }
 
 /** Create audio section */
-TSharedPtr<FMovieSceneImportAudioSectionData> FMovieSceneImportData::CreateAudioSection(FString InFilenameOrAssetPathName, bool bIsPathName, TSharedPtr<FMovieSceneImportAudioMasterTrackData> InMasterTrackData, int32 InRow, FFrameRate InFrameRate, FFrameNumber InStartFrame, FFrameNumber InEndFrame, FFrameNumber InStartOffsetFrame)
+TSharedPtr<FMovieSceneImportAudioSectionData> FMovieSceneImportData::CreateAudioSection(FString InFilenameOrAssetPathName, bool bIsPathName, TSharedPtr<FMovieSceneImportAudioData> InData, int32 InRow, FFrameRate InFrameRate, FFrameNumber InStartFrame, FFrameNumber InEndFrame, FFrameNumber InStartOffsetFrame)
 {
-	if (!MovieSceneData.IsValid() || MovieSceneData->MovieScene == nullptr || !InMasterTrackData.IsValid() || InMasterTrackData->MovieSceneTrack == nullptr)
+	if (!MovieSceneData.IsValid() || MovieSceneData->MovieScene == nullptr || !InData.IsValid() || InData->MovieSceneTrack == nullptr)
 	{
 		return nullptr;
 	}
 
-	UMovieSceneAudioTrack* MasterTrack = Cast<UMovieSceneAudioTrack>(InMasterTrackData->MovieSceneTrack);
-	if (MasterTrack == nullptr)
+	UMovieSceneAudioTrack* AudioTrack = Cast<UMovieSceneAudioTrack>(InData->MovieSceneTrack);
+	if (AudioTrack == nullptr)
 	{
 		return nullptr;
 	}
@@ -869,8 +869,8 @@ TSharedPtr<FMovieSceneImportAudioSectionData> FMovieSceneImportData::CreateAudio
 	FFrameNumber EndFrame = ConvertFrameTime(InEndFrame, InFrameRate, TickResolution).RoundToFrame();
 	int32 Duration = (EndFrame - StartFrame).Value;
 
-	MasterTrack->Modify();
-	UMovieSceneAudioSection* AudioSection = Cast<UMovieSceneAudioSection>(MasterTrack->AddNewSoundOnRow(SoundToAdd, StartFrame, InRow));
+	AudioTrack->Modify();
+	UMovieSceneAudioSection* AudioSection = Cast<UMovieSceneAudioSection>(AudioTrack->AddNewSoundOnRow(SoundToAdd, StartFrame, InRow));
 	if (AudioSection == nullptr)
 	{
 		return nullptr;
@@ -882,8 +882,8 @@ TSharedPtr<FMovieSceneImportAudioSectionData> FMovieSceneImportData::CreateAudio
 
 	TSharedPtr<FMovieSceneImportAudioSectionData> AudioSectionData = ConstructAudioSectionData(AudioSection);
 
-	InMasterTrackData->AudioSections.Add(AudioSectionData);
-	for (TSharedPtr<FMovieSceneImportAudioTrackData> TrackData : InMasterTrackData->AudioTracks)
+	InData->AudioSections.Add(AudioSectionData);
+	for (TSharedPtr<FMovieSceneImportAudioTrackData> TrackData : InData->AudioTracks)
 	{
 		if (InRow == TrackData->RowIndex)
 		{
@@ -917,20 +917,20 @@ bool FMovieSceneImportData::SetAudioSection(TSharedPtr<FMovieSceneImportAudioSec
 	return true;
 }
 
-bool FMovieSceneImportData::MoveAudioSection(TSharedPtr<FMovieSceneImportAudioSectionData> InAudioSectionData, TSharedPtr<FMovieSceneImportAudioMasterTrackData> InFromMasterTrackData, TSharedPtr<FMovieSceneImportAudioMasterTrackData> InToMasterTrackData, int32 InToRowIndex)
+bool FMovieSceneImportData::MoveAudioSection(TSharedPtr<FMovieSceneImportAudioSectionData> InAudioSectionData, TSharedPtr<FMovieSceneImportAudioData> InFromData, TSharedPtr<FMovieSceneImportAudioData> InToData, int32 InToRowIndex)
 {
-	if (!MovieSceneData.IsValid() || !InAudioSectionData.IsValid() || !InFromMasterTrackData.IsValid() || !InToMasterTrackData.IsValid())
+	if (!MovieSceneData.IsValid() || !InAudioSectionData.IsValid() || !InFromData.IsValid() || !InToData.IsValid())
 	{
 		return false;
 	}
 
-	UMovieSceneAudioTrack* FromTrack = Cast<UMovieSceneAudioTrack>(InFromMasterTrackData->MovieSceneTrack);
+	UMovieSceneAudioTrack* FromTrack = Cast<UMovieSceneAudioTrack>(InFromData->MovieSceneTrack);
 	if (FromTrack == nullptr)
 	{
 		return false;
 	}
 
-	UMovieSceneAudioTrack* ToTrack = Cast<UMovieSceneAudioTrack>(InToMasterTrackData->MovieSceneTrack);
+	UMovieSceneAudioTrack* ToTrack = Cast<UMovieSceneAudioTrack>(InToData->MovieSceneTrack);
 	if (ToTrack == nullptr)
 	{
 		return false;
@@ -946,8 +946,8 @@ bool FMovieSceneImportData::MoveAudioSection(TSharedPtr<FMovieSceneImportAudioSe
 	ToTrack->Modify();
 	ToTrack->AddSection(*AudioSection);
 
-	InFromMasterTrackData->AudioSections.Remove(InAudioSectionData);
-	for (TSharedPtr<FMovieSceneImportAudioTrackData> AudioTrackData : InFromMasterTrackData->AudioTracks)
+	InFromData->AudioSections.Remove(InAudioSectionData);
+	for (TSharedPtr<FMovieSceneImportAudioTrackData> AudioTrackData : InFromData->AudioTracks)
 	{
 		if (AudioTrackData->AudioSections.Contains(InAudioSectionData))
 		{
@@ -956,8 +956,8 @@ bool FMovieSceneImportData::MoveAudioSection(TSharedPtr<FMovieSceneImportAudioSe
 	}
 
 	bool bFoundTrack = false;
-	InToMasterTrackData->AudioSections.Add(InAudioSectionData);
-	for (TSharedPtr<FMovieSceneImportAudioTrackData> AudioTrackData : InToMasterTrackData->AudioTracks)
+	InToData->AudioSections.Add(InAudioSectionData);
+	for (TSharedPtr<FMovieSceneImportAudioTrackData> AudioTrackData : InToData->AudioTracks)
 	{
 		if (AudioTrackData->RowIndex == InToRowIndex)
 		{
@@ -972,37 +972,37 @@ bool FMovieSceneImportData::MoveAudioSection(TSharedPtr<FMovieSceneImportAudioSe
 		TSharedPtr<FMovieSceneImportAudioTrackData> TrackData = MakeShared<FMovieSceneImportAudioTrackData>();
 		TrackData->RowIndex = InToRowIndex;
 		TrackData->AudioSections.Add(InAudioSectionData);
-		InToMasterTrackData->AudioTracks.Add(TrackData);
+		InToData->AudioTracks.Add(TrackData);
 	}
 
 	return true;
 }
 
-TSharedPtr<FMovieSceneImportCinematicMasterTrackData> FMovieSceneImportData::GetCinematicMasterTrackData(bool CreateTrackIfNull) 
+TSharedPtr<FMovieSceneImportCinematicData> FMovieSceneImportData::GetCinematicData(bool CreateTrackIfNull) 
 {
 	if (!MovieSceneData.IsValid())
 	{
 		return nullptr;
 	}
-	if (!MovieSceneData->CinematicMasterTrack.IsValid() && CreateTrackIfNull)
+	if (!MovieSceneData->CinematicData.IsValid() && CreateTrackIfNull)
 	{
-		UMovieSceneCinematicShotTrack* CinematicMasterTrack = MovieSceneData->MovieScene->AddMasterTrack<UMovieSceneCinematicShotTrack>();
-		MovieSceneData->CinematicMasterTrack = ConstructCinematicMasterTrackData(CinematicMasterTrack);
+		UMovieSceneCinematicShotTrack* CinematicTrack = MovieSceneData->MovieScene->AddTrack<UMovieSceneCinematicShotTrack>();
+		MovieSceneData->CinematicData = ConstructCinematicData(CinematicTrack);
 	}
-	return MovieSceneData->CinematicMasterTrack;
+	return MovieSceneData->CinematicData;
 }
 
-TSharedPtr<FMovieSceneImportAudioMasterTrackData> FMovieSceneImportData::GetAudioMasterTrackData()
+TSharedPtr<FMovieSceneImportAudioData> FMovieSceneImportData::GetAudioData()
 {
 	if (!MovieSceneData.IsValid())
 	{
 		return nullptr;
 	}
-	for (TSharedPtr<FMovieSceneImportAudioMasterTrackData> MasterTrack : MovieSceneData->AudioMasterTracks)
+	for (TSharedPtr<FMovieSceneImportAudioData> AudioData : MovieSceneData->AudioData)
 	{
-		if (MasterTrack.IsValid())
+		if (AudioData.IsValid())
 		{
-			return MasterTrack;
+			return AudioData;
 		}
 	}
 	return nullptr;
@@ -1018,57 +1018,57 @@ TSharedPtr<FMovieSceneImportMovieSceneData> FMovieSceneImportData::ConstructMovi
 	MovieSceneData = MakeShared<FMovieSceneImportMovieSceneData>();
 	MovieSceneData->MovieScene = InMovieScene;
 
-	// Get cinematic master track
-	UMovieSceneCinematicShotTrack* CinematicMasterTrack = InMovieScene->FindMasterTrack<UMovieSceneCinematicShotTrack>();
-	if (CinematicMasterTrack != nullptr)
+	// Get cinematic track
+	UMovieSceneCinematicShotTrack* CinematicTrack = InMovieScene->FindTrack<UMovieSceneCinematicShotTrack>();
+	if (CinematicTrack != nullptr)
 	{
-		MovieSceneData->CinematicMasterTrack = ConstructCinematicMasterTrackData(CinematicMasterTrack);
-		if (!MovieSceneData->CinematicMasterTrack.IsValid())
+		MovieSceneData->CinematicData = ConstructCinematicData(CinematicTrack);
+		if (!MovieSceneData->CinematicData.IsValid())
 		{
 			return nullptr;
 		}
 	}
 
 	// Get audio tracks
-	TMap<int32, TSharedPtr<FMovieSceneImportAudioMasterTrackData>> AudioMasterTrackMap;
+	TMap<int32, TSharedPtr<FMovieSceneImportAudioData>> AudioDataMap;
 
-	const TArray<UMovieSceneTrack*> MasterTracks = InMovieScene->GetMasterTracks();
-	for (UMovieSceneTrack* MasterTrack : MasterTracks)
+	const TArray<UMovieSceneTrack*> Tracks = InMovieScene->GetTracks();
+	for (UMovieSceneTrack* Track : Tracks)
 	{
-		if (MasterTrack->IsA(UMovieSceneAudioTrack::StaticClass()))
+		if (Track->IsA(UMovieSceneAudioTrack::StaticClass()))
 		{
 
-			UMovieSceneAudioTrack* AudioTrack = Cast<UMovieSceneAudioTrack>(MasterTrack);
+			UMovieSceneAudioTrack* AudioTrack = Cast<UMovieSceneAudioTrack>(Track);
 			if (AudioTrack == nullptr)
 			{
 				continue;
 			}
 
-			TSharedPtr<FMovieSceneImportAudioMasterTrackData> AudioMasterTrackData = ConstructAudioMasterTrackData(AudioTrack);
-			if (!AudioMasterTrackData.IsValid())
+			TSharedPtr<FMovieSceneImportAudioData> AudioData = ConstructAudioData(AudioTrack);
+			if (!AudioData.IsValid())
 			{
 				continue;
 			}
 
-			AudioMasterTrackMap.Add(AudioTrack->GetSortingOrder(), AudioMasterTrackData);
+			AudioDataMap.Add(AudioTrack->GetSortingOrder(), AudioData);
 		}
 	}
 
 
-	// sort the audio master tracks by their sorting index and add to the AudioMasterTracks array
-	if (AudioMasterTrackMap.Num() > 0)
+	// sort the audio tracks by their sorting index and add to the AudioTracks array
+	if (AudioDataMap.Num() > 0)
 	{
-		AudioMasterTrackMap.KeySort([](int32 A, int32 B) {
+		AudioDataMap.KeySort([](int32 A, int32 B) {
 			return A < B; // sort keys in order
 		});
 
-		for (auto& Elem : AudioMasterTrackMap)
+		for (auto& Elem : AudioDataMap)
 		{
 			if (Elem.Value.IsValid())
 			{
-				MovieSceneData->AudioMasterTracks.Add(Elem.Value);
+				MovieSceneData->AudioData.Add(Elem.Value);
 
-				if (!SupportMultipleAudioMasterTracks())
+				if (!SupportMultipleAudioTracks())
 				{
 					break;
 				}
@@ -1080,20 +1080,20 @@ TSharedPtr<FMovieSceneImportMovieSceneData> FMovieSceneImportData::ConstructMovi
 }
 
 
-TSharedPtr<FMovieSceneImportCinematicMasterTrackData> FMovieSceneImportData::ConstructCinematicMasterTrackData(UMovieSceneCinematicShotTrack* InCinematicMasterTrack)
+TSharedPtr<FMovieSceneImportCinematicData> FMovieSceneImportData::ConstructCinematicData(UMovieSceneCinematicShotTrack* InCinematicTrack)
 {
-	if (!MovieSceneData.IsValid() || InCinematicMasterTrack == nullptr)
+	if (!MovieSceneData.IsValid() || InCinematicTrack == nullptr)
 	{
 		return nullptr;
 	}
 
-	TSharedPtr<FMovieSceneImportCinematicMasterTrackData> MasterTrackData = MakeShared<FMovieSceneImportCinematicMasterTrackData>();
-	MasterTrackData->MovieSceneTrack = InCinematicMasterTrack;
+	TSharedPtr<FMovieSceneImportCinematicData> CinematicData = MakeShared<FMovieSceneImportCinematicData>();
+	CinematicData->MovieSceneTrack = InCinematicTrack;
 
 	// Construct sections & create track row index array
 	TArray<int32> CinematicTrackRowIndices;
 
-	for (UMovieSceneSection* ShotSection : InCinematicMasterTrack->GetAllSections())
+	for (UMovieSceneSection* ShotSection : InCinematicTrack->GetAllSections())
 	{
 		UMovieSceneCinematicShotSection* CinematicSection = Cast<UMovieSceneCinematicShotSection>(ShotSection);
 
@@ -1112,27 +1112,27 @@ TSharedPtr<FMovieSceneImportCinematicMasterTrackData> FMovieSceneImportData::Con
 
 	for (int32 CinematicTrackRowIndex : CinematicTrackRowIndices)
 	{
-		TSharedPtr<FMovieSceneImportCinematicTrackData> TrackData = ConstructCinematicTrackData(InCinematicMasterTrack, CinematicTrackRowIndex);
+		TSharedPtr<FMovieSceneImportCinematicTrackData> TrackData = ConstructCinematicTrackData(InCinematicTrack, CinematicTrackRowIndex);
 		if (TrackData.IsValid())
 		{
-			MasterTrackData->CinematicTracks.Add(TrackData);
+			CinematicData->CinematicTracks.Add(TrackData);
 
 			for (TSharedPtr<FMovieSceneImportCinematicSectionData> SectionData : TrackData->CinematicSections)
 			{
 				if (SectionData.IsValid())
 				{
-					MasterTrackData->CinematicSections.Add(SectionData);
+					CinematicData->CinematicSections.Add(SectionData);
 				}
 			}
 		}
 	}
 
-	return MasterTrackData;
+	return CinematicData;
 }
 
-TSharedPtr<FMovieSceneImportCinematicTrackData> FMovieSceneImportData::ConstructCinematicTrackData(UMovieSceneCinematicShotTrack* InCinematicMasterTrack, int32 InRowIndex)
+TSharedPtr<FMovieSceneImportCinematicTrackData> FMovieSceneImportData::ConstructCinematicTrackData(UMovieSceneCinematicShotTrack* InCinematicTrack, int32 InRowIndex)
 {
-	if (!MovieSceneData.IsValid() || InCinematicMasterTrack == nullptr)
+	if (!MovieSceneData.IsValid() || InCinematicTrack == nullptr)
 	{
 		return nullptr;
 	}
@@ -1140,7 +1140,7 @@ TSharedPtr<FMovieSceneImportCinematicTrackData> FMovieSceneImportData::Construct
 	TSharedPtr<FMovieSceneImportCinematicTrackData> TrackData = MakeShared<FMovieSceneImportCinematicTrackData>();
 	TrackData->RowIndex = InRowIndex;
 
-	for (UMovieSceneSection* ShotSection : InCinematicMasterTrack->GetAllSections())
+	for (UMovieSceneSection* ShotSection : InCinematicTrack->GetAllSections())
 	{
 		UMovieSceneCinematicShotSection* CinematicSection = Cast<UMovieSceneCinematicShotSection>(ShotSection);
 
@@ -1159,21 +1159,21 @@ TSharedPtr<FMovieSceneImportCinematicTrackData> FMovieSceneImportData::Construct
 }
 
 
-TSharedPtr<FMovieSceneImportAudioMasterTrackData> FMovieSceneImportData::ConstructAudioMasterTrackData(UMovieSceneAudioTrack* InAudioMasterTrack)
+TSharedPtr<FMovieSceneImportAudioData> FMovieSceneImportData::ConstructAudioData(UMovieSceneAudioTrack* InAudioTrack)
 {
-	if (!MovieSceneData.IsValid() || InAudioMasterTrack == nullptr)
+	if (!MovieSceneData.IsValid() || InAudioTrack == nullptr)
 	{
 		return nullptr;
 	}
 
-	TSharedPtr<FMovieSceneImportAudioMasterTrackData> MasterTrackData = MakeShared<FMovieSceneImportAudioMasterTrackData>();
-	MasterTrackData->MovieSceneTrack = InAudioMasterTrack;
-	MasterTrackData->MaxRowIndex = 0;
+	TSharedPtr<FMovieSceneImportAudioData> AudioData = MakeShared<FMovieSceneImportAudioData>();
+	AudioData->MovieSceneTrack = InAudioTrack;
+	AudioData->MaxRowIndex = 0;
 
 	// Construct sections & create track row index array
 	TArray<int32> AudioTrackRowIndices;
 
-	for (UMovieSceneSection* ShotSection : InAudioMasterTrack->GetAllSections())
+	for (UMovieSceneSection* ShotSection : InAudioTrack->GetAllSections())
 	{
 		UMovieSceneAudioSection* AudioSection = Cast<UMovieSceneAudioSection>(ShotSection);
 
@@ -1184,9 +1184,9 @@ TSharedPtr<FMovieSceneImportAudioMasterTrackData> FMovieSceneImportData::Constru
 			{
 				AudioTrackRowIndices.AddUnique(RowIndex);
 
-				if (RowIndex > MasterTrackData->MaxRowIndex)
+				if (RowIndex > AudioData->MaxRowIndex)
 				{
-					MasterTrackData->MaxRowIndex = RowIndex;
+					AudioData->MaxRowIndex = RowIndex;
 				}
 			}
 		}
@@ -1197,27 +1197,27 @@ TSharedPtr<FMovieSceneImportAudioMasterTrackData> FMovieSceneImportData::Constru
 
 	for (int32 AudioTrackRowIndex : AudioTrackRowIndices)
 	{
-		TSharedPtr<FMovieSceneImportAudioTrackData> TrackData = ConstructAudioTrackData(InAudioMasterTrack, AudioTrackRowIndex);
+		TSharedPtr<FMovieSceneImportAudioTrackData> TrackData = ConstructAudioTrackData(InAudioTrack, AudioTrackRowIndex);
 		if (TrackData.IsValid())
 		{
-			MasterTrackData->AudioTracks.Add(TrackData);
+			AudioData->AudioTracks.Add(TrackData);
 
 			for (TSharedPtr<FMovieSceneImportAudioSectionData> SectionData : TrackData->AudioSections)
 			{
 				if (SectionData.IsValid())
 				{
-					MasterTrackData->AudioSections.Add(SectionData);
+					AudioData->AudioSections.Add(SectionData);
 				}
 			}
 		}
 	}
 
-	return MasterTrackData;
+	return AudioData;
 }
 
-TSharedPtr<FMovieSceneImportAudioTrackData> FMovieSceneImportData::ConstructAudioTrackData(UMovieSceneAudioTrack* InAudioMasterTrack, int32 InRowIndex)
+TSharedPtr<FMovieSceneImportAudioTrackData> FMovieSceneImportData::ConstructAudioTrackData(UMovieSceneAudioTrack* InAudioTrack, int32 InRowIndex)
 {
-	if (!MovieSceneData.IsValid() || InAudioMasterTrack == nullptr)
+	if (!MovieSceneData.IsValid() || InAudioTrack == nullptr)
 	{
 		return nullptr;
 	}
@@ -1225,7 +1225,7 @@ TSharedPtr<FMovieSceneImportAudioTrackData> FMovieSceneImportData::ConstructAudi
 	TSharedPtr<FMovieSceneImportAudioTrackData> TrackData = MakeShared<FMovieSceneImportAudioTrackData>();
 	TrackData->RowIndex = InRowIndex;
 
-	for (UMovieSceneSection* Section : InAudioMasterTrack->GetAllSections())
+	for (UMovieSceneSection* Section : InAudioTrack->GetAllSections())
 	{
 		UMovieSceneAudioSection* AudioSection = Cast<UMovieSceneAudioSection>(Section);
 
diff --git a/Engine/Source/Editor/MovieSceneTools/Private/MovieSceneTranslatorEDL.cpp b/Engine/Source/Editor/MovieSceneTools/Private/MovieSceneTranslatorEDL.cpp
--- a/Engine/Source/Editor/MovieSceneTools/Private/MovieSceneTranslatorEDL.cpp
+++ b/Engine/Source/Editor/MovieSceneTools/Private/MovieSceneTranslatorEDL.cpp
@@ -367,10 +367,10 @@ bool MovieSceneTranslatorEDL::ImportEDL(UMovieScene* InMovieScene, FFrameRate In
 	TArray<FShotData> ShotDataArray;
 	ParseFromEDL(InputString, TickResolution, InFrameRate, ShotDataArray);
 
-	UMovieSceneCinematicShotTrack* CinematicShotTrack = InMovieScene->FindMasterTrack<UMovieSceneCinematicShotTrack>();
+	UMovieSceneCinematicShotTrack* CinematicShotTrack = InMovieScene->FindTrack<UMovieSceneCinematicShotTrack>();
 	if (!CinematicShotTrack)
 	{
-		CinematicShotTrack = InMovieScene->AddMasterTrack<UMovieSceneCinematicShotTrack>();
+		CinematicShotTrack = InMovieScene->AddTrack<UMovieSceneCinematicShotTrack>();
 	}
 
 	for (FShotData ShotData : ShotDataArray)
@@ -447,13 +447,13 @@ bool MovieSceneTranslatorEDL::ExportEDL(const UMovieScene* InMovieScene, FFrameR
 
 	TArray<FShotData> ShotDataArray;
 
-	for (UMovieSceneTrack* MasterTrack : InMovieScene->GetMasterTracks())
+	for (UMovieSceneTrack* Track : InMovieScene->GetTracks())
 	{
 		// @todo: sequencer-timecode: deal with framerate differences??
 		TRange<FFrameNumber> PlaybackRange = InMovieScene->GetPlaybackRange();
-		if (MasterTrack->IsA(UMovieSceneCinematicShotTrack::StaticClass()))
+		if (Track->IsA(UMovieSceneCinematicShotTrack::StaticClass()))
 		{
-			UMovieSceneCinematicShotTrack* CinematicShotTrack = Cast<UMovieSceneCinematicShotTrack>(MasterTrack);
+			UMovieSceneCinematicShotTrack* CinematicShotTrack = Cast<UMovieSceneCinematicShotTrack>(Track);
 
 			for (UMovieSceneSection* ShotSection : CinematicShotTrack->GetAllSections())
 			{
@@ -490,9 +490,9 @@ bool MovieSceneTranslatorEDL::ExportEDL(const UMovieScene* InMovieScene, FFrameR
 				ShotDataArray.Add(FShotData(ShotName, ShotPath, FShotData::ETrackType::TT_Video, FShotData::EEditType::ET_Cut, SourceInFrame, SourceOutFrame, EditInFrame, EditOutFrame, bWithinPlaybackRange));
 			}
 		}
-		else if (MasterTrack->IsA(UMovieSceneAudioTrack::StaticClass()))
+		else if (Track->IsA(UMovieSceneAudioTrack::StaticClass()))
 		{
-			UMovieSceneAudioTrack* AudioTrack = Cast<UMovieSceneAudioTrack>(MasterTrack);
+			UMovieSceneAudioTrack* AudioTrack = Cast<UMovieSceneAudioTrack>(Track);
 
 			//@todo support audio clips
 		}
diff --git a/Engine/Source/Editor/MovieSceneTools/Private/SObjectBindingTag.cpp b/Engine/Source/Editor/MovieSceneTools/Private/SObjectBindingTag.cpp
--- a/Engine/Source/Editor/MovieSceneTools/Private/SObjectBindingTag.cpp
+++ b/Engine/Source/Editor/MovieSceneTools/Private/SObjectBindingTag.cpp
@@ -65,7 +65,7 @@ void SObjectBindingTags::OnBindingCacheUpdated(const FObjectBindingTagCache* Bin
 			.OnDeleted(OnDeleteDelegate)
 			.ColorTint(BindingCache->GetTagColor(TagName))
 			.Text(FText::FromName(TagName))
-			.ToolTipText(FText::Format(LOCTEXT("TagToolTip", "This object binding can be looked up, resolved, and overridden in the master sequence by using the tag '{0}'"), FText::FromName(TagName)))
+			.ToolTipText(FText::Format(LOCTEXT("TagToolTip", "This object binding can be looked up, resolved, and overridden in the root sequence by using the tag '{0}'"), FText::FromName(TagName)))
 		];
 
 		LeftPadding = 5.f;
diff --git a/Engine/Source/Editor/MovieSceneTools/Private/SequenceBindingTree.cpp b/Engine/Source/Editor/MovieSceneTools/Private/SequenceBindingTree.cpp
--- a/Engine/Source/Editor/MovieSceneTools/Private/SequenceBindingTree.cpp
+++ b/Engine/Source/Editor/MovieSceneTools/Private/SequenceBindingTree.cpp
@@ -61,9 +61,9 @@ bool FSequenceBindingTree::ConditionalRebuild(UMovieSceneSequence* InSequence, F
 			}
 
 			UMovieScene* ThisMovieScene = ThisSequence->GetMovieScene();
-			for (const UMovieSceneTrack* MasterTrack : ThisMovieScene->GetMasterTracks())
+			for (const UMovieSceneTrack* Track : ThisMovieScene->GetTracks())
 			{
-				const UMovieSceneSubTrack* SubTrack = Cast<const UMovieSceneSubTrack>(MasterTrack);
+				const UMovieSceneSubTrack* SubTrack = Cast<const UMovieSceneSubTrack>(Track);
 				if (SubTrack)
 				{
 					for (UMovieSceneSection* Section : SubTrack->GetAllSections())
@@ -192,9 +192,9 @@ void FSequenceBindingTree::Build(UMovieSceneSequence* InSequence, FSequenceIDSta
 	}
 
 	// Iterate all sub sections
-	for (const UMovieSceneTrack* MasterTrack : MovieScene->GetMasterTracks())
+	for (const UMovieSceneTrack* Track : MovieScene->GetTracks())
 	{
-		const UMovieSceneSubTrack* SubTrack = Cast<const UMovieSceneSubTrack>(MasterTrack);
+		const UMovieSceneSubTrack* SubTrack = Cast<const UMovieSceneSubTrack>(Track);
 		if (SubTrack)
 		{
 			for (UMovieSceneSection* Section : SubTrack->GetAllSections())
diff --git a/Engine/Source/Editor/MovieSceneTools/Private/TrackEditors/AudioTrackEditor.cpp b/Engine/Source/Editor/MovieSceneTools/Private/TrackEditors/AudioTrackEditor.cpp
--- a/Engine/Source/Editor/MovieSceneTools/Private/TrackEditors/AudioTrackEditor.cpp
+++ b/Engine/Source/Editor/MovieSceneTools/Private/TrackEditors/AudioTrackEditor.cpp
@@ -924,7 +924,7 @@ void FAudioTrackEditor::BuildAddTrackMenu(FMenuBuilder& MenuBuilder)
 {
 	MenuBuilder.AddMenuEntry(
 		LOCTEXT("AddTrack", "Audio Track"),
-		LOCTEXT("AddTooltip", "Adds a new master audio track that can play sounds."),
+		LOCTEXT("AddTooltip", "Adds a new audio track that can play sounds."),
 		FSlateIcon(FAppStyle::GetAppStyleSetName(), "Sequencer.Tracks.Audio"),
 		FUIAction(
 			FExecuteAction::CreateRaw(this, &FAudioTrackEditor::HandleAddAudioTrackMenuEntryExecute)
@@ -1090,7 +1090,7 @@ FReply FAudioTrackEditor::OnDrop(const FDragDropEvent& DragDropEvent, const FSeq
 			}
 			else
 			{
-				AnimatablePropertyChanged(FOnKeyProperty::CreateRaw(this, &FAudioTrackEditor::AddNewMasterSound, Sound, AudioTrack, DragDropParams.RowIndex));
+				AnimatablePropertyChanged(FOnKeyProperty::CreateRaw(this, &FAudioTrackEditor::AddNewSound, Sound, AudioTrack, DragDropParams.RowIndex));
 			}
 
 			bAnyDropped = true;
@@ -1144,7 +1144,7 @@ bool FAudioTrackEditor::HandleAssetAdded(UObject* Asset, const FGuid& TargetObje
 		else
 		{
 			int32 RowIndex = INDEX_NONE;
-			AnimatablePropertyChanged( FOnKeyProperty::CreateRaw(this, &FAudioTrackEditor::AddNewMasterSound, Sound, DummyTrack, RowIndex));
+			AnimatablePropertyChanged( FOnKeyProperty::CreateRaw(this, &FAudioTrackEditor::AddNewSound, Sound, DummyTrack, RowIndex));
 		}
 
 		return true;
@@ -1153,7 +1153,7 @@ bool FAudioTrackEditor::HandleAssetAdded(UObject* Asset, const FGuid& TargetObje
 }
 
 
-FKeyPropertyResult FAudioTrackEditor::AddNewMasterSound( FFrameNumber KeyTime, USoundBase* Sound, UMovieSceneAudioTrack* AudioTrack, int32 RowIndex )
+FKeyPropertyResult FAudioTrackEditor::AddNewSound( FFrameNumber KeyTime, USoundBase* Sound, UMovieSceneAudioTrack* AudioTrack, int32 RowIndex )
 {
 	FKeyPropertyResult KeyPropertyResult;
 
@@ -1165,11 +1165,11 @@ FKeyPropertyResult FAudioTrackEditor::AddNewMasterSound( FFrameNumber KeyTime, U
 
 	FocusedMovieScene->Modify();
 
-	FFindOrCreateMasterTrackResult<UMovieSceneAudioTrack> TrackResult;
+	FFindOrCreateRootTrackResult<UMovieSceneAudioTrack> TrackResult;
 	TrackResult.Track = AudioTrack;
 	if (!AudioTrack)
 	{
-		TrackResult = FindOrCreateMasterTrack<UMovieSceneAudioTrack>();
+		TrackResult = FindOrCreateRootTrack<UMovieSceneAudioTrack>();
 		AudioTrack = TrackResult.Track;
 	}
 
@@ -1261,7 +1261,7 @@ void FAudioTrackEditor::HandleAddAudioTrackMenuEntryExecute()
 	const FScopedTransaction Transaction(NSLOCTEXT("Sequencer", "AddAudioTrack_Transaction", "Add Audio Track"));
 	FocusedMovieScene->Modify();
 	
-	auto NewTrack = FocusedMovieScene->AddMasterTrack<UMovieSceneAudioTrack>();
+	auto NewTrack = FocusedMovieScene->AddTrack<UMovieSceneAudioTrack>();
 	ensure(NewTrack);
 
 	NewTrack->SetDisplayName(LOCTEXT("AudioTrackName", "Audio"));
diff --git a/Engine/Source/Editor/MovieSceneTools/Private/TrackEditors/CVarTrackEditor.cpp b/Engine/Source/Editor/MovieSceneTools/Private/TrackEditors/CVarTrackEditor.cpp
--- a/Engine/Source/Editor/MovieSceneTools/Private/TrackEditors/CVarTrackEditor.cpp
+++ b/Engine/Source/Editor/MovieSceneTools/Private/TrackEditors/CVarTrackEditor.cpp
@@ -107,7 +107,7 @@ void FCVarTrackEditor::HandleAddCVarTrackMenuEntryExecute()
 	const FScopedTransaction Transaction(NSLOCTEXT("Sequencer", "AddCVarTrack_Transaction", "Add Console Variable Track"));
 	FocusedMovieScene->Modify();
 	
-	UMovieSceneCVarTrack* NewTrack = FocusedMovieScene->AddMasterTrack<UMovieSceneCVarTrack>();
+	UMovieSceneCVarTrack* NewTrack = FocusedMovieScene->AddTrack<UMovieSceneCVarTrack>();
 	ensure(NewTrack);
 
 	NewTrack->SetDisplayName(LOCTEXT("CVarTrackName", "Console Variable"));
diff --git a/Engine/Source/Editor/MovieSceneTools/Private/TrackEditors/CinematicShotTrackEditor.cpp b/Engine/Source/Editor/MovieSceneTools/Private/TrackEditors/CinematicShotTrackEditor.cpp
--- a/Engine/Source/Editor/MovieSceneTools/Private/TrackEditors/CinematicShotTrackEditor.cpp
+++ b/Engine/Source/Editor/MovieSceneTools/Private/TrackEditors/CinematicShotTrackEditor.cpp
@@ -583,7 +583,7 @@ bool FCinematicShotTrackEditor::HandleAddCinematicShotTrackMenuEntryCanExecute()
 {
 	UMovieScene* FocusedMovieScene = GetFocusedMovieScene();
 
-	return ((FocusedMovieScene != nullptr) && (FocusedMovieScene->FindMasterTrack<UMovieSceneCinematicShotTrack>() == nullptr));
+	return ((FocusedMovieScene != nullptr) && (FocusedMovieScene->FindTrack<UMovieSceneCinematicShotTrack>() == nullptr));
 }
 
 
@@ -734,7 +734,7 @@ UMovieSceneCinematicShotTrack* FCinematicShotTrackEditor::FindOrCreateCinematicS
 		return nullptr;
 	}
 
-	UMovieSceneCinematicShotTrack* CinematicShotTrack = FocusedMovieScene->FindMasterTrack<UMovieSceneCinematicShotTrack>();
+	UMovieSceneCinematicShotTrack* CinematicShotTrack = FocusedMovieScene->FindTrack<UMovieSceneCinematicShotTrack>();
 	if (CinematicShotTrack != nullptr)
 	{
 		return CinematicShotTrack;
@@ -743,7 +743,7 @@ UMovieSceneCinematicShotTrack* FCinematicShotTrackEditor::FindOrCreateCinematicS
 	const FScopedTransaction Transaction(LOCTEXT("AddCinematicShotTrack_Transaction", "Add Cinematic Shot Track"));
 	FocusedMovieScene->Modify();
 
-	auto NewTrack = FocusedMovieScene->AddMasterTrack<UMovieSceneCinematicShotTrack>();
+	auto NewTrack = FocusedMovieScene->AddTrack<UMovieSceneCinematicShotTrack>();
 	ensure(NewTrack);
 
 	GetSequencer()->NotifyMovieSceneDataChanged( EMovieSceneDataChangeType::MovieSceneStructureItemAdded );
@@ -819,13 +819,13 @@ bool FCinematicShotTrackEditor::CanAddSubSequence(const UMovieSceneSequence& Seq
 
 	// make sure we are not contained in the other sequence (circular dependency)
 	// @todo sequencer: this check is not sufficient (does not prevent circular dependencies of 2+ levels)
-	UMovieSceneSubTrack* SequenceSubTrack = SequenceMovieScene->FindMasterTrack<UMovieSceneSubTrack>();
+	UMovieSceneSubTrack* SequenceSubTrack = SequenceMovieScene->FindTrack<UMovieSceneSubTrack>();
 	if (SequenceSubTrack && SequenceSubTrack->ContainsSequence(*FocusedSequence, true))
 	{
 		return false;
 	}
 
-	UMovieSceneCinematicShotTrack* SequenceCinematicTrack = SequenceMovieScene->FindMasterTrack<UMovieSceneCinematicShotTrack>();
+	UMovieSceneCinematicShotTrack* SequenceCinematicTrack = SequenceMovieScene->FindTrack<UMovieSceneCinematicShotTrack>();
 	if (SequenceCinematicTrack && SequenceCinematicTrack->ContainsSequence(*FocusedSequence, true))
 	{
 		return false;
diff --git a/Engine/Source/Editor/MovieSceneTools/Private/TrackEditors/DataLayerTrackEditor.cpp b/Engine/Source/Editor/MovieSceneTools/Private/TrackEditors/DataLayerTrackEditor.cpp
--- a/Engine/Source/Editor/MovieSceneTools/Private/TrackEditors/DataLayerTrackEditor.cpp
+++ b/Engine/Source/Editor/MovieSceneTools/Private/TrackEditors/DataLayerTrackEditor.cpp
@@ -351,7 +351,7 @@ void FDataLayerTrackEditor::HandleAddTrack()
 	const FScopedTransaction Transaction(LOCTEXT("AddDataLayerTrack_Transaction", "Add Data Layer Track"));
 	FocusedMovieScene->Modify();
 
-	UMovieSceneDataLayerTrack* NewTrack = FocusedMovieScene->AddMasterTrack<UMovieSceneDataLayerTrack>();
+	UMovieSceneDataLayerTrack* NewTrack = FocusedMovieScene->AddTrack<UMovieSceneDataLayerTrack>();
 	checkf(NewTrack, TEXT("Failed to create new data layer track."));
 
 	UMovieSceneDataLayerSection* NewSection = AddNewSection(FocusedMovieScene, NewTrack, EDataLayerRuntimeState::Activated);
diff --git a/Engine/Source/Editor/MovieSceneTools/Private/TrackEditors/EventTrackEditor.cpp b/Engine/Source/Editor/MovieSceneTools/Private/TrackEditors/EventTrackEditor.cpp
--- a/Engine/Source/Editor/MovieSceneTools/Private/TrackEditors/EventTrackEditor.cpp
+++ b/Engine/Source/Editor/MovieSceneTools/Private/TrackEditors/EventTrackEditor.cpp
@@ -287,11 +287,11 @@ void FEventTrackEditor::HandleAddEventTrackMenuEntryExecute(TArray<FGuid> InObje
 
 	if (!NewTracks.Num())
 	{
-		UMovieSceneEventTrack* NewMasterTrack = FocusedMovieScene->AddMasterTrack<UMovieSceneEventTrack>();
-		NewTracks.Add(NewMasterTrack);
+		UMovieSceneEventTrack* NewTrack = FocusedMovieScene->AddTrack<UMovieSceneEventTrack>();
+		NewTracks.Add(NewTrack);
 		if (GetSequencer().IsValid())
 		{
-			GetSequencer()->OnAddTrack(NewMasterTrack, FGuid());
+			GetSequencer()->OnAddTrack(NewTrack, FGuid());
 		}
 	}
 
diff --git a/Engine/Source/Editor/MovieSceneTools/Private/TrackEditors/FadeTrackEditor.cpp b/Engine/Source/Editor/MovieSceneTools/Private/TrackEditors/FadeTrackEditor.cpp
--- a/Engine/Source/Editor/MovieSceneTools/Private/TrackEditors/FadeTrackEditor.cpp
+++ b/Engine/Source/Editor/MovieSceneTools/Private/TrackEditors/FadeTrackEditor.cpp
@@ -178,7 +178,7 @@ void FFadeTrackEditor::HandleAddFadeTrackMenuEntryExecute()
 		return;
 	}
 
-	UMovieSceneTrack* FadeTrack = MovieScene->FindMasterTrack<UMovieSceneFadeTrack>();
+	UMovieSceneTrack* FadeTrack = MovieScene->FindTrack<UMovieSceneFadeTrack>();
 
 	if (FadeTrack != nullptr)
 	{
@@ -189,7 +189,7 @@ void FFadeTrackEditor::HandleAddFadeTrackMenuEntryExecute()
 
 	MovieScene->Modify();
 
-	FadeTrack = FindOrCreateMasterTrack<UMovieSceneFadeTrack>().Track;
+	FadeTrack = FindOrCreateRootTrack<UMovieSceneFadeTrack>().Track;
 	check(FadeTrack);
 
 	UMovieSceneSection* NewSection = FadeTrack->CreateNewSection();
@@ -208,7 +208,7 @@ bool FFadeTrackEditor::HandleAddFadeTrackMenuEntryCanExecute() const
 {
 	UMovieScene* FocusedMovieScene = GetFocusedMovieScene();
 	
-	return ((FocusedMovieScene != nullptr) && (FocusedMovieScene->FindMasterTrack<UMovieSceneFadeTrack>() == nullptr));
+	return ((FocusedMovieScene != nullptr) && (FocusedMovieScene->FindTrack<UMovieSceneFadeTrack>() == nullptr));
 }
 
 #undef LOCTEXT_NAMESPACE
diff --git a/Engine/Source/Editor/MovieSceneTools/Private/TrackEditors/LevelVisibilityTrackEditor.cpp b/Engine/Source/Editor/MovieSceneTools/Private/TrackEditors/LevelVisibilityTrackEditor.cpp
--- a/Engine/Source/Editor/MovieSceneTools/Private/TrackEditors/LevelVisibilityTrackEditor.cpp
+++ b/Engine/Source/Editor/MovieSceneTools/Private/TrackEditors/LevelVisibilityTrackEditor.cpp
@@ -138,7 +138,7 @@ void FLevelVisibilityTrackEditor::OnAddTrack()
 	const FScopedTransaction Transaction( LOCTEXT( "AddLevelVisibilityTrack_Transaction", "Add Level Visibility Track" ) );
 	FocusedMovieScene->Modify();
 
-	UMovieSceneLevelVisibilityTrack* NewTrack = FocusedMovieScene->AddMasterTrack<UMovieSceneLevelVisibilityTrack>();
+	UMovieSceneLevelVisibilityTrack* NewTrack = FocusedMovieScene->AddTrack<UMovieSceneLevelVisibilityTrack>();
 	checkf( NewTrack != nullptr, TEXT("Failed to create new level visibility track.") );
 
 	UMovieSceneLevelVisibilitySection* NewSection = AddNewSection( FocusedMovieScene, NewTrack, ELevelVisibility::Visible );
diff --git a/Engine/Source/Editor/MovieSceneTools/Private/TrackEditors/MaterialParameterCollectionTrackEditor.cpp b/Engine/Source/Editor/MovieSceneTools/Private/TrackEditors/MaterialParameterCollectionTrackEditor.cpp
--- a/Engine/Source/Editor/MovieSceneTools/Private/TrackEditors/MaterialParameterCollectionTrackEditor.cpp
+++ b/Engine/Source/Editor/MovieSceneTools/Private/TrackEditors/MaterialParameterCollectionTrackEditor.cpp
@@ -173,7 +173,7 @@ void FMaterialParameterCollectionTrackEditor::AddTrackToSequence(const FAssetDat
 	}
 
 	// Attempt to find an existing MPC track that animates this object
-	for (UMovieSceneTrack* Track : MovieScene->GetMasterTracks())
+	for (UMovieSceneTrack* Track : MovieScene->GetTracks())
 	{
 		if (auto* MPCTrack = Cast<UMovieSceneMaterialParameterCollectionTrack>(Track))
 		{
@@ -187,7 +187,7 @@ void FMaterialParameterCollectionTrackEditor::AddTrackToSequence(const FAssetDat
 	const FScopedTransaction Transaction(LOCTEXT("AddTrackDescription", "Add Material Parameter Collection Track"));
 
 	MovieScene->Modify();
-	UMovieSceneMaterialParameterCollectionTrack* Track = MovieScene->AddMasterTrack<UMovieSceneMaterialParameterCollectionTrack>();
+	UMovieSceneMaterialParameterCollectionTrack* Track = MovieScene->AddTrack<UMovieSceneMaterialParameterCollectionTrack>();
 	check(Track);
 
 	UMovieSceneSection* NewSection = Track->CreateNewSection();
diff --git a/Engine/Source/Editor/MovieSceneTools/Private/TrackEditors/SlomoTrackEditor.cpp b/Engine/Source/Editor/MovieSceneTools/Private/TrackEditors/SlomoTrackEditor.cpp
--- a/Engine/Source/Editor/MovieSceneTools/Private/TrackEditors/SlomoTrackEditor.cpp
+++ b/Engine/Source/Editor/MovieSceneTools/Private/TrackEditors/SlomoTrackEditor.cpp
@@ -96,7 +96,7 @@ void FSlomoTrackEditor::HandleAddSlomoTrackMenuEntryExecute()
 		return;
 	}
 
-	UMovieSceneTrack* SlomoTrack = MovieScene->FindMasterTrack<UMovieSceneSlomoTrack>();
+	UMovieSceneTrack* SlomoTrack = MovieScene->FindTrack<UMovieSceneSlomoTrack>();
 
 	if (SlomoTrack != nullptr)
 	{
@@ -107,7 +107,7 @@ void FSlomoTrackEditor::HandleAddSlomoTrackMenuEntryExecute()
 
 	MovieScene->Modify();
 
-	SlomoTrack = FindOrCreateMasterTrack<UMovieSceneSlomoTrack>().Track;
+	SlomoTrack = FindOrCreateRootTrack<UMovieSceneSlomoTrack>().Track;
 	check(SlomoTrack);
 
 	UMovieSceneSection* NewSection = SlomoTrack->CreateNewSection();
@@ -126,7 +126,7 @@ bool FSlomoTrackEditor::HandleAddSlomoTrackMenuEntryCanExecute() const
 {
 	UMovieScene* FocusedMovieScene = GetFocusedMovieScene();
 
-	return ((FocusedMovieScene != nullptr) && (FocusedMovieScene->FindMasterTrack<UMovieSceneSlomoTrack>() == nullptr));
+	return ((FocusedMovieScene != nullptr) && (FocusedMovieScene->FindTrack<UMovieSceneSlomoTrack>() == nullptr));
 }
 
 #undef LOCTEXT_NAMESPACE
diff --git a/Engine/Source/Editor/MovieSceneTools/Private/TrackEditors/SubTrackEditor.cpp b/Engine/Source/Editor/MovieSceneTools/Private/TrackEditors/SubTrackEditor.cpp
--- a/Engine/Source/Editor/MovieSceneTools/Private/TrackEditors/SubTrackEditor.cpp
+++ b/Engine/Source/Editor/MovieSceneTools/Private/TrackEditors/SubTrackEditor.cpp
@@ -76,7 +76,7 @@ public:
 
 		MenuBuilder.AddMenuEntry(
 			LOCTEXT("PlayableDirectly_Label", "Playable Directly"),
-			LOCTEXT("PlayableDirectly_Tip", "When enabled, this sequence will also support being played directly outside of the master sequence. Disable this to save some memory on complex hierarchies of sequences."),
+			LOCTEXT("PlayableDirectly_Tip", "When enabled, this sequence will also support being played directly outside of the root sequence. Disable this to save some memory on complex hierarchies of sequences."),
 			FSlateIcon(),
 			FUIAction(
 				FExecuteAction::CreateRaw(this, &FSubSection::TogglePlayableDirectly),
@@ -436,7 +436,7 @@ bool FSubTrackEditor::CanAddSubSequence(const UMovieSceneSequence& Sequence) con
 
 UMovieSceneSubTrack* FSubTrackEditor::CreateNewTrack(UMovieScene* MovieScene) const
 {
-	return MovieScene->AddMasterTrack<UMovieSceneSubTrack>();
+	return MovieScene->AddTrack<UMovieSceneSubTrack>();
 }
 
 void FSubTrackEditor::GetSupportedSequenceClassPaths(TArray<FTopLevelAssetPath>& ClassPaths) const
@@ -565,7 +565,7 @@ void FSubTrackEditor::InsertSequence(UMovieSceneTrack* Track)
 	UMovieSceneSubTrack* SubTrack = Cast<UMovieSceneSubTrack>(Track);
 	if (!SubTrack)
 	{
-		SubTrack = FindOrCreateMasterTrack<UMovieSceneSubTrack>().Track;
+		SubTrack = FindOrCreateRootTrack<UMovieSceneSubTrack>().Track;
 	}
 
 	FString NewSequencePath = FPaths::GetPath(GetSequencer()->GetFocusedMovieSceneSequence()->GetPathName());
@@ -642,7 +642,7 @@ FKeyPropertyResult FSubTrackEditor::HandleSequenceAdded(FFrameNumber KeyTime, UM
 	UMovieSceneSubTrack* SubTrack = Cast<UMovieSceneSubTrack>(Track);
 	if (!SubTrack)
 	{
-		SubTrack = FindOrCreateMasterTrack<UMovieSceneSubTrack>().Track;
+		SubTrack = FindOrCreateRootTrack<UMovieSceneSubTrack>().Track;
 	}
 
 	const FFrameRate TickResolution = Sequence->GetMovieScene()->GetTickResolution();
diff --git a/Engine/Source/Editor/MovieSceneTools/Private/TrackEditors/SubTrackEditorBase.cpp b/Engine/Source/Editor/MovieSceneTools/Private/TrackEditors/SubTrackEditorBase.cpp
--- a/Engine/Source/Editor/MovieSceneTools/Private/TrackEditors/SubTrackEditorBase.cpp
+++ b/Engine/Source/Editor/MovieSceneTools/Private/TrackEditors/SubTrackEditorBase.cpp
@@ -40,7 +40,7 @@ FSubSectionPainterResult FSubSectionPainterUtil::PaintSection(TSharedPtr<const I
     }
 
     UMovieScene* MovieScene = InnerSequence->GetMovieScene();
-    const int32 NumTracks = MovieScene->GetPossessableCount() + MovieScene->GetSpawnableCount() + MovieScene->GetMasterTracks().Num();
+    const int32 NumTracks = MovieScene->GetPossessableCount() + MovieScene->GetSpawnableCount() + MovieScene->GetTracks().Num();
 
     FVector2D TopLeft = InPainter.SectionGeometry.AbsoluteToLocal(InPainter.SectionClippingRect.GetTopLeft()) + FVector2D(1.f, -1.f);
 
@@ -432,13 +432,13 @@ bool FSubTrackEditorUtil::CanAddSubSequence(const UMovieSceneSequence* CurrentSe
 
 	// make sure we are not contained in the other sequence (circular dependency)
 	// @todo sequencer: this check is not sufficient (does not prevent circular dependencies of 2+ levels)
-	UMovieSceneSubTrack* SequenceSubTrack = SequenceMovieScene->FindMasterTrack<UMovieSceneSubTrack>();
+	UMovieSceneSubTrack* SequenceSubTrack = SequenceMovieScene->FindTrack<UMovieSceneSubTrack>();
 	if (SequenceSubTrack && SequenceSubTrack->ContainsSequence(*CurrentSequence, true))
 	{
 		return false;
 	}
 
-	UMovieSceneCinematicShotTrack* SequenceCinematicTrack = SequenceMovieScene->FindMasterTrack<UMovieSceneCinematicShotTrack>();
+	UMovieSceneCinematicShotTrack* SequenceCinematicTrack = SequenceMovieScene->FindTrack<UMovieSceneCinematicShotTrack>();
 	if (SequenceCinematicTrack && SequenceCinematicTrack->ContainsSequence(*CurrentSequence, true))
 	{
 		return false;
diff --git a/Engine/Source/Editor/MovieSceneTools/Public/MovieSceneToolHelpers.h b/Engine/Source/Editor/MovieSceneTools/Public/MovieSceneToolHelpers.h
--- a/Engine/Source/Editor/MovieSceneTools/Public/MovieSceneToolHelpers.h
+++ b/Engine/Source/Editor/MovieSceneTools/Public/MovieSceneToolHelpers.h
@@ -309,14 +309,14 @@ public:
 	* @param InMovieScene The movie scene to export frome
 	* @param MoviePlayer to use
 	* @param Bindings The sequencer binding map
-	* @param MasterTracks Master tracks to export
+	* @param Tracks The tracks to export
 	* @param NodeNameAdaptor Adaptor to look up actor names.
 	* @param InFBXFileName the fbx file name.
 	* @param Template Movie scene sequence id.
 	* @param RootToLocalTransform The root to local transform time.
 	* @return Whether the export was successful
 	*/
-	static bool ExportFBX(UWorld* World, UMovieScene* MovieScene, IMovieScenePlayer* Player, const TArray<FGuid>& Bindings, const TArray<UMovieSceneTrack*>& MasterTracks, INodeNameAdapter& NodeNameAdapter, FMovieSceneSequenceIDRef& Template,  const FString& InFBXFileName, FMovieSceneSequenceTransform& RootToLocalTransform);
+	static bool ExportFBX(UWorld* World, UMovieScene* MovieScene, IMovieScenePlayer* Player, const TArray<FGuid>& Bindings, const TArray<UMovieSceneTrack*>& Tracks, INodeNameAdapter& NodeNameAdapter, FMovieSceneSequenceIDRef& Template,  const FString& InFBXFileName, FMovieSceneSequenceTransform& RootToLocalTransform);
 
 	/**
 	* Import FBX with dialog
diff --git a/Engine/Source/Editor/MovieSceneTools/Public/MovieSceneTranslator.h b/Engine/Source/Editor/MovieSceneTools/Public/MovieSceneTranslator.h
--- a/Engine/Source/Editor/MovieSceneTools/Public/MovieSceneTranslator.h
+++ b/Engine/Source/Editor/MovieSceneTools/Public/MovieSceneTranslator.h
@@ -75,7 +75,7 @@ struct FMovieSceneExportAudioSectionData : public FMovieSceneExportSectionData
 
 struct FMovieSceneExportCinematicTrackData
 {
-	/** This indicates the sub-track's row index in the master track */
+	/** This indicates the sub-track's row index in the track */
 	int32 RowIndex;
 	TArray< TSharedPtr<FMovieSceneExportCinematicSectionData> > CinematicSections;
 };
@@ -84,17 +84,17 @@ struct FMovieSceneExportAudioTrackData
 {
 	int32 SampleRate;
 
-	/** This indicates the sub-track's row index in the master track */
+	/** This indicates the sub-track's row index in the track */
 	int32 RowIndex;
 	TArray< TSharedPtr<FMovieSceneExportAudioSectionData> > AudioSections;
 };
 
-struct FMovieSceneExportMasterTrackData
+struct FMovieSceneExportTrackData
 {
 	const UMovieSceneTrack* MovieSceneTrack;
 };
 
-struct FMovieSceneExportAudioMasterTrackData : public FMovieSceneExportMasterTrackData
+struct FMovieSceneExportAudioData : public FMovieSceneExportTrackData
 {
 	int32 SampleRate;
 
@@ -106,7 +106,7 @@ struct FMovieSceneExportAudioMasterTrackData : public FMovieSceneExportMasterTra
 
 };
 
-struct FMovieSceneExportCinematicMasterTrackData : public FMovieSceneExportMasterTrackData
+struct FMovieSceneExportCinematicData : public FMovieSceneExportTrackData
 {
 	/** Array of all sections in order they appear in UMovieSceneCinematicTrack*/
 	TArray< TSharedPtr<FMovieSceneExportCinematicSectionData> > CinematicSections;
@@ -123,8 +123,8 @@ struct FMovieSceneExportMovieSceneData
 	int32 Duration;
 	FFrameNumber PlaybackRangeStartFrame;
 	FFrameNumber PlaybackRangeEndFrame;
-	TSharedPtr<FMovieSceneExportCinematicMasterTrackData> CinematicMasterTrack;
-	TArray< TSharedPtr<FMovieSceneExportAudioMasterTrackData> > AudioMasterTracks;
+	TSharedPtr<FMovieSceneExportCinematicData> CinematicData;
+	TArray< TSharedPtr<FMovieSceneExportAudioData> > AudioData;
 	FString MovieExtension;
 };
 
@@ -184,18 +184,18 @@ private:
 	bool ConstructData(const UMovieScene* InMovieScene);
 	/** Loads intermediate movie scene data from Sequencer. */
 	bool ConstructMovieSceneData(const UMovieScene* InMovieScene);
-	/** Loads intermediate cinematic master track data from Sequencer. */
-	bool ConstructCinematicMasterTrackData(const UMovieScene* InMovieScene, const UMovieSceneCinematicShotTrack* InCinematicMasterTrack);
-	/** Loads intermediate cinematic track data from Sequencer. */
-	bool ConstructCinematicTrackData(const UMovieScene* InMovieScene, TSharedPtr<FMovieSceneExportCinematicMasterTrackData> InCinematicMasterTrack, int32 InRowIndex);
-	/** Loads intermediate audio track data from Sequencer. */
-	bool ConstructAudioMasterTrackData(const UMovieScene* InMovieScene, const UMovieSceneAudioTrack* InAudioMasterTrack, TMap<int32, TSharedPtr<FMovieSceneExportAudioMasterTrackData>>& InAudioTrackMap);
-	/** Loads intermediate audio track data from Sequencer. */
-	bool ConstructAudioTrackData(const UMovieScene* InMovieScene, TSharedPtr<FMovieSceneExportAudioMasterTrackData> InAudioMasterTrack, int32 InRowIndex);
+	/** Loads intermediate cinematic data from Sequencer. */
+	bool ConstructCinematicData(const UMovieScene* InMovieScene, const UMovieSceneCinematicShotTrack* InCinematicTrack);
+	/** Loads intermediate cinematic data from Sequencer. */
+	bool ConstructCinematicTrackData(const UMovieScene* InMovieScene, TSharedPtr<FMovieSceneExportCinematicData> InCinematicData, int32 InRowIndex);
+	/** Loads intermediate audio data from Sequencer. */
+	bool ConstructAudioData(const UMovieScene* InMovieScene, const UMovieSceneAudioTrack* InAudioTrack, TMap<int32, TSharedPtr<FMovieSceneExportAudioData>>& InAudioDataMap);
+	/** Loads intermediate audio data from Sequencer. */
+	bool ConstructAudioData(const UMovieScene* InMovieScene, TSharedPtr<FMovieSceneExportAudioData> InAudioData, int32 InRowIndex);
 	/** Loads intermediate cinematic section data from Sequencer. */
-	bool ConstructCinematicSectionData(const UMovieScene* InMovieScene, TSharedPtr<FMovieSceneExportCinematicMasterTrackData> InMasterTrackData, const UMovieSceneCinematicShotSection* InCinematicSection);
+	bool ConstructCinematicSectionData(const UMovieScene* InMovieScene, TSharedPtr<FMovieSceneExportCinematicData> InCinematicData, const UMovieSceneCinematicShotSection* InCinematicSection);
 	/** Loads intermediate audio section data from Sequencer. */
-	bool ConstructAudioSectionData(const UMovieScene* InMovieScene, TSharedPtr<FMovieSceneExportAudioMasterTrackData> InTrackData, const UMovieSceneAudioSection* InAudioSection);
+	bool ConstructAudioSectionData(const UMovieScene* InMovieScene, TSharedPtr<FMovieSceneExportAudioData> InAudioData, const UMovieSceneAudioSection* InAudioSection);
 	/** Loads intermediate common section data from Sequencer. */
 	bool ConstructSectionData(const UMovieScene* InMovieScene, TSharedPtr<FMovieSceneExportSectionData> InSectionData, const UMovieSceneSection* InSection, EMovieSceneTranslatorSectionType InSectionType, const FString& InSectionDisplayName);
 
@@ -247,12 +247,12 @@ struct FMovieSceneImportAudioTrackData
 	TArray< TSharedPtr<FMovieSceneImportAudioSectionData> > AudioSections;
 };
 
-struct FMovieSceneImportMasterTrackData
+struct FMovieSceneImportTrackData
 {
 	UMovieSceneTrack* MovieSceneTrack;
 };
         
-struct FMovieSceneImportAudioMasterTrackData : public FMovieSceneImportMasterTrackData
+struct FMovieSceneImportAudioData : public FMovieSceneImportTrackData
 {
 	/** Array of all sections in order they appear in UMovieSceneAudioTrack*/
 	TArray< TSharedPtr<FMovieSceneImportAudioSectionData> > AudioSections;
@@ -260,11 +260,11 @@ struct FMovieSceneImportAudioMasterTrackData : public FMovieSceneImportMasterTra
 	/** Array of sorted audio sub tracks, containing pointers to sections within the sub track row*/
 	TArray< TSharedPtr<FMovieSceneImportAudioTrackData> > AudioTracks;
 
-	/** Max row index existing in this master track */
+	/** Max row index existing in this track */
 	int32 MaxRowIndex;
 };
 
-struct FMovieSceneImportCinematicMasterTrackData : public FMovieSceneImportMasterTrackData
+struct FMovieSceneImportCinematicData : public FMovieSceneImportTrackData
 {
 	/** Array of all sections in order they appear in UMovieSceneCinematicTrack*/
 	TArray< TSharedPtr<FMovieSceneImportCinematicSectionData> > CinematicSections;
@@ -276,8 +276,8 @@ struct FMovieSceneImportCinematicMasterTrackData : public FMovieSceneImportMaste
 struct FMovieSceneImportMovieSceneData
 {
 	UMovieScene* MovieScene;
-	TSharedPtr<FMovieSceneImportCinematicMasterTrackData> CinematicMasterTrack;
-	TArray< TSharedPtr<FMovieSceneImportAudioMasterTrackData> > AudioMasterTracks;
+	TSharedPtr<FMovieSceneImportCinematicData> CinematicData;
+	TArray< TSharedPtr<FMovieSceneImportAudioData> > AudioData;
 };
 
 /**
@@ -296,8 +296,8 @@ public:
 	/** True when the export data was successfully constructed. */
 	bool IsImportDataValid() const;
 
-	/** Returns the cinematic master track data pointer or nullptr if one does not exist */
-	TSharedPtr<FMovieSceneImportCinematicMasterTrackData> GetCinematicMasterTrackData(bool CreateTrackIfNull);
+	/** Returns the cinematic data pointer or nullptr if one does not exist */
+	TSharedPtr<FMovieSceneImportCinematicData> GetCinematicData(bool CreateTrackIfNull);
 	/** Find cinematic section */
 	TSharedPtr<FMovieSceneImportCinematicSectionData> FindCinematicSection(const FString& InSectionPathName);
 	/** Create cinematic section */
@@ -305,29 +305,29 @@ public:
 	/** Set cinematic section */
 	bool SetCinematicSection(TSharedPtr<FMovieSceneImportCinematicSectionData> InSection, int32 InRow, FFrameRate InFrameRate, FFrameNumber InStartFrame, FFrameNumber InEndFrame, TOptional<FFrameNumber> InStartOffsetFrame);
 
-	/** Returns the audio master track data pointer or nullptr if one does not exist */
-	TSharedPtr<FMovieSceneImportAudioMasterTrackData> GetAudioMasterTrackData();
+	/** Returns the audio data pointer or nullptr if one does not exist */
+	TSharedPtr<FMovieSceneImportAudioData> GetAudioData();
 	/** Find audio sections */
-	TSharedPtr<FMovieSceneImportAudioSectionData>  FindAudioSection(const FString& InSectionPathName, TSharedPtr<FMovieSceneImportAudioMasterTrackData>& OutMasterTrackData);
+	TSharedPtr<FMovieSceneImportAudioSectionData>  FindAudioSection(const FString& InSectionPathName, TSharedPtr<FMovieSceneImportAudioData>& OutAudioData);
 	/** Create audio section */
-	TSharedPtr<FMovieSceneImportAudioSectionData> CreateAudioSection(FString InFilenameOrAssetPathName, bool bIsPathName, TSharedPtr<FMovieSceneImportAudioMasterTrackData> InMasterTrack, int32 InRow, FFrameRate InFrameRate, FFrameNumber InStartFrame, FFrameNumber InEndFrame, FFrameNumber InStartOffsetFrame);
+	TSharedPtr<FMovieSceneImportAudioSectionData> CreateAudioSection(FString InFilenameOrAssetPathName, bool bIsPathName, TSharedPtr<FMovieSceneImportAudioData> InAudioData, int32 InRow, FFrameRate InFrameRate, FFrameNumber InStartFrame, FFrameNumber InEndFrame, FFrameNumber InStartOffsetFrame);
 	/** Set audio section */
 	bool SetAudioSection(TSharedPtr<FMovieSceneImportAudioSectionData> InSection, int32 InRow, FFrameRate InFrameRate, FFrameNumber InStartFrame, FFrameNumber InEndFrame, FFrameNumber InStartOffsetFrame);
 	/** Move audio section */
-	bool MoveAudioSection(TSharedPtr<FMovieSceneImportAudioSectionData> InAudioSectionData, TSharedPtr<FMovieSceneImportAudioMasterTrackData> InFromMasterTrackData, TSharedPtr<FMovieSceneImportAudioMasterTrackData> InToMasterTrackData, int32 InToRowIndex);
+	bool MoveAudioSection(TSharedPtr<FMovieSceneImportAudioSectionData> InAudioSectionData, TSharedPtr<FMovieSceneImportAudioData> InFromAudioData, TSharedPtr<FMovieSceneImportAudioData> InToAudioData, int32 InToRowIndex);
 
 private:
 
 	/** Entry point for setting up intermediate data for use when importing. */
 	TSharedPtr<FMovieSceneImportMovieSceneData> ConstructMovieSceneData(UMovieScene* InMovieScene);
-	/** Gets cinematic master track data from Sequencer. */
-	TSharedPtr<FMovieSceneImportCinematicMasterTrackData> ConstructCinematicMasterTrackData(UMovieSceneCinematicShotTrack* InCinematicMasterTrack);
+	/** Gets cinematic data from Sequencer. */
+	TSharedPtr<FMovieSceneImportCinematicData> ConstructCinematicData(UMovieSceneCinematicShotTrack* InCinematicTrack);
 	/** Gets cinematic track data from Sequencer. */
-	TSharedPtr<FMovieSceneImportCinematicTrackData> ConstructCinematicTrackData(UMovieSceneCinematicShotTrack* InCinematicMasterTrack, int32 InRowIndex);
-	/** Gets audio master track data from Sequencer. */
-	TSharedPtr<FMovieSceneImportAudioMasterTrackData> ConstructAudioMasterTrackData(UMovieSceneAudioTrack* InAudioMasterTrack);
+	TSharedPtr<FMovieSceneImportCinematicTrackData> ConstructCinematicTrackData(UMovieSceneCinematicShotTrack* InCinematicTrack, int32 InRowIndex);
+	/** Gets audio data from Sequencer. */
+	TSharedPtr<FMovieSceneImportAudioData> ConstructAudioData(UMovieSceneAudioTrack* InAudioTrack);
 	/** Gets audio track data from Sequencer. */
-	TSharedPtr<FMovieSceneImportAudioTrackData> ConstructAudioTrackData(UMovieSceneAudioTrack* InAudioMasterTrack, int32 InRowIndex);
+	TSharedPtr<FMovieSceneImportAudioTrackData> ConstructAudioTrackData(UMovieSceneAudioTrack* InAudioTrack, int32 InRowIndex);
 	/** Gets cinematic section data from Sequencer. */
 	TSharedPtr<FMovieSceneImportCinematicSectionData> ConstructCinematicSectionData(UMovieSceneCinematicShotSection* InCinematicSection);
 	/** Gets audio section data from Sequencer. */
diff --git a/Engine/Source/Editor/MovieSceneTools/Public/TrackEditors/AudioTrackEditor.h b/Engine/Source/Editor/MovieSceneTools/Public/TrackEditors/AudioTrackEditor.h
--- a/Engine/Source/Editor/MovieSceneTools/Public/TrackEditors/AudioTrackEditor.h
+++ b/Engine/Source/Editor/MovieSceneTools/Public/TrackEditors/AudioTrackEditor.h
@@ -65,8 +65,8 @@ public:
 	
 protected:
 
-	/** Delegate for AnimatablePropertyChanged in HandleAssetAdded for master sounds */
-	FKeyPropertyResult AddNewMasterSound(FFrameNumber KeyTime, class USoundBase* Sound, UMovieSceneAudioTrack* Track, int32 RowIndex);
+	/** Delegate for AnimatablePropertyChanged in HandleAssetAdded for sounds */
+	FKeyPropertyResult AddNewSound(FFrameNumber KeyTime, class USoundBase* Sound, UMovieSceneAudioTrack* Track, int32 RowIndex);
 
 	/** Delegate for AnimatablePropertyChanged in HandleAssetAdded for attached sounds */
 	FKeyPropertyResult AddNewAttachedSound(FFrameNumber KeyTime, class USoundBase* Sound, UMovieSceneAudioTrack* Track, TArray<TWeakObjectPtr<UObject>> ObjectsToAttachTo);
diff --git a/Engine/Source/Editor/SequenceRecorder/Private/ActorRecording.cpp b/Engine/Source/Editor/SequenceRecorder/Private/ActorRecording.cpp
--- a/Engine/Source/Editor/SequenceRecorder/Private/ActorRecording.cpp
+++ b/Engine/Source/Editor/SequenceRecorder/Private/ActorRecording.cpp
@@ -784,10 +784,10 @@ bool UActorRecording::StopRecording(ULevelSequence* OriginalSequence, float Curr
 	if (CurrentSequence && OriginalSequence && OriginalSequence != CurrentSequence)
 	{
 		UMovieScene* MovieScene = OriginalSequence->GetMovieScene();
-		UMovieSceneSubTrack* SubTrack = Cast<UMovieSceneSubTrack>(MovieScene->FindMasterTrack(UMovieSceneSubTrack::StaticClass()));
+		UMovieSceneSubTrack* SubTrack = Cast<UMovieSceneSubTrack>(MovieScene->FindTrack(UMovieSceneSubTrack::StaticClass()));
 		if (!SubTrack)
 		{
-			SubTrack = Cast<UMovieSceneSubTrack>(MovieScene->AddMasterTrack(UMovieSceneSubTrack::StaticClass()));
+			SubTrack = Cast<UMovieSceneSubTrack>(MovieScene->AddTrack(UMovieSceneSubTrack::StaticClass()));
 		}
 
 		// Remove the current take if it exists
diff --git a/Engine/Source/Editor/SequenceRecorder/Private/SequenceRecorder.cpp b/Engine/Source/Editor/SequenceRecorder/Private/SequenceRecorder.cpp
--- a/Engine/Source/Editor/SequenceRecorder/Private/SequenceRecorder.cpp
+++ b/Engine/Source/Editor/SequenceRecorder/Private/SequenceRecorder.cpp
@@ -934,47 +934,47 @@ bool FSequenceRecorder::StopRecording(bool bAllowLooping)
 
 		if (RecordedSoundWaves.Num())
 		{
-			// Add a new master audio track to the level sequence		
+			// Add a new audio track to the level sequence		
 			UMovieScene* MovieScene = LevelSequence->GetMovieScene();
-			UMovieSceneAudioTrack* RecordedAudioMasterTrack = nullptr;
+			UMovieSceneAudioTrack* RecordedAudioTrack = nullptr;
 
 			FText RecordedAudioTrackName = Settings->AudioTrackName;
-			for (auto MasterTrack : MovieScene->GetMasterTracks())
+			for (auto Track : MovieScene->GetTracks())
 			{
-				if (MasterTrack->IsA(UMovieSceneAudioTrack::StaticClass()) && MasterTrack->GetDisplayName().EqualTo(RecordedAudioTrackName))
+				if (Track->IsA(UMovieSceneAudioTrack::StaticClass()) && Track->GetDisplayName().EqualTo(RecordedAudioTrackName))
 				{
-					RecordedAudioMasterTrack = Cast<UMovieSceneAudioTrack>(MasterTrack);
+					RecordedAudioTrack = Cast<UMovieSceneAudioTrack>(Track);
 				}
 			}
 
-			if (!RecordedAudioMasterTrack)
+			if (!RecordedAudioTrack)
 			{
-				RecordedAudioMasterTrack = MovieScene->AddMasterTrack<UMovieSceneAudioTrack>();
-				RecordedAudioMasterTrack->SetDisplayName(RecordedAudioTrackName);
+				RecordedAudioTrack = MovieScene->AddTrack<UMovieSceneAudioTrack>();
+				RecordedAudioTrack->SetDisplayName(RecordedAudioTrackName);
 			}
 
 			if (Settings->bReplaceRecordedAudio)
 			{
-				RecordedAudioMasterTrack->RemoveAllAnimationData();
+				RecordedAudioTrack->RemoveAllAnimationData();
 			}
 
 			for (USoundWave* RecordedAudio : RecordedSoundWaves)
 			{
 				int32 RowIndex = -1;
-				for (UMovieSceneSection* Section : RecordedAudioMasterTrack->GetAllSections())
+				for (UMovieSceneSection* Section : RecordedAudioTrack->GetAllSections())
 				{
 					RowIndex = FMath::Max(RowIndex, Section->GetRowIndex());
 				}
 
-				UMovieSceneAudioSection* NewAudioSection = NewObject<UMovieSceneAudioSection>(RecordedAudioMasterTrack, UMovieSceneAudioSection::StaticClass());
+				UMovieSceneAudioSection* NewAudioSection = NewObject<UMovieSceneAudioSection>(RecordedAudioTrack, UMovieSceneAudioSection::StaticClass());
 
-				FFrameRate TickResolution = RecordedAudioMasterTrack->GetTypedOuter<UMovieScene>()->GetTickResolution();
+				FFrameRate TickResolution = RecordedAudioTrack->GetTypedOuter<UMovieScene>()->GetTickResolution();
 
 				NewAudioSection->SetRowIndex(RowIndex + 1);
 				NewAudioSection->SetSound(RecordedAudio);
 				NewAudioSection->SetRange(TRange<FFrameNumber>(FFrameNumber(0), (RecordedAudio->GetDuration() * TickResolution).CeilToFrame()));
 
-				RecordedAudioMasterTrack->AddSection(*NewAudioSection);
+				RecordedAudioTrack->AddSection(*NewAudioSection);
 
 				if(Settings->bAutoSaveAsset || GEditor == nullptr)
 				{
diff --git a/Engine/Source/Editor/SequenceRecorder/Private/SequenceRecorderUtils.cpp b/Engine/Source/Editor/SequenceRecorder/Private/SequenceRecorderUtils.cpp
--- a/Engine/Source/Editor/SequenceRecorder/Private/SequenceRecorderUtils.cpp
+++ b/Engine/Source/Editor/SequenceRecorder/Private/SequenceRecorderUtils.cpp
@@ -281,11 +281,11 @@ void GatherTakeInfo(ULevelSequence* InLevelSequence, TArray<FTakeInfo>& TakeInfo
 {
 	UMovieScene* MovieScene = InLevelSequence->GetMovieScene();
 
-	for (auto MasterTrack : MovieScene->GetMasterTracks())
+	for (auto Track : MovieScene->GetTracks())
 	{
-		if (MasterTrack->IsA<UMovieSceneSubTrack>())
+		if (Track->IsA<UMovieSceneSubTrack>())
 		{
-			UMovieSceneSubTrack* SubTrack = Cast<UMovieSceneSubTrack>(MasterTrack);
+			UMovieSceneSubTrack* SubTrack = Cast<UMovieSceneSubTrack>(Track);
 			for (auto Section : SubTrack->GetAllSections())
 			{
 				if (Section->IsA<UMovieSceneSubSection>())
diff --git a/Engine/Source/Editor/SequenceRecorder/Public/AnimationRecorder.h b/Engine/Source/Editor/SequenceRecorder/Public/AnimationRecorder.h
--- a/Engine/Source/Editor/SequenceRecorder/Public/AnimationRecorder.h
+++ b/Engine/Source/Editor/SequenceRecorder/Public/AnimationRecorder.h
@@ -148,7 +148,7 @@ public:
 	/** Exclude list */
 	TArray<FString> ExcludeAnimationNames;
 public:
-	/** Helper function to get space bases depending on master pose component */
+	/** Helper function to get space bases depending on leader pose component */
 	static void GetBoneTransforms(USkeletalMeshComponent* Component, TArray<FTransform>& BoneTransforms);
 
 private:
diff --git a/Engine/Source/Editor/Sequencer/Private/MVVM/TrackModelStorageExtension.cpp b/Engine/Source/Editor/Sequencer/Private/MVVM/TrackModelStorageExtension.cpp
--- a/Engine/Source/Editor/Sequencer/Private/MVVM/TrackModelStorageExtension.cpp
+++ b/Engine/Source/Editor/Sequencer/Private/MVVM/TrackModelStorageExtension.cpp
@@ -33,30 +33,30 @@ void FTrackModelStorageExtension::OnReinitialize()
 	{
 		MovieScene->EventHandlers.Link(this);
 
-		TSet<UMovieSceneTrack*> MasterTracksInFolders;
+		TSet<UMovieSceneTrack*> TracksInFolders;
 
 		TArray<UMovieSceneFolder*> AllFolders;
 		GetMovieSceneFoldersRecursive(MovieScene->GetRootFolders(), AllFolders);
 		for (UMovieSceneFolder* Folder : AllFolders)
 		{
-			for (UMovieSceneTrack* Track : Folder->GetChildMasterTracks())
+			for (UMovieSceneTrack* Track : Folder->GetChildTracks())
 			{
-				MasterTracksInFolders.Add(Track);
+				TracksInFolders.Add(Track);
 			}
 		}
 
-		for (UMovieSceneTrack* Track : MovieScene->GetMasterTracks())
+		for (UMovieSceneTrack* Track : MovieScene->GetTracks())
 		{
-			if (!MasterTracksInFolders.Contains(Track))
+			if (!TracksInFolders.Contains(Track))
 			{
-				OnMasterTrackAdded(Track);
+				OnTrackAdded(Track);
 			}
 		}
 		if (UMovieSceneTrack* Track = MovieScene->GetCameraCutTrack())
 		{
-			if (!MasterTracksInFolders.Contains(Track))
+			if (!TracksInFolders.Contains(Track))
 			{
-				OnMasterTrackAdded(Track);
+				OnTrackAdded(Track);
 			}
 		}
 	}
@@ -128,12 +128,12 @@ TSharedPtr<FTrackModel> FTrackModelStorageExtension::FindModelForTrack(UMovieSce
 	return TrackToModel.FindRef(TrackAsKey).Pin();
 }
 
-void FTrackModelStorageExtension::OnMasterTrackAdded(UMovieSceneTrack* InTrack)
+void FTrackModelStorageExtension::OnTrackAdded(UMovieSceneTrack* InTrack)
 {
 	CreateModelForTrack(InTrack);
 }
 
-void FTrackModelStorageExtension::OnMasterTrackRemoved(UMovieSceneTrack* InTrack)
+void FTrackModelStorageExtension::OnTrackRemoved(UMovieSceneTrack* InTrack)
 {
 	FObjectKey TrackAsKey(InTrack);
 
diff --git a/Engine/Source/Editor/Sequencer/Private/MVVM/ViewModels/FolderModel.cpp b/Engine/Source/Editor/Sequencer/Private/MVVM/ViewModels/FolderModel.cpp
--- a/Engine/Source/Editor/Sequencer/Private/MVVM/ViewModels/FolderModel.cpp
+++ b/Engine/Source/Editor/Sequencer/Private/MVVM/ViewModels/FolderModel.cpp
@@ -137,7 +137,7 @@ void FFolderModel::RepopulateChildren()
 	}
 
 	// Create/recycle models for any tracks
-	for (UMovieSceneTrack* Track : Folder->GetChildMasterTracks())
+	for (UMovieSceneTrack* Track : Folder->GetChildTracks())
 	{
 		if (Track)
 		{
@@ -646,10 +646,10 @@ void FFolderModel::PerformDrop(const FViewModelPtr& TargetModel, const FDragDrop
 					{
 						OldFolder->SetFlags(RF_Transactional);
 						OldFolder->Modify();
-						OldFolder->RemoveChildMasterTrack(Track);
+						OldFolder->RemoveChildTrack(Track);
 					}
 
-					Folder->AddChildMasterTrack(Track);
+					Folder->AddChildTrack(Track);
 					bSuccess = true;
 				}
 			}
diff --git a/Engine/Source/Editor/Sequencer/Private/MVVM/ViewModels/ObjectBindingModel.cpp b/Engine/Source/Editor/Sequencer/Private/MVVM/ViewModels/ObjectBindingModel.cpp
--- a/Engine/Source/Editor/Sequencer/Private/MVVM/ViewModels/ObjectBindingModel.cpp
+++ b/Engine/Source/Editor/Sequencer/Private/MVVM/ViewModels/ObjectBindingModel.cpp
@@ -1044,13 +1044,13 @@ void FObjectBindingModel::AddSpawnOwnershipMenu(FMenuBuilder& MenuBuilder)
 	);
 
 	MenuBuilder.AddMenuEntry(
-		LOCTEXT("MasterSequence_Label", "Master Sequence"),
-		LOCTEXT("MasterSequence_Tooltip", "Indicates that the outermost sequence will own the spawned object. The object will be destroyed when the outermost sequence stops playing."),
+		LOCTEXT("RootSequence_Label", "Root Sequence"),
+		LOCTEXT("RootSequence_Tooltip", "Indicates that the outermost sequence will own the spawned object. The object will be destroyed when the outermost sequence stops playing."),
 		FSlateIcon(),
 		FUIAction(
-			FExecuteAction::CreateLambda(Callback, ESpawnOwnership::MasterSequence),
+			FExecuteAction::CreateLambda(Callback, ESpawnOwnership::RootSequence),
 			FCanExecuteAction(),
-			FIsActionChecked::CreateLambda([=]{ return Spawnable->GetSpawnOwnership() == ESpawnOwnership::MasterSequence; })
+			FIsActionChecked::CreateLambda([=]{ return Spawnable->GetSpawnOwnership() == ESpawnOwnership::RootSequence; })
 		),
 		NAME_None,
 		EUserInterfaceActionType::ToggleButton
diff --git a/Engine/Source/Editor/Sequencer/Private/MVVM/ViewModels/SequenceModel.cpp b/Engine/Source/Editor/Sequencer/Private/MVVM/ViewModels/SequenceModel.cpp
--- a/Engine/Source/Editor/Sequencer/Private/MVVM/ViewModels/SequenceModel.cpp
+++ b/Engine/Source/Editor/Sequencer/Private/MVVM/ViewModels/SequenceModel.cpp
@@ -311,7 +311,7 @@ void FSequenceModel::PerformDrop(const FViewModelPtr& TargetModel, const FDragDr
 					{
 						OldFolder->SetFlags(RF_Transactional);
 						OldFolder->Modify();
-						OldFolder->RemoveChildMasterTrack(Track);
+						OldFolder->RemoveChildTrack(Track);
 					}
 
 					bSuccess = true;
diff --git a/Engine/Source/Editor/Sequencer/Private/MVVM/ViewModels/TrackModel.cpp b/Engine/Source/Editor/Sequencer/Private/MVVM/ViewModels/TrackModel.cpp
--- a/Engine/Source/Editor/Sequencer/Private/MVVM/ViewModels/TrackModel.cpp
+++ b/Engine/Source/Editor/Sequencer/Private/MVVM/ViewModels/TrackModel.cpp
@@ -591,7 +591,7 @@ void FTrackModel::Resize(float NewSize)
 
 bool FTrackModel::CanDrag() const
 {
-	// Can only drag master tracks at the moment
+	// Can only drag root tracks at the moment
 	TSharedPtr<IObjectBindingExtension> ObjectBindingExtension = FindAncestorOfType<IObjectBindingExtension>();
 	return ObjectBindingExtension == nullptr;
 }
@@ -673,7 +673,7 @@ void FTrackModel::Delete()
 	if (TViewModelPtr<FFolderModel> ParentFolder = CastParent<FFolderModel>())
 	{
 		ParentFolder->GetFolder()->Modify();
-		ParentFolder->GetFolder()->RemoveChildMasterTrack(Track);
+		ParentFolder->GetFolder()->RemoveChildTrack(Track);
 	}
 
 	TSharedPtr<FSequenceModel> OwnerModel = FindAncestorOfType<FSequenceModel>();
@@ -698,7 +698,7 @@ void FTrackModel::Delete()
 	}
 	else
 	{
-		MovieScene->RemoveMasterTrack(*Track);
+		MovieScene->RemoveTrack(*Track);
 	}
 }
 
diff --git a/Engine/Source/Editor/Sequencer/Private/MVVM/ViewModels/TrackRowModel.cpp b/Engine/Source/Editor/Sequencer/Private/MVVM/ViewModels/TrackRowModel.cpp
--- a/Engine/Source/Editor/Sequencer/Private/MVVM/ViewModels/TrackRowModel.cpp
+++ b/Engine/Source/Editor/Sequencer/Private/MVVM/ViewModels/TrackRowModel.cpp
@@ -365,7 +365,7 @@ void FTrackRowModel::Delete()
 	if (TViewModelPtr<FFolderModel> ParentFolder = CastParent<FFolderModel>())
 	{
 		ParentFolder->GetFolder()->Modify();
-		ParentFolder->GetFolder()->RemoveChildMasterTrack(Track);
+		ParentFolder->GetFolder()->RemoveChildTrack(Track);
 	}
 
 	// Remove sub tracks belonging to this row only
diff --git a/Engine/Source/Editor/Sequencer/Private/MovieSceneCopyableTrack.h b/Engine/Source/Editor/Sequencer/Private/MovieSceneCopyableTrack.h
--- a/Engine/Source/Editor/Sequencer/Private/MovieSceneCopyableTrack.h
+++ b/Engine/Source/Editor/Sequencer/Private/MovieSceneCopyableTrack.h
@@ -18,10 +18,10 @@ public:
 	TObjectPtr<UMovieSceneTrack> Track;
 
 	UPROPERTY()
-	bool bIsAMasterTrack;
+	bool bIsRootTrack;
 
 	UPROPERTY()
-	bool bIsACameraCutTrack;
+	bool bIsCameraCutTrack;
 
 	UPROPERTY()
 	TArray<FName> FolderPath;
diff --git a/Engine/Source/Editor/Sequencer/Private/SSequencer.cpp b/Engine/Source/Editor/Sequencer/Private/SSequencer.cpp
--- a/Engine/Source/Editor/Sequencer/Private/SSequencer.cpp
+++ b/Engine/Source/Editor/Sequencer/Private/SSequencer.cpp
@@ -3581,9 +3581,9 @@ bool SSequencer::CanNavigateBreadcrumbs() const
 		UMovieScene* MovieScene = RootSequence ? RootSequence->GetMovieScene() : nullptr;
 		if (RootSequence)
 		{
-			for (UMovieSceneTrack* MasterTrack : MovieScene->GetMasterTracks())
+			for (UMovieSceneTrack* Track : MovieScene->GetTracks())
 			{
-				if (MasterTrack && MasterTrack->IsA<UMovieSceneSubTrack>())
+				if (Track && Track->IsA<UMovieSceneSubTrack>())
 				{
 					return true;
 				}
diff --git a/Engine/Source/Editor/Sequencer/Private/SSequencerHierarchyBrowser.cpp b/Engine/Source/Editor/Sequencer/Private/SSequencerHierarchyBrowser.cpp
--- a/Engine/Source/Editor/Sequencer/Private/SSequencerHierarchyBrowser.cpp
+++ b/Engine/Source/Editor/Sequencer/Private/SSequencerHierarchyBrowser.cpp
@@ -184,9 +184,9 @@ void SSequencerHierarchyBrowser::AddChildren(TSharedRef<FSequencerHierarchyNode>
 {
 	UMovieScene* MovieScene = Sequence->GetMovieScene();
 
-	for (UMovieSceneTrack* MasterTrack : MovieScene->GetMasterTracks())
+	for (UMovieSceneTrack* Track : MovieScene->GetTracks())
 	{
-		if (UMovieSceneCinematicShotTrack* CinematicShotTrack = Cast<UMovieSceneCinematicShotTrack>(MasterTrack))
+		if (UMovieSceneCinematicShotTrack* CinematicShotTrack = Cast<UMovieSceneCinematicShotTrack>(Track))
 		{
 			TArray<UMovieSceneCinematicShotSection*> ShotSections;
 			for (UMovieSceneSection* Section : CinematicShotTrack->GetAllSections())
@@ -209,7 +209,7 @@ void SSequencerHierarchyBrowser::AddChildren(TSharedRef<FSequencerHierarchyNode>
 				AddChildren(ChildNode.ToSharedRef(), ShotSection->GetSequence());
 			}
 		}
-		else if (UMovieSceneSubTrack* SubTrack = Cast<UMovieSceneSubTrack>(MasterTrack))
+		else if (UMovieSceneSubTrack* SubTrack = Cast<UMovieSceneSubTrack>(Track))
 		{
 			TArray<UMovieSceneSubSection*> SubSections;
 			for (UMovieSceneSection* Section : SubTrack->GetAllSections())
diff --git a/Engine/Source/Editor/Sequencer/Private/SSequencerTimePanel.cpp b/Engine/Source/Editor/Sequencer/Private/SSequencerTimePanel.cpp
--- a/Engine/Source/Editor/Sequencer/Private/SSequencerTimePanel.cpp
+++ b/Engine/Source/Editor/Sequencer/Private/SSequencerTimePanel.cpp
@@ -251,7 +251,7 @@ FReply SSequencerTimePanel::Apply()
 			};
 
 			UE::MovieScene::FSequenceVisitParams Params;
-			Params.bVisitMasterTracks = true;
+			Params.bVisitRootTracks = true;
 			Params.bVisitSubSequences = true;
 			FSequenceReadOnlyVisitor ReadOnlyVisitor;
 			VisitSequence(Sequence, Params, ReadOnlyVisitor);
