Sequencer: Remove Master from Sequencer

Fix static analysis

diff --git a/Engine/Source/Editor/MovieSceneTools/Private/FCPXML/FCPXMLImport.cpp b/Engine/Source/Editor/MovieSceneTools/Private/FCPXML/FCPXMLImport.cpp
--- a/Engine/Source/Editor/MovieSceneTools/Private/FCPXML/FCPXMLImport.cpp
+++ b/Engine/Source/Editor/MovieSceneTools/Private/FCPXML/FCPXMLImport.cpp
@@ -397,7 +397,7 @@ bool FFCPXMLImportVisitor::VisitAudioClipItemNode(TSharedRef<FFCPXMLClipItemNode
 	{
 		if (CurrAudioData.IsValid() && CurrAudioData->MovieSceneTrack != nullptr &&
 			AudioData.IsValid() && AudioData->MovieSceneTrack != nullptr &&
-			CurrAudioData->MovieSceneTrack->GetFullName() != CurrAudioData->MovieSceneTrack->GetFullName())
+			CurrAudioData->MovieSceneTrack->GetFullName() != AudioData->MovieSceneTrack->GetFullName())
 		{
 			// Move audio section
 			if (!ImportData->MoveAudioSection(AudioSectionData, AudioData, CurrAudioData, CurrAudioTrackRowIndex))
