New custom widget for EAlphaBlendOption properties showing a combo box with the curve graph and name

* New custom widget for EAlphaBlendOption properties showing a combo box with the curve graph and name inside it. Once clicking on it, a borderless overlay window will be opened showing a grid of available curve types (like in Sequencer).
* Refactor: Separated the easing customization into a grid compound widget that shows the curve types and the actual customization.
* The grid widget holding the curves now has a filter, to exclude unsupported curve types for various situations. (e.g. we do support more easing curve types than we do for alpha blends in state transitions).
* Created a new compound widget that shows a curve and the name below it, which is used as button content for the new combo box.
* Curve types use grouping set in the enum.

Testing:
* Made sure the easing curve type menu in Sequencer works the way it worked before.
* Made sure the blend modes in the state transitions and e.g. a blend poses anim graph node works the same way as before.

Initial slack channel discussion: https://epicgames.slack.com/archives/C044VB3LCKX/p1664989593947499
Continued slack thread: https://epicgames.slack.com/archives/C0450453V6J/p1665577875211299

diff --git a/Engine/Source/Editor/MovieSceneTools/Private/MovieSceneAlphaBlendOptionCustomization.cpp b/Engine/Source/Editor/MovieSceneTools/Private/MovieSceneAlphaBlendOptionCustomization.cpp
new file mode 100644
--- /dev/null
+++ b/Engine/Source/Editor/MovieSceneTools/Private/MovieSceneAlphaBlendOptionCustomization.cpp
@@ -0,0 +1,160 @@
+// Copyright Epic Games, Inc. All Rights Reserved.
+
+#include "MovieSceneAlphaBlendOptionCustomization.h"
+#include "MovieSceneBuiltInEasingFunctionGridWidget.h"
+
+#include "Delegates/Delegate.h"
+#include "DetailCategoryBuilder.h"
+#include "DetailLayoutBuilder.h"
+#include "DetailWidgetRow.h"
+#include "PropertyHandle.h"
+#include "ScopedTransaction.h"
+#include "Types/SlateEnums.h"
+#include "Types/SlateStructs.h"
+#include "UObject/Class.h"
+#include "UObject/UObjectGlobals.h"
+#include "UObject/UnrealType.h"
+#include "Widgets/Input/SButton.h"
+#include "Widgets/Input/SComboButton.h"
+#include "Widgets/Layout/SBox.h"
+#include "Widgets/SCompoundWidget.h"
+#include "Widgets/Text/STextBlock.h"
+#include "IDetailChildrenBuilder.h"
+#include "MovieSceneBuiltInEasingFunctionCustomization.h"
+
+void FAlphaBlendPropertyCustomization::SetAlphaBlendOption(const EAlphaBlendOption NewBlendOption)
+{
+	FScopedTransaction Transaction(NSLOCTEXT("AlphaBlendCustomization", "SetBlendMode", "Set Blend Mode"));
+
+	TypeProperty->NotifyPreChange();
+
+	TArray<void*> RawData;
+	TypeProperty->AccessRawData(RawData);
+
+	for (void* Ptr : RawData)
+	{
+		*((EAlphaBlendOption*)Ptr) = NewBlendOption;
+	}
+
+	ButtonContent->SetType(ConvertToEasingType(NewBlendOption));
+
+	TypeProperty->NotifyPostChange(EPropertyChangeType::ValueSet);
+	TypeProperty->NotifyFinishedChangingProperties();
+}
+
+EAlphaBlendOption FAlphaBlendPropertyCustomization::GetAlphaBlendOption() const
+{
+	if (TypeProperty)
+	{
+		TArray<void*> RawData;
+		TypeProperty->AccessRawData(RawData);
+
+		for (void* Ptr : RawData)
+		{
+			return *((EAlphaBlendOption*)Ptr);
+		}
+	}
+
+	return EAlphaBlendOption::Linear;
+}
+
+void FAlphaBlendPropertyCustomization::CustomizeChildren(TSharedRef<IPropertyHandle> StructPropertyHandle, IDetailChildrenBuilder& ChildBuilder, IPropertyTypeCustomizationUtils& StructCustomizationUtils)
+{
+	TypeProperty = StructPropertyHandle;
+
+	// Exclude all unsupported curve types
+	TSet<EMovieSceneBuiltInEasing> FilterExclude;
+	FilterExclude.Add(EMovieSceneBuiltInEasing::SinIn);
+	FilterExclude.Add(EMovieSceneBuiltInEasing::SinOut);
+	FilterExclude.Add(EMovieSceneBuiltInEasing::QuadIn);
+	FilterExclude.Add(EMovieSceneBuiltInEasing::QuadOut);
+	FilterExclude.Add(EMovieSceneBuiltInEasing::CubicOut);
+	FilterExclude.Add(EMovieSceneBuiltInEasing::CubicInOut);
+	FilterExclude.Add(EMovieSceneBuiltInEasing::QuartIn);
+	FilterExclude.Add(EMovieSceneBuiltInEasing::QuartOut);
+	FilterExclude.Add(EMovieSceneBuiltInEasing::QuintIn);
+	FilterExclude.Add(EMovieSceneBuiltInEasing::QuintOut);
+
+	ButtonContent = SNew(SBuiltInFunctionVisualizerWithText, ConvertToEasingType(FAlphaBlendPropertyCustomization::GetAlphaBlendOption()));
+
+	TSharedRef<SComboButton> ComboBox = SNew(SComboButton)
+		.ToolTipText(StructPropertyHandle->GetToolTipText())
+		// Show the currently used curve type inside the combo box
+		.ButtonContent()
+		[
+			ButtonContent.ToSharedRef()
+		]
+		// The grid of curve widgets, that appears after clicking the combo box
+		.MenuContent()
+		[
+			SNew(SEasingFunctionGridWidget)
+			.FilterExclude(FilterExclude)
+			.OnTypeChanged(this, &FAlphaBlendPropertyCustomization::SetType)
+		];
+
+	FDetailWidgetRow& Row = ChildBuilder.AddCustomRow(FText::GetEmpty())
+		.NameContent()
+		[
+			SNew(STextBlock)
+			.Text(FText::FromString("Mode"))
+			.Font(IDetailLayoutBuilder::GetDetailFont())
+		]
+		.ValueContent()
+		.VAlign(VAlign_Center)
+		[
+			ComboBox
+		];
+}
+
+void FAlphaBlendPropertyCustomization::SetType(EMovieSceneBuiltInEasing NewType)
+{
+	SetAlphaBlendOption(ConvertFromEasingType(NewType));
+}
+
+EMovieSceneBuiltInEasing FAlphaBlendPropertyCustomization::ConvertToEasingType(EAlphaBlendOption in) const
+{
+	switch (in)
+	{
+	case EAlphaBlendOption::Linear: return EMovieSceneBuiltInEasing::Linear;
+	case EAlphaBlendOption::Cubic: return EMovieSceneBuiltInEasing::CubicIn;
+	case EAlphaBlendOption::HermiteCubic: return EMovieSceneBuiltInEasing::HermiteCubicInOut;
+	case EAlphaBlendOption::Sinusoidal: return EMovieSceneBuiltInEasing::SinInOut;
+	case EAlphaBlendOption::QuadraticInOut: return EMovieSceneBuiltInEasing::QuadInOut;
+	case EAlphaBlendOption::CubicInOut: return EMovieSceneBuiltInEasing::CubicInOut;
+	case EAlphaBlendOption::QuarticInOut: return EMovieSceneBuiltInEasing::QuartInOut;
+	case EAlphaBlendOption::QuinticInOut: return EMovieSceneBuiltInEasing::QuintInOut;
+	case EAlphaBlendOption::CircularIn: return EMovieSceneBuiltInEasing::CircIn;
+	case EAlphaBlendOption::CircularOut: return EMovieSceneBuiltInEasing::CircOut;
+	case EAlphaBlendOption::CircularInOut: return EMovieSceneBuiltInEasing::CircInOut;
+	case EAlphaBlendOption::ExpIn: return EMovieSceneBuiltInEasing::ExpoIn;
+	case EAlphaBlendOption::ExpOut: return EMovieSceneBuiltInEasing::ExpoOut;
+	case EAlphaBlendOption::ExpInOut: return EMovieSceneBuiltInEasing::ExpoInOut;
+	case EAlphaBlendOption::Custom: return EMovieSceneBuiltInEasing::Custom;
+	}
+
+	return EMovieSceneBuiltInEasing::Linear;
+}
+
+EAlphaBlendOption FAlphaBlendPropertyCustomization::ConvertFromEasingType(EMovieSceneBuiltInEasing in) const
+{
+	switch (in)
+	{
+	case EMovieSceneBuiltInEasing::Linear: return EAlphaBlendOption::Linear;
+	case EMovieSceneBuiltInEasing::CubicIn: return EAlphaBlendOption::Cubic;
+	case EMovieSceneBuiltInEasing::HermiteCubicInOut: return EAlphaBlendOption::HermiteCubic;
+	case EMovieSceneBuiltInEasing::SinInOut: return EAlphaBlendOption::Sinusoidal;
+	case EMovieSceneBuiltInEasing::QuadInOut: return EAlphaBlendOption::QuadraticInOut;
+	case EMovieSceneBuiltInEasing::CubicInOut: return EAlphaBlendOption::CubicInOut;
+	case EMovieSceneBuiltInEasing::QuartInOut: return EAlphaBlendOption::QuarticInOut;
+	case EMovieSceneBuiltInEasing::QuintInOut: return EAlphaBlendOption::QuinticInOut;
+	case EMovieSceneBuiltInEasing::CircIn: return EAlphaBlendOption::CircularIn;
+	case EMovieSceneBuiltInEasing::CircOut: return EAlphaBlendOption::CircularOut;
+	case EMovieSceneBuiltInEasing::CircInOut: return EAlphaBlendOption::CircularInOut;
+	case EMovieSceneBuiltInEasing::ExpoIn: return EAlphaBlendOption::ExpIn;
+	case EMovieSceneBuiltInEasing::ExpoOut: return EAlphaBlendOption::ExpOut;
+	case EMovieSceneBuiltInEasing::ExpoInOut: return EAlphaBlendOption::ExpInOut;
+	case EMovieSceneBuiltInEasing::Custom: return EAlphaBlendOption::Custom;
+	}
+
+	return EAlphaBlendOption::Linear;
+}
\ No newline at end of file
diff --git a/Engine/Source/Editor/MovieSceneTools/Private/MovieSceneAlphaBlendOptionCustomization.h b/Engine/Source/Editor/MovieSceneTools/Private/MovieSceneAlphaBlendOptionCustomization.h
new file mode 100644
--- /dev/null
+++ b/Engine/Source/Editor/MovieSceneTools/Private/MovieSceneAlphaBlendOptionCustomization.h
@@ -0,0 +1,31 @@
+// Copyright Epic Games, Inc. All Rights Reserved.
+
+#pragma once
+
+#include "DetailWidgetRow.h"
+#include "IPropertyTypeCustomization.h"
+#include "PropertyHandle.h"
+#include "Generators/MovieSceneEasingCurves.h"
+#include "Runtime/Engine/Public/AlphaBlend.h"
+
+class SBuiltInFunctionVisualizerWithText;
+
+class FAlphaBlendPropertyCustomization : public IPropertyTypeCustomization
+{
+public:
+	virtual ~FAlphaBlendPropertyCustomization() = default;
+
+	void CustomizeHeader(TSharedRef<IPropertyHandle> StructPropertyHandle, FDetailWidgetRow& HeaderRow, IPropertyTypeCustomizationUtils& StructCustomizationUtils) override final {}
+	void CustomizeChildren(TSharedRef<IPropertyHandle> StructPropertyHandle, IDetailChildrenBuilder& ChildBuilder, IPropertyTypeCustomizationUtils& StructCustomizationUtils) override;
+
+	void SetAlphaBlendOption(const EAlphaBlendOption NewBlendOption);
+	EAlphaBlendOption GetAlphaBlendOption() const;
+
+private:
+	void SetType(EMovieSceneBuiltInEasing NewType);
+	EMovieSceneBuiltInEasing ConvertToEasingType(EAlphaBlendOption In) const;
+	EAlphaBlendOption ConvertFromEasingType(EMovieSceneBuiltInEasing In) const;
+
+	TSharedPtr<IPropertyHandle> TypeProperty;
+	TSharedPtr<SBuiltInFunctionVisualizerWithText> ButtonContent;
+};
\ No newline at end of file
diff --git a/Engine/Source/Editor/MovieSceneTools/Private/MovieSceneBuiltInEasingFunctionCustomization.cpp b/Engine/Source/Editor/MovieSceneTools/Private/MovieSceneBuiltInEasingFunctionCustomization.cpp
--- a/Engine/Source/Editor/MovieSceneTools/Private/MovieSceneBuiltInEasingFunctionCustomization.cpp
+++ b/Engine/Source/Editor/MovieSceneTools/Private/MovieSceneBuiltInEasingFunctionCustomization.cpp
@@ -1,6 +1,7 @@
 // Copyright Epic Games, Inc. All Rights Reserved.
 
 #include "MovieSceneBuiltInEasingFunctionCustomization.h"
+#include "MovieSceneBuiltInEasingFunctionGridWidget.h"
 
 #include "Containers/Array.h"
 #include "Containers/ContainerAllocationPolicies.h"
@@ -42,6 +43,7 @@
 #include "Widgets/SCompoundWidget.h"
 #include "Widgets/SOverlay.h"
 #include "Widgets/Text/STextBlock.h"
+#include "IDetailChildrenBuilder.h"
 
 class FActiveTimerHandle;
 class FPaintArgs;
@@ -50,216 +52,25 @@ class FWidgetStyle;
 struct FPointerEvent;
 struct FSlateBrush;
 
-struct FGroupedEasing
-{
-	FString GroupingName;
-	TArray<EMovieSceneBuiltInEasing, TInlineAllocator<3>> Values;
-};
-
-class SBuiltInFunctionVisualizer : public SCompoundWidget
-{
-public:
-
-	SLATE_BEGIN_ARGS(SBuiltInFunctionVisualizer){}
-	SLATE_END_ARGS()
-
-	void Construct(const FArguments& InArgs, EMovieSceneBuiltInEasing InValue)
-	{
-		InterpValue = FVector2D::ZeroVector;
-
-		UMovieSceneBuiltInEasingFunction* DefaultObject = GetMutableDefault<UMovieSceneBuiltInEasingFunction>();
-		EMovieSceneBuiltInEasing DefaultType = DefaultObject->Type;
-
-		DefaultObject->Type = InValue;
-
-		float Interp = 0.f;
-		while (Interp <= 1.f)
-		{
-			Samples.Add(FVector2D(Interp, DefaultObject->Evaluate(Interp)));
-			Interp += 0.05f;
-		}
-
-		DefaultObject->Type = DefaultType;
-		EasingType = InValue;
-
-		ChildSlot
-		[
-			SNew(SOverlay)
-		];
-	}
-
-	virtual void OnMouseEnter(const FGeometry& MyGeometry, const FPointerEvent& MouseEvent) override
-	{
-		if (!TimerHandle.IsValid())
-		{
-			MouseOverTime = FSlateApplication::Get().GetCurrentTime();
-			TimerHandle = RegisterActiveTimer(0.f, FWidgetActiveTimerDelegate::CreateSP(this, &SBuiltInFunctionVisualizer::TickInterp));
-		}
-	}
-
-	virtual void OnMouseLeave(const FPointerEvent& MouseEvent) override
-	{
-		if (TimerHandle.IsValid())
-		{
-			InterpValue = FVector2D::ZeroVector;
-			UnRegisterActiveTimer(TimerHandle.ToSharedRef());
-			TimerHandle = nullptr;
-		}
-	}
-
-	EActiveTimerReturnType TickInterp(const double InCurrentTime, const float InDeltaTime)
-	{
-		static float InterpInPad = .25f, InterpOutPad = .5f;
-		static float InterpDuration = .5f;
-
-		float TotalInterpTime = InterpInPad + InterpDuration + InterpOutPad;
-		InterpValue.X = FMath::Clamp((FMath::Fmod(float(InCurrentTime - MouseOverTime), TotalInterpTime) - InterpInPad) / InterpDuration, 0.f, 1.f);
-
-		UMovieSceneBuiltInEasingFunction* DefaultObject = GetMutableDefault<UMovieSceneBuiltInEasingFunction>();
-
-		EMovieSceneBuiltInEasing DefaultType = DefaultObject->Type;
-		DefaultObject->Type = EasingType;
-
-		InterpValue.Y = DefaultObject->Evaluate(InterpValue.X);
-		DefaultObject->Type = DefaultType;
-
-		return EActiveTimerReturnType::Continue;
-	}
-
-	virtual int32 OnPaint( const FPaintArgs& Args, const FGeometry& AllottedGeometry, const FSlateRect& MyClippingRect, FSlateWindowElementList& OutDrawElements, int32 LayerId, const FWidgetStyle& InWidgetStyle, bool bParentEnabled ) const
-	{
-		float VerticalPad = 0.2f;
-		FVector2D InverseVerticalSize(AllottedGeometry.Size.X, -AllottedGeometry.Size.Y);
-
-		const float VerticalBottom = AllottedGeometry.Size.Y-AllottedGeometry.Size.Y*VerticalPad*.5f;
-		const float CurveHeight = AllottedGeometry.Size.Y * (1.f-VerticalPad);
-		const float CurveWidth = AllottedGeometry.Size.X - 5.f;
-
-		TArray<FVector2D> Points;
-		for (FVector2D Sample : Samples)
-		{
-			FVector2D Offset(5.f, VerticalBottom);
-			Points.Add(Offset + FVector2D(
-				CurveWidth*Sample.X,
-				-CurveHeight * Sample.Y
-			));
-		}
-
-		FSlateDrawElement::MakeLines(
-			OutDrawElements,
-			LayerId,
-			AllottedGeometry.ToPaintGeometry(),
-			Points,
-			ESlateDrawEffect::None);
-
-		if (TimerHandle.IsValid())
-		{
-			FVector2D PointOffset(0.f, VerticalBottom - CurveHeight*InterpValue.Y - 4.f);
-
-			static const FSlateBrush* InterpPointBrush = FAppStyle::GetBrush("Sequencer.InterpLine");
-			FSlateDrawElement::MakeBox(
-				OutDrawElements,
-				LayerId+1,
-				AllottedGeometry.MakeChild(
-					FVector2D(AllottedGeometry.Size.X, 7.f),
-					FSlateLayoutTransform(PointOffset)
-				).ToPaintGeometry(),
-				InterpPointBrush,
-				ESlateDrawEffect::None,
-				FLinearColor::Green
-			);
-		}
-
-		return LayerId+1;
-	}
-
-private:
-
-	TSharedPtr<FActiveTimerHandle> TimerHandle;
-	double MouseOverTime;
-	EMovieSceneBuiltInEasing EasingType;
-
-	FVector2D InterpValue;
-	TArray<FVector2D> Samples;
-};
-
 void FMovieSceneBuiltInEasingFunctionCustomization::CustomizeDetails(IDetailLayoutBuilder& DetailBuilder)
 {
 	TypeProperty = DetailBuilder.GetProperty(GET_MEMBER_NAME_CHECKED(UMovieSceneBuiltInEasingFunction, Type));
 
-	const UEnum* EasingEnum = StaticEnum<EMovieSceneBuiltInEasing>();
-	check(EasingEnum)
-
-	TArray<FGroupedEasing> Groups;
-	auto FindGroup = [&](const FString& GroupName) -> FGroupedEasing&
-	{
-		for (int32 Index = Groups.Num() - 1; Index >= 0; --Index)
-		{
-			if (Groups[Index].GroupingName == GroupName)
-			{
-				return Groups[Index];
-			}
-		}
-
-		Groups.Emplace();
-		Groups.Last().GroupingName = GroupName;
-		return Groups.Last();
-	};
-
-	for (int32 NameIndex = 0; NameIndex < EasingEnum->NumEnums() - 1; ++NameIndex)
-	{
-		const FString& Grouping = EasingEnum->GetMetaData(TEXT("Grouping"), NameIndex);
-		EMovieSceneBuiltInEasing Value = (EMovieSceneBuiltInEasing)EasingEnum->GetValueByIndex(NameIndex);
-
-		FindGroup(Grouping).Values.Add(Value);
-	}
-
-
-	TSharedRef<SGridPanel> Grid = SNew(SGridPanel);
-
-	int32 RowIndex = 0;
-	for (const FGroupedEasing& Group : Groups)
-	{
-		for (int32 ColumnIndex = 0; ColumnIndex < Group.Values.Num(); ++ColumnIndex)
-		{
-			EMovieSceneBuiltInEasing Value = Group.Values[ColumnIndex];
-
-			Grid->AddSlot(ColumnIndex, RowIndex)
-			[
-				SNew(SBox)
-				.WidthOverride(100.f)
-				.HeightOverride(50.f)
-				[
-					SNew(SButton)
-					.ButtonStyle(FAppStyle::Get(), "HoverHintOnly")
-					.OnClicked(this, &FMovieSceneBuiltInEasingFunctionCustomization::SetType, Value)
-					[
-						SNew(SBuiltInFunctionVisualizer, Value)
-					]
-				]
-			];
-
-			Grid->AddSlot(ColumnIndex, RowIndex+1)
-			.HAlign(HAlign_Center)
-			[
-				SNew(STextBlock)
-				.Text(EasingEnum->GetDisplayNameTextByValue((int64)Value))
-			];
-		}
-
-		RowIndex += 2;
-	}
+	TSet<EMovieSceneBuiltInEasing> FilterExclude;
+	FilterExclude.Add(EMovieSceneBuiltInEasing::Custom);
 
 	IDetailCategoryBuilder& Category = DetailBuilder.EditCategory("Easing");
 	DetailBuilder.HideProperty(TypeProperty);
 	FDetailWidgetRow& Row = Category.AddCustomRow(FText());
 	Row.WholeRowContent()
 	[
-		Grid
+		SNew(SEasingFunctionGridWidget)
+		.FilterExclude(FilterExclude)
+		.OnTypeChanged(this, &FMovieSceneBuiltInEasingFunctionCustomization::SetType)
 	];
 }
 
-FReply FMovieSceneBuiltInEasingFunctionCustomization::SetType(EMovieSceneBuiltInEasing NewType)
+void FMovieSceneBuiltInEasingFunctionCustomization::SetType(EMovieSceneBuiltInEasing NewType)
 {
 	FScopedTransaction Transaction(NSLOCTEXT("EasingFunctionCustomization", "SetEasingType", "Set Easing Type"));	
 
@@ -275,6 +86,52 @@ FReply FMovieSceneBuiltInEasingFunctionCustomization::SetType(EMovieSceneBuiltIn
 
 	TypeProperty->NotifyPostChange(EPropertyChangeType::ValueSet);
 	TypeProperty->NotifyFinishedChangingProperties();
+}
+
+///////////////////////////////////////////////////////////////////////////////
+
+void SBuiltInFunctionVisualizerWithText::Construct(const FArguments& InArgs, EMovieSceneBuiltInEasing InValue)
+{
+	const UEnum* EasingEnum = StaticEnum<EMovieSceneBuiltInEasing>();
+	check(EasingEnum);
+
+	FunctionVisualizer = SNew(SBuiltInFunctionVisualizer, InValue);
+	FunctionName = SNew(STextBlock)
+		.Font(IDetailLayoutBuilder::GetDetailFont())
+		.Text(EasingEnum->GetDisplayNameTextByValue((int64)InValue));
 
-	return FReply::Unhandled();
+	ChildSlot
+		[
+			SNew(SVerticalBox)
+
+			// Add curve
+			+ SVerticalBox::Slot()
+			.HAlign(HAlign_Center)
+			[
+				SNew(SBox)
+				.WidthOverride(100.0f)
+				.HeightOverride(50.0f)
+				[
+					FunctionVisualizer.ToSharedRef()
+				]
+			]
+
+			// Add curve type name
+			+ SVerticalBox::Slot()
+			.HAlign(HAlign_Center)
+			.AutoHeight()
+			.Padding(0.0f, 0.0f, 0.0f, 2.0f)
+			[
+				FunctionName.ToSharedRef()
+			]
+		];
 }
+
+void SBuiltInFunctionVisualizerWithText::SetType(EMovieSceneBuiltInEasing InValue)
+{
+	FunctionVisualizer->SetType(InValue);
+
+	const UEnum* EasingEnum = StaticEnum<EMovieSceneBuiltInEasing>();
+	check(EasingEnum);
+	FunctionName->SetText(EasingEnum->GetDisplayNameTextByValue((int64)InValue));
+}
\ No newline at end of file
diff --git a/Engine/Source/Editor/MovieSceneTools/Private/MovieSceneBuiltInEasingFunctionCustomization.h b/Engine/Source/Editor/MovieSceneTools/Private/MovieSceneBuiltInEasingFunctionCustomization.h
--- a/Engine/Source/Editor/MovieSceneTools/Private/MovieSceneBuiltInEasingFunctionCustomization.h
+++ b/Engine/Source/Editor/MovieSceneTools/Private/MovieSceneBuiltInEasingFunctionCustomization.h
@@ -5,6 +5,10 @@
 #include "HAL/Platform.h"
 #include "IDetailCustomization.h"
 #include "Templates/SharedPointer.h"
+#include "Generators/MovieSceneEasingCurves.h"
+#include "Widgets/SCompoundWidget.h"
+#include "Widgets/Text/STextBlock.h"
+#include "MovieSceneBuiltInEasingFunctionVisualizer.h"
 
 class IDetailLayoutBuilder;
 
@@ -15,12 +19,25 @@ class IPropertyHandle;
 class FMovieSceneBuiltInEasingFunctionCustomization : public IDetailCustomization
 {
 public:
-
 	virtual void CustomizeDetails(IDetailLayoutBuilder& DetailBuilder) override;
 
-	FReply SetType(EMovieSceneBuiltInEasing NewType);
+	void SetType(EMovieSceneBuiltInEasing NewType);
 
 private:
-
 	TSharedPtr<IPropertyHandle> TypeProperty;
+};
+
+/** Widget showing the curve graph, and the curve name below it */
+class SBuiltInFunctionVisualizerWithText : public SCompoundWidget
+{
+public:
+	SLATE_BEGIN_ARGS(SBuiltInFunctionVisualizerWithText) {}
+	SLATE_END_ARGS()
+
+	void Construct(const FArguments& InArgs, EMovieSceneBuiltInEasing InValue);
+	void SetType(EMovieSceneBuiltInEasing InValue);
+
+private:
+	TSharedPtr<SBuiltInFunctionVisualizer> FunctionVisualizer;
+	TSharedPtr<STextBlock> FunctionName;
 };
\ No newline at end of file
diff --git a/Engine/Source/Editor/MovieSceneTools/Private/MovieSceneBuiltInEasingFunctionGridWidget.cpp b/Engine/Source/Editor/MovieSceneTools/Private/MovieSceneBuiltInEasingFunctionGridWidget.cpp
new file mode 100644
--- /dev/null
+++ b/Engine/Source/Editor/MovieSceneTools/Private/MovieSceneBuiltInEasingFunctionGridWidget.cpp
@@ -0,0 +1,146 @@
+// Copyright Epic Games, Inc. All Rights Reserved.
+
+#include "MovieSceneBuiltInEasingFunctionGridWidget.h"
+#include "MovieSceneBuiltInEasingFunctionVisualizer.h"
+
+#include "Containers/Array.h"
+#include "Containers/ContainerAllocationPolicies.h"
+#include "Containers/UnrealString.h"
+#include "Delegates/Delegate.h"
+#include "DetailCategoryBuilder.h"
+#include "DetailLayoutBuilder.h"
+#include "DetailWidgetRow.h"
+#include "Framework/Application/SlateApplication.h"
+#include "Generators/MovieSceneEasingCurves.h"
+#include "HAL/PlatformCrt.h"
+#include "Input/Reply.h"
+#include "Internationalization/Internationalization.h"
+#include "Internationalization/Text.h"
+#include "Layout/Children.h"
+#include "Layout/Geometry.h"
+#include "Math/Color.h"
+#include "Math/UnrealMathSSE.h"
+#include "Math/Vector2D.h"
+#include "Misc/AssertionMacros.h"
+#include "Misc/Attribute.h"
+#include "PropertyHandle.h"
+#include "Rendering/DrawElements.h"
+#include "Rendering/RenderingCommon.h"
+#include "Rendering/SlateLayoutTransform.h"
+#include "ScopedTransaction.h"
+#include "SlotBase.h"
+#include "Styling/AppStyle.h"
+#include "Types/SlateEnums.h"
+#include "Types/SlateStructs.h"
+#include "Types/WidgetActiveTimerDelegate.h"
+#include "UObject/Class.h"
+#include "UObject/UObjectGlobals.h"
+#include "UObject/UnrealType.h"
+#include "Widgets/DeclarativeSyntaxSupport.h"
+#include "Widgets/Input/SButton.h"
+#include "Widgets/Layout/SBox.h"
+#include "Widgets/Layout/SGridPanel.h"
+#include "Widgets/SCompoundWidget.h"
+#include "Widgets/SOverlay.h"
+#include "Widgets/Text/STextBlock.h"
+#include "IDetailChildrenBuilder.h"
+
+class FActiveTimerHandle;
+class FPaintArgs;
+class FSlateRect;
+class FWidgetStyle;
+struct FPointerEvent;
+struct FSlateBrush;
+
+SEasingFunctionGridWidget::FGroup& SEasingFunctionGridWidget::FindOrAddGroup(TArray<FGroup>& Groups, const FString& GroupName)
+{
+	for (int32 Index = Groups.Num() - 1; Index >= 0; --Index)
+	{
+		if (Groups[Index].GroupName == GroupName)
+		{
+			return Groups[Index];
+		}
+	}
+
+	Groups.Emplace();
+	Groups.Last().GroupName = GroupName;
+	return Groups.Last();
+};
+
+TArray<SEasingFunctionGridWidget::FGroup> SEasingFunctionGridWidget::ConstructGroups(const TSet<EMovieSceneBuiltInEasing>& FilterExclude)
+{
+	const UEnum* EasingEnum = StaticEnum<EMovieSceneBuiltInEasing>();
+	check(EasingEnum)
+
+	TArray<FGroup> Groups;
+
+	for (int32 NameIndex = 0; NameIndex < EasingEnum->NumEnums() - 1; ++NameIndex)
+	{
+		const FString& Grouping = EasingEnum->GetMetaData(TEXT("Grouping"), NameIndex);
+		EMovieSceneBuiltInEasing Value = (EMovieSceneBuiltInEasing)EasingEnum->GetValueByIndex(NameIndex);
+
+		if (FilterExclude.IsEmpty() || FilterExclude.Find(Value) == nullptr)
+		{
+			FindOrAddGroup(Groups, Grouping).Values.Add(Value);
+		}
+	}
+
+	return Groups;
+}
+
+FReply SEasingFunctionGridWidget::OnTypeButtonClicked(EMovieSceneBuiltInEasing type)
+{
+	OnClickedDelegate.ExecuteIfBound(type);
+	return FReply::Handled();
+}
+
+void SEasingFunctionGridWidget::Construct(const FArguments& InArgs)
+{
+	FilterExcludeAttribute = InArgs._FilterExclude;
+	OnClickedDelegate = InArgs._OnTypeChanged;
+
+	const UEnum* EasingEnum = StaticEnum<EMovieSceneBuiltInEasing>();
+	check(EasingEnum)
+
+	TArray<FGroup> Groups = ConstructGroups(FilterExcludeAttribute.Get());
+
+	TSharedRef<SGridPanel> Grid = SNew(SGridPanel);
+
+	int32 RowIndex = 0;
+	for (const FGroup& Group : Groups)
+	{
+		for (int32 ColumnIndex = 0; ColumnIndex < Group.Values.Num(); ++ColumnIndex)
+		{
+			EMovieSceneBuiltInEasing Value = Group.Values[ColumnIndex];
+
+			Grid->AddSlot(ColumnIndex, RowIndex)
+				[
+					SNew(SBox)
+					.WidthOverride(100.f)
+					.HeightOverride(50.f)
+					[
+						SNew(SButton)
+						.ButtonStyle(FAppStyle::Get(), "HoverHintOnly")
+						.OnClicked(this, &SEasingFunctionGridWidget::OnTypeButtonClicked, Value)
+						[
+							SNew(SBuiltInFunctionVisualizer, Value)
+						]
+					]
+				];
+
+			Grid->AddSlot(ColumnIndex, RowIndex + 1)
+				.HAlign(HAlign_Center)
+				[
+					SNew(STextBlock)
+					.Text(EasingEnum->GetDisplayNameTextByValue((int64)Value))
+				];
+		}
+
+		RowIndex += 2;
+	}
+
+	ChildSlot
+		[
+			Grid
+		];
+}
\ No newline at end of file
diff --git a/Engine/Source/Editor/MovieSceneTools/Private/MovieSceneBuiltInEasingFunctionGridWidget.h b/Engine/Source/Editor/MovieSceneTools/Private/MovieSceneBuiltInEasingFunctionGridWidget.h
new file mode 100644
--- /dev/null
+++ b/Engine/Source/Editor/MovieSceneTools/Private/MovieSceneBuiltInEasingFunctionGridWidget.h
@@ -0,0 +1,50 @@
+// Copyright Epic Games, Inc. All Rights Reserved.
+
+#pragma once
+
+#include "HAL/Platform.h"
+#include "IDetailCustomization.h"
+#include "Templates/SharedPointer.h"
+#include "Generators/MovieSceneEasingCurves.h"
+#include "Widgets/SCompoundWidget.h"
+
+class IDetailLayoutBuilder;
+
+enum class EMovieSceneBuiltInEasing : uint8;
+class FReply;
+class IPropertyHandle;
+
+DECLARE_DELEGATE_OneParam(FEasingFunctionGridWidget_OnClicked, EMovieSceneBuiltInEasing);
+
+/** Widget showing a grid of curves including their names.
+ * Curves are grouped within rows, e.g. Exponential in, out, in-out in one row.
+ * A filter can be used to exclude several curves from the grid.
+ */
+class SEasingFunctionGridWidget : public SCompoundWidget
+{
+public:
+	SLATE_BEGIN_ARGS(SEasingFunctionGridWidget) {}
+
+	/** The easing curve filter containing all curve types that should be excluded.In case the filter is empty, all curve types will be shown. */
+	SLATE_ATTRIBUTE(TSet<EMovieSceneBuiltInEasing>, FilterExclude)
+	SLATE_EVENT(FEasingFunctionGridWidget_OnClicked, OnTypeChanged)
+	
+	SLATE_END_ARGS()
+
+	void Construct(const FArguments& InArgs);
+
+private:
+	struct FGroup
+	{
+		FString GroupName;
+		TArray<EMovieSceneBuiltInEasing, TInlineAllocator<3>> Values;
+	};
+	
+	FGroup& FindOrAddGroup(TArray<FGroup>& Groups, const FString& GroupName);
+	TArray<FGroup> ConstructGroups(const TSet<EMovieSceneBuiltInEasing>& FilterExclude);
+
+	FReply OnTypeButtonClicked(EMovieSceneBuiltInEasing type);
+
+	FEasingFunctionGridWidget_OnClicked OnClickedDelegate;
+	TAttribute<TSet<EMovieSceneBuiltInEasing>> FilterExcludeAttribute;
+};
\ No newline at end of file
diff --git a/Engine/Source/Editor/MovieSceneTools/Private/MovieSceneBuiltInEasingFunctionVisualizer.cpp b/Engine/Source/Editor/MovieSceneTools/Private/MovieSceneBuiltInEasingFunctionVisualizer.cpp
new file mode 100644
--- /dev/null
+++ b/Engine/Source/Editor/MovieSceneTools/Private/MovieSceneBuiltInEasingFunctionVisualizer.cpp
@@ -0,0 +1,167 @@
+// Copyright Epic Games, Inc. All Rights Reserved.
+
+#include "MovieSceneBuiltInEasingFunctionVisualizer.h"
+
+#include "Containers/Array.h"
+#include "Containers/ContainerAllocationPolicies.h"
+#include "Containers/UnrealString.h"
+#include "Delegates/Delegate.h"
+#include "DetailCategoryBuilder.h"
+#include "DetailLayoutBuilder.h"
+#include "DetailWidgetRow.h"
+#include "Framework/Application/SlateApplication.h"
+#include "Generators/MovieSceneEasingCurves.h"
+#include "HAL/PlatformCrt.h"
+#include "Input/Reply.h"
+#include "Internationalization/Internationalization.h"
+#include "Internationalization/Text.h"
+#include "Layout/Children.h"
+#include "Layout/Geometry.h"
+#include "Math/Color.h"
+#include "Math/UnrealMathSSE.h"
+#include "Math/Vector2D.h"
+#include "Misc/AssertionMacros.h"
+#include "Misc/Attribute.h"
+#include "PropertyHandle.h"
+#include "Rendering/DrawElements.h"
+#include "Rendering/RenderingCommon.h"
+#include "Rendering/SlateLayoutTransform.h"
+#include "ScopedTransaction.h"
+#include "SlotBase.h"
+#include "Styling/AppStyle.h"
+#include "Types/SlateEnums.h"
+#include "Types/SlateStructs.h"
+#include "Types/WidgetActiveTimerDelegate.h"
+#include "UObject/Class.h"
+#include "UObject/UObjectGlobals.h"
+#include "UObject/UnrealType.h"
+#include "Widgets/DeclarativeSyntaxSupport.h"
+#include "Widgets/Input/SButton.h"
+#include "Widgets/Layout/SBox.h"
+#include "Widgets/Layout/SGridPanel.h"
+#include "Widgets/SCompoundWidget.h"
+#include "Widgets/SOverlay.h"
+#include "Widgets/Text/STextBlock.h"
+#include "IDetailChildrenBuilder.h"
+
+class FActiveTimerHandle;
+class FPaintArgs;
+class FSlateRect;
+class FWidgetStyle;
+struct FPointerEvent;
+struct FSlateBrush;
+
+void SBuiltInFunctionVisualizer::SetType(EMovieSceneBuiltInEasing InValue)
+{
+	InterpValue = FVector2D::ZeroVector;
+
+	UMovieSceneBuiltInEasingFunction* DefaultObject = GetMutableDefault<UMovieSceneBuiltInEasingFunction>();
+	EMovieSceneBuiltInEasing DefaultType = DefaultObject->Type;
+
+	DefaultObject->Type = InValue;
+
+	Samples.Reset();
+	float Interp = 0.f;
+	while (Interp <= 1.f)
+	{
+		Samples.Add(FVector2D(Interp, DefaultObject->Evaluate(Interp)));
+		Interp += 0.05f;
+	}
+
+	DefaultObject->Type = DefaultType;
+	EasingType = InValue;
+}
+
+void SBuiltInFunctionVisualizer::Construct(const FArguments& InArgs, EMovieSceneBuiltInEasing InValue)
+{
+	SetType(InValue);
+
+	ChildSlot
+		[
+			SNew(SOverlay)
+		];
+}
+
+void SBuiltInFunctionVisualizer::OnMouseEnter(const FGeometry& MyGeometry, const FPointerEvent& MouseEvent)
+{
+	if (!TimerHandle.IsValid())
+	{
+		MouseOverTime = FSlateApplication::Get().GetCurrentTime();
+		TimerHandle = RegisterActiveTimer(0.f, FWidgetActiveTimerDelegate::CreateSP(this, &SBuiltInFunctionVisualizer::TickInterp));
+	}
+}
+
+void SBuiltInFunctionVisualizer::OnMouseLeave(const FPointerEvent& MouseEvent)
+{
+	if (TimerHandle.IsValid())
+	{
+		InterpValue = FVector2D::ZeroVector;
+		UnRegisterActiveTimer(TimerHandle.ToSharedRef());
+		TimerHandle = nullptr;
+	}
+}
+
+EActiveTimerReturnType SBuiltInFunctionVisualizer::TickInterp(const double InCurrentTime, const float InDeltaTime)
+{
+	static float InterpInPad = .25f, InterpOutPad = .5f;
+
+	float TotalInterpTime = InterpInPad + AnimationDuration + InterpOutPad;
+	InterpValue.X = FMath::Clamp((FMath::Fmod(float(InCurrentTime - MouseOverTime), TotalInterpTime) - InterpInPad) / AnimationDuration, 0.f, 1.f);
+
+	UMovieSceneBuiltInEasingFunction* DefaultObject = GetMutableDefault<UMovieSceneBuiltInEasingFunction>();
+
+	EMovieSceneBuiltInEasing DefaultType = DefaultObject->Type;
+	DefaultObject->Type = EasingType;
+
+	InterpValue.Y = DefaultObject->Evaluate(InterpValue.X);
+	DefaultObject->Type = DefaultType;
+
+	return EActiveTimerReturnType::Continue;
+}
+
+int32 SBuiltInFunctionVisualizer::OnPaint(const FPaintArgs& Args, const FGeometry& AllottedGeometry, const FSlateRect& MyClippingRect, FSlateWindowElementList& OutDrawElements, int32 LayerId, const FWidgetStyle& InWidgetStyle, bool bParentEnabled) const
+{
+	float VerticalPad = 0.2f;
+	FVector2D InverseVerticalSize(AllottedGeometry.Size.X, -AllottedGeometry.Size.Y);
+
+	const float VerticalBottom = AllottedGeometry.Size.Y - AllottedGeometry.Size.Y * VerticalPad * .5f;
+	const float CurveHeight = AllottedGeometry.Size.Y * (1.f - VerticalPad);
+	const float CurveWidth = AllottedGeometry.Size.X - 5.f;
+
+	TArray<FVector2D> Points;
+	for (FVector2D Sample : Samples)
+	{
+		FVector2D Offset(5.f, VerticalBottom);
+		Points.Add(Offset + FVector2D(
+			CurveWidth * Sample.X,
+			-CurveHeight * Sample.Y
+		));
+	}
+
+	FSlateDrawElement::MakeLines(
+		OutDrawElements,
+		LayerId,
+		AllottedGeometry.ToPaintGeometry(),
+		Points,
+		ESlateDrawEffect::None);
+
+	if (TimerHandle.IsValid())
+	{
+		FVector2D PointOffset(0.f, VerticalBottom - CurveHeight * InterpValue.Y - 4.f);
+
+		static const FSlateBrush* InterpPointBrush = FAppStyle::GetBrush("Sequencer.InterpLine");
+		FSlateDrawElement::MakeBox(
+			OutDrawElements,
+			LayerId + 1,
+			AllottedGeometry.MakeChild(
+				FVector2D(AllottedGeometry.Size.X, 7.f),
+				FSlateLayoutTransform(PointOffset)
+			).ToPaintGeometry(),
+			InterpPointBrush,
+			ESlateDrawEffect::None,
+			FLinearColor::Green
+		);
+	}
+
+	return LayerId + 1;
+}
\ No newline at end of file
diff --git a/Engine/Source/Editor/MovieSceneTools/Private/MovieSceneBuiltInEasingFunctionVisualizer.h b/Engine/Source/Editor/MovieSceneTools/Private/MovieSceneBuiltInEasingFunctionVisualizer.h
new file mode 100644
--- /dev/null
+++ b/Engine/Source/Editor/MovieSceneTools/Private/MovieSceneBuiltInEasingFunctionVisualizer.h
@@ -0,0 +1,46 @@
+// Copyright Epic Games, Inc. All Rights Reserved.
+
+#pragma once
+
+#include "HAL/Platform.h"
+#include "IDetailCustomization.h"
+#include "Generators/MovieSceneEasingCurves.h"
+#include "Templates/SharedPointer.h"
+#include "Widgets/SCompoundWidget.h"
+
+class IDetailLayoutBuilder;
+
+enum class EMovieSceneBuiltInEasing : uint8;
+class FReply;
+class IPropertyHandle;
+
+/** Widget showing the curve graph */
+class SBuiltInFunctionVisualizer : public SCompoundWidget
+{
+public:
+	SLATE_BEGIN_ARGS(SBuiltInFunctionVisualizer) {}
+	SLATE_END_ARGS();
+
+	virtual ~SBuiltInFunctionVisualizer() = default;
+
+	void Construct(const FArguments& InArgs, EMovieSceneBuiltInEasing InValue);
+	
+	void SetType(EMovieSceneBuiltInEasing InValue);
+	void SetAnimationDuration(float TimeInSeconds) { AnimationDuration = TimeInSeconds; }
+
+private:
+	void OnMouseEnter(const FGeometry& MyGeometry, const FPointerEvent& MouseEvent) override final;
+	void OnMouseLeave(const FPointerEvent& MouseEvent) override final;
+
+	EActiveTimerReturnType TickInterp(const double InCurrentTime, const float InDeltaTime);
+
+	int32 OnPaint(const FPaintArgs& Args, const FGeometry& AllottedGeometry, const FSlateRect& MyClippingRect, FSlateWindowElementList& OutDrawElements, int32 LayerId, const FWidgetStyle& InWidgetStyle, bool bParentEnabled) const;
+
+	TSharedPtr<FActiveTimerHandle> TimerHandle;
+	double MouseOverTime;
+	EMovieSceneBuiltInEasing EasingType;
+
+	FVector2D InterpValue;
+	TArray<FVector2D> Samples;
+	float AnimationDuration = 0.5f; /** Preview animation duration in seconds. */
+};
\ No newline at end of file
diff --git a/Engine/Source/Editor/MovieSceneTools/Private/MovieSceneToolsModule.cpp b/Engine/Source/Editor/MovieSceneTools/Private/MovieSceneToolsModule.cpp
--- a/Engine/Source/Editor/MovieSceneTools/Private/MovieSceneToolsModule.cpp
+++ b/Engine/Source/Editor/MovieSceneTools/Private/MovieSceneToolsModule.cpp
@@ -60,6 +60,7 @@
 #include "MVVM/ViewModels/CinematicShotTrackModel.h"
 
 #include "MovieSceneBuiltInEasingFunctionCustomization.h"
+#include "MovieSceneAlphaBlendOptionCustomization.h"
 #include "MovieSceneObjectBindingIDCustomization.h"
 #include "MovieSceneEventCustomization.h"
 #include "SequencerClipboardReconciler.h"
@@ -184,6 +185,7 @@ void FMovieSceneToolsModule::StartupModule()
 		FPropertyEditorModule& PropertyModule = FModuleManager::LoadModuleChecked<FPropertyEditorModule>("PropertyEditor");
 		PropertyModule.RegisterCustomClassLayout("MovieSceneToolsProjectSettings", FOnGetDetailCustomizationInstance::CreateStatic(&FMovieSceneToolsProjectSettingsCustomization::MakeInstance));
 		PropertyModule.RegisterCustomClassLayout("MovieSceneBuiltInEasingFunction", FOnGetDetailCustomizationInstance::CreateLambda(&MakeShared<FMovieSceneBuiltInEasingFunctionCustomization>));
+		PropertyModule.RegisterCustomPropertyTypeLayout("EAlphaBlendOption", FOnGetPropertyTypeCustomizationInstance::CreateLambda(&MakeShared<FAlphaBlendPropertyCustomization>));
 		PropertyModule.RegisterCustomClassLayout("MovieSceneFloatPerlinNoiseChannelContainer", FOnGetDetailCustomizationInstance::CreateStatic(&FMovieSceneFloatPerlinNoiseChannelDetailsCustomization::MakeInstance));
 		PropertyModule.RegisterCustomClassLayout("MovieSceneDoublePerlinNoiseChannelContainer", FOnGetDetailCustomizationInstance::CreateStatic(&FMovieSceneDoublePerlinNoiseChannelDetailsCustomization::MakeInstance));
 		PropertyModule.RegisterCustomPropertyTypeLayout("MovieSceneObjectBindingID", FOnGetPropertyTypeCustomizationInstance::CreateLambda(&MakeShared<FMovieSceneObjectBindingIDCustomization>));
@@ -337,6 +339,7 @@ void FMovieSceneToolsModule::ShutdownModule()
 		FPropertyEditorModule& PropertyModule = FModuleManager::LoadModuleChecked<FPropertyEditorModule>("PropertyEditor");
 		PropertyModule.UnregisterCustomClassLayout("MovieSceneToolsProjectSettings");
 		PropertyModule.UnregisterCustomClassLayout("MovieSceneBuiltInEasingFunction");
+		PropertyModule.UnregisterCustomPropertyTypeLayout("EAlphaBlendOption");
 		PropertyModule.UnregisterCustomPropertyTypeLayout("MovieSceneObjectBindingID");
 		PropertyModule.UnregisterCustomPropertyTypeLayout("MovieSceneEvent");
 	}
diff --git a/Engine/Source/Runtime/Engine/Public/AlphaBlend.h b/Engine/Source/Runtime/Engine/Public/AlphaBlend.h
--- a/Engine/Source/Runtime/Engine/Public/AlphaBlend.h
+++ b/Engine/Source/Runtime/Engine/Public/AlphaBlend.h
@@ -11,36 +11,21 @@ class UCurveFloat;
 UENUM()
 enum class EAlphaBlendOption : uint8
 {
-	// Linear interpolation
-	Linear = 0,
-	// Cubic-in interpolation
-	Cubic,
-	// Hermite-Cubic
-	HermiteCubic,
-	// Sinusoidal interpolation
-	Sinusoidal,
-	// Quadratic in-out interpolation
-	QuadraticInOut,
-	// Cubic in-out interpolation
-	CubicInOut,
-	// Quartic in-out interpolation
-	QuarticInOut,
-	// Quintic in-out interpolation
-	QuinticInOut,
-	// Circular-in interpolation
-	CircularIn,
-	// Circular-out interpolation
-	CircularOut,
-	// Circular in-out interpolation
-	CircularInOut,
-	// Exponential-in interpolation
-	ExpIn,
-	// Exponential-Out interpolation
-	ExpOut,
-	// Exponential in-out interpolation
-	ExpInOut,
-	// Custom interpolation, will use custom curve inside an FAlphaBlend or linear if none has been set
-	Custom,
+	Linear = 0 UMETA(Grouping = Linear, DisplayName = "Linear", ToolTip = "Linear interpolation"),
+	Cubic UMETA(Grouping = Cubic, DisplayName = "Cubic In", ToolTip = "Cubic-in interpolation"),
+	HermiteCubic UMETA(Grouping = Cubic, DisplayName = "Hermite-Cubic InOut", ToolTip = "Hermite-Cubic"),
+	Sinusoidal UMETA(Grouping = Sinusoidal, DisplayName = "Sinusoidal", ToolTip = "Sinusoidal interpolation"),
+	QuadraticInOut UMETA(Grouping = Quadratic, DisplayName = "Quadratic InOut", ToolTip = "Quadratic in-out interpolation"),
+	CubicInOut UMETA(Grouping = Cubic, DisplayName = "Cubic InOut", ToolTip = "Cubic in-out interpolation"),
+	QuarticInOut UMETA(Grouping = Quartic, DisplayName = "Quartic InOut", ToolTip = "Quartic in-out interpolation"),
+	QuinticInOut UMETA(Grouping = Quintic, DisplayName = "Quintic InOut", ToolTip = "Quintic in-out interpolation"),
+	CircularIn UMETA(Grouping = Circular, DisplayName = "Circular In", ToolTip = "Circular-in interpolation"),
+	CircularOut UMETA(Grouping = Circular, DisplayName = "Circular Out", ToolTip = "Circular-out interpolation"),
+	CircularInOut UMETA(Grouping = Circular, DisplayName = "Circular InOut", ToolTip = "Circular in-out interpolation"),
+	ExpIn UMETA(Grouping = Exponential, DisplayName = "Exponential In", ToolTip = "Exponential-in interpolation"),
+	ExpOut UMETA(Grouping = Exponential, DisplayName = "Exponential Out", ToolTip = "Exponential-Out interpolation"),
+	ExpInOut UMETA(Grouping = Exponential, DisplayName = "Exponential InOut", ToolTip = "Exponential in-out interpolation"),
+	Custom UMETA(Grouping = Custom, DisplayName = "Custom", ToolTip = "Custom interpolation, will use custom curve inside an FAlphaBlend or linear if none has been set"),
 };
 
 /**
diff --git a/Engine/Source/Runtime/MovieScene/Private/Generators/MovieSceneEasingCurves.cpp b/Engine/Source/Runtime/MovieScene/Private/Generators/MovieSceneEasingCurves.cpp
--- a/Engine/Source/Runtime/MovieScene/Private/Generators/MovieSceneEasingCurves.cpp
+++ b/Engine/Source/Runtime/MovieScene/Private/Generators/MovieSceneEasingCurves.cpp
@@ -74,6 +74,8 @@ float UMovieSceneBuiltInEasingFunction::Evaluate(float Interp) const
 	case EMovieSceneBuiltInEasing::CubicOut: 	return PowOut(Interp,3);
 	case EMovieSceneBuiltInEasing::CubicInOut: 	return InTime < 1.f ? .5f*PowIn(InTime,3) : .5f + .5f*PowOut(OutTime,3);
 
+	case EMovieSceneBuiltInEasing::HermiteCubicInOut: return FMath::Clamp<float>(FMath::SmoothStep(0.0f, 1.0f, Interp), 0.0f, 1.0f);
+
 	case EMovieSceneBuiltInEasing::QuartIn: 	return PowIn(Interp,4);
 	case EMovieSceneBuiltInEasing::QuartOut: 	return PowOut(Interp,4);
 	case EMovieSceneBuiltInEasing::QuartInOut: 	return InTime < 1.f ? .5f*PowIn(InTime,4) : .5f + .5f*PowOut(OutTime,4);
diff --git a/Engine/Source/Runtime/MovieScene/Public/Generators/MovieSceneEasingCurves.h b/Engine/Source/Runtime/MovieScene/Public/Generators/MovieSceneEasingCurves.h
--- a/Engine/Source/Runtime/MovieScene/Public/Generators/MovieSceneEasingCurves.h
+++ b/Engine/Source/Runtime/MovieScene/Public/Generators/MovieSceneEasingCurves.h
@@ -27,7 +27,7 @@ enum class EMovieSceneBuiltInEasing : uint8
 	// Quadratic easing
 	QuadIn UMETA(Grouping=Quadratic,DisplayName="Quadratic In"), QuadOut UMETA(Grouping=Quadratic,DisplayName="Quadratic Out"), QuadInOut UMETA(Grouping=Quadratic,DisplayName="Quadratic InOut"),
 	// Cubic easing
-	CubicIn UMETA(Grouping=Cubic,DisplayName="Cubic In"), CubicOut UMETA(Grouping=Cubic,DisplayName="Cubic Out"), CubicInOut UMETA(Grouping=Cubic,DisplayName="Cubic InOut"),
+	CubicIn UMETA(Grouping=Cubic,DisplayName="Cubic In"), CubicOut UMETA(Grouping=Cubic,DisplayName="Cubic Out"), CubicInOut UMETA(Grouping=Cubic,DisplayName="Cubic InOut"), HermiteCubicInOut UMETA(Grouping = Cubic, DisplayName = "Hermite-Cubic InOut"),
 	// Quartic easing
 	QuartIn UMETA(Grouping=Quartic,DisplayName="Quartic In"), QuartOut UMETA(Grouping=Quartic,DisplayName="Quartic Out"), QuartInOut UMETA(Grouping=Quartic,DisplayName="Quartic InOut"),
 	// Quintic easing
@@ -36,6 +36,8 @@ enum class EMovieSceneBuiltInEasing : uint8
 	ExpoIn UMETA(Grouping=Exponential,DisplayName="Exponential In"), ExpoOut UMETA(Grouping=Exponential,DisplayName="Exponential Out"), ExpoInOut UMETA(Grouping=Exponential,DisplayName="Exponential InOut"),
 	// Circular easing
 	CircIn UMETA(Grouping=Circular,DisplayName="Circular In"), CircOut UMETA(Grouping=Circular,DisplayName="Circular Out"), CircInOut UMETA(Grouping=Circular,DisplayName="Circular InOut"),
+	// Custom
+	Custom UMETA(Grouping = Custom, DisplayName = "Custom"),
 };
 
 
