Fix channel curve colors incorrectly named when saved in the editor preferences

diff --git a/Engine/Source/Editor/MovieSceneTools/Private/Channels/ChannelCurveModel.cpp b/Engine/Source/Editor/MovieSceneTools/Private/Channels/ChannelCurveModel.cpp
--- a/Engine/Source/Editor/MovieSceneTools/Private/Channels/ChannelCurveModel.cpp
+++ b/Engine/Source/Editor/MovieSceneTools/Private/Channels/ChannelCurveModel.cpp
@@ -444,8 +444,13 @@ void FChannelCurveModel<ChannelType, ChannelValue, KeyType>::GetCurveColorObject
 {
 	if (UMovieSceneSection* Section = WeakSection.Get())
 	{
-
 		*OutObject = Section->GetImplicitObjectOwner();
+
+		if (const FMovieSceneChannelMetaData* MetaData = ChannelHandle.GetMetaData())
+		{
+			OutName = FString::Printf(TEXT( "%s.%s" ), *MetaData->Group.ToString(), *MetaData->DisplayText.ToString());
+			return;
+		}
 		OutName = GetIntentionName();
 		return;
 	}
