Sequencer: NLA.
Fix for some rigs that may have the root animation at some bone beyond the a direct child, we just search for the first bone with animation on it, also we now handle non-zero roots(like old ones from Maya/FBX) . Finally when manipulating the offset in the edit mode we need to do it in the space of the start of the section since that's the space in which we are manipulating.

diff --git a/Engine/Source/Editor/MovieSceneTools/Private/EditModes/SkeletalAnimationTrackEditMode.cpp b/Engine/Source/Editor/MovieSceneTools/Private/EditModes/SkeletalAnimationTrackEditMode.cpp
--- a/Engine/Source/Editor/MovieSceneTools/Private/EditModes/SkeletalAnimationTrackEditMode.cpp
+++ b/Engine/Source/Editor/MovieSceneTools/Private/EditModes/SkeletalAnimationTrackEditMode.cpp
@@ -217,18 +217,20 @@ static void RenderTrail(UObject* BoundObject, UMovieSceneSkeletalAnimationTrack*
 	const FLinearColor RootMotionColor(0.75, 0.75, 0.75, 1.0f);
 	TOptional<FVector> OldKeyPos_G;
 
-	FTransform SkelMeshTransform  = SkelMeshComp->GetSocketTransform(NAME_None);
 	int32 Index = 0;
-
-	for (const FTransform& Transform : SkelAnimTrack->RootMotionParams.RootTransforms)
+	if (SkelAnimTrack->GetAllSections().Num() > 0)
 	{
-		FVector NewKeyPos = Transform.GetLocation();
-		FVector NewKeyPos_G = SkelMeshTransform.TransformPosition(NewKeyPos);
-		if (OldKeyPos_G.IsSet())
+		const FTransform SkelMeshTransform = SkelAnimTrack->RootMotionParams.RootMotionStartOffset * SkelMeshComp->GetSocketTransform(NAME_None);
+		for (const FTransform& Transform : SkelAnimTrack->RootMotionParams.RootTransforms)
 		{
-			PDI->DrawLine(OldKeyPos_G.GetValue(), NewKeyPos_G, RootMotionColor, SDPG_Foreground);
+			FVector NewKeyPos = Transform.GetLocation();
+			FVector NewKeyPos_G = SkelMeshTransform.TransformPosition(NewKeyPos);
+			if (OldKeyPos_G.IsSet())
+			{
+				PDI->DrawLine(OldKeyPos_G.GetValue(), NewKeyPos_G, RootMotionColor, SDPG_Foreground);
+			}
+			OldKeyPos_G = NewKeyPos_G;
 		}
-		OldKeyPos_G = NewKeyPos_G;
 	}
 }
 
@@ -278,9 +280,7 @@ void FSkeletalAnimationTrackEditMode::Render(const FSceneView* View, FViewport*
 								//show root motions
 								if (SkelAnimTrack->bShowRootMotionTrail)
 								{
-									RenderTrail(BoundObject, SkelAnimTrack, InterrogationLinker, SkelMeshComp, Sequencer, PDI);
-									break;
-								
+									RenderTrail(BoundObject, SkelAnimTrack, InterrogationLinker, SkelMeshComp, Sequencer, PDI);								
 								}
 								//show skeletons
 
@@ -367,9 +367,9 @@ void FSkeletalAnimationTrackEditMode::Render(const FSceneView* View, FViewport*
 	}
 }
 
-bool FSkeletalAnimationTrackEditMode::IsRootSelected(UMovieSceneSkeletalAnimationSection* AnimSection)
+bool FSkeletalAnimationTrackEditMode::IsRootSelected(UMovieSceneSkeletalAnimationSection* AnimSection) const
 {
-	for (FSelectedRootData& RootData : SelectedRootData)
+	for (const FSelectedRootData& RootData : SelectedRootData)
 	{
 		if (RootData.SelectedSection.Get() == AnimSection)
 		{
@@ -432,22 +432,17 @@ FVector FSkeletalAnimationTrackEditMode::GetWidgetLocation() const
 	if (IsSomethingSelected() && WeakSequencer.IsValid())
 	{
 		FQualifiedFrameTime CurrentTime = WeakSequencer.Pin()->GetLocalTime();
-		int32 NumSelected = 0;
+		bool  bSelectionValid = false;
 		for (FSelectedRootData Data : SelectedRootData)
 		{
 			if (Data.SelectedMeshComp.IsValid() && Data.SelectedSection.IsValid() && Data.SelectedSection->GetRange().Contains(CurrentTime.Time.GetFrame()))
 			{
-				FTransform ComponentTransform = Data.SelectedMeshComp->GetComponentTransform();
-				FTransform PivotTransform, PivotParentTransform;
-				Data.CalcTransform(CurrentTime.Time,PivotTransform,PivotParentTransform);
-				PivotLocation += ComponentTransform.TransformPosition(PivotTransform.GetLocation());
-				++NumSelected;
+				bSelectionValid = true;
 			}
 		}
-		if (NumSelected > 0)
+		if (bSelectionValid)
 		{
-			PivotLocation /= (NumSelected);
-			PivotLocation.Z = RootTransform.GetLocation().Z;
+			PivotLocation = RootTransform.GetLocation();
 			return PivotLocation;
 		}
 	}
@@ -464,7 +459,7 @@ bool FSkeletalAnimationTrackEditMode::GetPivotForOrbit(FVector& OutPivot) const
 	return FEdMode::GetPivotForOrbit(OutPivot);
 }
 
-bool FSkeletalAnimationTrackEditMode::GetCurrentTransform(FTransform& OutTransform, FTransform& OutParentTransform) const
+bool FSkeletalAnimationTrackEditMode::GetTransformAtFirstSectionStart(FTransform& OutTransform, FTransform& OutParentTransform) const
 { 
 	if (IsSomethingSelected() && WeakSequencer.IsValid())
 	{
@@ -476,7 +471,8 @@ bool FSkeletalAnimationTrackEditMode::GetCurrentTransform(FTransform& OutTransfo
 			{
 				FTransform ComponentTransform = Data.SelectedMeshComp->GetComponentTransform();
 				FTransform PivotTransform, PivotParentTransform;
-				Data.CalcTransform(CurrentTime.Time, PivotTransform, PivotParentTransform);
+				//we want the transform at the start!
+				Data.CalcTransform(Data.SelectedSection->GetRange().GetLowerBoundValue(), PivotTransform, PivotParentTransform);
 				OutTransform = PivotTransform * ComponentTransform;
 				OutParentTransform = PivotParentTransform * ComponentTransform;
 				return true;
@@ -488,8 +484,9 @@ bool FSkeletalAnimationTrackEditMode::GetCurrentTransform(FTransform& OutTransfo
 
 bool FSkeletalAnimationTrackEditMode::GetCustomDrawingCoordinateSystem(FMatrix& OutMatrix, void* InData)
 {
+	//Get transform at the START of the section since that's the value we will manipulate
 	FTransform Transform, ParentTransform;
-	if (GetCurrentTransform(Transform, ParentTransform))
+	if (GetTransformAtFirstSectionStart(Transform, ParentTransform))
 	{
 		OutMatrix = Transform.ToMatrixNoScale().RemoveTranslation();
 		return true;
@@ -555,8 +552,9 @@ void FSelectedRootData::CalcTransform(const FFrameTime& FrameTime, FTransform& O
 			Param.FrameRate = MovieScene->GetTickResolution();
 			Param.CurrentTime = FrameTime;
 			Section->GetRootMotionTransform(AnimationPoseData, Param);
-			OutTransform = Param.OutTransform * Section->PreviousTransform;
-			OutParentTransform = Param.OutParentTransform * Section->PreviousTransform;
+			OutTransform = Param.OutTransform * Param.OutRootStartTransform * Section->PreviousTransform;
+			OutParentTransform = Param.OutParentTransform * Param.OutRootStartTransform * Section->PreviousTransform;
+
 		}
 	}
 }
@@ -614,23 +612,24 @@ bool FSkeletalAnimationTrackEditMode::InputDelta(FEditorViewportClient* InViewpo
 		const bool bDoTranslation = !Drag.IsZero() && (WidgetMode == UE::Widget::WM_Translate || WidgetMode == UE::Widget::WM_TranslateRotateZ);
 
 		FTransform CurrentTransform,ParentTransform;
-		if ((bDoRotation || bDoTranslation) && GetCurrentTransform(CurrentTransform,ParentTransform))
+		if ((bDoRotation || bDoTranslation) && GetTransformAtFirstSectionStart(CurrentTransform,ParentTransform))
 		{
-			if (bDoRotation)
+			for (const FSelectedRootData Data : SelectedRootData)
 			{
-				FQuat CurrentRotation = CurrentTransform.GetRotation();
-				CurrentRotation = (InRot.Quaternion() * CurrentRotation);
-				CurrentTransform.SetRotation(CurrentRotation);
-			}
+				if (bDoRotation)
+				{
+					FQuat CurrentRotation = CurrentTransform.GetRotation();
+					CurrentRotation = (InRot.Quaternion() * CurrentRotation);
+					CurrentTransform.SetRotation(CurrentRotation);
+				}
+
+				if (bDoTranslation)
+				{
+					FVector CurrentLocation = CurrentTransform.GetLocation();
+					CurrentLocation = CurrentLocation + InDrag;
+					CurrentTransform.SetLocation(CurrentLocation);
+				}
 
-			if (bDoTranslation)
-			{
-				FVector CurrentLocation = CurrentTransform.GetLocation();
-				CurrentLocation = CurrentLocation + InDrag;
-				CurrentTransform.SetLocation(CurrentLocation);
-			}
-			for (FSelectedRootData Data : SelectedRootData)
-			{
 				if (bManipulatorMadeChange == false)
 				{
 					bManipulatorMadeChange = true;
@@ -639,13 +638,13 @@ bool FSkeletalAnimationTrackEditMode::InputDelta(FEditorViewportClient* InViewpo
 				if (Data.SelectedSection.IsValid() && Data.SelectedMeshComp.IsValid())
 				{
 					CurrentTransform = CurrentTransform.GetRelativeTransform(ParentTransform);
-
+					
 					if (bDoRotation)
 					{
 						UMovieSceneSkeletalAnimationTrack* Track = Data.SelectedSection.Get()->GetTypedOuter<UMovieSceneSkeletalAnimationTrack>();
 						Track->Modify();
 						Data.SelectedSection.Get()->Modify();
-						Data.SelectedSection.Get()->StartRotationOffset = CurrentTransform.GetRotation().Rotator(); //(Rot)+Data.SelectedSection.Get()->StartRotationOffset;
+						Data.SelectedSection.Get()->StartRotationOffset = CurrentTransform.GetRotation().Rotator(); 
 						Data.SelectedSection.Get()->GetRootMotionParams()->bRootMotionsDirty = true;
 						TSharedPtr<ISequencer> Sequencer = WeakSequencer.Pin();
 						if (Sequencer.IsValid())
@@ -658,8 +657,7 @@ bool FSkeletalAnimationTrackEditMode::InputDelta(FEditorViewportClient* InViewpo
 						UMovieSceneSkeletalAnimationTrack* Track = Data.SelectedSection.Get()->GetTypedOuter<UMovieSceneSkeletalAnimationTrack>();
 						Track->Modify();
 						Data.SelectedSection.Get()->Modify();
-						Data.SelectedSection.Get()->StartLocationOffset = CurrentTransform.GetLocation(); 
-						
+						Data.SelectedSection.Get()->StartLocationOffset = CurrentTransform.GetLocation();
 						Data.SelectedSection.Get()->GetRootMotionParams()->bRootMotionsDirty = true;
 						TSharedPtr<ISequencer> Sequencer = WeakSequencer.Pin();
 						if (Sequencer.IsValid())
diff --git a/Engine/Source/Editor/MovieSceneTools/Private/EditModes/SkeletalAnimationTrackEditMode.h b/Engine/Source/Editor/MovieSceneTools/Private/EditModes/SkeletalAnimationTrackEditMode.h
--- a/Engine/Source/Editor/MovieSceneTools/Private/EditModes/SkeletalAnimationTrackEditMode.h
+++ b/Engine/Source/Editor/MovieSceneTools/Private/EditModes/SkeletalAnimationTrackEditMode.h
@@ -62,9 +62,9 @@ public:
 protected:
 
 	bool IsSomethingSelected() const;
-	bool GetCurrentTransform(FTransform& OutWorld, FTransform& OutParent) const;
-	bool IsRootSelected(UMovieSceneSkeletalAnimationSection* Section);
-
+	bool GetTransformAtFirstSectionStart(FTransform& OutWorld, FTransform& OutParent) const;
+	bool IsRootSelected(UMovieSceneSkeletalAnimationSection* Section) const;
+	
 protected:
 	/** Interrogator that is an gc object */
 	UE::MovieScene::FSystemInterrogator InterrogationLinker;
diff --git a/Engine/Source/Runtime/MovieSceneTracks/Private/Sections/MovieSceneSkeletalAnimationSection.cpp b/Engine/Source/Runtime/MovieSceneTracks/Private/Sections/MovieSceneSkeletalAnimationSection.cpp
--- a/Engine/Source/Runtime/MovieSceneTracks/Private/Sections/MovieSceneSkeletalAnimationSection.cpp
+++ b/Engine/Source/Runtime/MovieSceneTracks/Private/Sections/MovieSceneSkeletalAnimationSection.cpp
@@ -472,15 +472,7 @@ int32 UMovieSceneSkeletalAnimationSection::SetBoneIndexForRootMotionCalculations
 				const int32 BoneTreeIndex = AnimationTrack.BoneTreeIndex;
 				if (BoneTreeIndex != INDEX_NONE)
 				{
-
-					int32 ParentIndex = AnimSequence->GetSkeleton()->GetReferenceSkeleton().GetParentIndex(BoneTreeIndex);
-					if (ParentIndex == INDEX_NONE)
-					{
-						RootIndex = TrackIndex;
-					}
-					//if we are parent of the root and we have animation
-					else if (ParentIndex == RootIndex)
-					{
+	
 						const FRawAnimSequenceTrack& InternalTrackData = AnimationTrack.InternalTrackData;
 						const TArray<FVector3f>& TrackPosKeys = InternalTrackData.PosKeys;
 						for (const FVector3f& Vector : TrackPosKeys)
@@ -491,23 +483,10 @@ int32 UMovieSceneSkeletalAnimationSection::SetBoneIndexForRootMotionCalculations
 								break;
 							}
 						}
-						if (TempRootBoneIndex.IsSet() == false)
-						{
-							const TArray<FQuat4f>& TrackRotKeys = InternalTrackData.RotKeys;
-							for (const FQuat4f& Quat : TrackRotKeys)
-							{
-								if (Quat.IsIdentity() == false)
-								{
-									TempRootBoneIndex = BoneTreeIndex;
-									break;
-								}
-							}
-						}
 						if (TempRootBoneIndex.IsSet())
 						{
 							break;
-						}
-					}
+						}				
 
 				}
 			}
@@ -557,9 +536,20 @@ PRAGMA_ENABLE_DEPRECATION_WARNINGS
 #endif
 		}
 	}
-	return TempRootBoneIndex.GetValue();
+	return TempRootBoneIndex.IsSet() ? TempRootBoneIndex.GetValue() : 0;
+}
+
+FTransform UMovieSceneSkeletalAnimationSection::GetRootMotionStartOffset() const
+{
+	UAnimSequence* AnimSequence = Cast<UAnimSequence>(Params.Animation);
+	if (AnimSequence && TempRootBoneIndex.IsSet() && TempRootBoneIndex.GetValue() != 0)
+	{
+		return AnimSequence->ExtractRootTrackTransform(0, nullptr);
+	}
+	return FTransform::Identity;;
 }
 
+
 bool UMovieSceneSkeletalAnimationSection::GetRootMotionTransform(FAnimationPoseData& AnimationPoseData,FRootMotionTransformParam& InOutParams) const
 {
 	UAnimSequence* AnimSequence = Cast<UAnimSequence>(Params.Animation);
@@ -577,7 +567,9 @@ bool UMovieSceneSkeletalAnimationSection::GetRootMotionTransform(FAnimationPoseD
 
 		InOutParams.bOutIsAdditive = AnimSequence->GetAdditiveAnimType() != EAdditiveAnimationType::AAT_None;
 		FTransform StartBoneTransform;
-		if (TempRootBoneIndex.IsSet())
+		InOutParams.OutRootStartTransform = GetRootMotionStartOffset();
+
+		if (TempRootBoneIndex.IsSet() && TempRootBoneIndex.GetValue() != 0)
 		{
 			//get the start pose first since we pass out the pose and need the current
 			FCompactPoseBoneIndex PoseIndex = AnimationPoseData.GetPose().GetBoneContainer().GetCompactPoseIndexFromSkeletonIndex(TempRootBoneIndex.GetValue());
@@ -588,19 +580,17 @@ bool UMovieSceneSkeletalAnimationSection::GetRootMotionTransform(FAnimationPoseD
 			ExtractionContext.CurrentTime = CurrentTimeSeconds;
 			AnimSequence->GetAnimationPose(AnimationPoseData, ExtractionContext);
 			InOutParams.OutPoseTransform = AnimationPoseData.GetPose()[FCompactPoseBoneIndex(PoseIndex)];
-
 		}
 		else //not set then just use root.
 		{
 			StartBoneTransform = AnimSequence->ExtractRootTrackTransform(StartSeconds, nullptr);
 			InOutParams.OutPoseTransform = AnimSequence->ExtractRootTrackTransform(CurrentTimeSeconds, nullptr);
-
 		}
 		//note though we don't support mesh space addtive just local additive it will still work the same here for the root so 
 		if (!InOutParams.bOutIsAdditive)
 		{
 			const bool bStartTransformOffsetInBoneSpace = CVarStartTransformOffsetInBoneSpace.GetValueOnGameThread();
-			if (bStartTransformOffsetInBoneSpace)
+			if (TempRootBoneIndex.IsSet() && TempRootBoneIndex.GetValue() != 0 && bStartTransformOffsetInBoneSpace)
 			{
 				FTransform StartMatchedInRoot = StartBoneTransform * MatchedTransform;
 				FTransform LocalToRoot = (InOutParams.OutPoseTransform * StartBoneTransform.Inverse());
diff --git a/Engine/Source/Runtime/MovieSceneTracks/Private/Tracks/MovieSceneSkeletalAnimationTrack.cpp b/Engine/Source/Runtime/MovieSceneTracks/Private/Tracks/MovieSceneSkeletalAnimationTrack.cpp
--- a/Engine/Source/Runtime/MovieSceneTracks/Private/Tracks/MovieSceneSkeletalAnimationTrack.cpp
+++ b/Engine/Source/Runtime/MovieSceneTracks/Private/Tracks/MovieSceneSkeletalAnimationTrack.cpp
@@ -490,6 +490,7 @@ void UMovieSceneSkeletalAnimationTrack::SetUpRootMotions(bool bForce)
 						if (AnimSection->StartLocationOffset.IsNearlyZero() == false || AnimSection->StartRotationOffset.IsNearlyZero() == false ||
 							AnimSection->MatchedLocationOffset.IsNearlyZero() == false || AnimSection->MatchedRotationOffset.IsNearlyZero() == false)
 						{ 
+							RootMotionParams.RootMotionStartOffset = AnimSection->GetRootMotionStartOffset();
 							bAnySectionsHaveOffset = true;
 							FMemMark Mark(FMemStack::Get());
 							FCompactPose OutPose;
@@ -664,10 +665,11 @@ static int32 GetAnimationTrackIndex(const int32 BoneIndex, const UAnimSequence*
 	return INDEX_NONE;
 }
 
-static FTransform GetWorldTransformForBone(UAnimSequence* AnimSequence, USkeletalMeshComponent* MeshComponent, const FName& InBoneName, double Seconds)
+static FTransform GetTransformForBoneRelativeToIndex(UAnimSequence* AnimSequence, USkeletalMeshComponent* MeshComponent, const FName& InBoneName,
+	const FCompactPoseBoneIndex& ParentCPIndex, double Seconds)
 {
-	FTransform  WorldTransform = FTransform::Identity;
 
+	FTransform  WorldTransform = FTransform::Identity;
 	//AnimSequence->GetBoneTransform doesn't seem to be as accurate as GetAnimationPose
 	FMemMark Mark(FMemStack::Get());
 	FCompactPose OutPose;
@@ -678,6 +680,7 @@ static FTransform GetWorldTransformForBone(UAnimSequence* AnimSequence, USkeleta
 	{
 		RequiredBoneIndexArray[BoneIndex] = BoneIndex;
 	}
+
 	FBoneContainer BoneContainer(RequiredBoneIndexArray, CurveEvalOption, *AnimSequence->GetSkeleton());
 	OutPose.ResetToRefPose(BoneContainer);
 	FBlendedCurve OutCurve;
@@ -690,7 +693,7 @@ static FTransform GetWorldTransformForBone(UAnimSequence* AnimSequence, USkeleta
 	if (MeshIndex != INDEX_NONE)
 	{
 		FCompactPoseBoneIndex CPIndex = AnimationPoseData.GetPose().GetBoneContainer().MakeCompactPoseIndex(FMeshPoseBoneIndex(MeshIndex));
-		if (CPIndex != INDEX_NONE)
+		if (CPIndex != INDEX_NONE )
 		{
 			FTransform BoneTransform = FTransform::Identity;
 			WorldTransform *= BoneTransform;
@@ -698,7 +701,14 @@ static FTransform GetWorldTransformForBone(UAnimSequence* AnimSequence, USkeleta
 			{
 				BoneTransform = AnimationPoseData.GetPose()[CPIndex];
 				WorldTransform *= BoneTransform;
-				CPIndex = AnimationPoseData.GetPose().GetBoneContainer().GetParentBoneIndex(CPIndex);
+				if (CPIndex == ParentCPIndex)  //if we are the parent then we stop
+				{
+					CPIndex = FCompactPoseBoneIndex(INDEX_NONE);
+				}
+				else
+				{
+					CPIndex = AnimationPoseData.GetPose().GetBoneContainer().GetParentBoneIndex(CPIndex);
+				}
 			} while (CPIndex.IsValid());
 		}
 	}
@@ -961,10 +971,13 @@ void UMovieSceneSkeletalAnimationTrack::MatchSectionByBoneTransform(bool bMatchW
 		if (FirstAnimSequence && SecondAnimSequence)
 		{
 			double FirstSectionTime = FirstSection->MapTimeToAnimation(CurrentFrame, FrameRate);
-			FTransform  FirstTransform = GetWorldTransformForBone(FirstAnimSequence, SkelMeshComp, BoneName, FirstSectionTime);
+			//use same index for all
+			int32 Index = CurrentSection->SetBoneIndexForRootMotionCalculations(bBlendFirstChildOfRoot);
+			FCompactPoseBoneIndex ParentIndex(Index);
+			
+			FTransform  FirstTransform = GetTransformForBoneRelativeToIndex(FirstAnimSequence, SkelMeshComp, BoneName, ParentIndex, FirstSectionTime);
 			double SecondSectionTime = CurrentSection->MapTimeToAnimation(CurrentFrame, FrameRate);
-			FTransform  SecondTransform = GetWorldTransformForBone(SecondAnimSequence, SkelMeshComp, BoneName, SecondSectionTime);
-
+			FTransform  SecondTransform = GetTransformForBoneRelativeToIndex(SecondAnimSequence, SkelMeshComp, BoneName, ParentIndex,SecondSectionTime);
 			//Need to match the translations and rotations here 
 			//First need to get the correct rotation order based upon what's matching, otherwise if not all are matched 
 			//and one rotation is set last we will get errors.
@@ -1024,10 +1037,13 @@ void UMovieSceneSkeletalAnimationTrack::MatchSectionByBoneTransform(bool bMatchW
 
 		if (FirstAnimSequence && SecondAnimSequence)
 		{
+			//use same index for all
+			int32 Index = CurrentSection->SetBoneIndexForRootMotionCalculations(bBlendFirstChildOfRoot);
+			FCompactPoseBoneIndex ParentIndex(Index);
 			float FirstSectionTime = static_cast<float>(CurrentSection->MapTimeToAnimation(CurrentFrame, FrameRate));
-			FTransform  FirstTransform = GetWorldTransformForBone(FirstAnimSequence, SkelMeshComp, BoneName, FirstSectionTime);
+			FTransform  FirstTransform = GetTransformForBoneRelativeToIndex(FirstAnimSequence, SkelMeshComp, BoneName,ParentIndex, FirstSectionTime);
 			float SecondSectionTime = static_cast<float>(SecondSection->MapTimeToAnimation(CurrentFrame, FrameRate));
-			FTransform  SecondTransform = GetWorldTransformForBone(SecondAnimSequence, SkelMeshComp, BoneName, SecondSectionTime);
+			FTransform  SecondTransform = GetTransformForBoneRelativeToIndex(SecondAnimSequence, SkelMeshComp, BoneName,ParentIndex, SecondSectionTime);
 
 			//Need to match the translations and rotations here 
 			//First need to get the correct rotation order based upon what's matching, otherwise if not all are matched 
diff --git a/Engine/Source/Runtime/MovieSceneTracks/Public/Sections/MovieSceneSkeletalAnimationSection.h b/Engine/Source/Runtime/MovieSceneTracks/Public/Sections/MovieSceneSkeletalAnimationSection.h
--- a/Engine/Source/Runtime/MovieSceneTracks/Public/Sections/MovieSceneSkeletalAnimationSection.h
+++ b/Engine/Source/Runtime/MovieSceneTracks/Public/Sections/MovieSceneSkeletalAnimationSection.h
@@ -224,6 +224,7 @@ public:
 		FTransform OutTransform; //total transform including current pose plus offset
 		FTransform OutParentTransform; //offset transform not including original bone transform
 		FTransform OutPoseTransform; //original bone transform
+		FTransform OutRootStartTransform;//start of the root
 		float OutWeight; //weight at specified time
 	};
 
@@ -254,8 +255,9 @@ public:
 	MOVIESCENETRACKS_API void ToggleShowSkeleton();
 #endif
 
+	MOVIESCENETRACKS_API FTransform GetRootMotionStartOffset() const;
+
 private:
 	void MultiplyOutInverseOnNextClips(FVector PreviousMatchedLocationOffset, FRotator PreviousMatchedRotationOffset);
 
-
 };
diff --git a/Engine/Source/Runtime/MovieSceneTracks/Public/Tracks/MovieSceneSkeletalAnimationTrack.h b/Engine/Source/Runtime/MovieSceneTracks/Public/Tracks/MovieSceneSkeletalAnimationTrack.h
--- a/Engine/Source/Runtime/MovieSceneTracks/Public/Tracks/MovieSceneSkeletalAnimationTrack.h
+++ b/Engine/Source/Runtime/MovieSceneTracks/Public/Tracks/MovieSceneSkeletalAnimationTrack.h
@@ -24,7 +24,7 @@ struct FMovieSceneSkeletalAnimRootMotionTrackParams
 	FFrameTime EndFrame;
 	bool bRootMotionsDirty;
 	TArray<FTransform> RootTransforms;
-
+	FTransform RootMotionStartOffset; // root motion may not be in Mesh Space if we are putting values on a bone that is a child of a root with an offset
 	/** Get the Root Motion transform at the specified time.*/
 	MOVIESCENETRACKS_API TOptional<FTransform>  GetRootMotion(FFrameTime CurrentTime) const;
 	FMovieSceneSkeletalAnimRootMotionTrackParams() : bRootMotionsDirty(true) {}
