Sequencer: Camera Cut Track section now has more information when a camera cut is not found - instead of just showing no label it will instead specify if there is no binding specified, or if there is a binding specified it will say that the binding/object cannot be found.

diff --git a/Engine/Source/Editor/EditorStyle/Private/StarshipStyle.cpp b/Engine/Source/Editor/EditorStyle/Private/StarshipStyle.cpp
--- a/Engine/Source/Editor/EditorStyle/Private/StarshipStyle.cpp
+++ b/Engine/Source/Editor/EditorStyle/Private/StarshipStyle.cpp
@@ -2032,6 +2032,7 @@ void FStarshipEditorStyle::FStyle::SetupSequencerStyles()
 		Set("Sequencer.Section.Background", new BORDER_BRUSH(TEXT("Sequencer/SectionBackground"), FMargin(4.f / 16.f)));
 		Set("Sequencer.Section.BackgroundTint", new BOX_BRUSH(TEXT("Sequencer/SectionBackgroundTint"), FMargin(4 / 16.f)));
 		Set("Sequencer.Section.CollapsedSelectedSectionOverlay", new IMAGE_BRUSH(TEXT("Sequencer/Section/CollapsedSelectedSectionOverlay"), Icon16x16, FLinearColor::White, ESlateBrushTileType::Both));
+		Set("Sequencer.Section.ErroredSectionOverlay", new BORDER_BRUSH(TEXT("Sequencer/Section/CollapsedSelectedSectionOverlay"), FMargin(4.f / 16.f)));
 		Set("Sequencer.Section.SectionHeaderSelectedSectionOverlay", new IMAGE_BRUSH(TEXT("Sequencer/Section/SectionHeaderSelectedSectionOverlay"), Icon16x16, FLinearColor::White, ESlateBrushTileType::Both));
 		Set("Sequencer.Section.SelectedTrackTint", new BOX_BRUSH(TEXT("Sequencer/SelectedTrackTint"), FMargin(0.f, 0.5f)));
 		Set("Sequencer.Section.SelectionBorder", new BORDER_BRUSH(TEXT("Sequencer/SectionHighlight"), FMargin(7.f / 16.f)));
diff --git a/Engine/Source/Editor/MovieSceneTools/Private/Sections/CameraCutSection.cpp b/Engine/Source/Editor/MovieSceneTools/Private/Sections/CameraCutSection.cpp
--- a/Engine/Source/Editor/MovieSceneTools/Private/Sections/CameraCutSection.cpp
+++ b/Engine/Source/Editor/MovieSceneTools/Private/Sections/CameraCutSection.cpp
@@ -174,9 +174,29 @@ FMargin FCameraCutSection::GetContentPadding() const
 
 int32 FCameraCutSection::OnPaintSection(FSequencerSectionPainter& InPainter) const
 {
-	static const FSlateBrush* FilmBorder = FAppStyle::GetBrush("Sequencer.Section.FilmBorder");
-
 	InPainter.LayerId = InPainter.PaintSectionBackground();
+
+	// Draw a red frame around the edges to indicate an error since we can't highlight the error text right now.
+	AActor* CameraActor = GetCameraForFrame(Section->GetInclusiveStartFrame());
+	if (!CameraActor)
+	{
+		static const FSlateBrush* ErroredSectionOverlay = FAppStyle::Get().GetBrush("Sequencer.Section.ErroredSectionOverlay");
+		const ESlateDrawEffect DrawEffects = InPainter.bParentEnabled
+			? ESlateDrawEffect::None
+			: ESlateDrawEffect::DisabledEffect;
+
+		FLinearColor ErrorColor = FLinearColor::Red;
+		FSlateDrawElement::MakeBox(
+			InPainter.DrawElements,
+			InPainter.LayerId,
+			InPainter.SectionGeometry.ToPaintGeometry(FVector2D(1, 1), InPainter.SectionGeometry.GetLocalSize() - FVector2D(1, 1)),
+			ErroredSectionOverlay,
+			DrawEffects,
+			ErrorColor.CopyWithNewOpacity(0.8f)
+		);
+		InPainter.LayerId++;
+	}
+
 	return FViewportThumbnailSection::OnPaintSection(InPainter);
 }
 
@@ -188,6 +208,19 @@ FText FCameraCutSection::HandleThumbnailTextBlockText() const
 		return FText::FromString(CameraActor->GetActorLabel());
 	}
 
+	UMovieSceneCameraCutSection* CameraCutSection = Cast<UMovieSceneCameraCutSection>(Section);
+	if (CameraCutSection)
+	{
+		if(!CameraCutSection->GetCameraBindingID().IsValid())
+		{
+			return LOCTEXT("CameraBindingError_NoBinding", "No Object Binding specified.");
+		}
+		else
+		{
+			return LOCTEXT("CameraBindingError_MissingBinding", "Object Binding / Bound Object is missing!");
+		}
+	}
+
 	return FText::GetEmpty();
 }
 
