Sequencer: Add SupportsSequence implementation so that the track editor is only active for sequences that support it

This fixes OnInitialize() getting called (which has a workaround for edit modes, unfortunately) when the day sequence is active.

Moved callbacks to within OnInitialize since they don't need to run unless the track editor is active.

Note, derived class of MovieSceneSequence that need the Skeletal Animation Track Editor, should implement IsTrackSupported.

diff --git a/Engine/Plugins/MovieScene/ActorSequence/Source/ActorSequence/Private/ActorSequence.cpp b/Engine/Plugins/MovieScene/ActorSequence/Source/ActorSequence/Private/ActorSequence.cpp
--- a/Engine/Plugins/MovieScene/ActorSequence/Source/ActorSequence/Private/ActorSequence.cpp
+++ b/Engine/Plugins/MovieScene/ActorSequence/Source/ActorSequence/Private/ActorSequence.cpp
@@ -13,6 +13,7 @@
 #include "Tracks/MovieSceneAudioTrack.h"
 #include "Tracks/MovieSceneEventTrack.h"
 #include "Tracks/MovieSceneMaterialParameterCollectionTrack.h"
+#include "Tracks/MovieSceneSkeletalAnimationTrack.h"
 
 IMPLEMENT_MODULE(FDefaultModuleImpl, ActorSequence);
 
@@ -220,7 +221,8 @@ ETrackSupport UActorSequence::IsTrackSupported(TSubclassOf<class UMovieSceneTrac
 {
 	if (InTrackClass == UMovieSceneAudioTrack::StaticClass() ||
 		InTrackClass == UMovieSceneEventTrack::StaticClass() ||
-		InTrackClass == UMovieSceneMaterialParameterCollectionTrack::StaticClass())
+		InTrackClass == UMovieSceneMaterialParameterCollectionTrack::StaticClass() ||
+		InTrackClass == UMovieSceneSkeletalAnimationTrack::StaticClass())
 	{
 		return ETrackSupport::Supported;
 	}
diff --git a/Engine/Plugins/MovieScene/TemplateSequence/Source/TemplateSequence/Private/TemplateSequence.cpp b/Engine/Plugins/MovieScene/TemplateSequence/Source/TemplateSequence/Private/TemplateSequence.cpp
--- a/Engine/Plugins/MovieScene/TemplateSequence/Source/TemplateSequence/Private/TemplateSequence.cpp
+++ b/Engine/Plugins/MovieScene/TemplateSequence/Source/TemplateSequence/Private/TemplateSequence.cpp
@@ -11,6 +11,7 @@
 #include "MovieScene.h"
 #include "MovieSceneCommonHelpers.h"
 #include "MovieSceneTrack.h"
+#include "Tracks/MovieSceneSkeletalAnimationTrack.h"
 #include "Tracks/MovieSceneSpawnTrack.h"
 
 #include "TemplateSequenceActor.h"
@@ -272,7 +273,8 @@ FText UTemplateSequence::GetDisplayName() const
 
 ETrackSupport UTemplateSequence::IsTrackSupported(TSubclassOf<class UMovieSceneTrack> InTrackClass) const
 {
-	if (InTrackClass == UMovieSceneSpawnTrack::StaticClass())
+	if (InTrackClass == UMovieSceneSkeletalAnimationTrack::StaticClass() ||
+		InTrackClass == UMovieSceneSpawnTrack::StaticClass())
 	{
 		return ETrackSupport::Supported;
 	}
diff --git a/Engine/Source/Editor/MovieSceneTools/Private/TrackEditors/SkeletalAnimationTrackEditor.cpp b/Engine/Source/Editor/MovieSceneTools/Private/TrackEditors/SkeletalAnimationTrackEditor.cpp
--- a/Engine/Source/Editor/MovieSceneTools/Private/TrackEditors/SkeletalAnimationTrackEditor.cpp
+++ b/Engine/Source/Editor/MovieSceneTools/Private/TrackEditors/SkeletalAnimationTrackEditor.cpp
@@ -1061,14 +1061,14 @@ FSkeletalAnimationTrackEditor::FSkeletalAnimationTrackEditor( TSharedRef<ISequen
 	//We use the FGCObject pattern to keep the anim export option alive during the editor session
 
 	AnimSeqExportOption = NewObject<UAnimSeqExportOption>();
-
-	SequencerSavedHandle = InSequencer->OnPostSave().AddRaw(this, &FSkeletalAnimationTrackEditor::OnSequencerSaved);
-	SequencerChangedHandle = InSequencer->OnMovieSceneDataChanged().AddRaw(this, &FSkeletalAnimationTrackEditor::OnSequencerDataChanged);
-	FCoreUObjectDelegates::OnObjectPropertyChanged.AddRaw(this, &FSkeletalAnimationTrackEditor::OnPostPropertyChanged);
 }
 
 void FSkeletalAnimationTrackEditor::OnInitialize()
 {
+	SequencerSavedHandle = GetSequencer()->OnPostSave().AddRaw(this, &FSkeletalAnimationTrackEditor::OnSequencerSaved);
+	SequencerChangedHandle = GetSequencer()->OnMovieSceneDataChanged().AddRaw(this, &FSkeletalAnimationTrackEditor::OnSequencerDataChanged);
+	FCoreUObjectDelegates::OnObjectPropertyChanged.AddRaw(this, &FSkeletalAnimationTrackEditor::OnPostPropertyChanged);
+
 	++FSkeletalAnimationTrackEditor::NumberActive;
 
 	// Activate the default mode in case FEditorModeTools::Tick isn't run before here. 
@@ -1123,6 +1123,11 @@ TSharedRef<ISequencerTrackEditor> FSkeletalAnimationTrackEditor::CreateTrackEdit
 	return MakeShareable( new FSkeletalAnimationTrackEditor( InSequencer ) );
 }
 
+bool FSkeletalAnimationTrackEditor::SupportsSequence(UMovieSceneSequence* InSequence) const
+{
+	ETrackSupport TrackSupported = InSequence ? InSequence->IsTrackSupported(UMovieSceneSkeletalAnimationTrack::StaticClass()) : ETrackSupport::NotSupported;
+	return TrackSupported == ETrackSupport::Supported;
+}
 
 bool FSkeletalAnimationTrackEditor::SupportsType( TSubclassOf<UMovieSceneTrack> Type ) const
 {
diff --git a/Engine/Source/Editor/MovieSceneTools/Private/TrackEditors/SkeletalAnimationTrackEditor.h b/Engine/Source/Editor/MovieSceneTools/Private/TrackEditors/SkeletalAnimationTrackEditor.h
--- a/Engine/Source/Editor/MovieSceneTools/Private/TrackEditors/SkeletalAnimationTrackEditor.h
+++ b/Engine/Source/Editor/MovieSceneTools/Private/TrackEditors/SkeletalAnimationTrackEditor.h
@@ -64,6 +64,7 @@ public:
 	virtual void BuildObjectBindingTrackMenu(FMenuBuilder& MenuBuilder, const TArray<FGuid>& ObjectBindings, const UClass* ObjectClass) override;
 	virtual bool HandleAssetAdded(UObject* Asset, const FGuid& TargetObjectGuid) override;
 	virtual TSharedRef<ISequencerSection> MakeSectionInterface( UMovieSceneSection& SectionObject, UMovieSceneTrack& Track, FGuid ObjectBinding ) override;
+	virtual bool SupportsSequence(UMovieSceneSequence* InSequence) const override;
 	virtual bool SupportsType( TSubclassOf<UMovieSceneTrack> Type ) const override;
 	virtual void BuildTrackContextMenu( FMenuBuilder& MenuBuilder, UMovieSceneTrack* Track ) override;
 	virtual TSharedPtr<SWidget> BuildOutlinerEditWidget(const FGuid& ObjectBinding, UMovieSceneTrack* Track, const FBuildEditWidgetParams& Params) override;
diff --git a/Engine/Source/Runtime/LevelSequence/Private/LevelSequence.cpp b/Engine/Source/Runtime/LevelSequence/Private/LevelSequence.cpp
--- a/Engine/Source/Runtime/LevelSequence/Private/LevelSequence.cpp
+++ b/Engine/Source/Runtime/LevelSequence/Private/LevelSequence.cpp
@@ -26,6 +26,7 @@
 #include "Tracks/MovieSceneLevelVisibilityTrack.h"
 #include "Tracks/MovieSceneDataLayerTrack.h"
 #include "Tracks/MovieSceneMaterialParameterCollectionTrack.h"
+#include "Tracks/MovieSceneSkeletalAnimationTrack.h"
 #include "Tracks/MovieSceneSlomoTrack.h"
 #include "Tracks/MovieSceneSpawnTrack.h"
 #include "Tracks/MovieSceneSubTrack.h"
@@ -127,6 +128,7 @@ ETrackSupport ULevelSequence::IsTrackSupported(TSubclassOf<class UMovieSceneTrac
 		InTrackClass == UMovieSceneLevelVisibilityTrack::StaticClass() ||
 		InTrackClass == UMovieSceneDataLayerTrack::StaticClass() ||
 		InTrackClass == UMovieSceneMaterialParameterCollectionTrack::StaticClass() ||
+		InTrackClass == UMovieSceneSkeletalAnimationTrack::StaticClass() ||
 		InTrackClass == UMovieSceneSlomoTrack::StaticClass() ||
 		InTrackClass == UMovieSceneSpawnTrack::StaticClass() ||
 		InTrackClass == UMovieSceneSubTrack::StaticClass() ||
