Sequencer: Fix rename widget padding so that it fits at the film border edge when renaming.

#rnx

diff --git a/Engine/Source/Editor/MovieSceneTools/Private/Sections/CinematicShotSection.cpp b/Engine/Source/Editor/MovieSceneTools/Private/Sections/CinematicShotSection.cpp
--- a/Engine/Source/Editor/MovieSceneTools/Private/Sections/CinematicShotSection.cpp
+++ b/Engine/Source/Editor/MovieSceneTools/Private/Sections/CinematicShotSection.cpp
@@ -110,7 +110,8 @@ float FCinematicShotSection::GetSectionHeight() const
 
 FMargin FCinematicShotSection::GetContentPadding() const
 {
-	return FMargin(8.f, 15.f);
+	// When the rename widget is visible, use less padding so that the widget is visible over the film border (when thumbnails are not shown)
+	return FMargin(8.f, GetRenameVisibility() == EVisibility::Visible ? 10.f : 15.f);
 }
 
 void FCinematicShotSection::SetSingleTime(double GlobalTime)
diff --git a/Engine/Source/Editor/MovieSceneTools/Private/Sections/ThumbnailSection.cpp b/Engine/Source/Editor/MovieSceneTools/Private/Sections/ThumbnailSection.cpp
--- a/Engine/Source/Editor/MovieSceneTools/Private/Sections/ThumbnailSection.cpp
+++ b/Engine/Source/Editor/MovieSceneTools/Private/Sections/ThumbnailSection.cpp
@@ -76,7 +76,7 @@ TSharedRef<SWidget> FThumbnailSection::GenerateSectionWidget()
 	return SNew(SBox)
 		.HAlign(HAlign_Left)
 		.VAlign(VAlign_Top)
-		.Padding(GetContentPadding())
+		.Padding_Lambda([&]() { return GetContentPadding(); })
 		[
 			SAssignNew(NameWidget, SInlineEditableTextBlock)
 				.ToolTipText(CanRename() ? LOCTEXT("RenameThumbnail", "Click or hit F2 to rename") : FText::GetEmpty())
