Sequencer: Fix shot parsing so that if the current prefix is found at the beginning of the shot name, it will be stripped and the parsing of the shot number and take number will continue from there.

This fixes a situation where shot prefixes can collide with parsing shot and take numbers because they are looking for the same separators. For example if the shot name is: sq050_shot0010_01, that should be parsed as shot 10 and take number 1. If the user specifies that the shot prefix is sq050, that portion can be stripped so it doesn't erroneously parse as shot 50.

diff --git a/Engine/Source/Editor/MovieSceneTools/Private/MovieSceneToolHelpers.cpp b/Engine/Source/Editor/MovieSceneTools/Private/MovieSceneToolHelpers.cpp
--- a/Engine/Source/Editor/MovieSceneTools/Private/MovieSceneToolHelpers.cpp
+++ b/Engine/Source/Editor/MovieSceneTools/Private/MovieSceneToolHelpers.cpp
@@ -226,7 +226,7 @@ void MovieSceneToolHelpers::SplitSection(const TSet<TWeakObjectPtr<UMovieSceneSe
 	}
 }
 
-bool MovieSceneToolHelpers::ParseShotName(const FString& ShotName, FString& ShotPrefix, uint32& ShotNumber, uint32& TakeNumber, uint32& ShotNumberDigits, uint32& TakeNumberDigits)
+bool MovieSceneToolHelpers::ParseShotName(const FString& InShotName, FString& ShotPrefix, uint32& ShotNumber, uint32& TakeNumber, uint32& ShotNumberDigits, uint32& TakeNumberDigits)
 {
 	// Parse a shot name
 	// 
@@ -240,6 +240,16 @@ bool MovieSceneToolHelpers::ParseShotName(const FString& ShotName, FString& Shot
 	//  ShotNumber = 20
 	//  TakeNumber = 2
 	//
+
+	// If the shot prefix is known and is already at the beginning of the shot name, remove it and parse the rest
+	//
+	// sq050_sh0010_01
+	//  ShotPrefix = sq050_sh, ie. parse(0010_01)
+	//  ShotNumber = 10
+	//  TakeNumber = 1
+	FString ShotName = InShotName;
+	ShotName.RemoveFromStart(ShotPrefix);
+
 	const UMovieSceneToolsProjectSettings* ProjectSettings = GetDefault<UMovieSceneToolsProjectSettings>();
 
 	uint32 FirstShotNumberIndex = INDEX_NONE;
@@ -299,7 +309,11 @@ bool MovieSceneToolHelpers::ParseShotName(const FString& ShotName, FString& Shot
 
 	if (FirstShotNumberIndex != INDEX_NONE)
 	{
-		ShotPrefix = ShotName.Left(FirstShotNumberIndex);
+		if (FirstShotNumberIndex != 0)
+		{
+			ShotPrefix = ShotName.Left(FirstShotNumberIndex);
+		}
+
 		ShotNumber = FCString::Atoi(*ShotName.Mid(FirstShotNumberIndex, LastShotNumberIndex-FirstShotNumberIndex+1));
 	}
 
@@ -323,7 +337,7 @@ bool MovieSceneToolHelpers::ParseShotName(const FString& ShotName, FString& Shot
 			ShotPrefix = ShotName.Left(LastSlashPos);
 			ShotNumber = INDEX_NONE; // Nullify the shot number since we only have a shot prefix
 			TakeNumber = FCString::Atoi(*ShotName.RightChop(LastSlashPos+1));
-			TakeNumberDigits = LastSlashPos+1;
+			TakeNumberDigits = ShotName.Len() - (LastSlashPos+1);
 			return true;
 		}
 	}
