Sequencer: Fix shot parsing so that if the current prefix is found at the beginning of the shot name, it will be stripped and the parsing of the shot number and take number will continue from there.

This fixes a situation where shot prefixes can collide with parsing shot and take numbers because they are looking for the same separators. For example if the shot name is: sq050_shot0010_01, that should be parsed as shot 10 and take number 1. If the user specifies that the shot prefix is sq050, that portion can be stripped so it doesn't erroneously parse as shot 50.

diff --git a/Engine/Source/Editor/MovieSceneTools/Private/MovieSceneToolHelpers.cpp b/Engine/Source/Editor/MovieSceneTools/Private/MovieSceneToolHelpers.cpp
--- a/Engine/Source/Editor/MovieSceneTools/Private/MovieSceneToolHelpers.cpp
+++ b/Engine/Source/Editor/MovieSceneTools/Private/MovieSceneToolHelpers.cpp
@@ -334,7 +334,11 @@ bool MovieSceneToolHelpers::ParseShotName(const FString& InShotName, FString& Sh
 		int32 LastSlashPos = ShotName.Find(ProjectSettings->TakeSeparator, ESearchCase::IgnoreCase, ESearchDir::FromEnd);
 		if (LastSlashPos != INDEX_NONE)
 		{
-			ShotPrefix = ShotName.Left(LastSlashPos);
+			if (LastSlashPos != 0)
+			{
+				ShotPrefix = ShotName.Left(LastSlashPos);
+			}
+			
 			ShotNumber = INDEX_NONE; // Nullify the shot number since we only have a shot prefix
 			TakeNumber = FCString::Atoi(*ShotName.RightChop(LastSlashPos+1));
 			TakeNumberDigits = ShotName.Len() - (LastSlashPos+1);
