[UE-170238] Blend Curves: New Blend Curve UI is missing the standard "Cubic" blend.

diff --git a/Engine/Source/Editor/MovieSceneTools/Private/MovieSceneAlphaBlendOptionCustomization.cpp b/Engine/Source/Editor/MovieSceneTools/Private/MovieSceneAlphaBlendOptionCustomization.cpp
--- a/Engine/Source/Editor/MovieSceneTools/Private/MovieSceneAlphaBlendOptionCustomization.cpp
+++ b/Engine/Source/Editor/MovieSceneTools/Private/MovieSceneAlphaBlendOptionCustomization.cpp
@@ -69,7 +69,6 @@ void FAlphaBlendPropertyCustomization::CustomizeChildren(TSharedRef<IPropertyHan
 	FilterExclude.Add(EMovieSceneBuiltInEasing::QuadIn);
 	FilterExclude.Add(EMovieSceneBuiltInEasing::QuadOut);
 	FilterExclude.Add(EMovieSceneBuiltInEasing::CubicOut);
-	FilterExclude.Add(EMovieSceneBuiltInEasing::CubicInOut);
 	FilterExclude.Add(EMovieSceneBuiltInEasing::QuartIn);
 	FilterExclude.Add(EMovieSceneBuiltInEasing::QuartOut);
 	FilterExclude.Add(EMovieSceneBuiltInEasing::QuintIn);
