[UE-170238] Blend Curves: New Blend Curve UI is missing the standard "Cubic" blend

Additional fix: Turns out the cubic and cubic in-out are two different ways of cubic interpolations both with easing in and out.

diff --git a/Engine/Source/Editor/MovieSceneTools/Private/MovieSceneAlphaBlendOptionCustomization.cpp b/Engine/Source/Editor/MovieSceneTools/Private/MovieSceneAlphaBlendOptionCustomization.cpp
--- a/Engine/Source/Editor/MovieSceneTools/Private/MovieSceneAlphaBlendOptionCustomization.cpp
+++ b/Engine/Source/Editor/MovieSceneTools/Private/MovieSceneAlphaBlendOptionCustomization.cpp
@@ -68,6 +68,7 @@ void FAlphaBlendPropertyCustomization::CustomizeChildren(TSharedRef<IPropertyHan
 	FilterExclude.Add(EMovieSceneBuiltInEasing::SinOut);
 	FilterExclude.Add(EMovieSceneBuiltInEasing::QuadIn);
 	FilterExclude.Add(EMovieSceneBuiltInEasing::QuadOut);
+	FilterExclude.Add(EMovieSceneBuiltInEasing::CubicIn);
 	FilterExclude.Add(EMovieSceneBuiltInEasing::CubicOut);
 	FilterExclude.Add(EMovieSceneBuiltInEasing::QuartIn);
 	FilterExclude.Add(EMovieSceneBuiltInEasing::QuartOut);
@@ -115,11 +116,11 @@ EMovieSceneBuiltInEasing FAlphaBlendPropertyCustomization::ConvertToEasingType(E
 	switch (in)
 	{
 	case EAlphaBlendOption::Linear: return EMovieSceneBuiltInEasing::Linear;
-	case EAlphaBlendOption::Cubic: return EMovieSceneBuiltInEasing::CubicIn;
+	case EAlphaBlendOption::Cubic: return EMovieSceneBuiltInEasing::Cubic;
+	case EAlphaBlendOption::CubicInOut: return EMovieSceneBuiltInEasing::CubicInOut;
 	case EAlphaBlendOption::HermiteCubic: return EMovieSceneBuiltInEasing::HermiteCubicInOut;
 	case EAlphaBlendOption::Sinusoidal: return EMovieSceneBuiltInEasing::SinInOut;
 	case EAlphaBlendOption::QuadraticInOut: return EMovieSceneBuiltInEasing::QuadInOut;
-	case EAlphaBlendOption::CubicInOut: return EMovieSceneBuiltInEasing::CubicInOut;
 	case EAlphaBlendOption::QuarticInOut: return EMovieSceneBuiltInEasing::QuartInOut;
 	case EAlphaBlendOption::QuinticInOut: return EMovieSceneBuiltInEasing::QuintInOut;
 	case EAlphaBlendOption::CircularIn: return EMovieSceneBuiltInEasing::CircIn;
@@ -139,11 +140,11 @@ EAlphaBlendOption FAlphaBlendPropertyCustomization::ConvertFromEasingType(EMovie
 	switch (in)
 	{
 	case EMovieSceneBuiltInEasing::Linear: return EAlphaBlendOption::Linear;
-	case EMovieSceneBuiltInEasing::CubicIn: return EAlphaBlendOption::Cubic;
+	case EMovieSceneBuiltInEasing::Cubic: return EAlphaBlendOption::Cubic;
+	case EMovieSceneBuiltInEasing::CubicInOut: return EAlphaBlendOption::CubicInOut;
 	case EMovieSceneBuiltInEasing::HermiteCubicInOut: return EAlphaBlendOption::HermiteCubic;
 	case EMovieSceneBuiltInEasing::SinInOut: return EAlphaBlendOption::Sinusoidal;
 	case EMovieSceneBuiltInEasing::QuadInOut: return EAlphaBlendOption::QuadraticInOut;
-	case EMovieSceneBuiltInEasing::CubicInOut: return EAlphaBlendOption::CubicInOut;
 	case EMovieSceneBuiltInEasing::QuartInOut: return EAlphaBlendOption::QuarticInOut;
 	case EMovieSceneBuiltInEasing::QuintInOut: return EAlphaBlendOption::QuinticInOut;
 	case EMovieSceneBuiltInEasing::CircIn: return EAlphaBlendOption::CircularIn;
diff --git a/Engine/Source/Editor/MovieSceneTools/Private/MovieSceneBuiltInEasingFunctionCustomization.cpp b/Engine/Source/Editor/MovieSceneTools/Private/MovieSceneBuiltInEasingFunctionCustomization.cpp
--- a/Engine/Source/Editor/MovieSceneTools/Private/MovieSceneBuiltInEasingFunctionCustomization.cpp
+++ b/Engine/Source/Editor/MovieSceneTools/Private/MovieSceneBuiltInEasingFunctionCustomization.cpp
@@ -58,6 +58,7 @@ void FMovieSceneBuiltInEasingFunctionCustomization::CustomizeDetails(IDetailLayo
 
 	TSet<EMovieSceneBuiltInEasing> FilterExclude;
 	FilterExclude.Add(EMovieSceneBuiltInEasing::Custom);
+	FilterExclude.Add(EMovieSceneBuiltInEasing::Cubic);
 
 	IDetailCategoryBuilder& Category = DetailBuilder.EditCategory("Easing");
 	DetailBuilder.HideProperty(TypeProperty);
diff --git a/Engine/Source/Editor/MovieSceneTools/Private/MovieSceneBuiltInEasingFunctionGridWidget.cpp b/Engine/Source/Editor/MovieSceneTools/Private/MovieSceneBuiltInEasingFunctionGridWidget.cpp
--- a/Engine/Source/Editor/MovieSceneTools/Private/MovieSceneBuiltInEasingFunctionGridWidget.cpp
+++ b/Engine/Source/Editor/MovieSceneTools/Private/MovieSceneBuiltInEasingFunctionGridWidget.cpp
@@ -115,13 +115,14 @@ void SEasingFunctionGridWidget::Construct(const FArguments& InArgs)
 
 			Grid->AddSlot(ColumnIndex, RowIndex)
 				[
-					SNew(SBox)
-					.WidthOverride(100.f)
-					.HeightOverride(50.f)
+					SNew(SButton)
+					.HAlign(HAlign_Center)
+					.ButtonStyle(FAppStyle::Get(), "HoverHintOnly")
+					.OnClicked(this, &SEasingFunctionGridWidget::OnTypeButtonClicked, Value)
 					[
-						SNew(SButton)
-						.ButtonStyle(FAppStyle::Get(), "HoverHintOnly")
-						.OnClicked(this, &SEasingFunctionGridWidget::OnTypeButtonClicked, Value)
+						SNew(SBox)
+						.WidthOverride(100.f)
+						.HeightOverride(50.f)
 						[
 							SNew(SBuiltInFunctionVisualizer, Value)
 						]
diff --git a/Engine/Source/Editor/MovieSceneTools/Private/MovieSceneBuiltInEasingFunctionVisualizer.cpp b/Engine/Source/Editor/MovieSceneTools/Private/MovieSceneBuiltInEasingFunctionVisualizer.cpp
--- a/Engine/Source/Editor/MovieSceneTools/Private/MovieSceneBuiltInEasingFunctionVisualizer.cpp
+++ b/Engine/Source/Editor/MovieSceneTools/Private/MovieSceneBuiltInEasingFunctionVisualizer.cpp
@@ -65,7 +65,7 @@ void SBuiltInFunctionVisualizer::SetType(EMovieSceneBuiltInEasing InValue)
 	while (Interp <= 1.f)
 	{
 		Samples.Add(FVector2D(Interp, DefaultObject->Evaluate(Interp)));
-		Interp += 0.05f;
+		Interp += 0.01f;
 	}
 
 	DefaultObject->Type = DefaultType;
diff --git a/Engine/Source/Runtime/MovieScene/Private/Generators/MovieSceneEasingCurves.cpp b/Engine/Source/Runtime/MovieScene/Private/Generators/MovieSceneEasingCurves.cpp
--- a/Engine/Source/Runtime/MovieScene/Private/Generators/MovieSceneEasingCurves.cpp
+++ b/Engine/Source/Runtime/MovieScene/Private/Generators/MovieSceneEasingCurves.cpp
@@ -70,6 +70,7 @@ float UMovieSceneBuiltInEasingFunction::Evaluate(float Interp) const
 	case EMovieSceneBuiltInEasing::QuadOut: 	return PowOut(Interp,2);
 	case EMovieSceneBuiltInEasing::QuadInOut: 	return InTime < 1.f ? .5f*PowIn(InTime,2) : .5f + .5f*PowOut(OutTime,2);
 
+	case EMovieSceneBuiltInEasing::Cubic: 		return FMath::Clamp<float>(FMath::CubicInterp<float>(0.f, 0.f, 1.f, 0.f, Interp), 0.f, 1.f);
 	case EMovieSceneBuiltInEasing::CubicIn: 	return PowIn(Interp,3);
 	case EMovieSceneBuiltInEasing::CubicOut: 	return PowOut(Interp,3);
 	case EMovieSceneBuiltInEasing::CubicInOut: 	return InTime < 1.f ? .5f*PowIn(InTime,3) : .5f + .5f*PowOut(OutTime,3);
diff --git a/Engine/Source/Runtime/MovieScene/Public/Generators/MovieSceneEasingCurves.h b/Engine/Source/Runtime/MovieScene/Public/Generators/MovieSceneEasingCurves.h
--- a/Engine/Source/Runtime/MovieScene/Public/Generators/MovieSceneEasingCurves.h
+++ b/Engine/Source/Runtime/MovieScene/Public/Generators/MovieSceneEasingCurves.h
@@ -27,7 +27,7 @@ enum class EMovieSceneBuiltInEasing : uint8
 	// Quadratic easing
 	QuadIn UMETA(Grouping=Quadratic,DisplayName="Quadratic In"), QuadOut UMETA(Grouping=Quadratic,DisplayName="Quadratic Out"), QuadInOut UMETA(Grouping=Quadratic,DisplayName="Quadratic InOut"),
 	// Cubic easing
-	CubicIn UMETA(Grouping=Cubic,DisplayName="Cubic In"), CubicOut UMETA(Grouping=Cubic,DisplayName="Cubic Out"), CubicInOut UMETA(Grouping=Cubic,DisplayName="Cubic InOut"), HermiteCubicInOut UMETA(Grouping = Cubic, DisplayName = "Hermite-Cubic InOut"),
+	Cubic UMETA(Grouping = Cubic, DisplayName = "Cubic"), CubicIn UMETA(Grouping=Cubic,DisplayName="Cubic In"), CubicOut UMETA(Grouping=Cubic,DisplayName="Cubic Out"), CubicInOut UMETA(Grouping=Cubic,DisplayName="Cubic InOut"), HermiteCubicInOut UMETA(Grouping = Cubic, DisplayName = "Hermite-Cubic InOut"),
 	// Quartic easing
 	QuartIn UMETA(Grouping=Quartic,DisplayName="Quartic In"), QuartOut UMETA(Grouping=Quartic,DisplayName="Quartic Out"), QuartInOut UMETA(Grouping=Quartic,DisplayName="Quartic InOut"),
 	// Quintic easing
