Sequencer: fix incorrectly mirrored perlin noise curve preview

Slate's Y axis goes downwards...

diff --git a/Engine/Source/Editor/MovieSceneTools/Public/Channels/PerlinNoiseChannelInterface.h b/Engine/Source/Editor/MovieSceneTools/Public/Channels/PerlinNoiseChannelInterface.h
--- a/Engine/Source/Editor/MovieSceneTools/Public/Channels/PerlinNoiseChannelInterface.h
+++ b/Engine/Source/Editor/MovieSceneTools/Public/Channels/PerlinNoiseChannelInterface.h
@@ -167,7 +167,7 @@ struct TPerlinNoiseChannelInterface : ISequencerChannelInterface
 		{
 			double Seconds = PaintArgs.TimeToPixel.PixelToSeconds(X);
 			double Y = (double)TypedChannel->Evaluate(Seconds);
-			CurvePoints.Add(FVector2D(X, (Y * YScale) + YOffset));
+			CurvePoints.Add(FVector2D(X, YOffset - (Y * YScale)));
 		}
 
 		FSlateDrawElement::MakeLines(
