avoid registering transactions when evaluating sequencer during bake

in the context of spawnables, we don't need to register transaction, like selection change when destructing the actor for instance, otherwise it results in keeping useless references to that actor

diff --git a/Engine/Source/Editor/MovieSceneTools/Private/MovieSceneConstraintChannelHelper.cpp b/Engine/Source/Editor/MovieSceneTools/Private/MovieSceneConstraintChannelHelper.cpp
--- a/Engine/Source/Editor/MovieSceneTools/Private/MovieSceneConstraintChannelHelper.cpp
+++ b/Engine/Source/Editor/MovieSceneTools/Private/MovieSceneConstraintChannelHelper.cpp
@@ -229,6 +229,9 @@ void FCompensationEvaluator::ComputeLocalTransformsForBaking(UWorld* InWorld, co
 	const FConstraintsManagerController& Controller = FConstraintsManagerController::Get(InWorld);
 	static constexpr bool bSorted = true;
 	const TArray<ConstraintPtr> AllConstraints = Controller.GetAllConstraints(bSorted);
+
+	// avoid transacting when evaluating sequencer
+	TGuardValue<ITransaction*> TransactionGuard(GUndo, nullptr);
 	
 	const TArray<IMovieSceneToolsAnimationBakeHelper*>& BakeHelpers = FMovieSceneToolsModule::Get().GetAnimationBakeHelpers();
 	for (IMovieSceneToolsAnimationBakeHelper* BakeHelper : BakeHelpers)
