Sequencer: Sort camera names in Camera Cut Section menu

diff --git a/Engine/Source/Editor/MovieSceneTools/Private/Sections/CameraCutSection.cpp b/Engine/Source/Editor/MovieSceneTools/Private/Sections/CameraCutSection.cpp
--- a/Engine/Source/Editor/MovieSceneTools/Private/Sections/CameraCutSection.cpp
+++ b/Engine/Source/Editor/MovieSceneTools/Private/Sections/CameraCutSection.cpp
@@ -113,6 +113,8 @@ void FCameraCutSection::BuildSectionContextMenu(FMenuBuilder& MenuBuilder, const
 		return;
 	}
 
+	AllCameras.Sort([](const AActor& A, const AActor& B) { return A.GetActorLabel().Compare(B.GetActorLabel()) < 0; });
+
 	MenuBuilder.BeginSection(NAME_None, LOCTEXT("ChangeCameraMenuText", "Change Camera"));
 	{
 		for (auto EachCamera : AllCameras)
