Fix issue with movie scene quick binding not maintaining the default values of the template graph node

diff --git a/Engine/Source/Editor/MovieSceneTools/Private/MovieSceneEventCustomization.cpp b/Engine/Source/Editor/MovieSceneTools/Private/MovieSceneEventCustomization.cpp
--- a/Engine/Source/Editor/MovieSceneTools/Private/MovieSceneEventCustomization.cpp
+++ b/Engine/Source/Editor/MovieSceneTools/Private/MovieSceneEventCustomization.cpp
@@ -1243,7 +1243,8 @@ void FMovieSceneEventCustomization::SetEventEndpoint(UK2Node* NewEndpoint, UEdGr
 		}
 	}
 
-	TArray<FName> PayloadNames;
+	// Map of the Payload Variable Names to their Default Values as Strings
+	TMap<FName, FString> PayloadVariables;
 	if (PayloadTemplate && EnumHasAnyFlags(AutoCreatePayload, EAutoCreatePayload::Variables | EAutoCreatePayload::Pins))
 	{
 		UK2Node_EditablePinBase* EditableNode = Cast<UK2Node_EditablePinBase>(NewEndpoint);
@@ -1255,7 +1256,7 @@ void FMovieSceneEventCustomization::SetEventEndpoint(UK2Node* NewEndpoint, UEdGr
 				// Make a payload variable for this pin
 				if (EnumHasAnyFlags(AutoCreatePayload, EAutoCreatePayload::Variables))
 				{
-					PayloadNames.Add(PayloadPin->PinName);
+					PayloadVariables.Add(PayloadPin->PinName, PayloadPin->DefaultValue);
 				}
 
 				if (EditableNode && EnumHasAnyFlags(AutoCreatePayload, EAutoCreatePayload::Pins))
@@ -1285,11 +1286,14 @@ void FMovieSceneEventCustomization::SetEventEndpoint(UK2Node* NewEndpoint, UEdGr
 		{
 			FMovieSceneEventUtils::SetEndpoint(EntryPoint, EventSection, NewEndpoint, BoundObjectPin);
 
-			for (FName PayloadVar : PayloadNames)
+			for (const TPair<FName, FString>& PayloadVar : PayloadVariables)
 			{
-				if (!EntryPoint->PayloadVariables.Contains(PayloadVar))
+				if (!EntryPoint->PayloadVariables.Contains(PayloadVar.Key))
 				{
-					EntryPoint->PayloadVariables.Add(PayloadVar);
+					FMovieSceneEventPayloadVariable PayloadValue;
+					PayloadValue.Value = PayloadVar.Value;
+
+					EntryPoint->PayloadVariables.Add(PayloadVar.Key, MoveTemp(PayloadValue));
 				}
 			}
 		}
