Constraints: Clear out delegates from Sections.

diff --git a/Engine/Source/Editor/MovieSceneTools/Private/TrackEditors/TransformTrackEditor.cpp b/Engine/Source/Editor/MovieSceneTools/Private/TrackEditors/TransformTrackEditor.cpp
--- a/Engine/Source/Editor/MovieSceneTools/Private/TrackEditors/TransformTrackEditor.cpp
+++ b/Engine/Source/Editor/MovieSceneTools/Private/TrackEditors/TransformTrackEditor.cpp
@@ -83,6 +83,8 @@ F3DTransformTrackEditor::~F3DTransformTrackEditor()
 		GEditor->UnregisterForUndo(this);
 	}
 }
+//for 5.2 we will move this over to the header
+static TSet<TWeakObjectPtr<UMovieScene3DTransformSection>> SectionsToClear;
 
 void F3DTransformTrackEditor::OnRelease()
 {
@@ -1437,6 +1439,10 @@ void F3DTransformTrackEditor::HandleOnConstraintAdded(IMovieSceneConstrainedSect
 					}
 				});
 	}
+	// Store Section so we can remove these delegates
+	UMovieScene3DTransformSection* Section = Cast<UMovieScene3DTransformSection>(InSection);
+	SectionsToClear.Add(Section);
+
 	// handle key moved
 	if (!InConstraintChannel->OnKeyMovedEvent().IsBound())
 	{
@@ -1522,7 +1528,6 @@ void F3DTransformTrackEditor::HandleConstraintKeyMoved(IMovieSceneConstrainedSec
 	{
 		return;
 	}
-
 	UMovieSceneSection* Section = Cast<UMovieSceneSection>(InSection);
 
 	for (const FKeyMoveEventItem& MoveEventItem : InMovedItems)
@@ -1599,5 +1604,25 @@ void F3DTransformTrackEditor::ClearOutConstraintDelegates()
 		}
 	}
 	ConstraintHandlesToClear.Reset();
+
+	for (TWeakObjectPtr<UMovieScene3DTransformSection>& Section : SectionsToClear)
+	{
+		if (IMovieSceneConstrainedSection* CRSection = Section.Get())
+		{
+			// clear constraint channels
+			TArray<FConstraintAndActiveChannel>& ConstraintChannels = CRSection->GetConstraintsChannels();
+			for (FConstraintAndActiveChannel& Channel : ConstraintChannels)
+			{
+				Channel.ActiveChannel.OnKeyMovedEvent().Clear();
+				Channel.ActiveChannel.OnKeyDeletedEvent().Clear();
+			}
+
+			if (CRSection->OnConstraintRemovedHandle.IsValid())
+			{
+				CRSection->OnConstraintRemovedHandle.Reset();
+			}
+		}
+	}
+	SectionsToClear.Reset();
 }
 #undef LOCTEXT_NAMESPACE
