Sequencer: Check whether a sequence is allowed to be dropped

#rnx

diff --git a/Engine/Source/Editor/MovieSceneTools/Private/TrackEditors/SubTrackEditor.cpp b/Engine/Source/Editor/MovieSceneTools/Private/TrackEditors/SubTrackEditor.cpp
--- a/Engine/Source/Editor/MovieSceneTools/Private/TrackEditors/SubTrackEditor.cpp
+++ b/Engine/Source/Editor/MovieSceneTools/Private/TrackEditors/SubTrackEditor.cpp
@@ -351,7 +351,8 @@ bool FSubTrackEditor::OnAllowDrop(const FDragDropEvent& DragDropEvent, FSequence
 			continue;
 		}
 
-		if (UMovieSceneSequence* Sequence = Cast<UMovieSceneSequence>(AssetData.GetAsset()))
+		UMovieSceneSequence* Sequence = Cast<UMovieSceneSequence>(AssetData.GetAsset());
+		if (CanAddSubSequence(*Sequence))
 		{
 			FFrameRate TickResolution = SequencerPtr->GetFocusedTickResolution();
 
@@ -410,8 +411,7 @@ FReply FSubTrackEditor::OnDrop(const FDragDropEvent& DragDropEvent, const FSeque
 		}
 
 		UMovieSceneSequence* Sequence = Cast<UMovieSceneSequence>(AssetData.GetAsset());
-
-		if (Sequence)
+		if (CanAddSubSequence(*Sequence))
 		{
 			AnimatablePropertyChanged(FOnKeyProperty::CreateRaw(this, &FSubTrackEditor::HandleSequenceAdded, Sequence, DragDropParams.Track.Get(), DragDropParams.RowIndex));
 
