Sequencer: Improve context menu operations (Pre and Post infinity Extrapolation, Interpolation and Reduce Keys) so that they affected the selected channels (if selected) rather than all channels of the section.

Fixed a bug with the Interpolation and Reduce Keys menu operating on all channels of the section rather than just the found key areas. This seems like a regression possibly from MVVM refactor.

diff --git a/Engine/Source/Editor/MovieSceneTools/Private/Channels/BuiltInChannelEditors.cpp b/Engine/Source/Editor/MovieSceneTools/Private/Channels/BuiltInChannelEditors.cpp
--- a/Engine/Source/Editor/MovieSceneTools/Private/Channels/BuiltInChannelEditors.cpp
+++ b/Engine/Source/Editor/MovieSceneTools/Private/Channels/BuiltInChannelEditors.cpp
@@ -1402,8 +1402,64 @@ struct FCurveChannelSectionMenuExtension : TSharedFromThis<FCurveChannelSectionM
 		);
 	}
 
+	void GetChannels(TArray<FMovieSceneFloatChannel*>& FloatChannels, TArray<FMovieSceneDoubleChannel*>& DoubleChannels) const
+	{
+		ISequencer* Sequencer = WeakSequencer.Pin().Get();
+		if (!Sequencer)
+		{
+			return;
+		}
+
+		// Get selected channels
+		TArray<const IKeyArea*> KeyAreas;
+		Sequencer->GetSelectedKeyAreas(KeyAreas);
+		for (const IKeyArea* KeyArea : KeyAreas)
+		{
+			FMovieSceneChannelHandle Handle = KeyArea->GetChannel();
+			if (Handle.GetChannelTypeName() == FMovieSceneFloatChannel::StaticStruct()->GetFName())
+			{
+				FMovieSceneFloatChannel* Channel = static_cast<FMovieSceneFloatChannel*>(Handle.Get());
+				FloatChannels.Add(Channel);
+			}
+			else if (Handle.GetChannelTypeName() == FMovieSceneDoubleChannel::StaticStruct()->GetFName())
+			{
+				FMovieSceneDoubleChannel* Channel = static_cast<FMovieSceneDoubleChannel*>(Handle.Get());
+				DoubleChannels.Add(Channel);
+			}
+		}
+
+		// Otherwise, the channels of all the sections
+		if (FloatChannels.Num() + DoubleChannels.Num() == 0)
+		{
+			for (TWeakObjectPtr<UMovieSceneSection> WeakSection : Sections)
+			{
+				if (UMovieSceneSection* Section = WeakSection.Get())
+				{
+					FMovieSceneChannelProxy& ChannelProxy = Section->GetChannelProxy();
+					for (FMovieSceneFloatChannel* Channel : ChannelProxy.GetChannels<FMovieSceneFloatChannel>())
+					{
+						FloatChannels.Add(Channel);
+					}
+					for (FMovieSceneDoubleChannel* Channel : ChannelProxy.GetChannels<FMovieSceneDoubleChannel>())
+					{
+						DoubleChannels.Add(Channel);
+					}
+				}
+			}
+		}
+	}
+
 	void SetExtrapolationMode(ERichCurveExtrapolation ExtrapMode, bool bPreInfinity)
 	{
+		TArray<FMovieSceneFloatChannel*> FloatChannels;
+		TArray<FMovieSceneDoubleChannel*> DoubleChannels;
+		GetChannels(FloatChannels, DoubleChannels);
+
+		if (FloatChannels.Num() + DoubleChannels.Num() == 0)
+		{
+			return;
+		}
+
 		FScopedTransaction Transaction(LOCTEXT("SetExtrapolationMode_Transaction", "Set Extrapolation Mode"));
 
 		bool bAnythingChanged = false;
@@ -1418,24 +1474,17 @@ struct FCurveChannelSectionMenuExtension : TSharedFromThis<FCurveChannelSectionM
 		}
 
 		// Apply to all channels
-		for (TWeakObjectPtr<UMovieSceneSection> WeakSection : Sections)
+		for (FMovieSceneFloatChannel* Channel : FloatChannels)
 		{
-			if (UMovieSceneSection* Section = WeakSection.Get())
-			{
-				FMovieSceneChannelProxy& ChannelProxy = Section->GetChannelProxy();
-				for (FMovieSceneFloatChannel* Channel : ChannelProxy.GetChannels<FMovieSceneFloatChannel>())
-				{
-					TEnumAsByte<ERichCurveExtrapolation>& DestExtrap = bPreInfinity ? Channel->PreInfinityExtrap : Channel->PostInfinityExtrap;
-					DestExtrap = ExtrapMode;
-					bAnythingChanged = true;
-				}
-				for (FMovieSceneDoubleChannel* Channel : ChannelProxy.GetChannels<FMovieSceneDoubleChannel>())
-				{
-					TEnumAsByte<ERichCurveExtrapolation>& DestExtrap = bPreInfinity ? Channel->PreInfinityExtrap : Channel->PostInfinityExtrap;
-					DestExtrap = ExtrapMode;
-					bAnythingChanged = true;
-				}
-			}
+			TEnumAsByte<ERichCurveExtrapolation>& DestExtrap = bPreInfinity ? Channel->PreInfinityExtrap : Channel->PostInfinityExtrap;
+			DestExtrap = ExtrapMode;
+			bAnythingChanged = true;
+		}
+		for (FMovieSceneDoubleChannel* Channel : DoubleChannels)
+		{
+			TEnumAsByte<ERichCurveExtrapolation>& DestExtrap = bPreInfinity ? Channel->PreInfinityExtrap : Channel->PostInfinityExtrap;
+			DestExtrap = ExtrapMode;
+			bAnythingChanged = true;
 		}
 
 		if (bAnythingChanged)
@@ -1453,27 +1502,24 @@ struct FCurveChannelSectionMenuExtension : TSharedFromThis<FCurveChannelSectionM
 
 	bool IsExtrapolationModeSelected(ERichCurveExtrapolation ExtrapMode, bool bPreInfinity) const
 	{
-		for (TWeakObjectPtr<UMovieSceneSection> WeakSection : Sections)
+		TArray<FMovieSceneFloatChannel*> FloatChannels;
+		TArray<FMovieSceneDoubleChannel*> DoubleChannels;
+		GetChannels(FloatChannels, DoubleChannels);
+
+		for (FMovieSceneFloatChannel* Channel : FloatChannels)
 		{
-			if (UMovieSceneSection* Section = WeakSection.Get())
+			ERichCurveExtrapolation SourceExtrap = bPreInfinity ? Channel->PreInfinityExtrap : Channel->PostInfinityExtrap;
+			if (SourceExtrap != ExtrapMode)
 			{
-				FMovieSceneChannelProxy& ChannelProxy = Section->GetChannelProxy();
-				for (FMovieSceneFloatChannel* Channel : ChannelProxy.GetChannels<FMovieSceneFloatChannel>())
-				{
-					ERichCurveExtrapolation SourceExtrap = bPreInfinity ? Channel->PreInfinityExtrap : Channel->PostInfinityExtrap;
-					if (SourceExtrap != ExtrapMode)
-					{
-						return false;
-					}
-				}
-				for (FMovieSceneDoubleChannel* Channel : ChannelProxy.GetChannels<FMovieSceneDoubleChannel>())
-				{
-					ERichCurveExtrapolation SourceExtrap = bPreInfinity ? Channel->PreInfinityExtrap : Channel->PostInfinityExtrap;
-					if (SourceExtrap != ExtrapMode)
-					{
-						return false;
-					}
-				}
+				return false;
+			}
+		}
+		for (FMovieSceneDoubleChannel* Channel : DoubleChannels)
+		{
+			ERichCurveExtrapolation SourceExtrap = bPreInfinity ? Channel->PreInfinityExtrap : Channel->PostInfinityExtrap;
+			if (SourceExtrap != ExtrapMode)
+			{
+				return false;
 			}
 		}
 
diff --git a/Engine/Source/Editor/Sequencer/Private/SequencerContextMenus.cpp b/Engine/Source/Editor/Sequencer/Private/SequencerContextMenus.cpp
--- a/Engine/Source/Editor/Sequencer/Private/SequencerContextMenus.cpp
+++ b/Engine/Source/Editor/Sequencer/Private/SequencerContextMenus.cpp
@@ -784,12 +784,9 @@ void FSectionContextMenu::ReduceKeys()
 			{
 				Section->Modify();
 
-				for (const FMovieSceneChannelEntry& Entry : Section->GetChannelProxy().GetAllEntries())
+				if (FMovieSceneChannel* Channel = KeyArea->GetChannel().Get())
 				{
-					for (FMovieSceneChannel* Channel : Entry.GetChannels())
-					{
-						Channel->Optimize(Params);
-					}
+					Channel->Optimize(Params);
 				}
 			}
 		}
@@ -862,8 +859,10 @@ void FSectionContextMenu::SetInterpTangentMode(ERichCurveInterpMode InterpMode,
 			{
 				Section->Modify();
 
-				for (FMovieSceneFloatChannel* FloatChannel : Section->GetChannelProxy().GetChannels<FMovieSceneFloatChannel>())
+				FMovieSceneChannelHandle Handle = KeyArea->GetChannel();
+				if (Handle.GetChannelTypeName() == FMovieSceneFloatChannel::StaticStruct()->GetFName())
 				{
+					FMovieSceneFloatChannel* FloatChannel = static_cast<FMovieSceneFloatChannel*>(Handle.Get());
 					TMovieSceneChannelData<FMovieSceneFloatValue> ChannelData = FloatChannel->GetData();
 					TArrayView<FMovieSceneFloatValue> Values = ChannelData.GetValues();
 
@@ -876,9 +875,9 @@ void FSectionContextMenu::SetInterpTangentMode(ERichCurveInterpMode InterpMode,
 
 					FloatChannel->AutoSetTangents();
 				}
-
-				for (FMovieSceneDoubleChannel* DoubleChannel : Section->GetChannelProxy().GetChannels<FMovieSceneDoubleChannel>())
+				else if (Handle.GetChannelTypeName() == FMovieSceneDoubleChannel::StaticStruct()->GetFName())
 				{
+					FMovieSceneDoubleChannel* DoubleChannel = static_cast<FMovieSceneDoubleChannel*>(Handle.Get());
 					TMovieSceneChannelData<FMovieSceneDoubleValue> ChannelData = DoubleChannel->GetData();
 					TArrayView<FMovieSceneDoubleValue> Values = ChannelData.GetValues();
 
@@ -923,12 +922,9 @@ bool FSectionContextMenu::CanSetInterpTangentMode() const
 	{
 		if (KeyArea.IsValid())
 		{
-			UMovieSceneSection* Section = KeyArea->GetOwningSection();
-			if (Section)
-			{
-				return (Section->GetChannelProxy().GetChannels<FMovieSceneFloatChannel>().Num() != 0 ||
-					    Section->GetChannelProxy().GetChannels<FMovieSceneDoubleChannel>().Num() != 0);
-			}
+			FMovieSceneChannelHandle Handle = KeyArea->GetChannel();
+			return (Handle.GetChannelTypeName() == FMovieSceneFloatChannel::StaticStruct()->GetFName() ||
+					Handle.GetChannelTypeName() == FMovieSceneDoubleChannel::StaticStruct()->GetFName());
 		}
 	}
 
