new animation channel interface
    the constraint channel related code now lives outside of the CR plugin making it easier to reuse and maintain

diff --git a/Engine/Plugins/Animation/ControlRig/Source/ControlRigEditor/Private/ControlRigEditorModule.cpp b/Engine/Plugins/Animation/ControlRig/Source/ControlRigEditor/Private/ControlRigEditorModule.cpp
--- a/Engine/Plugins/Animation/ControlRig/Source/ControlRigEditor/Private/ControlRigEditorModule.cpp
+++ b/Engine/Plugins/Animation/ControlRig/Source/ControlRigEditor/Private/ControlRigEditorModule.cpp
@@ -124,6 +124,7 @@
 #include "UserDefinedStructure/UserDefinedStructEditorData.h"
 #include "UObject/FieldIterator.h"
 #include "AnimationToolMenuContext.h"
+#include "ControlConstraintChannelInterface.h"
 #include "RigVMModel/Nodes/RigVMAggregateNode.h"
 #include "RigVMUserWorkflowRegistry.h"
 #include "UserDefinedStructureCompilerUtils.h"
@@ -131,6 +132,8 @@
 #include "Units/ControlRigNodeWorkflow.h"
 #include "Units/RigDispatchFactory.h"
 #include "RigVMFunctions/Math/RigVMMathLibrary.h"
+#include "Constraints/ControlRigTransformableHandle.h"
+#include "Constraints/TransformConstraintChannelInterface.h"
 
 #define LOCTEXT_NAMESPACE "ControlRigEditorModule"
 
@@ -210,6 +213,10 @@ void FControlRigEditorModule::StartupModule()
 	SequencerModule.RegisterChannelInterface<FMovieSceneControlRigSpaceChannel>();
 	ControlRigParameterTrackCreateEditorHandle = SequencerModule.RegisterTrackEditor(FOnCreateTrackEditor::CreateStatic(&FControlRigParameterTrackEditor::CreateTrackEditor));
 
+	// register UTransformableControlHandle animatable interface
+	FConstraintChannelInterfaceRegistry& ConstraintChannelInterfaceRegistry = FConstraintChannelInterfaceRegistry::Get();
+	ConstraintChannelInterfaceRegistry.RegisterConstraintChannelInterface<UTransformableControlHandle>(MakeUnique<FControlConstraintChannelInterface>());
+	
 	AddControlRigExtenderToToolMenu("AssetEditor.AnimationEditor.ToolBar");
 
 	FEditorModeRegistry::Get().RegisterMode<FControlRigEditMode>(
diff --git a/Engine/Plugins/Animation/ControlRig/Source/ControlRigEditor/Private/ControlRigSequencerEditorLibrary.cpp b/Engine/Plugins/Animation/ControlRig/Source/ControlRigEditor/Private/ControlRigSequencerEditorLibrary.cpp
--- a/Engine/Plugins/Animation/ControlRig/Source/ControlRigEditor/Private/ControlRigSequencerEditorLibrary.cpp
+++ b/Engine/Plugins/Animation/ControlRig/Source/ControlRigEditor/Private/ControlRigSequencerEditorLibrary.cpp
@@ -45,7 +45,7 @@
 #include "EditorScriptingHelpers.h"
 #include "ConstraintsScripting.h"
 #include "Constraints/ControlRigTransformableHandle.h"
-#include "ConstraintChannelHelper.h"
+#include "Constraints/MovieSceneConstraintChannelHelper.h"
 #include "Sections/MovieSceneConstrainedSection.h"
 
 #include UE_INLINE_GENERATED_CPP_BY_NAME(ControlRigSequencerEditorLibrary)
@@ -469,7 +469,11 @@ UTickableConstraint* UControlRigSequencerEditorLibrary::AddConstraint(UWorld* Wo
 	}
 	
 	//add key
-	FConstraintChannelHelper::SmartConstraintKey(Constraint, TOptional<bool>(), TOptional<FFrameNumber>());
+	const TWeakPtr<ISequencer> WeakSequencer = GetSequencerFromAsset();
+	if (WeakSequencer.IsValid())
+	{
+		FMovieSceneConstraintChannelHelper::SmartConstraintKey(WeakSequencer.Pin(), Constraint, TOptional<bool>(), TOptional<FFrameNumber>());
+	}
 	return Constraint;
 }
 
@@ -509,7 +513,7 @@ bool UControlRigSequencerEditorLibrary::Compensate(UTickableConstraint* InConstr
 			InTime = FFrameRate::TransformTime(FFrameTime(InTime, 0), DisplayRate, TickResolution).RoundToFrame();
 		}
 		TOptional<FFrameNumber> OptTime(InTime);
-		FConstraintChannelHelper::Compensate(Constraint, OptTime);
+		FMovieSceneConstraintChannelHelper::Compensate(Sequencer, Constraint, OptTime);
 		return true;
 	}
 	else
@@ -529,7 +533,7 @@ bool UControlRigSequencerEditorLibrary::CompensateAll(UTickableConstraint* InCon
 	}
 	if (UTickableTransformConstraint* Constraint = Cast<UTickableTransformConstraint>(InConstraint))
 	{
-		FConstraintChannelHelper::Compensate(Constraint, TOptional<FFrameNumber>());
+		FMovieSceneConstraintChannelHelper::Compensate(WeakSequencer.Pin(), Constraint, TOptional<FFrameNumber>());
 		return true;
 	}
 	else
@@ -558,7 +562,7 @@ bool UControlRigSequencerEditorLibrary::SetConstraintActiveKey(UTickableConstrai
 		{
 			InTime = FFrameRate::TransformTime(FFrameTime(InTime, 0), DisplayRate, TickResolution).RoundToFrame();
 		}
-		return FConstraintChannelHelper::SmartConstraintKey(Constraint, TOptional<bool>(bActive), TOptional<FFrameNumber>(InTime));
+		return FMovieSceneConstraintChannelHelper::SmartConstraintKey(Sequencer, Constraint, TOptional<bool>(bActive), TOptional<FFrameNumber>(InTime));
 	}
 	else
 	{
diff --git a/Engine/Plugins/Animation/ControlRig/Source/ControlRigEditor/Private/EditMode/ControlRigEditMode.cpp b/Engine/Plugins/Animation/ControlRig/Source/ControlRigEditor/Private/EditMode/ControlRigEditMode.cpp
--- a/Engine/Plugins/Animation/ControlRig/Source/ControlRigEditor/Private/EditMode/ControlRigEditMode.cpp
+++ b/Engine/Plugins/Animation/ControlRig/Source/ControlRigEditor/Private/EditMode/ControlRigEditMode.cpp
@@ -61,7 +61,7 @@
 #include "LevelEditor.h"
 #include "InteractiveToolManager.h"
 #include "EdModeInteractiveToolsContext.h"
-#include "MovieSceneConstraintChannelHelper.h"
+#include "Constraints/MovieSceneConstraintChannelHelper.h"
 #include "Editor/EditorPerProjectUserSettings.h"
 #include "TransformConstraint.h"
 
diff --git a/Engine/Plugins/Animation/ControlRig/Source/ControlRigEditor/Private/Editor/Constraints/SConstraintsWidget.cpp b/Engine/Plugins/Animation/ControlRig/Source/ControlRigEditor/Private/Editor/Constraints/SConstraintsWidget.cpp
--- a/Engine/Plugins/Animation/ControlRig/Source/ControlRigEditor/Private/Editor/Constraints/SConstraintsWidget.cpp
+++ b/Engine/Plugins/Animation/ControlRig/Source/ControlRigEditor/Private/Editor/Constraints/SConstraintsWidget.cpp
@@ -3,7 +3,7 @@
 #include "SConstraintsWidget.h"
 
 #include "ActorPickerMode.h"
-#include "ConstraintChannelHelper.h"
+#include "Constraints/MovieSceneConstraintChannelHelper.h"
 #include "ConstraintsActor.h"
 #include "ControlRigEditorStyle.h"
 #include "DetailLayoutBuilder.h"
@@ -108,6 +108,23 @@ TArray<AActor*> GetCurrentSelection()
 		return IsValid(Actor) && Actor->IsSelected();
 	});	
 }
+
+TWeakPtr<ISequencer> GetSequencerChecked()
+{
+	const TWeakPtr<ISequencer> WeakSequencer = FBakingHelper::GetSequencer();
+	if (!WeakSequencer.IsValid() || !WeakSequencer.Pin()->GetFocusedMovieSceneSequence())
+	{
+		return nullptr;
+	}
+	
+	const UMovieScene* MovieScene = WeakSequencer.Pin()->GetFocusedMovieSceneSequence()->GetMovieScene();
+	if (!MovieScene)
+	{
+		return nullptr;
+	}
+
+	return WeakSequencer;
+}
 	
 }
 
@@ -301,8 +318,10 @@ void SDroppableConstraintItem::CreateConstraint(
 		return Component->HasAnySockets();
 	});
 
+	const TWeakPtr<ISequencer> WeakSequencer = GetSequencerChecked();
+	
 	// create constraints
-	auto CreateConstraint = [InCreationDelegate, InConstraintType](
+	auto CreateConstraint = [InCreationDelegate, InConstraintType, WeakSequencer](
 		const TArray<AActor*>& Selection, UObject* InParent, const FName& InSocketName)
 	{
 		UWorld* World = GetCurrentWorld();
@@ -321,8 +340,11 @@ void SDroppableConstraintItem::CreateConstraint(
 					FTransformConstraintUtils::CreateAndAddFromObjects(World, InParent, InSocketName, Child, NAME_None, InConstraintType);
 				if (Constraint)
 				{
-					FConstraintChannelHelper::SmartConstraintKey(Constraint, TOptional<bool>(), TOptional<FFrameNumber>());
 					bCreated = true;
+					if (WeakSequencer.IsValid())
+					{
+						FMovieSceneConstraintChannelHelper::SmartConstraintKey(WeakSequencer.Pin(), Constraint, TOptional<bool>(), TOptional<FFrameNumber>());
+					}
 				}
 			}
 		}
@@ -476,6 +498,9 @@ void SEditableConstraintItem::Construct(
 	};
 	UTickableConstraint* Constraint = GetConstraint();
 
+	// sequencer
+	TWeakPtr<ISequencer> WeakSequencer = GetSequencerChecked();
+	
 	// labels
 	FString ParentLabel(TEXT("undefined")), ChildLabel(TEXT("undefined"));
 	if (!InItem->Label.IsEmpty())
@@ -581,26 +606,32 @@ void SEditableConstraintItem::Construct(
 			SNew(SButton)
 			.ButtonStyle(FAppStyle::Get(), TEXT("SimpleButton"))
 			.ContentPadding(0)
-			.OnClicked_Lambda(	[Constraint]()
+			.OnClicked_Lambda(	[Constraint, WeakSequencer]()
 			{
 				if (!IsValid(Constraint))
 				{
 					return FReply::Handled();
 				}
+
+				if (!WeakSequencer.IsValid())
+				{
+					return FReply::Handled();
+				}
+				
 				if (UTickableTransformConstraint* TransformConstraint = Cast<UTickableTransformConstraint>(Constraint))
 				{
 					FScopedTransaction Transaction(LOCTEXT("CreateConstraintKey", "Create Constraint Key"));
-					FConstraintChannelHelper::SmartConstraintKey(TransformConstraint,TOptional<bool>(), TOptional<FFrameNumber>());
+					FMovieSceneConstraintChannelHelper::SmartConstraintKey(WeakSequencer.Pin(), TransformConstraint,TOptional<bool>(), TOptional<FFrameNumber>());
 				}
 				return FReply::Handled();
 			})
-			.IsEnabled_Lambda([]()
+			.IsEnabled_Lambda([WeakSequencer]()
 			{
-				return FConstraintChannelHelper::IsKeyframingAvailable();
+				return WeakSequencer.IsValid();
 			})
-			.Visibility_Lambda([]()
+			.Visibility_Lambda([WeakSequencer]()
 			{
-				const bool bIsKeyframingAvailable = FConstraintChannelHelper::IsKeyframingAvailable();
+				const bool bIsKeyframingAvailable = WeakSequencer.IsValid();
 				return bIsKeyframingAvailable ? EVisibility::Visible : EVisibility::Hidden;
 			})
 			.ToolTipText(LOCTEXT("KeyConstraintToolTip", "Add an active keyframe for that constraint."))
@@ -1042,14 +1073,8 @@ TSharedPtr<SWidget> SConstraintsEditionWidget::CreateContextMenu()
 				{
 					if (UTickableTransformConstraint* TransformConstraint = Cast<UTickableTransformConstraint>(Constraint))
 					{
-						const TWeakPtr<ISequencer> WeakSequencer = FBakingHelper::GetSequencer();
-						if (!WeakSequencer.IsValid() || !WeakSequencer.Pin()->GetFocusedMovieSceneSequence())
-						{
-							return;
-						}
-						const TSharedPtr<ISequencer> Sequencer = WeakSequencer.Pin();
-						const UMovieScene* MovieScene = Sequencer->GetFocusedMovieSceneSequence()->GetMovieScene();
-						if (!MovieScene)
+						const TWeakPtr<ISequencer> WeakSequencer = GetSequencerChecked();
+						if (!WeakSequencer.IsValid())
 						{
 							return;
 						}
@@ -1070,7 +1095,8 @@ TSharedPtr<SWidget> SConstraintsEditionWidget::CreateContextMenu()
 			return TransformConstraint->bDynamicOffset;
 		});
 
-		if (!bIsLookAtConstraint)
+		TWeakPtr<ISequencer> WeakSequencer = GetSequencerChecked();
+		if (!bIsLookAtConstraint && WeakSequencer.IsValid())
 		{
 			MenuBuilder.BeginSection("KeyConstraint", LOCTEXT("KeyConstraintHeader", "Keys"));
 			{
@@ -1078,17 +1104,13 @@ TSharedPtr<SWidget> SConstraintsEditionWidget::CreateContextMenu()
 				LOCTEXT("CompensateKeyLabel", "Compensate Key"),
 				FText::Format(LOCTEXT("CompensateKeyTooltip", "Compensate transform key for {0}."), ConstraintLabel),
 				FSlateIcon(),
-				FUIAction(FExecuteAction::CreateLambda([TransformConstraint]()
+				FUIAction(FExecuteAction::CreateLambda([TransformConstraint, WeakSequencer]()
 				{
-					const TWeakPtr<ISequencer> WeakSequencer = FBakingHelper::GetSequencer();
-					if (WeakSequencer.IsValid())
-					{
-						const TSharedPtr<ISequencer> Sequencer = WeakSequencer.Pin();
-						const FFrameRate TickResolution = Sequencer->GetFocusedTickResolution();
-						const FFrameTime FrameTime = Sequencer->GetLocalTime().ConvertTo(TickResolution);
-						const FFrameNumber Time = FrameTime.GetFrame();
-						FConstraintChannelHelper::Compensate(TransformConstraint, TOptional<FFrameNumber>(Time));
-					}
+					const TSharedPtr<ISequencer> Sequencer = WeakSequencer.Pin();
+					const FFrameRate TickResolution = Sequencer->GetFocusedTickResolution();
+					const FFrameTime FrameTime = Sequencer->GetLocalTime().ConvertTo(TickResolution);
+					const FFrameNumber Time = FrameTime.GetFrame();
+					FMovieSceneConstraintChannelHelper::Compensate(WeakSequencer.Pin(), TransformConstraint, TOptional<FFrameNumber>(Time));
 				}), IsCompensationEnabled),
 				NAME_None,
 				EUserInterfaceActionType::Button);
@@ -1097,9 +1119,9 @@ TSharedPtr<SWidget> SConstraintsEditionWidget::CreateContextMenu()
 				LOCTEXT("CompensateAllKeysLabel", "Compensate All Keys"),
 				FText::Format(LOCTEXT("CompensateAllKeysTooltip", "Compensate all transform keys for {0}."), ConstraintLabel),
 				FSlateIcon(),
-				FUIAction(FExecuteAction::CreateLambda([TransformConstraint]()
+				FUIAction(FExecuteAction::CreateLambda([TransformConstraint, WeakSequencer]()
 				{
-					FConstraintChannelHelper::Compensate(TransformConstraint, TOptional<FFrameNumber>());
+					FMovieSceneConstraintChannelHelper::Compensate(WeakSequencer.Pin(), TransformConstraint, TOptional<FFrameNumber>());
 				}), IsCompensationEnabled),
 				NAME_None,
 				EUserInterfaceActionType::Button);
diff --git a/Engine/Plugins/Animation/ControlRig/Source/ControlRigEditor/Private/Sequencer/ConstraintChannelHelper.cpp b/Engine/Plugins/Animation/ControlRig/Source/ControlRigEditor/Private/Sequencer/ConstraintChannelHelper.cpp
deleted file mode 100644
diff --git a/Engine/Plugins/Animation/ControlRig/Source/ControlRigEditor/Private/Sequencer/ControlConstraintChannelInterface.cpp b/Engine/Plugins/Animation/ControlRig/Source/ControlRigEditor/Private/Sequencer/ControlConstraintChannelInterface.cpp
new file mode 100644
--- /dev/null
+++ b/Engine/Plugins/Animation/ControlRig/Source/ControlRigEditor/Private/Sequencer/ControlConstraintChannelInterface.cpp
@@ -0,0 +1,274 @@
+﻿// Copyright Epic Games, Inc. All Rights Reserved.
+
+#include "ControlConstraintChannelInterface.h"
+
+#include "ControlRigSpaceChannelEditors.h"
+#include "ISequencer.h"
+
+#include "Constraints/ControlRigTransformableHandle.h"
+#include "Sequencer/MovieSceneControlRigParameterTrack.h"
+#include "LevelSequence.h"
+#include "MovieSceneToolHelpers.h"
+
+#include "TransformConstraint.h"
+#include "Algo/Copy.h"
+
+#include "ScopedTransaction.h"
+#include "Channels/MovieSceneFloatChannel.h"
+#include "Constraints/MovieSceneConstraintChannelHelper.inl"
+
+#define LOCTEXT_NAMESPACE "Constraints"
+
+namespace
+{
+	UMovieScene* GetMovieScene(const TSharedPtr<ISequencer>& InSequencer)
+	{
+		const UMovieSceneSequence* MovieSceneSequence = InSequencer ? InSequencer->GetFocusedMovieSceneSequence() : nullptr;
+		return MovieSceneSequence ? MovieSceneSequence->GetMovieScene() : nullptr;
+	}
+}
+
+UMovieSceneSection* FControlConstraintChannelInterface::GetHandleSection(
+	const UTransformableHandle* InHandle,
+	const TSharedPtr<ISequencer>& InSequencer)
+{
+	if (!IsValid(InHandle) || !InHandle->IsValid())
+	{
+		return nullptr;
+	}
+	
+	const UTransformableControlHandle* ControlHandle = static_cast<const UTransformableControlHandle*>(InHandle);
+	static constexpr bool bConstraintSection = false;
+	return GetControlSection(ControlHandle, InSequencer, bConstraintSection);
+}
+
+UMovieSceneSection* FControlConstraintChannelInterface::GetHandleConstraintSection(const UTransformableHandle* InHandle, const TSharedPtr<ISequencer>& InSequencer)
+{
+	if (!IsValid(InHandle) || !InHandle->IsValid())
+	{
+		return nullptr;
+	}
+
+	const UTransformableControlHandle* ControlHandle = static_cast<const UTransformableControlHandle*>(InHandle);
+	static constexpr bool bConstraintSection = true;
+	return GetControlSection(ControlHandle, InSequencer, bConstraintSection);
+}
+
+UWorld* FControlConstraintChannelInterface::GetHandleWorld(UTransformableHandle* InHandle)
+{
+	if (!IsValid(InHandle) || !InHandle->IsValid())
+	{
+		return nullptr;
+	}
+	
+	const UTransformableControlHandle* ControlHandle = static_cast<UTransformableControlHandle*>(InHandle);
+	const UControlRig* ControlRig = ControlHandle->ControlRig.LoadSynchronous();
+
+	return ControlRig ? ControlRig->GetWorld() : nullptr;
+}
+
+bool FControlConstraintChannelInterface::SmartConstraintKey(
+	UTickableTransformConstraint* InConstraint, const TOptional<bool>& InOptActive,
+	const FFrameNumber& InTime, const TSharedPtr<ISequencer>& InSequencer)
+{
+	const UTransformableControlHandle* ControlHandle = Cast<UTransformableControlHandle>(InConstraint->ChildTRSHandle);
+	if (!ControlHandle)
+	{
+		return false;
+	}
+
+	UMovieSceneControlRigParameterSection* ConstraintSection = GetControlSection(ControlHandle, InSequencer, true);
+	UMovieSceneControlRigParameterSection* TransformSection = GetControlSection(ControlHandle, InSequencer, false);
+	if ((!ConstraintSection) || (!TransformSection))
+	{
+		return false;
+	}
+	
+	FScopedTransaction Transaction(LOCTEXT("KeyConstraintaKehy", "Key Constraint Key"));
+	ConstraintSection->Modify();
+	TransformSection->Modify();
+
+	// set constraint as dynamic
+	InConstraint->bDynamicOffset = true;
+	
+	// add the channel
+	ConstraintSection->AddConstraintChannel(InConstraint);
+
+	// add key if needed
+	if (FConstraintAndActiveChannel* Channel = ConstraintSection->GetConstraintChannel(InConstraint->GetFName()))
+	{
+		bool ActiveValueToBeSet = false;
+		//add key if we can and make sure the key we are setting is what we want
+		if (CanAddKey(Channel->ActiveChannel, InTime, ActiveValueToBeSet) && (InOptActive.IsSet() == false || InOptActive.GetValue() == ActiveValueToBeSet))
+		{
+			const FFrameRate TickResolution = InSequencer->GetFocusedTickResolution();
+
+			const bool bNeedsCompensation = InConstraint->NeedsCompensation();
+				
+			TGuardValue<bool> CompensateGuard(FMovieSceneConstraintChannelHelper::bDoNotCompensate, true);
+				
+			UControlRig* ControlRig = ControlHandle->ControlRig.Get();
+			const FName& ControlName = ControlHandle->ControlName;
+				
+			// store the frames to compensate
+			const TArrayView<FMovieSceneFloatChannel*> Channels = ControlHandle->GetFloatChannels(TransformSection);
+			TArray<FFrameNumber> FramesToCompensate;
+			if (bNeedsCompensation)
+			{
+				FMovieSceneConstraintChannelHelper::GetFramesToCompensate<FMovieSceneFloatChannel>(Channel->ActiveChannel, ActiveValueToBeSet, InTime, Channels, FramesToCompensate);
+			}
+			else
+			{
+				FramesToCompensate.Add(InTime);
+			}
+
+
+			// store child and space transforms for these frames
+			FCompensationEvaluator Evaluator(InConstraint);
+			Evaluator.ComputeLocalTransforms(ControlRig->GetWorld(), InSequencer, FramesToCompensate, ActiveValueToBeSet);
+			TArray<FTransform>& ChildLocals = Evaluator.ChildLocals;
+			
+			// store tangents at this time
+			TArray<FMovieSceneTangentData> Tangents;
+			int32 ChannelIndex = 0, NumChannels = 0;
+
+			FChannelMapInfo* pChannelIndex = nullptr;
+			FRigControlElement* ControlElement = nullptr;
+			Tie(ControlElement, pChannelIndex) = FControlRigSpaceChannelHelpers::GetControlAndChannelInfo(ControlRig, TransformSection, ControlName);
+
+			if (pChannelIndex && ControlElement)
+			{
+				// get the number of float channels to treat
+				NumChannels = FControlRigSpaceChannelHelpers::GetNumFloatChannels(ControlElement->Settings.ControlType);
+				if (bNeedsCompensation && NumChannels > 0)
+				{
+					ChannelIndex = pChannelIndex->ChannelIndex;
+					EvaluateTangentAtThisTime<FMovieSceneFloatChannel>(ChannelIndex, NumChannels, TransformSection, InTime, Tangents);
+				}
+			}
+		
+			const EMovieSceneTransformChannel ChannelsToKey =InConstraint->GetChannelsToKey();
+			
+			// add child's transform key at Time-1 to keep animation
+			if (bNeedsCompensation)
+			{
+				const FFrameNumber TimeMinusOne(InTime - 1);
+
+				ControlHandle->AddTransformKeys({ TimeMinusOne },
+					{ ChildLocals[0] }, ChannelsToKey, TickResolution, nullptr,true);
+
+				// set tangents at Time-1
+				if (NumChannels > 0)
+				{
+					SetTangentsAtThisTime<FMovieSceneFloatChannel>(ChannelIndex, NumChannels, TransformSection, TimeMinusOne, Tangents);
+				}
+			}
+
+			// add active key
+			{
+				TMovieSceneChannelData<bool> ChannelData = Channel->ActiveChannel.GetData();
+				ChannelData.AddKey(InTime, ActiveValueToBeSet);
+			}
+
+			// compensate
+			{
+				// we need to remove the first transforms as we store NumFrames+1 transforms
+				ChildLocals.RemoveAt(0);
+
+				// add keys
+				ControlHandle->AddTransformKeys(FramesToCompensate,
+					ChildLocals, ChannelsToKey, TickResolution, nullptr,true);
+
+				// set tangents at Time
+				if (bNeedsCompensation && NumChannels > 0)
+				{
+					SetTangentsAtThisTime<FMovieSceneFloatChannel>(ChannelIndex, NumChannels, TransformSection, InTime, Tangents);
+				}
+			}
+			return true;
+		}
+	}
+	return false;
+}
+
+void FControlConstraintChannelInterface::AddHandleTransformKeys(
+	const TSharedPtr<ISequencer>& InSequencer,
+	const UTransformableHandle* InHandle,
+	const TArray<FFrameNumber>& InFrames,
+	const TArray<FTransform>& InLocalTransforms,
+	const EMovieSceneTransformChannel& InChannels)
+{
+	ensure(InLocalTransforms.Num());
+
+	if (!IsValid(InHandle) || !InHandle->IsValid())
+	{
+		return;
+	}
+
+	const UTransformableControlHandle* Handle = static_cast<const UTransformableControlHandle*>(InHandle);
+	const UMovieScene* MovieScene = InSequencer->GetFocusedMovieSceneSequence()->GetMovieScene();
+	UMovieSceneSection* Section = nullptr; //control rig doesn't need section it instead
+	Handle->AddTransformKeys(InFrames, InLocalTransforms, InChannels, MovieScene->GetTickResolution(), Section, true);
+}
+
+UMovieSceneControlRigParameterSection* FControlConstraintChannelInterface::GetControlSection(
+	const UTransformableControlHandle* InHandle,
+	const TSharedPtr<ISequencer>& InSequencer,
+	const bool bIsConstraint)
+{
+	if (!IsValid(InHandle) || !InHandle->IsValid())
+	{
+		return nullptr;
+	}
+	
+	const UControlRig* ControlRig = InHandle->ControlRig.LoadSynchronous();
+	if (!ControlRig)
+	{
+		return nullptr;
+	}
+	
+	const UMovieScene* MovieScene = GetMovieScene(InSequencer);
+	if (!MovieScene)
+	{
+		return nullptr;
+	}
+
+	auto GetControlRigTrack = [InHandle, MovieScene]()->UMovieSceneControlRigParameterTrack*
+	{
+		const TWeakObjectPtr<UControlRig> ControlRig = InHandle->ControlRig.LoadSynchronous();
+		if (ControlRig.IsValid())
+		{	
+			const TArray<FMovieSceneBinding>& Bindings = MovieScene->GetBindings();
+			for (const FMovieSceneBinding& Binding : Bindings)
+			{
+				UMovieSceneTrack* Track = MovieScene->FindTrack(UMovieSceneControlRigParameterTrack::StaticClass(), Binding.GetObjectGuid());
+				UMovieSceneControlRigParameterTrack* ControlRigTrack = Cast<UMovieSceneControlRigParameterTrack>(Track);
+				if (ControlRigTrack && ControlRigTrack->GetControlRig() == ControlRig)
+				{
+					return ControlRigTrack;
+				}
+			}
+		}
+		return nullptr;
+	};
+
+	UMovieSceneControlRigParameterTrack* ControlRigTrack = GetControlRigTrack();
+	if (!ControlRigTrack)
+	{
+		return nullptr;
+	}
+
+	UMovieSceneSection* Section = ControlRigTrack->FindSection(0);
+	if (bIsConstraint)
+	{
+		const TArray<UMovieSceneSection*>& AllSections = ControlRigTrack->GetAllSections();
+		if (!AllSections.IsEmpty())
+		{
+			Section = AllSections[0]; 
+		}
+	}
+
+	return Cast<UMovieSceneControlRigParameterSection>(Section);
+}
+
+#undef LOCTEXT_NAMESPACE
diff --git a/Engine/Plugins/Animation/ControlRig/Source/ControlRigEditor/Private/Sequencer/ControlRigParameterTrackEditor.cpp b/Engine/Plugins/Animation/ControlRig/Source/ControlRigEditor/Private/Sequencer/ControlRigParameterTrackEditor.cpp
--- a/Engine/Plugins/Animation/ControlRig/Source/ControlRigEditor/Private/Sequencer/ControlRigParameterTrackEditor.cpp
+++ b/Engine/Plugins/Animation/ControlRig/Source/ControlRigEditor/Private/Sequencer/ControlRigParameterTrackEditor.cpp
@@ -79,8 +79,7 @@
 #include "TransformConstraint.h"
 #include "Misc/ScopedSlowTask.h"
 #include "Misc/TransactionObjectEvent.h"
-#include "ConstraintChannelHelper.h"
-#include "MovieSceneConstraintChannelHelper.h"
+#include "Constraints/MovieSceneConstraintChannelHelper.h"
 #include "Constraints/ControlRigTransformableHandle.h"
 #include "PropertyEditorModule.h"
 
@@ -3473,8 +3472,7 @@ FKeyPropertyResult FControlRigParameterTrackEditor::AddKeysToControlRigHandle(US
 
 				// compensate constraints
 				const uint32 ControlHash = UTransformableControlHandle::ComputeHash(SectionControlRig, RigControlName);
-				FConstraintChannelHelper::CompensateIfNeeded(
-					SectionControlRig->GetWorld(), GetSequencer(), ParamSection, OptionalKeyTime, ControlHash);
+				FMovieSceneConstraintChannelHelper::CompensateIfNeeded(GetSequencer(), ParamSection, OptionalKeyTime, ControlHash);
 			}
 		}
 	}
diff --git a/Engine/Plugins/Animation/ControlRig/Source/ControlRigEditor/Private/Sequencer/ControlRigSpaceChannelEditors.cpp b/Engine/Plugins/Animation/ControlRig/Source/ControlRigEditor/Private/Sequencer/ControlRigSpaceChannelEditors.cpp
--- a/Engine/Plugins/Animation/ControlRig/Source/ControlRigEditor/Private/Sequencer/ControlRigSpaceChannelEditors.cpp
+++ b/Engine/Plugins/Animation/ControlRig/Source/ControlRigEditor/Private/Sequencer/ControlRigSpaceChannelEditors.cpp
@@ -1,7 +1,7 @@
 // Copyright Epic Games, Inc. All Rights Reserved.
 
 #include "ControlRigSpaceChannelEditors.h"
-#include "ConstraintChannelHelper.inl"
+#include "Constraints/MovieSceneConstraintChannelHelper.inl"
 #include "MovieSceneSequence.h"
 #include "MovieSceneSequenceEditor.h"
 #include "MovieSceneEventUtils.h"
diff --git a/Engine/Plugins/Animation/ControlRig/Source/ControlRigEditor/Private/Tools/ConstraintBaker.cpp b/Engine/Plugins/Animation/ControlRig/Source/ControlRigEditor/Private/Tools/ConstraintBaker.cpp
--- a/Engine/Plugins/Animation/ControlRig/Source/ControlRigEditor/Private/Tools/ConstraintBaker.cpp
+++ b/Engine/Plugins/Animation/ControlRig/Source/ControlRigEditor/Private/Tools/ConstraintBaker.cpp
@@ -2,350 +2,22 @@
 
 #include "Tools/ConstraintBaker.h"
 
-#include "IControlRigObjectBinding.h"
 #include "TransformConstraint.h"
+#include "TransformableHandle.h"
+#include "ConstraintChannel.h"
 
 #include "ISequencer.h"
 
-#include "ControlRig.h"
-#include "IKeyArea.h"
 #include "LevelSequence.h"
 #include "MovieSceneToolHelpers.h"
-#include "TransformableHandle.h"
-#include "Channels/MovieSceneChannelProxy.h"
-#include "Constraints/ControlRigTransformableHandle.h"
-#include "MovieSceneConstraintChannelHelper.h"
+
+#include "Constraints/MovieSceneConstraintChannelHelper.inl"
+#include "Constraints/TransformConstraintChannelInterface.h"
 #include "Sections/MovieScene3DTransformSection.h"
 #include "Sequencer/MovieSceneControlRigParameterSection.h"
-#include "ConstraintChannelHelper.h"
-#include "ConstraintChannel.h"
-#include "ConstraintChannelHelper.inl"
-#include "EntitySystem/MovieSceneDecompositionQuery.h"
-#include "Sequencer/MovieSceneControlRigParameterTrack.h"
-#include "Tracks/MovieScene3DTransformTrack.h"
-#include "Systems/MovieScenePropertyInstantiator.h"
 
 #define LOCTEXT_NAMESPACE "ConstraintBaker"
 
-namespace MoveMeFor52
-{
-	UMovieScene* GetMovieScene(const TSharedPtr<ISequencer>& InSequencer)
-	{
-		const UMovieSceneSequence* MovieSceneSequence = InSequencer ? InSequencer->GetFocusedMovieSceneSequence() : nullptr;
-		return MovieSceneSequence ? MovieSceneSequence->GetMovieScene() : nullptr;
-	}
-	
-	UMovieScene3DTransformSection* GetTransformSection(
-		const TSharedPtr<ISequencer>& InSequencer,
-		AActor* InActor,
-		const FTransform& InTransform0)
-	{
-		if (!InSequencer || !InSequencer->GetFocusedMovieSceneSequence())
-		{
-			return nullptr;
-		}
-
-		const FGuid Guid = InSequencer->GetHandleToObject(InActor, true);
-		if (!Guid.IsValid())
-		{
-			return nullptr;
-		}
-
-		return MovieSceneToolHelpers::GetTransformSection(InSequencer.Get(), Guid, InTransform0);
-	}
-
-	// NOTE, because of restriction for hotfix changes this code some of the following functions are duplicated in
-	// ConstraintChannelHelper.cpp. This will be factorized as part of a refactor for 5.2
-	UMovieScene3DTransformSection* GetComponentConstraintSection(
-		const TSharedPtr<ISequencer>& InSequencer,
-		const FGuid& InGuid,
-		const FTransform& InDefaultTransform)
-	{
-		UMovieScene* MovieScene = MoveMeFor52::GetMovieScene(InSequencer);
-		if (!MovieScene)
-		{
-			return nullptr;
-		}
-
-		UMovieScene3DTransformTrack* TransformTrack = MovieScene->FindTrack<UMovieScene3DTransformTrack>(InGuid);
-		if (!TransformTrack)
-		{
-			MovieScene->Modify();
-			TransformTrack = MovieScene->AddTrack<UMovieScene3DTransformTrack>(InGuid);
-		}
-		TransformTrack->Modify();
-
-		const TArray<UMovieSceneSection*>& AllSections = TransformTrack->GetAllSections();
-
-		static constexpr FFrameNumber Frame0;
-		bool bSectionAdded = false;
-		UMovieScene3DTransformSection* TransformSection = AllSections.IsEmpty() ?
-			Cast<UMovieScene3DTransformSection>(TransformTrack->FindOrAddSection(Frame0, bSectionAdded)) :
-			Cast<UMovieScene3DTransformSection>(AllSections[0]);
-		if (!TransformSection)
-		{
-			return nullptr;
-		}
-
-		TransformSection->Modify();
-		if (bSectionAdded)
-		{
-			TransformSection->SetRange(TRange<FFrameNumber>::All());
-
-			const FMovieSceneChannelProxy& SectionChannelProxy = TransformSection->GetChannelProxy();
-			const TMovieSceneChannelHandle<FMovieSceneDoubleChannel> DoubleChannels[] = {
-				SectionChannelProxy.GetChannelByName<FMovieSceneDoubleChannel>("Location.X"),
-				SectionChannelProxy.GetChannelByName<FMovieSceneDoubleChannel>("Location.Y"),
-				SectionChannelProxy.GetChannelByName<FMovieSceneDoubleChannel>("Location.Z"),
-				SectionChannelProxy.GetChannelByName<FMovieSceneDoubleChannel>("Rotation.X"),
-				SectionChannelProxy.GetChannelByName<FMovieSceneDoubleChannel>("Rotation.Y"),
-				SectionChannelProxy.GetChannelByName<FMovieSceneDoubleChannel>("Rotation.Z"),
-				SectionChannelProxy.GetChannelByName<FMovieSceneDoubleChannel>("Scale.X"),
-				SectionChannelProxy.GetChannelByName<FMovieSceneDoubleChannel>("Scale.Y"),
-				SectionChannelProxy.GetChannelByName<FMovieSceneDoubleChannel>("Scale.Z")
-			};
-
-			const FVector Location0 = InDefaultTransform.GetLocation();
-			const FRotator Rotation0 = InDefaultTransform.GetRotation().Rotator();
-			const FVector Scale3D0 = InDefaultTransform.GetScale3D();
-			const FTransform::FReal DefaultValues[] = { Location0.X, Location0.Y, Location0.Z,
-														Rotation0.Roll, Rotation0.Pitch, Rotation0.Yaw,
-														Scale3D0.X, Scale3D0.Y, Scale3D0.Z };
-			for (int32 Index = 0; Index < 9; Index++)
-			{
-				if (FMovieSceneDoubleChannel* Channel = DoubleChannels[Index].Get())
-				{
-					Channel->SetDefault(DefaultValues[Index]);
-				}
-			}
-		}
-
-		return TransformSection;
-	}
-
-	// NOTE this has to be moved to a public header like MovieSceneToolsHelpers.h (or similar) for 5.2
-	// see F3DTransformTrackEditor::RecomposeTransform for more details
-	void RecomposeTransforms(
-		const TSharedPtr<ISequencer>& InSequencer, USceneComponent* SceneComponent, UMovieSceneSection* Section,
-		const TArray<FFrameNumber>& InFrames, TArray<FTransform>& InOutTransforms)
-	{
-		using namespace UE::MovieScene;
-
-		FMovieSceneRootEvaluationTemplateInstance& EvaluationTemplate = InSequencer->GetEvaluationTemplate();
-		UMovieSceneEntitySystemLinker* EntityLinker = EvaluationTemplate.GetEntitySystemLinker();
-		if (!EntityLinker)
-		{
-			return;
-		}
-
-		UMovieScenePropertyInstantiatorSystem* System = EntityLinker->FindSystem<UMovieScenePropertyInstantiatorSystem>();
-		if (!System)
-		{
-			return;
-		}
-
-		TGuardValue<FEntityManager*> DebugVizGuard(GEntityManagerForDebuggingVisualizers, &EntityLinker->EntityManager);
-		TArray<FMovieSceneEntityID> ImportedEntityIDs;
-		EvaluationTemplate.FindEntitiesFromOwner(Section, InSequencer->GetFocusedTemplateID(), ImportedEntityIDs);
-		if (ImportedEntityIDs.Num())
-		{
-			const FBuiltInComponentTypes* BuiltInComponents = FBuiltInComponentTypes::Get();
-			FMovieSceneTracksComponentTypes* TrackComponents = FMovieSceneTracksComponentTypes::Get();
-
-			UMovieScene* MovieScene = InSequencer->GetFocusedMovieSceneSequence()->GetMovieScene();
-			const FFrameRate TickResolution = MovieScene->GetTickResolution();
-			const EMovieScenePlayerStatus::Type PlaybackStatus = InSequencer->GetPlaybackStatus();
-
-			TArray<FMovieSceneEntityID> EntityIDs;
-			FDecompositionQuery Query;
-			Query.Object   = SceneComponent;
-			Query.bConvertFromSourceEntityIDs = false;  // We already pass the children entity IDs
-			
-			// add keys
-			for (int32 Index = 0; Index < InFrames.Num(); ++Index)
-			{
-				const FFrameNumber& FrameNumber = InFrames[Index];
-				const FMovieSceneEvaluationRange EvaluationRange = FMovieSceneEvaluationRange(FFrameTime(FrameNumber), TickResolution);
-				const FMovieSceneContext Context = FMovieSceneContext(EvaluationRange, PlaybackStatus).SetHasJumped(true);
-
-				EvaluationTemplate.EvaluateSynchronousBlocking(Context, *InSequencer);
-
-				if (EntityIDs.IsEmpty())
-				{
-					// In order to check for the result channels later, we need to look up the children entities that are
-					// bound to the given animated object. Imported entities generally don't have the result channels.
-					FEntityTaskBuilder()
-					.ReadEntityIDs()
-					.Read(BuiltInComponents->ParentEntity)
-					.Read(BuiltInComponents->BoundObject)
-					.FilterAll({ TrackComponents->ComponentTransform.PropertyTag })
-					.Iterate_PerEntity(
-						&EntityLinker->EntityManager, 
-						[SceneComponent, ImportedEntityIDs, &EntityIDs](FMovieSceneEntityID EntityID, FMovieSceneEntityID ParentEntityID, UObject* BoundObject)
-						{
-							if (SceneComponent == BoundObject && ImportedEntityIDs.Contains(ParentEntityID))
-							{
-								EntityIDs.Add(EntityID);
-							}
-						});
-					Query.Entities = MakeArrayView(EntityIDs);
-				}
-				
-				FTransform& Transform = InOutTransforms[Index];
-
-				const FIntermediate3DTransform CurrentValue(Transform.GetTranslation(), Transform.GetRotation().Rotator(), Transform.GetScale3D());
-
-				TRecompositionResult<FIntermediate3DTransform> TransformData = System->RecomposeBlendOperational(TrackComponents->ComponentTransform, Query, CurrentValue);
-
-				double CurrentTransformChannels[9] = { 0, 0, 0, 0, 0, 0, 0, 0, 0 };
-				EMovieSceneTransformChannel ChannelsObtained(EMovieSceneTransformChannel::None);
-				check(EntityIDs.Num() == TransformData.Values.Num());
-				for (int32 EntityIndex = 0; EntityIndex < TransformData.Values.Num(); ++EntityIndex)
-				{
-					const FMovieSceneEntityID& EntityID = EntityIDs[EntityIndex];
-					
-					const FIntermediate3DTransform& EntityTransformData = TransformData.Values[EntityIndex];
-					FComponentMask EntityType = EntityLinker->EntityManager.GetEntityType(EntityID);
-					for (int32 CompositeIndex = 0; CompositeIndex < 9; ++CompositeIndex)
-					{
-						EMovieSceneTransformChannel ChannelMask = (EMovieSceneTransformChannel)(1 << CompositeIndex);
-						if (!EnumHasAnyFlags(ChannelsObtained, ChannelMask) && EntityType.Contains(BuiltInComponents->DoubleResult[CompositeIndex]))
-						{
-							EnumAddFlags(ChannelsObtained, (EMovieSceneTransformChannel)(1 << CompositeIndex));
-							CurrentTransformChannels[CompositeIndex] = EntityTransformData[CompositeIndex];
-						}
-					}
-				}
-
-				Transform = FTransform(
-					FRotator(CurrentTransformChannels[4], CurrentTransformChannels[5], CurrentTransformChannels[3]), // pitch yaw roll
-					FVector(CurrentTransformChannels[0], CurrentTransformChannels[1], CurrentTransformChannels[2]),
-					FVector(CurrentTransformChannels[6], CurrentTransformChannels[7], CurrentTransformChannels[8]));
-			}
-		}
-	}
-	
-	void BakeComponent(
-		const TSharedPtr<ISequencer>& InSequencer,
-		const UTransformableComponentHandle* InComponentHandle,
-		const TArray<FFrameNumber>& InFrames,
-		const TArray<FTransform>& InTransforms,
-		const EMovieSceneTransformChannel& InChannels)
-	{
-		ensure(InTransforms.Num());
-
-		if (!InComponentHandle->IsValid())
-		{
-			return;
-		}
-		AActor* Actor = InComponentHandle->Component->GetOwner();
-		if (!Actor)
-		{
-			return;
-		}
-		
-		UMovieScene3DTransformSection* TransformSection = MoveMeFor52::GetTransformSection(InSequencer, Actor, InTransforms[0]);
-		if (!TransformSection)
-		{
-			return;
-		}
-
-		const UMovieScene* MovieScene = InSequencer->GetFocusedMovieSceneSequence()->GetMovieScene();
-		const FFrameRate TickResolution = MovieScene->GetTickResolution();
-
-		const FGuid Guid = InSequencer->GetHandleToObject(Actor, false);
-		const bool bNeedsRecomposition = TransformSection != MoveMeFor52::GetComponentConstraintSection(InSequencer, Guid, InTransforms[0]);
-		if (bNeedsRecomposition)
-		{
-			TArray<FTransform> Transforms(InTransforms);
-			MoveMeFor52::RecomposeTransforms(InSequencer, InComponentHandle->Component.Get(), TransformSection, InFrames, Transforms);
-			InComponentHandle->AddTransformKeys(InFrames, Transforms, InChannels, TickResolution, TransformSection, true);
-		}
-		else
-		{
-			InComponentHandle->AddTransformKeys(InFrames, InTransforms, InChannels, TickResolution, TransformSection, true);
-		}
-	}
-
-	void BakeControl(
-		const TSharedPtr<ISequencer>& InSequencer,
-		const UTransformableControlHandle* InControlHandle,
-		const TArray<FFrameNumber>& InFrames,
-		const TArray<FTransform>& InLocalTransforms,
-		const EMovieSceneTransformChannel& InChannels)
-	{
-		ensure(InLocalTransforms.Num());
-
-		if (!InControlHandle->IsValid())
-		{
-			return;
-		}
-
-		const UMovieScene* MovieScene = InSequencer->GetFocusedMovieSceneSequence()->GetMovieScene();
-		UMovieSceneSection* Section = nullptr; //cnntrol rig doesn't need section it instead
-		InControlHandle->AddTransformKeys(InFrames, InLocalTransforms, InChannels, MovieScene->GetTickResolution(), Section, true);
-	}
-
-	UMovieSceneControlRigParameterSection* GetControlConstraintSection(
-	const UTransformableControlHandle* InHandle,
-	const TSharedPtr<ISequencer>& InSequencer)
-	{
-		const UMovieScene* MovieScene = MoveMeFor52::GetMovieScene(InSequencer);
-		if (!MovieScene)
-		{
-			return nullptr;
-		}
-
-		auto GetControlRigTrack = [InHandle, MovieScene]()->UMovieSceneControlRigParameterTrack*
-		{
-			const TWeakObjectPtr<UControlRig> ControlRig = InHandle->ControlRig.LoadSynchronous();
-			if (ControlRig.IsValid())
-			{	
-				const TArray<FMovieSceneBinding>& Bindings = MovieScene->GetBindings();
-				for (const FMovieSceneBinding& Binding : Bindings)
-				{
-					UMovieSceneTrack* Track = MovieScene->FindTrack(UMovieSceneControlRigParameterTrack::StaticClass(), Binding.GetObjectGuid());
-					UMovieSceneControlRigParameterTrack* ControlRigTrack = Cast<UMovieSceneControlRigParameterTrack>(Track);
-					if (ControlRigTrack && ControlRigTrack->GetControlRig() == ControlRig)
-					{
-						return ControlRigTrack;
-					}
-				}
-			}
-			return nullptr;
-		};
-
-		UMovieSceneControlRigParameterTrack* ControlRigTrack = GetControlRigTrack();
-		if (!ControlRigTrack)
-		{
-			return nullptr;
-		}
-
-		const TArray<UMovieSceneSection*>& AllSections = ControlRigTrack->GetAllSections();
-		UMovieSceneSection* Section = AllSections.IsEmpty() ? ControlRigTrack->FindSection(0) : AllSections[0];
-		return Cast<UMovieSceneControlRigParameterSection>(Section);
-	}
-}
-
-
-void FConstraintBaker::AddTransformKeys(
-	const TSharedPtr<ISequencer>& InSequencer,
-	UTransformableHandle* InHandle,
-	const TArray<FFrameNumber>& InFrames,
-	const TArray<FTransform>& InTransforms,
-	const EMovieSceneTransformChannel& InChannels)
-{
-	if (const UTransformableComponentHandle* ComponentHandle = Cast<UTransformableComponentHandle>(InHandle))
-	{
-		return MoveMeFor52::BakeComponent(InSequencer, ComponentHandle, InFrames, InTransforms, InChannels);
-	}
-	
-	if (const UTransformableControlHandle* ControlHandle = Cast<UTransformableControlHandle>(InHandle))
-	{
-		return MoveMeFor52::BakeControl(InSequencer, ControlHandle, InFrames, InTransforms, InChannels);
-	}
-}
-
 // at this stage, we suppose that any of these arguments are safe to use. Make sure to test them before using that function
 void FConstraintBaker::GetMinimalFramesToBake(
 	UWorld* InWorld,
@@ -493,35 +165,18 @@ void FConstraintBaker::Bake(UWorld* InWorld,
 		return;
 	}
 
-	const TObjectPtr<UTransformableHandle>& ChildHandle = InConstraint->ChildTRSHandle;
-	
-	//get the section to be used later to delete the extra transform keys at the frame -1 times, abort if not there for some reason
-	UMovieSceneSection* ConstraintSection = nullptr;
-	UMovieSceneSection* TransformSection = nullptr;
-	if (const UTransformableControlHandle* ControlHandle = Cast<UTransformableControlHandle>(ChildHandle))
-	{
-		TransformSection = FConstraintChannelHelper::GetControlSection(ControlHandle, InSequencer);
-		ConstraintSection = MoveMeFor52::GetControlConstraintSection(ControlHandle, InSequencer);
-	}
-	else if (const UTransformableComponentHandle* ComponentHandle = Cast<UTransformableComponentHandle>(ChildHandle))
-	{ 
-		//todo move to function also used by SmartConstraintKey
-		AActor* Actor = ComponentHandle->Component->GetOwner();
-		if (!Actor)
-		{
-			return;
-		}
-
-		const FTransform LocalTransform = ComponentHandle->GetLocalTransform();
-		const FGuid Guid = InSequencer->GetHandleToObject(Actor, true);
-		if (!Guid.IsValid())
-		{
-			return;
-		}
+	const TObjectPtr<UTransformableHandle>& Handle = InConstraint->ChildTRSHandle;
 
-		TransformSection = MovieSceneToolHelpers::GetTransformSection(InSequencer.Get(), Guid, LocalTransform);
-		ConstraintSection = MoveMeFor52::GetComponentConstraintSection(InSequencer, Guid, LocalTransform);
+	const FConstraintChannelInterfaceRegistry& InterfaceRegistry = FConstraintChannelInterfaceRegistry::Get();	
+	ITransformConstraintChannelInterface* Interface = InterfaceRegistry.FindConstraintChannelInterface(Handle->GetClass());
+	if (!Interface)
+	{
+		return;
 	}
+	
+	//get the section to be used later to delete the extra transform keys at the frame -1 times, abort if not there for some reason
+	UMovieSceneSection* ConstraintSection = Interface->GetHandleConstraintSection(Handle, InSequencer);
+	const UMovieSceneSection* TransformSection = Interface->GetHandleSection(Handle, InSequencer);
 
 	IMovieSceneConstrainedSection* ConstrainedSection = Cast<IMovieSceneConstrainedSection>(ConstraintSection);
 	if (ConstrainedSection == nullptr || TransformSection == nullptr)
@@ -588,8 +243,8 @@ void FConstraintBaker::Bake(UWorld* InWorld,
 	const TArrayView<const FFrameNumber> ConstraintFrames = ConstraintChannelData.GetTimes();
 	
 	// get transform channels
-	const TArrayView<FMovieSceneFloatChannel*> FloatTransformChannels = ChildHandle->GetFloatChannels(TransformSection);
-	const TArrayView<FMovieSceneDoubleChannel*> DoubleTransformChannels = ChildHandle->GetDoubleChannels(TransformSection);
+	const TArrayView<FMovieSceneFloatChannel*> FloatTransformChannels = Handle->GetFloatChannels(TransformSection);
+	const TArrayView<FMovieSceneDoubleChannel*> DoubleTransformChannels = Handle->GetDoubleChannels(TransformSection);
 
 	for (int32 Index = 0; Index < ConstraintFrames.Num(); ++Index)
 	{
@@ -615,7 +270,7 @@ void FConstraintBaker::Bake(UWorld* InWorld,
 
 	// now bake to channel curves
 	const EMovieSceneTransformChannel Channels = InConstraint->GetChannelsToKey();
-	AddTransformKeys(InSequencer, ChildHandle, FramesToBake, Transforms, Channels);
+	Interface->AddHandleTransformKeys(InSequencer, Handle, FramesToBake, Transforms, Channels);
 
 	// notify
 	InSequencer->RequestEvaluate();
diff --git a/Engine/Plugins/Animation/ControlRig/Source/ControlRigEditor/Public/ConstraintChannelHelper.h b/Engine/Plugins/Animation/ControlRig/Source/ControlRigEditor/Public/ConstraintChannelHelper.h
deleted file mode 100644
diff --git a/Engine/Plugins/Animation/ControlRig/Source/ControlRigEditor/Public/ControlConstraintChannelInterface.h b/Engine/Plugins/Animation/ControlRig/Source/ControlRigEditor/Public/ControlConstraintChannelInterface.h
new file mode 100644
--- /dev/null
+++ b/Engine/Plugins/Animation/ControlRig/Source/ControlRigEditor/Public/ControlConstraintChannelInterface.h
@@ -0,0 +1,52 @@
+﻿// Copyright Epic Games, Inc. All Rights Reserved.
+
+#pragma once
+
+#include "Templates/SharedPointer.h"
+#include "Constraints/TransformConstraintChannelInterface.h"
+
+class ISequencer;
+class UMovieSceneControlRigParameterSection;
+class IMovieSceneConstrainedSection;
+class UMovieSceneSection;
+class UTransformableHandle;
+class UTickableTransformConstraint;
+class UTransformableControlHandle;
+struct FFrameNumber;
+
+/**
+* Interface that defines animatable capabilities for UTransformableControlHandle
+*/
+
+class UMovieSceneControlRigParameterSection;
+class UTransformableControlHandle;
+
+struct FControlConstraintChannelInterface : ITransformConstraintChannelInterface
+{
+	virtual ~FControlConstraintChannelInterface() = default;
+
+	/** Get the UMovieSceneControlRigParameterSection from the ControlRig. */
+	
+	virtual UMovieSceneSection* GetHandleSection(const UTransformableHandle* InHandle, const TSharedPtr<ISequencer>& InSequencer) override;
+	virtual UMovieSceneSection* GetHandleConstraintSection(const UTransformableHandle* InHandle, const TSharedPtr<ISequencer>& InSequencer) override;
+
+	/** Get the world from the ControlRig. */
+	virtual UWorld* GetHandleWorld(UTransformableHandle* InHandle) override;
+
+	/** Add an active/inactive key to the constraint channel if needed and does the transform compensation on the control's transform channels. */
+	virtual bool SmartConstraintKey(
+		UTickableTransformConstraint* InConstraint, const TOptional<bool>& InOptActive,
+		const FFrameNumber& InTime, const TSharedPtr<ISequencer>& InSequencer) override;
+	
+	/** Add keys on the controls's transform channels. */
+	virtual void AddHandleTransformKeys(
+		const TSharedPtr<ISequencer>& InSequencer,
+		const UTransformableHandle* InHandle,
+		const TArray<FFrameNumber>& InFrames,
+		const TArray<FTransform>& InLocalTransforms,
+		const EMovieSceneTransformChannel& InChannels) override;
+
+private:
+	static UMovieSceneControlRigParameterSection* GetControlSection(const UTransformableControlHandle* InHandle, const TSharedPtr<ISequencer>& InSequencer, const bool bIsConstraint);
+};
+
diff --git a/Engine/Plugins/Animation/ControlRig/Source/ControlRigEditor/Public/Tools/ConstraintBaker.h b/Engine/Plugins/Animation/ControlRig/Source/ControlRigEditor/Public/Tools/ConstraintBaker.h
--- a/Engine/Plugins/Animation/ControlRig/Source/ControlRigEditor/Public/Tools/ConstraintBaker.h
+++ b/Engine/Plugins/Animation/ControlRig/Source/ControlRigEditor/Public/Tools/ConstraintBaker.h
@@ -11,7 +11,6 @@ class UTickableTransformConstraint;
 class ISequencer;
 class UWorld;
 struct FFrameNumber;
-enum class EMovieSceneTransformChannel : uint32;
 
 struct FConstraintBaker
 {
@@ -23,14 +22,6 @@ public:
 		const TSharedPtr<ISequencer>& InSequencer,
 		const TOptional<TArray<FFrameNumber>>& InFrames);
 	
-	/** Add InTransforms keys at InFrames into the InHandle transform animation channels. */
-	static void AddTransformKeys(
-		const TSharedPtr<ISequencer>& InSequencer,
-		UTransformableHandle* InHandle,
-		const TArray<FFrameNumber>& InFrames,
-		const TArray<FTransform>& InTransforms,
-		const EMovieSceneTransformChannel& InChannels);
-	
 private:
 
 	/** Computes the minimal frames to bake for that constraint. */
diff --git a/Engine/Source/Editor/MovieSceneTools/Private/Constraints/ComponentConstraintChannelInterface.cpp b/Engine/Source/Editor/MovieSceneTools/Private/Constraints/ComponentConstraintChannelInterface.cpp
new file mode 100644
--- /dev/null
+++ b/Engine/Source/Editor/MovieSceneTools/Private/Constraints/ComponentConstraintChannelInterface.cpp
@@ -0,0 +1,439 @@
+﻿// Copyright Epic Games, Inc. All Rights Reserved.
+
+
+#include "ComponentConstraintChannelInterface.h"
+
+#include "MovieSceneToolHelpers.h"
+#include "TransformableHandle.h"
+#include "TransformConstraint.h"
+#include "Constraints/MovieSceneConstraintChannelHelper.inl"
+#include "EntitySystem/MovieSceneDecompositionQuery.h"
+#include "Sections/MovieSceneConstrainedSection.h"
+#include "Sections/MovieScene3DTransformSection.h"
+#include "Systems/MovieScenePropertyInstantiator.h"
+#include "Tracks/MovieScene3DTransformTrack.h"
+#include "ScopedTransaction.h"
+
+UMovieSceneSection* FComponentConstraintChannelInterface::GetHandleSection(const UTransformableHandle* InHandle, const TSharedPtr<ISequencer>& InSequencer)
+{
+	if (!IsValid(InHandle) || !InHandle->IsValid())
+	{
+		return nullptr;
+	}
+
+	const UTransformableComponentHandle* ComponentHandle = static_cast<const UTransformableComponentHandle*>(InHandle);
+	static constexpr bool bConstraintSection = false;
+	return GetComponentSection(ComponentHandle, InSequencer, bConstraintSection);
+}
+
+UMovieSceneSection* FComponentConstraintChannelInterface::GetHandleConstraintSection(const UTransformableHandle* InHandle, const TSharedPtr<ISequencer>& InSequencer)
+{
+	if (!IsValid(InHandle) || !InHandle->IsValid())
+	{
+		return nullptr;
+	}
+
+	const UTransformableComponentHandle* ComponentHandle = static_cast<const UTransformableComponentHandle*>(InHandle);
+	static constexpr bool bConstraintSection = true;
+	return GetComponentSection(ComponentHandle, InSequencer, bConstraintSection);
+}
+
+UWorld* FComponentConstraintChannelInterface::GetHandleWorld(UTransformableHandle* InHandle)
+{
+	if (!IsValid(InHandle) || !InHandle->IsValid())
+	{
+		return nullptr;
+	}
+
+	const UTransformableComponentHandle* ComponentHandle = static_cast<const UTransformableComponentHandle*>(InHandle);
+	const AActor* Actor = ComponentHandle->Component->GetOwner();
+	return Actor ? Actor->GetWorld() : nullptr;
+}
+
+bool FComponentConstraintChannelInterface::SmartConstraintKey(
+	UTickableTransformConstraint* InConstraint,
+	const TOptional<bool>& InOptActive,
+	const FFrameNumber& InTime,
+	const TSharedPtr<ISequencer>& InSequencer)
+{
+	if (!IsValid(InConstraint->ChildTRSHandle) || !InConstraint->ChildTRSHandle->IsValid())
+	{
+		return false;
+	}
+	
+	const UTransformableComponentHandle* ComponentHandle = static_cast<UTransformableComponentHandle*>(InConstraint->ChildTRSHandle);
+
+	UMovieScene3DTransformSection* TransformSection = GetComponentSection(ComponentHandle, InSequencer, false);
+	UMovieScene3DTransformSection* ConstraintSection = GetComponentSection(ComponentHandle, InSequencer, true);
+	if ((!ConstraintSection) || (!TransformSection))
+	{
+		return false;
+	}
+
+	FScopedTransaction Transaction(NSLOCTEXT("Constraints", "KeyConstraintaKehy", "Key Constraint Key"));
+	TransformSection->Modify();
+	ConstraintSection->Modify();
+
+	// set constraint as dynamic
+	InConstraint->bDynamicOffset = true;
+	
+	// add the channel
+	if (ConstraintSection->HasConstraintChannel(InConstraint->GetFName()) == false)
+	{
+		ConstraintSection->AddConstraintChannel(InConstraint);
+		if (InSequencer.IsValid())
+		{
+			InSequencer->RecreateCurveEditor();
+		}
+	}
+
+	// add key if needed
+	if (FConstraintAndActiveChannel* Channel = ConstraintSection->GetConstraintChannel(InConstraint->GetFName()))
+	{
+		bool ActiveValueToBeSet = false;
+		//add key if we can and make sure the key we are setting is what we want
+		if (CanAddKey(Channel->ActiveChannel, InTime, ActiveValueToBeSet) && (InOptActive.IsSet() == false || InOptActive.GetValue() == ActiveValueToBeSet))
+		{
+			const bool bNeedsCompensation = InConstraint->NeedsCompensation();
+
+			//new for compensation
+
+			TGuardValue<bool> CompensateGuard(FMovieSceneConstraintChannelHelper::bDoNotCompensate, true);
+			// store the frames to compensate
+			const TArrayView<FMovieSceneDoubleChannel*> Channels = ComponentHandle->GetDoubleChannels(TransformSection);
+			TArray<FFrameNumber> FramesToCompensate;
+			if (bNeedsCompensation)
+			{
+				FMovieSceneConstraintChannelHelper::GetFramesToCompensate<FMovieSceneDoubleChannel>(Channel->ActiveChannel, ActiveValueToBeSet, InTime, Channels, FramesToCompensate);
+			}
+			else
+			{
+				FramesToCompensate.Add(InTime);
+			}
+
+			// store child and space transforms for these frames
+			FCompensationEvaluator Evaluator(InConstraint);
+			Evaluator.ComputeLocalTransforms(GetHandleWorld(InConstraint->ChildTRSHandle), InSequencer, FramesToCompensate, ActiveValueToBeSet);
+			TArray<FTransform>& ChildLocals = Evaluator.ChildLocals;
+
+			// store tangents at this time
+			TArray<FMovieSceneTangentData> Tangents;
+			const int32 ChannelIndex = 0;
+			const int32 NumChannels = 9;
+
+			if (bNeedsCompensation)
+			{
+				EvaluateTangentAtThisTime<FMovieSceneDoubleChannel>(ChannelIndex, NumChannels, TransformSection, InTime, Tangents);
+			}
+
+			const EMovieSceneTransformChannel ChannelsToKey = InConstraint->GetChannelsToKey();
+
+			// add child's transform key at Time-1 to keep animation
+			if (bNeedsCompensation)
+			{
+				const FFrameNumber TimeMinusOne(InTime - 1);
+
+				TArray<FTransform> TransformMinusOne({ChildLocals[0]} );
+				if (ConstraintSection != TransformSection)
+				{
+					RecomposeTransforms(InSequencer, ComponentHandle->Component.Get(), TransformSection, { TimeMinusOne },TransformMinusOne);
+				}
+				MovieSceneToolHelpers::AddTransformKeys(TransformSection, { TimeMinusOne }, TransformMinusOne, ChannelsToKey);
+
+				SetTangentsAtThisTime<FMovieSceneDoubleChannel>(ChannelIndex, NumChannels, TransformSection, TimeMinusOne, Tangents);
+			}
+
+			// add active key
+			{
+				TMovieSceneChannelData<bool> ChannelData = Channel->ActiveChannel.GetData();
+				ChannelData.AddKey(InTime, ActiveValueToBeSet);
+			}
+
+			// compensate
+			{
+				// we need to remove the first transforms as we store NumFrames+1 transforms
+				ChildLocals.RemoveAt(0);
+
+				if (ConstraintSection != TransformSection)
+				{
+					RecomposeTransforms(InSequencer, ComponentHandle->Component.Get(), TransformSection, FramesToCompensate,ChildLocals);
+				}
+
+				// add keys
+				MovieSceneToolHelpers::AddTransformKeys(TransformSection, FramesToCompensate, ChildLocals, ChannelsToKey);
+
+				// set tangents at Time
+				if (bNeedsCompensation)
+				{
+					SetTangentsAtThisTime<FMovieSceneDoubleChannel>(ChannelIndex,NumChannels, TransformSection, InTime, Tangents);
+				}
+			}
+			// evaluate the constraint, this is needed so the global transform will be set up on the component 
+			//Todo do we need to evaluate all constraints?
+			InConstraint->SetActive(true); //will be false
+			InConstraint->Evaluate();
+
+			//need to fire this event so the transform values set by the constraint propragate to the transform section
+			//first turn off autokey though
+			EAutoChangeMode AutoChangeMode = InSequencer->GetAutoChangeMode();
+			if (AutoChangeMode == EAutoChangeMode::AutoKey || AutoChangeMode == EAutoChangeMode::All)
+			{
+				InSequencer->SetAutoChangeMode(EAutoChangeMode::None);
+			};
+
+			AActor* Actor = ComponentHandle->Component->GetOwner();
+			FProperty* TransformProperty = FindFProperty<FProperty>(USceneComponent::StaticClass(), USceneComponent::GetRelativeLocationPropertyName());
+			FEditPropertyChain PropertyChain;
+			PropertyChain.AddHead(TransformProperty);
+			FCoreUObjectDelegates::OnPreObjectPropertyChanged.Broadcast(Actor, PropertyChain);
+			FPropertyChangedEvent PropertyChangedEvent(TransformProperty, EPropertyChangeType::ValueSet);
+			FCoreUObjectDelegates::OnObjectPropertyChanged.Broadcast(Actor, PropertyChangedEvent);
+			InSequencer->RequestEvaluate();
+
+			if (AutoChangeMode == EAutoChangeMode::AutoKey || AutoChangeMode == EAutoChangeMode::All)
+			{
+				InSequencer->SetAutoChangeMode(AutoChangeMode);
+			}
+			return true;
+		}
+	}
+	return false;
+}
+
+void FComponentConstraintChannelInterface::AddHandleTransformKeys(
+	const TSharedPtr<ISequencer>& InSequencer,
+	const UTransformableHandle* InHandle,
+	const TArray<FFrameNumber>& InFrames,
+	const TArray<FTransform>& InTransforms,
+	const EMovieSceneTransformChannel& InChannels)
+{
+	ensure(InTransforms.Num());
+	
+	if (!IsValid(InHandle) || !InHandle->IsValid())
+	{
+		return;
+	}
+
+	const UTransformableComponentHandle* Handle = static_cast<const UTransformableComponentHandle*>(InHandle);
+	AActor* Actor = Handle->Component->GetOwner();
+	if (!Actor)
+	{
+		return;
+	}
+
+	const FGuid Guid = InSequencer->GetHandleToObject(Actor, true);
+	if (!Guid.IsValid())
+	{
+		return;
+	}
+
+	UMovieScene3DTransformSection* TransformSection = GetComponentSection(Handle, InSequencer, false);
+	if (!TransformSection)
+	{
+		return;
+	}
+
+	const UMovieScene* MovieScene = InSequencer->GetFocusedMovieSceneSequence()->GetMovieScene();
+	
+	const UMovieScene3DTransformSection* ConstraintSection = GetComponentSection(Handle, InSequencer, true);
+	if (ConstraintSection && TransformSection != ConstraintSection)
+	{
+		TArray<FTransform> Transforms(InTransforms);
+		RecomposeTransforms(InSequencer, Handle->Component.Get(), TransformSection, InFrames, Transforms);
+		Handle->AddTransformKeys(InFrames, Transforms, InChannels, MovieScene->GetTickResolution(), TransformSection, true);
+	}
+	else
+	{
+		Handle->AddTransformKeys(InFrames, InTransforms, InChannels, MovieScene->GetTickResolution(), TransformSection, true);	
+	}
+}
+
+UMovieScene3DTransformSection* FComponentConstraintChannelInterface::GetComponentSection(
+	const UTransformableComponentHandle* InHandle, const TSharedPtr<ISequencer>& InSequencer, const bool bIsConstraint)
+{
+	if (!IsValid(InHandle) || !InHandle->IsValid())
+	{
+		return nullptr;
+	}
+
+	AActor* Actor = InHandle->Component->GetOwner();
+	if (!Actor)
+	{
+		return nullptr;
+	}
+
+	const FGuid Guid = InSequencer->GetHandleToObject(Actor, true);
+	if (!Guid.IsValid())
+	{
+		return nullptr;
+	}
+
+
+	const FTransform DefaultTransform = InHandle->GetLocalTransform();
+	if (!bIsConstraint)
+	{
+		return MovieSceneToolHelpers::GetTransformSection(InSequencer.Get(), Guid, DefaultTransform);
+	}
+
+	const UMovieSceneSequence* MovieSceneSequence = InSequencer ? InSequencer->GetFocusedMovieSceneSequence() : nullptr;
+	UMovieScene* MovieScene = MovieSceneSequence->GetMovieScene();
+	if (!MovieScene)
+	{
+		return nullptr;
+	}
+
+	UMovieScene3DTransformTrack* TransformTrack = MovieScene->FindTrack<UMovieScene3DTransformTrack>(Guid);
+	if (!TransformTrack)
+	{
+		MovieScene->Modify();
+		TransformTrack = MovieScene->AddTrack<UMovieScene3DTransformTrack>(Guid);
+	}
+	TransformTrack->Modify();
+
+	static constexpr FFrameNumber Frame0;
+	bool bSectionAdded = false;
+
+	const TArray<UMovieSceneSection*>& AllSections = TransformTrack->GetAllSections();
+	UMovieScene3DTransformSection* TransformSection = AllSections.IsEmpty() ?
+		Cast<UMovieScene3DTransformSection>(TransformTrack->FindOrAddSection(Frame0, bSectionAdded)) :
+		Cast<UMovieScene3DTransformSection>(AllSections[0]);
+	if (!TransformSection)
+	{
+		return nullptr;
+	}
+
+	TransformSection->Modify();
+	if (bSectionAdded)
+	{
+		TransformSection->SetRange(TRange<FFrameNumber>::All());
+
+		const FMovieSceneChannelProxy& SectionChannelProxy = TransformSection->GetChannelProxy();
+		const TMovieSceneChannelHandle<FMovieSceneDoubleChannel> DoubleChannels[] = {
+			SectionChannelProxy.GetChannelByName<FMovieSceneDoubleChannel>("Location.X"),
+			SectionChannelProxy.GetChannelByName<FMovieSceneDoubleChannel>("Location.Y"),
+			SectionChannelProxy.GetChannelByName<FMovieSceneDoubleChannel>("Location.Z"),
+			SectionChannelProxy.GetChannelByName<FMovieSceneDoubleChannel>("Rotation.X"),
+			SectionChannelProxy.GetChannelByName<FMovieSceneDoubleChannel>("Rotation.Y"),
+			SectionChannelProxy.GetChannelByName<FMovieSceneDoubleChannel>("Rotation.Z"),
+			SectionChannelProxy.GetChannelByName<FMovieSceneDoubleChannel>("Scale.X"),
+			SectionChannelProxy.GetChannelByName<FMovieSceneDoubleChannel>("Scale.Y"),
+			SectionChannelProxy.GetChannelByName<FMovieSceneDoubleChannel>("Scale.Z")
+		};
+
+		const FVector Location0 = DefaultTransform.GetLocation();
+		const FRotator Rotation0 = DefaultTransform.GetRotation().Rotator();
+		const FVector Scale3D0 = DefaultTransform.GetScale3D();
+		const FTransform::FReal DefaultValues[] = { Location0.X, Location0.Y, Location0.Z,
+													Rotation0.Roll, Rotation0.Pitch, Rotation0.Yaw,
+													Scale3D0.X, Scale3D0.Y, Scale3D0.Z };
+		for (int32 Index = 0; Index < 9; Index++)
+		{
+			if (FMovieSceneDoubleChannel* Channel = DoubleChannels[Index].Get())
+			{
+				Channel->SetDefault(DefaultValues[Index]);
+			}
+		}
+	}
+
+	return TransformSection;
+}
+
+void FComponentConstraintChannelInterface::RecomposeTransforms(
+	const TSharedPtr<ISequencer>& InSequencer, USceneComponent* SceneComponent, UMovieSceneSection* Section,
+	const TArray<FFrameNumber>& InFrames, TArray<FTransform>& InOutTransforms) const
+{
+	using namespace UE::MovieScene;
+
+	FMovieSceneRootEvaluationTemplateInstance& EvaluationTemplate = InSequencer->GetEvaluationTemplate();
+	UMovieSceneEntitySystemLinker* EntityLinker = EvaluationTemplate.GetEntitySystemLinker();
+	if (!EntityLinker)
+	{
+		return;
+	}
+
+	UMovieScenePropertyInstantiatorSystem* System = EntityLinker->FindSystem<UMovieScenePropertyInstantiatorSystem>();
+	if (!System)
+	{
+		return;
+	}
+
+	TGuardValue<FEntityManager*> DebugVizGuard(GEntityManagerForDebuggingVisualizers, &EntityLinker->EntityManager);
+	TArray<FMovieSceneEntityID> ImportedEntityIDs;
+	EvaluationTemplate.FindEntitiesFromOwner(Section, InSequencer->GetFocusedTemplateID(), ImportedEntityIDs);
+	if (ImportedEntityIDs.Num())
+	{
+		const FBuiltInComponentTypes* BuiltInComponents = FBuiltInComponentTypes::Get();
+		FMovieSceneTracksComponentTypes* TrackComponents = FMovieSceneTracksComponentTypes::Get();
+
+		UMovieScene* MovieScene = InSequencer->GetFocusedMovieSceneSequence()->GetMovieScene();
+		const FFrameRate TickResolution = MovieScene->GetTickResolution();
+		const EMovieScenePlayerStatus::Type PlaybackStatus = InSequencer->GetPlaybackStatus();
+
+		TArray<FMovieSceneEntityID> EntityIDs;
+		FDecompositionQuery Query;
+		Query.Object   = SceneComponent;
+		Query.bConvertFromSourceEntityIDs = false;  // We already pass the children entity IDs
+		
+		// add keys
+		for (int32 Index = 0; Index < InFrames.Num(); ++Index)
+		{
+			const FFrameNumber& FrameNumber = InFrames[Index];
+			const FMovieSceneEvaluationRange EvaluationRange = FMovieSceneEvaluationRange(FFrameTime(FrameNumber), TickResolution);
+			const FMovieSceneContext Context = FMovieSceneContext(EvaluationRange, PlaybackStatus).SetHasJumped(true);
+
+			EvaluationTemplate.EvaluateSynchronousBlocking(Context, *InSequencer);
+
+			if (EntityIDs.IsEmpty())
+			{
+				// In order to check for the result channels later, we need to look up the children entities that are
+				// bound to the given animated object. Imported entities generally don't have the result channels.
+				FEntityTaskBuilder()
+				.ReadEntityIDs()
+				.Read(BuiltInComponents->ParentEntity)
+				.Read(BuiltInComponents->BoundObject)
+				.FilterAll({ TrackComponents->ComponentTransform.PropertyTag })
+				.Iterate_PerEntity(
+					&EntityLinker->EntityManager, 
+					[SceneComponent, ImportedEntityIDs, &EntityIDs](FMovieSceneEntityID EntityID, FMovieSceneEntityID ParentEntityID, UObject* BoundObject)
+					{
+						if (SceneComponent == BoundObject && ImportedEntityIDs.Contains(ParentEntityID))
+						{
+							EntityIDs.Add(EntityID);
+						}
+					});
+				Query.Entities = MakeArrayView(EntityIDs);
+			}
+			
+			FTransform& Transform = InOutTransforms[Index];
+
+			const FIntermediate3DTransform CurrentValue(Transform.GetTranslation(), Transform.GetRotation().Rotator(), Transform.GetScale3D());
+
+			TRecompositionResult<FIntermediate3DTransform> TransformData = System->RecomposeBlendOperational(TrackComponents->ComponentTransform, Query, CurrentValue);
+
+			double CurrentTransformChannels[9] = { 0, 0, 0, 0, 0, 0, 0, 0, 0 };
+			EMovieSceneTransformChannel ChannelsObtained(EMovieSceneTransformChannel::None);
+			check(EntityIDs.Num() == TransformData.Values.Num());
+			for (int32 EntityIndex = 0; EntityIndex < TransformData.Values.Num(); ++EntityIndex)
+			{
+				const FMovieSceneEntityID& EntityID = EntityIDs[EntityIndex];
+				
+				const FIntermediate3DTransform& EntityTransformData = TransformData.Values[EntityIndex];
+				FComponentMask EntityType = EntityLinker->EntityManager.GetEntityType(EntityID);
+				for (int32 CompositeIndex = 0; CompositeIndex < 9; ++CompositeIndex)
+				{
+					EMovieSceneTransformChannel ChannelMask = (EMovieSceneTransformChannel)(1 << CompositeIndex);
+					if (!EnumHasAnyFlags(ChannelsObtained, ChannelMask) && EntityType.Contains(BuiltInComponents->DoubleResult[CompositeIndex]))
+					{
+						EnumAddFlags(ChannelsObtained, (EMovieSceneTransformChannel)(1 << CompositeIndex));
+						CurrentTransformChannels[CompositeIndex] = EntityTransformData[CompositeIndex];
+					}
+				}
+			}
+
+			Transform = FTransform(
+				FRotator(CurrentTransformChannels[4], CurrentTransformChannels[5], CurrentTransformChannels[3]), // pitch yaw roll
+				FVector(CurrentTransformChannels[0], CurrentTransformChannels[1], CurrentTransformChannels[2]),
+				FVector(CurrentTransformChannels[6], CurrentTransformChannels[7], CurrentTransformChannels[8]));
+		}
+	}
+}
\ No newline at end of file
diff --git a/Engine/Source/Editor/MovieSceneTools/Private/Constraints/ComponentConstraintChannelInterface.h b/Engine/Source/Editor/MovieSceneTools/Private/Constraints/ComponentConstraintChannelInterface.h
new file mode 100644
--- /dev/null
+++ b/Engine/Source/Editor/MovieSceneTools/Private/Constraints/ComponentConstraintChannelInterface.h
@@ -0,0 +1,46 @@
+﻿// Copyright Epic Games, Inc. All Rights Reserved.
+
+#pragma once
+
+#include "Constraints/TransformConstraintChannelInterface.h"
+
+/**
+* Interface that defines animatable capabilities for UTransformableComponentHandle
+*/
+
+class UTransformableComponentHandle;
+class UMovieScene3DTransformSection;
+class USceneComponent;
+
+struct FComponentConstraintChannelInterface : ITransformConstraintChannelInterface
+{
+	virtual ~FComponentConstraintChannelInterface() = default;
+
+	/** Get the UMovieScene3DTransformSection from the component. */
+	virtual UMovieSceneSection* GetHandleSection(const UTransformableHandle* InHandle, const TSharedPtr<ISequencer>& InSequencer) override;
+	virtual UMovieSceneSection* GetHandleConstraintSection(const UTransformableHandle* InHandle, const TSharedPtr<ISequencer>& InSequencer) override;
+
+	/** Get the world from the component. */
+	virtual UWorld* GetHandleWorld(UTransformableHandle* InHandle) override;
+	
+	/** Add an active/inactive key to the constraint channel if needed and does the transform compensation on the component's transform channels. */
+	virtual bool SmartConstraintKey(
+		UTickableTransformConstraint* InConstraint, const TOptional<bool>& InOptActive,
+		const FFrameNumber& InTime, const TSharedPtr<ISequencer>& InSequencer) override;
+
+	/** Add keys on the component's transform channels. */
+	virtual void AddHandleTransformKeys(
+		const TSharedPtr<ISequencer>& InSequencer,
+		const UTransformableHandle* InHandle,
+		const TArray<FFrameNumber>& InFrames,
+		const TArray<FTransform>& InLocalTransforms,
+		const EMovieSceneTransformChannel& InChannels) override;
+
+private:
+		static UMovieScene3DTransformSection* GetComponentSection(const UTransformableComponentHandle* InHandle, const TSharedPtr<ISequencer>& InSequencer, const bool bInConstraint);
+	
+		void RecomposeTransforms(
+			const TSharedPtr<ISequencer>& InSequencer,
+			USceneComponent* SceneComponent, UMovieSceneSection* Section,
+			const TArray<FFrameNumber>& InFrames, TArray<FTransform>& InOutTransforms) const;
+};
diff --git a/Engine/Source/Editor/MovieSceneTools/Private/Constraints/MovieSceneConstraintChannelHelper.cpp b/Engine/Source/Editor/MovieSceneTools/Private/Constraints/MovieSceneConstraintChannelHelper.cpp
new file mode 100644
--- /dev/null
+++ b/Engine/Source/Editor/MovieSceneTools/Private/Constraints/MovieSceneConstraintChannelHelper.cpp
@@ -0,0 +1,950 @@
+﻿// Copyright Epic Games, Inc. All Rights Reserved.
+
+#include "Constraints/MovieSceneConstraintChannelHelper.inl"
+
+#include "ISequencer.h"
+
+#include "TransformableHandle.h"
+#include "LevelSequence.h"
+#include "MovieSceneToolHelpers.h"
+
+#include "TransformConstraint.h"
+#include "Algo/Copy.h"
+
+#include "ScopedTransaction.h"
+#include "Channels/MovieSceneFloatChannel.h"
+#include "Channels/MovieSceneDoubleChannel.h"
+#include "LevelEditorViewport.h"
+#include "MovieSceneSpawnableAnnotation.h"
+#include "Sections/MovieSceneConstrainedSection.h"
+#include "MovieSceneToolsModule.h"
+#include "Constraints/TransformConstraintChannelInterface.h"
+
+
+/*
+*
+*  FCompensationEvaluator
+*
+*/
+
+FCompensationEvaluator::FCompensationEvaluator(UTickableTransformConstraint* InConstraint)
+	: Constraint(InConstraint)
+	, Handle(InConstraint ? InConstraint->ChildTRSHandle : nullptr)
+{}
+
+void FCompensationEvaluator::ComputeLocalTransforms(
+	UWorld* InWorld, const TSharedPtr<ISequencer>& InSequencer,
+	const TArray<FFrameNumber>& InFrames, const bool bToActive)
+{
+	if (InFrames.IsEmpty())
+	{
+		return;
+	}
+
+	using ConstraintPtr = TObjectPtr<UTickableConstraint>;
+	const TArray< ConstraintPtr > Constraints = GetHandleTransformConstraints(InWorld);
+	if (Constraints.IsEmpty())
+	{
+		return;
+	}
+	
+	const TArray<ConstraintPtr> ConstraintsMinusThis =
+		Constraints.FilterByPredicate([this](const ConstraintPtr& InConstraint)
+		{
+			return InConstraint != Constraint;
+		});
+
+	// find last active constraint in the list that is different than the on we want to compensate for
+	auto GetLastActiveConstraint = [ConstraintsMinusThis]()
+	{
+		// find last active constraint in the list that is different than the one we want to compensate for
+		const int32 LastActiveIndex = FTransformConstraintUtils::GetLastActiveConstraintIndex(ConstraintsMinusThis);
+
+		// if found, return its parent global transform
+		return LastActiveIndex > INDEX_NONE ? Cast<UTickableTransformConstraint>(ConstraintsMinusThis[LastActiveIndex]) : nullptr;
+	};
+
+	UMovieScene* MovieScene = InSequencer->GetFocusedMovieSceneSequence()->GetMovieScene();
+	const FFrameRate TickResolution = MovieScene->GetTickResolution();
+	const EMovieScenePlayerStatus::Type PlaybackStatus = InSequencer->GetPlaybackStatus();
+
+	const int32 NumFrames = InFrames.Num();
+
+	// resize arrays to num frames + 1 as we also evaluate at InFrames[0]-1
+	ChildLocals.SetNum(NumFrames + 1);
+	ChildGlobals.SetNum(NumFrames + 1);
+	SpaceGlobals.SetNum(NumFrames + 1);
+	const TArray<IMovieSceneToolsAnimationBakeHelper*>& BakeHelpers = FMovieSceneToolsModule::Get().GetAnimationBakeHelpers();
+	for (IMovieSceneToolsAnimationBakeHelper* BakeHelper : BakeHelpers)
+	{
+		if (BakeHelper)
+		{
+			BakeHelper->StartBaking(MovieScene);
+		}
+	}
+
+	// get all constraints to evaluate
+	const FConstraintsManagerController& Controller = FConstraintsManagerController::Get(InWorld);
+	static constexpr bool bSorted = true;
+	const TArray<ConstraintPtr> AllConstraints = Controller.GetAllConstraints(bSorted);
+	
+	for (int32 Index = 0; Index < NumFrames + 1; ++Index)
+	{
+		const FFrameNumber FrameNumber = (Index == 0) ? InFrames[0] - 1 : InFrames[Index - 1];
+
+		// evaluate animation
+		const FMovieSceneEvaluationRange EvaluationRange = FMovieSceneEvaluationRange(FFrameTime(FrameNumber), TickResolution);
+		const FMovieSceneContext Context = FMovieSceneContext(EvaluationRange, PlaybackStatus).SetHasJumped(true);
+
+		for (IMovieSceneToolsAnimationBakeHelper* BakeHelper : BakeHelpers)
+		{
+			if (BakeHelper)
+			{
+				BakeHelper->PreEvaluation(MovieScene, FrameNumber);
+			}
+		}
+		InSequencer->GetEvaluationTemplate().EvaluateSynchronousBlocking(Context, *InSequencer);
+
+		// evaluate constraints
+		for (const UTickableConstraint* InConstraint : AllConstraints)
+		{
+			InConstraint->Evaluate();
+		}
+
+		for (IMovieSceneToolsAnimationBakeHelper* BakeHelper : BakeHelpers)
+		{
+			if (BakeHelper)
+			{
+				BakeHelper->PostEvaluation(MovieScene, FrameNumber);
+			}
+		}
+		// evaluate ControlRig?
+		// ControlRig->Evaluate_AnyThread();
+
+		FTransform& ChildLocal = ChildLocals[Index];
+		FTransform& ChildGlobal = ChildGlobals[Index];
+		FTransform& SpaceGlobal = SpaceGlobals[Index];
+
+		// store child transforms        	
+		ChildLocal = Handle->GetLocalTransform();
+		ChildGlobal = Handle->GetGlobalTransform();
+
+		const UTickableTransformConstraint* LastConstraint = GetLastActiveConstraint();
+
+		// store constraint/parent space global transform
+		if (bToActive)
+		{
+			// if activating the constraint, store last constraint or parent space at T[0]-1
+			// and constraint space for all other times
+			if (Index == 0)
+			{
+				if (LastConstraint)
+				{
+					SpaceGlobal = LastConstraint->GetParentGlobalTransform();
+					TOptional<FTransform> Relative =
+						FTransformConstraintUtils::GetConstraintsRelativeTransform(ConstraintsMinusThis, ChildLocal, ChildGlobal);
+					if (Relative)
+					{
+						ChildLocal = *Relative;
+					}
+				}
+			}
+			else
+			{
+				SpaceGlobal = Constraint->GetParentGlobalTransform();
+				ChildLocal = FTransformConstraintUtils::ComputeRelativeTransform(
+					ChildLocal, ChildGlobal, SpaceGlobal, Constraint);
+			}
+		}
+		else
+		{
+			// if deactivating the constraint, store constraint space at T[0]-1
+			// and last constraint or parent space for all other times
+			if (Index == 0)
+			{
+				SpaceGlobal = Constraint->GetParentGlobalTransform();
+				ChildLocal = FTransformConstraintUtils::ComputeRelativeTransform(
+					ChildLocal, ChildGlobal, SpaceGlobal, Constraint);
+			}
+			else
+			{
+				if (LastConstraint)
+				{
+					SpaceGlobal = LastConstraint->GetParentGlobalTransform();
+					TOptional<FTransform> Relative =
+						FTransformConstraintUtils::GetConstraintsRelativeTransform(ConstraintsMinusThis, ChildLocal, ChildGlobal);
+					if (Relative)
+					{
+						ChildLocal = *Relative;
+					}
+				}
+			}
+		}
+	}
+	for (IMovieSceneToolsAnimationBakeHelper* BakeHelper : BakeHelpers)
+	{
+		if (BakeHelper)
+		{
+			BakeHelper->StopBaking(MovieScene);
+		}
+	}
+}
+void FCompensationEvaluator::ComputeLocalTransformsForBaking(UWorld* InWorld, const TSharedPtr<ISequencer>& InSequencer, const TArray<FFrameNumber>& InFrames)
+{
+	if (InFrames.IsEmpty())
+	{
+		return;
+	}
+
+	using ConstraintPtr = TObjectPtr<UTickableConstraint>;
+	const TArray< ConstraintPtr > Constraints = GetHandleTransformConstraints(InWorld);
+	
+	const TArray< ConstraintPtr > ConstraintsMinusThis =
+		Constraints.FilterByPredicate([this](const ConstraintPtr& InConstraint)
+		{
+			return InConstraint != Constraint;
+		});
+
+	auto GetLastActiveConstraint = [ConstraintsMinusThis]()
+	{
+		// find last active constraint in the list that is different than the one we want to compensate for
+		const int32 LastActiveIndex = FTransformConstraintUtils::GetLastActiveConstraintIndex(ConstraintsMinusThis);
+
+		// if found, return its parent global transform
+		return LastActiveIndex > INDEX_NONE ? Cast<UTickableTransformConstraint>(ConstraintsMinusThis[LastActiveIndex]) : nullptr;
+	};
+
+	UMovieScene* MovieScene = InSequencer->GetFocusedMovieSceneSequence()->GetMovieScene();
+	const FFrameRate TickResolution = MovieScene->GetTickResolution();
+	const EMovieScenePlayerStatus::Type PlaybackStatus = InSequencer->GetPlaybackStatus();
+
+	const int32 NumFrames = InFrames.Num();
+
+	ChildLocals.SetNum(NumFrames);
+	ChildGlobals.SetNum(NumFrames);
+	SpaceGlobals.SetNum(NumFrames);
+
+	// get all constraints for evaluation
+	const FConstraintsManagerController& Controller = FConstraintsManagerController::Get(InWorld);
+	static constexpr bool bSorted = true;
+	const TArray<ConstraintPtr> AllConstraints = Controller.GetAllConstraints(bSorted);
+
+	// avoid transacting when evaluating sequencer
+	TGuardValue<ITransaction*> TransactionGuard(GUndo, nullptr);
+	
+	const TArray<IMovieSceneToolsAnimationBakeHelper*>& BakeHelpers = FMovieSceneToolsModule::Get().GetAnimationBakeHelpers();
+	for (IMovieSceneToolsAnimationBakeHelper* BakeHelper : BakeHelpers)
+	{
+		if (BakeHelper)
+		{
+			BakeHelper->StartBaking(MovieScene);
+		}
+	}
+	for (int32 Index = 0; Index < NumFrames; ++Index)
+	{
+		const FFrameNumber& FrameNumber = InFrames[Index];
+
+		// evaluate animation
+		const FMovieSceneEvaluationRange EvaluationRange = FMovieSceneEvaluationRange(FFrameTime(FrameNumber), TickResolution);
+		const FMovieSceneContext Context = FMovieSceneContext(EvaluationRange, PlaybackStatus).SetHasJumped(true);
+
+		for (IMovieSceneToolsAnimationBakeHelper* BakeHelper : BakeHelpers)
+		{
+			if (BakeHelper)
+			{
+				BakeHelper->PreEvaluation(MovieScene, FrameNumber);
+			}
+		}
+		InSequencer->GetEvaluationTemplate().EvaluateSynchronousBlocking(Context, *InSequencer);
+
+		// evaluate constraints
+		for (UTickableConstraint* InConstraint : AllConstraints)
+		{
+			InConstraint->Evaluate(true);
+		}
+
+		for (IMovieSceneToolsAnimationBakeHelper* BakeHelper : BakeHelpers)
+		{
+			if (BakeHelper)
+			{
+				BakeHelper->PostEvaluation(MovieScene, FrameNumber);
+			}
+		}
+
+		// evaluate ControlRig?
+		// ControlRig->Evaluate_AnyThread();
+
+		FTransform& ChildLocal = ChildLocals[Index];
+		FTransform& ChildGlobal = ChildGlobals[Index];
+		FTransform& SpaceGlobal = SpaceGlobals[Index];
+
+		// store child transforms        	
+		ChildLocal = Handle->GetLocalTransform();
+		ChildGlobal = Handle->GetGlobalTransform();
+
+		// store constraint/parent space global transform
+		if (const UTickableTransformConstraint* LastConstraint = GetLastActiveConstraint())
+		{
+			SpaceGlobal = LastConstraint->GetParentGlobalTransform();
+			TOptional<FTransform> Relative =
+				FTransformConstraintUtils::GetConstraintsRelativeTransform(ConstraintsMinusThis, ChildLocal, ChildGlobal);
+			if (Relative)
+			{
+				ChildLocal = *Relative;
+			}
+		}
+	}
+	for (IMovieSceneToolsAnimationBakeHelper* BakeHelper : BakeHelpers)
+	{
+		if (BakeHelper)
+		{
+			BakeHelper->StopBaking(MovieScene);
+		}
+	}
+
+}
+void FCompensationEvaluator::ComputeLocalTransformsBeforeDeletion(
+	UWorld* InWorld,
+	const TSharedPtr<ISequencer>& InSequencer,
+	const TArray<FFrameNumber>& InFrames)
+{
+	if (InFrames.IsEmpty())
+	{
+		return;
+	}
+
+	using ConstraintPtr = TObjectPtr<UTickableConstraint>;
+	const TArray<ConstraintPtr> Constraints = GetHandleTransformConstraints(InWorld);
+	const TArray<ConstraintPtr> ConstraintsMinusThis = Constraints.FilterByPredicate(
+		[this](const ConstraintPtr& InConstraint)
+		{
+			return InConstraint != Constraint;
+		});
+
+	// find last active constraint in the list that is different than the on we want to compensate for
+	auto GetLastActiveConstraint = [this, ConstraintsMinusThis]()
+	{
+		// find last active constraint in the list that is different than the on we want to compensate for
+		const int32 LastActiveIndex = FTransformConstraintUtils::GetLastActiveConstraintIndex(ConstraintsMinusThis);
+
+		// if found, return its parent global transform
+		return LastActiveIndex > INDEX_NONE ? Cast<UTickableTransformConstraint>(ConstraintsMinusThis[LastActiveIndex]) : nullptr;
+	};
+
+	// get all constraints for evaluation
+	const FConstraintsManagerController& Controller = FConstraintsManagerController::Get(InWorld);
+	static constexpr bool bSorted = true;
+	const TArray<ConstraintPtr> AllConstraints = Controller.GetAllConstraints(bSorted);
+	
+	UMovieScene* MovieScene = InSequencer->GetFocusedMovieSceneSequence()->GetMovieScene();
+	const FFrameRate TickResolution = MovieScene->GetTickResolution();
+	const EMovieScenePlayerStatus::Type PlaybackStatus = InSequencer->GetPlaybackStatus();
+
+	const int32 NumFrames = InFrames.Num();
+
+	ChildLocals.SetNum(NumFrames);
+	ChildGlobals.SetNum(NumFrames);
+	SpaceGlobals.SetNum(NumFrames);
+	const TArray<IMovieSceneToolsAnimationBakeHelper*>& BakeHelpers = FMovieSceneToolsModule::Get().GetAnimationBakeHelpers();
+	for (IMovieSceneToolsAnimationBakeHelper* BakeHelper : BakeHelpers)
+	{
+		if (BakeHelper)
+		{
+			BakeHelper->StartBaking(MovieScene);
+		}
+	}
+	for (int32 Index = 0; Index < NumFrames; ++Index)
+	{
+		const FFrameNumber& FrameNumber = InFrames[Index];
+
+		// evaluate animation
+		const FMovieSceneEvaluationRange EvaluationRange = FMovieSceneEvaluationRange(FFrameTime(FrameNumber), TickResolution);
+		const FMovieSceneContext Context = FMovieSceneContext(EvaluationRange, PlaybackStatus).SetHasJumped(true);
+
+		for (IMovieSceneToolsAnimationBakeHelper* BakeHelper : BakeHelpers)
+		{
+			if (BakeHelper)
+			{
+				BakeHelper->PreEvaluation(MovieScene, FrameNumber);
+			}
+		}
+		InSequencer->GetEvaluationTemplate().EvaluateSynchronousBlocking(Context, *InSequencer);
+
+		// evaluate constraints
+		for (const UTickableConstraint* InConstraint : AllConstraints)
+		{
+			InConstraint->Evaluate(true);
+		}
+		
+		for (IMovieSceneToolsAnimationBakeHelper* BakeHelper : BakeHelpers)
+		{
+			if (BakeHelper)
+			{
+				BakeHelper->PostEvaluation(MovieScene, FrameNumber);
+			}
+		}
+		// evaluate ControlRig?
+		// ControlRig->Evaluate_AnyThread();
+
+		FTransform& ChildLocal = ChildLocals[Index];
+		FTransform& ChildGlobal = ChildGlobals[Index];
+		FTransform& SpaceGlobal = SpaceGlobals[Index];
+
+		// store child transforms        	
+		ChildLocal = Handle->GetLocalTransform();
+		ChildGlobal = Handle->GetGlobalTransform();
+
+		// store constraint/parent space global transform
+		if (const UTickableTransformConstraint* LastConstraint = GetLastActiveConstraint())
+		{
+			SpaceGlobal = LastConstraint->GetParentGlobalTransform();
+			TOptional<FTransform> Relative =
+				FTransformConstraintUtils::GetConstraintsRelativeTransform(Constraints, ChildLocal, ChildGlobal);
+			if(Relative)
+			{
+				ChildLocal = *Relative;
+			}
+		}
+	}
+	for (IMovieSceneToolsAnimationBakeHelper* BakeHelper : BakeHelpers)
+	{
+		if (BakeHelper)
+		{
+			BakeHelper->StopBaking(MovieScene);
+		}
+	}
+}
+
+void FCompensationEvaluator::ComputeCompensation(UWorld* InWorld, const TSharedPtr<ISequencer>& InSequencer, const FFrameNumber& InTime)
+{
+	using ConstraintPtr = TObjectPtr<UTickableConstraint>;
+	const TArray<ConstraintPtr> Constraints = GetHandleTransformConstraints(InWorld);
+	if (Constraints.IsEmpty())
+	{
+		return;
+	}
+
+	// find last active constraint in the list that is different than the on we want to compensate for
+	auto GetLastActiveConstraint = [this, Constraints]()
+	{
+		// find last active constraint in the list that is different than the on we want to compensate for
+		const int32 LastActiveIndex = FTransformConstraintUtils::GetLastActiveConstraintIndex(Constraints);
+
+		// if found, return its parent global transform
+		return LastActiveIndex > INDEX_NONE ? Cast<UTickableTransformConstraint>(Constraints[LastActiveIndex]) : nullptr;
+	};
+
+	// get all constraints for evaluation
+	const FConstraintsManagerController& Controller = FConstraintsManagerController::Get(InWorld);
+	static constexpr bool bSorted = true;
+	const TArray<ConstraintPtr> AllConstraints = Controller.GetAllConstraints(bSorted);
+	
+	UMovieScene* MovieScene = InSequencer->GetFocusedMovieSceneSequence()->GetMovieScene();
+	const TArray<IMovieSceneToolsAnimationBakeHelper*>& BakeHelpers = FMovieSceneToolsModule::Get().GetAnimationBakeHelpers();
+	for (IMovieSceneToolsAnimationBakeHelper* BakeHelper : BakeHelpers)
+	{
+		if (BakeHelper)
+		{
+			BakeHelper->StartBaking(MovieScene);
+		}
+	}
+	
+	auto EvaluateAt = [InSequencer, &AllConstraints, &BakeHelpers](const FFrameNumber InFrame)
+	{
+
+		UMovieScene* MovieScene = InSequencer->GetFocusedMovieSceneSequence()->GetMovieScene();
+		const FFrameRate TickResolution = MovieScene->GetTickResolution();
+		const EMovieScenePlayerStatus::Type PlaybackStatus = InSequencer->GetPlaybackStatus();
+
+		const FMovieSceneEvaluationRange EvaluationRange0 = FMovieSceneEvaluationRange(FFrameTime(InFrame), TickResolution);
+		const FMovieSceneContext Context0 = FMovieSceneContext(EvaluationRange0, PlaybackStatus).SetHasJumped(true);
+
+		for (IMovieSceneToolsAnimationBakeHelper* BakeHelper : BakeHelpers)
+		{
+			if (BakeHelper)
+			{
+				BakeHelper->PreEvaluation(MovieScene, InFrame);
+			}
+		}
+		InSequencer->GetEvaluationTemplate().EvaluateSynchronousBlocking(Context0, *InSequencer);
+
+		for (const UTickableConstraint* InConstraint : AllConstraints)
+		{
+			InConstraint->Evaluate(true);
+		}
+
+		for (IMovieSceneToolsAnimationBakeHelper* BakeHelper : BakeHelpers)
+		{
+			if (BakeHelper)
+			{
+				BakeHelper->PostEvaluation(MovieScene, InFrame);
+			}
+		}
+
+		// ControlRig->Evaluate_AnyThread();
+	};
+
+
+	// allocate
+	ChildLocals.SetNum(1);
+	ChildGlobals.SetNum(1);
+	SpaceGlobals.SetNum(1);
+
+	
+	// evaluate at InTime and store global
+	EvaluateAt(InTime);
+	ChildGlobals[0] = Handle->GetGlobalTransform();
+
+	// evaluate at InTime-1 and store local
+	EvaluateAt(InTime - 1);
+	ChildLocals[0] = Handle->GetLocalTransform();
+
+	for (IMovieSceneToolsAnimationBakeHelper* BakeHelper : BakeHelpers)
+	{
+		if (BakeHelper)
+		{
+			BakeHelper->StopBaking(MovieScene);
+		}
+	}
+	
+	// if constraint at T-1 then switch to its space
+	if (const UTickableTransformConstraint* LastConstraint = GetLastActiveConstraint())
+	{
+		SpaceGlobals[0] = LastConstraint->GetParentGlobalTransform();
+		TOptional<FTransform> Relative =
+			FTransformConstraintUtils::GetConstraintsRelativeTransform(Constraints, ChildLocals[0], ChildGlobals[0]);
+		if(Relative)
+		{
+			ChildLocals[0] = *Relative;
+		}
+	}
+	else // switch to parent space
+	{
+		const FTransform ChildLocal = ChildLocals[0];
+		Handle->SetGlobalTransform(ChildGlobals[0]);
+		ChildLocals[0] = Handle->GetLocalTransform();
+		Handle->SetLocalTransform(ChildLocal);
+	}
+}
+
+TArray< TObjectPtr<UTickableConstraint> > FCompensationEvaluator::GetHandleTransformConstraints(UWorld* InWorld) const
+{
+	using ConstraintPtr = TObjectPtr<UTickableConstraint>;
+	
+	if (Handle)
+	{
+		// get sorted transform constraints
+		const FConstraintsManagerController& Controller = FConstraintsManagerController::Get(InWorld);
+		static constexpr bool bSorted = true;
+		const TArray< ConstraintPtr > Constraints = Controller.GetParentConstraints(Handle->GetHash(), bSorted);
+		return Constraints.FilterByPredicate([](const ConstraintPtr& InConstraint)
+		{
+			return IsValid(InConstraint) && InConstraint.IsA<UTickableTransformConstraint>();  
+		});
+	}
+
+	static const TArray< ConstraintPtr > DummyArray;
+	return DummyArray;
+}
+
+bool FMovieSceneConstraintChannelHelper::bDoNotCompensate = false;
+
+
+/*
+*
+*  Constraint Channel Helpers
+*
+*/
+void FMovieSceneConstraintChannelHelper::HandleConstraintRemoved(
+	UTickableConstraint* InConstraint,
+	const FMovieSceneConstraintChannel* InConstraintChannel,
+	const TSharedPtr<ISequencer>& InSequencer,
+	UMovieSceneSection* InSection)
+{
+	UTickableTransformConstraint* Constraint = Cast<UTickableTransformConstraint>(InConstraint);
+	if (!Constraint || !Constraint->NeedsCompensation() || !InConstraintChannel || !InSection)
+	{
+		return;
+	}
+	
+	InSection->Modify();
+	TGuardValue<bool> CompensateGuard(bDoNotCompensate, true);
+	if (const UTransformableHandle* ControlHandle =Constraint->ChildTRSHandle)
+	{
+		const TArrayView<const FFrameNumber> Times = InConstraintChannel->GetData().GetTimes();
+		if (Times.IsEmpty())
+		{
+			return;
+		}
+
+		// get transform channels
+		const TArrayView<FMovieSceneFloatChannel*> FloatTransformChannels = ControlHandle->GetFloatChannels(InSection);
+		const TArrayView<FMovieSceneDoubleChannel*> DoubleTransformChannels = ControlHandle->GetDoubleChannels(InSection);
+
+		// get frames after this time
+		TArray<FFrameNumber> FramesToCompensate;
+		if (FloatTransformChannels.Num() > 0)
+		{
+			FMovieSceneConstraintChannelHelper::GetFramesWithinActiveState(*InConstraintChannel, FloatTransformChannels, FramesToCompensate);
+		}
+		else
+		{
+			FMovieSceneConstraintChannelHelper::GetFramesWithinActiveState(*InConstraintChannel, DoubleTransformChannels, FramesToCompensate);
+		}
+		// do the compensation
+		UWorld* World = GCurrentLevelEditingViewportClient ? GCurrentLevelEditingViewportClient->GetWorld() : nullptr;
+
+		FCompensationEvaluator Evaluator(Constraint);
+		Evaluator.ComputeLocalTransformsBeforeDeletion(World, InSequencer, FramesToCompensate);
+		const TArray<FTransform>& ChildLocals = Evaluator.ChildLocals;
+
+		const EMovieSceneTransformChannel ChannelsToKey = Constraint->GetChannelsToKey();
+		const FFrameRate TickResolution = InSequencer->GetFocusedTickResolution();
+
+		ControlHandle->AddTransformKeys(FramesToCompensate,
+			ChildLocals, ChannelsToKey, TickResolution, InSection);
+	}
+}
+
+void FMovieSceneConstraintChannelHelper::HandleConstraintKeyDeleted(
+	UTickableTransformConstraint* InConstraint,
+	const FMovieSceneConstraintChannel* InConstraintChannel,
+	const TSharedPtr<ISequencer>& InSequencer,
+	UMovieSceneSection* InSection,
+	const FFrameNumber& InTime)
+{
+	if (!InConstraint || !InConstraint->NeedsCompensation())
+	{
+		return;
+	}
+	
+	const FFrameNumber TimeMinusOne(InTime - 1);
+
+	bool CurrentValue = false, PreviousValue = false;
+	InConstraintChannel->Evaluate(TimeMinusOne, PreviousValue);
+	InConstraintChannel->Evaluate(InTime, CurrentValue);
+
+	if (CurrentValue == PreviousValue)
+	{
+		const int32 NumKeys = InConstraintChannel->GetNumKeys();
+		if (NumKeys > 1)
+		{
+			return;
+		}
+	}
+
+	TGuardValue<bool> CompensateGuard(bDoNotCompensate, true);
+
+	if (const UTransformableHandle* ControlHandle = InConstraint->ChildTRSHandle)
+	{
+		// get transform channels
+		const TArrayView<FMovieSceneFloatChannel*> FloatTransformChannels = ControlHandle->GetFloatChannels(InSection);
+		const TArrayView<FMovieSceneDoubleChannel*> DoubleTransformChannels = ControlHandle->GetDoubleChannels(InSection);
+
+		// get frames after this time
+		TArray<FFrameNumber> FramesToCompensate;
+		if (FloatTransformChannels.Num() > 0)
+		{
+			GetFramesAfter(*InConstraintChannel, InTime, FloatTransformChannels, FramesToCompensate);
+		}
+		else
+		{
+			GetFramesAfter(*InConstraintChannel,  InTime, DoubleTransformChannels, FramesToCompensate);
+		}
+		// do the compensation
+		FCompensationEvaluator Evaluator(InConstraint);
+		UWorld* World = GCurrentLevelEditingViewportClient ? GCurrentLevelEditingViewportClient->GetWorld() : nullptr;
+		Evaluator.ComputeLocalTransforms(World, InSequencer, FramesToCompensate, PreviousValue);
+		TArray<FTransform>& ChildLocals = Evaluator.ChildLocals;
+		//turn off constraint, if we delete the key it may not evaluate to false
+		InConstraint->SetActive(false);
+
+		if (ChildLocals.Num() < 2)
+		{
+			return;
+		}
+		ChildLocals.RemoveAt(0);
+
+		const EMovieSceneTransformChannel ChannelsToKey = InConstraint->GetChannelsToKey();
+		const FFrameRate TickResolution = InSequencer->GetFocusedTickResolution();
+
+		ControlHandle->AddTransformKeys(FramesToCompensate,
+			ChildLocals, ChannelsToKey, TickResolution, InSection);
+
+		// now delete any extra TimeMinusOne
+		if (FloatTransformChannels.Num() > 0)
+		{
+			FMovieSceneConstraintChannelHelper::DeleteTransformKeys(FloatTransformChannels, TimeMinusOne);
+		}
+		else
+		{
+			FMovieSceneConstraintChannelHelper::DeleteTransformKeys(DoubleTransformChannels, TimeMinusOne);
+		}
+	}
+}
+
+void FMovieSceneConstraintChannelHelper::HandleConstraintKeyMoved(
+	const UTickableTransformConstraint* InConstraint,
+	const FMovieSceneConstraintChannel* InConstraintChannel,
+	UMovieSceneSection* InSection,
+	const FFrameNumber& InCurrentFrame, const FFrameNumber& InNextFrame)
+{
+	const FFrameNumber Delta = InNextFrame - InCurrentFrame;
+	if (Delta == 0)
+	{
+		return;
+	}
+
+	if (!InConstraint || !InConstraintChannel || !InSection)
+	{
+		return;
+	}
+
+	if (const UTransformableHandle* ControlHandle = InConstraint->ChildTRSHandle)
+	{
+
+		// get transform channels
+		const TArrayView<FMovieSceneFloatChannel*> FloatTransformChannels = ControlHandle->GetFloatChannels(InSection);
+		const TArrayView<FMovieSceneDoubleChannel*> DoubleTransformChannels = ControlHandle->GetDoubleChannels(InSection);
+
+		// move them
+		if (FloatTransformChannels.Num() > 0)
+		{
+			FMovieSceneConstraintChannelHelper::MoveTransformKeys(FloatTransformChannels, InCurrentFrame, InNextFrame);
+		}
+		else
+		{
+			FMovieSceneConstraintChannelHelper::MoveTransformKeys(DoubleTransformChannels, InCurrentFrame, InNextFrame);
+		}
+	}
+
+}
+
+bool FMovieSceneConstraintChannelHelper::SmartConstraintKey(
+	const TSharedPtr<ISequencer>& InSequencer,
+	UTickableTransformConstraint* InConstraint, 
+	const TOptional<bool>& InOptActive, 
+	const TOptional<FFrameNumber>& InOptFrameTime)
+{
+	if (!InSequencer.IsValid() || !InSequencer->GetFocusedMovieSceneSequence())
+	{
+		return false;
+	}
+	
+	ITransformConstraintChannelInterface* Interface = GetHandleInterface(InConstraint->ChildTRSHandle);
+	if (!Interface)
+	{
+		return false;
+	}
+	
+	FFrameNumber Time;
+	if (InOptFrameTime.IsSet())
+	{
+		Time = InOptFrameTime.GetValue();
+	}
+	else
+	{
+		const FFrameRate TickResolution = InSequencer->GetFocusedTickResolution();
+		const FFrameTime FrameTime = InSequencer->GetLocalTime().ConvertTo(TickResolution);
+		Time = FrameTime.GetFrame();
+	}
+
+	const bool bSucceeded = Interface->SmartConstraintKey(InConstraint, InOptActive, Time, InSequencer);
+
+	//todo need to revisit this to see if we need to create this even if we don't set a key, it's harmless I think either way
+	CreateBindingIDForHandle(InSequencer, InConstraint->ChildTRSHandle);
+	CreateBindingIDForHandle(InSequencer, InConstraint->ParentTRSHandle);
+
+	return bSucceeded;
+}
+
+void FMovieSceneConstraintChannelHelper::Compensate(
+	const TSharedPtr<ISequencer>& InSequencer,
+	const UTickableTransformConstraint* InConstraint,
+	const TOptional<FFrameNumber>& InOptTime)
+{
+	if (!InSequencer.IsValid() || !InSequencer->GetFocusedMovieSceneSequence())
+	{
+		return;
+	}
+
+	const TObjectPtr<UTransformableHandle>& Handle = InConstraint->ChildTRSHandle;
+	
+	ITransformConstraintChannelInterface* Interface = GetHandleInterface(Handle);
+	if (!Interface)
+	{
+		return;
+	}
+
+	IMovieSceneConstrainedSection* Section = Cast<IMovieSceneConstrainedSection>(Interface->GetHandleConstraintSection(Handle, InSequencer));
+	const UWorld* World = Interface->GetHandleWorld(Handle);
+
+	if (!Section || !IsValid(World))
+	{
+		return;
+	}
+
+	CompensateIfNeeded(InSequencer, Section, InOptTime, Handle->GetHash());
+}
+
+void FMovieSceneConstraintChannelHelper::CompensateIfNeeded(
+	const TSharedPtr<ISequencer>& InSequencer,
+	IMovieSceneConstrainedSection* ConstraintSection,
+	const TOptional<FFrameNumber>& OptionalTime,
+	const int32 InChildHash)
+{
+	if (bDoNotCompensate)
+	{
+		return;
+	}
+
+	TGuardValue<bool> CompensateGuard(bDoNotCompensate, true);
+
+	// Frames to compensate
+	TArray<FFrameNumber> OptionalTimeArray;
+	if (OptionalTime.IsSet())
+	{
+		OptionalTimeArray.Add(OptionalTime.GetValue());
+	}
+
+	auto GetConstraintTimesToCompensate = [&OptionalTimeArray](const FConstraintAndActiveChannel& Channel)->TArrayView<const FFrameNumber>
+	{
+		if (OptionalTimeArray.IsEmpty())
+		{
+			return Channel.ActiveChannel.GetData().GetTimes();
+		}
+		return OptionalTimeArray;
+	};
+
+	// gather all transform constraints' channels for
+	const TArray<FConstraintAndActiveChannel>& ConstraintChannels = ConstraintSection->GetConstraintsChannels();
+	TArray<FConstraintAndActiveChannel> TransformConstraintsChannels;
+	Algo::CopyIf(ConstraintChannels, TransformConstraintsChannels,
+		[InChildHash](const FConstraintAndActiveChannel& InChannel)
+		{
+			if (!InChannel.Constraint.IsValid())
+			{
+				return false;
+			}
+
+			if ((InChildHash != INDEX_NONE) && (InChannel.Constraint->GetTargetHash() != InChildHash))
+			{
+				return false;
+			}
+
+			const UTickableTransformConstraint* Constraint = Cast<UTickableTransformConstraint>(InChannel.Constraint.Get());
+			return Constraint && (Constraint->GetTargetHash() == InChildHash) && Constraint->NeedsCompensation();
+		}
+	);
+
+	// we only need to treat one single constraint per child as FCompensationEvaluator::ComputeCompensation will
+	// compensate within the last active constraint's space
+	using CompensationData = TPair< UTickableTransformConstraint*, TArray<FFrameNumber> >;
+	TArray< CompensationData > ToCompensate;
+
+	// store constraints and times where compensation is needed 
+	for (const FConstraintAndActiveChannel& Channel : TransformConstraintsChannels)
+	{
+		const TArrayView<const FFrameNumber> FramesToCompensate = GetConstraintTimesToCompensate(Channel);
+		for (const FFrameNumber& Time : FramesToCompensate)
+		{
+			const FFrameNumber TimeMinusOne(Time - 1);
+
+			bool CurrentValue = false, PreviousValue = false;
+			Channel.ActiveChannel.Evaluate(TimeMinusOne, PreviousValue);
+			Channel.ActiveChannel.Evaluate(Time, CurrentValue);
+
+			if (CurrentValue != PreviousValue) //if they are the same no need to do anything
+			{
+				UTickableTransformConstraint* Constraint = Cast<UTickableTransformConstraint>(Channel.Constraint.Get());
+
+				// is the child already in that array?
+				int32 DataIndex = ToCompensate.IndexOfByPredicate([Constraint](const CompensationData& InData)
+				{
+					return InData.Key->GetTargetHash() == Constraint->GetTargetHash();
+				});
+
+				// if not, add the constraint
+				if (DataIndex == INDEX_NONE)
+				{
+					DataIndex = ToCompensate.Emplace(Constraint, TArray<FFrameNumber>() );
+				}
+
+				// store the time it needs to be compensated at
+				TArray<FFrameNumber>& Times = ToCompensate[DataIndex].Value;
+				Times.AddUnique(Time);
+			}
+		}
+	}
+
+	// compensate
+	bool bNeedsEvaluation = false;
+	for (const CompensationData& Data: ToCompensate)
+	{
+		UTickableTransformConstraint* Constraint = Data.Key;
+		const TObjectPtr<UTransformableHandle>& Handle = Constraint->ChildTRSHandle;
+		if (ITransformConstraintChannelInterface* Interface = GetHandleInterface(Handle))
+		{
+			UWorld* World = Interface->GetHandleWorld(Handle);
+			
+			FCompensationEvaluator Evaluator(Constraint);
+			const EMovieSceneTransformChannel ChannelsToKey = Constraint->GetChannelsToKey();
+			for (const FFrameNumber& Time : Data.Value)
+			{
+				// compute transform to set
+				// if switching from active to inactive then we must add a key at T-1 in the constraint space
+				// if switching from inactive to active then we must add a key at T-1 in the previous constraint or parent space
+				Evaluator.ComputeCompensation(World, InSequencer, Time);
+				const TArray<FTransform>& LocalTransforms = Evaluator.ChildLocals;
+
+				const FFrameNumber TimeMinusOne(Time - 1);
+				Interface->AddHandleTransformKeys(InSequencer, Handle, { TimeMinusOne }, LocalTransforms, ChannelsToKey);
+
+				bNeedsEvaluation = true;
+			}
+		}
+	}
+
+	if (bNeedsEvaluation)
+	{
+		InSequencer->ForceEvaluate();
+	}
+}
+
+ITransformConstraintChannelInterface* FMovieSceneConstraintChannelHelper::GetHandleInterface(const UTransformableHandle* InHandle)
+{
+	if (!IsValid(InHandle) || !InHandle->IsValid())
+	{
+		return nullptr;
+	}
+
+	const FConstraintChannelInterfaceRegistry& InterfaceRegistry = FConstraintChannelInterfaceRegistry::Get();	
+	return InterfaceRegistry.FindConstraintChannelInterface(InHandle->GetClass());
+}
+
+void FMovieSceneConstraintChannelHelper::CreateBindingIDForHandle(const TSharedPtr<ISequencer>& InSequencer, UTransformableHandle* InHandle)
+{
+	if (InHandle == nullptr || InSequencer.IsValid() == false)
+	{
+		return;
+	}
+	
+	if (const USceneComponent* SceneComponent = Cast<USceneComponent>(InHandle->GetTarget().Get()))
+	{
+		if (AActor* Actor = SceneComponent->GetTypedOuter<AActor>())
+		{
+			TOptional<FMovieSceneSpawnableAnnotation> Spawnable = FMovieSceneSpawnableAnnotation::Find(Actor);
+			if (Spawnable.IsSet())
+			{
+				// Check whether the spawnable is underneath the current sequence, if so, we can remap it to a local sequence ID
+				InHandle->ConstraintBindingID = UE::MovieScene::FRelativeObjectBindingID(InSequencer->GetFocusedTemplateID(), Spawnable->SequenceID, Spawnable->ObjectBindingID,
+					*(InSequencer.Get()));
+			}
+			else
+			{
+				const FGuid Guid = InSequencer->GetHandleToObject(Actor, false); //don't create it???
+				InHandle->ConstraintBindingID = UE::MovieScene::FRelativeObjectBindingID(Guid);
+			}
+		}
+	}
+}
\ No newline at end of file
diff --git a/Engine/Source/Editor/MovieSceneTools/Private/Constraints/TransformConstraintChannelInterface.cpp b/Engine/Source/Editor/MovieSceneTools/Private/Constraints/TransformConstraintChannelInterface.cpp
new file mode 100644
--- /dev/null
+++ b/Engine/Source/Editor/MovieSceneTools/Private/Constraints/TransformConstraintChannelInterface.cpp
@@ -0,0 +1,53 @@
+﻿// Copyright Epic Games, Inc. All Rights Reserved.
+
+#include "Constraints/TransformConstraintChannelInterface.h"
+#include "Sections/MovieSceneConstrainedSection.h"
+#include "ConstraintChannel.h"
+
+// IMovieSceneConstrainedSection* ITransformConstraintChannelInterface::GetHandleConstraintSection(const UTransformableHandle* InHandle, const TSharedPtr<ISequencer>& InSequencer)
+// {
+// 	return Cast<IMovieSceneConstrainedSection>(GetHandleSection(InHandle, InSequencer));
+// }
+
+bool ITransformConstraintChannelInterface::CanAddKey(const FMovieSceneConstraintChannel& ActiveChannel, const FFrameNumber& InTime, bool& ActiveValue)
+{
+	const TMovieSceneChannelData<const bool> ChannelData = ActiveChannel.GetData();
+	const TArrayView<const FFrameNumber> Times = ChannelData.GetTimes();
+	if (Times.IsEmpty())
+	{
+		ActiveValue = true;
+		return true;
+	}
+
+	const TArrayView<const bool> Values = ChannelData.GetValues();
+	if (InTime < Times[0])
+	{
+		if (!Values[0])
+		{
+			ActiveValue = true;
+			return true;
+		}
+		return false;
+	}
+	
+	if (InTime > Times.Last())
+	{
+		ActiveValue = !Values.Last();
+		return true;
+	}
+
+	return false;
+}
+
+FConstraintChannelInterfaceRegistry& FConstraintChannelInterfaceRegistry::Get()
+{
+	static FConstraintChannelInterfaceRegistry Singleton;
+	return Singleton;
+}
+
+ITransformConstraintChannelInterface* FConstraintChannelInterfaceRegistry::FindConstraintChannelInterface(const UClass* InClass) const
+{
+	const TUniquePtr<ITransformConstraintChannelInterface>* Interface = HandleToInterfaceMap.Find(InClass);
+	ensureMsgf(Interface, TEXT("No constraint channel interface found for class %s. Did you call RegisterConstraintChannelInterface<> for that class?"), *InClass->GetName());
+	return Interface ? Interface->Get() : nullptr;
+}
\ No newline at end of file
diff --git a/Engine/Source/Editor/MovieSceneTools/Private/MovieSceneConstraintChannelHelper.cpp b/Engine/Source/Editor/MovieSceneTools/Private/MovieSceneConstraintChannelHelper.cpp
deleted file mode 100644
diff --git a/Engine/Source/Editor/MovieSceneTools/Private/MovieSceneToolsModule.cpp b/Engine/Source/Editor/MovieSceneTools/Private/MovieSceneToolsModule.cpp
--- a/Engine/Source/Editor/MovieSceneTools/Private/MovieSceneToolsModule.cpp
+++ b/Engine/Source/Editor/MovieSceneTools/Private/MovieSceneToolsModule.cpp
@@ -110,6 +110,10 @@
 #include "EditorAnimUtils.h"
 #include "Animation/AnimSequence.h"
 
+#include "TransformableHandle.h"
+#include "Constraints/ComponentConstraintChannelInterface.h"
+#include "Constraints/TransformConstraintChannelInterface.h"
+
 #define LOCTEXT_NAMESPACE "FMovieSceneToolsModule"
 
 TAutoConsoleVariable<bool> CVarDuplicateLinkedAnimSequence(TEXT("Sequencer.DuplicateLinkedAnimSequence"), false, TEXT("When true when we duplicate a level sequence that has a linked anim sequence it will duplicate and link the anim sequencel, if false we leave any link alone."));
@@ -259,6 +263,10 @@ void FMovieSceneToolsModule::StartupModule()
 		NSLOCTEXT("SkeletalAnimationTrackEditorMode", "SkelAnimTrackEditMode", "Skeletal Anim Track Mode"),
 		FSlateIcon(),
 		false);
+
+	// register UTransformableComponentHandle animatable interface
+	FConstraintChannelInterfaceRegistry& ConstraintChannelInterfaceRegistry = FConstraintChannelInterfaceRegistry::Get();
+	ConstraintChannelInterfaceRegistry.RegisterConstraintChannelInterface<UTransformableComponentHandle>(MakeUnique<FComponentConstraintChannelInterface>());
 }
 
 void FMovieSceneToolsModule::ShutdownModule()
diff --git a/Engine/Source/Editor/MovieSceneTools/Private/TrackEditors/TransformTrackEditor.cpp b/Engine/Source/Editor/MovieSceneTools/Private/TrackEditors/TransformTrackEditor.cpp
--- a/Engine/Source/Editor/MovieSceneTools/Private/TrackEditors/TransformTrackEditor.cpp
+++ b/Engine/Source/Editor/MovieSceneTools/Private/TrackEditors/TransformTrackEditor.cpp
@@ -38,7 +38,7 @@
 #include "Editor.h"
 #include "LevelEditorViewport.h"
 #include "TransformConstraint.h"
-#include "MovieSceneConstraintChannelHelper.h"
+#include "Constraints/MovieSceneConstraintChannelHelper.h"
 #include "Misc/TransactionObjectEvent.h"
 #include "Engine/Selection.h"
 
diff --git a/Engine/Source/Editor/MovieSceneTools/Public/ConstraintChannelHelper.inl b/Engine/Source/Editor/MovieSceneTools/Public/ConstraintChannelHelper.inl
deleted file mode 100644
diff --git a/Engine/Source/Editor/MovieSceneTools/Public/Constraints/MovieSceneConstraintChannelHelper.h b/Engine/Source/Editor/MovieSceneTools/Public/Constraints/MovieSceneConstraintChannelHelper.h
new file mode 100644
--- /dev/null
+++ b/Engine/Source/Editor/MovieSceneTools/Public/Constraints/MovieSceneConstraintChannelHelper.h
@@ -0,0 +1,129 @@
+﻿// Copyright Epic Games, Inc. All Rights Reserved.
+
+#pragma once
+
+#include "Templates/SharedPointer.h"
+#include "UObject/ObjectPtr.h"
+
+class IMovieSceneConstrainedSection;
+struct ITransformConstraintChannelInterface;
+class UWorld;
+class UTickableConstraint;
+class ISequencer;
+class UMovieSceneSection;
+class UTickableTransformConstraint;
+class UTransformableHandle;
+struct FMovieSceneConstraintChannel;
+struct FFrameNumber;
+
+
+struct MOVIESCENETOOLS_API FCompensationEvaluator
+{
+public:
+	TArray<FTransform> ChildLocals;
+	TArray<FTransform> ChildGlobals;
+	TArray<FTransform> SpaceGlobals;
+
+	FCompensationEvaluator(UTickableTransformConstraint* InConstraint);
+
+	void ComputeLocalTransforms(UWorld* InWorld, const TSharedPtr<ISequencer>& InSequencer, const TArray<FFrameNumber>& InFrames, const bool bToActive);
+	void ComputeLocalTransformsBeforeDeletion(UWorld* InWorld, const TSharedPtr<ISequencer>& InSequencer, const TArray<FFrameNumber>& InFrames);
+	void ComputeCompensation(UWorld* InWorld, const TSharedPtr<ISequencer>& InSequencer, const FFrameNumber& InTime);
+	void ComputeLocalTransformsForBaking(UWorld* InWorld, const TSharedPtr<ISequencer>& InSequencer, const TArray<FFrameNumber>& InFrames);
+private:
+
+	TArray< TObjectPtr<UTickableConstraint> > GetHandleTransformConstraints(UWorld* InWorld) const;
+
+	UTickableTransformConstraint* Constraint = nullptr;
+	UTransformableHandle* Handle = nullptr;
+};
+
+struct MOVIESCENETOOLS_API FMovieSceneConstraintChannelHelper
+{
+public:
+
+	/** Adds an active key if needed and does the compensation when switching. Will use the optional active and time if set. 
+	Will return true if key is actually set, may not be if the value is the same.*/
+	static bool SmartConstraintKey(
+		const TSharedPtr<ISequencer>& InSequencer,
+		UTickableTransformConstraint* InConstraint, 
+		const TOptional<bool>& InOptActive,
+		const TOptional<FFrameNumber>& InOptFrameTime);
+	
+	/** Compensate transform on handles when a constraint switches state. */
+	static void Compensate(
+		const TSharedPtr<ISequencer>& InSequencer,
+		const UTickableTransformConstraint* InConstraint,
+		const TOptional<FFrameNumber>& InOptTime);
+	
+	static void CompensateIfNeeded(
+		const TSharedPtr<ISequencer>& InSequencer,
+		IMovieSceneConstrainedSection* Section,
+		const TOptional<FFrameNumber>& OptionalTime,
+		const int32 InChildHash = INDEX_NONE);
+		
+	static void HandleConstraintRemoved(
+		UTickableConstraint* InConstraint,
+		const FMovieSceneConstraintChannel* InChannel,
+		const TSharedPtr<ISequencer>& InSequencer,
+		UMovieSceneSection* InSection);
+
+	static void HandleConstraintKeyDeleted(
+		UTickableTransformConstraint* InConstraint,
+		const FMovieSceneConstraintChannel* InConstraintChannel,
+		const TSharedPtr<ISequencer>& InSequencer,
+		UMovieSceneSection* InSection,
+		const FFrameNumber& InTime);
+
+	static void HandleConstraintKeyMoved(
+		const UTickableTransformConstraint* InConstraint,
+		const FMovieSceneConstraintChannel* InConstraintChannel,
+		UMovieSceneSection* InSection,
+		const FFrameNumber& InCurrentFrame, const FFrameNumber& InNextFrame);
+
+	/** @todo documentation. */
+	template<typename ChannelType>
+	static void GetFramesToCompensate(
+		const FMovieSceneConstraintChannel& InActiveChannel,
+		const bool InActiveValueToBeSet,
+		const FFrameNumber& InTime,
+		const TArrayView<ChannelType*>& InChannels,
+		TArray<FFrameNumber>& OutFramesAfter);
+
+	/** @todo documentation. */
+	template< typename ChannelType >
+	static void GetFramesAfter(
+		const FMovieSceneConstraintChannel& InActiveChannel,
+		const FFrameNumber& InTime,
+		const TArrayView<ChannelType*>& InChannels,
+		TArray<FFrameNumber>& OutFrames);
+
+	/** @todo documentation. */
+	template< typename ChannelType >
+	static void GetFramesWithinActiveState(
+		const FMovieSceneConstraintChannel& InActiveChannel,
+		const TArrayView<ChannelType*>& InChannels,
+		TArray<FFrameNumber>& OutFrames);
+
+	/** @todo documentation. */
+	template< typename ChannelType >
+	static void MoveTransformKeys(
+		const TArrayView<ChannelType*>& InChannels,
+		const FFrameNumber& InCurrentTime,
+		const FFrameNumber& InNextTime);
+
+	/** @todo documentation. */
+	template< typename ChannelType >
+	static void DeleteTransformKeys(
+		const TArrayView<ChannelType*>& InChannels,
+		const FFrameNumber& InTime);
+
+	static bool bDoNotCompensate;
+
+private:
+	/** Get the animatable interface for that handle if registered. */
+	static ITransformConstraintChannelInterface* GetHandleInterface(const UTransformableHandle* InHandle);
+	
+	/** For the given handle create any movie scene binding for it based upon the current sequencer that's open*/
+	static void CreateBindingIDForHandle(const TSharedPtr<ISequencer>& InSequencer, UTransformableHandle* InHandle);
+};
diff --git a/Engine/Source/Editor/MovieSceneTools/Public/Constraints/MovieSceneConstraintChannelHelper.inl b/Engine/Source/Editor/MovieSceneTools/Public/Constraints/MovieSceneConstraintChannelHelper.inl
new file mode 100644
--- /dev/null
+++ b/Engine/Source/Editor/MovieSceneTools/Public/Constraints/MovieSceneConstraintChannelHelper.inl
@@ -0,0 +1,369 @@
+﻿// Copyright Epic Games, Inc. All Rights Reserved.
+
+#pragma once
+
+#include "ConstraintChannel.h"
+#include "MovieSceneConstraintChannelHelper.h"
+#include "Channels/MovieSceneCurveChannelCommon.h"
+#include "Channels/MovieSceneChannelData.h"
+#include "Channels/MovieSceneChannelProxy.h"
+#include "MovieSceneSection.h"
+#include "Algo/Unique.h"
+
+template<typename ChannelType>
+void FMovieSceneConstraintChannelHelper::GetFramesToCompensate(
+	const FMovieSceneConstraintChannel& InActiveChannel,
+	const bool InActiveValueToBeSet,
+	const FFrameNumber& InTime,
+	const TArrayView<ChannelType*>& InChannels,
+	TArray<FFrameNumber>& OutFramesAfter)
+{
+	using ChannelValueType = typename ChannelType::ChannelValueType;
+
+	const bool bHasKeys = (InActiveChannel.GetNumKeys() > 0);
+	
+	OutFramesAfter.Reset();
+
+	// add the current frame
+	OutFramesAfter.Add(InTime);
+
+	// add the next frames that need transform compensation 
+	for (const ChannelType* InChannel: InChannels)
+	{
+		const TMovieSceneChannelData<const ChannelValueType> ChannelData = InChannel->GetData();
+		const TArrayView<const FFrameNumber> Times = ChannelData.GetTimes();
+		if (!Times.IsEmpty())
+		{
+			// look for the first next key frame for this channel 
+			const int32 NextTimeIndex = Algo::UpperBound(Times, InTime);
+			if (Times.IsValidIndex(NextTimeIndex))
+			{
+				// store the time while the state is different
+				for (int32 Index = NextTimeIndex; Index < Times.Num(); ++Index)
+				{
+					if (!bHasKeys)
+					{
+						OutFramesAfter.Add(Times[Index]);
+					}
+					else
+					{
+						bool NextValue = false; InActiveChannel.Evaluate(Times[Index], NextValue);
+						if (NextValue == InActiveValueToBeSet)
+						{
+							break;
+						}
+						OutFramesAfter.Add(Times[Index]);
+					}
+				}
+			}
+		}
+	}
+
+	// uniqueness
+	OutFramesAfter.Sort();
+	OutFramesAfter.SetNum(Algo::Unique(OutFramesAfter));
+}
+
+template< typename ChannelType >
+void FMovieSceneConstraintChannelHelper::GetFramesAfter(
+	const FMovieSceneConstraintChannel& InActiveChannel,
+	const FFrameNumber& InTime,
+	const TArrayView<ChannelType*>& InChannels,
+	TArray<FFrameNumber>& OutFrames)
+{
+	using ChannelValueType = typename ChannelType::ChannelValueType;
+	
+	OutFrames.Reset();
+	
+	const TMovieSceneChannelData<const bool> ConstraintChannelData = InActiveChannel.GetData();
+	const int32 KeyIndex = ConstraintChannelData.FindKey(InTime);
+	if (!ConstraintChannelData.GetTimes().IsValidIndex(KeyIndex))
+	{
+		return;
+	}
+
+	const bool CurrentValue = ConstraintChannelData.GetValues()[KeyIndex];
+	
+	// compute last frame to compensate
+	auto GetEndOfCompensationTime = [KeyIndex](const TMovieSceneChannelData<const bool>& InData)
+	{
+		const TArrayView<const bool> Values = InData.GetValues();
+		const TArrayView<const FFrameNumber> Times = InData.GetTimes();
+
+		const bool CurrentValue = Values[KeyIndex];
+		for (int32 NextIndex = KeyIndex+1; NextIndex < Times.Num(); ++NextIndex)
+		{
+			if (Values[NextIndex] != CurrentValue)
+			{
+				return TOptional<FFrameNumber>(Times[NextIndex]);
+			}
+		}
+		return TOptional<FFrameNumber>();
+	};
+	const TOptional<FFrameNumber> EndOfCompensationTime = GetEndOfCompensationTime(ConstraintChannelData);
+
+	const bool bHasEndTime = EndOfCompensationTime.IsSet();
+	
+	// add the current frame
+	OutFrames.Add(InTime);
+
+	// add the next frames that need transform compensation 
+	for (const ChannelType* InChannel: InChannels)
+	{
+		const TMovieSceneChannelData<const ChannelValueType> ChannelData = InChannel->GetData();
+		const TArrayView<const FFrameNumber> Times = ChannelData.GetTimes();
+		if (!Times.IsEmpty())
+		{
+			// look for the first next key frame for this channel 
+			const int32 NextTimeIndex = Algo::UpperBound(Times, InTime);
+			if (Times.IsValidIndex(NextTimeIndex))
+			{
+				// store the time while the state is different
+				for (int32 Index = NextTimeIndex; Index < Times.Num(); ++Index)
+				{
+					if (!bHasEndTime || Times[Index] < EndOfCompensationTime.GetValue() )
+					{
+						OutFrames.Add(Times[Index]);
+					}
+				}
+			}
+		}
+	}
+
+	// uniqueness
+	OutFrames.Sort();
+	OutFrames.SetNum(Algo::Unique(OutFrames));
+}
+
+template< typename ChannelType >
+void FMovieSceneConstraintChannelHelper::GetFramesWithinActiveState(
+	const FMovieSceneConstraintChannel& InActiveChannel,
+	const TArrayView<ChannelType*>& InChannels,
+	TArray<FFrameNumber>& OutFrames)
+{
+	using ChannelValueType = typename ChannelType::ChannelValueType;
+	
+	OutFrames.Reset();
+	
+	const TMovieSceneChannelData<const bool> ConstraintChannelData = InActiveChannel.GetData();
+	const TArrayView<const FFrameNumber>& ActiveTimes = ConstraintChannelData.GetTimes();
+	if (ActiveTimes.IsEmpty())
+	{
+		return;
+	}
+
+	const FFrameNumber& FirstTime = ActiveTimes[0];
+	const FFrameNumber& LastTime = ActiveTimes.Last();
+	
+	// add active times
+	for (int32 Index = 0; Index < ActiveTimes.Num(); ++Index)
+	{
+		OutFrames.Add(ActiveTimes[Index]);
+	}
+
+	const bool bIsLastStateInactive = ConstraintChannelData.GetValues().Last() == false; 
+
+	// add frames where the constraint is active 
+	for (const ChannelType* InChannel: InChannels)
+	{
+		const TMovieSceneChannelData<const ChannelValueType> ChannelData = InChannel->GetData();
+		const TArrayView<const FFrameNumber> Times = ChannelData.GetTimes();
+		if (!Times.IsEmpty())
+		{
+			// look for the first next key frame for this channel 
+			const int32 NextTimeIndex = Algo::UpperBound(Times,FirstTime);
+			if (Times.IsValidIndex(NextTimeIndex))
+			{
+				// store the time is the state is active
+				for (int32 Index = NextTimeIndex; Index < Times.Num(); ++Index)
+				{
+					bool bIsActive = false;
+					InActiveChannel.Evaluate(Times[Index], bIsActive);
+					if (bIsActive)
+					{
+						OutFrames.Add(Times[Index]);
+					}
+
+					if (bIsLastStateInactive && Times[Index] > LastTime)
+					{
+						break;
+					}
+				}
+			}
+		}
+	}
+
+	// uniqueness
+	OutFrames.Sort();
+	OutFrames.SetNum(Algo::Unique(OutFrames));
+}
+
+template< typename ChannelType >
+void FMovieSceneConstraintChannelHelper::MoveTransformKeys(
+	const TArrayView<ChannelType*>& InChannels,
+	const FFrameNumber& InCurrentTime,
+	const FFrameNumber& InNextTime)
+{
+	const FFrameNumber Delta = InNextTime - InCurrentTime;
+	if (Delta == 0)
+	{
+		return;
+	}
+	
+	for (ChannelType* Channel: InChannels)
+	{
+		TMovieSceneChannelData<typename ChannelType::ChannelValueType> Data = Channel->GetData();
+		const TArrayView<const FFrameNumber> Times = Data.GetTimes();
+		const int32 NumTimes = Times.Num();
+
+		if (Delta > 0) //if we are moving keys positively in time we start from end frames and move them so we can use indices
+		{
+			for (int32 KeyIndex = NumTimes - 1; KeyIndex >= 0; --KeyIndex)
+			{
+				const FFrameNumber& Frame = Times[KeyIndex];
+				const FFrameNumber AbsDiff = FMath::Abs(Frame - InCurrentTime);
+				if (AbsDiff<= 1)
+				{
+					Data.MoveKey(KeyIndex, Frame + Delta);
+				}
+			}
+		}
+		else
+		{
+			for (int32 KeyIndex = 0; KeyIndex < NumTimes; ++KeyIndex)
+			{
+				const FFrameNumber& Frame = Times[KeyIndex];
+				const FFrameNumber AbsDiff = FMath::Abs( Frame - InCurrentTime);
+				if (AbsDiff <= 1)
+				{
+					Data.MoveKey(KeyIndex, Frame + Delta);
+				}
+			}
+		}
+	}
+}
+
+template< typename ChannelType >
+void FMovieSceneConstraintChannelHelper::DeleteTransformKeys(
+	const TArrayView<ChannelType*>& InChannels,
+	const FFrameNumber& InTime)
+{
+	for (ChannelType* Channel: InChannels)
+	{
+		TMovieSceneChannelData<typename ChannelType::ChannelValueType> Data = Channel->GetData();
+		const TArrayView<const FFrameNumber> Times = Data.GetTimes();
+		
+		const int32 KeyIndex = Algo::LowerBound(Times, InTime);
+		if (Times.IsValidIndex(KeyIndex) && Times[KeyIndex] == InTime)
+		{
+			Data.RemoveKey(KeyIndex);
+		}
+	}
+}
+
+template<typename ChannelType>
+void EvaluateTangentAtThisTime(int32 ChannelIndex, int32 NumChannels,
+	UMovieSceneSection* Section,
+	FFrameNumber Time,
+	TArray<FMovieSceneTangentData>& OutTangents)
+{
+	using ChannelValueType = typename ChannelType::ChannelValueType;
+
+	// NOTE this might be moved to FMovieSceneFloatChannel
+	auto EvaluateTangent = [](const ChannelType* InChannel, FFrameNumber InTime)
+	{
+		const TMovieSceneChannelData<const ChannelValueType> ChannelInterface = InChannel->GetData();
+		const TArrayView<const FFrameNumber> Times = ChannelInterface.GetTimes();
+
+		// Need at least two keys to evaluate a derivative
+		if (Times.Num() < 2)
+		{
+			static const FMovieSceneTangentData DefaultTangent;
+			return DefaultTangent;
+		}
+
+		const TArrayView<const ChannelValueType> Values = ChannelInterface.GetValues();
+
+		// look around to get the closest key tangent if fairly close
+		const int32 Tolerance = static_cast<int32>(InChannel->GetTickResolution().AsDecimal() * 0.01);
+		// NOTE FindKey might return Times.Num() (see AlgoImpl::LowerBoundInternal)
+		const int32 ExistingIndex = ChannelInterface.FindKey(InTime, Tolerance);
+		if (Times.IsValidIndex(ExistingIndex) && Values.IsValidIndex(ExistingIndex))
+		{
+			const FFrameNumber& Time = Times[ExistingIndex];
+			const int32 DiffToKey = InTime.Value - Time.Value;
+
+			// if the closest key is within a threshold, we use it's tangent directly instead of computing one
+			if (FMath::Abs(DiffToKey) <= Tolerance)
+			{
+				const ChannelValueType Value = Values[ExistingIndex];
+				if (DiffToKey == 0)
+				{
+					return Value.Tangent;
+				}
+
+				FMovieSceneTangentData Tangent = Value.Tangent;
+				if (DiffToKey < 0) // pretty close to the next key
+				{
+					Tangent.LeaveTangent = 0.f;
+				}
+				else // pretty close to the previous key
+				{
+					Tangent.ArriveTangent = 0.f;
+				}
+				return Tangent;
+			}
+		}
+
+		// compute tangent using central difference
+		// NOTE we may wanna compute a backward / forward difference instead
+		const int32 Delta = static_cast<int32>(InChannel->GetTickResolution().AsDecimal() * 0.1);
+
+		float PrevValue = 0.f; InChannel->Evaluate(InTime - Delta, PrevValue);
+		float NextValue = 0.f; InChannel->Evaluate(InTime + Delta, NextValue);
+		const float TangentValue = (NextValue - PrevValue) / (2.f * Delta);
+
+		FMovieSceneTangentData TangentData;
+		TangentData.ArriveTangent = TangentValue;
+		TangentData.LeaveTangent = TangentValue;
+
+		return TangentData;
+	};
+
+
+	// compute and store tangents
+	OutTangents.SetNum(NumChannels);
+
+	const TArrayView<ChannelType*> Channels = Section->GetChannelProxy().GetChannels<ChannelType>();
+	for (int32 Index = 0; Index < NumChannels; ++Index, ++ChannelIndex)
+	{
+		OutTangents[Index] = EvaluateTangent(Channels[ChannelIndex], Time);
+	}
+}
+
+// NOTE we may pass an enum to tell which tangent we wanna set (arrive, leave, both)
+template<typename ChannelType>
+void SetTangentsAtThisTime(int32 ChannelIndex,
+	int32 NumChannels,
+	UMovieSceneSection* Section,
+	FFrameNumber Time,
+	const TArray<FMovieSceneTangentData>& InTangents)
+{
+	using ChannelValueType = typename ChannelType::ChannelValueType;
+
+	const TArrayView<ChannelType*> Channels = Section->GetChannelProxy().GetChannels<ChannelType>();
+
+	for (int32 Index = 0; Index < NumChannels; ++Index, ++ChannelIndex)
+	{
+		TMovieSceneChannelData<ChannelValueType> ChannelInterface = Channels[ChannelIndex]->GetData();
+		TArrayView<ChannelValueType> Values = ChannelInterface.GetValues();
+		const int32 KeyIndex = ChannelInterface.FindKey(Time);
+		if (KeyIndex != INDEX_NONE)
+		{
+			ChannelValueType& Value = Values[KeyIndex];
+			Value.Tangent.ArriveTangent = InTangents[Index].ArriveTangent;
+			Value.Tangent.LeaveTangent = InTangents[Index].LeaveTangent;
+			Value.TangentMode = RCTM_Break;
+		}
+	}
+}
\ No newline at end of file
diff --git a/Engine/Source/Editor/MovieSceneTools/Public/Constraints/TransformConstraintChannelInterface.h b/Engine/Source/Editor/MovieSceneTools/Public/Constraints/TransformConstraintChannelInterface.h
new file mode 100644
--- /dev/null
+++ b/Engine/Source/Editor/MovieSceneTools/Public/Constraints/TransformConstraintChannelInterface.h
@@ -0,0 +1,80 @@
+﻿// Copyright Epic Games, Inc. All Rights Reserved.
+
+#pragma once
+
+#include "Templates/SharedPointer.h"
+#include "Templates/UniquePtr.h"
+#include "Containers/Map.h"
+
+struct FMovieSceneConstraintChannel;
+class IMovieSceneConstrainedSection;
+class UWorld;
+class ISequencer;
+class UMovieSceneSection;
+class UTickableTransformConstraint;
+class UTransformableHandle;
+struct FFrameNumber;
+enum class EMovieSceneTransformChannel : uint32;
+
+/**
+* Abstract interface that defines animatable capabilities for transformable handles
+*/
+
+struct ITransformConstraintChannelInterface
+{
+	virtual ~ITransformConstraintChannelInterface() = default;
+
+	/** Get the section where the channels live from the object that the handle wraps. */
+	virtual UMovieSceneSection* GetHandleSection(const UTransformableHandle* InHandle, const TSharedPtr<ISequencer>& InSequencer) = 0;
+	virtual UMovieSceneSection* GetHandleConstraintSection(const UTransformableHandle* InHandle, const TSharedPtr<ISequencer>& InSequencer) = 0;
+	// MOVIESCENETOOLS_API IMovieSceneConstrainedSection* GetHandleConstraintSection(const UTransformableHandle* InHandle, const TSharedPtr<ISequencer>& InSequencer);
+
+	/** Get the world from the object that the handle wraps. */
+	virtual UWorld* GetHandleWorld(UTransformableHandle* InHandle) = 0;
+
+	/** Add an active/inactive key to the constraint channel if needed and does the transform compensation on the transform channels. */
+	virtual bool SmartConstraintKey(
+		UTickableTransformConstraint* InConstraint, const TOptional<bool>& InOptActive,
+		const FFrameNumber& InTime, const TSharedPtr<ISequencer>& InSequencer) = 0;
+
+	/** Add keys on the transform channels of the the object that the handle wraps. */
+	virtual void AddHandleTransformKeys(
+		const TSharedPtr<ISequencer>& InSequencer,
+		const UTransformableHandle* InHandle,
+		const TArray<FFrameNumber>& InFrames,
+		const TArray<FTransform>& InLocalTransforms,
+		const EMovieSceneTransformChannel& InChannels) = 0;
+
+protected:
+	static MOVIESCENETOOLS_API bool CanAddKey(const FMovieSceneConstraintChannel& ActiveChannel, const FFrameNumber& InTime, bool& ActiveValue);
+};
+
+/**
+* Handle animatable interface registry
+*/
+
+class FConstraintChannelInterfaceRegistry
+{
+public:
+	~FConstraintChannelInterfaceRegistry() = default;
+
+	/** Get the singleton registry object */
+	static MOVIESCENETOOLS_API FConstraintChannelInterfaceRegistry& Get();
+
+	/** Register an interface for the HandleType static class */
+	template<typename HandleType>
+	void RegisterConstraintChannelInterface(TUniquePtr<ITransformConstraintChannelInterface>&& InInterface)
+	{
+		UClass* HandleClass = HandleType::StaticClass();
+		check(!HandleToInterfaceMap.Contains(HandleClass));
+		HandleToInterfaceMap.Add(HandleClass, MoveTemp(InInterface));
+	}
+
+	/** Find the registered interface from the given class. Returns a nullptr if noting registered for that class. */
+	MOVIESCENETOOLS_API ITransformConstraintChannelInterface* FindConstraintChannelInterface(const UClass* InClass) const;
+
+private:
+	FConstraintChannelInterfaceRegistry() = default;
+
+	TMap<UClass*, TUniquePtr<ITransformConstraintChannelInterface>> HandleToInterfaceMap;
+};
\ No newline at end of file
diff --git a/Engine/Source/Editor/MovieSceneTools/Public/MovieSceneConstraintChannelHelper.h b/Engine/Source/Editor/MovieSceneTools/Public/MovieSceneConstraintChannelHelper.h
deleted file mode 100644
