fix static analysis

diff --git a/Engine/Source/Editor/MovieSceneTools/Private/Constraints/ComponentConstraintChannelInterface.cpp b/Engine/Source/Editor/MovieSceneTools/Private/Constraints/ComponentConstraintChannelInterface.cpp
--- a/Engine/Source/Editor/MovieSceneTools/Private/Constraints/ComponentConstraintChannelInterface.cpp
+++ b/Engine/Source/Editor/MovieSceneTools/Private/Constraints/ComponentConstraintChannelInterface.cpp
@@ -276,7 +276,7 @@ UMovieScene3DTransformSection* FComponentConstraintChannelInterface::GetComponen
 	}
 
 	const UMovieSceneSequence* MovieSceneSequence = InSequencer ? InSequencer->GetFocusedMovieSceneSequence() : nullptr;
-	UMovieScene* MovieScene = MovieSceneSequence->GetMovieScene();
+	UMovieScene* MovieScene = MovieSceneSequence ? MovieSceneSequence->GetMovieScene() : nullptr;
 	if (!MovieScene)
 	{
 		return nullptr;
