Fixes for TOptional forward includes.

diff --git a/Engine/Source/Editor/ContentBrowserData/Public/ContentBrowserDataFilter.h b/Engine/Source/Editor/ContentBrowserData/Public/ContentBrowserDataFilter.h
--- a/Engine/Source/Editor/ContentBrowserData/Public/ContentBrowserDataFilter.h
+++ b/Engine/Source/Editor/ContentBrowserData/Public/ContentBrowserDataFilter.h
@@ -11,6 +11,7 @@
 #include "HAL/Platform.h"
 #include "Misc/AssertionMacros.h"
 #include "Misc/EnumClassFlags.h"
+#include "Misc/OptionalFwd.h"
 #include "Templates/SharedPointer.h"
 #include "UObject/Class.h"
 #include "UObject/NameTypes.h"
@@ -22,7 +23,6 @@
 class FNamePermissionList;
 class FPathPermissionList;
 class UContentBrowserDataSource;
-template <typename OptionalType> struct TOptional;
 
 /** Flags controlling which item types should be included */
 UENUM()
diff --git a/Engine/Source/Editor/CurveEditor/Private/DragOperations/ICurveEditorDragOperation.cpp b/Engine/Source/Editor/CurveEditor/Private/DragOperations/ICurveEditorDragOperation.cpp
--- a/Engine/Source/Editor/CurveEditor/Private/DragOperations/ICurveEditorDragOperation.cpp
+++ b/Engine/Source/Editor/CurveEditor/Private/DragOperations/ICurveEditorDragOperation.cpp
@@ -6,7 +6,6 @@ class FSlateWindowElementList;
 struct FCurvePointHandle;
 struct FGeometry;
 struct FPointerEvent;
-template <typename OptionalType> struct TOptional;
 
 void ICurveEditorDragOperation::BeginDrag(FVector2D InitialPosition, FVector2D CurrentPosition, const FPointerEvent& MouseEvent)
 {
diff --git a/Engine/Source/Editor/CurveEditor/Public/ICurveEditorDragOperation.h b/Engine/Source/Editor/CurveEditor/Public/ICurveEditorDragOperation.h
--- a/Engine/Source/Editor/CurveEditor/Public/ICurveEditorDragOperation.h
+++ b/Engine/Source/Editor/CurveEditor/Public/ICurveEditorDragOperation.h
@@ -9,7 +9,7 @@
 #include "Input/Reply.h"
 #include "InputCoreTypes.h"
 #include "Math/Vector2D.h"
-#include "Misc/Optional.h"
+#include "Misc/OptionalFwd.h"
 #include "ScopedTransaction.h"
 #include "Templates/UniquePtr.h"
 
@@ -18,7 +18,6 @@ class FSlateWindowElementList;
 struct FCurvePointHandle;
 struct FGeometry;
 struct FPointerEvent;
-template <typename OptionalType> struct TOptional;
 
 /**
  * Interface for all drag operations in the curve editor
diff --git a/Engine/Source/Editor/CurveEditor/Public/RichCurveEditorModel.h b/Engine/Source/Editor/CurveEditor/Public/RichCurveEditorModel.h
--- a/Engine/Source/Editor/CurveEditor/Public/RichCurveEditorModel.h
+++ b/Engine/Source/Editor/CurveEditor/Public/RichCurveEditorModel.h
@@ -10,6 +10,7 @@
 #include "IBufferedCurveModel.h"
 #include "Math/Range.h"
 #include "Misc/Attribute.h"
+#include "Misc/OptionalFwd.h"
 #include "Templates/Tuple.h"
 #include "Templates/UniquePtr.h"
 #include "UObject/UnrealType.h"
@@ -26,7 +27,6 @@ struct FKeyDrawInfo;
 struct FKeyHandle;
 struct FKeyPosition;
 struct FRichCurve;
-template <typename OptionalType> struct TOptional;
 
 class CURVEEDITOR_API FRichCurveEditorModel : public FCurveModel
 {
diff --git a/Engine/Source/Editor/CurveTableEditor/Private/RealCurveModel.h b/Engine/Source/Editor/CurveTableEditor/Private/RealCurveModel.h
--- a/Engine/Source/Editor/CurveTableEditor/Private/RealCurveModel.h
+++ b/Engine/Source/Editor/CurveTableEditor/Private/RealCurveModel.h
@@ -6,6 +6,7 @@
 #include "CoreTypes.h"
 #include "CurveEditorTypes.h"
 #include "CurveModel.h"
+#include "Misc/OptionalFwd.h"
 #include "Templates/Tuple.h"
 #include "UObject/UnrealType.h"
 #include "UObject/WeakObjectPtrTemplates.h"
@@ -18,7 +19,6 @@ struct FKeyDrawInfo;
 struct FKeyHandle;
 struct FKeyPosition;
 struct FRealCurve;
-template <typename OptionalType> struct TOptional;
 
 /*
  * RealCurveModel implements the FCurveModel interface for FRealCurves which allows
diff --git a/Engine/Source/Editor/MovieSceneTools/Private/Channels/EventChannelCurveModel.cpp b/Engine/Source/Editor/MovieSceneTools/Private/Channels/EventChannelCurveModel.cpp
--- a/Engine/Source/Editor/MovieSceneTools/Private/Channels/EventChannelCurveModel.cpp
+++ b/Engine/Source/Editor/MovieSceneTools/Private/Channels/EventChannelCurveModel.cpp
@@ -29,7 +29,6 @@
 class FCurveEditor;
 class UObject;
 struct FCurveEditorScreenSpace;
-template <typename OptionalType> struct TOptional;
 
 ECurveEditorViewID FEventChannelCurveModel::EventView = ECurveEditorViewID::Invalid;
 
diff --git a/Engine/Source/Editor/MovieSceneTools/Private/Channels/EventChannelCurveModel.h b/Engine/Source/Editor/MovieSceneTools/Private/Channels/EventChannelCurveModel.h
--- a/Engine/Source/Editor/MovieSceneTools/Private/Channels/EventChannelCurveModel.h
+++ b/Engine/Source/Editor/MovieSceneTools/Private/Channels/EventChannelCurveModel.h
@@ -10,6 +10,7 @@
 #include "CurveEditorTypes.h"
 #include "CurveModel.h"
 #include "Curves/KeyHandle.h"
+#include "Misc/OptionalFwd.h"
 #include "Templates/SharedPointer.h"
 #include "Templates/Tuple.h"
 #include "UObject/UnrealType.h"
@@ -25,7 +26,6 @@ struct FKeyAttributes;
 struct FKeyDrawInfo;
 struct FKeyPosition;
 struct FMovieSceneEventChannel;
-template <typename OptionalType> struct TOptional;
 
 class FEventChannelCurveModel : public FCurveModel
 {
diff --git a/Engine/Source/Editor/MovieSceneTools/Public/Channels/ChannelCurveModel.h b/Engine/Source/Editor/MovieSceneTools/Public/Channels/ChannelCurveModel.h
--- a/Engine/Source/Editor/MovieSceneTools/Public/Channels/ChannelCurveModel.h
+++ b/Engine/Source/Editor/MovieSceneTools/Public/Channels/ChannelCurveModel.h
@@ -15,6 +15,7 @@
 #include "Curves/KeyHandle.h"
 #include "Delegates/IDelegateInstance.h"
 #include "IBufferedCurveModel.h"
+#include "Misc/OptionalFwd.h"
 #include "MovieSceneSection.h"
 #include "Templates/SharedPointer.h"
 #include "Templates/Tuple.h"
@@ -33,7 +34,6 @@ struct FCurveEditorScreenSpace;
 struct FKeyAttributes;
 struct FKeyDrawInfo;
 struct FKeyPosition;
-template <typename OptionalType> struct TOptional;
 
 template <class ChannelType, class ChannelValue, class KeyType>
 class FChannelCurveModel : public FCurveModel
diff --git a/Engine/Source/Editor/Persona/Private/SAnimCurveEd.h b/Engine/Source/Editor/Persona/Private/SAnimCurveEd.h
--- a/Engine/Source/Editor/Persona/Private/SAnimCurveEd.h
+++ b/Engine/Source/Editor/Persona/Private/SAnimCurveEd.h
@@ -10,6 +10,7 @@
 #include "Input/Reply.h"
 #include "Math/Vector2D.h"
 #include "Misc/Attribute.h"
+#include "Misc/OptionalFwd.h"
 #include "SCurveEditor.h"
 #include "Widgets/DeclarativeSyntaxSupport.h"
 
@@ -20,7 +21,6 @@ class FWidgetStyle;
 class UObject;
 struct FGeometry;
 struct FPointerEvent;
-template <typename OptionalType> struct TOptional;
 
 DECLARE_DELEGATE_OneParam(FOnSelectionChanged, const TArray<UObject*>&)
 DECLARE_DELEGATE(FOnUpdatePanel)
diff --git a/Engine/Source/Runtime/Interchange/Core/Public/InterchangePipelineBase.h b/Engine/Source/Runtime/Interchange/Core/Public/InterchangePipelineBase.h
--- a/Engine/Source/Runtime/Interchange/Core/Public/InterchangePipelineBase.h
+++ b/Engine/Source/Runtime/Interchange/Core/Public/InterchangePipelineBase.h
@@ -10,6 +10,7 @@
 #include "InterchangeResultsContainer.h"
 #include "InterchangeSourceData.h"
 #include "Misc/AssertionMacros.h"
+#include "Misc/OptionalFwd.h"
 #include "Nodes/InterchangeBaseNodeContainer.h"
 #include "Nodes/InterchangeFactoryBaseNode.h"
 #include "UObject/NameTypes.h"
@@ -26,7 +27,6 @@ class UInterchangeBaseNodeContainer;
 class UInterchangeResult;
 class UInterchangeSourceData;
 struct FFrame;
-template <typename OptionalType> struct TOptional;
 
 UENUM(BlueprintType)
 enum class EInterchangePipelineTask : uint8
diff --git a/Engine/Source/Runtime/MovieScene/Public/MovieSceneSpawnRegister.h b/Engine/Source/Runtime/MovieScene/Public/MovieSceneSpawnRegister.h
--- a/Engine/Source/Runtime/MovieScene/Public/MovieSceneSpawnRegister.h
+++ b/Engine/Source/Runtime/MovieScene/Public/MovieSceneSpawnRegister.h
@@ -11,6 +11,7 @@
 #include "Internationalization/Text.h"
 #include "Misc/AssertionMacros.h"
 #include "Misc/Guid.h"
+#include "Misc/OptionalFwd.h"
 #include "MovieSceneFwd.h"
 #include "MovieSceneSequenceID.h"
 #include "Templates/SharedPointer.h"
@@ -22,7 +23,6 @@
 class UClass;
 class UObject;
 template <typename FuncType> class TFunctionRef;
-template <typename OptionalType> struct TOptional;
 
 #if UE_ENABLE_INCLUDE_ORDER_DEPRECATED_IN_5_1
 #include "MovieSceneSpawnable.h"
