Sequencer: Remove take separator limitation for 1 character

diff --git a/Engine/Source/Editor/MovieSceneTools/Private/MovieSceneToolsProjectSettingsCustomization.cpp b/Engine/Source/Editor/MovieSceneTools/Private/MovieSceneToolsProjectSettingsCustomization.cpp
--- a/Engine/Source/Editor/MovieSceneTools/Private/MovieSceneToolsProjectSettingsCustomization.cpp
+++ b/Engine/Source/Editor/MovieSceneTools/Private/MovieSceneToolsProjectSettingsCustomization.cpp
@@ -8,17 +8,12 @@
 #include "HAL/Platform.h"
 #include "Internationalization/Text.h"
 #include "Misc/AssertionMacros.h"
+#include "MovieSceneFwd.h"
 #include "MovieSceneToolsProjectSettings.h"
 #include "PropertyHandle.h"
 #include "UObject/NameTypes.h"
 #include "UObject/UObjectGlobals.h"
 
-//#include "Editor/PropertyEditor/Public/IDetailsView.h"
-//#include "IDetailCustomization.h"
-//#include "DetailLayoutBuilder.h"
-//#include "PropertyHandle.h"
-
-
 TSharedRef<IDetailCustomization> FMovieSceneToolsProjectSettingsCustomization::MakeInstance()
 {
 	return MakeShareable(new FMovieSceneToolsProjectSettingsCustomization);
@@ -28,6 +23,9 @@ void FMovieSceneToolsProjectSettingsCustomization::CustomizeDetails(IDetailLayou
 {
 	TSharedRef<IPropertyHandle> TakeSeparatorProperty = DetailBuilder.GetProperty(GET_MEMBER_NAME_CHECKED(UMovieSceneToolsProjectSettings, TakeSeparator));
 	TakeSeparatorProperty->SetOnPropertyValueChanged(FSimpleDelegate::CreateRaw(this, &FMovieSceneToolsProjectSettingsCustomization::OnTakeSeparatorUpdated));
+
+	TSharedRef<IPropertyHandle> SubSequenceSeparatorProperty = DetailBuilder.GetProperty(GET_MEMBER_NAME_CHECKED(UMovieSceneToolsProjectSettings, SubSequenceSeparator));
+	SubSequenceSeparatorProperty->SetOnPropertyValueChanged(FSimpleDelegate::CreateRaw(this, &FMovieSceneToolsProjectSettingsCustomization::OnSubSequenceSeparatorUpdated));
 }
 
 void FMovieSceneToolsProjectSettingsCustomization::OnTakeSeparatorUpdated()
@@ -36,14 +34,26 @@ void FMovieSceneToolsProjectSettingsCustomization::OnTakeSeparatorUpdated()
 
 	FString TakeSeparator = ProjectSettings->TakeSeparator;
 
-	// Make sure the take separator is a valid single character
+	// Make sure the take separator is a valid character
 	FText OutErrorMessage;
 	if (!FName(*TakeSeparator).IsValidXName(INVALID_OBJECTNAME_CHARACTERS INVALID_LONGPACKAGE_CHARACTERS, &OutErrorMessage))
 	{
+		UE_LOG(LogMovieScene, Warning, TEXT("Invalid separator: %s"), *OutErrorMessage.ToString());
 		ProjectSettings->TakeSeparator = TEXT("_");
 	}
-	else if (ProjectSettings->TakeSeparator.Len() > 1)
+}
+
+void FMovieSceneToolsProjectSettingsCustomization::OnSubSequenceSeparatorUpdated()
+{
+	UMovieSceneToolsProjectSettings* ProjectSettings = GetMutableDefault<UMovieSceneToolsProjectSettings>();
+
+	FString SubSequenceSeparator = ProjectSettings->SubSequenceSeparator;
+
+	// Make sure the subsequence separator is a valid character
+	FText OutErrorMessage;
+	if (!FName(*SubSequenceSeparator).IsValidXName(INVALID_OBJECTNAME_CHARACTERS INVALID_LONGPACKAGE_CHARACTERS, &OutErrorMessage))
 	{
-		ProjectSettings->TakeSeparator.LeftChopInline(ProjectSettings->TakeSeparator.Len()-1);
+		UE_LOG(LogMovieScene, Warning, TEXT("Invalid separator: %s"), *OutErrorMessage.ToString());
+		ProjectSettings->SubSequenceSeparator = TEXT("_");
 	}
 }
diff --git a/Engine/Source/Editor/MovieSceneTools/Private/MovieSceneToolsProjectSettingsCustomization.h b/Engine/Source/Editor/MovieSceneTools/Private/MovieSceneToolsProjectSettingsCustomization.h
--- a/Engine/Source/Editor/MovieSceneTools/Private/MovieSceneToolsProjectSettingsCustomization.h
+++ b/Engine/Source/Editor/MovieSceneTools/Private/MovieSceneToolsProjectSettingsCustomization.h
@@ -15,4 +15,5 @@ public:
 
 private:
 	void OnTakeSeparatorUpdated();
+	void OnSubSequenceSeparatorUpdated();
 };
diff --git a/Engine/Source/Editor/MovieSceneTools/Public/MovieSceneToolsProjectSettings.h b/Engine/Source/Editor/MovieSceneTools/Public/MovieSceneToolsProjectSettings.h
--- a/Engine/Source/Editor/MovieSceneTools/Public/MovieSceneToolsProjectSettings.h
+++ b/Engine/Source/Editor/MovieSceneTools/Public/MovieSceneToolsProjectSettings.h
@@ -96,11 +96,11 @@ public:
 	UPROPERTY(config, EditAnywhere, Category=Shots, meta = (UIMin = "1", UIMax = "10"))
 	uint32 FirstTakeNumber;
 
-	/** A single character separator between the shot number and the take number. */
+	/** Character separator between the shot number and the take number. */
 	UPROPERTY(config, EditAnywhere, Category=Shots)
 	FString TakeSeparator;
 
-	/** A single character separator between the take number and the sub sequence name. */
+	/** Character separator between the take number and the sub sequence name. */
 	UPROPERTY(config, EditAnywhere, Category=Shots)
 	FString SubSequenceSeparator;
 
