Sequencer: Remove bIsReference from payload pin type.

Pins for ref parameters for functions default to bIsReference but the payload cannot be by reference.

diff --git a/Engine/Source/Editor/MovieSceneTools/Private/MovieSceneEventCustomization.cpp b/Engine/Source/Editor/MovieSceneTools/Private/MovieSceneEventCustomization.cpp
--- a/Engine/Source/Editor/MovieSceneTools/Private/MovieSceneEventCustomization.cpp
+++ b/Engine/Source/Editor/MovieSceneTools/Private/MovieSceneEventCustomization.cpp
@@ -1261,6 +1261,9 @@ void FMovieSceneEventCustomization::SetEventEndpoint(UK2Node* NewEndpoint, UEdGr
 
 				if (EditableNode && EnumHasAnyFlags(AutoCreatePayload, EAutoCreatePayload::Pins))
 				{
+					// Pins for ref parameters for functions default to bIsReference but the payload cannot be by reference.
+					PayloadPin->PinType.bIsReference = false;
+					
 					UEdGraphPin* NewPin = EditableNode->CreateUserDefinedPin(PayloadPin->PinName, PayloadPin->PinType, EGPD_Output);
 					if (PayloadTemplate != NewEndpoint && NewPin)
 					{
