Sequencer: Remove unused code

diff --git a/Engine/Source/Editor/MovieSceneTools/Private/TrackEditors/SubTrackEditor.cpp b/Engine/Source/Editor/MovieSceneTools/Private/TrackEditors/SubTrackEditor.cpp
--- a/Engine/Source/Editor/MovieSceneTools/Private/TrackEditors/SubTrackEditor.cpp
+++ b/Engine/Source/Editor/MovieSceneTools/Private/TrackEditors/SubTrackEditor.cpp
@@ -473,24 +473,6 @@ void FSubTrackEditor::HandleAddSubTrackMenuEntryExecute()
 	}
 }
 
-/** Helper function - get the first PIE world (or first PIE client world if there is more than one) */
-static UWorld* GetFirstPIEWorld()
-{
-	for (const FWorldContext& Context : GEngine->GetWorldContexts())
-	{
-		if (Context.World()->IsPlayInEditor())
-		{
-			if(Context.World()->GetNetMode() == ENetMode::NM_Standalone ||
-				(Context.World()->GetNetMode() == ENetMode::NM_Client && Context.PIEInstance == 2))
-			{
-				return Context.World();
-			}
-		}
-	}
-
-	return nullptr;
-}
-
 TSharedRef<SWidget> FSubTrackEditor::HandleAddSubSequenceComboButtonGetMenuContent(UMovieSceneTrack* InTrack)
 {
 	FMenuBuilder MenuBuilder(true, nullptr);
