Sequencer: Fixed Linked Anim Sequence Saving with nested sequences

diff --git a/Engine/Source/Editor/MovieSceneTools/Private/TrackEditors/SkeletalAnimationTrackEditor.cpp b/Engine/Source/Editor/MovieSceneTools/Private/TrackEditors/SkeletalAnimationTrackEditor.cpp
--- a/Engine/Source/Editor/MovieSceneTools/Private/TrackEditors/SkeletalAnimationTrackEditor.cpp
+++ b/Engine/Source/Editor/MovieSceneTools/Private/TrackEditors/SkeletalAnimationTrackEditor.cpp
@@ -1197,8 +1197,8 @@ void FSkeletalAnimationTrackEditor::OnSequencerSaved(ISequencer& )
 			{
 				UMovieScene* MovieScene = SequencerPtr->GetFocusedMovieSceneSequence()->GetMovieScene();
 				FMovieSceneSequenceIDRef Template = SequencerPtr->GetFocusedTemplateID();
-				FMovieSceneSequenceTransform RootToLocalTransform;
-				for (int32 Index = LevelAnimLink->AnimSequenceLinks.Num() -1; Index >=0 ; --Index) 
+				FMovieSceneSequenceTransform RootToLocalTransform = SequencerPtr->GetFocusedMovieSceneSequenceTransform();
+				for (int32 Index = LevelAnimLink->AnimSequenceLinks.Num() -1; Index >=0 ; --Index)
 				{
 					FLevelSequenceAnimSequenceLinkItem& Item = LevelAnimLink->AnimSequenceLinks[Index];
 					UAnimSequence* AnimSequence = Item.ResolveAnimSequence();
