Sequencer: Fix crash in event customization when the blueprint graph is deleted

diff --git a/Engine/Source/Editor/MovieSceneTools/Private/MovieSceneEventCustomization.cpp b/Engine/Source/Editor/MovieSceneTools/Private/MovieSceneEventCustomization.cpp
--- a/Engine/Source/Editor/MovieSceneTools/Private/MovieSceneEventCustomization.cpp
+++ b/Engine/Source/Editor/MovieSceneTools/Private/MovieSceneEventCustomization.cpp
@@ -343,7 +343,7 @@ void FMovieSceneEventCustomization::CustomizeChildren(TSharedRef<IPropertyHandle
 
 
 	UFunction* CommonFunction = nullptr;
-	UBlueprint* Blueprint = CommonEndpoint->GetBlueprint();
+	UBlueprint* Blueprint = CommonEndpoint->HasValidBlueprint() ? CommonEndpoint->GetBlueprint() : nullptr;
 	if (Blueprint)
 	{
 		if (UK2Node_Event* Event = Cast<UK2Node_Event>(CommonEndpoint))
@@ -507,8 +507,13 @@ void FMovieSceneEventCustomization::OnPayloadVariableChanged(TSharedRef<FStructO
 
 	if (bChangedAnything)
 	{
-		UBlueprint* BP = GetCommonEndpoint()->GetBlueprint();
-		FBlueprintEditorUtils::MarkBlueprintAsModified(BP);
+		UK2Node* CommonEndpoint = GetCommonEndpoint();
+
+		UBlueprint* BP = CommonEndpoint && CommonEndpoint->HasValidBlueprint() ? GetCommonEndpoint()->GetBlueprint() : nullptr;
+		if (BP)
+		{
+			FBlueprintEditorUtils::MarkBlueprintAsModified(BP);
+		}
 	}
 	else
 	{
@@ -557,13 +562,21 @@ void FMovieSceneEventCustomization::OnSetCallInEditorCheckState(ECheckBoxState N
 		{
 			CustomEvent->Modify();
 			CustomEvent->bCallInEditor = bCallInEditor;
-			Blueprints.Add(CustomEvent->GetBlueprint());
+
+			if (CustomEvent->HasValidBlueprint())
+			{
+				Blueprints.Add(CustomEvent->GetBlueprint());
+			}
 		}
 		else if (UK2Node_FunctionEntry* FunctionEntry = Cast<UK2Node_FunctionEntry>(Endpoint))
 		{
 			FunctionEntry->Modify();
 			FunctionEntry->MetaData.bCallInEditor = bCallInEditor;
-			Blueprints.Add(FunctionEntry->GetBlueprint());
+		
+			if (FunctionEntry->HasValidBlueprint())
+			{
+				Blueprints.Add(FunctionEntry->GetBlueprint());
+			}
 		}
 	}
 
@@ -1212,7 +1225,7 @@ void FMovieSceneEventCustomization::SetEventEndpoint(UK2Node* NewEndpoint, UEdGr
 
 	// If we're assigning a new valid endpoint, it must reside within the same blueprint as everything we're assigning it to.
 	// Anything else must be implemented as a call function node connected to a custom event node
-	UBlueprint* Blueprint = NewEndpoint ? NewEndpoint->GetBlueprint() : nullptr;
+	UBlueprint* Blueprint = NewEndpoint && NewEndpoint->HasValidBlueprint() ? NewEndpoint->GetBlueprint() : nullptr;
 	if (Blueprint)
 	{
 		for (UObject* Outer : EditObjects)
