Sequencer: Add unwound key when adding a key explicitly (autokeying was fine and already does this)

diff --git a/Engine/Source/Editor/MovieSceneTools/Private/TrackEditors/TransformTrackEditor.cpp b/Engine/Source/Editor/MovieSceneTools/Private/TrackEditors/TransformTrackEditor.cpp
--- a/Engine/Source/Editor/MovieSceneTools/Private/TrackEditors/TransformTrackEditor.cpp
+++ b/Engine/Source/Editor/MovieSceneTools/Private/TrackEditors/TransformTrackEditor.cpp
@@ -1057,6 +1057,26 @@ void F3DTransformTrackEditor::ProcessKeyOperation(FFrameNumber InKeyTime, const
 	Operation.IterateOperations(Iterator);
 }
 
+
+int32 GetPreviousKey(FMovieSceneDoubleChannel& Channel, FFrameNumber Time)
+{
+	TArray<FFrameNumber> KeyTimes;
+	TArray<FKeyHandle> KeyHandles;
+
+	TRange<FFrameNumber> Range;
+	Range.SetLowerBound(TRangeBound<FFrameNumber>::Open());
+	Range.SetUpperBound(TRangeBound<FFrameNumber>::Exclusive(Time));
+	Channel.GetData().GetKeys(Range, &KeyTimes, &KeyHandles);
+
+	if (KeyHandles.Num() <= 0)
+	{
+		return INDEX_NONE;
+	}
+
+	int32 Index = Channel.GetData().GetIndex(KeyHandles[KeyHandles.Num() - 1]);
+	return Index;
+}
+
 void F3DTransformTrackEditor::ProcessKeyOperation(UObject* ObjectToKey, TArrayView<const UE::Sequencer::FKeySectionOperation> SectionsToKey, ISequencer& InSequencer, FFrameNumber KeyTime)
 {
 	USceneComponent* Component = MovieSceneHelpers::SceneComponentFromRuntimeObject(ObjectToKey);
@@ -1148,7 +1168,20 @@ void F3DTransformTrackEditor::ProcessKeyOperation(UObject* ObjectToKey, TArrayVi
 					FMovieSceneDoubleChannel* Channel = static_cast<FMovieSceneDoubleChannel*>(Handle.Get());
 
 					double Value = RecomposedTransform[Handle.GetChannelIndex()];
-					EMovieSceneKeyInterpolation Interpolation = GetInterpolationMode(Channel,KeyTime,InSequencer.GetKeyInterpolation());
+
+					if (KeyArea->GetName() == "Rotation.X" ||
+						KeyArea->GetName() == "Rotation.Y" ||
+						KeyArea->GetName() == "Rotation.Z")
+					{
+						int32 PreviousKey = GetPreviousKey(*Channel, KeyTime);
+						if (PreviousKey != INDEX_NONE && PreviousKey < Channel->GetData().GetValues().Num())
+						{
+							double OldValue = Channel->GetData().GetValues()[PreviousKey].Value;
+							Value = UnwindChannel(OldValue, Value);
+						}
+					}
+
+					EMovieSceneKeyInterpolation Interpolation = GetInterpolationMode(Channel, KeyTime, InSequencer.GetKeyInterpolation());
 					AddKeyToChannel(Channel, KeyTime, Value, Interpolation);
 				}
 				else
