Sequencer: Add option to select the material parameter collection asset from the track

diff --git a/Engine/Source/Editor/MovieSceneTools/Private/TrackEditors/MaterialParameterCollectionTrackEditor.cpp b/Engine/Source/Editor/MovieSceneTools/Private/TrackEditors/MaterialParameterCollectionTrackEditor.cpp
--- a/Engine/Source/Editor/MovieSceneTools/Private/TrackEditors/MaterialParameterCollectionTrackEditor.cpp
+++ b/Engine/Source/Editor/MovieSceneTools/Private/TrackEditors/MaterialParameterCollectionTrackEditor.cpp
@@ -9,6 +9,7 @@
 #include "ContentBrowserDelegates.h"
 #include "ContentBrowserModule.h"
 #include "Delegates/Delegate.h"
+#include "Editor/UnrealEdEngine.h"
 #include "Framework/Application/SlateApplication.h"
 #include "Framework/Commands/UIAction.h"
 #include "Framework/MultiBox/MultiBoxBuilder.h"
@@ -133,6 +134,16 @@ void FMaterialParameterCollectionTrackEditor::BuildTrackContextMenu(FMenuBuilder
 		SubMenuBuilder.AddWidget(CreateAssetPicker(FOnAssetSelected::CreateLambda(AssignAsset), FOnAssetEnterPressed::CreateLambda(AssignAssetEnterPressed), GetSequencer()), FText::GetEmpty(), true);
 	};
 
+	UMaterialParameterCollection* MPC = MPCTrack ? MPCTrack->MPC : nullptr;
+
+	MenuBuilder.AddMenuEntry(
+		FText::Format(LOCTEXT("SelectAssetFormat", "Select {0}"), MPC ? FText::FromString(MPC->GetName()) : FText::GetEmpty()),
+		FText::Format(LOCTEXT("SelectAssetTooltipFormat", "Select {0}"), MPC ? FText::FromString(MPC->GetName()) : FText::GetEmpty()),
+		FSlateIcon(),
+		FUIAction(FExecuteAction::CreateRaw(this, &FMaterialParameterCollectionTrackEditor::OnSelectMPC, MPC),
+				  FCanExecuteAction::CreateLambda([MPC] { return MPC != nullptr; }))
+	);
+
 	MenuBuilder.AddSubMenu(
 		LOCTEXT("SetAsset", "Set Asset"),
 		LOCTEXT("SetAsset_ToolTip", "Sets the Material Parameter Collection that this track animates."),
@@ -140,6 +151,19 @@ void FMaterialParameterCollectionTrackEditor::BuildTrackContextMenu(FMenuBuilder
 	);
 }
 
+void FMaterialParameterCollectionTrackEditor::OnSelectMPC(UMaterialParameterCollection* MPC)
+{
+	if (!MPC)
+	{
+		return;
+	}
+
+	TArray<UObject*> ObjectsToFocus;
+	ObjectsToFocus.Add(MPC);
+
+	GEditor->SyncBrowserToObjects(ObjectsToFocus);
+}
+
 
 void FMaterialParameterCollectionTrackEditor::BuildAddTrackMenu(FMenuBuilder& MenuBuilder)
 {
diff --git a/Engine/Source/Editor/MovieSceneTools/Private/TrackEditors/MaterialParameterCollectionTrackEditor.h b/Engine/Source/Editor/MovieSceneTools/Private/TrackEditors/MaterialParameterCollectionTrackEditor.h
--- a/Engine/Source/Editor/MovieSceneTools/Private/TrackEditors/MaterialParameterCollectionTrackEditor.h
+++ b/Engine/Source/Editor/MovieSceneTools/Private/TrackEditors/MaterialParameterCollectionTrackEditor.h
@@ -55,6 +55,8 @@ private:
 	/** Provides the contents of the add parameter menu. */
 	TSharedRef<SWidget> OnGetAddParameterMenuContent(UMovieSceneMaterialParameterCollectionTrack* MaterialTrack);
 
+	void OnSelectMPC(UMaterialParameterCollection* MPC);
+
 	void AddScalarParameter(UMovieSceneMaterialParameterCollectionTrack* Track, FCollectionScalarParameter Parameter);
 	void AddVectorParameter(UMovieSceneMaterialParameterCollectionTrack* Track, FCollectionVectorParameter Parameter);
 
