Sequencer: Control Rig: Fix Setting Keys/Undoing now working when setting defaults.  Regression from a fix UE-165296 were we were too agressive and turned off ally setting of defaults, we only need to turn it off for controls that aren't 'adding' keys. To do this we tweak what we did before and add a flag to only set default on channels that are actually setting values via the bAddKey flag.

diff --git a/Engine/Plugins/Animation/ControlRig/Source/ControlRigEditor/Private/Sequencer/ControlRigParameterTrackEditor.cpp b/Engine/Plugins/Animation/ControlRig/Source/ControlRigEditor/Private/Sequencer/ControlRigParameterTrackEditor.cpp
--- a/Engine/Plugins/Animation/ControlRig/Source/ControlRigEditor/Private/Sequencer/ControlRigParameterTrackEditor.cpp
+++ b/Engine/Plugins/Animation/ControlRig/Source/ControlRigEditor/Private/Sequencer/ControlRigParameterTrackEditor.cpp
@@ -3432,7 +3432,7 @@ FKeyPropertyResult FControlRigParameterTrackEditor::AddKeysToControlRigHandle(US
 			UMovieSceneControlRigParameterSection* ParamSection = Cast<UMovieSceneControlRigParameterSection>(SectionToKey);
 			if (!ParamSection->GetDoNotKey())
 			{
-				KeyPropertyResult |= AddKeysToSection(SectionToKey, KeyTime, GeneratedKeys, KeyMode, FKeyframeTrackEditor::ESetDefault::DoNotSetDefault);
+				KeyPropertyResult |= AddKeysToSection(SectionToKey, KeyTime, GeneratedKeys, KeyMode, EKeyFrameTrackEditorSetDefault::SetDefaultOnAddKeys);
 			}
 		}
 
diff --git a/Engine/Source/Editor/MovieSceneTools/Public/KeyframeTrackEditor.h b/Engine/Source/Editor/MovieSceneTools/Public/KeyframeTrackEditor.h
--- a/Engine/Source/Editor/MovieSceneTools/Public/KeyframeTrackEditor.h
+++ b/Engine/Source/Editor/MovieSceneTools/Public/KeyframeTrackEditor.h
@@ -25,6 +25,16 @@
 #include "Channels/MovieSceneBoolChannel.h"
 
 
+enum class EKeyFrameTrackEditorSetDefault
+{
+	/** Do not set default */
+	DoNotSetDefault,
+	/** Set default just on keys that have bAddKey set */
+	SetDefaultOnAddKeys,
+	/** Set default */
+	SetDefault,
+};
+
 struct IImpl
 {
 	virtual ~IImpl() {}
@@ -32,7 +42,7 @@ struct IImpl
 	/* Returns whether a key was created */
 	virtual bool Apply(UMovieSceneSection* Section, FMovieSceneChannelProxy& Proxy, FFrameNumber InTime, EMovieSceneKeyInterpolation InterpolationMode, bool bKeyEvenIfUnchanged, bool bKeyEvenIfEmpty) const { return false; }
 
-	virtual void ApplyDefault(UMovieSceneSection* Section, FMovieSceneChannelProxy& Proxy) const { }
+	virtual void ApplyDefault(UMovieSceneSection* Section, FMovieSceneChannelProxy& Proxy, EKeyFrameTrackEditorSetDefault SetDefault) const { }
 
 	virtual bool ModifyByCurrentAndWeight(FMovieSceneChannelProxy& Proxy, FFrameNumber InTime, void* VCurrentValue, float Weight) { return false; }
 
@@ -81,7 +91,7 @@ struct TAddKeyImpl : IImpl
 		return bKeyCreated;
 	}
 
-	virtual void ApplyDefault(UMovieSceneSection* Section, FMovieSceneChannelProxy& Proxy) const override
+	virtual void ApplyDefault(UMovieSceneSection* Section, FMovieSceneChannelProxy& Proxy, EKeyFrameTrackEditorSetDefault SetDefault) const override
 	{
 		using namespace UE::MovieScene;
 
@@ -142,12 +152,13 @@ struct TAddKeyImpl<FMovieSceneBoolChannel, bool> : IImpl
 		return bKeyCreated;
 	}
 
-	virtual void ApplyDefault(UMovieSceneSection* Section, FMovieSceneChannelProxy& Proxy) const override
+	virtual void ApplyDefault(UMovieSceneSection* Section, FMovieSceneChannelProxy& Proxy, EKeyFrameTrackEditorSetDefault SetDefault) const override
 	{
 		using namespace UE::MovieScene;
 
 		FMovieSceneBoolChannel* Channel = Proxy.GetChannel<FMovieSceneBoolChannel>(ChannelIndex);
-		if (Channel && Channel->GetData().GetTimes().Num() == 0 && Channel->GetDefault() != ValueToSet)
+		if (Channel && Channel->GetData().GetTimes().Num() == 0 && Channel->GetDefault() != ValueToSet
+			&& (SetDefault == EKeyFrameTrackEditorSetDefault::SetDefault || (SetDefault == EKeyFrameTrackEditorSetDefault::SetDefaultOnAddKeys && bAddKey)))
 		{
 			if (Section->TryModify())
 			{
@@ -202,12 +213,13 @@ struct TAddKeyImpl<FMovieSceneDoubleChannel, double> : IImpl
 		return bKeyCreated;
 	}
 
-	virtual void ApplyDefault(UMovieSceneSection* Section, FMovieSceneChannelProxy& Proxy) const override
+	virtual void ApplyDefault(UMovieSceneSection* Section, FMovieSceneChannelProxy& Proxy, EKeyFrameTrackEditorSetDefault SetDefault) const override
 	{
 		using namespace UE::MovieScene;
 
 		FMovieSceneDoubleChannel* Channel = Proxy.GetChannel<FMovieSceneDoubleChannel>(ChannelIndex);
-		if (Channel && Channel->GetData().GetTimes().Num() == 0  && Channel->GetDefault() != ValueToSet)
+		if (Channel && Channel->GetData().GetTimes().Num() == 0  && Channel->GetDefault() != ValueToSet
+			&& (SetDefault == EKeyFrameTrackEditorSetDefault::SetDefault || (SetDefault == EKeyFrameTrackEditorSetDefault::SetDefaultOnAddKeys && bAddKey)))
 		{
 			if (Section->TryModify())
 			{
@@ -280,12 +292,13 @@ struct TAddKeyImpl<FMovieSceneFloatChannel, float> : IImpl
 		return bKeyCreated;
 	}
 
-	virtual void ApplyDefault(UMovieSceneSection* Section, FMovieSceneChannelProxy& Proxy) const override
+	virtual void ApplyDefault(UMovieSceneSection* Section, FMovieSceneChannelProxy& Proxy, EKeyFrameTrackEditorSetDefault SetDefault) const override
 	{
 		using namespace UE::MovieScene;
 
 		FMovieSceneFloatChannel* Channel = Proxy.GetChannel<FMovieSceneFloatChannel>(ChannelIndex);
-		if (Channel && Channel->GetData().GetTimes().Num() == 0  && Channel->GetDefault() != ValueToSet)
+		if (Channel && Channel->GetData().GetTimes().Num() == 0  && Channel->GetDefault() != ValueToSet
+			&& (SetDefault == EKeyFrameTrackEditorSetDefault::SetDefault || (SetDefault == EKeyFrameTrackEditorSetDefault::SetDefaultOnAddKeys && bAddKey)))
 		{
 			if (Section->TryModify())
 			{
@@ -358,12 +371,13 @@ struct TAddKeyImpl<FMovieSceneIntegerChannel, int32> : IImpl
 		return bKeyCreated;
 	}
 
-	virtual void ApplyDefault(UMovieSceneSection* Section, FMovieSceneChannelProxy& Proxy) const override
+	virtual void ApplyDefault(UMovieSceneSection* Section, FMovieSceneChannelProxy& Proxy, EKeyFrameTrackEditorSetDefault SetDefault) const override
 	{
 		using namespace UE::MovieScene;
 
 		FMovieSceneIntegerChannel* Channel = Proxy.GetChannel<FMovieSceneIntegerChannel>(ChannelIndex);
-		if (Channel && Channel->GetData().GetTimes().Num() == 0 && Channel->GetDefault() != ValueToSet)
+		if (Channel && Channel->GetData().GetTimes().Num() == 0 && Channel->GetDefault() != ValueToSet
+			&& (SetDefault == EKeyFrameTrackEditorSetDefault::SetDefault || (SetDefault == EKeyFrameTrackEditorSetDefault::SetDefaultOnAddKeys && bAddKey)))
 		{
 			if (Section->TryModify())
 			{
@@ -642,15 +656,8 @@ private:
 	}
 
 protected:
-	enum class ESetDefault
-	{
-		/** Do not set default. */
-		DoNotSetDefault,
-		/** Set default */
-		SetDefault,
-	};
 	/* Returns whether a section was added */
-	FKeyPropertyResult AddKeysToSection(UMovieSceneSection* Section, FFrameNumber KeyTime, const FGeneratedTrackKeys& Keys, ESequencerKeyMode KeyMode, ESetDefault SetDefault = ESetDefault::SetDefault)
+	FKeyPropertyResult AddKeysToSection(UMovieSceneSection* Section, FFrameNumber KeyTime, const FGeneratedTrackKeys& Keys, ESequencerKeyMode KeyMode, EKeyFrameTrackEditorSetDefault SetDefault = EKeyFrameTrackEditorSetDefault::SetDefault)
 	{
 		FKeyPropertyResult KeyPropertyResult;
 
@@ -658,7 +665,7 @@ protected:
 
 		FMovieSceneChannelProxy& Proxy = Section->GetChannelProxy();
 			
-		const bool bSetDefaults = GetSequencer()->GetAutoSetTrackDefaults() && (SetDefault == ESetDefault::SetDefault);
+		const bool bSetDefaults = GetSequencer()->GetAutoSetTrackDefaults() && (SetDefault != EKeyFrameTrackEditorSetDefault::DoNotSetDefault);
 
 		// The default value is a value for the channel when there are no keyframes. For example, if you add keys and 
 		// then delete them all, the default value is the value of the channel. In the implementation of ApplyDefault, 
@@ -668,7 +675,7 @@ protected:
 		{
 			for (const FMovieSceneChannelValueSetter& GeneratedKey : Keys)
 			{
-				GeneratedKey->ApplyDefault(Section, Proxy);
+				GeneratedKey->ApplyDefault(Section, Proxy, SetDefault);
 			}
 		}
 
