Constraints: handle parent constraint scaling property change to update scale channels

diff --git a/Engine/Plugins/Animation/ControlRig/Source/ControlRigEditor/Private/Sequencer/ControlRigParameterTrackEditor.cpp b/Engine/Plugins/Animation/ControlRig/Source/ControlRigEditor/Private/Sequencer/ControlRigParameterTrackEditor.cpp
--- a/Engine/Plugins/Animation/ControlRig/Source/ControlRigEditor/Private/Sequencer/ControlRigParameterTrackEditor.cpp
+++ b/Engine/Plugins/Animation/ControlRig/Source/ControlRigEditor/Private/Sequencer/ControlRigParameterTrackEditor.cpp
@@ -82,6 +82,7 @@
 #include "Constraints/MovieSceneConstraintChannelHelper.h"
 #include "Constraints/ControlRigTransformableHandle.h"
 #include "PropertyEditorModule.h"
+#include "Constraints/TransformConstraintChannelInterface.h"
 
 #define LOCTEXT_NAMESPACE "FControlRigParameterTrackEditor"
 
@@ -2417,7 +2418,7 @@ FMovieSceneTrackEditor::FFindOrCreateTrackResult FControlRigParameterTrackEditor
 	return Result;
 }
 
-UMovieSceneControlRigParameterTrack* FControlRigParameterTrackEditor::FindTrack(UControlRig* InControlRig)
+UMovieSceneControlRigParameterTrack* FControlRigParameterTrackEditor::FindTrack(const UControlRig* InControlRig) const
 {
 	if (!GetSequencer().IsValid())
 	{
@@ -2556,6 +2557,8 @@ void FControlRigParameterTrackEditor::HandleSpaceKeyMoved(
 
 void FControlRigParameterTrackEditor::ClearOutAllSpaceAndConstraintDelegates(const UControlRig* InOptionalControlRig) const
 {
+	UTickableTransformConstraint::GetOnConstraintChanged().RemoveAll(this);
+	
 	const UMovieScene* MovieScene = GetSequencer().IsValid() && GetSequencer()->GetFocusedMovieSceneSequence() ? GetSequencer()->GetFocusedMovieSceneSequence()->GetMovieScene() : nullptr;
 	if (!MovieScene)
 	{
@@ -2701,6 +2704,11 @@ void FControlRigParameterTrackEditor::HandleOnConstraintAdded(
 	{
 		HandleConstraintRemoved(InSection);
 	}
+
+	if (!UTickableTransformConstraint::GetOnConstraintChanged().IsBoundToObject(this))
+	{
+		UTickableTransformConstraint::GetOnConstraintChanged().AddRaw(this, &FControlRigParameterTrackEditor::HandleConstraintPropertyChanged);
+	}
 }
 
 void FControlRigParameterTrackEditor::HandleConstraintKeyDeleted(
@@ -2809,6 +2817,50 @@ void FControlRigParameterTrackEditor::HandleConstraintRemoved(IMovieSceneConstra
 	}
 }
 
+void FControlRigParameterTrackEditor::HandleConstraintPropertyChanged(UTickableTransformConstraint* InConstraint, const FPropertyChangedEvent& InPropertyChangedEvent) const
+{
+	if (!IsValid(InConstraint))
+	{
+		return;
+	}
+
+	// find constraint section
+	const UTransformableControlHandle* Handle = Cast<UTransformableControlHandle>(InConstraint->ChildTRSHandle);
+	if (!IsValid(Handle) || !Handle->IsValid())
+	{
+		return;
+	}
+
+	const FConstraintChannelInterfaceRegistry& InterfaceRegistry = FConstraintChannelInterfaceRegistry::Get();	
+	ITransformConstraintChannelInterface* Interface = InterfaceRegistry.FindConstraintChannelInterface(Handle->GetClass());
+	if (!Interface)
+	{
+		return;
+	}
+	
+	UMovieSceneSection* Section = Interface->GetHandleConstraintSection(Handle, GetSequencer());
+	IMovieSceneConstrainedSection* ConstraintSection = Cast<IMovieSceneConstrainedSection>(Section);
+	if (!ConstraintSection)
+	{
+		return;
+	}
+
+	// find corresponding channel
+	const TArray<FConstraintAndActiveChannel>& ConstraintChannels = ConstraintSection->GetConstraintsChannels();
+	const FConstraintAndActiveChannel* Channel = ConstraintChannels.FindByPredicate([InConstraint](const FConstraintAndActiveChannel& Channel)
+	{
+		return Channel.Constraint == InConstraint || Channel.ConstraintCopyToSpawn == InConstraint;
+	});
+
+	if (!Channel)
+	{
+		return;
+	}
+
+	FMovieSceneConstraintChannelHelper::HandleConstraintPropertyChanged(
+			InConstraint, Channel->ActiveChannel, InPropertyChangedEvent, GetSequencer(), Section);
+}
+
 void FControlRigParameterTrackEditor::SetUpEditModeIfNeeded(UControlRig* ControlRig)
 {
 	if (ControlRig)
diff --git a/Engine/Plugins/Animation/ControlRig/Source/ControlRigEditor/Private/Sequencer/ControlRigParameterTrackEditor.h b/Engine/Plugins/Animation/ControlRig/Source/ControlRigEditor/Private/Sequencer/ControlRigParameterTrackEditor.h
--- a/Engine/Plugins/Animation/ControlRig/Source/ControlRigEditor/Private/Sequencer/ControlRigParameterTrackEditor.h
+++ b/Engine/Plugins/Animation/ControlRig/Source/ControlRigEditor/Private/Sequencer/ControlRigParameterTrackEditor.h
@@ -20,6 +20,7 @@
 #include "EditorUndoClient.h"
 #include "ScopedTransaction.h"
 
+class UTickableTransformConstraint;
 struct FAssetData;
 class FMenuBuilder;
 class USkeleton;
@@ -160,6 +161,7 @@ private:
 		const FMovieSceneConstraintChannel* InConstraintChannel,
 		const TArray<FKeyMoveEventItem>& InMovedItems);
 	void HandleConstraintRemoved(IMovieSceneConstrainedSection* InSection);
+	void HandleConstraintPropertyChanged(UTickableTransformConstraint* InConstraint, const FPropertyChangedEvent& InPropertyChangedEvent) const;
 
 	/** Select control rig if not selected, select controls from key areas */
 	void SelectRigsAndControls(UControlRig* Subject, const TArray<const IKeyArea*>& KeyAreas);
@@ -174,7 +176,7 @@ private:
 	void ImportFBX(UMovieSceneControlRigParameterTrack* InTrack, UMovieSceneControlRigParameterSection* InSection, 
 		TArray<FFBXNodeAndChannels>* NodeAndChannels);
 	/** Find Track for given ControlRig*/
-	UMovieSceneControlRigParameterTrack* FindTrack(UControlRig* InControlRig);
+	UMovieSceneControlRigParameterTrack* FindTrack(const UControlRig* InControlRig) const;
 
 	/** Select Bones to Animate on FK Rig*/
 	void SelectFKBonesToAnimate(UFKControlRig* FKControlRig, UMovieSceneControlRigParameterTrack* Track);
diff --git a/Engine/Source/Editor/MovieSceneTools/Private/Constraints/MovieSceneConstraintChannelHelper.cpp b/Engine/Source/Editor/MovieSceneTools/Private/Constraints/MovieSceneConstraintChannelHelper.cpp
--- a/Engine/Source/Editor/MovieSceneTools/Private/Constraints/MovieSceneConstraintChannelHelper.cpp
+++ b/Engine/Source/Editor/MovieSceneTools/Private/Constraints/MovieSceneConstraintChannelHelper.cpp
@@ -527,6 +527,89 @@ void FCompensationEvaluator::ComputeCompensation(UWorld* InWorld, const TSharedP
 	}
 }
 
+void FCompensationEvaluator::CacheTransforms(UWorld* InWorld, const TSharedPtr<ISequencer>& InSequencer, const TArray<FFrameNumber>& InTime)
+{
+	if (InTime.IsEmpty())
+	{
+		return;
+	}
+
+	using ConstraintPtr = TObjectPtr<UTickableConstraint>;
+
+	// get all constraints for evaluation
+	const FConstraintsManagerController& Controller = FConstraintsManagerController::Get(InWorld);
+	static constexpr bool bSorted = true;
+	const TArray<ConstraintPtr> AllConstraints = Controller.GetAllConstraints(bSorted);
+	
+	UMovieScene* MovieScene = InSequencer->GetFocusedMovieSceneSequence()->GetMovieScene();
+	const FFrameRate TickResolution = MovieScene->GetTickResolution();
+	const EMovieScenePlayerStatus::Type PlaybackStatus = InSequencer->GetPlaybackStatus();
+
+	const int32 NumFrames = InTime.Num();
+
+	ChildLocals.SetNum(NumFrames);
+	ChildGlobals.SetNum(NumFrames);
+	SpaceGlobals.SetNum(NumFrames);
+
+	const TArray<IMovieSceneToolsAnimationBakeHelper*>& BakeHelpers = FMovieSceneToolsModule::Get().GetAnimationBakeHelpers();
+	
+	auto EvaluateAt = [&](const FFrameNumber InFrame)
+	{
+		const FMovieSceneEvaluationRange EvaluationRange = FMovieSceneEvaluationRange(FFrameTime(InFrame), TickResolution);
+		const FMovieSceneContext Context = FMovieSceneContext(EvaluationRange, PlaybackStatus).SetHasJumped(true);
+
+		for (IMovieSceneToolsAnimationBakeHelper* BakeHelper : BakeHelpers)
+		{
+			if (BakeHelper)
+			{
+				BakeHelper->PreEvaluation(MovieScene, InFrame);
+			}
+		}
+		InSequencer->GetEvaluationTemplate().EvaluateSynchronousBlocking(Context, *InSequencer);
+
+		// evaluate constraints
+		for (const UTickableConstraint* InConstraint : AllConstraints)
+		{
+			InConstraint->Evaluate(true);
+		}
+		
+		for (IMovieSceneToolsAnimationBakeHelper* BakeHelper : BakeHelpers)
+		{
+			if (BakeHelper)
+			{
+				BakeHelper->PostEvaluation(MovieScene, InFrame);
+			}
+		}
+	};
+
+	for (IMovieSceneToolsAnimationBakeHelper* BakeHelper : BakeHelpers)
+	{
+		if (BakeHelper)
+		{
+			BakeHelper->StartBaking(MovieScene);
+		}
+	}
+	
+	for (int32 Index = 0; Index < NumFrames; ++Index)
+	{
+		// evaluate animation
+		EvaluateAt(InTime[Index]);
+
+		// store transforms        	
+		ChildLocals[Index] = Handle->GetLocalTransform();
+		ChildGlobals[Index] = Handle->GetGlobalTransform();
+		SpaceGlobals[Index] = Constraint->GetParentGlobalTransform();
+	}
+	
+	for (IMovieSceneToolsAnimationBakeHelper* BakeHelper : BakeHelpers)
+	{
+		if (BakeHelper)
+		{
+			BakeHelper->StopBaking(MovieScene);
+		}
+	}
+}
+
 TArray< TObjectPtr<UTickableConstraint> > FCompensationEvaluator::GetHandleTransformConstraints(UWorld* InWorld) const
 {
 	using ConstraintPtr = TObjectPtr<UTickableConstraint>;
@@ -947,4 +1030,109 @@ void FMovieSceneConstraintChannelHelper::CreateBindingIDForHandle(const TSharedP
 			}
 		}
 	}
+}
+
+void FMovieSceneConstraintChannelHelper::HandleConstraintPropertyChanged(
+	UTickableTransformConstraint* InConstraint,
+	const FMovieSceneConstraintChannel& InActiveChannel,
+	const FPropertyChangedEvent& InPropertyChangedEvent,
+	const TSharedPtr<ISequencer>& InSequencer,
+	UMovieSceneSection* InSection)
+{
+	if (!InConstraint || !InSection || !InSequencer.IsValid())
+	{
+		return;
+	}
+	
+	const FName PropertyName = InPropertyChangedEvent.GetPropertyName();
+	if (PropertyName == UTickableParentConstraint::GetScalingPropertyName())
+	{
+		return CompensateScale(Cast<UTickableParentConstraint>(InConstraint), InActiveChannel, InSequencer, InSection);
+	}
+}
+
+void FMovieSceneConstraintChannelHelper::CompensateScale(
+	UTickableParentConstraint* InParentConstraint,
+	const FMovieSceneConstraintChannel& InActiveChannel,
+	const TSharedPtr<ISequencer>& InSequencer,
+	UMovieSceneSection* InSection)
+{
+	if (!InParentConstraint)
+	{
+		return;
+	}
+
+	TObjectPtr<UTransformableHandle> Handle = InParentConstraint->ChildTRSHandle;
+	ITransformConstraintChannelInterface* Interface = GetHandleInterface(Handle);
+	if (!Interface)
+	{
+		return;
+	}
+
+	const TArrayView<const FFrameNumber> Times = InActiveChannel.GetTimes();
+	if (Times.IsEmpty())
+	{
+	    return;
+	}
+
+	// get transform channels
+	const TArrayView<FMovieSceneFloatChannel*> FloatTransformChannels = Handle->GetFloatChannels(InSection);
+	const TArrayView<FMovieSceneDoubleChannel*> DoubleTransformChannels = Handle->GetDoubleChannels(InSection);
+
+	// get frames after this time
+	TArray<FFrameNumber> ActiveTimes;
+	if (!FloatTransformChannels.IsEmpty())
+	{
+	    GetFramesWithinActiveState(InActiveChannel, FloatTransformChannels, ActiveTimes);
+	}
+	else
+	{
+	    GetFramesWithinActiveState(InActiveChannel, DoubleTransformChannels, ActiveTimes);
+	}
+
+	if (ActiveTimes.IsEmpty())
+	{
+	    return;
+	}
+
+	const bool bRefScalingValue = InParentConstraint->IsScalingEnabled();
+	
+	// if scaling has been enabled (bRefScalingValue == true), it means that it was not before the property has changed so
+	// the current scale channels values represent the local scale values of the handle
+	// if scaling has been disabled (bRefScalingValue == false), it means that it was before the property has changed so
+	// the current scale channels values represent the offset in the constraint space
+	
+	InParentConstraint->SetScaling(!bRefScalingValue);
+
+	TGuardValue<bool> CompensateGuard(bDoNotCompensate, true);
+	InSection->Modify();
+	
+	FCompensationEvaluator Evaluator(InParentConstraint);
+	UWorld* World = GCurrentLevelEditingViewportClient ? GCurrentLevelEditingViewportClient->GetWorld() : nullptr;
+	Evaluator.CacheTransforms(World, InSequencer, ActiveTimes);
+
+	TArray<FTransform>& ChildLocals = Evaluator.ChildLocals;
+	const TArray<FTransform>& ChildGlobals = Evaluator.ChildGlobals;
+	const TArray<FTransform>& SpaceGlobals = Evaluator.SpaceGlobals;
+
+	const int32 NumTransforms = ChildLocals.Num();
+
+	if (bRefScalingValue)
+	{
+	    // local scale values have to be switched to the constraint space to represent the offset
+	    for (int Index = 0; Index < NumTransforms; Index++)
+	    {
+    		FTransform& ChildLocal = ChildLocals[Index];
+    		const FTransform Offset = ChildGlobals[Index].GetRelativeTransform(SpaceGlobals[Index]);
+    		ChildLocal.SetScale3D(Offset.GetScale3D());
+	    }
+	}
+	// else ChildLocals already represents the data that needs to be keyed as it is the result of
+	// the constraint evaluation so it just needs to be keyed
+
+	// add keys
+	Interface->AddHandleTransformKeys(InSequencer, Handle, ActiveTimes, ChildLocals, EMovieSceneTransformChannel::Scale);
+
+	// reset scaling to reference value
+	InParentConstraint->SetScaling(bRefScalingValue);
 }
\ No newline at end of file
diff --git a/Engine/Source/Editor/MovieSceneTools/Private/TrackEditors/TransformTrackEditor.cpp b/Engine/Source/Editor/MovieSceneTools/Private/TrackEditors/TransformTrackEditor.cpp
--- a/Engine/Source/Editor/MovieSceneTools/Private/TrackEditors/TransformTrackEditor.cpp
+++ b/Engine/Source/Editor/MovieSceneTools/Private/TrackEditors/TransformTrackEditor.cpp
@@ -38,7 +38,9 @@
 #include "Editor.h"
 #include "LevelEditorViewport.h"
 #include "TransformConstraint.h"
+#include "TransformableHandle.h"
 #include "Constraints/MovieSceneConstraintChannelHelper.h"
+#include "Constraints/TransformConstraintChannelInterface.h"
 #include "Misc/TransactionObjectEvent.h"
 #include "Engine/Selection.h"
 
@@ -1450,6 +1452,7 @@ void F3DTransformTrackEditor::HandleOnConstraintAdded(IMovieSceneConstrainedSect
 					}
 				});
 	}
+
 	// Store Section so we can remove these delegates
 	UMovieScene3DTransformSection* Section = Cast<UMovieScene3DTransformSection>(InSection);
 	SectionsToClear.Add(Section);
@@ -1481,6 +1484,11 @@ void F3DTransformTrackEditor::HandleOnConstraintAdded(IMovieSceneConstrainedSect
 	{
 		HandleConstraintRemoved(InSection);
 	}
+
+	if (!UTickableTransformConstraint::GetOnConstraintChanged().IsBoundToObject(this))
+	{
+		UTickableTransformConstraint::GetOnConstraintChanged().AddRaw(this, &F3DTransformTrackEditor::HandleConstraintPropertyChanged);
+	}
 }
 
 static UTickableTransformConstraint* GetTickableTransformConstraint(IMovieSceneConstrainedSection* InSection, const FMovieSceneConstraintChannel* InConstraintChannel)
@@ -1602,6 +1610,50 @@ void F3DTransformTrackEditor::HandleConstraintRemoved(IMovieSceneConstrainedSect
 	}
 }
 
+void F3DTransformTrackEditor::HandleConstraintPropertyChanged(UTickableTransformConstraint* InConstraint, const FPropertyChangedEvent& InPropertyChangedEvent) const
+{
+	if (!IsValid(InConstraint))
+	{
+		return;
+	}
+
+	// find constraint section
+	const UTransformableComponentHandle* Handle = Cast<UTransformableComponentHandle>(InConstraint->ChildTRSHandle);
+	if (!IsValid(Handle) || !Handle->IsValid())
+	{
+		return;
+	}
+
+	const FConstraintChannelInterfaceRegistry& InterfaceRegistry = FConstraintChannelInterfaceRegistry::Get();	
+	ITransformConstraintChannelInterface* Interface = InterfaceRegistry.FindConstraintChannelInterface(Handle->GetClass());
+	if (!Interface)
+	{
+		return;
+	}
+	
+	UMovieSceneSection* Section = Interface->GetHandleConstraintSection(Handle, GetSequencer());
+	IMovieSceneConstrainedSection* ConstraintSection = Cast<IMovieSceneConstrainedSection>(Section);
+	if (!ConstraintSection)
+	{
+		return;
+	}
+
+	// find corresponding channel
+	const TArray<FConstraintAndActiveChannel>& ConstraintChannels = ConstraintSection->GetConstraintsChannels();
+	const FConstraintAndActiveChannel* Channel = ConstraintChannels.FindByPredicate([InConstraint](const FConstraintAndActiveChannel& Channel)
+	{
+		return Channel.Constraint == InConstraint || Channel.ConstraintCopyToSpawn == InConstraint;
+	});
+
+	if (!Channel)
+	{
+		return;
+	}
+
+	FMovieSceneConstraintChannelHelper::HandleConstraintPropertyChanged(
+			InConstraint, Channel->ActiveChannel, InPropertyChangedEvent, GetSequencer(), Section);
+}
+
 void F3DTransformTrackEditor::ClearOutConstraintDelegates() 
 {
 	UWorld* World = GCurrentLevelEditingViewportClient ? GCurrentLevelEditingViewportClient->GetWorld() : nullptr;
@@ -1634,5 +1686,7 @@ void F3DTransformTrackEditor::ClearOutConstraintDelegates()
 		}
 	}
 	SectionsToClear.Reset();
+
+	UTickableTransformConstraint::GetOnConstraintChanged().RemoveAll(this);
 }
 #undef LOCTEXT_NAMESPACE
diff --git a/Engine/Source/Editor/MovieSceneTools/Private/TrackEditors/TransformTrackEditor.h b/Engine/Source/Editor/MovieSceneTools/Private/TrackEditors/TransformTrackEditor.h
--- a/Engine/Source/Editor/MovieSceneTools/Private/TrackEditors/TransformTrackEditor.h
+++ b/Engine/Source/Editor/MovieSceneTools/Private/TrackEditors/TransformTrackEditor.h
@@ -20,6 +20,7 @@
 class AActor;
 struct FAssetData;
 class SHorizontalBox;
+class UTickableTransformConstraint;
 
 namespace UE { namespace MovieScene { struct FIntermediate3DTransform; } }
 
@@ -187,6 +188,8 @@ private:
 	void HandleConstraintKeyMoved(IMovieSceneConstrainedSection* InSection, const FMovieSceneConstraintChannel* InConstraintChannel,
 		const TArray<FKeyMoveEventItem>& InMovedItems);
 	void HandleConstraintRemoved(IMovieSceneConstrainedSection* InSection);
+	void HandleConstraintPropertyChanged(UTickableTransformConstraint* InConstraint, const FPropertyChangedEvent& InPropertyChangedEvent) const;
+	
 	void ClearOutConstraintDelegates();
 private:
 
diff --git a/Engine/Source/Editor/MovieSceneTools/Public/Constraints/MovieSceneConstraintChannelHelper.h b/Engine/Source/Editor/MovieSceneTools/Public/Constraints/MovieSceneConstraintChannelHelper.h
--- a/Engine/Source/Editor/MovieSceneTools/Public/Constraints/MovieSceneConstraintChannelHelper.h
+++ b/Engine/Source/Editor/MovieSceneTools/Public/Constraints/MovieSceneConstraintChannelHelper.h
@@ -12,6 +12,7 @@ class UTickableConstraint;
 class ISequencer;
 class UMovieSceneSection;
 class UTickableTransformConstraint;
+class UTickableParentConstraint;
 class UTransformableHandle;
 struct FMovieSceneConstraintChannel;
 struct FFrameNumber;
@@ -30,6 +31,8 @@ public:
 	void ComputeLocalTransformsBeforeDeletion(UWorld* InWorld, const TSharedPtr<ISequencer>& InSequencer, const TArray<FFrameNumber>& InFrames);
 	void ComputeCompensation(UWorld* InWorld, const TSharedPtr<ISequencer>& InSequencer, const FFrameNumber& InTime);
 	void ComputeLocalTransformsForBaking(UWorld* InWorld, const TSharedPtr<ISequencer>& InSequencer, const TArray<FFrameNumber>& InFrames);
+	void CacheTransforms(UWorld* InWorld, const TSharedPtr<ISequencer>& InSequencer, const TArray<FFrameNumber>& InFrames);
+	
 private:
 
 	TArray< TObjectPtr<UTickableConstraint> > GetHandleTransformConstraints(UWorld* InWorld) const;
@@ -118,6 +121,13 @@ public:
 		const TArrayView<ChannelType*>& InChannels,
 		const FFrameNumber& InTime);
 
+	static void HandleConstraintPropertyChanged(
+		UTickableTransformConstraint* InConstraint,
+		const FMovieSceneConstraintChannel& InActiveChannel,
+		const FPropertyChangedEvent& InPropertyChangedEvent,
+		const TSharedPtr<ISequencer>& InSequencer,
+		UMovieSceneSection* InSection);
+
 	static bool bDoNotCompensate;
 
 private:
@@ -126,4 +136,10 @@ private:
 	
 	/** For the given handle create any movie scene binding for it based upon the current sequencer that's open*/
 	static void CreateBindingIDForHandle(const TSharedPtr<ISequencer>& InSequencer, UTransformableHandle* InHandle);
+
+	static void CompensateScale(
+		UTickableParentConstraint* InParentConstraint,
+		const FMovieSceneConstraintChannel& InActiveChannel,
+		const TSharedPtr<ISequencer>& InSequencer,
+		UMovieSceneSection* InSection);
 };
diff --git a/Engine/Source/Runtime/Experimental/Animation/Constraints/Private/TransformConstraint.cpp b/Engine/Source/Runtime/Experimental/Animation/Constraints/Private/TransformConstraint.cpp
--- a/Engine/Source/Runtime/Experimental/Animation/Constraints/Private/TransformConstraint.cpp
+++ b/Engine/Source/Runtime/Experimental/Animation/Constraints/Private/TransformConstraint.cpp
@@ -198,6 +198,10 @@ private:
  * UTickableTransformConstraint
  **/
 
+#if WITH_EDITOR
+UTickableTransformConstraint::FOnConstraintChanged UTickableTransformConstraint::OnConstraintChanged;
+#endif // WITH_EDITOR
+
 int64 UTickableTransformConstraint::GetType() const
 {
 	return static_cast<int64>(Type);
@@ -305,6 +309,11 @@ void UTickableTransformConstraint::PostEditChangeProperty(FPropertyChangedEvent&
 	}
 }
 
+UTickableTransformConstraint::FOnConstraintChanged& UTickableTransformConstraint::GetOnConstraintChanged()
+{
+	return OnConstraintChanged;	
+}
+
 #endif
 
 void UTickableTransformConstraint::UnregisterDelegates() const
@@ -1167,25 +1176,43 @@ void UTickableParentConstraint::PostEditChangeProperty(FPropertyChangedEvent& Pr
 {
 	Super::PostEditChangeProperty(PropertyChangedEvent);
 
+	auto UpdateOffset = [&]()
+	{
+		FTransform ChildGlobalTransform = GetChildGlobalTransform();
+		if (!bScaling)
+		{
+			ChildGlobalTransform.RemoveScaling();
+		}
+		const FTransform ParentWorldTransform = GetParentGlobalTransform();
+		OffsetTransform = ChildGlobalTransform.GetRelativeTransform(ParentWorldTransform);
+	};
+	
 	const FName PropertyName = PropertyChangedEvent.GetPropertyName();
 	if (PropertyName == GET_MEMBER_NAME_CHECKED(UTickableTransformConstraint, bDynamicOffset))
 	{
 		if(bDynamicOffset)
 		{
 			Cache.CachedInputHash = CalculateInputHash();
-			
-			FTransform ChildGlobalTransform = GetChildGlobalTransform();
-			if (!bScaling)
-			{
-				ChildGlobalTransform.RemoveScaling();
-			}
-			const FTransform ParentWorldTransform = GetParentGlobalTransform();
-			OffsetTransform = ChildGlobalTransform.GetRelativeTransform(ParentWorldTransform);
-			
+			UpdateOffset();			
 			Evaluate();
 		}
 		return;
 	}
+
+	if (PropertyName == GET_MEMBER_NAME_CHECKED(UTickableParentConstraint, bScaling))
+	{
+		// notify scale change. note that this is currently the only property change we monitor but this broadcast
+		// call could be made higher in the hierarchy to monitor other changes.
+		OnConstraintChanged.Broadcast(this, PropertyChangedEvent);
+		
+		if (bMaintainOffset || bDynamicOffset)
+		{
+			Cache.CachedInputHash = CalculateInputHash();
+			UpdateOffset();
+			Evaluate();
+			return;
+		}
+	}
 }
 
 #endif
diff --git a/Engine/Source/Runtime/Experimental/Animation/Constraints/Public/TransformConstraint.h b/Engine/Source/Runtime/Experimental/Animation/Constraints/Public/TransformConstraint.h
--- a/Engine/Source/Runtime/Experimental/Animation/Constraints/Public/TransformConstraint.h
+++ b/Engine/Source/Runtime/Experimental/Animation/Constraints/Public/TransformConstraint.h
@@ -6,6 +6,9 @@
 #include "ConstraintsManager.h"
 #include "Kismet/BlueprintFunctionLibrary.h"
 
+#include "Delegates/DelegateCombinations.h"
+#include "Delegates/Delegate.h"
+
 #include "TransformConstraint.generated.h"
 
 enum class EHandleEvent : uint8;
@@ -160,6 +163,16 @@ public:
 	// UObject interface
 	virtual void PostEditChangeProperty(FPropertyChangedEvent& PropertyChangedEvent) override;
 	// End of UObject interface
+
+	/** Returns a delegate that can be used to monitor for property changes. This can be used to monitor constraints changes
+	 * only instead of using FCoreUObjectDelegates::OnObjectPropertyChanged that is listening to every objects.
+	 * Note that bScaling is currently the only property change we monitor but this can be used for other properties.
+	 */
+	DECLARE_MULTICAST_DELEGATE_TwoParams(FOnConstraintChanged, UTickableTransformConstraint*, const FPropertyChangedEvent&);
+	static FOnConstraintChanged& GetOnConstraintChanged();
+
+protected:
+	static FOnConstraintChanged OnConstraintChanged;
 #endif
 };
 
@@ -323,6 +336,12 @@ public:
 		return bScaling;
 	}
 	
+	void SetScaling(const bool bInScale)
+	{
+		bScaling = bInScale;
+	}
+
+	static FName GetScalingPropertyName() { return GET_MEMBER_NAME_CHECKED(UTickableParentConstraint, bScaling); }
 
 	/** Updates the dynamic offset based on external child's transform changes. */
 	virtual void OnHandleModified(UTransformableHandle* InHandle, EHandleEvent InEvent) override;
