Non-unity fixes
#rnx

diff --git a/Engine/Plugins/Experimental/TypedElementsDataStorage/Source/TypedElementsDataStorage/Private/TypedElementDatabase.h b/Engine/Plugins/Experimental/TypedElementsDataStorage/Source/TypedElementsDataStorage/Private/TypedElementDatabase.h
--- a/Engine/Plugins/Experimental/TypedElementsDataStorage/Source/TypedElementsDataStorage/Private/TypedElementDatabase.h
+++ b/Engine/Plugins/Experimental/TypedElementsDataStorage/Source/TypedElementsDataStorage/Private/TypedElementDatabase.h
@@ -8,6 +8,7 @@
 #include "MassEntityQuery.h"
 #include "TypedElementHandleStore.h"
 #include "UObject/ObjectMacros.h"
+#include "UObject/StrongObjectPtr.h"
 
 #include "TypedElementDatabase.generated.h"
 
diff --git a/Engine/Source/Editor/MovieSceneTools/Private/TrackEditors/MaterialParameterCollectionTrackEditor.cpp b/Engine/Source/Editor/MovieSceneTools/Private/TrackEditors/MaterialParameterCollectionTrackEditor.cpp
--- a/Engine/Source/Editor/MovieSceneTools/Private/TrackEditors/MaterialParameterCollectionTrackEditor.cpp
+++ b/Engine/Source/Editor/MovieSceneTools/Private/TrackEditors/MaterialParameterCollectionTrackEditor.cpp
@@ -9,6 +9,7 @@
 #include "ContentBrowserDelegates.h"
 #include "ContentBrowserModule.h"
 #include "Delegates/Delegate.h"
+#include "Editor.h"
 #include "Editor/UnrealEdEngine.h"
 #include "Framework/Application/SlateApplication.h"
 #include "Framework/Commands/UIAction.h"
diff --git a/Engine/Source/Editor/MovieSceneTools/Private/TrackEditors/MaterialParameterCollectionTrackEditor.h b/Engine/Source/Editor/MovieSceneTools/Private/TrackEditors/MaterialParameterCollectionTrackEditor.h
--- a/Engine/Source/Editor/MovieSceneTools/Private/TrackEditors/MaterialParameterCollectionTrackEditor.h
+++ b/Engine/Source/Editor/MovieSceneTools/Private/TrackEditors/MaterialParameterCollectionTrackEditor.h
@@ -12,6 +12,7 @@ class ISequencer;
 class ISequencerSection;
 class ISequencerTrackEditor;
 class SWidget;
+class UMaterialParameterCollection;
 class UMovieSceneMaterialParameterCollectionTrack;
 class UMovieSceneSection;
 class UMovieSceneSequence;
