Sequencer: Fix a crash undoing the creation of a subsequence.

Move sequence creation outside scope of transaction so that the asset creation is not transacted and the internals will not be undone

diff --git a/Engine/Source/Editor/MovieSceneTools/Private/MovieSceneToolHelpers.cpp b/Engine/Source/Editor/MovieSceneTools/Private/MovieSceneToolHelpers.cpp
--- a/Engine/Source/Editor/MovieSceneTools/Private/MovieSceneToolHelpers.cpp
+++ b/Engine/Source/Editor/MovieSceneTools/Private/MovieSceneToolHelpers.cpp
@@ -532,6 +532,41 @@ FString MovieSceneToolHelpers::GenerateNewShotName(const TArray<UMovieSceneSecti
 	return ComposeShotName(ProjectSettings->ShotPrefix, ProjectSettings->FirstShotNumber, ProjectSettings->FirstTakeNumber, ProjectSettings->ShotNumDigits, ProjectSettings->TakeNumDigits);
 }
 
+UMovieSceneSequence* MovieSceneToolHelpers::CreateSequence(FString& NewSequenceName, FString& NewSequencePath, UMovieSceneSubSection* SectionToDuplicate)
+{
+	// Create a new level sequence asset with the appropriate name
+	IAssetTools& AssetTools = FModuleManager::GetModuleChecked<FAssetToolsModule>("AssetTools").Get();
+
+	UObject* NewAsset = nullptr;
+	for (TObjectIterator<UClass> It; It; ++It)
+	{
+		UClass* CurrentClass = *It;
+		if (CurrentClass->IsChildOf(UFactory::StaticClass()) && !(CurrentClass->HasAnyClassFlags(CLASS_Abstract)))
+		{
+			UFactory* Factory = Cast<UFactory>(CurrentClass->GetDefaultObject());
+			if (Factory->CanCreateNew() && Factory->ImportPriority >= 0 && Factory->SupportedClass == ULevelSequence::StaticClass())
+			{
+				if (SectionToDuplicate != nullptr)
+				{
+					NewAsset = AssetTools.DuplicateAssetWithDialog(NewSequenceName, NewSequencePath, SectionToDuplicate->GetSequence());
+				}
+				else
+				{
+					NewAsset = AssetTools.CreateAssetWithDialog(NewSequenceName, NewSequencePath, ULevelSequence::StaticClass(), Factory);
+				}
+				break;
+			}
+		}
+	}
+
+	if (NewAsset == nullptr)
+	{
+		return nullptr;
+	}
+
+	return Cast<UMovieSceneSequence>(NewAsset);
+}
+
 UMovieSceneSubSection* MovieSceneToolHelpers::CreateSubSequence(FString& NewSequenceName, FString& NewSequencePath, FFrameNumber NewSequenceStartTime, UMovieSceneSubTrack* SubTrack, UMovieSceneSubSection* SectionToDuplicate)
 {
 	// Create a new level sequence asset with the appropriate name
diff --git a/Engine/Source/Editor/MovieSceneTools/Private/TrackEditors/CinematicShotTrackEditor.cpp b/Engine/Source/Editor/MovieSceneTools/Private/TrackEditors/CinematicShotTrackEditor.cpp
--- a/Engine/Source/Editor/MovieSceneTools/Private/TrackEditors/CinematicShotTrackEditor.cpp
+++ b/Engine/Source/Editor/MovieSceneTools/Private/TrackEditors/CinematicShotTrackEditor.cpp
@@ -369,24 +369,28 @@ FReply FCinematicShotTrackEditor::OnDrop(const FDragDropEvent& DragDropEvent, co
 
 void FCinematicShotTrackEditor::InsertShot()
 {
-	const FScopedTransaction Transaction(LOCTEXT("InsertShot_Transaction", "Insert Shot"));
-
 	FFrameTime NewShotStartTime = GetSequencer()->GetLocalTime().Time;
 
 	UMovieSceneCinematicShotTrack* CinematicShotTrack = FindOrCreateCinematicShotTrack();
 	FString NewShotName = MovieSceneToolHelpers::GenerateNewShotName(CinematicShotTrack->GetAllSections(), NewShotStartTime.FrameNumber);
 	FString NewSequencePath = MovieSceneToolHelpers::GenerateNewShotPath(GetSequencer()->GetFocusedMovieSceneSequence()->GetMovieScene(), NewShotName);
 
-	UMovieSceneSubSection* NewShot = MovieSceneToolHelpers::CreateSubSequence(NewShotName, NewSequencePath, NewShotStartTime.FrameNumber, CinematicShotTrack);
-	if (NewShot)
+	if (UMovieSceneSequence* NewSequence = MovieSceneToolHelpers::CreateSequence(NewShotName, NewSequencePath))
 	{
-		NewShot->SetRowIndex(MovieSceneToolHelpers::FindAvailableRowIndex(CinematicShotTrack, NewShot));
-	}
+		const FScopedTransaction Transaction(LOCTEXT("InsertShot_Transaction", "Insert Shot"));
 
-	GetSequencer()->NotifyMovieSceneDataChanged( EMovieSceneDataChangeType::MovieSceneStructureItemAdded );
-	GetSequencer()->EmptySelection();
-	GetSequencer()->SelectSection(NewShot);
-	GetSequencer()->ThrobSectionSelection();
+		int32 Duration = UE::MovieScene::DiscreteSize(NewSequence->GetMovieScene()->GetPlaybackRange());
+
+		if (UMovieSceneSubSection* NewShot = CinematicShotTrack->AddSequence(NewSequence, NewShotStartTime.FrameNumber, Duration))
+		{
+			NewShot->SetRowIndex(MovieSceneToolHelpers::FindAvailableRowIndex(CinematicShotTrack, NewShot));
+
+			GetSequencer()->NotifyMovieSceneDataChanged(EMovieSceneDataChangeType::MovieSceneStructureItemAdded);
+			GetSequencer()->EmptySelection();
+			GetSequencer()->SelectSection(NewShot);
+			GetSequencer()->ThrobSectionSelection();
+		}
+	}
 }
 
 
@@ -419,28 +423,35 @@ void FCinematicShotTrackEditor::InsertFiller()
 
 void FCinematicShotTrackEditor::DuplicateShot(UMovieSceneCinematicShotSection* Section)
 {
-	const FScopedTransaction Transaction(LOCTEXT("DuplicateShot_Transaction", "Duplicate Shot"));
-
 	UMovieSceneCinematicShotTrack* CinematicShotTrack = FindOrCreateCinematicShotTrack();
 
 	FFrameNumber StartTime = Section->HasStartFrame() ? Section->GetInclusiveStartFrame() : 0;
 	FString NewShotName = MovieSceneToolHelpers::GenerateNewShotName(CinematicShotTrack->GetAllSections(), StartTime);
 	FString NewSequencePath = FPaths::GetPath(Section->GetSequence()->GetPathName());
 
+
+
 	// Duplicate the shot and put it on the next available row
-	UMovieSceneSubSection* NewShot = MovieSceneToolHelpers::CreateSubSequence(NewShotName, NewSequencePath, StartTime, CinematicShotTrack, Section);
-	if (NewShot)
+	UMovieSceneSequence* NewSequence = MovieSceneToolHelpers::CreateSequence(NewShotName, NewSequencePath, Section);
+	if (NewSequence)
 	{
-		NewShot->SetRange(Section->GetRange());
-		NewShot->SetRowIndex(MovieSceneToolHelpers::FindAvailableRowIndex(CinematicShotTrack, NewShot));
-		NewShot->Parameters.StartFrameOffset = Section->Parameters.StartFrameOffset;
-		NewShot->Parameters.TimeScale = Section->Parameters.TimeScale;
-		NewShot->SetPreRollFrames(Section->GetPreRollFrames());
+		const FScopedTransaction Transaction(LOCTEXT("DuplicateShot_Transaction", "Duplicate Shot"));
 
-		GetSequencer()->NotifyMovieSceneDataChanged( EMovieSceneDataChangeType::MovieSceneStructureItemAdded );
-		GetSequencer()->EmptySelection();
-		GetSequencer()->SelectSection(NewShot);
-		GetSequencer()->ThrobSectionSelection();
+		int32 Duration = UE::MovieScene::DiscreteSize(Section->GetRange());
+
+		if (UMovieSceneSubSection* NewShot = CinematicShotTrack->AddSequence(NewSequence, StartTime, Duration))
+		{
+			NewShot->SetRange(Section->GetRange());
+			NewShot->SetRowIndex(MovieSceneToolHelpers::FindAvailableRowIndex(CinematicShotTrack, NewShot));
+			NewShot->Parameters.StartFrameOffset = Section->Parameters.StartFrameOffset;
+			NewShot->Parameters.TimeScale = Section->Parameters.TimeScale;
+			NewShot->SetPreRollFrames(Section->GetPreRollFrames());
+
+			GetSequencer()->NotifyMovieSceneDataChanged( EMovieSceneDataChangeType::MovieSceneStructureItemAdded );
+			GetSequencer()->EmptySelection();
+			GetSequencer()->SelectSection(NewShot);
+			GetSequencer()->ThrobSectionSelection();
+		}
 	}
 }
 
@@ -459,8 +470,6 @@ void FCinematicShotTrackEditor::RenameShot(UMovieSceneCinematicShotSection* Sect
 
 void FCinematicShotTrackEditor::NewTake(UMovieSceneCinematicShotSection* Section)
 {
-	const FScopedTransaction Transaction(LOCTEXT("NewTake_Transaction", "New Take"));
-
 	FString ShotPrefix;
 	uint32 ShotNumber = INDEX_NONE;
 	uint32 TakeNumber = INDEX_NONE;
@@ -497,10 +506,13 @@ void FCinematicShotTrackEditor::NewTake(UMovieSceneCinematicShotSection* Section
 		UMovieSceneCinematicShotTrack* CinematicShotTrack = FindOrCreateCinematicShotTrack();
 		FString NewSequencePath = FPaths::GetPath(Section->GetSequence()->GetPathName());
 
-		UMovieSceneSubSection* NewShot = MovieSceneToolHelpers::CreateSubSequence(NewShotName, NewSequencePath, NewShotStartTime, CinematicShotTrack, Section);
-
-		if (NewShot)
+		if (UMovieSceneSequence* NewSequence = MovieSceneToolHelpers::CreateSequence(NewShotName, NewSequencePath, Section))
 		{
+			const FScopedTransaction Transaction(LOCTEXT("NewTake_Transaction", "New Take"));
+
+			int32 Duration = UE::MovieScene::DiscreteSize(Section->GetRange());
+
+			UMovieSceneSubSection* NewShot = CinematicShotTrack->AddSequence(NewSequence, NewShotStartTime, Duration);
 			CinematicShotTrack->RemoveSection(*Section);
 
 			NewShot->SetRange(NewShotRange);
diff --git a/Engine/Source/Editor/MovieSceneTools/Private/TrackEditors/SubTrackEditor.cpp b/Engine/Source/Editor/MovieSceneTools/Private/TrackEditors/SubTrackEditor.cpp
--- a/Engine/Source/Editor/MovieSceneTools/Private/TrackEditors/SubTrackEditor.cpp
+++ b/Engine/Source/Editor/MovieSceneTools/Private/TrackEditors/SubTrackEditor.cpp
@@ -541,8 +541,6 @@ void FSubTrackEditor::HandleAddSubSequenceComboButtonMenuEntryEnterPressed(const
 
 void FSubTrackEditor::InsertSequence(UMovieSceneTrack* Track)
 {
-	const FScopedTransaction Transaction(LOCTEXT("InsertSequence_Transaction", "Insert Sequence"));
-
 	FFrameTime NewSectionStartTime = GetSequencer()->GetLocalTime().Time;
 
 	UMovieSceneSubTrack* SubTrack = Cast<UMovieSceneSubTrack>(Track);
@@ -557,16 +555,22 @@ void FSubTrackEditor::InsertSequence(UMovieSceneTrack* Track)
 	FAssetToolsModule& AssetToolsModule = FModuleManager::GetModuleChecked<FAssetToolsModule>("AssetTools");
 	AssetToolsModule.Get().CreateUniqueAssetName(NewSequencePath + TEXT("/") + NewSequenceName, TEXT(""), NewSequencePath, NewSequenceName);
 
-	UMovieSceneSubSection* NewSection = MovieSceneToolHelpers::CreateSubSequence(NewSequenceName, NewSequencePath, NewSectionStartTime.FrameNumber, SubTrack);
-	if (NewSection)
+	if (UMovieSceneSequence* NewSequence = MovieSceneToolHelpers::CreateSequence(NewSequenceName, NewSequencePath))
 	{
-		NewSection->SetRowIndex(MovieSceneToolHelpers::FindAvailableRowIndex(Track, NewSection));
-	}
+		const FScopedTransaction Transaction(LOCTEXT("InsertSequence_Transaction", "Insert Sequence"));
 
-	GetSequencer()->NotifyMovieSceneDataChanged(EMovieSceneDataChangeType::MovieSceneStructureItemAdded);
-	GetSequencer()->EmptySelection();
-	GetSequencer()->SelectSection(NewSection);
-	GetSequencer()->ThrobSectionSelection();
+		int32 Duration = UE::MovieScene::DiscreteSize(NewSequence->GetMovieScene()->GetPlaybackRange());
+
+		if (UMovieSceneSubSection* NewSection = SubTrack->AddSequence(NewSequence, NewSectionStartTime.FrameNumber, Duration))
+		{
+			NewSection->SetRowIndex(MovieSceneToolHelpers::FindAvailableRowIndex(Track, NewSection));
+
+			GetSequencer()->NotifyMovieSceneDataChanged(EMovieSceneDataChangeType::MovieSceneStructureItemAdded);
+			GetSequencer()->EmptySelection();
+			GetSequencer()->SelectSection(NewSection);
+			GetSequencer()->ThrobSectionSelection();
+		}
+	}
 }
 
 FKeyPropertyResult FSubTrackEditor::AddKeyInternal(FFrameNumber KeyTime, UMovieSceneSequence* InMovieSceneSequence, UMovieSceneTrack* InTrack, int32 RowIndex)
diff --git a/Engine/Source/Editor/MovieSceneTools/Public/MovieSceneToolHelpers.h b/Engine/Source/Editor/MovieSceneTools/Public/MovieSceneToolHelpers.h
--- a/Engine/Source/Editor/MovieSceneTools/Public/MovieSceneToolHelpers.h
+++ b/Engine/Source/Editor/MovieSceneTools/Public/MovieSceneToolHelpers.h
@@ -170,15 +170,16 @@ public:
 	static FString GenerateNewShotName(const TArray<UMovieSceneSection*>& AllSections, FFrameNumber Time);
 
 	/*
-	 * Create sub sequence
+	 * Create sequence
 	 *
 	 * @param NewSequenceName The new sequence name.
 	 * @param NewSequencePath The new sequence path. 
-	 * @param NewSequenceStartTime The time to start the new sequence at.
-	 * @param SubTrack The track to put the sequence onto.
 	 * @param SectionToDuplicate The section to duplicate.
 	 * @return The new subsequence.
 	 */
+	static UMovieSceneSequence* CreateSequence(FString& NewSequenceName, FString& NewSequencePath, UMovieSceneSubSection* SectionToDuplicate = nullptr);
+
+	UE_DEPRECATED(5.2, "CreateSubSequence has been deprecated in favor of CreateSequence so that asset creation can be invoked outside of a transaction (with the intent that it is undoable) and then adding the sequence to a subtrack can be invoked within a transaction")
 	static UMovieSceneSubSection* CreateSubSequence(FString& NewSequenceName, FString& NewSequencePath, FFrameNumber NewSequenceStartTime, UMovieSceneSubTrack* SubTrack, UMovieSceneSubSection* SectionToDuplicate = nullptr);
 
 	/**
