Sequencer: refactor utilities for managing BP director event track endpoints to be more generally usable

This prepares the terrain for the user-defined dynamic bindings, which also need to create similar endpoints in the BP director.

diff --git a/Engine/Plugins/MovieScene/SequencerScripting/Source/SequencerScriptingEditor/Private/SequencerTools.cpp b/Engine/Plugins/MovieScene/SequencerScripting/Source/SequencerScriptingEditor/Private/SequencerTools.cpp
--- a/Engine/Plugins/MovieScene/SequencerScripting/Source/SequencerScriptingEditor/Private/SequencerTools.cpp
+++ b/Engine/Plugins/MovieScene/SequencerScripting/Source/SequencerScriptingEditor/Private/SequencerTools.cpp
@@ -784,7 +784,7 @@ FMovieSceneEvent USequencerToolsFunctionLibrary::CreateEvent(UMovieSceneSequence
 	InSection->Modify();
 	FMovieSceneEventUtils::BindEventSectionToBlueprint(InSection, InEndpoint.EventEndpoint->GetBlueprint());
 
-	UEdGraphPin* BoundObjectPin = FMovieSceneEventUtils::FindBoundObjectPin(InEndpoint.EventEndpoint, BoundObjectPinClass);
+	UEdGraphPin* BoundObjectPin = FMovieSceneDirectorBlueprintUtils::FindCallTargetPin(InEndpoint.EventEndpoint, BoundObjectPinClass);
 	FMovieSceneEventUtils::SetEndpoint(&Event, InSection, InEndpoint.EventEndpoint, BoundObjectPin);
 
 	if (InEndpoint.PayloadNames.Num() != InPayload.Num())
@@ -832,10 +832,11 @@ FSequencerQuickBindingResult USequencerToolsFunctionLibrary::CreateQuickBinding(
 
 	UMovieScene* MovieScene = InSequence->GetMovieScene();
 	
-	FMovieSceneEventEndpointParameters Params;
-	Params.SanitizedObjectName = InObject->GetName();
-	Params.SanitizedEventName = InFunctionName;
-	Params.BoundObjectPinClass = InObject->GetClass();
+	FMovieSceneDirectorBlueprintEndpointDefinition EndpointDefinition;
+	EndpointDefinition.EndpointType = EMovieSceneDirectorBlueprintEndpointType::Event;
+	EndpointDefinition.EndpointName = InFunctionName;
+	EndpointDefinition.PossibleCallTargetClass = InObject->GetClass();
+	EndpointDefinition.AddExtraOutputPin(InObject->GetName(), UEdGraphSchema_K2::PC_Object, InObject->GetClass());
 	UFunction* Function = InObject->GetClass()->FindFunctionByName(FName(InFunctionName));
 	if (Function == nullptr)
 	{
@@ -847,12 +848,12 @@ FSequencerQuickBindingResult USequencerToolsFunctionLibrary::CreateQuickBinding(
 	UBlueprintFunctionNodeSpawner* BlueprintFunctionNodeSpawner = UBlueprintFunctionNodeSpawner::Create(Function);
 	FBlueprintActionMenuItem Action(BlueprintFunctionNodeSpawner);
 
-	UK2Node_CustomEvent* NewEventEndpoint = FMovieSceneEventUtils::CreateUserFacingEvent(Blueprint, Params);
+	UK2Node_CustomEvent* NewEventEndpoint = FMovieSceneDirectorBlueprintUtils::CreateEventEndpoint(Blueprint, EndpointDefinition);
 	NewEventEndpoint->bCallInEditor = bCallInEditor;
 	Result.EventEndpoint = NewEventEndpoint;
 
 	UEdGraphPin* ThenPin = NewEventEndpoint->FindPin(UEdGraphSchema_K2::PN_Then, EGPD_Output);
-	UEdGraphPin* BoundObjectPin = FMovieSceneEventUtils::FindBoundObjectPin(NewEventEndpoint, Params.BoundObjectPinClass);
+	UEdGraphPin* BoundObjectPin = FMovieSceneDirectorBlueprintUtils::FindCallTargetPin(NewEventEndpoint, EndpointDefinition.PossibleCallTargetClass);
 
 	FVector2D NodePosition(NewEventEndpoint->NodePosX + 400.f, NewEventEndpoint->NodePosY);
 	UEdGraphNode* NewNode = Action.PerformAction(NewEventEndpoint->GetGraph(), BoundObjectPin ? BoundObjectPin : ThenPin, NodePosition);
diff --git a/Engine/Source/Editor/MovieSceneTools/Private/MovieSceneDirectorBlueprintEndpointCustomization.cpp b/Engine/Source/Editor/MovieSceneTools/Private/MovieSceneDirectorBlueprintEndpointCustomization.cpp
new file mode 100644
--- /dev/null
+++ b/Engine/Source/Editor/MovieSceneTools/Private/MovieSceneDirectorBlueprintEndpointCustomization.cpp
@@ -0,0 +1,1516 @@
+// Copyright Epic Games, Inc. All Rights Reserved.
+
+#include "MovieSceneDirectorBlueprintEndpointCustomization.h"
+#include "Modules/ModuleManager.h"
+
+#include "UObject/UnrealType.h"
+#include "Algo/Find.h"
+#include "ISequencerModule.h"
+#include "MovieSceneSequence.h"
+#include "ScopedTransaction.h"
+
+#include "Framework/Application/SlateApplication.h"
+#include "Framework/MultiBox/MultiBoxBuilder.h"
+#include "Widgets/DeclarativeSyntaxSupport.h"
+#include "Widgets/Text/STextBlock.h"
+#include "Widgets/Images/SImage.h"
+#include "Widgets/Layout/SBox.h"
+#include "Widgets/Input/SCheckBox.h"
+#include "Widgets/Input/SComboButton.h"
+#include "Widgets/Input/SHyperlink.h"
+#include "Widgets/Input/SEditableTextBox.h"
+#include "Widgets/Views/SExpanderArrow.h"
+
+#include "PropertyHandle.h"
+#include "IPropertyUtilities.h"
+#include "PropertyCustomizationHelpers.h"
+#include "IDetailChildrenBuilder.h"
+#include "DetailLayoutBuilder.h"
+#include "DetailCategoryBuilder.h"
+
+#include "K2Node_Variable.h"
+#include "K2Node_CustomEvent.h"
+#include "K2Node_CallFunction.h"
+#include "K2Node_FunctionEntry.h"
+#include "K2Node_FunctionResult.h"
+#include "Kismet2/KismetEditorUtilities.h"
+#include "Kismet2/BlueprintEditorUtils.h"
+#include "Kismet2/Kismet2NameValidators.h"
+#include "SGraphActionMenu.h"
+#include "BlueprintActionMenuUtils.h"
+#include "BlueprintActionMenuBuilder.h"
+#include "BlueprintActionMenuItem.h"
+#include "BlueprintFunctionNodeSpawner.h"
+
+#include "Styling/AppStyle.h"
+#include "EditorFontGlyphs.h"
+
+#define LOCTEXT_NAMESPACE "MovieSceneDirectorBlueprintEndpointCustomization"
+
+void FMovieSceneDirectorBlueprintEndpointCustomization::CustomizeHeader(TSharedRef<IPropertyHandle> InPropertyHandle, FDetailWidgetRow& HeaderRow, IPropertyTypeCustomizationUtils& CustomizationUtils)
+{
+}
+
+void FMovieSceneDirectorBlueprintEndpointCustomization::CustomizeChildren(TSharedRef<IPropertyHandle> InPropertyHandle, IDetailChildrenBuilder& ChildBuilder, IPropertyTypeCustomizationUtils& CustomizationUtils)
+{
+	PropertyHandle = InPropertyHandle;
+	PropertyUtilities = CustomizationUtils.GetPropertyUtilities();
+
+	ChildBuilder.AddCustomRow(FText())
+	.NameContent()
+	[
+		SNew(STextBlock)
+		.Font(CustomizationUtils.GetRegularFont())
+		.Text(LOCTEXT("EndpointValueText", "Endpoint"))
+	]
+	.ValueContent()
+	.MinDesiredWidth(200.f)
+	[
+		SNew(SHorizontalBox)
+
+		+ SHorizontalBox::Slot()
+		[
+			SNew(SComboButton)
+			.ButtonStyle(FAppStyle::Get(), "HoverHintOnly")
+			.ForegroundColor(FSlateColor::UseForeground())
+			.OnGetMenuContent(this, &FMovieSceneDirectorBlueprintEndpointCustomization::GetMenuContent)
+			.CollapseMenuOnParentFocus(true)
+			.ContentPadding(FMargin(4.f, 0.f))
+			.ButtonContent()
+			[
+				SNew(SHorizontalBox)
+
+				+ SHorizontalBox::Slot()
+				.Padding(FMargin(0.f, 0.f, 4.f, 0.f))
+				.VAlign(VAlign_Center)
+				.AutoWidth()
+				[
+					SNew(SImage)
+					.Image(this, &FMovieSceneDirectorBlueprintEndpointCustomization::GetEndpointIcon)
+				]
+
+				+ SHorizontalBox::Slot()
+				.Padding(FMargin(0.f, 0.f, 4.f, 0.f))
+				.VAlign(VAlign_Center)
+				[
+					SNew(STextBlock)
+					.Font(CustomizationUtils.GetRegularFont())
+					.Text(this, &FMovieSceneDirectorBlueprintEndpointCustomization::GetEndpointName)
+				]
+			]
+		]
+
+		+ SHorizontalBox::Slot()
+		.VAlign(VAlign_Center)
+		.AutoWidth()
+		[
+			PropertyCustomizationHelpers::MakeBrowseButton(FSimpleDelegate::CreateSP(this, &FMovieSceneDirectorBlueprintEndpointCustomization::NavigateToDefinition), LOCTEXT("NavigateToDefinition_Tip", "Navigate to this endpoint's definition"))
+		]
+	];
+
+	const bool bAnyBoundEndpoints = GetAllValidEndpoints().Num() != 0;
+	if (bAnyBoundEndpoints)
+	{
+		FText CallInEditorText    = LOCTEXT("CallInEditor_Label",   "Call In Editor");
+		FText CallInEditorTooltip = LOCTEXT("CallInEditor_Tooltip", "When checked, this endpoint will be triggered in the Editor outside of PIE.\n\nBEWARE: ANY CHANGES MADE AS A RESULT OF THIS ENDPOINT BEING CALLED MAY END UP BEING SAVED IN THE CURRENT LEVEL OR ASSET.");
+
+		ChildBuilder.AddCustomRow(CallInEditorText)
+		.NameContent()
+		[
+			SNew(STextBlock)
+			.Text(CallInEditorText)
+			.ToolTipText(CallInEditorTooltip)
+			.Font(IDetailLayoutBuilder::GetDetailFont())
+		]
+		.ValueContent()
+		[
+			SNew(SCheckBox)
+			.ToolTipText(CallInEditorTooltip)
+			.IsChecked(this, &FMovieSceneDirectorBlueprintEndpointCustomization::GetCallInEditorCheckState)
+			.OnCheckStateChanged(this, &FMovieSceneDirectorBlueprintEndpointCustomization::OnSetCallInEditorCheckState)
+		];
+	}
+
+	UK2Node* CommonEndpoint = GetCommonEndpoint();
+	if (!CommonEndpoint)
+	{
+		return;
+	}
+
+	TArray<FWellKnownParameterCandidates> WellKnownParameterCandidates;
+	GetWellKnownParameterCandidates(CommonEndpoint, WellKnownParameterCandidates);
+	for (int32 ParamIndex = 0; ParamIndex < WellKnownParameterCandidates.Num(); ++ParamIndex)
+	{
+		const FWellKnownParameterCandidates& Candidates = WellKnownParameterCandidates[ParamIndex];
+		const FWellKnownParameterMetadata& ParamMetadata = Candidates.Metadata;
+
+		if (!Candidates.bShowUnmatchedParameters && Candidates.CandidatePinNames.IsEmpty())
+		{
+			continue;
+		}
+
+		ChildBuilder.AddCustomRow(ParamMetadata.PickerLabel)
+		.NameContent()
+		[
+			SNew(STextBlock)
+			.Text(ParamMetadata.PickerLabel)
+			.ToolTipText(ParamMetadata.PickerTooltip)
+			.Font(IDetailLayoutBuilder::GetDetailFont())
+		]
+		.ValueContent()
+		[
+			SNew(SComboButton)
+			.ButtonStyle(FAppStyle::Get(), "HoverHintOnly")
+			.ForegroundColor(FSlateColor::UseForeground())
+			.ToolTipText(ParamMetadata.PickerTooltip)
+			.OnGetMenuContent(this, &FMovieSceneDirectorBlueprintEndpointCustomization::GetWellKnownParameterPinMenuContent, ParamIndex)
+			.CollapseMenuOnParentFocus(true)
+			.ContentPadding(FMargin(4.f, 0.f))
+			.ButtonContent()
+			[
+				SNew(SBox)
+				.HeightOverride(21.f)
+				[
+					SNew(SHorizontalBox)
+
+					+ SHorizontalBox::Slot()
+					.Padding(FMargin(0.f, 0.f, 4.f, 0.f))
+					.VAlign(VAlign_Center)
+					.AutoWidth()
+					[
+						SNew(SImage)
+						.Image(this, &FMovieSceneDirectorBlueprintEndpointCustomization::GetWellKnownParameterPinIcon, ParamIndex)
+					]
+
+					+ SHorizontalBox::Slot()
+					.Padding(FMargin(0.f, 0.f, 4.f, 0.f))
+					.VAlign(VAlign_Center)
+					[
+						SNew(STextBlock)
+						.Font(CustomizationUtils.GetRegularFont())
+						.Text(this, &FMovieSceneDirectorBlueprintEndpointCustomization::GetWellKnownParameterPinText, ParamIndex)
+					]
+				]
+			]
+		];
+	}
+
+	UFunction* CommonFunction = nullptr;
+	UBlueprint* Blueprint = CommonEndpoint->HasValidBlueprint() ? CommonEndpoint->GetBlueprint() : nullptr;
+	if (Blueprint)
+	{
+		if (UK2Node_Event* Event = Cast<UK2Node_Event>(CommonEndpoint))
+		{
+			CommonFunction = Blueprint->SkeletonGeneratedClass ?
+				Blueprint->SkeletonGeneratedClass->FindFunctionByName(Event->GetFunctionName()) :
+				nullptr;
+		}
+		else if (UK2Node_FunctionEntry* EndpointEntry = Cast<UK2Node_FunctionEntry>(CommonEndpoint))
+		{
+			CommonFunction = EndpointEntry->FindSignatureFunction();
+		}
+		else
+		{
+			// @todo: Error not supported
+		}
+
+		Blueprint->OnCompiled().AddSP(this, &FMovieSceneDirectorBlueprintEndpointCustomization::OnBlueprintCompiled);
+	}
+
+	if (CommonFunction)
+	{
+		IDetailCategoryBuilder& DetailCategoryBuilder = ChildBuilder
+			.GetParentCategory()
+			.GetParentLayout()
+			.EditCategory("Payload", LOCTEXT("PayloadLabel", "Payload"), ECategoryPriority::Uncommon);
+
+		bool bPayloadUpToDate = true;
+
+		TArray<UObject*> EditObjects;
+		GetEditObjects(EditObjects);
+
+		TArray<void*> RawData;
+		PropertyHandle->AccessRawData(RawData);
+
+		FPayloadVariableMap PayloadVariables;
+		TArray<FName> WellKnownParameters;
+		TArray<FName, TInlineAllocator<8>> AllValidNames;
+
+		const FString* WorldContextParamName = CommonFunction->FindMetaData(FBlueprintMetadata::MD_WorldContext);
+
+		for (int32 Index = 0; Index < RawData.Num(); ++Index)
+		{
+			AllValidNames.Empty();
+			PayloadVariables.Empty();
+			WellKnownParameters.Empty();
+
+			GetPayloadVariables(EditObjects[Index], RawData[Index], PayloadVariables);
+			GetWellKnownParameterPinNames(EditObjects[Index], RawData[Index], WellKnownParameters);
+
+			TSharedPtr<FStructOnScope> StructData = MakeShared<FStructOnScope>(CommonFunction);
+
+			for (FProperty* Field : TFieldRange<FProperty>(CommonFunction))
+			{
+				// Ignore parameters we can't use as inputs.
+				if (Field->HasAnyPropertyFlags(CPF_OutParm | CPF_ReturnParm | CPF_ReferenceParm))
+				{
+					continue;
+				}
+
+				// Ignore "well-known parameters", i.e. parameters for which the system
+				// will pass specific values not meant to be defined by the user.
+				if (WellKnownParameters.Contains(Field->GetFName()))
+				{
+					continue;
+				}
+
+				// Ignore the world context parameter.
+				if (WorldContextParamName && Field->GetName() == *WorldContextParamName)
+				{
+					continue;
+				}
+
+				const FString* PayloadVariable = PayloadVariables.Find(Field->GetFName());
+				if (PayloadVariable)
+				{
+					AllValidNames.Add(Field->GetFName());
+
+					// We have an override for this variable
+					const bool bImportSuccess = FBlueprintEditorUtils::PropertyValueFromString(Field, *PayloadVariable, StructData->GetStructMemory());
+					if (!bImportSuccess)
+					{
+						// @todo: error
+					}
+				}
+				else if (UEdGraphPin* Pin = CommonEndpoint->FindPin(Field->GetFName(), EGPD_Output))
+				{
+					if (!Pin->DefaultValue.IsEmpty())
+					{
+						const bool bImportSuccess = FBlueprintEditorUtils::PropertyValueFromString(Field, Pin->DefaultValue, StructData->GetStructMemory());
+						if (!bImportSuccess)
+						{
+							// @todo: error
+						}
+					}
+				}
+
+				IDetailPropertyRow* ExternalRow = DetailCategoryBuilder.AddExternalStructureProperty(StructData.ToSharedRef(), Field->GetFName(), EPropertyLocation::Default, FAddPropertyParams().ForceShowProperty());
+
+				TSharedPtr<IPropertyHandle> LocalVariableProperty = ExternalRow->GetPropertyHandle();
+				FSimpleDelegate Delegate = FSimpleDelegate::CreateSP(this, &FMovieSceneDirectorBlueprintEndpointCustomization::OnPayloadVariableChanged, StructData.ToSharedRef(), LocalVariableProperty);
+
+				LocalVariableProperty->SetOnPropertyValueChanged(Delegate);
+				LocalVariableProperty->SetOnChildPropertyValueChanged(Delegate);
+			}
+
+			bPayloadUpToDate &= (AllValidNames.Num() == PayloadVariables.Num());
+		}
+
+		if (!bPayloadUpToDate)
+		{
+			DetailCategoryBuilder.AddCustomRow(FText())
+			.WholeRowContent()
+			[
+				SNew(SHorizontalBox)
+
+				+ SHorizontalBox::Slot()
+				.Padding(FMargin(5.f, 0.f, 5.f, 0.f))
+				.AutoWidth()
+				[
+					SNew(STextBlock)
+					.TextStyle(FAppStyle::Get(), "Log.Warning")
+					.Font(FAppStyle::GetFontStyle("FontAwesome.10"))
+					.Text(FEditorFontGlyphs::Exclamation_Triangle)
+				]
+
+				+ SHorizontalBox::Slot()
+				.Padding(FMargin(0.f, 0.f, 5.f, 0.f))
+				[
+					SNew(STextBlock)
+					.TextStyle(FAppStyle::Get(), "Log.Warning")
+					.Text(LOCTEXT("PayloadOutOfDateError", "Payload variables may be out-of-date. Please compile the blueprint."))
+				]
+			];
+		}
+	}
+}
+
+void FMovieSceneDirectorBlueprintEndpointCustomization::OnPayloadVariableChanged(TSharedRef<FStructOnScope> InStructData, TSharedPtr<IPropertyHandle> LocalVariableProperty)
+{
+	// This function should only ever be bound if all the entry points call the same function
+	FProperty* Property = LocalVariableProperty->GetProperty();
+	if (!Property)
+	{
+		return;
+	}
+
+	FScopedTransaction Transaction(LOCTEXT("SetPayloadValue", "Set payload value"));
+
+	bool bChangedAnything = false;
+
+	FString NewValueString;
+
+	const bool bSuccessfulTextExport = FBlueprintEditorUtils::PropertyValueToString(Property, InStructData->GetStructMemory(), NewValueString, nullptr);
+	if (!bSuccessfulTextExport)
+	{
+		// @todo: error
+		return;
+	}
+
+	TArray<UObject*> EditObjects;
+	GetEditObjects(EditObjects);
+
+	TArray<void*> RawData;
+	PropertyHandle->AccessRawData(RawData);
+
+	if (!ensure(RawData.Num() == EditObjects.Num()))
+	{
+		return;
+	}
+
+	for (int32 Index = 0; Index < RawData.Num(); ++Index)
+	{
+		bChangedAnything |= SetPayloadVariable(EditObjects[Index], RawData[Index], Property->GetFName(), NewValueString);
+	}
+
+	if (bChangedAnything)
+	{
+		UK2Node* CommonEndpoint = GetCommonEndpoint();
+
+		UBlueprint* BP = (CommonEndpoint && CommonEndpoint->HasValidBlueprint()) ? CommonEndpoint->GetBlueprint() : nullptr;
+		if (BP)
+		{
+			FBlueprintEditorUtils::MarkBlueprintAsModified(BP);
+		}
+	}
+	else
+	{
+		Transaction.Cancel();
+	}
+}
+
+ECheckBoxState FMovieSceneDirectorBlueprintEndpointCustomization::GetCallInEditorCheckState() const
+{
+	ECheckBoxState CheckState = ECheckBoxState::Undetermined;
+	for (UK2Node* Endpoint : GetAllValidEndpoints())
+	{
+		bool bCallInEditor = false;
+		if (UK2Node_CustomEvent* CustomEvent = Cast<UK2Node_CustomEvent>(Endpoint))
+		{
+			bCallInEditor = CustomEvent->bCallInEditor;
+		}
+		else if (UK2Node_FunctionEntry* FunctionEntry = Cast<UK2Node_FunctionEntry>(Endpoint))
+		{
+			bCallInEditor = FunctionEntry->MetaData.bCallInEditor;
+		}
+
+		if (CheckState == ECheckBoxState::Undetermined)
+		{
+			CheckState = bCallInEditor ? ECheckBoxState::Checked : ECheckBoxState::Unchecked;
+		}
+		else if (bCallInEditor != (CheckState == ECheckBoxState::Checked) )
+		{
+			return ECheckBoxState::Undetermined;
+		}
+	}
+	return CheckState;
+}
+
+void FMovieSceneDirectorBlueprintEndpointCustomization::OnSetCallInEditorCheckState(ECheckBoxState NewCheckedState)
+{
+	FScopedTransaction Transaction(LOCTEXT("SetCallInEditor", "Set Call in Editor"));
+
+	const bool bCallInEditor = (NewCheckedState == ECheckBoxState::Checked);
+
+	TSet<UBlueprint*> Blueprints;
+
+	for (UK2Node* Endpoint : GetAllValidEndpoints())
+	{
+		if (UK2Node_CustomEvent* CustomEvent = Cast<UK2Node_CustomEvent>(Endpoint))
+		{
+			CustomEvent->Modify();
+			CustomEvent->bCallInEditor = bCallInEditor;
+
+			if (CustomEvent->HasValidBlueprint())
+			{
+				Blueprints.Add(CustomEvent->GetBlueprint());
+			}
+		}
+		else if (UK2Node_FunctionEntry* FunctionEntry = Cast<UK2Node_FunctionEntry>(Endpoint))
+		{
+			FunctionEntry->Modify();
+			FunctionEntry->MetaData.bCallInEditor = bCallInEditor;
+
+			if (FunctionEntry->HasValidBlueprint())
+			{
+				Blueprints.Add(FunctionEntry->GetBlueprint());
+			}
+		}
+	}
+
+	for (UBlueprint* Blueprint : Blueprints)
+	{
+		FBlueprintEditorUtils::MarkBlueprintAsModified(Blueprint);
+	}
+}
+
+void FMovieSceneDirectorBlueprintEndpointCustomization::OnBlueprintCompiled(UBlueprint*)
+{
+	PropertyUtilities->ForceRefresh();
+}
+
+UMovieSceneSequence* FMovieSceneDirectorBlueprintEndpointCustomization::GetCommonSequence() const
+{
+	TArray<UObject*> EditObjects;
+	GetEditObjects(EditObjects);
+
+	UMovieSceneSequence* CommonSequence = nullptr;
+
+	for (UObject* Obj : EditObjects)
+	{
+		UMovieSceneSequence* ThisSequence = Obj ? Obj->GetTypedOuter<UMovieSceneSequence>() : nullptr;
+		if (CommonSequence && CommonSequence != ThisSequence)
+		{
+			return nullptr;
+		}
+
+		CommonSequence = ThisSequence;
+	}
+	return CommonSequence;
+}
+
+void FMovieSceneDirectorBlueprintEndpointCustomization::IterateEndpoints(TFunctionRef<bool(UK2Node*)> Callback) const
+{
+	TArray<UObject*> EditObjects;
+	GetEditObjects(EditObjects);
+
+	TArray<void*> RawData;
+	PropertyHandle->AccessRawData(RawData);
+
+	if (!ensure(RawData.Num() == EditObjects.Num()))
+	{
+		return;
+	}
+
+	for (int32 Index = 0; Index < RawData.Num(); ++Index)
+	{
+		UK2Node* Endpoint = FindEndpoint(EditObjects[Index], RawData[Index]);
+		if (Endpoint)
+		{
+			if (!Callback(Endpoint))
+			{
+				return;
+			}
+		}
+	}
+}
+
+TArray<UK2Node*> FMovieSceneDirectorBlueprintEndpointCustomization::GetAllValidEndpoints() const
+{
+	TArray<UK2Node*> Endpoints;
+
+	IterateEndpoints(
+		[&Endpoints](UK2Node* Endpoint)
+		{
+			if (Endpoint)
+			{
+				Endpoints.Add(Endpoint);
+			}
+			return true;
+		}
+	);
+
+	return Endpoints;
+}
+
+UK2Node* FMovieSceneDirectorBlueprintEndpointCustomization::GetCommonEndpoint() const
+{
+	UK2Node* CommonEndpoint = nullptr;
+
+	IterateEndpoints(
+		[&CommonEndpoint](UK2Node* Endpoint)
+		{
+			if (CommonEndpoint && Endpoint != CommonEndpoint)
+			{
+				CommonEndpoint = nullptr;
+				return false;
+			}
+
+			CommonEndpoint = Endpoint;
+			return true;
+		}
+	);
+
+	return CommonEndpoint;
+}
+
+void FMovieSceneDirectorBlueprintEndpointCustomization::GetEditObjects(TArray<UObject*>& OutObjects) const
+{
+	PropertyHandle->GetOuterObjects(OutObjects);
+}
+
+TSharedRef<SWidget> FMovieSceneDirectorBlueprintEndpointCustomization::GetMenuContent()
+{
+	FMenuBuilder MenuBuilder(true, nullptr, nullptr, true);
+
+	UMovieSceneSequence* Sequence = GetCommonSequence();
+
+	MenuBuilder.AddMenuEntry(
+		LOCTEXT("CreateEndpoint_Text",    "Create New Endpoint"),
+		LOCTEXT("CreateEndpoint_Tooltip", "Creates a new endpoint in this sequence's blueprint."),
+		FSlateIcon(FAppStyle::GetAppStyleSetName(), "Sequencer.CreateEventBinding"),
+		FUIAction(
+			FExecuteAction::CreateSP(this, &FMovieSceneDirectorBlueprintEndpointCustomization::CreateEndpoint)
+		)
+	);
+
+	const bool bAnyBoundEndpoints = GetAllValidEndpoints().Num() != 0;
+	if (!bAnyBoundEndpoints && Sequence)
+	{
+		MenuBuilder.AddSubMenu(
+			LOCTEXT("CreateQuickBinding_Text",    "Quick Bind"),
+			LOCTEXT("CreateQuickBinding_Tooltip", "Shows a list of functions on this object binding that can be bound directly to this endpoint."),
+			FNewMenuDelegate::CreateSP(this, &FMovieSceneDirectorBlueprintEndpointCustomization::PopulateQuickBindSubMenu, Sequence),
+			false /* bInOpenSubMenuOnClick */,
+			FSlateIcon(FAppStyle::GetAppStyleSetName(), "Sequencer.CreateQuickBinding"),
+			false /* bInShouldWindowAfterMenuSelection */
+		);
+	}
+	else
+	{
+		MenuBuilder.AddSubMenu(
+			LOCTEXT("Rebind_Text",         "Rebind To"),
+			LOCTEXT("Rebind_Text_Tooltip", "Rebinds this endpoint to a different function call or event node."),
+			FNewMenuDelegate::CreateSP(this, &FMovieSceneDirectorBlueprintEndpointCustomization::PopulateRebindSubMenu, Sequence),
+			false /* bInOpenSubMenuOnClick */,
+			FSlateIcon(FAppStyle::GetAppStyleSetName(), "Sequencer.CreateQuickBinding"),
+			false /* bInShouldCloseWindowAfterMenuSelection */
+		);
+
+		MenuBuilder.AddMenuEntry(
+			LOCTEXT("ClearEndpoint_Text",    "Clear"),
+			LOCTEXT("ClearEndpoint_Tooltip", "Unbinds this endpoint from its current binding."),
+			FSlateIcon(FAppStyle::GetAppStyleSetName(), "Sequencer.ClearEventBinding"),
+			FUIAction(
+				FExecuteAction::CreateSP(this, &FMovieSceneDirectorBlueprintEndpointCustomization::ClearEndpoint)
+			)
+		);
+	}
+
+	return MenuBuilder.MakeWidget();
+}
+
+TSharedRef<SWidget> FMovieSceneDirectorBlueprintEndpointCustomization::GetWellKnownParameterPinMenuContent(int32 ParameterIndex)
+{
+	UK2Node* CommonEndpoint = GetCommonEndpoint();
+	if (!CommonEndpoint)
+	{
+		return SNew(STextBlock).Text(LOCTEXT("WellKnownParameterPinError_MultipleEndpoints", "Cannot choose a parameter pin with multiple endpoints selected"));
+	}
+
+	TArray<FWellKnownParameterCandidates> WellKnownParameters;
+	GetWellKnownParameterCandidates(CommonEndpoint, WellKnownParameters);
+	if (!WellKnownParameters.IsValidIndex(ParameterIndex) ||
+			WellKnownParameters[ParameterIndex].CandidatePinNames.IsEmpty())
+	{
+		return SNew(STextBlock).Text(LOCTEXT("WellKnownParameterPinError_NoPins", "No compatible pins were found."));
+	}
+
+	FMenuBuilder MenuBuilder(true, nullptr, nullptr, true);
+	for (FName CandidatePinName : WellKnownParameters[ParameterIndex].CandidatePinNames)
+	{
+		FText PinText = FText::FromName(CandidatePinName);
+		MenuBuilder.AddMenuEntry(
+			PinText,
+			FText::Format(LOCTEXT("SetWellKnownParameterPin_Tooltip", "When calling the endpoint, pass this parameter through pin {0}."), FText::FromName(CandidatePinName)),
+			FSlateIcon(),
+			FUIAction(
+				FExecuteAction::CreateSP(this, &FMovieSceneDirectorBlueprintEndpointCustomization::SetWellKnownParameterPinName, ParameterIndex, CandidatePinName),
+				FCanExecuteAction(),
+				FIsActionChecked::CreateSP(this, &FMovieSceneDirectorBlueprintEndpointCustomization::CompareWellKnownParameterPinName, ParameterIndex, CandidatePinName)
+				),
+			NAME_None,
+			EUserInterfaceActionType::RadioButton
+			);
+	}
+
+	MenuBuilder.AddMenuEntry(
+		LOCTEXT("ClearWellKnownParameterPin_Label", "Clear"),
+		LOCTEXT("ClearWellKnownParameterPin_Tooltip", "Clears the parameter binding, meaning that this parameter won't be passed to the endpoint anymore."),
+		FSlateIcon(),
+		FUIAction(
+			FExecuteAction::CreateSP(this, &FMovieSceneDirectorBlueprintEndpointCustomization::SetWellKnownParameterPinName, ParameterIndex, FName())
+		)
+	);
+
+	return MenuBuilder.MakeWidget();
+}
+
+bool FMovieSceneDirectorBlueprintEndpointCustomization::CompareWellKnownParameterPinName(int32 ParameterIndex, FName InPinName) const
+{
+	TArray<UObject*> EditObjects;
+	GetEditObjects(EditObjects);
+
+	TArray<void*> RawData;
+	PropertyHandle->AccessRawData(RawData);
+
+	check(EditObjects.Num() == RawData.Num());
+
+	TArray<FName> WellKnownParameters;
+	for (int32 Index = 0; Index < RawData.Num(); ++Index)
+	{
+		WellKnownParameters.Empty();
+		GetWellKnownParameterPinNames(EditObjects[Index], RawData[Index], WellKnownParameters);
+		if (WellKnownParameters.IsValidIndex(ParameterIndex) &&
+				WellKnownParameters[ParameterIndex] == InPinName)
+		{
+			return true;
+		}
+	}
+
+	return false;
+}
+
+void FMovieSceneDirectorBlueprintEndpointCustomization::SetWellKnownParameterPinName(int32 ParameterIndex, FName InNewBoundPinName)
+{
+	FScopedTransaction Transaction(LOCTEXT("SetWellKnownParameterPinTransaction", "Set Bound Object Pin"));
+
+	TArray<UObject*> EditObjects;
+	GetEditObjects(EditObjects);
+
+	TArray<void*> RawData;
+	PropertyHandle->AccessRawData(RawData);
+
+	check(EditObjects.Num() == RawData.Num());
+
+	for (int32 Index = 0; Index < RawData.Num(); ++Index)
+	{
+		EditObjects[Index]->Modify();
+
+		const bool bWasSet = SetWellKnownParameterPinName(EditObjects[Index], RawData[Index], ParameterIndex, InNewBoundPinName);
+		ensureMsgf(bWasSet, TEXT("Trying to set bound pin %s for parameter %d but sub-class did not handle it."), 
+				*InNewBoundPinName.ToString(), ParameterIndex);
+	}
+
+	// Ensure that anything listening for property changed notifications are notified of the new binding
+	PropertyHandle->NotifyFinishedChangingProperties();
+
+	FSequenceDataMap AllSequenceData;
+	GatherSequenceData(AllSequenceData);
+	for (const TPair<UMovieSceneSequence*, FSequenceData>& SequenceData : AllSequenceData)
+	{
+		FKismetEditorUtilities::CompileBlueprint(SequenceData.Value.Blueprint);
+	}
+
+	PropertyUtilities->ForceRefresh();
+}
+
+const FSlateBrush* FMovieSceneDirectorBlueprintEndpointCustomization::GetWellKnownParameterPinIcon(int32 ParameterIndex) const
+{
+	const bool bIsUnbound = CompareWellKnownParameterPinName(ParameterIndex, NAME_None);
+	if (bIsUnbound)
+	{
+		return nullptr;
+	}
+
+	return FAppStyle::GetBrush("Graph.Pin.Disconnected_VarA");
+}
+
+FText FMovieSceneDirectorBlueprintEndpointCustomization::GetWellKnownParameterPinText(int32 ParameterIndex) const
+{
+	TArray<UObject*> EditObjects;
+	GetEditObjects(EditObjects);
+
+	TArray<void*> RawData;
+	PropertyHandle->AccessRawData(RawData);
+
+	check(EditObjects.Num() == RawData.Num());
+
+	FName CommonPinName;
+	TArray<FName> WellKnownParameters;
+
+	for (int32 Index = 0; Index < RawData.Num(); ++Index)
+	{
+		WellKnownParameters.Empty();
+		GetWellKnownParameterPinNames(EditObjects[Index], RawData[Index], WellKnownParameters);
+		if (WellKnownParameters.IsValidIndex(ParameterIndex))
+		{
+			FName PinName = WellKnownParameters[ParameterIndex];
+			if (CommonPinName == NAME_None)
+			{
+				CommonPinName = PinName;
+			}
+			else if (CommonPinName != PinName)
+			{
+				return LOCTEXT("MultiplePinValues", "Multiple Values");
+			}
+		}
+	}
+
+	return FText::FromName(CommonPinName);
+}
+
+void FMovieSceneDirectorBlueprintEndpointCustomization::PopulateQuickBindSubMenu(FMenuBuilder& MenuBuilder, UMovieSceneSequence* Sequence)
+{
+	FMovieSceneSequenceEditor* SequenceEditor = FMovieSceneSequenceEditor::Find(Sequence);
+	if (!SequenceEditor)
+	{
+		return;
+	}
+
+	UBlueprint* Blueprint = SequenceEditor->GetOrCreateDirectorBlueprint(Sequence);
+	if (!Blueprint)
+	{
+		return;
+	}
+
+	FMovieSceneDirectorBlueprintEndpointDefinition EndpointDefinition = GenerateEndpointDefinition(Sequence);
+
+	TSharedRef<SGraphActionMenu> ActionMenu = SNew(SGraphActionMenu)
+		.OnCreateCustomRowExpander_Static([](const FCustomExpanderData& Data) -> TSharedRef<SExpanderArrow> { return SNew(SExpanderArrow, Data.TableRow); })
+		.OnCollectAllActions(this, &FMovieSceneDirectorBlueprintEndpointCustomization::CollectQuickBindActions, Blueprint, EndpointDefinition)
+		.OnActionSelected(this, &FMovieSceneDirectorBlueprintEndpointCustomization::HandleQuickBindActionSelected, Blueprint, EndpointDefinition);
+
+	ActionMenu->RegisterActiveTimer(0.f, FWidgetActiveTimerDelegate::CreateLambda(
+		[FilterTextBox = ActionMenu->GetFilterTextBox()](double, float)
+		{
+			FSlateApplication::Get().SetKeyboardFocus(FilterTextBox);
+			return EActiveTimerReturnType::Stop;
+		}
+	));
+
+	MenuBuilder.AddWidget(
+		SNew(SBox)
+		.WidthOverride(300.f)
+		.MaxDesiredHeight(500.f)
+		[
+			ActionMenu
+		],
+		FText()
+	);
+}
+
+void FMovieSceneDirectorBlueprintEndpointCustomization::CollectQuickBindActions(FGraphActionListBuilderBase& OutAllActions, UBlueprint* Blueprint, FMovieSceneDirectorBlueprintEndpointDefinition EndpointDefinition)
+{
+	// Build up the context object
+	auto RejectAnyNonFunctions = [](const FBlueprintActionFilter& Filter, FBlueprintActionInfo& BlueprintAction)
+	{
+		const UFunction* Function = BlueprintAction.GetAssociatedFunction();
+		return Function == nullptr || Function->HasAnyFunctionFlags(FUNC_BlueprintPure);
+	};
+	auto RejectAnyUnboundActions = [](const FBlueprintActionFilter& Filter, FBlueprintActionInfo& BlueprintAction)
+	{
+		return BlueprintAction.GetBindings().Num() <= 0;
+	};
+
+	FObjectProperty* ReturnProperty = EndpointDefinition.EndpointSignature ?
+		CastField<FObjectProperty>(EndpointDefinition.EndpointSignature->GetReturnProperty()) : nullptr;
+	auto RejectAnyIncompatibleReturnValues = [ReturnProperty](const FBlueprintActionFilter& Filter, FBlueprintActionInfo& BlueprintAction)
+	{
+		const UFunction* Function = BlueprintAction.GetAssociatedFunction();
+		const FObjectProperty* FunctionReturnProperty = CastField<FObjectProperty>(Function->GetReturnProperty());
+		return !FunctionReturnProperty || FunctionReturnProperty->PropertyClass != ReturnProperty->PropertyClass;
+	};
+
+	FBlueprintActionMenuBuilder ContextMenuBuilder;
+
+	if (EndpointDefinition.PossibleCallTargetClass)
+	{
+		// Add actions that are relevant to the bound object from the pin class
+		FBlueprintActionFilter CallOnMemberFilter(FBlueprintActionFilter::BPFILTER_RejectGlobalFields | FBlueprintActionFilter::BPFILTER_RejectPermittedSubClasses);
+		CallOnMemberFilter.PermittedNodeTypes.Add(UK2Node_CallFunction::StaticClass());
+		CallOnMemberFilter.Context.Blueprints.Add(Blueprint);
+
+		for (FObjectProperty* ObjectProperty : TFieldRange<FObjectProperty>(EndpointDefinition.PossibleCallTargetClass))
+		{
+			if (ObjectProperty->HasAnyPropertyFlags(CPF_BlueprintVisible) && (ObjectProperty->HasMetaData(FBlueprintMetadata::MD_ExposeFunctionCategories) || FBlueprintEditorUtils::IsSCSComponentProperty(ObjectProperty)))
+			{
+				CallOnMemberFilter.Context.SelectedObjects.Add(ObjectProperty);
+				FBlueprintActionFilter::AddUnique(CallOnMemberFilter.TargetClasses, ObjectProperty->PropertyClass);
+			}
+		}
+
+		FBlueprintActionFilter::AddUnique(CallOnMemberFilter.TargetClasses, EndpointDefinition.PossibleCallTargetClass);
+
+		// This removes duplicate entries (ie. Set Static Mesh and Set Static Mesh (StaticMeshComponent)), 
+		// but also prevents displaying functions on BP components. Comment out for now.
+		//CallOnMemberFilter.AddRejectionTest(FBlueprintActionFilter::FRejectionTestDelegate::CreateStatic(RejectAnyUnboundActions));
+		
+		CallOnMemberFilter.AddRejectionTest(FBlueprintActionFilter::FRejectionTestDelegate::CreateStatic(RejectAnyNonFunctions));
+
+		if (ReturnProperty)
+		{
+			CallOnMemberFilter.AddRejectionTest(FBlueprintActionFilter::FRejectionTestDelegate::CreateLambda(RejectAnyIncompatibleReturnValues));
+		}
+
+		ContextMenuBuilder.AddMenuSection(CallOnMemberFilter, FText::FromName(EndpointDefinition.PossibleCallTargetClass->GetFName()), 0);
+	}
+
+	{
+		// Add all actions that are relevant to the sequence director BP itself
+		FBlueprintActionFilter MenuFilter(FBlueprintActionFilter::BPFILTER_RejectGlobalFields | FBlueprintActionFilter::BPFILTER_RejectPermittedSubClasses);
+
+		MenuFilter.PermittedNodeTypes.Add(UK2Node_CallFunction::StaticClass());
+
+		MenuFilter.Context.Blueprints.Add(Blueprint);
+		MenuFilter.Context.Graphs.Append(Blueprint->UbergraphPages);
+		MenuFilter.Context.Graphs.Append(Blueprint->FunctionGraphs);
+
+		if (Blueprint->SkeletonGeneratedClass)
+		{
+			FBlueprintActionFilter::AddUnique(MenuFilter.TargetClasses, Blueprint->SkeletonGeneratedClass);
+		}
+		MenuFilter.AddRejectionTest(FBlueprintActionFilter::FRejectionTestDelegate::CreateStatic(RejectAnyNonFunctions));
+		
+		if (ReturnProperty)
+		{
+			MenuFilter.AddRejectionTest(FBlueprintActionFilter::FRejectionTestDelegate::CreateLambda(RejectAnyIncompatibleReturnValues));
+		}
+
+		ContextMenuBuilder.AddMenuSection(MenuFilter, LOCTEXT("SequenceDirectorMenu", "This Sequence"), 0);
+	}
+
+	{
+		OnCollectQuickBindActions(Blueprint, ContextMenuBuilder);
+	}
+
+	ContextMenuBuilder.RebuildActionList();
+
+	OutAllActions.Append(ContextMenuBuilder);
+}
+
+void FMovieSceneDirectorBlueprintEndpointCustomization::HandleQuickBindActionSelected(const TArray<TSharedPtr<FEdGraphSchemaAction>>& SelectedAction, ESelectInfo::Type InSelectionType, UBlueprint* Blueprint, FMovieSceneDirectorBlueprintEndpointDefinition EndpointDefinition)
+{
+	if (InSelectionType != ESelectInfo::OnMouseClick && InSelectionType != ESelectInfo::OnKeyPress)
+	{
+		return;
+	}
+
+	for (TSharedPtr<FEdGraphSchemaAction> Action : SelectedAction)
+	{
+		if (Action->GetTypeId() == FBlueprintActionMenuItem::StaticGetTypeId())
+		{
+			const UBlueprintFunctionNodeSpawner* FunctionNodeSpawner = Cast<UBlueprintFunctionNodeSpawner>(static_cast<FBlueprintActionMenuItem*>(Action.Get())->GetRawAction());
+			const UFunction* FunctionToCall = FunctionNodeSpawner ? FunctionNodeSpawner->GetFunction() : nullptr;
+
+			if (FunctionToCall)
+			{
+				UClass* OuterClass = CastChecked<UClass>(FunctionToCall->GetOuter());
+				if (OuterClass->ClassGeneratedBy == Blueprint && Blueprint->SkeletonGeneratedClass)
+				{
+					// Attempt to locate a custom event or a function graph of this name on the blueprint
+					for (UEdGraph* Graph : Blueprint->UbergraphPages)
+					{
+						for (UEdGraphNode* Node : Graph->Nodes)
+						{
+							UK2Node_CustomEvent* CustomEvent = Cast<UK2Node_CustomEvent>(Node);
+							if (CustomEvent && Blueprint->SkeletonGeneratedClass->FindFunctionByName(CustomEvent->GetFunctionName()) == FunctionToCall)
+							{
+								// Use this custom event
+								SetEndpoint(EndpointDefinition, CustomEvent, CustomEvent, EAutoCreatePayload::Variables);
+								return;
+							}
+						}
+					}
+
+					for (UEdGraph* Graph : Blueprint->FunctionGraphs)
+					{
+						if (Blueprint->SkeletonGeneratedClass->FindFunctionByName(Graph->GetFName()) != FunctionToCall)
+						{
+							continue;
+						}
+
+						// Use this function graph for the event endpoint
+						for (UEdGraphNode* Node : Graph->Nodes)
+						{
+							UK2Node_FunctionEntry* FunctionEntry = Cast<UK2Node_FunctionEntry>(Node);
+							if (FunctionEntry)
+							{
+								SetEndpoint(EndpointDefinition, FunctionEntry, FunctionEntry, EAutoCreatePayload::Variables);
+								return;
+							}
+						}
+					}
+				}
+			}
+		}
+
+		if (Action)
+		{
+			UK2Node* NewEndpoint = FMovieSceneDirectorBlueprintUtils::CreateEndpoint(Blueprint, EndpointDefinition);
+
+			UEdGraphPin* EndpointThenPin = NewEndpoint->FindPin(UEdGraphSchema_K2::PN_Then, EGPD_Output);
+			UEdGraphPin* CallTargetPin = FMovieSceneDirectorBlueprintUtils::FindCallTargetPin(NewEndpoint, EndpointDefinition.PossibleCallTargetClass);
+
+			FVector2D NodePosition(NewEndpoint->NodePosX + 400.f, NewEndpoint->NodePosY + 100.f);
+			UEdGraphNode* NewNode = Action->PerformAction(NewEndpoint->GetGraph(), CallTargetPin ? CallTargetPin : EndpointThenPin, NodePosition);
+
+			if (NewNode)
+			{
+				// If the new node has an exec pin, connect it. It may not have one if it's a BlueprintPure function.
+				UEdGraphPin* NewNodeExecPin = NewNode->FindPin(UEdGraphSchema_K2::PN_Execute, EGPD_Input);
+				if (EndpointThenPin && NewNodeExecPin)
+				{
+					EndpointThenPin->MakeLinkTo(NewNodeExecPin);
+				}
+
+				TArray<UK2Node_FunctionResult*> ResultNodes;
+				NewNode->GetGraph()->GetNodesOfClass(ResultNodes);
+				if (ResultNodes.Num() > 0)
+				{
+					// If there is a result node, move it past the endpoint call and connect it.
+					ResultNodes[0]->NodePosX = NodePosition.X + 400.f;
+
+					UEdGraphPin* NewNodeThenPin = NewNode->FindPin(UEdGraphSchema_K2::PN_Execute, EGPD_Input);
+					UEdGraphPin* ResultExecPin = ResultNodes[0]->FindPin(UEdGraphSchema_K2::PN_Execute, EGPD_Input);
+					if (NewNodeThenPin && ResultExecPin)
+					{
+						NewNodeThenPin->MakeLinkTo(ResultExecPin);
+					}
+
+					// If the new node has a return value, and if the endpoint has one too, try to connect them together.
+					UEdGraphPin* OutputPin = ResultNodes[0]->FindPin(UEdGraphSchema_K2::PN_ReturnValue, EGPD_Input);
+					UEdGraphPin* NewNodeReturnValuePin = NewNode->FindPin(UEdGraphSchema_K2::PN_ReturnValue, EGPD_Output);
+					if (OutputPin && NewNodeReturnValuePin)
+					{
+						// Connect the nodes.
+						NewNodeReturnValuePin->MakeLinkTo(OutputPin);
+					}
+				}
+			}
+
+			SetEndpoint(EndpointDefinition, NewEndpoint, Cast<UK2Node>(NewNode), EAutoCreatePayload::Pins | EAutoCreatePayload::Variables);
+		}
+	}
+}
+
+void FMovieSceneDirectorBlueprintEndpointCustomization::PopulateRebindSubMenu(FMenuBuilder& MenuBuilder, UMovieSceneSequence* Sequence)
+{
+	FMovieSceneSequenceEditor* SequenceEditor = FMovieSceneSequenceEditor::Find(Sequence);
+	if (!SequenceEditor)
+	{
+		return;
+	}
+
+	UBlueprint* Blueprint = SequenceEditor->GetOrCreateDirectorBlueprint(Sequence);
+	if (!Blueprint)
+	{
+		return;
+	}
+
+	FMovieSceneDirectorBlueprintEndpointDefinition EndpointDefinition = GenerateEndpointDefinition(Sequence);
+
+	TSharedRef<SGraphActionMenu> ActionMenu = SNew(SGraphActionMenu)
+		.OnCreateCustomRowExpander_Static([](const FCustomExpanderData& Data) -> TSharedRef<SExpanderArrow> { return SNew(SExpanderArrow, Data.TableRow); })
+		.OnCollectAllActions(this, &FMovieSceneDirectorBlueprintEndpointCustomization::CollectAllRebindActions, Blueprint, EndpointDefinition)
+		.OnActionSelected(this, &FMovieSceneDirectorBlueprintEndpointCustomization::HandleRebindActionSelected, Blueprint, EndpointDefinition);
+
+	ActionMenu->RegisterActiveTimer(0.f, FWidgetActiveTimerDelegate::CreateLambda(
+		[FilterTextBox = ActionMenu->GetFilterTextBox()](double, float)
+		{
+			FSlateApplication::Get().SetKeyboardFocus(FilterTextBox);
+			return EActiveTimerReturnType::Stop;
+		}
+	));
+
+	MenuBuilder.AddWidget(
+		SNew(SBox)
+		.WidthOverride(300.f)
+		.MaxDesiredHeight(500.f)
+		[
+			ActionMenu
+		],
+		FText()
+	);
+}
+
+void FMovieSceneDirectorBlueprintEndpointCustomization::CollectAllRebindActions(FGraphActionListBuilderBase& OutAllActions, UBlueprint* Blueprint, FMovieSceneDirectorBlueprintEndpointDefinition EndpointDefinition)
+{
+	// Build up the context object
+	auto RejectAnyForeignFunctions = [](const FBlueprintActionFilter& Filter, FBlueprintActionInfo& BlueprintAction, UBlueprint* Blueprint)
+	{
+		const UBlueprintFunctionNodeSpawner* FunctionNodeSpawner = Cast<UBlueprintFunctionNodeSpawner>(BlueprintAction.NodeSpawner);
+		const UFunction* FunctionToCall = FunctionNodeSpawner ? FunctionNodeSpawner->GetFunction() : nullptr;
+
+		if (!FunctionToCall || FunctionToCall->HasAnyFunctionFlags(FUNC_BlueprintPure))
+		{
+			return true;
+		}
+
+		UClass* OuterClass = CastChecked<UClass>(FunctionToCall->GetOuter());
+		return OuterClass->ClassGeneratedBy != Blueprint;
+	};
+
+	FObjectProperty* ReturnProperty = EndpointDefinition.EndpointSignature ?
+		CastField<FObjectProperty>(EndpointDefinition.EndpointSignature->GetReturnProperty()) : nullptr;
+	auto RejectAnyIncompatibleReturnValues = [ReturnProperty](const FBlueprintActionFilter& Filter, FBlueprintActionInfo& BlueprintAction)
+	{
+		const UFunction* Function = BlueprintAction.GetAssociatedFunction();
+		const FObjectProperty* FunctionReturnProperty = CastField<FObjectProperty>(Function->GetReturnProperty());
+		return !FunctionReturnProperty || FunctionReturnProperty->PropertyClass != ReturnProperty->PropertyClass;
+	};
+
+	FBlueprintActionMenuBuilder ContextMenuBuilder;
+
+	{
+		FBlueprintActionFilter MenuFilter(FBlueprintActionFilter::BPFILTER_RejectGlobalFields | FBlueprintActionFilter::BPFILTER_RejectPermittedSubClasses);
+		MenuFilter.PermittedNodeTypes.Add(UK2Node_CallFunction::StaticClass());
+
+		MenuFilter.Context.Blueprints.Add(Blueprint);
+
+		MenuFilter.AddRejectionTest(FBlueprintActionFilter::FRejectionTestDelegate::CreateStatic(RejectAnyForeignFunctions, Blueprint));
+
+		if (ReturnProperty)
+		{
+			MenuFilter.AddRejectionTest(FBlueprintActionFilter::FRejectionTestDelegate::CreateLambda(RejectAnyIncompatibleReturnValues));
+		}
+
+		if (Blueprint->SkeletonGeneratedClass)
+		{
+			FBlueprintActionFilter::AddUnique(MenuFilter.TargetClasses, Blueprint->SkeletonGeneratedClass);
+		}
+
+		ContextMenuBuilder.AddMenuSection(MenuFilter, LOCTEXT("SequenceDirectorMenu", "This Sequence"), 0);
+	}
+
+	{
+		OnCollectAllRebindActions(Blueprint, ContextMenuBuilder);
+	}
+
+	ContextMenuBuilder.RebuildActionList();
+
+	OutAllActions.Append(ContextMenuBuilder);
+}
+
+void FMovieSceneDirectorBlueprintEndpointCustomization::HandleRebindActionSelected(const TArray<TSharedPtr<FEdGraphSchemaAction>>& SelectedAction, ESelectInfo::Type InSelectionType, UBlueprint* Blueprint, FMovieSceneDirectorBlueprintEndpointDefinition EndpointDefinition)
+{
+	if (InSelectionType != ESelectInfo::OnMouseClick && InSelectionType != ESelectInfo::OnKeyPress)
+	{
+		return;
+	}
+
+	if (!Blueprint->SkeletonGeneratedClass)
+	{
+		return;
+	}
+
+	for (TSharedPtr<FEdGraphSchemaAction> Action : SelectedAction)
+	{
+		if (Action->GetTypeId() != FBlueprintActionMenuItem::StaticGetTypeId())
+		{
+			continue;
+		}
+
+		const UBlueprintFunctionNodeSpawner* FunctionNodeSpawner = Cast<UBlueprintFunctionNodeSpawner>(static_cast<FBlueprintActionMenuItem*>(Action.Get())->GetRawAction());
+		const UFunction* FunctionToCall = FunctionNodeSpawner ? FunctionNodeSpawner->GetFunction() : nullptr;
+		if (!FunctionToCall)
+		{
+			continue;
+		}
+
+		// Give the opportunity for our sub-class implementation to do custom rebinding.
+		FSequenceDataMap AllSequenceData;
+		GatherSequenceData(AllSequenceData);
+
+		for (const TPair<UMovieSceneSequence*, FSequenceData>& Pair : AllSequenceData)
+		{
+			const bool bDidRebind = OnRebindEndpoint(
+					Pair.Key, Pair.Value.Blueprint, Pair.Value.EditObjects, Pair.Value.RawData,
+					EndpointDefinition, Action);
+			if (bDidRebind)
+			{
+				return;
+			}
+		}
+
+		// Default implementation only rebinds to an existing endpoint in our current director blueprint.
+		UClass* OuterClass = CastChecked<UClass>(FunctionToCall->GetOuter());
+		if (OuterClass->ClassGeneratedBy != Blueprint)
+		{
+			continue;
+		}
+
+		// Attempt to locate a custom event or a function graph of this name on the blueprint
+		for (UEdGraph* Graph : Blueprint->UbergraphPages)
+		{
+			for (UEdGraphNode* Node : Graph->Nodes)
+			{
+				UK2Node_CustomEvent* CustomEvent = Cast<UK2Node_CustomEvent>(Node);
+				if (CustomEvent && Blueprint->SkeletonGeneratedClass->FindFunctionByName(CustomEvent->GetFunctionName()) == FunctionToCall)
+				{
+					// Use this custom event
+					SetEndpoint(EndpointDefinition, CustomEvent, CustomEvent, EAutoCreatePayload::Variables);
+					return;
+				}
+			}
+		}
+
+		for (UEdGraph* Graph : Blueprint->FunctionGraphs)
+		{
+			if (Blueprint->SkeletonGeneratedClass->FindFunctionByName(Graph->GetFName()) != FunctionToCall)
+			{
+				continue;
+			}
+
+			// Use this function graph for the event endpoint
+			for (UEdGraphNode* Node : Graph->Nodes)
+			{
+				UK2Node_FunctionEntry* FunctionEntry = Cast<UK2Node_FunctionEntry>(Node);
+				if (FunctionEntry)
+				{
+					SetEndpoint(EndpointDefinition, FunctionEntry, FunctionEntry, EAutoCreatePayload::Variables);
+					return;
+				}
+			}
+		}
+	}
+
+	ensureMsgf(false, TEXT("Unknown blueprint action type encountered for rebinding"));
+}
+
+const FSlateBrush* FMovieSceneDirectorBlueprintEndpointCustomization::GetEndpointIcon() const
+{
+	UK2Node* CommonEndpoint = GetCommonEndpoint();
+
+	if (CommonEndpoint)
+	{
+		FLinearColor Color;
+		FSlateIcon EndpointIcon = CommonEndpoint->GetIconAndTint(Color);
+		return EndpointIcon.GetIcon();
+	}
+	else
+	{
+		TArray<void*> RawData;
+		PropertyHandle->AccessRawData(RawData);
+		if (RawData.Num() > 1)
+		{
+			return FAppStyle::GetBrush("Sequencer.MultipleEvents");
+		}
+	}
+
+	return FAppStyle::GetBrush("Sequencer.UnboundEvent");
+}
+
+FText FMovieSceneDirectorBlueprintEndpointCustomization::GetEndpointName() const
+{
+	UEdGraphNode* CommonEndpoint = GetCommonEndpoint();
+
+	if (CommonEndpoint)
+	{
+		return CommonEndpoint->GetNodeTitle(ENodeTitleType::MenuTitle);
+	}
+	else
+	{
+		TArray<void*> RawData;
+		PropertyHandle->AccessRawData(RawData);
+		if (RawData.Num() != 1)
+		{
+			return LOCTEXT("MultipleValuesText", "Multiple Values");
+		}
+	}
+
+	return LOCTEXT("UnboundText", "Unbound");
+}
+
+void FMovieSceneDirectorBlueprintEndpointCustomization::GatherSequenceData(FSequenceDataMap& AllSequenceData)
+{
+	TArray<UObject*> EditObjects;
+	GetEditObjects(EditObjects);
+
+	TArray<void*> RawData;
+	PropertyHandle->AccessRawData(RawData);
+
+	check(EditObjects.Num() == RawData.Num());
+
+	for (int32 Index = 0; Index < RawData.Num(); ++Index)
+	{
+		UMovieSceneSequence*       Sequence           = EditObjects[Index]->GetTypedOuter<UMovieSceneSequence>();
+		FMovieSceneSequenceEditor* SequenceEditor     = FMovieSceneSequenceEditor::Find(Sequence);
+		UBlueprint*                SequenceDirectorBP = SequenceEditor ? SequenceEditor->GetOrCreateDirectorBlueprint(Sequence) : nullptr;
+
+		FSequenceData& SequenceData = AllSequenceData.FindOrAdd(Sequence);
+		ensure(SequenceData.Blueprint == nullptr || SequenceData.Blueprint == SequenceDirectorBP);
+		SequenceData.Blueprint = SequenceDirectorBP;
+		SequenceData.EditObjects.Add(EditObjects[Index]);
+		SequenceData.RawData.Add(RawData[Index]);
+	}
+}
+
+void FMovieSceneDirectorBlueprintEndpointCustomization::SetEndpoint(const FMovieSceneDirectorBlueprintEndpointDefinition& EndpointDefinition, UK2Node* NewEndpoint, UK2Node* PayloadTemplate, EAutoCreatePayload AutoCreatePayload)
+{
+	FScopedTransaction Transaction(LOCTEXT("SetDirectorBlueprintEndpoint", "Set Director Blueprint Endpoint"));
+
+	// Modify and assign the blueprint for outer sections
+	FSequenceDataMap AllSequenceData;
+	GatherSequenceData(AllSequenceData);
+
+	// If we're assigning a new valid endpoint, it must reside within the same blueprint as everything we're assigning it to.
+	// Anything else must be implemented as a call function node connected to a custom event node or function graph.
+	UMovieSceneSequence* Sequence = nullptr;
+	UBlueprint* Blueprint = (NewEndpoint && NewEndpoint->HasValidBlueprint()) ? NewEndpoint->GetBlueprint() : nullptr;
+	for (const TPair<UMovieSceneSequence*, FSequenceData>& Pair : AllSequenceData)
+	{
+		if (NewEndpoint == nullptr)
+		{
+			Blueprint = Pair.Value.Blueprint;
+		}
+		if (Sequence == nullptr)
+		{
+			Sequence = Pair.Key;
+		}
+
+		if (!ensureAlwaysMsgf(
+				Pair.Value.Blueprint == Blueprint && Pair.Key == Sequence, 
+				TEXT("Attempting to assign an endpoint to objects with different Sequence Director Blueprints.")))
+		{
+			Transaction.Cancel();
+			return;
+		}
+
+		Blueprint->Modify();
+	}
+
+	if (!ensure(Sequence))
+	{
+		Transaction.Cancel();
+		return;
+	}
+
+	UEdGraphPin* CallTargetPin = FMovieSceneDirectorBlueprintUtils::FindCallTargetPin(NewEndpoint, EndpointDefinition.PossibleCallTargetClass);
+
+	// Map of the Payload Variable Names to their Default Values as Strings
+	TMap<FName, FString> PayloadVariables;
+	if (PayloadTemplate && EnumHasAnyFlags(AutoCreatePayload, EAutoCreatePayload::Variables | EAutoCreatePayload::Pins))
+	{
+		UFunction* PayloadTemplateFunction = nullptr;
+		if (UK2Node_Event* EventNode = Cast<UK2Node_Event>(PayloadTemplate))
+		{
+			PayloadTemplateFunction = EventNode->FindEventSignatureFunction();
+		}
+		else if (UK2Node_FunctionEntry* FunctionEntryNode = Cast<UK2Node_FunctionEntry>(PayloadTemplate))
+		{
+			PayloadTemplateFunction = FunctionEntryNode->FindSignatureFunction();
+		}
+		else if (UK2Node_CallFunction* CallFunctionNode = Cast<UK2Node_CallFunction>(PayloadTemplate))
+		{
+			PayloadTemplateFunction = CallFunctionNode->GetTargetFunction();
+		}
+
+		TSet<FName> NonPayloadPins;
+		if (PayloadTemplateFunction)
+		{
+			const FString* WorldContextParamName = PayloadTemplateFunction->FindMetaData(FBlueprintMetadata::MD_WorldContext);
+			if (WorldContextParamName)
+			{
+				NonPayloadPins.Add(FName(*WorldContextParamName));
+			}
+		}
+
+		UK2Node_EditablePinBase* EditableNode = Cast<UK2Node_EditablePinBase>(NewEndpoint);
+
+		for (UEdGraphPin* PayloadPin : PayloadTemplate->Pins)
+		{
+			if (PayloadPin != CallTargetPin && 
+					PayloadPin->LinkedTo.Num() == 0 && 
+					PayloadPin->Direction == EGPD_Input && 
+					PayloadPin->PinType.PinCategory != UEdGraphSchema_K2::PC_Exec && 
+					PayloadPin->PinName != UEdGraphSchema_K2::PN_Self &&
+					!NonPayloadPins.Contains(PayloadPin->PinName))
+			{
+				// Make a payload variable for this pin
+				if (EnumHasAnyFlags(AutoCreatePayload, EAutoCreatePayload::Variables))
+				{
+					PayloadVariables.Add(PayloadPin->PinName, PayloadPin->DefaultValue);
+				}
+
+				// Make a matching user pin on the endpoint node
+				if (EditableNode && EnumHasAnyFlags(AutoCreatePayload, EAutoCreatePayload::Pins))
+				{
+					// Pins for ref parameters for functions default to bIsReference but the payload cannot be by reference.
+					PayloadPin->PinType.bIsReference = false;
+
+					UEdGraphPin* NewPin = EditableNode->CreateUserDefinedPin(PayloadPin->PinName, PayloadPin->PinType, EGPD_Output);
+					if (PayloadTemplate != NewEndpoint && NewPin)
+					{
+						NewPin->MakeLinkTo(PayloadPin);
+					}
+				}
+			}
+		}
+	}
+
+	// Create payload variables for new parameters, remove payload variables for parameters
+	// that don't exist anymore.
+	const FString EmptyValue;
+	for (const TPair<UMovieSceneSequence*, FSequenceData>& Pair : AllSequenceData)
+	{
+		const TArray<void*>& RawData = Pair.Value.RawData;
+		const TArray<UObject*>& EditObjects = Pair.Value.EditObjects;
+
+		for (int32 Index = 0; Index < RawData.Num(); ++Index)
+		{
+			FPayloadVariableMap OldPayloadVariables;
+			GetPayloadVariables(EditObjects[Index], RawData[Index], OldPayloadVariables);
+
+			for (const TPair<FName, FString>& PayloadVar : PayloadVariables)
+			{
+				if (!OldPayloadVariables.Contains(PayloadVar.Key))
+				{
+					SetPayloadVariable(EditObjects[Index], RawData[Index], PayloadVar.Key, PayloadVar.Value);
+				}
+			}
+
+			for (const TPair<FName, FString>& PayloadVar : OldPayloadVariables)
+			{
+				if (!PayloadVariables.Contains(PayloadVar.Key))
+				{
+					SetPayloadVariable(EditObjects[Index], RawData[Index], PayloadVar.Key, EmptyValue);
+				}
+			}
+		}
+
+		OnSetEndpoint(Sequence, Blueprint, EditObjects, RawData, EndpointDefinition, NewEndpoint);
+	}
+
+	// Ensure that anything listening for property changed notifications are notified of the new binding
+	PropertyHandle->NotifyFinishedChangingProperties();
+
+	// Compile the blueprint now that clients have had a chance to update underlying data (we do this after to ensure we are compiling the correct data)
+	if (Blueprint)
+	{
+		FKismetEditorUtilities::CompileBlueprint(Blueprint);
+	}
+
+	// Forcibly update the panel now that our endpoint has changed
+	PropertyUtilities->ForceRefresh();
+
+	if (NewEndpoint)
+	{
+		FKismetEditorUtilities::BringKismetToFocusAttentionOnObject(NewEndpoint, false);
+	}
+}
+
+void FMovieSceneDirectorBlueprintEndpointCustomization::CreateEndpoint()
+{
+	struct FSequenceData
+	{
+		TArray<void*> RawData;
+		TArray<UObject*> EditObjects;
+	};
+
+	TSortedMap<UMovieSceneSequence*, FSequenceData> PerSequenceData;
+
+	// Populate all the sequences represented by this customization
+	{
+		TArray<void*> RawData;
+		GetPropertyHandle()->AccessRawData(RawData);
+
+		TArray<UObject*> EditObjects;
+		GetEditObjects(EditObjects);
+
+		check(RawData.Num() == EditObjects.Num());
+
+		for (int32 Index = 0; Index < RawData.Num(); ++Index)
+		{
+			FSequenceData& SequenceData = PerSequenceData.FindOrAdd(EditObjects[Index]->GetTypedOuter<UMovieSceneSequence>());
+			SequenceData.RawData.Add(RawData[Index]);
+			SequenceData.EditObjects.Add(EditObjects[Index]);
+		}
+	}
+
+	// Create user facing endpoint
+	FScopedTransaction Transaction(LOCTEXT("CreateEndpoint", "Create Endpoint"));
+
+	UK2Node* LastNewEndpoint = nullptr;
+	TArray<UBlueprint*> BlueprintsToRecompile;
+
+	for (const TPair<UMovieSceneSequence*, FSequenceData>& SequencePair : PerSequenceData)
+	{
+		FMovieSceneSequenceEditor* SequenceEditor = FMovieSceneSequenceEditor::Find(SequencePair.Key);
+		if (!SequenceEditor)
+		{
+			continue;
+		}
+
+		UBlueprint* SequenceDirectorBP = SequenceEditor->GetOrCreateDirectorBlueprint(SequencePair.Key);
+		if (!SequenceDirectorBP)
+		{
+			continue;
+		}
+
+		FMovieSceneDirectorBlueprintEndpointDefinition EndpointDefinition = GenerateEndpointDefinition(SequencePair.Key);
+
+		SequenceDirectorBP->Modify();
+		BlueprintsToRecompile.Add(SequenceDirectorBP);
+
+		UK2Node* NewEndpoint = FMovieSceneDirectorBlueprintUtils::CreateEndpoint(SequenceDirectorBP, EndpointDefinition);
+		if (!NewEndpoint)
+		{
+			continue;
+		}
+
+		OnCreateEndpoint(SequencePair.Key, SequenceDirectorBP, SequencePair.Value.EditObjects, SequencePair.Value.RawData, EndpointDefinition, NewEndpoint);
+
+		FBlueprintEditorUtils::MarkBlueprintAsModified(SequenceDirectorBP);
+
+		LastNewEndpoint = NewEndpoint;
+	}
+
+	// Ensure that anything listening for property changed notifications are notified of the new binding
+	PropertyHandle->NotifyFinishedChangingProperties();
+
+	// Compile the blueprint now that clients have had a chance to update underlying data (we do this after to ensure we are compiling the correct data)
+	for (UBlueprint* Blueprint : BlueprintsToRecompile)
+	{
+		FKismetEditorUtilities::CompileBlueprint(Blueprint);
+	}
+
+	PropertyUtilities->ForceRefresh();
+
+	// Focus the first created endpoint
+	if (LastNewEndpoint)
+	{
+		FKismetEditorUtilities::BringKismetToFocusAttentionOnObject(LastNewEndpoint, false);
+	}
+}
+
+UK2Node* FMovieSceneDirectorBlueprintEndpointCustomization::FindEndpoint(UObject* EditObject, void* RawData) const
+{
+	if (!EditObject || !RawData)
+	{
+		return nullptr;
+	}
+
+	UMovieSceneSequence* Sequence = EditObject->GetTypedOuter<UMovieSceneSequence>();
+	FMovieSceneSequenceEditor* SequenceEditor = FMovieSceneSequenceEditor::Find(Sequence);
+	if (!SequenceEditor)
+	{
+		return nullptr;
+	}
+
+	UBlueprint* SequenceDirectorBP = SequenceEditor->FindDirectorBlueprint(Sequence);
+	if (!SequenceDirectorBP)
+	{
+		return nullptr;
+	}
+	
+	return FindEndpoint(Sequence, SequenceDirectorBP, EditObject, RawData);
+}
+
+void FMovieSceneDirectorBlueprintEndpointCustomization::ClearEndpoint()
+{
+	FMovieSceneDirectorBlueprintEndpointDefinition EmptyDefinition;
+	SetEndpoint(EmptyDefinition, nullptr, nullptr, EAutoCreatePayload::None);
+}
+
+void FMovieSceneDirectorBlueprintEndpointCustomization::NavigateToDefinition()
+{
+	UEdGraphNode* CommonEndpoint = GetCommonEndpoint();
+	if (CommonEndpoint)
+	{
+		FKismetEditorUtilities::BringKismetToFocusAttentionOnObject(CommonEndpoint, false);
+	}
+}
+
+#undef LOCTEXT_NAMESPACE
+
diff --git a/Engine/Source/Editor/MovieSceneTools/Private/MovieSceneDirectorBlueprintUtils.cpp b/Engine/Source/Editor/MovieSceneTools/Private/MovieSceneDirectorBlueprintUtils.cpp
new file mode 100644
--- /dev/null
+++ b/Engine/Source/Editor/MovieSceneTools/Private/MovieSceneDirectorBlueprintUtils.cpp
@@ -0,0 +1,434 @@
+// Copyright Epic Games, Inc. All Rights Reserved.
+
+#include "MovieSceneDirectorBlueprintUtils.h"
+
+#include "EdGraph/EdGraph.h"
+#include "EdGraph/EdGraphNode.h"
+#include "EdGraph/EdGraphPin.h"
+#include "EdGraph/EdGraphSchema.h"
+#include "Engine/Blueprint.h"
+#include "GameFramework/Actor.h"
+#include "K2Node_CallFunction.h"
+#include "K2Node_CustomEvent.h"
+#include "K2Node_DynamicCast.h"
+#include "K2Node_FunctionEntry.h"
+#include "K2Node_FunctionResult.h"
+#include "K2Node_MakeArray.h"
+#include "Kismet/KismetSystemLibrary.h"
+#include "Kismet2/BlueprintEditorUtils.h"
+#include "Kismet2/CompilerResultsLog.h"
+#include "KismetCompiler.h"
+
+#include "MovieScene.h"
+#include "MovieSceneFwd.h"
+#include "MovieScenePossessable.h"
+#include "MovieSceneSpawnable.h"
+
+#define LOCTEXT_NAMESPACE "MovieSceneDirectorBlueprintUtils"
+
+FString FMovieSceneDirectorBlueprintEndpointParameter::SanitizePinName(const FString& InPinName)
+{
+	FString SanitizedPinName(InPinName);
+	for (TCHAR& Char : SanitizedPinName)
+	{
+		if (FCString::Strchr(INVALID_NAME_CHARACTERS, Char) != nullptr)
+		{
+			Char = '_';
+		}
+	}
+	return SanitizedPinName;
+}
+
+UK2Node* FMovieSceneDirectorBlueprintUtils::CreateEndpoint(UBlueprint* Blueprint, const FMovieSceneDirectorBlueprintEndpointDefinition& EndpointDefinition)
+{
+	switch (EndpointDefinition.EndpointType)
+	{
+		case EMovieSceneDirectorBlueprintEndpointType::Event:
+			return CreateEventEndpoint(Blueprint, EndpointDefinition);
+		case EMovieSceneDirectorBlueprintEndpointType::Function:
+			return CreateFunctionEndpoint(Blueprint, EndpointDefinition);
+		default:
+			return nullptr;
+	}
+}
+
+UK2Node_CustomEvent* FMovieSceneDirectorBlueprintUtils::CreateEventEndpoint(UBlueprint* Blueprint, const FMovieSceneDirectorBlueprintEndpointDefinition& EndpointDefinition)
+{
+	// Create a custom event node to replace the original event node imported from text
+	UEdGraph* SequenceEventGraph = FindOrCreateEventGraph(Blueprint, EndpointDefinition.GraphName);
+	UK2Node_CustomEvent* CustomEventNode = NewObject<UK2Node_CustomEvent>(SequenceEventGraph);
+
+	FString EndpointName = FMovieSceneDirectorBlueprintEndpointParameter::SanitizePinName(EndpointDefinition.EndpointName);
+	check(EndpointName.Len() != 0);
+	CustomEventNode->CustomFunctionName = FBlueprintEditorUtils::FindUniqueKismetName(Blueprint, EndpointName);
+
+	// Ensure that it is editable
+	CustomEventNode->bIsEditable = true;
+
+	// Setup defaults
+	CustomEventNode->CreateNewGuid();
+	CustomEventNode->PostPlacedNewNode();
+	CustomEventNode->AllocateDefaultPins();
+
+	const FVector2D NewPosition = SequenceEventGraph->GetGoodPlaceForNewNode();
+	CustomEventNode->NodePosX = NewPosition.X;
+	CustomEventNode->NodePosY = NewPosition.Y;
+
+	// Add extra pins and invoke custom setup
+	for (const FMovieSceneDirectorBlueprintEndpointParameter& ExtraPin : EndpointDefinition.ExtraPins)
+	{
+		FEdGraphPinType PinType;
+		PinType.PinCategory = ExtraPin.PinTypeCategory;
+		PinType.PinSubCategoryObject = ExtraPin.PinTypeClass;
+		FString PinName = FMovieSceneDirectorBlueprintEndpointParameter::SanitizePinName(ExtraPin.PinName);
+		CustomEventNode->CreateUserDefinedPin(FName(PinName), PinType, ExtraPin.PinDirection);
+	}
+
+	SequenceEventGraph->AddNode(CustomEventNode, false, false);
+
+	FBlueprintEditorUtils::ValidateBlueprintChildVariables(Blueprint, CustomEventNode->CustomFunctionName);
+	FBlueprintEditorUtils::MarkBlueprintAsStructurallyModified(Blueprint);
+	return CustomEventNode;
+}
+
+UK2Node_FunctionEntry* FMovieSceneDirectorBlueprintUtils::CreateFunctionEndpoint(UBlueprint* Blueprint, const FMovieSceneDirectorBlueprintEndpointDefinition& EndpointDefinition)
+{
+	// Create the new function graph
+	FString EndpointName = FMovieSceneDirectorBlueprintEndpointParameter::SanitizePinName(EndpointDefinition.EndpointName);
+	check(EndpointName.Len() != 0);
+	FName FuncGraphName = FBlueprintEditorUtils::FindUniqueKismetName(Blueprint, EndpointName);
+	UEdGraph* FuncGraph = FBlueprintEditorUtils::CreateNewGraph(Blueprint, FuncGraphName, UEdGraph::StaticClass(), UEdGraphSchema_K2::StaticClass());
+	const bool bIsUserCreated = true;
+	FBlueprintEditorUtils::AddFunctionGraph<UFunction>(Blueprint, FuncGraph, bIsUserCreated, EndpointDefinition.EndpointSignature);
+
+	// Find the entry node.
+	TArray<UK2Node_FunctionEntry*> EntryNodes;
+	FuncGraph->GetNodesOfClass(EntryNodes);
+	check(EntryNodes.Num() == 1);
+
+	// Add extra pins and invoke custom setup
+	for (const FMovieSceneDirectorBlueprintEndpointParameter& ExtraPin : EndpointDefinition.ExtraPins)
+	{
+		FEdGraphPinType PinType;
+		PinType.PinCategory = ExtraPin.PinTypeCategory;
+		PinType.PinSubCategoryObject = ExtraPin.PinTypeClass;
+		FString PinName = FMovieSceneDirectorBlueprintEndpointParameter::SanitizePinName(ExtraPin.PinName);
+		EntryNodes[0]->CreateUserDefinedPin(FName(PinName), PinType, ExtraPin.PinDirection);
+	}
+
+	FBlueprintEditorUtils::ValidateBlueprintChildVariables(Blueprint, EntryNodes[0]->CustomGeneratedFunctionName);
+	FBlueprintEditorUtils::MarkBlueprintAsStructurallyModified(Blueprint);
+	return EntryNodes[0];
+}
+
+UEdGraph* FMovieSceneDirectorBlueprintUtils::FindOrCreateEventGraph(UBlueprint* Blueprint, const FString& GraphName)
+{
+	UEdGraph* SequenceEventGraph = FindObject<UEdGraph>(Blueprint, *GraphName);
+	if (!SequenceEventGraph)
+	{
+		SequenceEventGraph = FBlueprintEditorUtils::CreateNewGraph(Blueprint, FName(GraphName), UEdGraph::StaticClass(), UEdGraphSchema_K2::StaticClass());
+
+		SequenceEventGraph->GraphGuid = FGuid::NewGuid();
+		Blueprint->UbergraphPages.Add(SequenceEventGraph);
+	}
+	return SequenceEventGraph;
+}
+
+FMovieSceneDirectorBlueprintEntrypointResult FMovieSceneDirectorBlueprintUtils::GenerateEntryPoint(const FMovieSceneDirectorBlueprintEndpointCall& EndpointCall, FKismetCompilerContext* Compiler)
+{
+	check(Compiler && EndpointCall.Endpoint);
+
+	UEdGraphNode* Endpoint = EndpointCall.Endpoint;
+	UBlueprint* Blueprint = Compiler->Blueprint;
+
+	UFunction* EndpointFunction = nullptr;
+
+	// Find the function that we need to call on the skeleton class
+	{
+		if (UK2Node_Event* Event = Cast<UK2Node_Event>(Endpoint))
+		{
+			EndpointFunction = Blueprint->SkeletonGeneratedClass ? Blueprint->SkeletonGeneratedClass->FindFunctionByName(Event->GetFunctionName()) : nullptr;
+		}
+		else if (UK2Node_FunctionEntry* EndpointEntry = Cast<UK2Node_FunctionEntry>(Endpoint))
+		{
+			EndpointFunction = EndpointEntry->FindSignatureFunction();
+		}
+		else
+		{
+			Compiler->MessageLog.Error(*LOCTEXT("InvalidEndpoint_Error", "Sequencer event is bound to an invalid endpoint node @@").ToString(), Endpoint);
+		}
+	}
+
+	if (EndpointFunction == nullptr)
+	{
+		return FMovieSceneDirectorBlueprintEntrypointResult();
+	}
+
+	// @todo: use a more destriptive name for event entry points?
+	static FString DefaultEventEntryName = TEXT("SequenceEvent__ENTRYPOINT");
+	UEdGraph* EntryPointGraph = Compiler->SpawnIntermediateFunctionGraph(DefaultEventEntryName + Blueprint->GetName());
+	check(EntryPointGraph->Nodes.Num() == 1);
+
+	const UEdGraphSchema*  Schema        = EntryPointGraph->GetSchema();
+	UK2Node_FunctionEntry* FunctionEntry = CastChecked<UK2Node_FunctionEntry>(EntryPointGraph->Nodes[0]);
+
+	FMovieSceneDirectorBlueprintEntrypointResult Result;
+	Result.Entrypoint = FunctionEntry;
+	Result.CompiledFunctionName = FunctionEntry->GetGraph()->GetFName();
+
+	// -------------------------------------------------------------------------------------
+	// Locate and initialize the function entry node
+	{
+		int32 ExtraFunctionFlags = ( FUNC_BlueprintEvent | FUNC_Public );
+		FunctionEntry->AddExtraFlags(ExtraFunctionFlags);
+		FunctionEntry->bIsEditable = false;
+		FunctionEntry->MetaData.Category = LOCTEXT("DefaultCategory", "Sequencer Event Endpoints");
+		FunctionEntry->MetaData.bCallInEditor = EndpointFunction->GetBoolMetaData(FBlueprintMetadata::MD_CallInEditor);
+	}
+
+	// -------------------------------------------------------------------------------------
+	// Create a function call node to invoke the endpoint function itself
+	UK2Node_CallFunction* CallFunctionNode = NewObject<UK2Node_CallFunction>(EntryPointGraph);
+	{
+		CallFunctionNode->CreateNewGuid();
+		CallFunctionNode->PostPlacedNewNode();
+
+		CallFunctionNode->FunctionReference.SetSelfMember(EndpointFunction->GetFName());
+
+		CallFunctionNode->ReconstructNode();
+
+		CallFunctionNode->NodePosX = FunctionEntry->NodePosX + 400.f;
+		CallFunctionNode->NodePosY = FunctionEntry->NodePosY - 16.f;
+
+		EntryPointGraph->AddNode(CallFunctionNode, false, false);
+	}
+
+	// -------------------------------------------------------------------------------------
+	// If the endpoint function has a return value, create a return node for it
+	UK2Node_FunctionResult* ResultNode = nullptr;
+	if (FProperty* ReturnParam = EndpointFunction->GetReturnProperty())
+	{
+		ResultNode = NewObject<UK2Node_FunctionResult>(EntryPointGraph);
+
+		ResultNode->CreateNewGuid();
+		ResultNode->PostPlacedNewNode();
+
+		ResultNode->FunctionReference.SetSelfMember(EndpointFunction->GetFName());
+
+		ResultNode->ReconstructNode();
+
+		ResultNode->NodePosX = CallFunctionNode->NodePosX + 400.f;
+		ResultNode->NodePosY = CallFunctionNode->NodePosY - 16.f;
+
+		EntryPointGraph->AddNode(ResultNode, false, false);
+	}
+
+	// -------------------------------------------------------------------------------------
+	// Reserve pins that want to be assigned a value at runtime
+	for (const FName& ExposedPinName : EndpointCall.ExposedPinNames)
+	{
+		UEdGraphPin* ExposedPin = CallFunctionNode->FindPin(ExposedPinName, EGPD_Input);
+		if (ExposedPin)
+		{
+			UEdGraphPin* BoundParamPinInput = FunctionEntry->CreateUserDefinedPin(ExposedPin->PinName, ExposedPin->PinType, EGPD_Output, true);
+			if (BoundParamPinInput)
+			{
+				BoundParamPinInput->MakeLinkTo(ExposedPin);
+			}
+		}
+	}
+
+	// -------------------------------------------------------------------------------------
+	// Wire up the function entry 'then' pin to the call function 'execute' pin
+	// If there's a result node, wire up the call function's result to its 'Return Value' pin
+	{
+		// Wire up the exec nodes
+		UEdGraphPin* ThenPin = FunctionEntry->FindPin(UEdGraphSchema_K2::PN_Then, EGPD_Output);
+		UEdGraphPin* ExecPin = CallFunctionNode->FindPin(UEdGraphSchema_K2::PN_Execute, EGPD_Input);
+
+		if (ensure(ThenPin && ExecPin))
+		{
+			ThenPin->MakeLinkTo(ExecPin);
+		}
+
+		if (ResultNode)
+		{
+			UEdGraphPin* CallResultPin = CallFunctionNode->FindPin(UEdGraphSchema_K2::PN_ReturnValue, EGPD_Output);
+			UEdGraphPin* OutputResultPin = ResultNode->FindPin(UEdGraphSchema_K2::PN_ReturnValue, EGPD_Input);
+
+			if (ensure(CallResultPin && OutputResultPin))
+			{
+				CallResultPin->MakeLinkTo(OutputResultPin);
+			}
+
+			ThenPin = CallFunctionNode->FindPin(UEdGraphSchema_K2::PN_Then, EGPD_Output);
+			ExecPin = ResultNode->FindPin(UEdGraphSchema_K2::PN_Execute, EGPD_Input);
+
+			if (ensure(ThenPin && ExecPin))
+			{
+				ThenPin->MakeLinkTo(ExecPin);
+			}
+		}
+	}
+
+	// -------------------------------------------------------------------------------------
+	// Set pin defaults for each of the relevant pins on the call function node according to the payload
+	{
+		TArray<FName, TInlineAllocator<8>> ValidPinNames;
+
+		// Gather all input pins for the function call
+		for (int32 PinIndex = 0; PinIndex < CallFunctionNode->Pins.Num(); ++PinIndex)
+		{
+			// Only consider input pins that are data, and not already connected
+			UEdGraphPin* Pin = CallFunctionNode->Pins[PinIndex];
+			if (Pin->Direction != EGPD_Input || Pin->PinName == UEdGraphSchema_K2::PN_Execute || Pin->PinName == UEdGraphSchema_K2::PN_Self || Pin->LinkedTo.Num() != 0)
+			{
+				continue;
+			}
+
+			UClass* PinObjectType = Cast<UClass>(Pin->PinType.PinSubCategoryObject.Get());
+			const bool bIsArrayPin = Pin->PinType.IsArray();
+			const bool bIsRawActorPin = Pin->PinType.PinCategory == UEdGraphSchema_K2::PC_Object && PinObjectType && PinObjectType->IsChildOf(AActor::StaticClass());
+
+			if (!ensureMsgf(!bIsArrayPin, TEXT("Array parameters aren't supported")))
+			{
+				continue;
+			}
+
+			const FString* PayloadVariable = EndpointCall.PayloadVariables.Find(Pin->PinName);
+			if (!PayloadVariable)
+			{
+				continue;
+			}
+
+			bool bSuccess = false;
+			if (bIsRawActorPin)
+			{
+				bSuccess = GenerateEntryPointRawActorParameter(Compiler, EntryPointGraph, Endpoint, CallFunctionNode, Pin, *PayloadVariable);
+			}
+			else
+			{
+				bool bMarkAsModified = false;
+				Schema->TrySetDefaultValue(*Pin, *PayloadVariable, bMarkAsModified);
+				bSuccess = true;
+			}
+
+			if (bSuccess)
+			{
+				ValidPinNames.Add(Pin->PinName);
+			}
+		}
+
+		// Remove any invalid payload variable names
+		if (ValidPinNames.Num() != EndpointCall.PayloadVariables.Num())
+		{
+			const FText ErrorFmt = LOCTEXT("PayloadParameterRemoved", "Stale Sequencer endpoint payload parameter {0} has been removed.");
+			for (auto It = EndpointCall.PayloadVariables.CreateConstIterator(); It; ++It)
+			{
+				if (!ValidPinNames.Contains(It->Key))
+				{
+					Compiler->MessageLog.Note(*FText::Format(ErrorFmt, FText::FromName(It->Key)).ToString());
+					Result.StalePayloadVariables.Add(It->Key);
+				}
+			}
+		}
+	}
+
+	// -------------------------------------------------------------------------------------
+	// Generate any other custom pins and nodes
+	if (EndpointCall.CustomizeEndpointCall.IsBound())
+	{
+		FMovieSceneCustomizeDirectorBlueprintEndpointCallParams Params;
+		Params.FunctionGraph = EntryPointGraph;
+		Params.FunctionEntryNode = FunctionEntry;
+		Params.CallFunctionNode = CallFunctionNode;
+		Params.FunctionResultNode = ResultNode;
+		EndpointCall.CustomizeEndpointCall.Execute(Params);
+	}
+	// WARNING: After this call, we have to be careful about looking for stuff in the function graph,
+	//		    since there could be any number of new nodes and pins.
+
+	return Result;
+}
+
+bool FMovieSceneDirectorBlueprintUtils::GenerateEntryPointRawActorParameter(
+			FKismetCompilerContext* Compiler, UEdGraph* Graph, UEdGraphNode* Endpoint, 
+			UK2Node* OriginNode, UEdGraphPin* DestPin, const FString& PayloadValue)
+{
+	// Raw actor properties are represented as soft object ptrs under the hood so that we can still serialize them
+	// To make this work, we have to make the following graph:
+	// MakeSoftObjectPath('<ActorPath>') -> Conv_SoftObjPathToSoftObjRef() -> Conv_SoftObjectReferenceToObject() -> Cast<ActorType>()
+	UK2Node_CallFunction* MakeSoftObjectPathNode = Compiler->SpawnIntermediateNode<UK2Node_CallFunction>(OriginNode, Graph);
+	MakeSoftObjectPathNode->FunctionReference.SetExternalMember(GET_FUNCTION_NAME_CHECKED(UKismetSystemLibrary, MakeSoftObjectPath), UKismetSystemLibrary::StaticClass());
+	MakeSoftObjectPathNode->AllocateDefaultPins();
+
+	UK2Node_CallFunction* ConvertToSoftObjectRef = Compiler->SpawnIntermediateNode<UK2Node_CallFunction>(OriginNode, Graph);
+	ConvertToSoftObjectRef->FunctionReference.SetExternalMember(GET_FUNCTION_NAME_CHECKED(UKismetSystemLibrary, Conv_SoftObjPathToSoftObjRef), UKismetSystemLibrary::StaticClass());
+	ConvertToSoftObjectRef->AllocateDefaultPins();
+
+	UK2Node_CallFunction* ConvertToObjectFunc = Compiler->SpawnIntermediateNode<UK2Node_CallFunction>(OriginNode, Graph);
+	ConvertToObjectFunc->FunctionReference.SetExternalMember(GET_FUNCTION_NAME_CHECKED(UKismetSystemLibrary, Conv_SoftObjectReferenceToObject), UKismetSystemLibrary::StaticClass());
+	ConvertToObjectFunc->AllocateDefaultPins();
+
+	UClass* PinObjectType = Cast<UClass>(DestPin->PinType.PinSubCategoryObject.Get());
+
+	UK2Node_DynamicCast* CastNode = Compiler->SpawnIntermediateNode<UK2Node_DynamicCast>(OriginNode, Graph);
+	CastNode->SetPurity(true);
+	CastNode->TargetType = PinObjectType;
+	CastNode->AllocateDefaultPins();
+
+	UEdGraphPin* PathInput     = MakeSoftObjectPathNode->FindPin(FName("PathString"));
+	UEdGraphPin* PathOutput    = MakeSoftObjectPathNode->GetReturnValuePin();
+	UEdGraphPin* SoftRefInput  = ConvertToSoftObjectRef->FindPin(FName("SoftObjectPath"));
+	UEdGraphPin* SoftRefOutput = ConvertToSoftObjectRef->GetReturnValuePin();
+	UEdGraphPin* ConvertInput  = ConvertToObjectFunc->FindPin(FName("SoftObject"));
+	UEdGraphPin* ConvertOutput = ConvertToObjectFunc->GetReturnValuePin();
+	UEdGraphPin* CastInput     = CastNode->GetCastSourcePin();
+	UEdGraphPin* CastOutput    = CastNode->GetCastResultPin();
+
+	if (!PathInput || !PathOutput || !SoftRefInput || !SoftRefOutput || !ConvertInput || !ConvertOutput || !CastInput || !CastOutput)
+	{
+		Compiler->MessageLog.Error(*LOCTEXT("ActorFacadeError", "GenerateEntryPoint: Failed to generate soft-ptr facade for AActor payload property property @@. @@").ToString(), *DestPin->PinName.ToString(), Endpoint);
+		return false;
+	}
+
+	const UEdGraphSchema*  Schema = Graph->GetSchema();
+
+	// Set the default value for the path string
+	const bool bMarkAsModified = false;
+	Schema->TrySetDefaultValue(*PathInput, PayloadValue, bMarkAsModified);
+
+	bool bSuccess = true;
+	bSuccess &= Schema->TryCreateConnection(PathOutput, SoftRefInput);
+	bSuccess &= Schema->TryCreateConnection(SoftRefOutput, ConvertInput);
+	bSuccess &= Schema->TryCreateConnection(ConvertOutput, CastInput);
+	bSuccess &= Schema->TryCreateConnection(CastOutput, DestPin);
+
+	if (!bSuccess)
+	{
+		Compiler->MessageLog.Error(*LOCTEXT("ActorFacadeConnectionError", "GenerateEntryPoint: Failed to connect nodes for soft-ptr facade in AActor payload property @@. @@").ToString(), *DestPin->PinName.ToString(), Endpoint);
+	}
+
+	return bSuccess;
+}
+
+UEdGraphPin* FMovieSceneDirectorBlueprintUtils::FindCallTargetPin(UK2Node* InEndpoint, UClass* CallTargetClass)
+{
+	if (!CallTargetClass)
+	{
+		return nullptr;
+	}
+
+	for (UEdGraphPin* Pin : InEndpoint->Pins)
+	{
+		if (Pin->Direction == EGPD_Output && Pin->PinType.PinCategory == UEdGraphSchema_K2::PC_Object && Pin->PinType.PinSubCategoryObject == CallTargetClass)
+		{
+			return Pin;
+		}
+	}
+
+	return nullptr;
+}
+
+
+#undef LOCTEXT_NAMESPACE
diff --git a/Engine/Source/Editor/MovieSceneTools/Private/MovieSceneEventBlueprintExtension.cpp b/Engine/Source/Editor/MovieSceneTools/Private/MovieSceneEventBlueprintExtension.cpp
--- a/Engine/Source/Editor/MovieSceneTools/Private/MovieSceneEventBlueprintExtension.cpp
+++ b/Engine/Source/Editor/MovieSceneTools/Private/MovieSceneEventBlueprintExtension.cpp
@@ -57,7 +57,7 @@ void UMovieSceneEventBlueprintExtension::HandleGenerateFunctionGraphs(FKismetCom
 			UEdGraphNode* Endpoint = FMovieSceneEventUtils::FindEndpoint(&EntryPoint, EventSection, CompilerContext->Blueprint);
 			if (Endpoint)
 			{
-				UK2Node_FunctionEntry* FunctionEntry = FMovieSceneEventUtils::GenerateEntryPoint(EventSection, &EntryPoint, CompilerContext, Endpoint);
+				UK2Node_FunctionEntry* FunctionEntry = FMovieSceneEventUtils::GenerateEntryPoint(&EntryPoint, CompilerContext, Endpoint);
 				if (FunctionEntry)
 				{
 					EntryPoint.CompiledFunctionName = FunctionEntry->GetGraph()->GetFName();
diff --git a/Engine/Source/Editor/MovieSceneTools/Private/MovieSceneEventCustomization.cpp b/Engine/Source/Editor/MovieSceneTools/Private/MovieSceneEventCustomization.cpp
--- a/Engine/Source/Editor/MovieSceneTools/Private/MovieSceneEventCustomization.cpp
+++ b/Engine/Source/Editor/MovieSceneTools/Private/MovieSceneEventCustomization.cpp
@@ -1,167 +1,26 @@
 // Copyright Epic Games, Inc. All Rights Reserved.
 
 #include "MovieSceneEventCustomization.h"
-#include "Modules/ModuleManager.h"
-#include "MovieSceneEventUtils.h"
 
-#include "UObject/UnrealType.h"
-#include "Channels/MovieSceneEvent.h"
-#include "Tracks/MovieSceneEventTrack.h"
-#include "Sections/MovieSceneEventSectionBase.h"
-#include "Algo/Find.h"
 #include "MovieScene.h"
+#include "MovieSceneDirectorBlueprintUtils.h"
+#include "MovieSceneEventUtils.h"
 #include "MovieSceneSequence.h"
-#include "ISequencerModule.h"
-#include "ScopedTransaction.h"
-
-#include "Framework/Application/SlateApplication.h"
-#include "Framework/MultiBox/MultiBoxBuilder.h"
-#include "Widgets/DeclarativeSyntaxSupport.h"
-#include "Widgets/Text/STextBlock.h"
-#include "Widgets/Images/SImage.h"
-#include "Widgets/Layout/SBox.h"
-#include "Widgets/Input/SCheckBox.h"
-#include "Widgets/Input/SComboButton.h"
-#include "Widgets/Input/SHyperlink.h"
-#include "Widgets/Input/SEditableTextBox.h"
-#include "Widgets/Views/SExpanderArrow.h"
+#include "MovieSceneSequenceEditor.h"
 
+#include "Channels/MovieSceneEvent.h"
 #include "PropertyHandle.h"
-#include "IPropertyUtilities.h"
-#include "PropertyCustomizationHelpers.h"
-#include "IDetailChildrenBuilder.h"
-#include "DetailLayoutBuilder.h"
-#include "DetailCategoryBuilder.h"
+#include "Sections/MovieSceneEventSectionBase.h"
+#include "Tracks/MovieSceneEventTrack.h"
 
-#include "K2Node_Variable.h"
 #include "K2Node_CustomEvent.h"
-#include "K2Node_CallFunction.h"
-#include "K2Node_FunctionEntry.h"
-#include "Kismet2/KismetEditorUtilities.h"
 #include "Kismet2/BlueprintEditorUtils.h"
-#include "Kismet2/Kismet2NameValidators.h"
-#include "SGraphActionMenu.h"
-#include "BlueprintActionMenuUtils.h"
-#include "BlueprintActionMenuBuilder.h"
-#include "BlueprintActionMenuItem.h"
-#include "BlueprintFunctionNodeSpawner.h"
 
 #include "Styling/AppStyle.h"
 #include "EditorFontGlyphs.h"
 
 #define LOCTEXT_NAMESPACE "MovieSceneEventCustomization"
 
-
-namespace UE_MovieSceneEventCustomization
-{
-	void CollectQuickBindActions(FGraphActionListBuilderBase& OutAllActions, UBlueprint* Blueprint, UClass* BoundObjectPinClass)
-	{
-		// Build up the context object
-		auto RejectAnyNonFunctions = [](const FBlueprintActionFilter& Filter, FBlueprintActionInfo& BlueprintAction)
-		{
-			const UFunction* Function = BlueprintAction.GetAssociatedFunction();
-			return Function == nullptr || Function->HasAnyFunctionFlags(FUNC_BlueprintPure);
-		};
-		auto RejectAnyUnboundActions = [](const FBlueprintActionFilter& Filter, FBlueprintActionInfo& BlueprintAction)
-		{
-			return BlueprintAction.GetBindings().Num() <= 0;
-		};
-
-
-		FBlueprintActionMenuBuilder ContextMenuBuilder;
-
-		if (BoundObjectPinClass)
-		{
-			// Add actions that are relevant to the bound object from the pin class
-			FBlueprintActionFilter CallOnMemberFilter(FBlueprintActionFilter::BPFILTER_RejectGlobalFields | FBlueprintActionFilter::BPFILTER_RejectPermittedSubClasses);
-			CallOnMemberFilter.PermittedNodeTypes.Add(UK2Node_CallFunction::StaticClass());
-			CallOnMemberFilter.Context.Blueprints.Add(Blueprint);
-
-			for (FObjectProperty* ObjectProperty : TFieldRange<FObjectProperty>(BoundObjectPinClass))
-			{
-				if (ObjectProperty->HasAnyPropertyFlags(CPF_BlueprintVisible) && (ObjectProperty->HasMetaData(FBlueprintMetadata::MD_ExposeFunctionCategories) || FBlueprintEditorUtils::IsSCSComponentProperty(ObjectProperty)))
-				{
-					CallOnMemberFilter.Context.SelectedObjects.Add(ObjectProperty);
-					FBlueprintActionFilter::AddUnique(CallOnMemberFilter.TargetClasses, ObjectProperty->PropertyClass);
-				}
-			}
-
-			FBlueprintActionFilter::AddUnique(CallOnMemberFilter.TargetClasses, BoundObjectPinClass);
-	
-			// This removes duplicate entries (ie. Set Static Mesh and Set Static Mesh (StaticMeshComponent)), 
-			// but also prevents displaying functions on BP components. Comment out for now.
-			//CallOnMemberFilter.AddRejectionTest(FBlueprintActionFilter::FRejectionTestDelegate::CreateStatic(RejectAnyUnboundActions));
-			
-			CallOnMemberFilter.AddRejectionTest(FBlueprintActionFilter::FRejectionTestDelegate::CreateStatic(RejectAnyNonFunctions));
-
-			ContextMenuBuilder.AddMenuSection(CallOnMemberFilter, FText::FromName(BoundObjectPinClass->GetFName()), 0);
-		}
-
-		{
-			// Add all actions that are relevant to the sequence director BP itself
-			FBlueprintActionFilter MenuFilter(FBlueprintActionFilter::BPFILTER_RejectGlobalFields | FBlueprintActionFilter::BPFILTER_RejectPermittedSubClasses);
-
-			MenuFilter.PermittedNodeTypes.Add(UK2Node_CallFunction::StaticClass());
-
-			MenuFilter.Context.Blueprints.Add(Blueprint);
-			MenuFilter.Context.Graphs.Append(Blueprint->UbergraphPages);
-			MenuFilter.Context.Graphs.Append(Blueprint->FunctionGraphs);
-
-			if (Blueprint->SkeletonGeneratedClass)
-			{
-				FBlueprintActionFilter::AddUnique(MenuFilter.TargetClasses, Blueprint->SkeletonGeneratedClass);
-			}
-			MenuFilter.AddRejectionTest(FBlueprintActionFilter::FRejectionTestDelegate::CreateStatic(RejectAnyNonFunctions));
-
-
-			ContextMenuBuilder.AddMenuSection(MenuFilter, LOCTEXT("SequenceDirectorMenu", "This Sequence"), 0);
-		}
-
-		ContextMenuBuilder.RebuildActionList();
-
-		OutAllActions.Append(ContextMenuBuilder);
-	}
-
-	void CollectAllRebindActions(FGraphActionListBuilderBase& OutAllActions, UBlueprint* Blueprint)
-	{
-		// Build up the context object
-		auto RejectAnyForeignFunctions = [](const FBlueprintActionFilter& Filter, FBlueprintActionInfo& BlueprintAction, UBlueprint* Blueprint)
-		{
-			const UBlueprintFunctionNodeSpawner* FunctionNodeSpawner = Cast<UBlueprintFunctionNodeSpawner>(BlueprintAction.NodeSpawner);
-			const UFunction* FunctionToCall = FunctionNodeSpawner ? FunctionNodeSpawner->GetFunction() : nullptr;
-
-			if (!FunctionToCall || FunctionToCall->HasAnyFunctionFlags(FUNC_BlueprintPure))
-			{
-				return true;
-			}
-
-			UClass* OuterClass = CastChecked<UClass>(FunctionToCall->GetOuter());
-			return OuterClass->ClassGeneratedBy != Blueprint;
-		};
-
-		FBlueprintActionFilter MenuFilter(FBlueprintActionFilter::BPFILTER_RejectGlobalFields | FBlueprintActionFilter::BPFILTER_RejectPermittedSubClasses);
-		MenuFilter.PermittedNodeTypes.Add(UK2Node_CallFunction::StaticClass());
-
-		MenuFilter.Context.Blueprints.Add(Blueprint);
-
-		MenuFilter.AddRejectionTest(FBlueprintActionFilter::FRejectionTestDelegate::CreateStatic(RejectAnyForeignFunctions, Blueprint));
-
-		if (Blueprint->SkeletonGeneratedClass)
-		{
-			FBlueprintActionFilter::AddUnique(MenuFilter.TargetClasses, Blueprint->SkeletonGeneratedClass);
-		}
-
-		FBlueprintActionMenuBuilder ContextMenuBuilder;
-
-		ContextMenuBuilder.AddMenuSection(MenuFilter, LOCTEXT("SequenceDirectorMenu", "This Sequence"), 0);
-		ContextMenuBuilder.RebuildActionList();
-		OutAllActions.Append(ContextMenuBuilder);
-	}
-}
-
-
-
-
 TSharedRef<IPropertyTypeCustomization> FMovieSceneEventCustomization::MakeInstance()
 {
 	return MakeShared<FMovieSceneEventCustomization>();
@@ -183,1285 +42,175 @@ void FMovieSceneEventCustomization::GetEditObjects(TArray<UObject*>& OutObjects)
 	}
 	else
 	{
-		PropertyHandle->GetOuterObjects(OutObjects);
+		GetPropertyHandle()->GetOuterObjects(OutObjects);
 	}
 }
 
-void FMovieSceneEventCustomization::CustomizeHeader(TSharedRef<IPropertyHandle> InPropertyHandle, FDetailWidgetRow& HeaderRow, IPropertyTypeCustomizationUtils& CustomizationUtils)
-{}
-
-void FMovieSceneEventCustomization::CustomizeChildren(TSharedRef<IPropertyHandle> InPropertyHandle, IDetailChildrenBuilder& ChildBuilder, IPropertyTypeCustomizationUtils& CustomizationUtils)
+void FMovieSceneEventCustomization::GetPayloadVariables(UObject* EditObject, void* RawData, FPayloadVariableMap& OutPayloadVariables) const
 {
-	PropertyHandle = InPropertyHandle;
-	PropertyUtilities = CustomizationUtils.GetPropertyUtilities();
-
-	TArray<void*> RawData;
-	PropertyHandle->AccessRawData(RawData);
-
-	ChildBuilder.AddCustomRow(FText())
-	.NameContent()
-	[
-		SNew(STextBlock)
-		.Font(CustomizationUtils.GetRegularFont())
-		.Text(LOCTEXT("EventValueText", "Event"))
-	]
-	.ValueContent()
-	.MinDesiredWidth(200.f)
-	[
-		SNew(SHorizontalBox)
-
-		+ SHorizontalBox::Slot()
-		[
-			SNew(SComboButton)
-			.ButtonStyle(FAppStyle::Get(), "HoverHintOnly")
-			.ForegroundColor(FSlateColor::UseForeground())
-			.OnGetMenuContent(this, &FMovieSceneEventCustomization::GetMenuContent)
-			.CollapseMenuOnParentFocus(true)
-			.ContentPadding(FMargin(4.f, 0.f))
-			.ButtonContent()
-			[
-				SNew(SHorizontalBox)
-
-				+ SHorizontalBox::Slot()
-				.Padding(FMargin(0.f, 0.f, 4.f, 0.f))
-				.VAlign(VAlign_Center)
-				.AutoWidth()
-				[
-					SNew(SImage)
-					.Image(this, &FMovieSceneEventCustomization::GetEventIcon)
-				]
-
-				+ SHorizontalBox::Slot()
-				.Padding(FMargin(0.f, 0.f, 4.f, 0.f))
-				.VAlign(VAlign_Center)
-				[
-					SNew(STextBlock)
-					.Font(CustomizationUtils.GetRegularFont())
-					.Text(this, &FMovieSceneEventCustomization::GetEventName)
-				]
-			]
-		]
-
-		+ SHorizontalBox::Slot()
-		.VAlign(VAlign_Center)
-		.AutoWidth()
-		[
-			PropertyCustomizationHelpers::MakeBrowseButton(FSimpleDelegate::CreateSP(this, &FMovieSceneEventCustomization::NavigateToDefinition), LOCTEXT("NavigateToDefinition_Tip", "Navigate to this event's definition"))
-		]
-	];
-
-	const bool bAnyBoundEvents = GetAllValidEndpoints().Num() != 0;
-	if (bAnyBoundEvents)
-	{
-		FText CallInEditorText    = LOCTEXT("CallInEditor_Label",   "Call In Editor");
-		FText CallInEditorTooltip = LOCTEXT("CallInEditor_Tooltip", "When checked, this event will be triggered in the Editor outside of PIE.\n\nBEWARE: ANY CHANGES MADE AS A RESULT OF THIS EVENT BEING CALLED MAY END UP BEING SAVED IN THE CURRENT LEVEL OR ASSET.");
-
-		ChildBuilder.AddCustomRow(CallInEditorText)
-		.NameContent()
-		[
-			SNew(STextBlock)
-			.Text(CallInEditorText)
-			.ToolTipText(CallInEditorTooltip)
-			.Font(IDetailLayoutBuilder::GetDetailFont())
-		]
-		.ValueContent()
-		[
-			SNew(SCheckBox)
-			.ToolTipText(CallInEditorTooltip)
-			.IsChecked(this, &FMovieSceneEventCustomization::GetCallInEditorCheckState)
-			.OnCheckStateChanged(this, &FMovieSceneEventCustomization::OnSetCallInEditorCheckState)
-		];
-	}
-
-
-	UK2Node* CommonEndpoint = GetCommonEndpoint();
-	if (!CommonEndpoint)
-	{
-		return;
-	}
-
-	TArray<UEdGraphPin*> ValidBoundObjectPins;
-	for (UEdGraphPin* Pin : CommonEndpoint->Pins)
-	{
-		if (Pin->Direction == EGPD_Output && (Pin->PinType.PinCategory == UEdGraphSchema_K2::PC_Object || Pin->PinType.PinCategory == UEdGraphSchema_K2::PC_Interface))
-		{
-			ValidBoundObjectPins.Add(Pin);
-		}
-	}
-
-
-	if (ValidBoundObjectPins.Num() > 0)
-	{
-		FText BoundObjectPinText    = LOCTEXT("BoundObjectPin_Label", "Pass Bound Object To");
-		FText BoundObjectPinTooltip = LOCTEXT("BoundObjectPin_Tooltip", "Specifies a pin to pass the bound object(s) through when the event is triggered. Interface and object pins are both supported.");
-
-		ChildBuilder.AddCustomRow(BoundObjectPinText)
-		.NameContent()
-		[
-			SNew(STextBlock)
-			.Text(BoundObjectPinText)
-			.ToolTipText(BoundObjectPinTooltip)
-			.Font(IDetailLayoutBuilder::GetDetailFont())
-		]
-		.ValueContent()
-		[
-			SNew(SComboButton)
-			.ButtonStyle(FAppStyle::Get(), "HoverHintOnly")
-			.ForegroundColor(FSlateColor::UseForeground())
-			.ToolTipText(BoundObjectPinTooltip)
-			.OnGetMenuContent(this, &FMovieSceneEventCustomization::GetBoundObjectPinMenuContent)
-			.CollapseMenuOnParentFocus(true)
-			.ContentPadding(FMargin(4.f, 0.f))
-			.ButtonContent()
-			[
-				SNew(SBox)
-				.HeightOverride(21.f)
-				[
-					SNew(SHorizontalBox)
-
-					+ SHorizontalBox::Slot()
-					.Padding(FMargin(0.f, 0.f, 4.f, 0.f))
-					.VAlign(VAlign_Center)
-					.AutoWidth()
-					[
-						SNew(SImage)
-						.Image(this, &FMovieSceneEventCustomization::GetBoundObjectPinIcon)
-					]
-
-					+ SHorizontalBox::Slot()
-					.Padding(FMargin(0.f, 0.f, 4.f, 0.f))
-					.VAlign(VAlign_Center)
-					[
-						SNew(STextBlock)
-						.Font(CustomizationUtils.GetRegularFont())
-						.Text(this, &FMovieSceneEventCustomization::GetBoundObjectPinText)
-					]
-				]
-			]
-		];
-	}
-
-
-	UFunction* CommonFunction = nullptr;
-	UBlueprint* Blueprint = CommonEndpoint->HasValidBlueprint() ? CommonEndpoint->GetBlueprint() : nullptr;
-	if (Blueprint)
-	{
-		if (UK2Node_Event* Event = Cast<UK2Node_Event>(CommonEndpoint))
-		{
-			CommonFunction = Blueprint->SkeletonGeneratedClass ? Blueprint->SkeletonGeneratedClass->FindFunctionByName(Event->GetFunctionName()) : nullptr;
-		}
-		else if (UK2Node_FunctionEntry* EndpointEntry = Cast<UK2Node_FunctionEntry>(CommonEndpoint))
-		{
-			CommonFunction = EndpointEntry->FindSignatureFunction();
-		}
-		else
-		{
-			// @todo: Error not supported
-		}
-
-		Blueprint->OnCompiled().AddSP(this, &FMovieSceneEventCustomization::OnBlueprintCompiled);
-	}
-
-	if (CommonFunction)
+	const FMovieSceneEvent* EntryPoint = static_cast<FMovieSceneEvent*>(RawData);
+	for (const TPair<FName, FMovieSceneEventPayloadVariable>& Pair : EntryPoint->PayloadVariables)
 	{
-		IDetailCategoryBuilder& DetailCategoryBuilder = ChildBuilder.GetParentCategory().GetParentLayout().EditCategory("Payload", LOCTEXT("PayloadLabel", "Payload"), ECategoryPriority::Uncommon);
-
-		bool bPayloadUpToDate = true;
-
-		for (int32 Index = 0; Index < RawData.Num(); ++Index)
-		{
-			TSharedPtr<FStructOnScope> StructData = MakeShared<FStructOnScope>(CommonFunction);
-
-			const FMovieSceneEvent* EntryPoint = static_cast<FMovieSceneEvent*>(RawData[Index]);
-
-			TArray<FName, TInlineAllocator<8>> AllValidNames;
-
-			for (FProperty* Field : TFieldRange<FProperty>(CommonFunction))
-			{
-				if (Field->HasAnyPropertyFlags(CPF_OutParm | CPF_ReturnParm | CPF_ReferenceParm) || Field->GetFName() == EntryPoint->BoundObjectPinName)
-				{
-					continue;
-				}
-
-				const FMovieSceneEventPayloadVariable* PayloadVariable = EntryPoint->PayloadVariables.Find(Field->GetFName());
-				if (PayloadVariable)
-				{
-					AllValidNames.Add(Field->GetFName());
-
-					// We have an override for this variable
-					const bool bImportSuccess = FBlueprintEditorUtils::PropertyValueFromString(Field, PayloadVariable->Value, StructData->GetStructMemory());
-					if (!bImportSuccess)
-					{
-						// @todo: error
-					}
-				}
-				else if (UEdGraphPin* Pin = CommonEndpoint->FindPin(Field->GetFName(), EGPD_Output))
-				{
-					if (!Pin->DefaultValue.IsEmpty())
-					{
-						const bool bImportSuccess = FBlueprintEditorUtils::PropertyValueFromString(Field, Pin->DefaultValue, StructData->GetStructMemory());
-						if (!bImportSuccess)
-						{
-							// @todo: error
-						}
-					}
-				}
-
-				IDetailPropertyRow* ExternalRow = DetailCategoryBuilder.AddExternalStructureProperty(StructData.ToSharedRef(), Field->GetFName(), EPropertyLocation::Default, FAddPropertyParams().ForceShowProperty());
-
-				TSharedPtr<IPropertyHandle> LocalVariableProperty = ExternalRow->GetPropertyHandle();
-				FSimpleDelegate Delegate = FSimpleDelegate::CreateSP(this, &FMovieSceneEventCustomization::OnPayloadVariableChanged, StructData.ToSharedRef(), LocalVariableProperty);
-
-				LocalVariableProperty->SetOnPropertyValueChanged(Delegate);
-				LocalVariableProperty->SetOnChildPropertyValueChanged(Delegate);
-			}
-
-			bPayloadUpToDate &= (AllValidNames.Num() == EntryPoint->PayloadVariables.Num());
-		}
-
-		if (!bPayloadUpToDate)
-		{
-			DetailCategoryBuilder.AddCustomRow(FText())
-			.WholeRowContent()
-			[
-				SNew(SHorizontalBox)
-
-				+ SHorizontalBox::Slot()
-				.Padding(FMargin(5.f, 0.f, 5.f, 0.f))
-				.AutoWidth()
-				[
-					SNew(STextBlock)
-					.TextStyle(FAppStyle::Get(), "Log.Warning")
-					.Font(FAppStyle::GetFontStyle("FontAwesome.10"))
-					.Text(FEditorFontGlyphs::Exclamation_Triangle)
-				]
-
-				+ SHorizontalBox::Slot()
-				.Padding(FMargin(0.f, 0.f, 5.f, 0.f))
-				[
-					SNew(STextBlock)
-					.TextStyle(FAppStyle::Get(), "Log.Warning")
-					.Text(LOCTEXT("PayloadOutOfDateError", "Payload variables may be out-of-date. Please compile the blueprint."))
-				]
-			];
-		}
+		OutPayloadVariables.Add(Pair.Key, Pair.Value.Value);
 	}
 }
 
-void FMovieSceneEventCustomization::OnPayloadVariableChanged(TSharedRef<FStructOnScope> InStructData, TSharedPtr<IPropertyHandle> LocalVariableProperty)
+bool FMovieSceneEventCustomization::SetPayloadVariable(UObject* EditObject, void* RawData, FName FieldName, const FString& NewVariableValue)
 {
-	// This function should only ever be bound if all the entry points call the same function
-	FProperty* Property = LocalVariableProperty->GetProperty();
-	if (!Property)
-	{
-		return;
-	}
-
-	FScopedTransaction Transaction(LOCTEXT("SetPayloadValue", "Set payload value"));
-
-	bool bChangedAnything = false;
-
-	UBlueprint* Blueprint = nullptr;
-
-	FString NewValueString;
-
-	const bool bSuccessfulTextExport = FBlueprintEditorUtils::PropertyValueToString(Property, InStructData->GetStructMemory(), NewValueString, nullptr);
-	if (!bSuccessfulTextExport)
-	{
-		// @todo: error
-		return;
-	}
-
-
-	TArray<UObject*> EditObjects;
-	GetEditObjects(EditObjects);
-
-	TArray<void*> RawData;
-	PropertyHandle->AccessRawData(RawData);
-
-	if (!ensure(RawData.Num() == EditObjects.Num()))
-	{
-		return;
-	}
-
-	for (int32 Index = 0; Index < RawData.Num(); ++Index)
+	UMovieSceneEventSectionBase* EventSection = Cast<UMovieSceneEventSectionBase>(EditObject);
+	FMovieSceneEvent* EntryPoint = static_cast<FMovieSceneEvent*>(RawData);
+	if (!EventSection || !EntryPoint)
 	{
-		UMovieSceneEventSectionBase* EventSection = Cast<UMovieSceneEventSectionBase>(EditObjects[Index]);
-		FMovieSceneEvent*  EntryPoint   = static_cast<FMovieSceneEvent*>(RawData[Index]);
-		if (!EventSection || !EntryPoint)
-		{
-			continue;
-		}
-
-		EventSection->Modify();
-
-		FMovieSceneEventPayloadVariable* PayloadVariable = EntryPoint->PayloadVariables.Find(Property->GetFName());
-		if (!PayloadVariable)
-		{
-			PayloadVariable = &EntryPoint->PayloadVariables.Add(Property->GetFName());
-		}
-
-		PayloadVariable->Value = NewValueString;
-		bChangedAnything = true;
+		return false;
 	}
 
-	if (bChangedAnything)
-	{
-		UK2Node* CommonEndpoint = GetCommonEndpoint();
+	EventSection->Modify();
 
-		UBlueprint* BP = CommonEndpoint && CommonEndpoint->HasValidBlueprint() ? GetCommonEndpoint()->GetBlueprint() : nullptr;
-		if (BP)
-		{
-			FBlueprintEditorUtils::MarkBlueprintAsModified(BP);
-		}
-	}
-	else
+	FMovieSceneEventPayloadVariable* PayloadVariable = EntryPoint->PayloadVariables.Find(FieldName);
+	if (!PayloadVariable)
 	{
-		Transaction.Cancel();
+		PayloadVariable = &EntryPoint->PayloadVariables.Add(FieldName);
 	}
-}
-
-ECheckBoxState FMovieSceneEventCustomization::GetCallInEditorCheckState() const
-{
-	ECheckBoxState CheckState = ECheckBoxState::Undetermined;
-	for (UK2Node* Endpoint : GetAllValidEndpoints())
-	{
-		bool bCallInEditor = false;
-		if (UK2Node_CustomEvent* CustomEvent = Cast<UK2Node_CustomEvent>(Endpoint))
-		{
-			bCallInEditor = CustomEvent->bCallInEditor;
-		}
-		else if (UK2Node_FunctionEntry* FunctionEntry = Cast<UK2Node_FunctionEntry>(Endpoint))
-		{
-			bCallInEditor = FunctionEntry->MetaData.bCallInEditor;
-		}
 
-		if (CheckState == ECheckBoxState::Undetermined)
-		{
-			CheckState = bCallInEditor ? ECheckBoxState::Checked : ECheckBoxState::Unchecked;
-		}
-		else if (bCallInEditor != (CheckState == ECheckBoxState::Checked) )
-		{
-			return ECheckBoxState::Undetermined;
-		}
-	}
-	return CheckState;
+	PayloadVariable->Value = NewVariableValue;
+	return true;
 }
 
-void FMovieSceneEventCustomization::OnSetCallInEditorCheckState(ECheckBoxState NewCheckedState)
+UK2Node* FMovieSceneEventCustomization::FindEndpoint(UMovieSceneSequence* Sequence, UBlueprint* Blueprint, UObject* EditObject, void* RawData) const
 {
-	FScopedTransaction Transaction(LOCTEXT("SetCallInEditor", "Set Call in Editor"));
-
-	const bool bCallInEditor = (NewCheckedState == ECheckBoxState::Checked);
-
-	TSet<UBlueprint*> Blueprints;
-
-	for (UK2Node* Endpoint : GetAllValidEndpoints())
-	{
-		if (UK2Node_CustomEvent* CustomEvent = Cast<UK2Node_CustomEvent>(Endpoint))
-		{
-			CustomEvent->Modify();
-			CustomEvent->bCallInEditor = bCallInEditor;
-
-			if (CustomEvent->HasValidBlueprint())
-			{
-				Blueprints.Add(CustomEvent->GetBlueprint());
-			}
-		}
-		else if (UK2Node_FunctionEntry* FunctionEntry = Cast<UK2Node_FunctionEntry>(Endpoint))
-		{
-			FunctionEntry->Modify();
-			FunctionEntry->MetaData.bCallInEditor = bCallInEditor;
-		
-			if (FunctionEntry->HasValidBlueprint())
-			{
-				Blueprints.Add(FunctionEntry->GetBlueprint());
-			}
-		}
-	}
-
-	for (UBlueprint* Blueprint : Blueprints)
+	UMovieSceneEventSectionBase* EventSection = Cast<UMovieSceneEventSectionBase>(EditObject);
+	FMovieSceneEvent* EntryPoint = static_cast<FMovieSceneEvent*>(RawData);
+	if (EventSection && EntryPoint)
 	{
-		FBlueprintEditorUtils::MarkBlueprintAsModified(Blueprint);
+		return FMovieSceneEventUtils::FindEndpoint(EntryPoint, EventSection, Blueprint);
 	}
+	return nullptr;
 }
 
-void FMovieSceneEventCustomization::OnBlueprintCompiled(UBlueprint*)
+void FMovieSceneEventCustomization::GetWellKnownParameterPinNames(UObject* EditObject, void* RawData, TArray<FName>& OutWellKnownParameters) const
 {
-	PropertyUtilities->ForceRefresh();
+	// We only have one well-known parameter: the (optional) bound object parameter.
+	FMovieSceneEvent* EntryPoint = static_cast<FMovieSceneEvent*>(RawData);
+	OutWellKnownParameters.Add(EntryPoint->BoundObjectPinName);
 }
 
-UMovieSceneSequence* FMovieSceneEventCustomization::GetCommonSequence() const
+void FMovieSceneEventCustomization::GetWellKnownParameterCandidates(UK2Node* Endpoint, TArray<FWellKnownParameterCandidates>& OutCandidates) const
 {
-	TArray<UObject*> EditObjects;
-	GetEditObjects(EditObjects);
+	FWellKnownParameterCandidates BoundObjectCandidates;
+	BoundObjectCandidates.Metadata.PickerLabel = LOCTEXT("BoundObjectPin_Label", "Pass Bound Object To");
+	BoundObjectCandidates.Metadata.PickerTooltip = LOCTEXT("BoundObjectPin_Tooltip", "Specifies a pin to pass the bound object(s) through when the event is triggered. Interface and object pins are both supported.");
 
-	UMovieSceneSequence* CommonEventSequence = nullptr;
-
-	for (UObject* Obj : EditObjects)
+	for (UEdGraphPin* Pin : Endpoint->Pins)
 	{
-		UMovieSceneSequence* ThisSequence = Obj ? Obj->GetTypedOuter<UMovieSceneSequence>() : nullptr;
-		if (CommonEventSequence && CommonEventSequence != ThisSequence)
+		if (Pin->Direction == EGPD_Output && 
+				(Pin->PinType.PinCategory == UEdGraphSchema_K2::PC_Object || 
+				 Pin->PinType.PinCategory == UEdGraphSchema_K2::PC_Interface))
 		{
-			return nullptr;
+			BoundObjectCandidates.CandidatePinNames.Add(Pin->GetFName());
 		}
-
-		CommonEventSequence = ThisSequence;
 	}
-	return CommonEventSequence;
+	
+	OutCandidates.Add(BoundObjectCandidates);
 }
 
-UMovieSceneEventTrack* FMovieSceneEventCustomization::GetCommonTrack() const
+bool FMovieSceneEventCustomization::SetWellKnownParameterPinName(UObject* EditObject, void* RawData, int32 ParameterIndex, FName BoundPinName)
 {
-	TArray<UObject*> EditObjects;
-	GetEditObjects(EditObjects);
-
-	UMovieSceneEventTrack* CommonEventTrack = nullptr;
-
-	for (UObject* Obj : EditObjects)
+	if (ensure(ParameterIndex == 0))
 	{
-		UMovieSceneEventTrack* ThisTrack = Obj ? Obj->GetTypedOuter<UMovieSceneEventTrack>() : nullptr;
-		if (CommonEventTrack && CommonEventTrack != ThisTrack)
-		{
-			return nullptr;
-		}
-
-		CommonEventTrack = ThisTrack;
+		FMovieSceneEvent* EntryPoint = static_cast<FMovieSceneEvent*>(RawData);
+		EntryPoint->BoundObjectPinName = BoundPinName;
+		return true;
 	}
-	return CommonEventTrack;
+	return false;
 }
 
-void FMovieSceneEventCustomization::IterateEndpoints(TFunctionRef<bool(UK2Node*)> Callback) const
+FMovieSceneDirectorBlueprintEndpointDefinition FMovieSceneEventCustomization::GenerateEndpointDefinition(UMovieSceneSequence* Sequence)
 {
-	TArray<UObject*> EditObjects;
-	GetEditObjects(EditObjects);
-
-	TArray<void*> RawData;
-	PropertyHandle->AccessRawData(RawData);
-
-	if (!ensure(RawData.Num() == EditObjects.Num()))
-	{
-		return;
-	}
+	UMovieScene* MovieScene = Sequence->GetMovieScene();
 
-	for (int32 Index = 0; Index < RawData.Num(); ++Index)
+	TOptional<FGuid> CommonObjectBindingID;
 	{
-		UMovieSceneEventSectionBase* EventSection = Cast<UMovieSceneEventSectionBase>(EditObjects[Index]);
-		FMovieSceneEvent*  EntryPoint   = static_cast<FMovieSceneEvent*>(RawData[Index]);
-		if (!EventSection || !EntryPoint)
-		{
-			continue;
-		}
+		TArray<UObject*> EditObjects;
+		GetEditObjects(EditObjects);
 
-		UMovieSceneSequence*       Sequence       = EventSection->GetTypedOuter<UMovieSceneSequence>();
-		FMovieSceneSequenceEditor* SequenceEditor = FMovieSceneSequenceEditor::Find(Sequence);
-		if (!SequenceEditor)
+		for (UObject* Outer : EditObjects)
 		{
-			continue;
-		}
+			FGuid ThisBindingID;
+			MovieScene->FindTrackBinding(*Outer->GetTypedOuter<UMovieSceneTrack>(), ThisBindingID);
 
-		UBlueprint* SequenceDirectorBP = SequenceEditor->FindDirectorBlueprint(Sequence);
-		if (SequenceDirectorBP)
-		{
-			UK2Node* Endpoint = FMovieSceneEventUtils::FindEndpoint(EntryPoint, EventSection, SequenceDirectorBP);
-			if (!Callback(Endpoint))
+			if (CommonObjectBindingID.IsSet() && CommonObjectBindingID != ThisBindingID)
 			{
-				return;
+				CommonObjectBindingID.Reset();
+				break;
 			}
+
+			CommonObjectBindingID = ThisBindingID;
 		}
 	}
+
+	return FMovieSceneEventUtils::GenerateEventDefinition(MovieScene, CommonObjectBindingID.Get(FGuid()));
 }
 
-TArray<UK2Node*> FMovieSceneEventCustomization::GetAllValidEndpoints() const
+void FMovieSceneEventCustomization::OnCreateEndpoint(UMovieSceneSequence* Sequence, UBlueprint* Blueprint, const TArray<UObject*> EditObjects, const TArray<void*> RawData, const FMovieSceneDirectorBlueprintEndpointDefinition& EndpointDefinition, UK2Node* NewEndpoint)
 {
-	TArray<UK2Node*> Endpoints;
+	check(RawData.Num() == EditObjects.Num());
 
-	IterateEndpoints(
-		[&Endpoints](UK2Node* Endpoint)
-		{
-			if (Endpoint)
-			{
-				Endpoints.Add(Endpoint);
-			}
-			return true;
-		}
-	);
+	struct FSectionData
+	{
+		TArray<FMovieSceneEvent*> EntryPoints;
+	};
 
-	return Endpoints;
-}
+	TSortedMap<UMovieSceneEventSectionBase*, FSectionData> PerSectionData;
 
-UK2Node* FMovieSceneEventCustomization::GetCommonEndpoint() const
-{
-	UK2Node* CommonEndpoint = nullptr;
+	for (int32 Index = 0; Index < RawData.Num(); ++Index)
+	{
+		UMovieSceneEventSectionBase* EventSection = Cast<UMovieSceneEventSectionBase>(EditObjects[Index]);
+		FMovieSceneEvent* EntryPoint = static_cast<FMovieSceneEvent*>(RawData[Index]);
 
-	IterateEndpoints(
-		[&CommonEndpoint](UK2Node* Endpoint)
+		if (EventSection)
 		{
-			if (CommonEndpoint && Endpoint != CommonEndpoint)
-			{
-				CommonEndpoint = nullptr;
-				return false;
-			}
-
-			CommonEndpoint = Endpoint;
-			return true;
+			FSectionData& SectionData = PerSectionData.FindOrAdd(EventSection);
+			SectionData.EntryPoints.Add(EntryPoint);
 		}
-	);
-
-	return CommonEndpoint;
-}
-
-TSharedRef<SWidget> FMovieSceneEventCustomization::GetMenuContent()
-{
-	FMenuBuilder MenuBuilder(true, nullptr, nullptr, true);
-
-	UMovieSceneSequence*       Sequence       = GetCommonSequence();
-	FMovieSceneSequenceEditor* SequenceEditor = FMovieSceneSequenceEditor::Find(Sequence);
-
-	UBlueprint* DirectorBP = SequenceEditor ? SequenceEditor->FindDirectorBlueprint(Sequence) : nullptr;
+	}
 
-	MenuBuilder.AddMenuEntry(
-		LOCTEXT("CreateEventEndpoint_Text",    "Create New Endpoint"),
-		LOCTEXT("CreateEventEndpoint_Tooltip", "Creates a new event endpoint in this sequence's blueprint."),
-		FSlateIcon(FAppStyle::GetAppStyleSetName(), "Sequencer.CreateEventBinding"),
-		FUIAction(
-			FExecuteAction::CreateSP(this, &FMovieSceneEventCustomization::CreateEventEndpoint)
-		)
-	);
+	UEdGraphPin* BoundObjectPin = FMovieSceneDirectorBlueprintUtils::FindCallTargetPin(NewEndpoint, EndpointDefinition.PossibleCallTargetClass);
 
-	const bool bAnyBoundEvents = GetAllValidEndpoints().Num() != 0;
-	if (!bAnyBoundEvents && Sequence)
-	{
-		MenuBuilder.AddSubMenu(
-			LOCTEXT("CreateQuickBinding_Text",    "Quick Bind"),
-			LOCTEXT("CreateQuickBinding_Tooltip", "Shows a list of functions on this object binding that can be bound directly to this event."),
-			FNewMenuDelegate::CreateSP(this, &FMovieSceneEventCustomization::PopulateQuickBindSubMenu, Sequence),
-			false /* bInOpenSubMenuOnClick */,
-			FSlateIcon(FAppStyle::GetAppStyleSetName(), "Sequencer.CreateQuickBinding"),
-			false /* bInShouldWindowAfterMenuSelection */
-		);
-	}
-	else
+	for (const TPair<UMovieSceneEventSectionBase*, FSectionData>& SectionPair : PerSectionData)
 	{
-		MenuBuilder.AddSubMenu(
-			LOCTEXT("Rebind_Text",         "Rebind To"),
-			LOCTEXT("Rebind_Text_Tooltip", "Rebinds this event to a different function call or event node."),
-			FNewMenuDelegate::CreateSP(this, &FMovieSceneEventCustomization::PopulateRebindSubMenu, Sequence),
-			false /* bInOpenSubMenuOnClick */,
-			FSlateIcon(FAppStyle::GetAppStyleSetName(), "Sequencer.CreateQuickBinding"),
-			false /* bInShouldCloseWindowAfterMenuSelection */
-		);
+		SectionPair.Key->Modify();
+		FMovieSceneEventUtils::BindEventSectionToBlueprint(SectionPair.Key, Blueprint);
 
-		MenuBuilder.AddMenuEntry(
-			LOCTEXT("ClearEventEndpoint_Text",    "Clear"),
-			LOCTEXT("ClearEventEndpoint_Tooltip", "Unbinds this event from its current binding."),
-			FSlateIcon(FAppStyle::GetAppStyleSetName(), "Sequencer.ClearEventBinding"),
-			FUIAction(
-				FExecuteAction::CreateSP(this, &FMovieSceneEventCustomization::ClearEventEndpoint)
-			)
-		);
+		for (FMovieSceneEvent* EntryPoint : SectionPair.Value.EntryPoints)
+		{
+			FMovieSceneEventUtils::SetEndpoint(EntryPoint, SectionPair.Key, NewEndpoint, BoundObjectPin);
+		}
 	}
-
-	return MenuBuilder.MakeWidget();
 }
 
-TSharedRef<SWidget> FMovieSceneEventCustomization::GetBoundObjectPinMenuContent()
+void FMovieSceneEventCustomization::OnSetEndpoint(UMovieSceneSequence* Sequence, UBlueprint* Blueprint, const TArray<UObject*> EditObjects, const TArray<void*> RawData, const FMovieSceneDirectorBlueprintEndpointDefinition& EndpointDefinition, UK2Node* NewEndpoint)
 {
-	UK2Node* CommonEndpoint = GetCommonEndpoint();
-	if (CommonEndpoint)
-	{
-		bool bAnyValidPins        = false;
-		bool bBoundObjectPinIsSet = false;
+	check(EditObjects.Num() == RawData.Num());
 
-		FMenuBuilder MenuBuilder(true, nullptr, nullptr, true);
-		for (UEdGraphPin* Pin : CommonEndpoint->Pins)
+	for (int32 Index = 0; Index < RawData.Num(); ++Index)
+	{
+		UMovieSceneEventSectionBase* EventSection = Cast<UMovieSceneEventSectionBase>(EditObjects[Index]);
+		if (EventSection)
 		{
-			if (Pin->Direction == EGPD_Output && (Pin->PinType.PinCategory == UEdGraphSchema_K2::PC_Object || Pin->PinType.PinCategory == UEdGraphSchema_K2::PC_Interface))
+			EventSection->Modify();
+			if (Blueprint)
 			{
-				bAnyValidPins = true;
-
-				FText PinText = FText::FromName(Pin->PinName);
-				MenuBuilder.AddMenuEntry(
-					PinText,
-					FText::Format(LOCTEXT("SetBoundObjectPin_Tooltip", "When calling this event with a bound object, pass the object through pin {0}."), PinText),
-					FSlateIcon(),
-					FUIAction(
-						FExecuteAction::CreateSP(this, &FMovieSceneEventCustomization::SetBoundObjectPinName, Pin->PinName),
-						FCanExecuteAction(),
-						FIsActionChecked::CreateSP(this, &FMovieSceneEventCustomization::CompareBoundObjectPinName, Pin->PinName)
-					),
-					NAME_None,
-					EUserInterfaceActionType::RadioButton
-				);
+				FMovieSceneEventUtils::BindEventSectionToBlueprint(EventSection, Blueprint);
 			}
 		}
 
-		MenuBuilder.AddMenuEntry(
-			LOCTEXT("ClearBoundObjectPin_Label", "Clear"),
-			LOCTEXT("ClearBoundObjectPin_Tooltip", "Clears the reference to the bound object pin, meaning any bound objects will not be passed to the event"),
-			FSlateIcon(),
-			FUIAction(
-				FExecuteAction::CreateSP(this, &FMovieSceneEventCustomization::SetBoundObjectPinName, FName())
-			)
-		);
-
-		if (bAnyValidPins || bBoundObjectPinIsSet)
-		{
-			return MenuBuilder.MakeWidget();
-		}
-		else
+		FMovieSceneEvent* EntryPoint = static_cast<FMovieSceneEvent*>(RawData[Index]);
+		if (EntryPoint)
 		{
-			return SNew(STextBlock).Text(LOCTEXT("BoundObjectPinError_NoPins", "No compatible pins were found. Only Object and Interface pins can be bound to objects."));
+			UEdGraphPin* CallTargetPin = FMovieSceneDirectorBlueprintUtils::FindCallTargetPin(NewEndpoint, EndpointDefinition.PossibleCallTargetClass);
+			FMovieSceneEventUtils::SetEndpoint(EntryPoint, EventSection, NewEndpoint, CallTargetPin);
 		}
 	}
-
-	return SNew(STextBlock).Text(LOCTEXT("BoundObjectPinError_MultipleEvents", "Cannot choose a bound object pin with multiple events selected"));
-}
-
-bool FMovieSceneEventCustomization::CompareBoundObjectPinName(FName InPinName) const
-{
-	bool bMatch = false;
-
-	auto Enumerator = [&bMatch, InPinName](void* Ptr, const int32, const int32)
-	{
-		bMatch = (static_cast<FMovieSceneEvent*>(Ptr)->BoundObjectPinName == InPinName);
-		return bMatch;
-	};
-
-	PropertyHandle->EnumerateRawData(Enumerator);
-
-	return bMatch;
-}
-
-void FMovieSceneEventCustomization::SetBoundObjectPinName(FName InNewBoundObjectPinName)
-{
-	FScopedTransaction Transaction(LOCTEXT("SetBoundObjectPinTransaction", "Set Bound Object Pin"));
-
-	TArray<UObject*> EditObjects;
-	GetEditObjects(EditObjects);
-
-	for (UObject* Object : EditObjects)
-	{
-		Object->Modify();
-	}
-
-	auto Enumerator = [InNewBoundObjectPinName](void* Ptr, const int32, const int32)
-	{
-		static_cast<FMovieSceneEvent*>(Ptr)->BoundObjectPinName = InNewBoundObjectPinName;
-		return true;
-	};
-	PropertyHandle->EnumerateRawData(Enumerator);
-
-	// Ensure that anything listening for property changed notifications are notified of the new binding
-	PropertyHandle->NotifyFinishedChangingProperties();
-	PropertyUtilities->ForceRefresh();
-}
-
-const FSlateBrush* FMovieSceneEventCustomization::GetBoundObjectPinIcon() const
-{
-	if (CompareBoundObjectPinName(NAME_None))
-	{
-		return nullptr;
-	}
-
-	return FAppStyle::GetBrush("Graph.Pin.Disconnected_VarA");
-}
-
-FText FMovieSceneEventCustomization::GetBoundObjectPinText() const
-{
-	TArray<void*> RawData;
-	PropertyHandle->AccessRawData(RawData);
-
-	FName CommonPinName;
-
-	for (void* Ptr : RawData)
-	{
-		FName PinName = static_cast<FMovieSceneEvent*>(Ptr)->BoundObjectPinName;
-		if (CommonPinName == NAME_None)
-		{
-			CommonPinName = PinName;
-		}
-		else if (CommonPinName != PinName)
-		{
-			return LOCTEXT("MultiplePinValues", "Multiple Values");
-		}
-	}
-
-	return FText::FromName(CommonPinName);
-}
-
-void FMovieSceneEventCustomization::PopulateQuickBindSubMenu(FMenuBuilder& MenuBuilder, UMovieSceneSequence* Sequence)
-{
-	FMovieSceneSequenceEditor* SequenceEditor = FMovieSceneSequenceEditor::Find(Sequence);
-	if (!SequenceEditor)
-	{
-		return;
-	}
-
-	UBlueprint* Blueprint = SequenceEditor->GetOrCreateDirectorBlueprint(Sequence);
-	if (!Blueprint)
-	{
-		return;
-	}
-
-	UMovieScene* MovieScene = Sequence->GetMovieScene();
-
-	TOptional<FGuid> CommonObjectBindingID;
-	{
-		TArray<UObject*> EditObjects;
-		GetEditObjects(EditObjects);
-
-		for (UObject* Outer : EditObjects)
-		{
-			FGuid ThisBindingID;
-			MovieScene->FindTrackBinding(*Outer->GetTypedOuter<UMovieSceneTrack>(), ThisBindingID);
-
-			if (CommonObjectBindingID.IsSet() && CommonObjectBindingID != ThisBindingID)
-			{
-				CommonObjectBindingID.Reset();
-				break;
-			}
-
-			CommonObjectBindingID = ThisBindingID;
-		}
-	}
-
-
-	FMovieSceneEventEndpointParameters EventParams = FMovieSceneEventEndpointParameters::Generate(MovieScene, CommonObjectBindingID.Get(FGuid()));
-
-	TSharedRef<SGraphActionMenu> ActionMenu = SNew(SGraphActionMenu)
-		.OnCreateCustomRowExpander_Static([](const FCustomExpanderData& Data) -> TSharedRef<SExpanderArrow> { return SNew(SExpanderArrow, Data.TableRow); })
-		.OnCollectAllActions_Static(UE_MovieSceneEventCustomization::CollectQuickBindActions, Blueprint, EventParams.BoundObjectPinClass)
-		.OnActionSelected(this, &FMovieSceneEventCustomization::HandleQuickBindActionSelected, Blueprint, EventParams);
-
-	ActionMenu->RegisterActiveTimer(0.f, FWidgetActiveTimerDelegate::CreateLambda(
-		[FilterTextBox = ActionMenu->GetFilterTextBox()](double, float)
-		{
-			FSlateApplication::Get().SetKeyboardFocus(FilterTextBox);
-			return EActiveTimerReturnType::Stop;
-		}
-	));
-
-	MenuBuilder.AddWidget(
-		SNew(SBox)
-		.WidthOverride(300.f)
-		.MaxDesiredHeight(500.f)
-		[
-			ActionMenu
-		],
-		FText()
-	);
-}
-
-void FMovieSceneEventCustomization::HandleQuickBindActionSelected(const TArray<TSharedPtr<FEdGraphSchemaAction>>& SelectedAction, ESelectInfo::Type InSelectionType, UBlueprint* Blueprint, FMovieSceneEventEndpointParameters Params)
-{
-	if (InSelectionType != ESelectInfo::OnMouseClick && InSelectionType != ESelectInfo::OnKeyPress)
-	{
-		return;
-	}
-
-	for (TSharedPtr<FEdGraphSchemaAction> Action : SelectedAction)
-	{
-		if (Action->GetTypeId() == FBlueprintActionMenuItem::StaticGetTypeId())
-		{
-			const UBlueprintFunctionNodeSpawner* FunctionNodeSpawner = Cast<UBlueprintFunctionNodeSpawner>(static_cast<FBlueprintActionMenuItem*>(Action.Get())->GetRawAction());
-			const UFunction* FunctionToCall = FunctionNodeSpawner ? FunctionNodeSpawner->GetFunction() : nullptr;
-
-			if (FunctionToCall)
-			{
-				UClass* OuterClass = CastChecked<UClass>(FunctionToCall->GetOuter());
-				if (OuterClass->ClassGeneratedBy == Blueprint && Blueprint->SkeletonGeneratedClass)
-				{
-					// Attempt to locate a custom event or a function graph of this name on the blueprint
-					for (UEdGraph* Graph : Blueprint->UbergraphPages)
-					{
-						for (UEdGraphNode* Node : Graph->Nodes)
-						{
-							UK2Node_CustomEvent* CustomEvent = Cast<UK2Node_CustomEvent>(Node);
-							if (CustomEvent && Blueprint->SkeletonGeneratedClass->FindFunctionByName(CustomEvent->GetFunctionName()) == FunctionToCall)
-							{
-								// Use this custom event
-								UEdGraphPin* BoundObjectPin = FMovieSceneEventUtils::FindBoundObjectPin(CustomEvent, Params.BoundObjectPinClass);
-								this->SetEventEndpoint(CustomEvent, BoundObjectPin, CustomEvent, EAutoCreatePayload::Variables);
-								return;
-							}
-						}
-					}
-
-					for (UEdGraph* Graph : Blueprint->FunctionGraphs)
-					{
-						if (Blueprint->SkeletonGeneratedClass->FindFunctionByName(Graph->GetFName()) != FunctionToCall)
-						{
-							continue;
-						}
-
-						// Use this function graph for the event endpoint
-						for (UEdGraphNode* Node : Graph->Nodes)
-						{
-							UK2Node_FunctionEntry* FunctionEntry = Cast<UK2Node_FunctionEntry>(Node);
-							if (FunctionEntry)
-							{
-								UEdGraphPin* BoundObjectPin = FMovieSceneEventUtils::FindBoundObjectPin(FunctionEntry, Params.BoundObjectPinClass);
-								this->SetEventEndpoint(FunctionEntry, BoundObjectPin, FunctionEntry, EAutoCreatePayload::Variables);
-								return;
-							}
-						}
-					}
-				}
-			}
-		}
-
-		if (Action)
-		{
-			UK2Node_CustomEvent* NewEventEndpoint = FMovieSceneEventUtils::CreateUserFacingEvent(Blueprint, Params);
-
-			UEdGraphPin* ThenPin = NewEventEndpoint->FindPin(UEdGraphSchema_K2::PN_Then, EGPD_Output);
-			UEdGraphPin* BoundObjectPin = FMovieSceneEventUtils::FindBoundObjectPin(NewEventEndpoint, Params.BoundObjectPinClass);
-
-			FVector2D NodePosition(NewEventEndpoint->NodePosX + 400.f, NewEventEndpoint->NodePosY);
-			UEdGraphNode* NewNode = Action->PerformAction(NewEventEndpoint->GetGraph(), BoundObjectPin ? BoundObjectPin : ThenPin, NodePosition);
-
-			if (NewNode)
-			{
-				UEdGraphPin* ExecPin = NewNode->FindPin(UEdGraphSchema_K2::PN_Execute, EGPD_Input);
-				if (ensure(ThenPin && ExecPin))
-				{
-					ThenPin->MakeLinkTo(ExecPin);
-				}
-			}
-
-			this->SetEventEndpoint(NewEventEndpoint, BoundObjectPin, Cast<UK2Node>(NewNode), EAutoCreatePayload::Pins | EAutoCreatePayload::Variables);
-		}
-	}
-}
-
-void FMovieSceneEventCustomization::PopulateRebindSubMenu(FMenuBuilder& MenuBuilder, UMovieSceneSequence* Sequence)
-{
-	FMovieSceneSequenceEditor* SequenceEditor = FMovieSceneSequenceEditor::Find(Sequence);
-	if (!SequenceEditor)
-	{
-		return;
-	}
-
-	UBlueprint* Blueprint = SequenceEditor->GetOrCreateDirectorBlueprint(Sequence);
-	if (!Blueprint)
-	{
-		return;
-	}
-
-	UMovieScene* MovieScene = Sequence->GetMovieScene();
-
-	TOptional<FGuid> CommonObjectBindingID;
-	{
-		TArray<UObject*> EditObjects;
-		GetEditObjects(EditObjects);
-
-		for (UObject* Outer : EditObjects)
-		{
-			FGuid ThisBindingID;
-			MovieScene->FindTrackBinding(*Outer->GetTypedOuter<UMovieSceneTrack>(), ThisBindingID);
-
-			if (CommonObjectBindingID.IsSet() && CommonObjectBindingID != ThisBindingID)
-			{
-				CommonObjectBindingID.Reset();
-				break;
-			}
-
-			CommonObjectBindingID = ThisBindingID;
-		}
-	}
-
-	FMovieSceneEventEndpointParameters EventParams = FMovieSceneEventEndpointParameters::Generate(MovieScene, CommonObjectBindingID.Get(FGuid()));
-
-	TSharedRef<SGraphActionMenu> ActionMenu = SNew(SGraphActionMenu)
-		.OnCreateCustomRowExpander_Static([](const FCustomExpanderData& Data) -> TSharedRef<SExpanderArrow> { return SNew(SExpanderArrow, Data.TableRow); })
-		.OnCollectAllActions_Static(UE_MovieSceneEventCustomization::CollectAllRebindActions, Blueprint)
-		.OnActionSelected(this, &FMovieSceneEventCustomization::HandleRebindActionSelected, Blueprint, EventParams.BoundObjectPinClass);
-
-	ActionMenu->RegisterActiveTimer(0.f, FWidgetActiveTimerDelegate::CreateLambda(
-		[FilterTextBox = ActionMenu->GetFilterTextBox()](double, float)
-		{
-			FSlateApplication::Get().SetKeyboardFocus(FilterTextBox);
-			return EActiveTimerReturnType::Stop;
-		}
-	));
-
-	MenuBuilder.AddWidget(
-		SNew(SBox)
-		.WidthOverride(300.f)
-		.MaxDesiredHeight(500.f)
-		[
-			ActionMenu
-		],
-		FText()
-	);
-}
-
-void FMovieSceneEventCustomization::HandleRebindActionSelected(const TArray<TSharedPtr<FEdGraphSchemaAction>>& SelectedAction, ESelectInfo::Type InSelectionType, UBlueprint* Blueprint, UClass* BoundObjectPinClass)
-{
-	if (InSelectionType != ESelectInfo::OnMouseClick && InSelectionType != ESelectInfo::OnKeyPress)
-	{
-		return;
-	}
-
-	if (!Blueprint->SkeletonGeneratedClass)
-	{
-		return;
-	}
-
-	for (TSharedPtr<FEdGraphSchemaAction> Action : SelectedAction)
-	{
-		if (Action->GetTypeId() != FBlueprintActionMenuItem::StaticGetTypeId())
-		{
-			continue;
-		}
-
-		const UBlueprintFunctionNodeSpawner* FunctionNodeSpawner = Cast<UBlueprintFunctionNodeSpawner>(static_cast<FBlueprintActionMenuItem*>(Action.Get())->GetRawAction());
-		const UFunction* FunctionToCall = FunctionNodeSpawner ? FunctionNodeSpawner->GetFunction() : nullptr;
-		if (!FunctionToCall)
-		{
-			continue;
-		}
-
-		UClass* OuterClass = CastChecked<UClass>(FunctionToCall->GetOuter());
-		if (OuterClass->ClassGeneratedBy != Blueprint)
-		{
-			continue;
-		}
-
-		// Attempt to locate a custom event or a function graph of this name on the blueprint
-		for (UEdGraph* Graph : Blueprint->UbergraphPages)
-		{
-			for (UEdGraphNode* Node : Graph->Nodes)
-			{
-				UK2Node_CustomEvent* CustomEvent = Cast<UK2Node_CustomEvent>(Node);
-				if (CustomEvent && Blueprint->SkeletonGeneratedClass->FindFunctionByName(CustomEvent->GetFunctionName()) == FunctionToCall)
-				{
-					// Use this custom event
-					UEdGraphPin* BoundObjectPin = FMovieSceneEventUtils::FindBoundObjectPin(CustomEvent, BoundObjectPinClass);
-					this->SetEventEndpoint(CustomEvent, BoundObjectPin, CustomEvent, EAutoCreatePayload::Variables);
-					return;
-				}
-			}
-		}
-
-		for (UEdGraph* Graph : Blueprint->FunctionGraphs)
-		{
-			if (Blueprint->SkeletonGeneratedClass->FindFunctionByName(Graph->GetFName()) != FunctionToCall)
-			{
-				continue;
-			}
-
-			// Use this function graph for the event endpoint
-			for (UEdGraphNode* Node : Graph->Nodes)
-			{
-				UK2Node_FunctionEntry* FunctionEntry = Cast<UK2Node_FunctionEntry>(Node);
-				if (FunctionEntry)
-				{
-					UEdGraphPin* BoundObjectPin = FMovieSceneEventUtils::FindBoundObjectPin(FunctionEntry, BoundObjectPinClass);
-					this->SetEventEndpoint(FunctionEntry, BoundObjectPin, FunctionEntry, EAutoCreatePayload::Variables);
-					return;
-				}
-			}
-		}
-	}
-
-	ensureMsgf(false, TEXT("Unknown blueprint action type encountered for rebinding"));
-}
-
-const FSlateBrush* FMovieSceneEventCustomization::GetEventIcon() const
-{
-	UK2Node* CommonEndpoint = GetCommonEndpoint();
-
-	if (CommonEndpoint)
-	{
-		FLinearColor Color;
-		FSlateIcon EndpointIcon = CommonEndpoint->GetIconAndTint(Color);
-		return EndpointIcon.GetIcon();
-	}
-	else
-	{
-		TArray<void*> RawData;
-		PropertyHandle->AccessRawData(RawData);
-		if (RawData.Num() > 1)
-		{
-			return FAppStyle::GetBrush("Sequencer.MultipleEvents");
-		}
-	}
-
-	return FAppStyle::GetBrush("Sequencer.UnboundEvent");
-}
-
-FText FMovieSceneEventCustomization::GetEventName() const
-{
-	UEdGraphNode* CommonEndpoint = GetCommonEndpoint();
-
-	if (CommonEndpoint)
-	{
-		return CommonEndpoint->GetNodeTitle(ENodeTitleType::MenuTitle);
-	}
-	else
-	{
-		TArray<void*> RawData;
-		PropertyHandle->AccessRawData(RawData);
-		if (RawData.Num() != 1)
-		{
-			return LOCTEXT("MultipleValuesText", "Multiple Values");
-		}
-	}
-
-	return LOCTEXT("UnboundText", "Unbound");
-}
-
-void FMovieSceneEventCustomization::SetEventEndpoint(UK2Node* NewEndpoint, UEdGraphPin* BoundObjectPin, UK2Node* PayloadTemplate, EAutoCreatePayload AutoCreatePayload)
-{
-	FScopedTransaction Transaction(LOCTEXT("SetEventEndpoint", "Set Event Endpoint"));
-
-	// Modify and assign the blueprint for outer sections
-	TArray<UObject*> EditObjects;
-	GetEditObjects(EditObjects);
-
-	TArray<UBlueprint*> AllBlueprints;
-
-	// If we're assigning a new valid endpoint, it must reside within the same blueprint as everything we're assigning it to.
-	// Anything else must be implemented as a call function node connected to a custom event node
-	UBlueprint* Blueprint = NewEndpoint && NewEndpoint->HasValidBlueprint() ? NewEndpoint->GetBlueprint() : nullptr;
-	if (Blueprint)
-	{
-		for (UObject* Outer : EditObjects)
-		{
-			UMovieSceneSequence*       Sequence           = Outer->GetTypedOuter<UMovieSceneSequence>();
-			FMovieSceneSequenceEditor* SequenceEditor     = FMovieSceneSequenceEditor::Find(Sequence);
-			UBlueprint*                SequenceDirectorBP = SequenceEditor ? SequenceEditor->GetOrCreateDirectorBlueprint(Sequence) : nullptr;
-
-			if (!ensureAlwaysMsgf(SequenceDirectorBP == Blueprint, TEXT("Attempting to assign an event endpoint to an event with a different sequence director Blueprint.")))
-			{
-				Transaction.Cancel();
-				return;
-			}
-		}
-
-		Blueprint->Modify();
-	}
-
-	for (UObject* Outer : EditObjects)
-	{
-		UMovieSceneEventSectionBase* BaseEventSection = Cast<UMovieSceneEventSectionBase>(Outer);
-		if (BaseEventSection)
-		{
-			BaseEventSection->Modify();
-			if (Blueprint)
-			{
-				FMovieSceneEventUtils::BindEventSectionToBlueprint(BaseEventSection, Blueprint);
-			}
-		}
-	}
-
-	// Map of the Payload Variable Names to their Default Values as Strings
-	TMap<FName, FString> PayloadVariables;
-	if (PayloadTemplate && EnumHasAnyFlags(AutoCreatePayload, EAutoCreatePayload::Variables | EAutoCreatePayload::Pins))
-	{
-		UK2Node_EditablePinBase* EditableNode = Cast<UK2Node_EditablePinBase>(NewEndpoint);
-
-		for (UEdGraphPin* PayloadPin : PayloadTemplate->Pins)
-		{
-			if (PayloadPin != BoundObjectPin && PayloadPin->Direction == EGPD_Input && PayloadPin->PinType.PinCategory != UEdGraphSchema_K2::PC_Exec && PayloadPin->LinkedTo.Num() == 0 && PayloadPin->PinName != UEdGraphSchema_K2::PN_Self)
-			{
-				// Make a payload variable for this pin
-				if (EnumHasAnyFlags(AutoCreatePayload, EAutoCreatePayload::Variables))
-				{
-					PayloadVariables.Add(PayloadPin->PinName, PayloadPin->DefaultValue);
-				}
-
-				if (EditableNode && EnumHasAnyFlags(AutoCreatePayload, EAutoCreatePayload::Pins))
-				{
-					// Pins for ref parameters for functions default to bIsReference but the payload cannot be by reference.
-					PayloadPin->PinType.bIsReference = false;
-					
-					UEdGraphPin* NewPin = EditableNode->CreateUserDefinedPin(PayloadPin->PinName, PayloadPin->PinType, EGPD_Output);
-					if (PayloadTemplate != NewEndpoint && NewPin)
-					{
-						NewPin->MakeLinkTo(PayloadPin);
-					}
-				}
-			}
-		}
-	}
-
-	// Assign the endpoints to all events
-	TArray<void*> RawData;
-	PropertyHandle->AccessRawData(RawData);
-
-	check(EditObjects.Num() == RawData.Num());
-
-	for (int32 Index = 0; Index < RawData.Num(); ++Index)
-	{
-		UMovieSceneEventSectionBase* EventSection = Cast<UMovieSceneEventSectionBase>(EditObjects[Index]);
-		FMovieSceneEvent*  EntryPoint   = static_cast<FMovieSceneEvent*>(RawData[Index]);
-
-		if (EntryPoint)
-		{
-			FMovieSceneEventUtils::SetEndpoint(EntryPoint, EventSection, NewEndpoint, BoundObjectPin);
-
-			for (const TPair<FName, FString>& PayloadVar : PayloadVariables)
-			{
-				if (!EntryPoint->PayloadVariables.Contains(PayloadVar.Key))
-				{
-					FMovieSceneEventPayloadVariable PayloadValue;
-					PayloadValue.Value = PayloadVar.Value;
-
-					EntryPoint->PayloadVariables.Add(PayloadVar.Key, MoveTemp(PayloadValue));
-				}
-			}
-		}
-	}
-
-	// Ensure that anything listening for property changed notifications are notified of the new binding
-	PropertyHandle->NotifyFinishedChangingProperties();
-
-	// Compile the blueprint now that clients have had a chance to update underlying data (we do this after to ensure we are compiling the correct data
-	if (Blueprint)
-	{
-		FKismetEditorUtilities::CompileBlueprint(Blueprint);
-	}
-
-	// Forcibly update the panel now that our endpoint has changed
-	PropertyUtilities->ForceRefresh();
-
-	if (NewEndpoint)
-	{
-		FKismetEditorUtilities::BringKismetToFocusAttentionOnObject(NewEndpoint, false);
-	}
-}
-
-void FMovieSceneEventCustomization::CreateEventEndpoint()
-{
-	struct FSectionData
-	{
-		TArray<FMovieSceneEvent*, TInlineAllocator<1>>  EntryPoints;
-	};
-	struct FTrackData
-	{
-		TSortedMap<UMovieSceneEventSectionBase*, FSectionData>  Sections;
-	};
-	struct FSequenceData
-	{
-		TSortedMap<UMovieSceneEventTrack*, FTrackData> Tracks;
-	};
-
-	TSortedMap<UMovieSceneSequence*, FSequenceData> PerSequenceData;
-
-	// Populate all the sequences represented by this customization
-	{
-		TArray<void*> RawData;
-		PropertyHandle->AccessRawData(RawData);
-
-		TArray<UObject*> EditObjects;
-		GetEditObjects(EditObjects);
-
-		check(RawData.Num() == EditObjects.Num());
-
-		for (int32 Index = 0; Index < RawData.Num(); ++Index)
-		{
-			UMovieSceneEventSectionBase* EventSection = Cast<UMovieSceneEventSectionBase>(EditObjects[Index]);
-			FMovieSceneEvent*  EntryPoint   = static_cast<FMovieSceneEvent*>(RawData[Index]);
-
-			if (EventSection)
-			{
-				FSequenceData& SequenceData = PerSequenceData.FindOrAdd(EventSection->GetTypedOuter<UMovieSceneSequence>());
-				FTrackData&    TrackData    = SequenceData.Tracks.FindOrAdd(EventSection->GetTypedOuter<UMovieSceneEventTrack>());
-				FSectionData&  SectionData  = TrackData.Sections.FindOrAdd(EventSection);
-
-				SectionData.EntryPoints.Add(EntryPoint);
-			}
-		}
-	}
-
-	UK2Node_CustomEvent* NewEventEndpoint = nullptr;
-
-	FScopedTransaction Transaction(LOCTEXT("CreateEventEndpoint", "Create Event Endpoint"));
-
-	for (const TPair<UMovieSceneSequence*, FSequenceData>& SequencePair : PerSequenceData)
-	{
-		FMovieSceneSequenceEditor* SequenceEditor = FMovieSceneSequenceEditor::Find(SequencePair.Key);
-		if (!SequenceEditor)
-		{
-			continue;
-		}
-
-		UBlueprint* SequenceDirectorBP = SequenceEditor->GetOrCreateDirectorBlueprint(SequencePair.Key);
-		if (!SequenceDirectorBP)
-		{
-			continue;
-		}
-
-		UMovieScene* MovieScene = SequencePair.Key->GetMovieScene();
-
-		TOptional<FGuid> CommonObjectBindingID;
-		for (const TPair<UMovieSceneEventTrack*, FTrackData>& TrackPair : SequencePair.Value.Tracks)
-		{
-			FGuid ThisBindingID;
-			MovieScene->FindTrackBinding(*TrackPair.Key, ThisBindingID);
-
-			if (CommonObjectBindingID.IsSet() && CommonObjectBindingID != ThisBindingID)
-			{
-				CommonObjectBindingID.Reset();
-				break;
-			}
-
-			CommonObjectBindingID = ThisBindingID;
-		}
-
-		FMovieSceneEventEndpointParameters Parameters = FMovieSceneEventEndpointParameters::Generate(MovieScene, CommonObjectBindingID.Get(FGuid()));
-
-		SequenceDirectorBP->Modify();
-
-		NewEventEndpoint = FMovieSceneEventUtils::CreateUserFacingEvent(SequenceDirectorBP, Parameters);
-		if (!NewEventEndpoint)
-		{
-			continue;
-		}
-
-		UEdGraphPin* BoundObjectPin = FMovieSceneEventUtils::FindBoundObjectPin(NewEventEndpoint, Parameters.BoundObjectPinClass);
-		for (const TPair<UMovieSceneEventTrack*, FTrackData>& TrackPair : SequencePair.Value.Tracks)
-		{
-			for (const TPair<UMovieSceneEventSectionBase*, FSectionData>& SectionPair : TrackPair.Value.Sections)
-			{
-				SectionPair.Key->Modify();
-				FMovieSceneEventUtils::BindEventSectionToBlueprint(SectionPair.Key, SequenceDirectorBP);
-
-				for (FMovieSceneEvent* EntryPoint : SectionPair.Value.EntryPoints)
-				{
-					FMovieSceneEventUtils::SetEndpoint(EntryPoint, SectionPair.Key, NewEventEndpoint, BoundObjectPin);
-				}
-			}
-		}
-
-		FBlueprintEditorUtils::MarkBlueprintAsModified(SequenceDirectorBP);
-	}
-
-	// Ensure that anything listening for property changed notifications are notified of the new binding
-	PropertyHandle->NotifyFinishedChangingProperties();
-	PropertyUtilities->ForceRefresh();
-
-	if (NewEventEndpoint)
-	{
-		FKismetEditorUtilities::BringKismetToFocusAttentionOnObject(NewEventEndpoint, false);
-	}
-}
-
-void FMovieSceneEventCustomization::ClearEventEndpoint()
-{
-	SetEventEndpoint(nullptr, nullptr, nullptr, EAutoCreatePayload::None);
-}
-
-void FMovieSceneEventCustomization::NavigateToDefinition()
-{
-	UEdGraphNode* CommonEndpoint = GetCommonEndpoint();
-	if (CommonEndpoint)
-	{
-		FKismetEditorUtilities::BringKismetToFocusAttentionOnObject(CommonEndpoint, false);
-	}
 }
 
 #undef LOCTEXT_NAMESPACE
diff --git a/Engine/Source/Editor/MovieSceneTools/Private/MovieSceneEventUtils.cpp b/Engine/Source/Editor/MovieSceneTools/Private/MovieSceneEventUtils.cpp
--- a/Engine/Source/Editor/MovieSceneTools/Private/MovieSceneEventUtils.cpp
+++ b/Engine/Source/Editor/MovieSceneTools/Private/MovieSceneEventUtils.cpp
@@ -31,8 +31,6 @@
 #include "K2Node_FunctionEntry.h"
 #include "Kismet/KismetSystemLibrary.h"
 #include "Kismet2/BlueprintEditorUtils.h"
-#include "Kismet2/CompilerResultsLog.h"
-#include "KismetCompiler.h"
 #include "Logging/LogCategory.h"
 #include "Logging/LogMacros.h"
 #include "Math/Vector2D.h"
@@ -66,7 +64,7 @@ class UBlueprintExtension;
 
 #define LOCTEXT_NAMESPACE "MovieSceneEventUtils"
 
-FMovieSceneEventEndpointParameters FMovieSceneEventEndpointParameters::Generate(UMovieSceneEventTrack* Track)
+FMovieSceneDirectorBlueprintEndpointDefinition FMovieSceneEventUtils::GenerateEventDefinition(UMovieSceneTrack* Track)
 {
 	check(Track);
 
@@ -75,53 +73,57 @@ FMovieSceneEventEndpointParameters FMovieSceneEventEndpointParameters::Generate(
 	FGuid ObjectBindingID;
 	if (MovieScene->FindTrackBinding(*Track, ObjectBindingID))
 	{
-		return Generate(MovieScene, ObjectBindingID);
+		return GenerateEventDefinition(MovieScene, ObjectBindingID);
 	}
 
-	FMovieSceneEventEndpointParameters Params;
-	Params.SanitizedObjectName = TEXT("None");
-	Params.SanitizedEventName  = TEXT("SequenceEvent");
+	FMovieSceneDirectorBlueprintEndpointDefinition Definition;
+	Definition.EndpointType = EMovieSceneDirectorBlueprintEndpointType::Event;
+	Definition.EndpointName  = TEXT("SequenceEvent");
+	Definition.GraphName = TEXT("Sequencer Events");
 
-	return Params;
+	return Definition;
 }
 
-FMovieSceneEventEndpointParameters FMovieSceneEventEndpointParameters::Generate(UMovieScene* MovieScene, const FGuid& ObjectBindingID)
+FMovieSceneDirectorBlueprintEndpointDefinition FMovieSceneEventUtils::GenerateEventDefinition(UMovieScene* MovieScene, const FGuid& ObjectBindingID)
 {
 	check(MovieScene);
 
-	FMovieSceneEventEndpointParameters Params;
+	FMovieSceneDirectorBlueprintEndpointDefinition Definition;
+	Definition.EndpointType = EMovieSceneDirectorBlueprintEndpointType::Event;
+	Definition.GraphName = TEXT("Sequencer Events");
 
 	if (ObjectBindingID.IsValid())
 	{
-		Params.SanitizedObjectName = MovieScene->GetObjectDisplayName(ObjectBindingID).ToString();
+		FString BoundObjectName = MovieScene->GetObjectDisplayName(ObjectBindingID).ToString();
 
-		for (TCHAR& Char : Params.SanitizedObjectName)
+		UClass* BoundObjectClass = nullptr;
+		if (FMovieScenePossessable* Possessable = MovieScene->FindPossessable(ObjectBindingID))
 		{
-			if (FCString::Strchr(INVALID_NAME_CHARACTERS, Char) != nullptr)
-			{
-				Char = '_';
-			}
+			BoundObjectClass = const_cast<UClass*>(Possessable->GetPossessedObjectClass());
 		}
-
-		Params.SanitizedEventName = Params.SanitizedObjectName + TEXT("_Event");
-
-		if (FMovieScenePossessable* Possessable = MovieScene->FindPossessable(ObjectBindingID))
+		else if (FMovieSceneSpawnable* Spawnable = MovieScene->FindSpawnable(ObjectBindingID))
 		{
-			Params.BoundObjectPinClass = const_cast<UClass*>(Possessable->GetPossessedObjectClass());
+			BoundObjectClass = Spawnable->GetObjectTemplate()->GetClass();
 		}
-		else if(FMovieSceneSpawnable* Spawnable = MovieScene->FindSpawnable(ObjectBindingID))
+
+		FMovieSceneDirectorBlueprintEndpointParameter BoundObjectParam;
 		{
-			Params.BoundObjectPinClass = Spawnable->GetObjectTemplate()->GetClass();
+			BoundObjectParam.PinDirection = EGPD_Output;
+			BoundObjectParam.PinTypeCategory = UEdGraphSchema_K2::PC_Object;
+			BoundObjectParam.PinName = BoundObjectName;
+			BoundObjectParam.PinTypeClass = BoundObjectClass;
+			Definition.ExtraPins.Add(BoundObjectParam);
 		}
+
+		Definition.EndpointName = BoundObjectName + TEXT("_Event");
+		Definition.PossibleCallTargetClass = BoundObjectClass;
 	}
 	else
 	{
-		Params.BoundObjectPinClass = nullptr;
-		Params.SanitizedObjectName = TEXT("None");
-		Params.SanitizedEventName  = TEXT("SequenceEvent");
+		Definition.EndpointName  = TEXT("SequenceEvent");
 	}
 
-	return Params;
+	return Definition;
 }
 
 UK2Node_CustomEvent* FMovieSceneEventUtils::BindNewUserFacingEvent(FMovieSceneEvent* EntryPoint, UMovieSceneEventSectionBase* EventSection, UBlueprint* Blueprint)
@@ -138,64 +140,19 @@ UK2Node_CustomEvent* FMovieSceneEventUtils::BindNewUserFacingEvent(FMovieSceneEv
 	FMovieSceneEventUtils::BindEventSectionToBlueprint(EventSection, Blueprint);
 
 	// Create the new user-facing event node
-	FMovieSceneEventEndpointParameters Params = FMovieSceneEventEndpointParameters::Generate(Track);
-	UK2Node_CustomEvent* NewEventNode = CreateUserFacingEvent(Blueprint, Params);
+	FMovieSceneDirectorBlueprintEndpointDefinition EndpointDefinition = FMovieSceneEventUtils::GenerateEventDefinition(Track);
+	UK2Node_CustomEvent* NewEventNode = FMovieSceneDirectorBlueprintUtils::CreateEventEndpoint(Blueprint, EndpointDefinition);
 
 	if (NewEventNode)
 	{
 		// Bind the node to the event entry point
-		UEdGraphPin* BoundObjectPin = FMovieSceneEventUtils::FindBoundObjectPin(NewEventNode, Params.BoundObjectPinClass);
+		UEdGraphPin* BoundObjectPin = FMovieSceneDirectorBlueprintUtils::FindCallTargetPin(NewEventNode, EndpointDefinition.PossibleCallTargetClass);
 		FMovieSceneEventUtils::SetEndpoint(EntryPoint, EventSection, NewEventNode, BoundObjectPin);
 	}
 
 	return NewEventNode;
 }
 
-UK2Node_CustomEvent* FMovieSceneEventUtils::CreateUserFacingEvent(UBlueprint* Blueprint, const FMovieSceneEventEndpointParameters& Parameters)
-{
-	static const TCHAR* const EventGraphName = TEXT("Sequencer Events");
-	UEdGraph* SequenceEventGraph = FindObject<UEdGraph>(Blueprint, EventGraphName);
-	if (!SequenceEventGraph)
-	{
-		SequenceEventGraph = FBlueprintEditorUtils::CreateNewGraph(Blueprint, EventGraphName, UEdGraph::StaticClass(), UEdGraphSchema_K2::StaticClass());
-
-		SequenceEventGraph->GraphGuid = FGuid::NewGuid();
-		Blueprint->UbergraphPages.Add(SequenceEventGraph);
-	}
-
-	// Create a custom event node to replace the original event node imported from text
-	UK2Node_CustomEvent* CustomEventNode = NewObject<UK2Node_CustomEvent>(SequenceEventGraph);
-
-	check(Parameters.SanitizedEventName.Len() != 0);
-	CustomEventNode->CustomFunctionName = FBlueprintEditorUtils::FindUniqueKismetName(Blueprint, Parameters.SanitizedEventName);
-
-	// Ensure that it is editable
-	CustomEventNode->bIsEditable = true;
-
-	CustomEventNode->CreateNewGuid();
-	CustomEventNode->PostPlacedNewNode();
-	CustomEventNode->AllocateDefaultPins();
-
-	const FVector2D NewPosition = SequenceEventGraph->GetGoodPlaceForNewNode();
-	CustomEventNode->NodePosX = NewPosition.X;
-	CustomEventNode->NodePosY = NewPosition.Y;
-
-	if (Parameters.BoundObjectPinClass != nullptr)
-	{
-		FEdGraphPinType PinType;
-		PinType.PinCategory = UEdGraphSchema_K2::PC_Object;
-		PinType.PinSubCategoryObject = Parameters.BoundObjectPinClass;
-
-		CustomEventNode->CreateUserDefinedPin(*Parameters.SanitizedObjectName, PinType, EGPD_Output, true);
-	}
-
-	SequenceEventGraph->AddNode(CustomEventNode, false, false);
-
-	FBlueprintEditorUtils::ValidateBlueprintChildVariables(Blueprint, CustomEventNode->CustomFunctionName);
-	FBlueprintEditorUtils::MarkBlueprintAsStructurallyModified(Blueprint);
-	return CustomEventNode;
-}
-
 UK2Node* FMovieSceneEventUtils::FindEndpoint(FMovieSceneEvent* EntryPoint, UMovieSceneEventSectionBase* EventSection, UBlueprint* OwnerBlueprint)
 {
 	check(OwnerBlueprint);
@@ -257,24 +214,6 @@ UK2Node* FMovieSceneEventUtils::FindEndpoint(FMovieSceneEvent* EntryPoint, UMovi
 	return nullptr;
 }
 
-UEdGraphPin* FMovieSceneEventUtils::FindBoundObjectPin(UK2Node* InEndpoint, UClass* BoundObjectPinClass)
-{
-	if (!BoundObjectPinClass)
-	{
-		return nullptr;
-	}
-
-	for (UEdGraphPin* Pin : InEndpoint->Pins)
-	{
-		if (Pin->Direction == EGPD_Output && Pin->PinType.PinCategory == UEdGraphSchema_K2::PC_Object && Pin->PinType.PinSubCategoryObject == BoundObjectPinClass)
-		{
-			return Pin;
-		}
-	}
-
-	return nullptr;
-}
-
 void FMovieSceneEventUtils::SetEndpoint(FMovieSceneEvent* EntryPoint, UMovieSceneEventSectionBase* EventSection, UK2Node* InNewEndpoint, UEdGraphPin* BoundObjectPin)
 {
 	check(EntryPoint);
@@ -292,7 +231,6 @@ void FMovieSceneEventUtils::SetEndpoint(FMovieSceneEvent* EntryPoint, UMovieScen
 
 		checkf(bIsFunction || bIsCustomEvent, TEXT("Only functions and custom events are supported as event endpoints"));
 
-
 		if (BoundObjectPin)
 		{
 			EntryPoint->BoundObjectPinName = BoundObjectPin->GetFName();
@@ -312,200 +250,21 @@ void FMovieSceneEventUtils::SetEndpoint(FMovieSceneEvent* EntryPoint, UMovieScen
 	}
 }
 
-
-UK2Node_FunctionEntry* FMovieSceneEventUtils::GenerateEntryPoint(UMovieSceneEventSectionBase* EventSection, FMovieSceneEvent* EntrypointDefinition, FKismetCompilerContext* Compiler, UEdGraphNode* Endpoint)
+UK2Node_FunctionEntry* FMovieSceneEventUtils::GenerateEntryPoint(FMovieSceneEvent* EntrypointDefinition, FKismetCompilerContext* Compiler, UEdGraphNode* Endpoint)
 {
-	check(EventSection && Compiler && Endpoint);
-
-	UBlueprint* Blueprint = Compiler->Blueprint;
-
-	UFunction* EndpointFunction = nullptr;
-
-	// Find the function that we need to call on the skeleton class
+	FMovieSceneDirectorBlueprintEndpointCall EndpointCall;
+	for (const TPair<FName, FMovieSceneEventPayloadVariable>& Pair : EntrypointDefinition->PayloadVariables)
 	{
-		if (UK2Node_Event* Event = Cast<UK2Node_Event>(Endpoint))
-		{
-			EndpointFunction = Blueprint->SkeletonGeneratedClass ? Blueprint->SkeletonGeneratedClass->FindFunctionByName(Event->GetFunctionName()) : nullptr;
-		}
-		else if (UK2Node_FunctionEntry* EndpointEntry = Cast<UK2Node_FunctionEntry>(Endpoint))
-		{
-			EndpointFunction = EndpointEntry->FindSignatureFunction();
-		}
-		else
-		{
-			Compiler->MessageLog.Error(*LOCTEXT("InvalidEndpoint_Error", "Sequencer event is bound to an invalid endpoint node @@").ToString(), Endpoint);
-		}
+		EndpointCall.PayloadVariables.Add(Pair.Key, Pair.Value.Value);
 	}
-
-	if (EndpointFunction == nullptr)
-	{
-		return nullptr;
-	}
-
-	// @todo: use a more destriptive name for event entry points?
-	static FString DefaultEventEntryName = TEXT("SequenceEvent__ENTRYPOINT");
-	UEdGraph* EntryPointGraph = Compiler->SpawnIntermediateFunctionGraph(DefaultEventEntryName + Blueprint->GetName());
-	check(EntryPointGraph->Nodes.Num() == 1);
-
-	const UEdGraphSchema*  Schema        = EntryPointGraph->GetSchema();
-	UK2Node_FunctionEntry* FunctionEntry = CastChecked<UK2Node_FunctionEntry>(EntryPointGraph->Nodes[0]);
-
-	// -------------------------------------------------------------------------------------
-	// Locate and initialize the function entry node
+	if (!EntrypointDefinition->BoundObjectPinName.IsNone())
 	{
-		int32 ExtraFunctionFlags = ( FUNC_BlueprintEvent | FUNC_Public );
-		FunctionEntry->AddExtraFlags(ExtraFunctionFlags);
-		FunctionEntry->bIsEditable = false;
-		FunctionEntry->MetaData.Category = LOCTEXT("DefaultCategory", "Sequencer Event Endpoints");
-		FunctionEntry->MetaData.bCallInEditor = EndpointFunction->GetBoolMetaData(FBlueprintMetadata::MD_CallInEditor);
+		EndpointCall.ExposedPinNames.Add(EntrypointDefinition->BoundObjectPinName);
 	}
-
-	// -------------------------------------------------------------------------------------
-	// Create a function call node to invoke the endpoint function itself
-	UK2Node_CallFunction* CallFunctionNode = NewObject<UK2Node_CallFunction>(EntryPointGraph);
-	{
-		CallFunctionNode->CreateNewGuid();
-		CallFunctionNode->PostPlacedNewNode();
-
-		CallFunctionNode->FunctionReference.SetSelfMember(EndpointFunction->GetFName());
-
-		CallFunctionNode->ReconstructNode();
-
-		CallFunctionNode->NodePosX = FunctionEntry->NodePosX + 400.f;
-		CallFunctionNode->NodePosY = FunctionEntry->NodePosY - 16.f;
-
-		EntryPointGraph->AddNode(CallFunctionNode, false, false);
-	}
-
-	// -------------------------------------------------------------------------------------
-	// Create a pin for the bound object if possible
-	if (EntrypointDefinition->BoundObjectPinName != NAME_None)
-	{
-		UEdGraphPin* BoundObjectPin = CallFunctionNode->FindPin(EntrypointDefinition->BoundObjectPinName, EGPD_Input);
-		if (BoundObjectPin)
-		{
-			UEdGraphPin* BoundObjectPinInput = FunctionEntry->CreateUserDefinedPin(BoundObjectPin->PinName, BoundObjectPin->PinType, EGPD_Output, true);
-			if (BoundObjectPinInput)
-			{
-				BoundObjectPinInput->MakeLinkTo(BoundObjectPin);
-			}
-		}
-	}
-
-	// -------------------------------------------------------------------------------------
-	// Wire up the function entry 'then' pin to the call function 'execute' pin
-	{
-		// Wire up the exec nodes
-		UEdGraphPin* ThenPin = FunctionEntry->FindPin(UEdGraphSchema_K2::PN_Then, EGPD_Output);
-		UEdGraphPin* ExecPin = CallFunctionNode->FindPin(UEdGraphSchema_K2::PN_Execute, EGPD_Input);
-
-		if (ensure(ThenPin && ExecPin))
-		{
-			ThenPin->MakeLinkTo(ExecPin);
-		}
-	}
-
-	// -------------------------------------------------------------------------------------
-	// Set pin defaults for each of the relevant pins on the call function node according to the payload
-	{
-		TArray<FName, TInlineAllocator<8>> ValidPinNames;
-
-		// Gather all input pins for the function call
-		for (UEdGraphPin* Pin : CallFunctionNode->Pins)
-		{
-			// Only consider input pins that are data, and not already connected
-			if (Pin->Direction != EGPD_Input || Pin->PinName == UEdGraphSchema_K2::PN_Execute || Pin->PinName == UEdGraphSchema_K2::PN_Self || Pin->LinkedTo.Num() != 0)
-			{
-				continue;
-			}
-
-			const FMovieSceneEventPayloadVariable* PayloadVariable = EntrypointDefinition->PayloadVariables.Find(Pin->PinName);
-			if (!PayloadVariable)
-			{
-				continue;
-			}
-
-			UClass* PinObjectType = Cast<UClass>(Pin->PinType.PinSubCategoryObject.Get());
-			const bool bIsRawActorPin = Pin->PinType.PinCategory == UEdGraphSchema_K2::PC_Object && PinObjectType && PinObjectType->IsChildOf(AActor::StaticClass());
-
-			if (bIsRawActorPin)
-			{
-				// Raw actor properties are represented as soft object ptrs under the hood so that we can still serialize them
-				// To make this work, we have to make the following graph:
-				// MakeSoftObjectPath('<ActorPath>') -> Conv_SoftObjPathToSoftObjRef() -> Conv_SoftObjectReferenceToObject() -> Cast<ActorType>()
-				UK2Node_CallFunction* MakeSoftObjectPathNode = Compiler->SpawnIntermediateNode<UK2Node_CallFunction>(CallFunctionNode, EntryPointGraph);
-				MakeSoftObjectPathNode->FunctionReference.SetExternalMember(GET_FUNCTION_NAME_CHECKED(UKismetSystemLibrary, MakeSoftObjectPath), UKismetSystemLibrary::StaticClass());
-				MakeSoftObjectPathNode->AllocateDefaultPins();
-
-				UK2Node_CallFunction* ConvertToSoftObjectRef = Compiler->SpawnIntermediateNode<UK2Node_CallFunction>(CallFunctionNode, EntryPointGraph);
-				ConvertToSoftObjectRef->FunctionReference.SetExternalMember(GET_FUNCTION_NAME_CHECKED(UKismetSystemLibrary, Conv_SoftObjPathToSoftObjRef), UKismetSystemLibrary::StaticClass());
-				ConvertToSoftObjectRef->AllocateDefaultPins();
-
-				UK2Node_CallFunction* ConvertToObjectFunc = Compiler->SpawnIntermediateNode<UK2Node_CallFunction>(CallFunctionNode, EntryPointGraph);
-				ConvertToObjectFunc->FunctionReference.SetExternalMember(GET_FUNCTION_NAME_CHECKED(UKismetSystemLibrary, Conv_SoftObjectReferenceToObject), UKismetSystemLibrary::StaticClass());
-				ConvertToObjectFunc->AllocateDefaultPins();
-
-				UK2Node_DynamicCast* CastNode = Compiler->SpawnIntermediateNode<UK2Node_DynamicCast>(CallFunctionNode, EntryPointGraph);
-				CastNode->SetPurity(true);
-				CastNode->TargetType = PinObjectType;
-				CastNode->AllocateDefaultPins();
-
-				UEdGraphPin* PathInput     = MakeSoftObjectPathNode->FindPin(FName("PathString"));
-				UEdGraphPin* PathOutput    = MakeSoftObjectPathNode->GetReturnValuePin();
-				UEdGraphPin* SoftRefInput  = ConvertToSoftObjectRef->FindPin(FName("SoftObjectPath"));
-				UEdGraphPin* SoftRefOutput = ConvertToSoftObjectRef->GetReturnValuePin();
-				UEdGraphPin* ConvertInput  = ConvertToObjectFunc->FindPin(FName("SoftObject"));
-				UEdGraphPin* ConvertOutput = ConvertToObjectFunc->GetReturnValuePin();
-				UEdGraphPin* CastInput     = CastNode->GetCastSourcePin();
-				UEdGraphPin* CastOutput    = CastNode->GetCastResultPin();
-
-				if (!PathInput || !PathOutput || !SoftRefInput || !SoftRefOutput || !ConvertInput || !ConvertOutput || !CastInput || !CastOutput)
-				{
-					Compiler->MessageLog.Error(*LOCTEXT("ActorFacadeError", "GenerateEntryPoint: Failed to generate soft-ptr facade for AActor payload property property @@. @@").ToString(), *Pin->PinName.ToString(), Endpoint);
-					continue;
-				}
-
-				// Set the default value for the path string
-				const bool bMarkAsModified = false;
-				Schema->TrySetDefaultValue(*PathInput, PayloadVariable->Value, bMarkAsModified);
-
-				bool bSuccess = true;
-				bSuccess &= Schema->TryCreateConnection(PathOutput, SoftRefInput);
-				bSuccess &= Schema->TryCreateConnection(SoftRefOutput, ConvertInput);
-				bSuccess &= Schema->TryCreateConnection(ConvertOutput, CastInput);
-				bSuccess &= Schema->TryCreateConnection(CastOutput, Pin);
-
-				if (!bSuccess)
-				{
-					Compiler->MessageLog.Error(*LOCTEXT("ActorFacadeConnectionError", "GenerateEntryPoint: Failed to connect nodes for soft-ptr facade in AActor payload property @@. @@").ToString(), *Pin->PinName.ToString(), Endpoint);
-				}
-
-				ValidPinNames.Add(Pin->PinName);
-			}
-			else
-			{
-				bool bMarkAsModified = false;
-				Schema->TrySetDefaultValue(*Pin, PayloadVariable->Value, bMarkAsModified);
-
-				ValidPinNames.Add(Pin->PinName);
-			}
-		}
-
-		// Remove any invalid payload variable names
-		if (ValidPinNames.Num() != EntrypointDefinition->PayloadVariables.Num())
-		{
-			for (auto It = EntrypointDefinition->PayloadVariables.CreateIterator(); It; ++It)
-			{
-				if (!ValidPinNames.Contains(It->Key))
-				{
-					Compiler->MessageLog.Note(*FText::Format(LOCTEXT("PayloadParameterRemoved", "Stale Sequencer event payload parameter %s has been removed."), FText::FromName(It->Key)).ToString());
-					It.RemoveCurrent();
-				}
-			}
-		}
-	}
-
-	return FunctionEntry;
+	EndpointCall.Endpoint = Endpoint;
+	FMovieSceneDirectorBlueprintEntrypointResult Result = FMovieSceneDirectorBlueprintUtils::GenerateEntryPoint(EndpointCall, Compiler);
+	Result.CleanUpStalePayloadVariables(EntrypointDefinition->PayloadVariables);
+	return Result.Entrypoint;
 }
 
 void FMovieSceneEventUtils::BindEventSectionToBlueprint(UMovieSceneEventSectionBase* EventSection, UBlueprint* DirectorBP)
@@ -598,4 +357,4 @@ void FMovieSceneEventUtils::RemoveUnusedCustomEvents(const TArray<TWeakObjectPtr
 	}
 }
 
-#undef LOCTEXT_NAMESPACE
\ No newline at end of file
+#undef LOCTEXT_NAMESPACE
diff --git a/Engine/Source/Editor/MovieSceneTools/Public/MovieSceneDirectorBlueprintEndpointCustomization.h b/Engine/Source/Editor/MovieSceneTools/Public/MovieSceneDirectorBlueprintEndpointCustomization.h
new file mode 100644
--- /dev/null
+++ b/Engine/Source/Editor/MovieSceneTools/Public/MovieSceneDirectorBlueprintEndpointCustomization.h
@@ -0,0 +1,302 @@
+// Copyright Epic Games, Inc. All Rights Reserved.
+
+#pragma once
+
+#include "CoreMinimal.h"
+#include "Containers/ContainersFwd.h"
+#include "IPropertyTypeCustomization.h"
+#include "MovieSceneDirectorBlueprintUtils.h"
+#include "Types/SlateEnums.h"
+
+class UK2Node;
+class UBlueprint;
+class UEdGraphPin;
+class UEdGraphNode;
+class FMenuBuilder;
+class FStructOnScope;
+class IPropertyHandle;
+class FDetailWidgetRow;
+class SWidget;
+class UMovieSceneSection;
+class UMovieSceneSequence;
+class UK2Node_CustomEvent;
+class UK2Node_FunctionEntry;
+class IDetailChildrenBuilder;
+struct FAssetData;
+struct FBlueprintActionMenuBuilder;
+struct FEdGraphSchemaAction;
+struct FGraphActionListBuilderBase;
+
+enum class ECheckBoxState : uint8;
+
+enum class EAutoCreatePayload : uint8
+{
+	None      = 0,
+	Pins      = 1 << 0,
+	Variables = 1 << 1,
+};
+ENUM_CLASS_FLAGS(EAutoCreatePayload)
+
+/**
+ * Base class for details view customizations that operate on sequence director blueprint endpoints.
+ */
+class MOVIESCENETOOLS_API FMovieSceneDirectorBlueprintEndpointCustomization : public IPropertyTypeCustomization
+{
+public:
+
+	// IPropertyTypeCustomization interface
+	virtual void CustomizeHeader(TSharedRef<IPropertyHandle> InPropertyHandle, FDetailWidgetRow& HeaderRow, IPropertyTypeCustomizationUtils& CustomizationUtils) override;
+	virtual void CustomizeChildren(TSharedRef<IPropertyHandle> InPropertyHandle, IDetailChildrenBuilder& ChildBuilder, IPropertyTypeCustomizationUtils& CustomizationUtils) override;
+
+protected:
+
+	using FPayloadVariableMap = TMap<FName, FString, TInlineSetAllocator<8>>;
+
+	/**
+	 * User-interface information for "well-known parameters", i.e. parameters that can be bound
+	 * to values provided by the underlying system, as opposed to values entered by the user.
+	 */
+	struct FWellKnownParameterMetadata
+	{
+		FText PickerLabel;
+		FText PickerTooltip;
+		FText CandidateTooltip;
+	};
+
+	/**
+	 * Structure specifying candidate pins for a given "well-known parameter".
+	 */
+	struct FWellKnownParameterCandidates
+	{
+		/** Names of the pins that can be bound to the "well-known parameter" described in the metadata */
+		TArray<FName> CandidatePinNames;
+		/** User-interface information for the "well-known parameter" */
+		FWellKnownParameterMetadata Metadata;
+		/** Whether to show this "well-known" parameter when no candidate pins have been found */
+		bool bShowUnmatchedParameters = false;
+	};
+
+	/**
+	 * Gets the payload values to pass to the endpoint. 
+	 *
+	 * @param EditObject		The edited object for which to get the endpoint call payload variables
+	 * @param RawData			The raw data for which to get the endpoint call payload variables
+	 * @param OutPayloadVariables	The string values of the payload variables to pass to the endpoint
+	 */
+	virtual void GetPayloadVariables(UObject* EditObject, void* RawData, FPayloadVariableMap& OutPayloadVariables) const = 0;
+
+	/**
+	 * Sets a payload value to pass to the endpoint.
+	 *
+	 * @param EditObject		The edited object for which to set the endpoint call payload variable
+	 * @param RawData			The raw data for which to set the endpoint call payload variable
+	 * @param FieldName			The name of the payload variable to set
+	 * @param NewVariableValue	The string value of the payload variable to set
+	 * @return					Whether the payload variable was set
+	 */
+	virtual bool SetPayloadVariable(UObject* EditObject, void* RawData, FName FieldName, const FString& NewVariableValue) = 0;
+
+	/**
+	 * Get the pin names setup with "well-known parameters" for the currently edited objects.
+	 *
+	 * The output array of pin names should be of the same size as GetWellKnownParameterCandidates, with NAME_None
+	 * for those parameters that aren't hooked up to any pin.
+	 *
+	 * @param EditObject		The edited object on which to get the connected pin names
+	 * @param RawData			The raw data on which to get the connected pin names
+	 * @param OutParameters		The list of pin names connected to "well-known parameters"
+	 */
+	virtual void GetWellKnownParameterPinNames(UObject* EditObject, void* RawData, TArray<FName>& OutParameters) const {}
+
+	/**
+	 * Get the pins that are candidates for any "well-known parameter".
+	 *
+	 * The output array of pin candidates should be of the same size as GetWellKnownParameterPinNames, with an empty
+	 * array of candidates for those parameters that currently have no compatible pins.
+	 *
+	 * @param EditObject		The edited object on which to set the connected pin name
+	 * @param RawData			The raw data on which to set the connected pin name
+	 */
+	virtual void GetWellKnownParameterCandidates(UK2Node* Endpoint, TArray<FWellKnownParameterCandidates>& OutCandidates) const {}
+
+	/**
+	 * Set the pin name for the given "well-known parameter".
+	 *
+	 * @param EditObject		The edited object on which to set the connected pin name
+	 * @param RawData			The raw data on which to set the connected pin name
+	 * @param ParameterIndex	The index of the "well-known parameter", as reported by GetWellKnownParameterCandidates
+	 * @return					Whether the pin name was successfully set
+	 */
+	virtual bool SetWellKnownParameterPinName(UObject* EditObject, void* RawData, int32 ParameterIndex, FName BoundPinName) { return false; }
+
+	/**
+	 * Generate endpoint definition for the sequence.
+	 */
+	virtual FMovieSceneDirectorBlueprintEndpointDefinition GenerateEndpointDefinition(UMovieSceneSequence* Sequence) = 0;
+
+	/**
+	 * Called to find the endpoint node in the director blueprint.
+	 */
+	virtual UK2Node* FindEndpoint(UMovieSceneSequence* Sequence, UBlueprint* Blueprint, UObject* EditObject, void* RawData) const = 0;
+
+	/**
+	 * Called when an endpoint has been created.
+	 */
+	virtual void OnCreateEndpoint(UMovieSceneSequence* Sequence, UBlueprint* Blueprint, const TArray<UObject*> EditObjects, const TArray<void*> RawData, const FMovieSceneDirectorBlueprintEndpointDefinition& EndpointDefinition, UK2Node* NewEndpoint) = 0;
+
+	/**
+	 * Called when an endpoint has been set.
+	 */
+	virtual void OnSetEndpoint(UMovieSceneSequence* Sequence, UBlueprint* Blueprint, const TArray<UObject*> EditObjects, const TArray<void*> RawData, const FMovieSceneDirectorBlueprintEndpointDefinition& EndpointDefinition, UK2Node* NewEndpoint) = 0;
+
+	/**
+	 * Called when an endpoint was set and has been changed to another endpoint.
+	 */
+	virtual bool OnRebindEndpoint(UMovieSceneSequence* Sequence, UBlueprint* Blueprint, const TArray<UObject*> EditObjects, const TArray<void*> RawData, const FMovieSceneDirectorBlueprintEndpointDefinition& EndpointDefinition, TSharedPtr<FEdGraphSchemaAction> Action) { return false; }
+
+	/**
+	 * Get all the objects that the endpoints reside within.
+	 * By default, this gets the outer objects of the property handle.
+     */
+	virtual void GetEditObjects(TArray<UObject*>& OutObjects) const;
+
+	/**
+	 * Collect blueprint actions for quick-binding a non-connected endpoint.
+	 */
+	virtual void OnCollectQuickBindActions(UBlueprint* Blueprint, FBlueprintActionMenuBuilder& MenuBuilder) {}
+
+	/**
+	 * Collect blueprint actions for rebinding an already connected endpoint.
+	 */
+	virtual void OnCollectAllRebindActions(UBlueprint* Blueprint, FBlueprintActionMenuBuilder& MenuBuilder) {}
+
+protected:
+
+	/**
+	 * Clear the endpoint represented by this property handle.
+	 * @note: Does not delete the endpoint in the blueprint itself.
+	 */
+	void ClearEndpoint();
+
+	/**
+	 * Creates a single new endpoint represented by this property handle.
+	 */
+	void CreateEndpoint();
+
+	/**
+	 * Find the endpoint node in the director blueprint.
+	 */
+	UK2Node* FindEndpoint(UObject* EditObject, void* RawData) const;
+
+	/**
+	 * Assigns the specified function entry to this property handle.
+	 */
+	void SetEndpoint(const FMovieSceneDirectorBlueprintEndpointDefinition& EndpointDefinition, UK2Node* NewEndpoint, UK2Node* PayloadTemplate, EAutoCreatePayload AutoCreatePayload);
+
+	/**
+	 * Navigate to the definition of the endpoint represented by this property handle.
+	 */
+	void NavigateToDefinition();
+
+	/**
+	 * Generate the content of the main combo button menu dropdown
+	 */
+	TSharedRef<SWidget> GetMenuContent();
+
+	/**
+	 * Generate the content of the quick bind sub-menu dropdown (shown if the endpoint is not already bound)
+	 */
+	void PopulateQuickBindSubMenu(FMenuBuilder& MenuBuilder, UMovieSceneSequence* Sequence);
+
+	/**
+	 * Generate the content of the rebind sub-menu dropdown (shown if the endpoint is already bound)
+	 */
+	void PopulateRebindSubMenu(FMenuBuilder& MenuBuilder, UMovieSceneSequence* Sequence);
+
+	/**
+	 * Called when a payload variable property value is changed on the details panel
+	 */
+	void OnPayloadVariableChanged(TSharedRef<FStructOnScope> InStructData, TSharedPtr<IPropertyHandle> LocalVariableProperty);
+
+protected:
+
+	struct FSequenceData
+	{
+		UBlueprint* Blueprint;
+		TArray<UObject*> EditObjects;
+		TArray<void*> RawData;
+	};
+	using FSequenceDataMap = TMap<UMovieSceneSequence*, FSequenceData>;
+
+	/**
+	 * Gather data for the sequences being edited.
+	 */
+	void GatherSequenceData(FSequenceDataMap& AllSequenceData);
+
+	/**
+	 * Get the property handle for this customization.
+	 */
+	TSharedPtr<IPropertyHandle> GetPropertyHandle() const { return PropertyHandle; }
+
+	/**
+	 * Get the property utilities for this customization.
+	 */
+	TSharedPtr<IPropertyUtilities> GetPropertyUtilities() const { return PropertyUtilities; }
+
+	/**
+	 * Get the sequence that is common to all the values in our property handle, or nullptr if they are not all the same.
+	 */
+	UMovieSceneSequence* GetCommonSequence() const;
+
+	/**
+	 * Get the endpoint that is common to all the values in our property handle, or nullptr if they are not all the same.
+	 */
+	UK2Node* GetCommonEndpoint() const;
+
+	/**
+	 * Get all the endpoints of all the values in our property handle.
+	 */
+	TArray<UK2Node*> GetAllValidEndpoints() const;
+
+	/**
+	 * Invoke a callback on all the endpoints.
+	 */
+	void IterateEndpoints(TFunctionRef<bool(UK2Node*)> Callback) const;
+
+private:
+
+	void CollectQuickBindActions(FGraphActionListBuilderBase& OutAllActions, UBlueprint* Blueprint, FMovieSceneDirectorBlueprintEndpointDefinition EndpointDefinition);
+	void CollectAllRebindActions(FGraphActionListBuilderBase& OutAllActions, UBlueprint* Blueprint, FMovieSceneDirectorBlueprintEndpointDefinition EndpointDefinition);
+
+private:
+
+	/** Get the name of the endpoint to display on the main combo button */
+	FText GetEndpointName() const;
+	/** Get the icon of the endpoint to display on the main combo button */
+	const FSlateBrush* GetEndpointIcon() const;
+
+	void HandleRebindActionSelected(const TArray<TSharedPtr<FEdGraphSchemaAction>>& SelectedAction, ESelectInfo::Type InSelectionType, UBlueprint* Blueprint, FMovieSceneDirectorBlueprintEndpointDefinition EndpointDefinition);
+	void HandleQuickBindActionSelected(const TArray<TSharedPtr<FEdGraphSchemaAction>>& SelectedAction, ESelectInfo::Type InSelectionType, UBlueprint* Blueprint, FMovieSceneDirectorBlueprintEndpointDefinition EndpointDefinition);
+
+	void OnBlueprintCompiled(UBlueprint*);
+
+	ECheckBoxState GetCallInEditorCheckState() const;
+	void OnSetCallInEditorCheckState(ECheckBoxState NewCheckBoxState);
+
+	TSharedRef<SWidget> GetWellKnownParameterPinMenuContent(int32 ParameterIndex);
+	const FSlateBrush* GetWellKnownParameterPinIcon(int32 ParameterIndex) const;
+	FText GetWellKnownParameterPinText(int32 ParameterIndex) const;
+	bool CompareWellKnownParameterPinName(int32 ParameterIndex, FName InPinName) const;
+	void SetWellKnownParameterPinName(int32 ParameterIndex, FName InNewWellKnownParameterPinName);
+
+private:
+
+	/** A cache of the common endpoint that is only used when the menu is open to avoid re-computing it every frame. */
+	TWeakObjectPtr<UK2Node_FunctionEntry> CachedCommonEndpoint;
+
+	/** The property handle we're reflecting */
+	TSharedPtr<IPropertyHandle> PropertyHandle;
+
+	/** Property utilities for the property we're editing */
+	TSharedPtr<IPropertyUtilities> PropertyUtilities;
+};
diff --git a/Engine/Source/Editor/MovieSceneTools/Public/MovieSceneDirectorBlueprintUtils.h b/Engine/Source/Editor/MovieSceneTools/Public/MovieSceneDirectorBlueprintUtils.h
new file mode 100644
--- /dev/null
+++ b/Engine/Source/Editor/MovieSceneTools/Public/MovieSceneDirectorBlueprintUtils.h
@@ -0,0 +1,215 @@
+// Copyright Epic Games, Inc. All Rights Reserved.
+
+#pragma once
+
+#include "Containers/Array.h"
+#include "Containers/UnrealString.h"
+#include "CoreTypes.h"
+#include "Delegates/Delegate.h"
+#include "Delegates/DelegateCombinations.h"
+#include "EdGraph/EdGraphNode.h"
+#include "UObject/NameTypes.h"
+#include "UObject/WeakObjectPtr.h"
+#include "UObject/WeakObjectPtrTemplates.h"
+
+class FKismetCompilerContext;
+class UBlueprint;
+class UClass;
+class UEdGraph;
+class UEdGraphNode;
+class UEdGraphPin;
+class UK2Node;
+class UK2Node_CallFunction;
+class UK2Node_CustomEvent;
+class UK2Node_FunctionEntry;
+class UK2Node_FunctionResult;
+class UMovieScene;
+class UMovieSceneSequence;
+class UMovieSceneTrack;
+struct FGuid;
+
+/**
+ * Sequence Director Blueprint endpoint type
+ */
+enum class EMovieSceneDirectorBlueprintEndpointType
+{
+	None = 0,
+	Event,
+	Function
+};
+
+/**
+ * Structure for creating extra parameter pins on an enpdoint
+ */
+struct MOVIESCENETOOLS_API FMovieSceneDirectorBlueprintEndpointParameter
+{
+	FString PinName;
+	EEdGraphPinDirection PinDirection;
+	FName PinTypeCategory;
+	TWeakObjectPtr<UObject> PinTypeClass;
+
+	static FString SanitizePinName(const FString& InPinName);
+};
+
+/**
+ * Parameters for creating a blueprint endpoint
+ */
+struct MOVIESCENETOOLS_API FMovieSceneDirectorBlueprintEndpointDefinition
+{
+	/** Type of endpoint */
+	EMovieSceneDirectorBlueprintEndpointType EndpointType;
+
+	/** Name of the event or function graph */
+	FString GraphName;
+
+	/** Name of the event or function */
+	FString EndpointName;
+
+	/** Function signature, for function endpoint types */
+	UFunction* EndpointSignature = nullptr;
+
+	/** Optional class on which the endpoint can call functions */
+	UClass* PossibleCallTargetClass = nullptr;
+
+	/** Extra pins to create on the endpoint */
+	TArray<FMovieSceneDirectorBlueprintEndpointParameter> ExtraPins;
+
+	FMovieSceneDirectorBlueprintEndpointDefinition()
+		: EndpointType(EMovieSceneDirectorBlueprintEndpointType::None)
+	{}
+
+	void AddExtraInputPin(const FString& InName, FName InPinCategory, TWeakObjectPtr<UObject> InPinTypeClass = nullptr)
+	{
+		FMovieSceneDirectorBlueprintEndpointParameter Parameter;
+		Parameter.PinName = InName;
+		Parameter.PinDirection = EGPD_Input;
+		Parameter.PinTypeCategory = InPinCategory;
+		Parameter.PinTypeClass = InPinTypeClass;
+		ExtraPins.Add(Parameter);
+	}
+
+	void AddExtraOutputPin(const FString& InName, FName InPinCategory, TWeakObjectPtr<UObject> InPinTypeClass = nullptr)
+	{
+		FMovieSceneDirectorBlueprintEndpointParameter Parameter;
+		Parameter.PinName = InName;
+		Parameter.PinDirection = EGPD_Output;
+		Parameter.PinTypeCategory = InPinCategory;
+		Parameter.PinTypeClass = InPinTypeClass;
+		ExtraPins.Add(Parameter);
+	}
+};
+
+/** Parameters for endpoint call customization callback */
+struct MOVIESCENETOOLS_API FMovieSceneCustomizeDirectorBlueprintEndpointCallParams
+{
+	UEdGraph* FunctionGraph = nullptr;
+	UK2Node_FunctionEntry* FunctionEntryNode = nullptr;
+	UK2Node_CallFunction* CallFunctionNode = nullptr;
+	UK2Node_FunctionResult* FunctionResultNode = nullptr;
+};
+
+DECLARE_DELEGATE_OneParam(FMovieSceneCustomizeDirectorBlueprintEndpointCallDelegate, const FMovieSceneCustomizeDirectorBlueprintEndpointCallParams&);
+
+/**
+ * Parameters for creating an entry point, i.e. a call to an endpoint
+ */
+struct MOVIESCENETOOLS_API FMovieSceneDirectorBlueprintEndpointCall
+{
+	/** The endpoint node to call */
+	UEdGraphNode* Endpoint = nullptr;
+
+	/** List of endpoint node pins that should be exposed to the call, for receiving special values at runtime */
+	TArray<FName> ExposedPinNames;
+
+	/** The values to pass to the endpoint node's parameters */
+	TMap<FName, FString> PayloadVariables;
+
+	/** Callback for customizing the endpoint call */
+	FMovieSceneCustomizeDirectorBlueprintEndpointCallDelegate CustomizeEndpointCall;
+};
+
+/**
+ * Result structure for a created blueprint entry point, i.e. a call to an endpoint
+ */
+struct MOVIESCENETOOLS_API FMovieSceneDirectorBlueprintEntrypointResult
+{
+	/** The name of the function that calls the endpoint */
+	FName CompiledFunctionName;
+	/** The entry point node */
+	UK2Node_FunctionEntry* Entrypoint = nullptr;
+	/** A list of payload variables that aren't valid anymore */
+	TArray<FName> StalePayloadVariables;
+
+	/** A utility method for cleaning up stale payload variables */
+	template<typename Container>
+	void CleanUpStalePayloadVariables(Container& InPayloadVariables) const
+	{
+		for (FName StalePayloadVariable : StalePayloadVariables)
+		{
+			InPayloadVariables.Remove(StalePayloadVariable);
+		}
+	}
+};
+
+/**
+ * Static utility library for dealing with movie-scene blueprint endpoints at edit/cook-time
+ */
+struct MOVIESCENETOOLS_API FMovieSceneDirectorBlueprintUtils
+{
+	/**
+	 * Create a new endpoint for the specified sequence.
+	 *
+	 * @param Blueprint       (Required, non-null) The blueprint within which to create the new custom event
+	 * @param EndpointDefinition  The end point definition
+	 * @return A valid pointer to the blueprint node for the endpoint
+	 */
+	static UK2Node* CreateEndpoint(UBlueprint* Blueprint, const FMovieSceneDirectorBlueprintEndpointDefinition& EndpointDefinition);
+
+	/**
+	 * Create a new event endpoint for the specified sequence.
+	 *
+	 * By default, event endpoints are presented as a UK2Node_CustomEvent in an event graph called "Sequencer Events" with a pin for the bound object.
+	 *
+	 * @param Blueprint       (Required, non-null) The blueprint within which to create the new custom event
+	 * @param EndpointDefinition  The end point definition
+	 * @return A valid pointer to the blueprint node for the endpoint
+	 */
+	static UK2Node_CustomEvent* CreateEventEndpoint(UBlueprint* Blueprint, const FMovieSceneDirectorBlueprintEndpointDefinition& EndpointDefinition);
+	
+	/**
+	 * Create a new function endpoint for the specified sequence.
+	 *
+	 * @param Blueprint       (Required, non-null) The blueprint within which to create the new custom event
+	 * @param EndpointDefinition  The end point definition
+	 * @return A valid pointer to the blueprint node for the endpoint
+	 */
+	static UK2Node_FunctionEntry* CreateFunctionEndpoint(UBlueprint* Blueprint, const FMovieSceneDirectorBlueprintEndpointDefinition& EndpointDefinition);
+
+	/**
+	 * Main generation function for creating the necessary function graph for calling an event endpoint with 
+	 * the parameters specified in the entry point payload
+	 *
+	 * @param EntryPoint             (Required, non-null) The event entry point definition to generate a function graph for
+	 * @param Compiler               (Required, non-null) The kismet compiler context for generation
+	 * @param Endpoint               (Required, non-null) The endpoint node that the event is bound to
+	 */
+	static FMovieSceneDirectorBlueprintEntrypointResult GenerateEntryPoint(const FMovieSceneDirectorBlueprintEndpointCall& EndpointCall, FKismetCompilerContext* Compiler);
+
+	/**
+	 * Attempt to locate an output pin on the specified endpoint node that matches the specified pin class
+	 * 
+	 * @param InEndpoint             (Required, non-null) The endpoint node to search on
+	 * @param BoundObjectPinClass    (Required, non-null) The class of the object pin to locate
+	 * @return An output pin that is of type PC_Object with the specified class as its PinSubCategoryObject, or nullptr if one was not found
+	 */
+	static UEdGraphPin* FindCallTargetPin(UK2Node* InEndpoint, UClass* CallTargetClass);
+
+private:
+
+	static UEdGraph* FindOrCreateEventGraph(UBlueprint* Blueprint, const FString& GraphName);
+
+	static bool GenerateEntryPointRawActorParameter(
+			FKismetCompilerContext* Compiler, UEdGraph* Graph, UEdGraphNode* Endpoint, 
+			UK2Node* OriginNode, UEdGraphPin* DestPin, const FString& PayloadValue);
+};
+
diff --git a/Engine/Source/Editor/MovieSceneTools/Public/MovieSceneEventCustomization.h b/Engine/Source/Editor/MovieSceneTools/Public/MovieSceneEventCustomization.h
--- a/Engine/Source/Editor/MovieSceneTools/Public/MovieSceneEventCustomization.h
+++ b/Engine/Source/Editor/MovieSceneTools/Public/MovieSceneEventCustomization.h
@@ -4,162 +4,36 @@
 
 #include "CoreMinimal.h"
 #include "IPropertyTypeCustomization.h"
+#include "MovieSceneDirectorBlueprintEndpointCustomization.h"
 #include "Types/SlateEnums.h"
 
-struct FAssetData;
-struct FEdGraphSchemaAction;
-struct FMovieSceneEventEndpointParameters;
-
-class UK2Node;
-class UBlueprint;
-class UEdGraphPin;
-class UEdGraphNode;
-class FMenuBuilder;
-class FStructOnScope;
-class IPropertyHandle;
-class FDetailWidgetRow;
-class SWidget;
-class UMovieSceneSection;
-class UMovieSceneSequence;
 class UMovieSceneEventTrack;
-class UK2Node_FunctionEntry;
-class IDetailChildrenBuilder;
-
-enum class ECheckBoxState : uint8;
-
-template <typename InKeyType, typename InValueType, typename ArrayAllocator, typename SortPredicate> class TSortedMap;
-
-enum class EAutoCreatePayload : uint8
-{
-	None      = 0,
-	Pins      = 1 << 0,
-	Variables = 1 << 1,
-};
-ENUM_CLASS_FLAGS(EAutoCreatePayload)
 
 /**
  * Customization for FMovieSceneEvent structs.
  * Will deduce the event's section either from the outer objects on the details customization, or use the one provided on construction (for instanced property type customizations)
  */
-class FMovieSceneEventCustomization : public IPropertyTypeCustomization
+class FMovieSceneEventCustomization : public FMovieSceneDirectorBlueprintEndpointCustomization
 {
 public:
 	static TSharedRef<IPropertyTypeCustomization> MakeInstance();
 	static TSharedRef<IPropertyTypeCustomization> MakeInstance(UMovieSceneSection* InSection);
 
-	virtual void CustomizeHeader( TSharedRef<IPropertyHandle> PropertyHandle, FDetailWidgetRow& HeaderRow, IPropertyTypeCustomizationUtils& CustomizationUtils ) override;
-	virtual void CustomizeChildren( TSharedRef<IPropertyHandle> PropertyHandle, IDetailChildrenBuilder& ChildBuilder, IPropertyTypeCustomizationUtils& CustomizationUtils ) override;
-
-private:
-
-	/**
-	 * Clear the endpoint for the event(s) represented by this property handle.
-	 * @note: Does not delete the endpoint in the blueprint itself.
-	 */
-	void ClearEventEndpoint();
-
-
-	/**
-	 * Creates a single new endpoint for the event(s) represented by this property handle.
-	 */
-	void CreateEventEndpoint();
+protected:
 
-
-	/**
-	 * Assigns the specified function entry to the event(s) represented by this property handle.
-	 */
-	void SetEventEndpoint(UK2Node* NewEndpoint, UEdGraphPin* BoundObjectPin, UK2Node* PayloadTemplate, EAutoCreatePayload AutoCreatePayload);
-
-
-	/**
-	 * Navigate to the definition of the endpoint specified by the event(s) represented by this property handle.
-	 */
-	void NavigateToDefinition();
-
-
-	/**
-	 * Generate the content of the main combo button menu dropdown
-	 */
-	TSharedRef<SWidget> GetMenuContent();
-
-
-	/**
-	 * Generate the content of the quick bind sub-menu dropdown (shown if the event is not already bound)
-	 */
-	void PopulateQuickBindSubMenu(FMenuBuilder& MenuBuilder, UMovieSceneSequence* Sequence);
-
-	/**
-	 * Generate the content of the rebind sub-menu dropdown (shown if the event is already bound)
-	 */
-	void PopulateRebindSubMenu(FMenuBuilder& MenuBuilder, UMovieSceneSequence* Sequence);
-
-	/**
-	 * Called when a payload variable property value is changed on the details panel
-	 */
-	void OnPayloadVariableChanged(TSharedRef<FStructOnScope> InStructData, TSharedPtr<IPropertyHandle> LocalVariableProperty);
-
-private:
-
-	/**
-	 * Get the sequence that is common to all the events represented by this property handle, or nullptr if they are not all the same.
-	 */
-	UMovieSceneSequence* GetCommonSequence() const;
-
-
-	/**
-	 * Get the track that is common to all the events represented by this property handle, or nullptr if they are not all the same.
-	 */
-	UMovieSceneEventTrack* GetCommonTrack() const;
-
-
-	/**
-	 * Get the endpoint that is common to all the events represented by this property handle, or nullptr if they are not all the same.
-	 */
-	UK2Node* GetCommonEndpoint() const;
-
-	TArray<UK2Node*> GetAllValidEndpoints() const;
-
-	void IterateEndpoints(TFunctionRef<bool(UK2Node*)> Callback) const;
-
-	/**
-	 * Get all the objects that the events reside within.
-	 */
-	void GetEditObjects(TArray<UObject*>& OutObjects) const;
-
-private:
-
-	/** Get the name of the event to display on the main combo button */
-	FText GetEventName() const;
-	/** Get the icon of the event to display on the main combo button */
-	const FSlateBrush* GetEventIcon() const;
-
-	void HandleRebindActionSelected(const TArray<TSharedPtr<FEdGraphSchemaAction>>& SelectedAction, ESelectInfo::Type InSelectionType, UBlueprint* Blueprint, UClass* BoundObjectPinClass);
-	void HandleQuickBindActionSelected(const TArray<TSharedPtr<FEdGraphSchemaAction>>& SelectedAction, ESelectInfo::Type InSelectionType, UBlueprint* Blueprint, FMovieSceneEventEndpointParameters Params);
-
-	void OnBlueprintCompiled(UBlueprint*);
-
-	ECheckBoxState GetCallInEditorCheckState() const;
-	void OnSetCallInEditorCheckState(ECheckBoxState NewCheckBoxState);
-
-	TSharedRef<SWidget> GetBoundObjectPinMenuContent();
-	bool CompareBoundObjectPinName(FName InPinName) const;
-	const FSlateBrush* GetBoundObjectPinIcon() const;
-	FText GetBoundObjectPinText() const;
-	void SetBoundObjectPinName(FName InNewBoundObjectPinName);
+	virtual void GetPayloadVariables(UObject* EditObject, void* RawData, FPayloadVariableMap& OutPayloadVariables) const override;
+	virtual bool SetPayloadVariable(UObject* EditObject, void* RawData, FName FieldName, const FString& NewVariableValue) override;
+	virtual UK2Node* FindEndpoint(UMovieSceneSequence* Sequence, UBlueprint* Blueprint, UObject* EditObject, void* RawData) const override;
+	virtual void GetWellKnownParameterPinNames(UObject* EditObject, void* RawData, TArray<FName>& OutWellKnownParameters) const override;
+	virtual void GetWellKnownParameterCandidates(UK2Node* Endpoint, TArray<FWellKnownParameterCandidates>& OutCandidates) const override;
+	virtual bool SetWellKnownParameterPinName(UObject* EditObject, void* RawData, int32 ParameterIndex, FName BoundPinName) override;
+	virtual FMovieSceneDirectorBlueprintEndpointDefinition GenerateEndpointDefinition(UMovieSceneSequence* Sequence) override;
+	virtual void OnCreateEndpoint(UMovieSceneSequence* Sequence, UBlueprint* Blueprint, const TArray<UObject*> EditObjects, const TArray<void*> RawData, const FMovieSceneDirectorBlueprintEndpointDefinition& EndpointDefinition, UK2Node* NewEndpoint) override;
+	virtual void OnSetEndpoint(UMovieSceneSequence* Sequence, UBlueprint* Blueprint, const TArray<UObject*> EditObjects, const TArray<void*> RawData, const FMovieSceneDirectorBlueprintEndpointDefinition& EndpointDefinition, UK2Node* NewEndpoint) override;
+	virtual void GetEditObjects(TArray<UObject*>& OutObjects) const override;
 
 private:
 
 	/** Externally supplied section that the event(s) we're reflecting reside within */
 	TWeakObjectPtr<UMovieSceneSection> WeakExternalSection;
-
-	/** A cache of the common endpoint that is only used when the menu is open to avoid re-computing it every frame. */
-	TWeakObjectPtr<UK2Node_FunctionEntry> CachedCommonEndpoint;
-
-	/** The property handle we're reflecting */
-	TSharedPtr<IPropertyHandle> PropertyHandle;
-
-	TSharedPtr<IPropertyUtilities> PropertyUtilities;
-
-	template<typename KeyType, typename ValueType>
-	using TInlineMap = TSortedMap<KeyType, ValueType, TInlineAllocator<1>>;
 };
diff --git a/Engine/Source/Editor/MovieSceneTools/Public/MovieSceneEventUtils.h b/Engine/Source/Editor/MovieSceneTools/Public/MovieSceneEventUtils.h
--- a/Engine/Source/Editor/MovieSceneTools/Public/MovieSceneEventUtils.h
+++ b/Engine/Source/Editor/MovieSceneTools/Public/MovieSceneEventUtils.h
@@ -5,6 +5,7 @@
 #include "Containers/Array.h"
 #include "Containers/UnrealString.h"
 #include "CoreTypes.h"
+#include "MovieSceneDirectorBlueprintUtils.h"
 #include "UObject/WeakObjectPtrTemplates.h"
 
 class FKismetCompilerContext;
@@ -24,33 +25,20 @@ class UMovieSceneTrack;
 struct FGuid;
 struct FMovieSceneEvent;
 
-
-
-struct MOVIESCENETOOLS_API FMovieSceneEventEndpointParameters
-{
-	FString SanitizedObjectName;
-	FString SanitizedEventName;
-	UClass* BoundObjectPinClass = nullptr;
-
-	static FMovieSceneEventEndpointParameters Generate(UMovieSceneEventTrack* Track);
-	static FMovieSceneEventEndpointParameters Generate(UMovieScene* MovieScene, const FGuid& ObjectBindingID);
-};
-
-
 /**
  * Static utility library for dealing with movie-scene events at edit/cook-time
  */
 struct MOVIESCENETOOLS_API FMovieSceneEventUtils
 {
 	/**
-	 * Create a new event endpoint for the specified sequence. By default, event end points are presented as a UK2Node_CustomEvent in an event graph called "Sequencer Events" with a pin for the bound object.
-	 * @note: Does not bind any event sections to the UBlueprintExtension which is necessary for the event's entry point to be generated during blueprint compilation - that must be done before (or after) this call.
-	 *
-	 * @param Blueprint       (Required, non-null) The blueprint within which to create the new custom event
-	 * @param Parameters      The end point parameters that define how to name and construct the event
-	 * @return A valid pointer to the blueprint node for the event endpoint
+	 * Generate event endpoint definition for the given track
 	 */
-	static UK2Node_CustomEvent* CreateUserFacingEvent(UBlueprint* Blueprint, const FMovieSceneEventEndpointParameters& Parameters);
+	static FMovieSceneDirectorBlueprintEndpointDefinition GenerateEventDefinition(UMovieSceneTrack* Track);
+
+	/**
+	 * Generate event endpoint definition for the given object binding
+	 */
+	static FMovieSceneDirectorBlueprintEndpointDefinition GenerateEventDefinition(UMovieScene* MovieScene, const FGuid& ObjectBindingID);
 
 
 	/**
@@ -63,7 +51,6 @@ struct MOVIESCENETOOLS_API FMovieSceneEventUtils
 	 */
 	static UK2Node_CustomEvent* BindNewUserFacingEvent(FMovieSceneEvent* EntryPoint, UMovieSceneEventSectionBase* EventSection, UBlueprint* Blueprint);
 
-
 	/**
 	 * Main generation function for creating the necessary function graph for calling an event endpoint with the parameters specified in the entry point payload
 	 *
@@ -72,7 +59,7 @@ struct MOVIESCENETOOLS_API FMovieSceneEventUtils
 	 * @param Compiler               (Required, non-null) The kismet compiler context for generation
 	 * @param Endpoint               (Required, non-null) The endpoint node that the event is bound to
 	 */
-	static UK2Node_FunctionEntry* GenerateEntryPoint(UMovieSceneEventSectionBase* EventSection, FMovieSceneEvent* EntryPoint, FKismetCompilerContext* Compiler, UEdGraphNode* Endpoint);
+	static UK2Node_FunctionEntry* GenerateEntryPoint(FMovieSceneEvent* EntryPoint, FKismetCompilerContext* Compiler, UEdGraphNode* Endpoint);
 
 	/**
 	 * Create the necessary mapping between a director Blueprint and and event section in order for the event section to be compiled correctly.
@@ -98,8 +85,6 @@ struct MOVIESCENETOOLS_API FMovieSceneEventUtils
 	 */
 	static void RemoveUnusedCustomEvents(const TArray<TWeakObjectPtr<UMovieSceneEventSectionBase>>& EventSections, UBlueprint* DirectorBP);
 
-public:
-
 	/**
 	 * Attempt to locate the blueprint node that relates to an event's end-point.
 	 * 
@@ -110,17 +95,6 @@ public:
 	 */
 	static UK2Node* FindEndpoint(FMovieSceneEvent* EntryPoint, UMovieSceneEventSectionBase* EventSection, UBlueprint* OwnerBlueprint);
 
-
-	/**
-	 * Attempt to locate an output pin on the specified endpoint node that matches the specified pin class
-	 * 
-	 * @param InEndpoint             (Required, non-null) The endpoint node to search on
-	 * @param BoundObjectPinClass    (Required, non-null) The class of the object pin to locate
-	 * @return An output pin that is of type PC_Object with the specified class as its PinSubCategoryObject, or nullptr if one was not found
-	 */
-	static UEdGraphPin* FindBoundObjectPin(UK2Node* InEndpoint, UClass* BoundObjectPinClass);
-
-
 	/**
 	 * Set the specified EntryPoint to be bound to a new K2 node
 	 * 
