Sequencer: make it possible to expand skeletal animation sections back to the left

diff --git a/Engine/Source/Editor/MovieSceneTools/Private/TrackEditors/SkeletalAnimationTrackEditor.cpp b/Engine/Source/Editor/MovieSceneTools/Private/TrackEditors/SkeletalAnimationTrackEditor.cpp
--- a/Engine/Source/Editor/MovieSceneTools/Private/TrackEditors/SkeletalAnimationTrackEditor.cpp
+++ b/Engine/Source/Editor/MovieSceneTools/Private/TrackEditors/SkeletalAnimationTrackEditor.cpp
@@ -531,6 +531,7 @@ void FSkeletalAnimationSection::BeginDilateSection()
 {
 	Section.PreviousPlayRate = Section.Params.PlayRate; //make sure to cache the play rate
 }
+
 void FSkeletalAnimationSection::DilateSection(const TRange<FFrameNumber>& NewRange, float DilationFactor)
 {
 	Section.Params.PlayRate = Section.PreviousPlayRate / DilationFactor;
@@ -542,7 +543,6 @@ UMovieSceneSection* FSkeletalAnimationSection::GetSectionObject()
 	return &Section;
 }
 
-
 FText FSkeletalAnimationSection::GetSectionTitle() const
 {
 	if (Section.Params.Animation != nullptr)
@@ -694,7 +694,7 @@ int32 FSkeletalAnimationSection::OnPaintSection( FSequencerSectionPainter& Paint
 void FSkeletalAnimationSection::BeginResizeSection()
 {
 	InitialFirstLoopStartOffsetDuringResize = Section.Params.FirstLoopStartFrameOffset;
-	InitialStartTimeDuringResize   = Section.HasStartFrame() ? Section.GetInclusiveStartFrame() : 0;
+	InitialStartTimeDuringResize = Section.HasStartFrame() ? Section.GetInclusiveStartFrame() : 0;
 }
 
 void FSkeletalAnimationSection::ResizeSection(ESequencerSectionResizeMode ResizeMode, FFrameNumber ResizeTime)
@@ -702,31 +702,27 @@ void FSkeletalAnimationSection::ResizeSection(ESequencerSectionResizeMode Resize
 	// Adjust the start offset when resizing from the beginning
 	if (ResizeMode == SSRM_LeadingEdge)
 	{
-		FFrameRate FrameRate   = Section.GetTypedOuter<UMovieScene>()->GetTickResolution();
-		FFrameNumber StartOffset = FrameRate.AsFrameNumber((ResizeTime - InitialStartTimeDuringResize) / FrameRate * Section.Params.PlayRate);
+		// Get the effective animation length, in frames (rounded up), after taking into account start/end trimming.
+		const FFrameRate FrameRate = Section.GetTypedOuter<UMovieScene>()->GetTickResolution();
+		const FFrameNumber SeqLength = FrameRate.AsFrameTime(Section.Params.GetSequenceLength()).CeilToFrame() - Section.Params.StartFrameOffset - Section.Params.EndFrameOffset;
 
-		StartOffset += InitialFirstLoopStartOffsetDuringResize;
+		// Note that there is no scalar-multiplication support for frame numbers, so we need to multiply Value directly.
+		FFrameNumber ResizeAmount = (ResizeTime - InitialStartTimeDuringResize);
+		ResizeAmount.Value = (int32)FMath::Floor(ResizeAmount.Value * Section.Params.PlayRate);
+		FFrameNumber NewFirstLoopStartFrameOffset = InitialFirstLoopStartOffsetDuringResize + ResizeAmount;
 
-		if (StartOffset < 0)
+		// If the start offset exceeds the length of one loop, trim it back.
+		if (SeqLength > 0)
 		{
-			FFrameTime FrameTimeOver = FFrameTime::FromDecimal(StartOffset.Value / Section.Params.PlayRate);
-
-			// Ensure start offset is not less than 0 and adjust ResizeTime
-			ResizeTime = ResizeTime - FrameTimeOver.GetFrame();
-
-			StartOffset = FFrameNumber(0);
+			NewFirstLoopStartFrameOffset = NewFirstLoopStartFrameOffset % SeqLength;
 		}
-		else
+		// If the start offset is negative, add an extra loop at the beginning by making this start offset the complement.
+		if (NewFirstLoopStartFrameOffset < 0)
 		{
-			// If the start offset exceeds the length of one loop, trim it back.
-			const FFrameNumber SeqLength = FrameRate.AsFrameNumber(Section.Params.GetSequenceLength()) - Section.Params.StartFrameOffset - Section.Params.EndFrameOffset;
-			if (SeqLength > 0)
-			{
-				StartOffset = StartOffset % SeqLength;
-			}
+			NewFirstLoopStartFrameOffset = SeqLength + NewFirstLoopStartFrameOffset;
 		}
 
-		Section.Params.FirstLoopStartFrameOffset = StartOffset;
+		Section.Params.FirstLoopStartFrameOffset = NewFirstLoopStartFrameOffset;
 	}
 
 	ISequencerSection::ResizeSection(ResizeMode, ResizeTime);
