Sequencer: Constraints: Baking would be off with sub sequences starting at non zero frame, need to modify time by usual inverse focused time.

diff --git a/Engine/Source/Editor/MovieSceneTools/Private/Constraints/ComponentConstraintChannelInterface.cpp b/Engine/Source/Editor/MovieSceneTools/Private/Constraints/ComponentConstraintChannelInterface.cpp
--- a/Engine/Source/Editor/MovieSceneTools/Private/Constraints/ComponentConstraintChannelInterface.cpp
+++ b/Engine/Source/Editor/MovieSceneTools/Private/Constraints/ComponentConstraintChannelInterface.cpp
@@ -375,11 +375,12 @@ void FComponentConstraintChannelInterface::RecomposeTransforms(
 		FDecompositionQuery Query;
 		Query.Object   = SceneComponent;
 		Query.bConvertFromSourceEntityIDs = false;  // We already pass the children entity IDs
-		
+		FMovieSceneSequenceTransform RootToLocalTransform = InSequencer->GetFocusedMovieSceneSequenceTransform();
+
 		// add keys
 		for (int32 Index = 0; Index < InFrames.Num(); ++Index)
 		{
-			const FFrameNumber& FrameNumber = InFrames[Index];
+			const FFrameNumber& FrameNumber = (FFrameTime(InFrames[Index]) * RootToLocalTransform.InverseLinearOnly()).GetFrame();
 			const FMovieSceneEvaluationRange EvaluationRange = FMovieSceneEvaluationRange(FFrameTime(FrameNumber), TickResolution);
 			const FMovieSceneContext Context = FMovieSceneContext(EvaluationRange, PlaybackStatus).SetHasJumped(true);
 
diff --git a/Engine/Source/Editor/MovieSceneTools/Private/Constraints/MovieSceneConstraintChannelHelper.cpp b/Engine/Source/Editor/MovieSceneTools/Private/Constraints/MovieSceneConstraintChannelHelper.cpp
--- a/Engine/Source/Editor/MovieSceneTools/Private/Constraints/MovieSceneConstraintChannelHelper.cpp
+++ b/Engine/Source/Editor/MovieSceneTools/Private/Constraints/MovieSceneConstraintChannelHelper.cpp
@@ -87,11 +87,12 @@ void FCompensationEvaluator::ComputeLocalTransforms(
 	const FConstraintsManagerController& Controller = FConstraintsManagerController::Get(InWorld);
 	static constexpr bool bSorted = true;
 	const TArray<ConstraintPtr> AllConstraints = Controller.GetAllConstraints(bSorted);
-	
+	FMovieSceneSequenceTransform RootToLocalTransform = InSequencer->GetFocusedMovieSceneSequenceTransform();
+
 	for (int32 Index = 0; Index < NumFrames + 1; ++Index)
 	{
-		const FFrameNumber FrameNumber = (Index == 0) ? InFrames[0] - 1 : InFrames[Index - 1];
-
+		FFrameNumber FrameNumber = (Index == 0) ? InFrames[0] - 1 : InFrames[Index - 1];
+		FrameNumber = (FFrameTime(FrameNumber) * RootToLocalTransform.InverseLinearOnly()).GetFrame();
 		// evaluate animation
 		const FMovieSceneEvaluationRange EvaluationRange = FMovieSceneEvaluationRange(FFrameTime(FrameNumber), TickResolution);
 		const FMovieSceneContext Context = FMovieSceneContext(EvaluationRange, PlaybackStatus).SetHasJumped(true);
@@ -240,9 +241,11 @@ void FCompensationEvaluator::ComputeLocalTransformsForBaking(UWorld* InWorld, co
 			BakeHelper->StartBaking(MovieScene);
 		}
 	}
+	FMovieSceneSequenceTransform RootToLocalTransform = InSequencer->GetFocusedMovieSceneSequenceTransform();
+
 	for (int32 Index = 0; Index < NumFrames; ++Index)
 	{
-		const FFrameNumber& FrameNumber = InFrames[Index];
+		const FFrameNumber& FrameNumber = (FFrameTime(InFrames[Index]) * RootToLocalTransform.InverseLinearOnly()).GetFrame();
 
 		// evaluate animation
 		const FMovieSceneEvaluationRange EvaluationRange = FMovieSceneEvaluationRange(FFrameTime(FrameNumber), TickResolution);
@@ -353,9 +356,11 @@ void FCompensationEvaluator::ComputeLocalTransformsBeforeDeletion(
 			BakeHelper->StartBaking(MovieScene);
 		}
 	}
+	FMovieSceneSequenceTransform RootToLocalTransform = InSequencer->GetFocusedMovieSceneSequenceTransform();
+
 	for (int32 Index = 0; Index < NumFrames; ++Index)
 	{
-		const FFrameNumber& FrameNumber = InFrames[Index];
+		const FFrameNumber& FrameNumber = (FFrameTime(InFrames[Index]) * RootToLocalTransform.InverseLinearOnly()).GetFrame();
 
 		// evaluate animation
 		const FMovieSceneEvaluationRange EvaluationRange = FMovieSceneEvaluationRange(FFrameTime(FrameNumber), TickResolution);
@@ -449,8 +454,10 @@ void FCompensationEvaluator::ComputeCompensation(UWorld* InWorld, const TSharedP
 		}
 	}
 	
-	auto EvaluateAt = [InSequencer, &AllConstraints, &BakeHelpers](const FFrameNumber InFrame)
+	auto EvaluateAt = [InSequencer, &AllConstraints, &BakeHelpers](FFrameNumber InFrame)
 	{
+		FMovieSceneSequenceTransform RootToLocalTransform = InSequencer->GetFocusedMovieSceneSequenceTransform();
+		InFrame = (FFrameTime(InFrame) * RootToLocalTransform.InverseLinearOnly()).GetFrame();
 
 		UMovieScene* MovieScene = InSequencer->GetFocusedMovieSceneSequence()->GetMovieScene();
 		const FFrameRate TickResolution = MovieScene->GetTickResolution();
@@ -553,10 +560,12 @@ void FCompensationEvaluator::CacheTransforms(UWorld* InWorld, const TSharedPtr<I
 
 	const TArray<IMovieSceneToolsAnimationBakeHelper*>& BakeHelpers = FMovieSceneToolsModule::Get().GetAnimationBakeHelpers();
 	
-	auto EvaluateAt = [&](const FFrameNumber InFrame)
+	auto EvaluateAt = [&](FFrameNumber InFrame)
 	{
 		const FMovieSceneEvaluationRange EvaluationRange = FMovieSceneEvaluationRange(FFrameTime(InFrame), TickResolution);
 		const FMovieSceneContext Context = FMovieSceneContext(EvaluationRange, PlaybackStatus).SetHasJumped(true);
+		FMovieSceneSequenceTransform RootToLocalTransform = InSequencer->GetFocusedMovieSceneSequenceTransform();
+		InFrame = (FFrameTime(InFrame) * RootToLocalTransform.InverseLinearOnly()).GetFrame();
 
 		for (IMovieSceneToolsAnimationBakeHelper* BakeHelper : BakeHelpers)
 		{
