Sequencer: user-defined bindings

This change adds the ability to define custom logic (via the directory blueprint) for object bindings. Instead of binding to the usual possessable or spawnable, this will optionally run arbitrary BP logic defined by the user to resolve the bound object. This makes it possible, for instance, to bind to the player pawn, or any other actor that might not exist at edit time.

The user interface is similar to that of the event track, and therefore reuses the director blueprint endpoint picker. There is however a new blueprint extension to compile these functions into the spawnables/possessables, and of course a spot to call these functions when appropriate.

diff --git a/Engine/Source/Editor/MovieSceneTools/Private/MovieSceneDynamicBindingBlueprintExtension.cpp b/Engine/Source/Editor/MovieSceneTools/Private/MovieSceneDynamicBindingBlueprintExtension.cpp
new file mode 100644
--- /dev/null
+++ b/Engine/Source/Editor/MovieSceneTools/Private/MovieSceneDynamicBindingBlueprintExtension.cpp
@@ -0,0 +1,142 @@
+// Copyright Epic Games, Inc. All Rights Reserved.
+
+#include "MovieSceneDynamicBindingBlueprintExtension.h"
+
+#include "CoreGlobals.h"
+#include "Delegates/Delegate.h"
+#include "EdGraph/EdGraph.h"
+#include "Engine/Blueprint.h"
+#include "HAL/Platform.h"
+#include "K2Node.h"
+#include "K2Node_CallFunction.h"
+#include "K2Node_FunctionEntry.h"
+#include "K2Node_FunctionResult.h"
+#include "KismetCompiler.h"
+#include "Misc/AssertionMacros.h"
+#include "MovieScene.h"
+#include "MovieSceneSequence.h"
+#include "MovieSceneDynamicBinding.h"
+#include "MovieSceneDynamicBindingUtils.h"
+#include "MovieSceneDirectorBlueprintUtils.h"
+#include "UObject/NameTypes.h"
+#include "UObject/PropertyPortFlags.h"
+#include "UObject/UnrealNames.h"
+#include "UObject/WeakObjectPtr.h"
+
+class UEdGraphNode;
+
+void UMovieSceneDynamicBindingBlueprintExtension::BindTo(TWeakObjectPtr<UMovieSceneSequence> InMovieSceneSequence)
+{
+	check(!WeakMovieSceneSequence.IsValid() || WeakMovieSceneSequence == InMovieSceneSequence);
+	WeakMovieSceneSequence = InMovieSceneSequence;
+}
+
+void UMovieSceneDynamicBindingBlueprintExtension::HandlePreloadObjectsForCompilation(UBlueprint* OwningBlueprint)
+{
+	if (UMovieSceneSequence* MovieSceneSequence = WeakMovieSceneSequence.Get())
+	{
+		UBlueprint::ForceLoad(MovieSceneSequence);
+		if (UMovieScene* MovieScene = MovieSceneSequence->GetMovieScene())
+		{
+			UBlueprint::ForceLoad(MovieScene);
+		}
+	}
+}
+
+void UMovieSceneDynamicBindingBlueprintExtension::HandleGenerateFunctionGraphs(FKismetCompilerContext* CompilerContext)
+{
+	UMovieSceneSequence* MovieSceneSequence = WeakMovieSceneSequence.Get();
+	UMovieScene* MovieScene = MovieSceneSequence ? MovieSceneSequence->GetMovieScene() : nullptr;
+	if (!MovieScene)
+	{
+		return;
+	}
+
+	ensureMsgf(!MovieScene->HasAnyFlags(RF_NeedLoad), TEXT("Attempting to generate entry point functions before a movie scene has been loaded"));
+
+	// Generate a function graph for each endpoint used by the bound sequence. This function graph is simply a call
+	// to the endpoint function with the payload variables set as the call parameters.
+	auto GenerateFunctionGraphs = 
+		[CompilerContext]
+		(const FGuid& ObjectBinding, FMovieSceneDynamicBinding& DynamicBinding)
+	{
+		UEdGraphNode* Endpoint = Cast<UK2Node>(DynamicBinding.WeakEndpoint.Get());
+		if (Endpoint)
+		{
+			// Set up the endpoint call, with our payload variables.
+			FMovieSceneDirectorBlueprintEndpointCall EndpointCall;
+			EndpointCall.Endpoint = Endpoint;
+			if (!DynamicBinding.ResolveParamsPinName.IsNone())
+			{
+				EndpointCall.ExposedPinNames.Add(DynamicBinding.ResolveParamsPinName);
+			}
+			for (auto& Pair : DynamicBinding.PayloadVariables)
+			{
+				EndpointCall.PayloadVariables.Add(Pair.Key, Pair.Value.Value);
+			}
+
+			// Create the endpoint call, and clean-up stale payload variables.
+			FMovieSceneDirectorBlueprintEntrypointResult EntrypointResult = FMovieSceneDirectorBlueprintUtils::GenerateEntryPoint(EndpointCall, CompilerContext);
+			DynamicBinding.CompiledFunctionName = EntrypointResult.CompiledFunctionName;
+			EntrypointResult.CleanUpStalePayloadVariables(DynamicBinding.PayloadVariables);
+		}
+	};
+	FMovieSceneDynamicBindingUtils::IterateDynamicBindings(MovieScene, GenerateFunctionGraphs);
+
+	// This callback sets the generated function calls back onto the spawnables/possessables in the sequence,
+	// and keeps a pointer to any needed arguments on this function, so that we can pass special values later
+	// at runtime.
+	TWeakObjectPtr<UMovieSceneSequence> WeakMovieSceneSequenceCopy(WeakMovieSceneSequence);
+	auto OnFunctionListGenerated = [WeakMovieSceneSequenceCopy](FKismetCompilerContext* CompilerContext)
+	{
+		UMovieSceneSequence* MovieSceneSequence = WeakMovieSceneSequenceCopy.Get();
+		UMovieScene* MovieScene = MovieSceneSequence ? MovieSceneSequence->GetMovieScene() : nullptr;
+		if (!ensureMsgf(MovieSceneSequence, TEXT("A movie scene was garbage-collected while its director blueprint was being compiled!")))
+		{
+			return;
+		}
+
+		UBlueprint* Blueprint = CompilerContext->Blueprint;
+		if (!Blueprint->GeneratedClass)
+		{
+			return;
+		}
+
+		TArray<FMovieSceneDynamicBinding*> DynamicBindings;
+		FMovieSceneDynamicBindingUtils::GatherDynamicBindings(MovieScene, DynamicBindings);
+		for (FMovieSceneDynamicBinding* DynamicBinding : DynamicBindings)
+		{
+			// Set the pointer to the resolver function to invoke.
+			if (DynamicBinding->CompiledFunctionName != NAME_None)
+			{
+				DynamicBinding->Function = Blueprint->GeneratedClass->FindFunctionByName(DynamicBinding->CompiledFunctionName);
+			}
+			else
+			{
+				DynamicBinding->Function = nullptr;
+			}
+
+			DynamicBinding->CompiledFunctionName = NAME_None;
+
+			// Set the pointer to the resolve parameter field, if any.
+			if (DynamicBinding->Function && !DynamicBinding->ResolveParamsPinName.IsNone())
+			{
+				FProperty* ResolveParamsProp = DynamicBinding->Function->FindPropertyByName(DynamicBinding->ResolveParamsPinName);
+				DynamicBinding->ResolveParamsProperty = ResolveParamsProp;
+			}
+			else
+			{
+				DynamicBinding->ResolveParamsProperty = nullptr;
+			}
+		}
+
+		if (!Blueprint->bIsRegeneratingOnLoad)
+		{
+			MovieScene->MarkAsChanged();
+			MovieScene->MarkPackageDirty();
+		}
+	};
+
+	CompilerContext->OnFunctionListCompiled().AddLambda(OnFunctionListGenerated);
+}
+
diff --git a/Engine/Source/Editor/MovieSceneTools/Private/MovieSceneDynamicBindingBlueprintExtension.h b/Engine/Source/Editor/MovieSceneTools/Private/MovieSceneDynamicBindingBlueprintExtension.h
new file mode 100644
--- /dev/null
+++ b/Engine/Source/Editor/MovieSceneTools/Private/MovieSceneDynamicBindingBlueprintExtension.h
@@ -0,0 +1,44 @@
+// Copyright Epic Games, Inc. All Rights Reserved.
+
+#pragma once
+
+#include "Blueprint/BlueprintExtension.h"
+#include "Containers/Array.h"
+#include "HAL/PlatformCrt.h"
+#include "UObject/UObjectGlobals.h"
+#include "UObject/WeakObjectPtrTemplates.h"
+
+#include "MovieSceneDynamicBindingBlueprintExtension.generated.h"
+
+class FKismetCompilerContext;
+class UBlueprint;
+class UMovieSceneSequence;
+class UObject;
+struct FMovieSceneCustomizeDirectorBlueprintEndpointCallParams;
+
+/**
+ * An extension for a sequence's director blueprint to compile dynamic binding endpoints.
+ */
+UCLASS()
+class UMovieSceneDynamicBindingBlueprintExtension : public UBlueprintExtension
+{
+public:
+
+	GENERATED_BODY()
+
+	/**
+	 * Bind this extension to the given sequence. This may be called multiple times, but the extension
+	 * must only be bound to one sequence.
+	 */
+	void BindTo(TWeakObjectPtr<UMovieSceneSequence> InMovieSceneSequence);
+
+private:
+
+	virtual void HandlePreloadObjectsForCompilation(UBlueprint* OwningBlueprint) override final;
+	virtual void HandleGenerateFunctionGraphs(FKismetCompilerContext* CompilerContext) override final;
+
+	/** The movie scene that is bound to the director blueprint */
+	UPROPERTY()
+	TWeakObjectPtr<UMovieSceneSequence> WeakMovieSceneSequence;
+};
+
diff --git a/Engine/Source/Editor/MovieSceneTools/Private/MovieSceneDynamicBindingCustomization.cpp b/Engine/Source/Editor/MovieSceneTools/Private/MovieSceneDynamicBindingCustomization.cpp
new file mode 100644
--- /dev/null
+++ b/Engine/Source/Editor/MovieSceneTools/Private/MovieSceneDynamicBindingCustomization.cpp
@@ -0,0 +1,250 @@
+// Copyright Epic Games, Inc. All Rights Reserved.
+
+#include "MovieSceneDynamicBindingCustomization.h"
+
+#include "MovieScene.h"
+#include "MovieSceneDynamicBinding.h"
+#include "MovieSceneDynamicBindingUtils.h"
+#include "MovieSceneSequence.h"
+
+#include "AssetRegistry/ARFilter.h"
+#include "AssetRegistry/AssetData.h"
+#include "AssetRegistry/AssetRegistryModule.h"
+#include "AssetRegistry/IAssetRegistry.h"
+
+#include "BlueprintActionFilter.h"
+#include "BlueprintActionMenuBuilder.h"
+#include "BlueprintActionMenuItem.h"
+#include "BlueprintFunctionNodeSpawner.h"
+#include "EdGraphSchema_K2.h"
+#include "IPropertyUtilities.h"
+#include "K2Node.h"
+#include "K2Node_CallFunction.h"
+#include "K2Node_FunctionEntry.h"
+#include "K2Node_FunctionResult.h"
+#include "Kismet2/BlueprintEditorUtils.h"
+#include "PropertyHandle.h"
+
+#define LOCTEXT_NAMESPACE "MovieSceneDynamicBindingCustomization"
+
+TSharedRef<IPropertyTypeCustomization> FMovieSceneDynamicBindingCustomization::MakeInstance()
+{
+	TSharedRef<FMovieSceneDynamicBindingCustomization> Instance = MakeShared<FMovieSceneDynamicBindingCustomization>();
+	return Instance;
+}
+
+TSharedRef<IPropertyTypeCustomization> FMovieSceneDynamicBindingCustomization::MakeInstance(UMovieScene* InMovieScene, FGuid InObjectBinding)
+{
+	TSharedRef<FMovieSceneDynamicBindingCustomization> Instance = MakeShared<FMovieSceneDynamicBindingCustomization>();
+	Instance->EditedMovieScene = InMovieScene;
+	Instance->ObjectBinding = InObjectBinding;
+	return Instance;
+}
+
+void FMovieSceneDynamicBindingCustomization::GetPayloadVariables(UObject* EditObject, void* RawData, FPayloadVariableMap& OutPayloadVariables) const
+{
+	const FMovieSceneDynamicBinding* DynamicBinding = static_cast<FMovieSceneDynamicBinding*>(RawData);
+	for (const TPair<FName, FMovieSceneDynamicBindingPayloadVariable>& Pair : DynamicBinding->PayloadVariables)
+	{
+		OutPayloadVariables.Add(Pair.Key, Pair.Value.Value);
+	}
+}
+
+bool FMovieSceneDynamicBindingCustomization::SetPayloadVariable(UObject* EditObject, void* RawData, FName FieldName, const FString& NewVariableValue)
+{
+	UMovieScene* MovieScene = Cast<UMovieScene>(EditObject);
+	FMovieSceneDynamicBinding* DynamicBinding = static_cast<FMovieSceneDynamicBinding*>(RawData);
+	if (!MovieScene || !DynamicBinding)
+	{
+		return false;
+	}
+
+	MovieScene->Modify();
+
+	if (!NewVariableValue.IsEmpty())
+	{
+		FMovieSceneDynamicBindingPayloadVariable* PayloadVariable = DynamicBinding->PayloadVariables.Find(FieldName);
+		if (!PayloadVariable)
+		{
+			PayloadVariable = &DynamicBinding->PayloadVariables.Add(FieldName);
+		}
+
+		PayloadVariable->Value = NewVariableValue;
+	}
+	else
+	{
+		DynamicBinding->PayloadVariables.Remove(FieldName);
+	}
+
+	return true;
+}
+
+UK2Node* FMovieSceneDynamicBindingCustomization::FindEndpoint(UMovieSceneSequence* Sequence, UBlueprint* Blueprint, UObject* EditObject, void* RawData) const
+{
+	FMovieSceneDynamicBinding* DynamicBinding = static_cast<FMovieSceneDynamicBinding*>(RawData);
+	if (UK2Node* Node = Cast<UK2Node>(DynamicBinding->WeakEndpoint.Get()))
+	{
+		return Node;
+	}
+	return nullptr;
+}
+
+void FMovieSceneDynamicBindingCustomization::GetWellKnownParameterPinNames(UObject* EditObject, void* RawData, TArray<FName>& OutWellKnownParameters) const
+{
+	FMovieSceneDynamicBinding* DynamicBinding = static_cast<FMovieSceneDynamicBinding*>(RawData);
+	OutWellKnownParameters.Add(DynamicBinding->ResolveParamsPinName);
+}
+
+void FMovieSceneDynamicBindingCustomization::GetWellKnownParameterCandidates(UK2Node* Endpoint, TArray<FWellKnownParameterCandidates>& OutCandidates) const
+{
+	FWellKnownParameterCandidates ResolveParamsCandidates;
+	ResolveParamsCandidates.Metadata.PickerLabel = LOCTEXT("ResolveParamsPin_Label", "Pass Resolve Params To");
+	ResolveParamsCandidates.Metadata.PickerTooltip = LOCTEXT("ResolveParamsPin_Tooltip", "Specifies a pin to pass the resolve parameters through when the binding is resolved.");
+
+	for (UEdGraphPin* Pin : Endpoint->Pins)
+	{
+		// Parameter pins are outputs on the function entry node.
+		if (Pin->Direction != EGPD_Output)
+		{
+			continue;
+		}
+
+		// Pin of type FMovieSceneDynamicBindingResolveParams is eligible for passing the resolve params.
+		if (Pin->PinType.PinCategory == UEdGraphSchema_K2::PC_Struct &&
+				Pin->PinType.PinSubCategoryObject == FMovieSceneDynamicBindingResolveParams::StaticStruct())
+		{
+			ResolveParamsCandidates.CandidatePinNames.Add(Pin->GetFName());
+		}
+	}
+	
+	OutCandidates.Add(ResolveParamsCandidates);
+}
+
+bool FMovieSceneDynamicBindingCustomization::SetWellKnownParameterPinName(UObject* EditObject, void* RawData, int32 ParameterIndex, FName BoundPinName)
+{
+	FMovieSceneDynamicBinding* DynamicBinding = static_cast<FMovieSceneDynamicBinding*>(RawData);
+	switch (ParameterIndex)
+	{
+		case 0:
+			DynamicBinding->ResolveParamsPinName = BoundPinName;
+			return true;
+	}
+	return false;
+}
+
+FMovieSceneDirectorBlueprintEndpointDefinition FMovieSceneDynamicBindingCustomization::GenerateEndpointDefinition(UMovieSceneSequence* Sequence)
+{
+	FMovieSceneDirectorBlueprintEndpointDefinition Definition;
+	Definition.EndpointType = EMovieSceneDirectorBlueprintEndpointType::Function;
+
+	// We use a dummy utility class to get a function signature that takes no parameter and returns a UObject pointer.
+	static const FName SampleResolveBindingFuncName("SampleResolveBinding");
+	UClass* EndpointUtilClass = UMovieSceneDynamicBindingEndpointUtil::StaticClass();
+	Definition.EndpointSignature = EndpointUtilClass->FindFunctionByName(SampleResolveBindingFuncName);
+	check(Definition.EndpointSignature);
+
+	UMovieScene* MovieScene = Sequence->GetMovieScene();
+	if (FMovieScenePossessable* Possessable = MovieScene->FindPossessable(ObjectBinding))
+	{
+		Definition.EndpointName = Possessable->GetName() + "_DynamicBinding";
+	}
+	else if (FMovieSceneSpawnable* Spawnable = MovieScene->FindSpawnable(ObjectBinding))
+	{
+		Definition.EndpointName = Spawnable->GetName() + "_DynamicBinding";
+	}
+
+	return Definition;
+}
+
+void FMovieSceneDynamicBindingCustomization::OnCreateEndpoint(UMovieSceneSequence* Sequence, UBlueprint* Blueprint, const TArray<UObject*> EditObjects, const TArray<void*> RawData, const FMovieSceneDirectorBlueprintEndpointDefinition& EndpointDefinition, UK2Node* NewEndpoint)
+{
+	UMovieScene* MovieScene = Sequence->GetMovieScene();
+
+	MovieScene->Modify();
+
+	for (int32 Index = 0; Index < RawData.Num(); ++Index)
+	{
+		ensureMsgf(
+			Cast<UMovieScene>(EditObjects[Index]) == MovieScene, 
+			TEXT("Editing dynamic binding endpoint for a different sequence"));
+
+		FMovieSceneDynamicBinding* DynamicBinding = static_cast<FMovieSceneDynamicBinding*>(RawData[Index]);
+
+		FMovieSceneDynamicBindingUtils::SetEndpoint(MovieScene, DynamicBinding, NewEndpoint);
+	}
+
+	FMovieSceneDynamicBindingUtils::EnsureBlueprintExtensionCreated(Sequence, Blueprint);
+}
+
+void FMovieSceneDynamicBindingCustomization::OnSetEndpoint(UMovieSceneSequence* Sequence, UBlueprint* Blueprint, const TArray<UObject*> EditObjects, const TArray<void*> RawData, const FMovieSceneDirectorBlueprintEndpointDefinition& EndpointDefinition, UK2Node* NewEndpoint)
+{
+	check(EditObjects.Num() == RawData.Num());
+
+	for (int32 Index = 0; Index < RawData.Num(); ++Index)
+	{
+		UMovieScene* MovieScene = Cast<UMovieScene>(EditObjects[Index]);
+		FMovieSceneDynamicBinding* DynamicBinding = static_cast<FMovieSceneDynamicBinding*>(RawData[Index]);
+
+		FMovieSceneDynamicBindingUtils::SetEndpoint(MovieScene, DynamicBinding, NewEndpoint);
+
+		FMovieSceneDynamicBindingUtils::EnsureBlueprintExtensionCreated(Sequence, Blueprint);
+	}
+}
+
+void FMovieSceneDynamicBindingCustomization::GetEditObjects(TArray<UObject*>& OutObjects) const
+{
+	OutObjects.Add(EditedMovieScene);
+}
+
+void FMovieSceneDynamicBindingCustomization::OnCollectQuickBindActions(UBlueprint* Blueprint, FBlueprintActionMenuBuilder& MenuBuilder)
+{
+	CollectResolverLibraryBindActions(Blueprint, MenuBuilder, false);
+}
+
+void FMovieSceneDynamicBindingCustomization::CollectResolverLibraryBindActions(UBlueprint* Blueprint, FBlueprintActionMenuBuilder& MenuBuilder, bool bIsRebinding)
+{
+	// We don't show the resolver library endpoints for rebinding, because we should only rebind
+	// to other function graphs of the director blueprint.
+	if (bIsRebinding)
+	{
+		return;
+	}
+
+	// We want the ability to create CallFunction nodes for any static method that we think can be used
+	// as a bound object resolution function.
+	FBlueprintActionFilter MenuFilter(FBlueprintActionFilter::BPFILTER_RejectGlobalFields | FBlueprintActionFilter::BPFILTER_RejectPermittedSubClasses);
+	MenuFilter.PermittedNodeTypes.Add(UK2Node_CallFunction::StaticClass());
+	MenuFilter.Context.Blueprints.Add(Blueprint);
+	
+	// Add any class that has the "SequencerBindingResolverLibrary" meta as a target class.
+	//
+	// We don't consider *all* blueprint function libraries because there are many, many of them that expose
+	// functions that are, technically speaking, compatible with bound object resolution (i.e. they return
+	// a UObject pointer) but that are completely non-sensical in this context.
+	const static FName SequencerBindingResolverLibraryMeta("SequencerBindingResolverLibrary");
+	for (TObjectIterator<UClass> ClassIt; ClassIt; ++ClassIt)
+	{
+		UClass* CurrentClass = *ClassIt;
+		if (CurrentClass->HasMetaData(SequencerBindingResolverLibraryMeta))
+		{
+			FBlueprintActionFilter::Add(MenuFilter.TargetClasses, CurrentClass);
+		}
+	}
+	
+	auto RejectAnyIncompatibleReturnValues = [](const FBlueprintActionFilter& Filter, FBlueprintActionInfo& BlueprintAction)
+	{
+		const UFunction* Function = BlueprintAction.GetAssociatedFunction();
+		const FObjectProperty* FunctionReturnProperty = CastField<FObjectProperty>(Function->GetReturnProperty());
+		return 
+			!FunctionReturnProperty || 
+			!FunctionReturnProperty->PropertyClass || 
+			!FunctionReturnProperty->PropertyClass->IsChildOf<UObject>();
+	};
+
+	MenuFilter.AddRejectionTest(FBlueprintActionFilter::FRejectionTestDelegate::CreateLambda(RejectAnyIncompatibleReturnValues));
+
+	MenuBuilder.AddMenuSection(MenuFilter, LOCTEXT("DynamicBindingCustomization", "Resolver Library"), 0);
+}
+
+#undef LOCTEXT_NAMESPACE
+
diff --git a/Engine/Source/Editor/MovieSceneTools/Private/MovieSceneDynamicBindingUtils.cpp b/Engine/Source/Editor/MovieSceneTools/Private/MovieSceneDynamicBindingUtils.cpp
new file mode 100644
--- /dev/null
+++ b/Engine/Source/Editor/MovieSceneTools/Private/MovieSceneDynamicBindingUtils.cpp
@@ -0,0 +1,56 @@
+// Copyright Epic Games, Inc. All Rights Reserved.
+
+#include "MovieSceneDynamicBindingUtils.h"
+
+#include "K2Node.h"
+#include "K2Node_FunctionEntry.h"
+#include "MovieSceneSequence.h"
+#include "MovieSceneDynamicBindingBlueprintExtension.h"
+
+void FMovieSceneDynamicBindingUtils::SetEndpoint(UMovieScene* MovieScene, FMovieSceneDynamicBinding* DynamicBinding, UK2Node* NewEndpoint)
+{
+	UK2Node* ExistingEndpoint = CastChecked<UK2Node>(DynamicBinding->WeakEndpoint.Get(), ECastCheckedType::NullAllowed);
+	if (ExistingEndpoint)
+	{
+		ExistingEndpoint->OnUserDefinedPinRenamed().RemoveAll(MovieScene);
+	}
+
+	if (NewEndpoint)
+	{
+		checkf(
+			NewEndpoint->IsA<UK2Node_FunctionEntry>(),
+			TEXT("Only functions are supported as dynamic binding endpoints"));
+
+		NewEndpoint->OnUserDefinedPinRenamed().AddUObject(MovieScene, &UMovieScene::OnDynamicBindingUserDefinedPinRenamed);
+		DynamicBinding->WeakEndpoint = NewEndpoint;
+	}
+	else
+	{
+		DynamicBinding->WeakEndpoint = nullptr;
+	}
+}
+
+void FMovieSceneDynamicBindingUtils::EnsureBlueprintExtensionCreated(UMovieSceneSequence* MovieSceneSequence, UBlueprint* Blueprint)
+{
+	if (!Blueprint)
+	{
+		return;
+	}
+
+	check(MovieSceneSequence);
+
+	for (TObjectPtr<UBlueprintExtension> Extension : Blueprint->GetExtensions())
+	{
+		UMovieSceneDynamicBindingBlueprintExtension* DynamidBindingExtension = Cast<UMovieSceneDynamicBindingBlueprintExtension>(Extension);
+		if (DynamidBindingExtension)
+		{
+			DynamidBindingExtension->BindTo(MovieSceneSequence);
+			return;
+		}
+	}
+
+	UMovieSceneDynamicBindingBlueprintExtension* DynamidBindingExtension = NewObject<UMovieSceneDynamicBindingBlueprintExtension>(Blueprint);
+	DynamidBindingExtension->BindTo(MovieSceneSequence);
+	Blueprint->AddExtension(DynamidBindingExtension);
+}
+
diff --git a/Engine/Source/Editor/MovieSceneTools/Private/MovieSceneToolsModule.cpp b/Engine/Source/Editor/MovieSceneTools/Private/MovieSceneToolsModule.cpp
--- a/Engine/Source/Editor/MovieSceneTools/Private/MovieSceneToolsModule.cpp
+++ b/Engine/Source/Editor/MovieSceneTools/Private/MovieSceneToolsModule.cpp
@@ -62,6 +62,7 @@
 #include "MovieSceneBuiltInEasingFunctionCustomization.h"
 #include "MovieSceneAlphaBlendOptionCustomization.h"
 #include "MovieSceneObjectBindingIDCustomization.h"
+#include "MovieSceneDynamicBindingCustomization.h"
 #include "MovieSceneEventCustomization.h"
 #include "SequencerClipboardReconciler.h"
 #include "ClipboardTypes.h"
@@ -193,6 +194,7 @@ void FMovieSceneToolsModule::StartupModule()
 		PropertyModule.RegisterCustomClassLayout("MovieSceneFloatPerlinNoiseChannelContainer", FOnGetDetailCustomizationInstance::CreateStatic(&FMovieSceneFloatPerlinNoiseChannelDetailsCustomization::MakeInstance));
 		PropertyModule.RegisterCustomClassLayout("MovieSceneDoublePerlinNoiseChannelContainer", FOnGetDetailCustomizationInstance::CreateStatic(&FMovieSceneDoublePerlinNoiseChannelDetailsCustomization::MakeInstance));
 		PropertyModule.RegisterCustomPropertyTypeLayout("MovieSceneObjectBindingID", FOnGetPropertyTypeCustomizationInstance::CreateLambda(&MakeShared<FMovieSceneObjectBindingIDCustomization>));
+		PropertyModule.RegisterCustomPropertyTypeLayout("MovieSceneDynamicBinding", FOnGetPropertyTypeCustomizationInstance::CreateStatic(&FMovieSceneDynamicBindingCustomization::MakeInstance));
 		PropertyModule.RegisterCustomPropertyTypeLayout("MovieSceneEvent", FOnGetPropertyTypeCustomizationInstance::CreateStatic(&FMovieSceneEventCustomization::MakeInstance));
 		PropertyModule.RegisterCustomPropertyTypeLayout("MovieSceneCVarOverrides", FOnGetPropertyTypeCustomizationInstance::CreateStatic(&UE::MovieScene::FCVarOverridesPropertyTypeCustomization::MakeInstance));
 
@@ -235,7 +237,8 @@ void FMovieSceneToolsModule::StartupModule()
 		));
 	}
 
-	FixupPayloadParameterNameHandle = UMovieSceneEventSectionBase::FixupPayloadParameterNameEvent.AddStatic(FixupPayloadParameterNameForSection);
+	FixupDynamicBindingPayloadParameterNameHandle = UMovieScene::FixupDynamicBindingPayloadParameterNameEvent.AddStatic(FixupPayloadParameterNameForDynamicBinding);
+	FixupEventSectionPayloadParameterNameHandle = UMovieSceneEventSectionBase::FixupPayloadParameterNameEvent.AddStatic(FixupPayloadParameterNameForSection);
 	UMovieSceneEventSectionBase::UpgradeLegacyEventEndpoint.BindStatic(UpgradeLegacyEventEndpointForSection);
 	UMovieSceneEventSectionBase::PostDuplicateSectionEvent.BindStatic(PostDuplicateEventSection);
 	UMovieSceneEventSectionBase::RemoveForCookEvent.BindStatic(RemoveForCookEventSection);
@@ -271,7 +274,8 @@ void FMovieSceneToolsModule::StartupModule()
 
 void FMovieSceneToolsModule::ShutdownModule()
 {
-	UMovieSceneEventSectionBase::FixupPayloadParameterNameEvent.Remove(FixupPayloadParameterNameHandle);
+	UMovieScene::FixupDynamicBindingPayloadParameterNameEvent.Remove(FixupDynamicBindingPayloadParameterNameHandle);
+	UMovieSceneEventSectionBase::FixupPayloadParameterNameEvent.Remove(FixupEventSectionPayloadParameterNameHandle);
 	UMovieSceneEventSectionBase::UpgradeLegacyEventEndpoint = UMovieSceneEventSectionBase::FUpgradeLegacyEventEndpoint();
 	UMovieSceneEventSectionBase::PostDuplicateSectionEvent = UMovieSceneEventSectionBase::FPostDuplicateEvent();
 	UMovieSceneEventSectionBase::RemoveForCookEvent = UMovieSceneEventSectionBase::FRemoveForCookEvent();
@@ -569,6 +573,33 @@ void FMovieSceneToolsModule::FixupPayloadParameterNameForSection(UMovieSceneEven
 	}
 }
 
+void FMovieSceneToolsModule::FixupPayloadParameterNameForDynamicBinding(UMovieScene* MovieScene, UK2Node* InNode, FName OldPinName, FName NewPinName)
+{
+	check(MovieScene);
+
+	auto FixupPayloadParameterName = [InNode, OldPinName, NewPinName](FMovieSceneDynamicBinding& DynamicBinding)
+	{
+		if (DynamicBinding.WeakEndpoint.Get() == InNode)
+		{
+			if (FMovieSceneDynamicBindingPayloadVariable* Variable = DynamicBinding.PayloadVariables.Find(OldPinName))
+			{
+				DynamicBinding.PayloadVariables.Add(NewPinName, MoveTemp(*Variable));
+				DynamicBinding.PayloadVariables.Remove(OldPinName);
+			}
+		}
+	};
+
+	for (int32 Index = 0, PossessableCount = MovieScene->GetPossessableCount(); Index < PossessableCount; ++Index)
+	{
+		FixupPayloadParameterName(MovieScene->GetPossessable(Index).DynamicBinding);
+	}
+
+	for (int32 Index = 0, SpawnableCount = MovieScene->GetSpawnableCount(); Index < SpawnableCount; ++Index)
+	{
+		FixupPayloadParameterName(MovieScene->GetSpawnable(Index).DynamicBinding);
+	}
+}
+
 void FMovieSceneToolsModule::RegisterClipboardConversions()
 {
 	using namespace MovieSceneClipboard;
diff --git a/Engine/Source/Editor/MovieSceneTools/Public/MovieSceneDynamicBindingCustomization.h b/Engine/Source/Editor/MovieSceneTools/Public/MovieSceneDynamicBindingCustomization.h
new file mode 100644
--- /dev/null
+++ b/Engine/Source/Editor/MovieSceneTools/Public/MovieSceneDynamicBindingCustomization.h
@@ -0,0 +1,57 @@
+// Copyright Epic Games, Inc. All Rights Reserved.
+
+#pragma once
+
+#include "CoreMinimal.h"
+#include "MovieSceneDirectorBlueprintEndpointCustomization.h"
+#include "IPropertyTypeCustomization.h"
+#include "Types/SlateEnums.h"
+
+class UK2Node;
+class UBlueprint;
+class UEdGraphPin;
+class IPropertyHandle;
+class FDetailWidgetRow;
+class SWidget;
+
+struct FMovieSceneDynamicBinding;
+
+enum class ECheckBoxState : uint8;
+
+/**
+ * Customization for dynamic spawn/possession endpoint picker.
+ */
+class MOVIESCENETOOLS_API FMovieSceneDynamicBindingCustomization : public FMovieSceneDirectorBlueprintEndpointCustomization
+{
+public:
+
+	static TSharedRef<IPropertyTypeCustomization> MakeInstance();
+	static TSharedRef<IPropertyTypeCustomization> MakeInstance(UMovieScene* InMovieScene, FGuid InObjectBinding);
+
+protected:
+
+	virtual void GetPayloadVariables(UObject* EditObject, void* RawData, FPayloadVariableMap& OutPayloadVariables) const override;
+	virtual bool SetPayloadVariable(UObject* EditObject, void* RawData, FName FieldName, const FString& NewVariableValue) override;
+	virtual UK2Node* FindEndpoint(UMovieSceneSequence* Sequence, UBlueprint* Blueprint, UObject* EditObject, void* RawData) const override;
+	virtual void GetWellKnownParameterPinNames(UObject* EditObject, void* RawData, TArray<FName>& OutWellKnownParameters) const override;
+	virtual void GetWellKnownParameterCandidates(UK2Node* Endpoint, TArray<FWellKnownParameterCandidates>& OutCandidates) const override;
+	virtual bool SetWellKnownParameterPinName(UObject* EditObject, void* RawData, int32 ParameterIndex, FName BoundPinName) override;
+	virtual FMovieSceneDirectorBlueprintEndpointDefinition GenerateEndpointDefinition(UMovieSceneSequence* Sequence) override;
+	virtual void OnCreateEndpoint(UMovieSceneSequence* Sequence, UBlueprint* Blueprint, const TArray<UObject*> EditObjects, const TArray<void*> RawData, const FMovieSceneDirectorBlueprintEndpointDefinition& EndpointDefinition, UK2Node* NewEndpoint) override;
+	virtual void OnSetEndpoint(UMovieSceneSequence* Sequence, UBlueprint* Blueprint, const TArray<UObject*> EditObjects, const TArray<void*> RawData, const FMovieSceneDirectorBlueprintEndpointDefinition& EndpointDefinition, UK2Node* NewEndpoint) override;
+	virtual void GetEditObjects(TArray<UObject*>& OutObjects) const override;
+	virtual void OnCollectQuickBindActions(UBlueprint* Blueprint, FBlueprintActionMenuBuilder& MenuBuilder) override;
+
+private:
+
+	void SetEndpointImpl(UMovieScene* MovieScene, FMovieSceneDynamicBinding* DynamicBinding, UBlueprint* Blueprint, UK2Node* NewEndpoint);
+	void EnsureBlueprintExtensionCreated(UMovieScene* MovieScene, UBlueprint* Blueprint);
+
+	void CollectResolverLibraryBindActions(UBlueprint* Blueprint, FBlueprintActionMenuBuilder& MenuBuilder, bool bIsRebinding);
+
+private:
+
+	UMovieScene* EditedMovieScene;
+	FGuid ObjectBinding;
+};
+
diff --git a/Engine/Source/Editor/MovieSceneTools/Public/MovieSceneDynamicBindingUtils.h b/Engine/Source/Editor/MovieSceneTools/Public/MovieSceneDynamicBindingUtils.h
new file mode 100644
--- /dev/null
+++ b/Engine/Source/Editor/MovieSceneTools/Public/MovieSceneDynamicBindingUtils.h
@@ -0,0 +1,72 @@
+// Copyright Epic Games, Inc. All Rights Reserved.
+
+#pragma once
+
+#include "CoreTypes.h"
+#include "MovieScene.h"
+#include "MovieSceneDynamicBinding.h"
+#include "MovieScenePossessable.h"
+#include "MovieSceneSpawnable.h"
+
+#include "MovieSceneDynamicBindingUtils.generated.h"
+
+/**
+ * A utility class for managing dynamic binding endpoints.
+ */
+struct MOVIESCENETOOLS_API FMovieSceneDynamicBindingUtils
+{
+	/**
+	 * Set an endpoint on the given dynamic binding.
+	 */
+	static void SetEndpoint(UMovieScene* MovieScene, FMovieSceneDynamicBinding* DynamicBinding, UK2Node* NewEndpoint);
+
+	/**
+	 * Ensures that the dynamic binding blueprint extension has been added to the given sequence's director blueprint.
+	 */
+	static void EnsureBlueprintExtensionCreated(UMovieSceneSequence* MovieSceneSequence, UBlueprint* Blueprint);
+
+	/**
+	 * Utility function for iterating all dynamic bindings in a sequence.
+	 */
+	template<typename Callback>
+	static void IterateDynamicBindings(UMovieScene* InMovieScene, Callback&& InCallback)
+	{
+		for (int32 Index = 0, PossessableCount = InMovieScene->GetPossessableCount(); Index < PossessableCount; ++Index)
+		{
+			FMovieScenePossessable& Possessable = InMovieScene->GetPossessable(Index);
+			FMovieSceneDynamicBinding& DynamicBinding = Possessable.DynamicBinding;
+			InCallback(Possessable.GetGuid(), DynamicBinding);
+		}
+
+		for (int32 Index = 0, SpawnableCount = InMovieScene->GetSpawnableCount(); Index < SpawnableCount; ++Index)
+		{
+			FMovieSceneSpawnable& Spawnable = InMovieScene->GetSpawnable(Index);
+			FMovieSceneDynamicBinding& DynamicBinding = Spawnable.DynamicBinding;
+			InCallback(Spawnable.GetGuid(), DynamicBinding);
+		}
+	}
+
+	/**
+	 * Utility function for gathering all dynamic bindings in a sequence into a container.
+	 */
+	static void GatherDynamicBindings(UMovieScene* InMovieScene, TArray<FMovieSceneDynamicBinding*>& OutDynamicBindings)
+	{
+		IterateDynamicBindings(InMovieScene, [&](const FGuid&, FMovieSceneDynamicBinding& Item)
+			{
+				OutDynamicBindings.Add(&Item);
+			});
+	}
+};
+
+/**
+ * Dummy class, used for easily getting a valid UFunction that helps prepare blueprint function graphs.
+ */
+UCLASS()
+class UMovieSceneDynamicBindingEndpointUtil : public UObject
+{
+	GENERATED_BODY()
+
+	UFUNCTION()
+	UObject* SampleResolveBinding() { return nullptr; }
+};
+
diff --git a/Engine/Source/Editor/MovieSceneTools/Public/MovieSceneToolsModule.h b/Engine/Source/Editor/MovieSceneTools/Public/MovieSceneToolsModule.h
--- a/Engine/Source/Editor/MovieSceneTools/Public/MovieSceneToolsModule.h
+++ b/Engine/Source/Editor/MovieSceneTools/Public/MovieSceneToolsModule.h
@@ -76,6 +76,7 @@ private:
 	void RegisterClipboardConversions();
 
 	static void FixupPayloadParameterNameForSection(UMovieSceneEventSectionBase* Section, UK2Node* InNode, FName OldPinName, FName NewPinName);
+	static void FixupPayloadParameterNameForDynamicBinding(UMovieScene* MovieScene, UK2Node* InNode, FName OldPinName, FName NewPinName);
 	static bool UpgradeLegacyEventEndpointForSection(UMovieSceneEventSectionBase* Section);
 	static void PostDuplicateEventSection(UMovieSceneEventSectionBase* Section);
 	static void RemoveForCookEventSection(UMovieSceneEventSectionBase* Section);
@@ -127,7 +128,8 @@ private:
 	FDelegateHandle CinematicShotTrackModelHandle;
 
 	FDelegateHandle GenerateEventEntryPointsHandle;
-	FDelegateHandle FixupPayloadParameterNameHandle;
+	FDelegateHandle FixupDynamicBindingPayloadParameterNameHandle;
+	FDelegateHandle FixupEventSectionPayloadParameterNameHandle;
 	FDelegateHandle UpgradeLegacyEventEndpointHandle;
 
 	FDelegateHandle OnObjectsReplacedHandle;
diff --git a/Engine/Source/Editor/Sequencer/Private/MVVM/ViewModels/ObjectBindingModel.cpp b/Engine/Source/Editor/Sequencer/Private/MVVM/ViewModels/ObjectBindingModel.cpp
--- a/Engine/Source/Editor/Sequencer/Private/MVVM/ViewModels/ObjectBindingModel.cpp
+++ b/Engine/Source/Editor/Sequencer/Private/MVVM/ViewModels/ObjectBindingModel.cpp
@@ -12,31 +12,36 @@
 #include "MVVM/ViewModels/OutlinerViewModelDragDropOp.h"
 #include "MVVM/Views/SOutlinerObjectBindingView.h"
 #include "MVVM/Extensions/IRecyclableExtension.h"
-#include "Modules/ModuleManager.h"
+#include "Algo/Sort.h"
+#include "ClassViewerModule.h"
+#include "Containers/ArrayBuilder.h"
+#include "Engine/LevelStreaming.h"
+#include "Framework/Application/SlateApplication.h"
+#include "Framework/MultiBox/MultiBoxBuilder.h"
+#include "IDetailCustomization.h"
+#include "IDetailsView.h"
 #include "ISequencerModule.h"
+#include "ISequencerTrackEditor.h"
+#include "IStructureDetailsView.h"
+#include "Modules/ModuleManager.h"
 #include "MovieScene.h"
-#include "MovieSceneFolder.h"
 #include "MovieSceneBinding.h"
+#include "MovieSceneDynamicBindingCustomization.h"
+#include "MovieSceneFolder.h"
+#include "ObjectBindingTagCache.h"
+#include "ObjectEditorUtils.h"
+#include "PropertyEditorModule.h"
 #include "PropertyPath.h"
+#include "SObjectBindingTag.h"
 #include "ScopedTransaction.h"
 #include "Sequencer.h"
-#include "SequencerNodeTree.h"
 #include "SequencerCommands.h"
+#include "SequencerNodeTree.h"
 #include "SequencerSettings.h"
 #include "SequencerUtilities.h"
-#include "ISequencerTrackEditor.h"
-#include "SObjectBindingTag.h"
-#include "Containers/ArrayBuilder.h"
-#include "ObjectBindingTagCache.h"
-#include "ObjectEditorUtils.h"
 #include "Styling/AppStyle.h"
 #include "Styling/SlateIconFinder.h"
 #include "Tracks/MovieSceneSpawnTrack.h"
-#include "ClassViewerModule.h"
-#include "Algo/Sort.h"
-#include "Engine/LevelStreaming.h"
-#include "Framework/Application/SlateApplication.h"
-#include "Framework/MultiBox/MultiBoxBuilder.h"
 
 #define LOCTEXT_NAMESPACE "ObjectBindingModel"
 
@@ -961,6 +966,11 @@ void FObjectBindingModel::BuildContextMenu(FMenuBuilder& MenuBuilder)
 				LOCTEXT("ChangeClassTooltip", "Change the class (object template) that this spawns from"),
 				FNewMenuDelegate::CreateSP(this, &FObjectBindingModel::AddChangeClassMenu));
 
+			MenuBuilder.AddSubMenu(
+				LOCTEXT("DynamicSpawn", "Dynamic Spawn"),
+				LOCTEXT("DynamicSpawnTooltip", "Specify a Blueprint method that will spawn or otherwise acquire a compatible actor for this binding"),
+				FNewMenuDelegate::CreateSP(this, &FObjectBindingModel::AddDynamicSpawnMenu));
+
 			MenuBuilder.AddMenuEntry(
 				LOCTEXT("ContinuouslyRespawn", "Continuously Respawn"),
 				LOCTEXT("ContinuouslyRespawnTooltip", "When enabled, this spawnable will always be respawned if it gets destroyed externally. When disabled, this object will only ever be spawned once for each spawn key even if it is destroyed externally"),
@@ -998,6 +1008,11 @@ void FObjectBindingModel::BuildContextMenu(FMenuBuilder& MenuBuilder)
 
 			MenuBuilder.AddMenuEntry( FSequencerCommands::Get().ConvertToSpawnable );
 
+			MenuBuilder.AddSubMenu(
+				LOCTEXT("DynamicPossession", "Dynamic Possession"),
+				LOCTEXT("DynamicPossessionTooltip", "Specify a Blueprint method that will find a compatible possessable actor for this binding"),
+				FNewMenuDelegate::CreateSP(this, &FObjectBindingModel::AddDynamicPossessionMenu));
+
 			MenuBuilder.EndSection();
 		}
 
@@ -1228,6 +1243,102 @@ void FObjectBindingModel::HandleTemplateActorClassPicked(UClass* ChosenClass)
 	}
 }
 
+void FObjectBindingModel::AddDynamicSpawnMenu(FMenuBuilder& MenuBuilder)
+{
+	UMovieScene* MovieScene = OwnerModel->GetMovieScene();
+	FMovieSceneSpawnable* Spawnable = MovieScene->FindSpawnable(ObjectBindingID);
+	if (!Spawnable)
+	{
+		return;
+	}
+
+	AddDynamicBindingMenu(MenuBuilder, Spawnable->DynamicBinding);
+}
+
+void FObjectBindingModel::AddDynamicPossessionMenu(FMenuBuilder& MenuBuilder)
+{
+	UMovieScene* MovieScene = OwnerModel->GetMovieScene();
+	FMovieScenePossessable* Possessable = MovieScene->FindPossessable(ObjectBindingID);
+	if (!Possessable)
+	{
+		return;
+	}
+
+	AddDynamicBindingMenu(MenuBuilder, Possessable->DynamicBinding);
+}
+
+void FObjectBindingModel::AddDynamicBindingMenu(FMenuBuilder& MenuBuilder, FMovieSceneDynamicBinding& DynamicBinding)
+{
+	FDetailsViewArgs DetailsViewArgs;
+	{
+		DetailsViewArgs.bAllowSearch = false;
+		DetailsViewArgs.bCustomFilterAreaLocation = true;
+		DetailsViewArgs.bCustomNameAreaLocation = true;
+		DetailsViewArgs.bHideSelectionTip = true;
+		DetailsViewArgs.bLockable = false;
+		DetailsViewArgs.bSearchInitialKeyFocus = true;
+		DetailsViewArgs.bUpdatesFromSelection = false;
+		DetailsViewArgs.bShowOptions = false;
+		DetailsViewArgs.bShowModifiedPropertiesOption = false;
+		DetailsViewArgs.bShowScrollBar = false;
+	}
+
+	FStructureDetailsViewArgs StructureViewArgs;
+	{
+		StructureViewArgs.bShowObjects = false;
+		StructureViewArgs.bShowAssets = true;
+		StructureViewArgs.bShowClasses = true;
+		StructureViewArgs.bShowInterfaces = false;
+	}
+
+	TSharedRef<IStructureDetailsView> StructureDetailsView = FModuleManager::GetModuleChecked<FPropertyEditorModule>("PropertyEditor")
+		.CreateStructureDetailView(DetailsViewArgs, StructureViewArgs, nullptr);
+
+	// Register details customizations for this instance
+	TSharedPtr<FSequencer> Sequencer = OwnerModel->GetSequencerImpl();
+	UMovieSceneSequence* Sequence = Sequencer->GetRootMovieSceneSequence();
+	StructureDetailsView->GetDetailsView()->RegisterInstancedCustomPropertyTypeLayout(
+		FMovieSceneDynamicBinding::StaticStruct()->GetFName(),
+		FOnGetPropertyTypeCustomizationInstance::CreateStatic(&FMovieSceneDynamicBindingCustomization::MakeInstance, Sequence->GetMovieScene(), ObjectBindingID));
+
+	// We can't just show the FMovieSceneDynamicBinding struct in the details view, because Slate only uses
+	// the above details view customization for *properties* (not for the root object). So here we put a copy of
+	// our dynamic binding struct inside a container, and when the details view is done setting values on it,
+	// we copy these values back to the original dynamic binding.
+	TSharedPtr<FStructOnScope> StructOnScope = MakeShared<FStructOnScope>(FMovieSceneDynamicBindingContainer::StaticStruct());
+	FMovieSceneDynamicBindingContainer* BufferContainer = (FMovieSceneDynamicBindingContainer*)StructOnScope->GetStructMemory();
+	BufferContainer->DynamicBinding = DynamicBinding;
+	StructureDetailsView->SetStructureData(StructOnScope);
+
+	StructureDetailsView->GetOnFinishedChangingPropertiesDelegate().AddSP(this, &FObjectBindingModel::OnFinishedChangingDynamicBindingProperties, StructOnScope);
+
+	MenuBuilder.BeginSection(NAME_None, LOCTEXT("DynamicBindingHeader", "Dynamic Binding"));
+	{
+		TSharedRef<SWidget> Widget = StructureDetailsView->GetWidget().ToSharedRef();
+		MenuBuilder.AddWidget(Widget, FText());
+	}
+	MenuBuilder.EndSection();
+}
+
+void FObjectBindingModel::OnFinishedChangingDynamicBindingProperties(const FPropertyChangedEvent& ChangeEvent, TSharedPtr<FStructOnScope> ValueStruct)
+{
+	auto* Container = (FMovieSceneDynamicBindingContainer*)ValueStruct->GetStructMemory();
+
+	UMovieScene* MovieScene = OwnerModel->GetMovieScene();
+	FMovieScenePossessable* Possessable = MovieScene->FindPossessable(ObjectBindingID);
+	if (Possessable)
+	{
+		Possessable->DynamicBinding = Container->DynamicBinding;
+		return;
+	}
+	FMovieSceneSpawnable* Spawnable = MovieScene->FindSpawnable(ObjectBindingID);
+	if (Spawnable)
+	{
+		Spawnable->DynamicBinding = Container->DynamicBinding;
+		return;
+	}
+}
+
 void FObjectBindingModel::AddTagMenu(FMenuBuilder& MenuBuilder)
 {
 	MenuBuilder.AddMenuEntry(FSequencerCommands::Get().OpenTaggedBindingManager);
diff --git a/Engine/Source/Editor/Sequencer/Public/MVVM/ViewModels/ObjectBindingModel.h b/Engine/Source/Editor/Sequencer/Public/MVVM/ViewModels/ObjectBindingModel.h
--- a/Engine/Source/Editor/Sequencer/Public/MVVM/ViewModels/ObjectBindingModel.h
+++ b/Engine/Source/Editor/Sequencer/Public/MVVM/ViewModels/ObjectBindingModel.h
@@ -20,6 +20,8 @@ class UMovieScene;
 class UMovieSceneTrack;
 class FMenuBuilder;
 class FPropertyPath;
+class FStructOnScope;
+struct FMovieSceneDynamicBinding;
 enum class ECheckBoxState : uint8;
 
 namespace UE
@@ -128,6 +130,8 @@ private:
 	void AddSpawnLevelMenu(FMenuBuilder& MenuBuilder);
 	void AddTagMenu(FMenuBuilder& MenuBuilder);
 	void AddChangeClassMenu(FMenuBuilder& MenuBuilder);
+	void AddDynamicSpawnMenu(FMenuBuilder& MenuBuilder);
+	void AddDynamicPossessionMenu(FMenuBuilder& MenuBuilder);
 
 private:
 
@@ -141,6 +145,9 @@ private:
 	void HandleAddTag(FName TagName);
 	void HandleTemplateActorClassPicked(UClass* ChosenClass);
 
+	void AddDynamicBindingMenu(FMenuBuilder& MenuBuilder, FMovieSceneDynamicBinding& DynamicBinding);
+	void OnFinishedChangingDynamicBindingProperties(const FPropertyChangedEvent& ChangeEvent, TSharedPtr<FStructOnScope> ValueStruct);
+
 protected:
 
 	FGuid ObjectBindingID;
diff --git a/Engine/Source/Runtime/MovieScene/Private/Evaluation/MovieSceneEvaluationState.cpp b/Engine/Source/Runtime/MovieScene/Private/Evaluation/MovieSceneEvaluationState.cpp
--- a/Engine/Source/Runtime/MovieScene/Private/Evaluation/MovieSceneEvaluationState.cpp
+++ b/Engine/Source/Runtime/MovieScene/Private/Evaluation/MovieSceneEvaluationState.cpp
@@ -1,14 +1,17 @@
 // Copyright Epic Games, Inc. All Rights Reserved.
 
 #include "Evaluation/MovieSceneEvaluationState.h"
-#include "MovieSceneSequence.h"
-#include "MovieScene.h"
-#include "IMovieScenePlayer.h"
-#include "IMovieScenePlaybackClient.h"
-#include "MovieSceneObjectBindingID.h"
 
 #include "Algo/Sort.h"
 #include "Algo/Unique.h"
+#include "Engine/World.h"
+#include "Evaluation/MovieSceneEvaluationTemplateInstance.h"
+#include "IMovieScenePlaybackClient.h"
+#include "IMovieScenePlayer.h"
+#include "MovieScene.h"
+#include "MovieSceneDynamicBindingInvoker.h"
+#include "MovieSceneObjectBindingID.h"
+#include "MovieSceneSequence.h"
 
 DECLARE_CYCLE_STAT(TEXT("Find Bound Objects"), MovieSceneEval_FindBoundObjects, STATGROUP_MovieSceneEval);
 DECLARE_CYCLE_STAT(TEXT("Iterate Bound Objects"), MovieSceneEval_IterateBoundObjects, STATGROUP_MovieSceneEval);
@@ -388,7 +391,11 @@ void FMovieSceneObjectCache::UpdateBindings(const FGuid& InGuid, IMovieScenePlay
 				}
 				else
 				{
-					Player.ResolveBoundObjects(InGuid, SequenceID, *Sequence, ResolutionContext, FoundObjects);
+					const bool bUseDefault = FMovieSceneDynamicBindingInvoker::ResolveDynamicBinding(Player, Sequence, SequenceID, InGuid, Possessable->DynamicBinding, FoundObjects);
+					if (bUseDefault)
+					{
+						Player.ResolveBoundObjects(InGuid, SequenceID, *Sequence, ResolutionContext, FoundObjects);
+					}
 				}
 				
 				Bindings = BoundObjects.Find(InGuid);
diff --git a/Engine/Source/Runtime/MovieScene/Private/MovieScene.cpp b/Engine/Source/Runtime/MovieScene/Private/MovieScene.cpp
--- a/Engine/Source/Runtime/MovieScene/Private/MovieScene.cpp
+++ b/Engine/Source/Runtime/MovieScene/Private/MovieScene.cpp
@@ -44,6 +44,7 @@ TOptional<TRangeBound<FFrameNumber>> GetMaxUpperBound(const UMovieSceneTrack* Tr
 #if WITH_EDITOR
 
 UMovieScene::FIsTrackClassAllowedEvent UMovieScene::IsTrackClassAllowedEvent;
+UMovieScene::FFixupDynamicBindingPayloadParameterNameEvent UMovieScene::FixupDynamicBindingPayloadParameterNameEvent;
 
 bool UMovieScene::IsTrackClassAllowed(UClass* InClass)
 {
diff --git a/Engine/Source/Runtime/MovieScene/Private/MovieSceneDynamicBinding.cpp b/Engine/Source/Runtime/MovieScene/Private/MovieSceneDynamicBinding.cpp
new file mode 100644
--- /dev/null
+++ b/Engine/Source/Runtime/MovieScene/Private/MovieSceneDynamicBinding.cpp
@@ -0,0 +1,39 @@
+// Copyright Epic Games, Inc. All Rights Reserved.
+
+#include "MovieSceneDynamicBinding.h"
+
+#include "EngineUtils.h"
+#include "Engine/Engine.h"
+#include "Engine/World.h"
+#include "GameFramework/PlayerController.h"
+#include "MovieScene.h"
+#include "MovieSceneBinding.h"
+#include "MovieScenePossessable.h"
+#include "MovieSceneSequence.h"
+#include "MovieSceneSpawnable.h"
+
+#include UE_INLINE_GENERATED_CPP_BY_NAME(MovieSceneDynamicBinding)
+
+UObject* UBuiltInDynamicBindingResolverLibrary::ResolveToPlayerPawn(UObject* WorldContextObject, int32 PlayerControllerIndex)
+{
+	if (UWorld* World = GEngine->GetWorldFromContextObject(WorldContextObject, EGetWorldErrorMode::LogAndReturnNull))
+	{
+		const int32 NumPlayerControllers = World->GetNumPlayerControllers();
+		if (ensure(PlayerControllerIndex < NumPlayerControllers))
+		{
+			FConstPlayerControllerIterator It = World->GetPlayerControllerIterator();
+			for (int32 Index = 0; Index < PlayerControllerIndex; ++Index)
+			{
+				++It;
+			}
+			if (APlayerController* PlayerController = It->Get())
+			{
+				APawn* Pawn = PlayerController->GetPawn();
+				return (UObject*)Pawn;
+			}
+		}
+	}
+
+	return nullptr;
+}
+
diff --git a/Engine/Source/Runtime/MovieScene/Private/MovieSceneDynamicBindingInvoker.cpp b/Engine/Source/Runtime/MovieScene/Private/MovieSceneDynamicBindingInvoker.cpp
new file mode 100644
--- /dev/null
+++ b/Engine/Source/Runtime/MovieScene/Private/MovieSceneDynamicBindingInvoker.cpp
@@ -0,0 +1,134 @@
+// Copyright Epic Games, Inc. All Rights Reserved.
+
+#include "MovieSceneDynamicBindingInvoker.h"
+
+#include "CoreMinimal.h"
+#include "Engine/World.h"
+#include "Evaluation/MovieSceneEvaluationTemplateInstance.h"
+#include "IMovieScenePlayer.h"
+#include "MovieSceneDynamicBinding.h"
+
+bool FMovieSceneDynamicBindingInvoker::ResolveDynamicBinding(IMovieScenePlayer& Player, UMovieSceneSequence* Sequence, const FMovieSceneSequenceID& SequenceID, const FGuid& InGuid, const FMovieSceneDynamicBinding& DynamicBinding, TArray<UObject*, TInlineAllocator<1>>& OutObjects)
+{
+	UObject* Object = ResolveDynamicBinding(Player, Sequence, SequenceID, InGuid, DynamicBinding);
+	if (Object != nullptr)
+	{
+		OutObjects.Add(Object);
+
+		// We have successfully called the dynamic binding function, don't use the default behavior.
+		return false;
+	}
+
+	// No valid object found, fallback to default behavior.
+	return true;
+}
+
+UObject* FMovieSceneDynamicBindingInvoker::ResolveDynamicBinding(IMovieScenePlayer& Player, UMovieSceneSequence* Sequence, const FMovieSceneSequenceID& SequenceID, const FGuid& InGuid, const FMovieSceneDynamicBinding& DynamicBinding)
+{
+	if (!ensure(Sequence))
+	{
+		// Sequence is somehow null... fallback to default behavior.
+		return nullptr;
+	}
+
+	UFunction* DynamicBindingFunc = DynamicBinding.Function.Get();
+	if (!DynamicBindingFunc)
+	{
+		// No dynamic binding, fallback to default behavior.
+		return nullptr;
+	}
+
+	UObject* DirectorInstance = Player.GetEvaluationTemplate().GetOrCreateDirectorInstance(SequenceID, Player);
+	if (!DirectorInstance)
+	{
+#if !NO_LOGGING
+		UE_LOG(LogMovieScene, Warning, 
+				TEXT("%s: Failed to resolve dynamic binding '%s' because no director instance was available."), 
+				*Sequence->GetName(), *DynamicBindingFunc->GetName());
+#endif
+		// Fallback to default behavior.
+		return nullptr;
+	}
+
+#if WITH_EDITOR
+	const static FName NAME_CallInEditor(TEXT("CallInEditor"));
+
+	UWorld* World = DirectorInstance->GetWorld();
+	const bool bIsGameWorld = World && World->IsGameWorld();
+
+	if (!bIsGameWorld && !DynamicBindingFunc->HasMetaData(NAME_CallInEditor))
+	{
+		UE_LOG(LogMovieScene, Verbose,
+				TEXT("%s: Refusing to resolve dynamic binding '%s' in editor world because function '%s' has 'Call in Editor' set to false."),
+				*Sequence->GetName(), *LexToString(InGuid), *DynamicBindingFunc->GetName());
+		// Fallback to default behavior.
+		return nullptr;
+	}
+#endif // WITH_EDITOR
+
+	UE_LOG(LogMovieScene, VeryVerbose,
+			TEXT("%s: Resolving dynamic binding '%s' with function '%s'."),
+			*Sequence->GetName(), *LexToString(InGuid), *DynamicBindingFunc->GetName());
+
+	FMovieSceneDynamicBindingResolveParams ResolveParams;
+	ResolveParams.ObjectBindingID = InGuid;
+	ResolveParams.Sequence = Sequence;
+	ResolveParams.RootSequence = Player.GetEvaluationTemplate().GetRootSequence();
+	UObject* Object = InvokeDynamicBinding(DirectorInstance, DynamicBinding, ResolveParams);
+
+	return Object;
+}
+
+UObject* FMovieSceneDynamicBindingInvoker::InvokeDynamicBinding(UObject* DirectorInstance, const FMovieSceneDynamicBinding& DynamicBinding, const FMovieSceneDynamicBindingResolveParams& ResolveParams)
+{
+	// Parse all function parameters.
+	UFunction* DynamicBindingFunc = DynamicBinding.Function.Get();
+	check(DynamicBindingFunc);
+	uint8* Parameters = (uint8*)FMemory_Alloca(DynamicBindingFunc->ParmsSize + DynamicBindingFunc->MinAlignment);
+	Parameters = Align(Parameters, DynamicBindingFunc->MinAlignment);
+
+	// Initialize parameters.
+	FMemory::Memzero(Parameters, DynamicBindingFunc->ParmsSize);
+
+	FObjectPropertyBase* ReturnProp = nullptr;
+	for (TFieldIterator<FProperty> It(DynamicBindingFunc); It; ++It)
+	{
+		FProperty* LocalProp = *It;
+		checkSlow(LocalProp);
+		if (!LocalProp->HasAnyPropertyFlags(CPF_ZeroConstructor))
+		{
+			LocalProp->InitializeValue_InContainer(Parameters);
+		}
+
+		if (LocalProp->HasAnyPropertyFlags(CPF_ReturnParm))
+		{
+			ensure(ReturnProp == nullptr);
+			ReturnProp = CastFieldChecked<FObjectPropertyBase>(LocalProp);
+		}
+	}
+	
+	// Set the resolve parameter struct if we need to pass it to the function.
+	if (FProperty* ResolveParamsProp = DynamicBinding.ResolveParamsProperty.Get())
+	{
+		ResolveParamsProp->SetValue_InContainer(Parameters, &ResolveParams);
+	}
+
+	// Invoke the function.
+	DirectorInstance->ProcessEvent(DynamicBindingFunc, Parameters);
+
+	// Grab the result value.
+	UObject* ResultObject = nullptr;
+	if (ensure(ReturnProp))
+	{
+		ResultObject = ReturnProp->GetObjectPropertyValue_InContainer(Parameters);
+	}
+
+	// Destroy parameters.
+	for (TFieldIterator<FProperty> It(DynamicBindingFunc); It; ++It)
+	{
+		It->DestroyValue_InContainer(Parameters);
+	}
+
+	return ResultObject;
+}
+
diff --git a/Engine/Source/Runtime/MovieScene/Private/MovieSceneDynamicBindingInvoker.h b/Engine/Source/Runtime/MovieScene/Private/MovieSceneDynamicBindingInvoker.h
new file mode 100644
--- /dev/null
+++ b/Engine/Source/Runtime/MovieScene/Private/MovieSceneDynamicBindingInvoker.h
@@ -0,0 +1,29 @@
+// Copyright Epic Games, Inc. All Rights Reserved.
+
+#pragma once
+
+#include "Containers/Array.h"
+#include "CoreTypes.h"
+
+class UObject;
+class IMovieScenePlayer;
+class UMovieSceneSequence;
+struct FGuid;
+struct FMovieSceneSequenceID;
+struct FMovieSceneDynamicBinding;
+struct FMovieSceneDynamicBindingResolveParams;
+
+/**
+ * Utility class for invoking dynamic binding endpoints.
+ */
+struct FMovieSceneDynamicBindingInvoker
+{
+	/** Invoke the dynamic binding, if any, and add the result to the given array of objects */
+	static bool ResolveDynamicBinding(IMovieScenePlayer& Player, UMovieSceneSequence* Sequence, const FMovieSceneSequenceID& SequenceID, const FGuid& InGuid, const FMovieSceneDynamicBinding& DynamicBinding, TArray<UObject*, TInlineAllocator<1>>& OutObjects);
+	/** Invoke the dynamic binding, if any, and return the result */
+	static UObject* ResolveDynamicBinding(IMovieScenePlayer& Player, UMovieSceneSequence* Sequence, const FMovieSceneSequenceID& SequenceID, const FGuid& InGuid, const FMovieSceneDynamicBinding& DynamicBinding);
+
+private:
+	static UObject* InvokeDynamicBinding(UObject* DirectorInstance, const FMovieSceneDynamicBinding& DynamicBinding, const FMovieSceneDynamicBindingResolveParams& Params);
+};
+
diff --git a/Engine/Source/Runtime/MovieScene/Private/MovieSceneSpawnRegister.cpp b/Engine/Source/Runtime/MovieScene/Private/MovieSceneSpawnRegister.cpp
--- a/Engine/Source/Runtime/MovieScene/Private/MovieSceneSpawnRegister.cpp
+++ b/Engine/Source/Runtime/MovieScene/Private/MovieSceneSpawnRegister.cpp
@@ -4,6 +4,7 @@
 #include "MovieScene.h"
 #include "IMovieScenePlayer.h"
 #include "Evaluation/MovieSceneEvaluationTemplateInstance.h"
+#include "MovieSceneDynamicBindingInvoker.h"
 #include "MovieSceneSpawnableAnnotation.h"
 
 TWeakObjectPtr<> FMovieSceneSpawnRegister::FindSpawnedObject(const FGuid& BindingId, FMovieSceneSequenceIDRef TemplateID) const
@@ -36,11 +37,21 @@ UObject* FMovieSceneSpawnRegister::SpawnObject(const FGuid& BindingId, UMovieSce
 		return nullptr;
 	}
 
-	UObject* SpawnedActor = SpawnObject(*Spawnable, TemplateID, Player);
+	UMovieSceneSequence* Sequence = Player.GetEvaluationTemplate().GetSequence(TemplateID);
+	if (!ensure(Sequence))
+	{
+		return nullptr;
+	}
+
+	// See if there is some dynamic binding logic to invoke, otherwise spawn the actor
+	UObject* SpawnedActor = FMovieSceneDynamicBindingInvoker::ResolveDynamicBinding(Player, Sequence, TemplateID, BindingId, Spawnable->DynamicBinding);
+	if (!SpawnedActor)
+	{
+		SpawnedActor = SpawnObject(*Spawnable, TemplateID, Player);
+	}
 	
 	if (SpawnedActor)
 	{
-		UMovieSceneSequence* Sequence = Player.GetEvaluationTemplate().GetSequence(TemplateID);
 		FMovieSceneSpawnableAnnotation::Add(SpawnedActor, BindingId, TemplateID, Sequence);
 
 		FMovieSceneSpawnRegisterKey Key(TemplateID, BindingId);
diff --git a/Engine/Source/Runtime/MovieScene/Public/Evaluation/MovieSceneEvaluationState.h b/Engine/Source/Runtime/MovieScene/Public/Evaluation/MovieSceneEvaluationState.h
--- a/Engine/Source/Runtime/MovieScene/Public/Evaluation/MovieSceneEvaluationState.h
+++ b/Engine/Source/Runtime/MovieScene/Public/Evaluation/MovieSceneEvaluationState.h
@@ -26,6 +26,8 @@
 class IMovieScenePlayer;
 class UMovieSceneSequence;
 class UObject;
+struct FMovieSceneDynamicBinding;
+struct FMovieSceneDynamicBindingResolveParams;
 struct FMovieSceneEvaluationKey;
 struct FMovieSceneObjectBindingID;
 struct FSharedPersistentDataKey;
@@ -140,6 +142,16 @@ private:
 	 */
 	void UpdateBindings(const FGuid& InGuid, IMovieScenePlayer& Player);
 
+	/**
+	 * Handles optional dynamic binding for a given object binding.
+	 */
+	bool ResolveDynamicBinding(const FGuid& InGuid, const FMovieSceneDynamicBinding& DynamicBinding, IMovieScenePlayer& Player, TArray<UObject*, TInlineAllocator<1>>& OutObjects);
+
+	/**
+	 * Invokes the custom function used to resolve a given dynamic object binding.
+	 */
+	UObject* InvokeDynamicBinding(UObject* DirectorInstance, const FMovieSceneDynamicBinding& DynamicBinding, const FMovieSceneDynamicBindingResolveParams& ResolveParams);
+
 	/**
 	 * Invalidate the object bindings for a specific object binding ID
 	 */
diff --git a/Engine/Source/Runtime/MovieScene/Public/MovieScene.h b/Engine/Source/Runtime/MovieScene/Public/MovieScene.h
--- a/Engine/Source/Runtime/MovieScene/Public/MovieScene.h
+++ b/Engine/Source/Runtime/MovieScene/Public/MovieScene.h
@@ -47,6 +47,7 @@
 class FArchive;
 class FObjectPreSaveContext;
 class UClass;
+class UK2Node;
 class UMovieSceneFolder;
 class UMovieSceneSection;
 class UMovieSceneTrack;
@@ -580,6 +581,15 @@ public:
 
 	static bool IsTrackClassAllowed(UClass* InClass);
 
+	void OnDynamicBindingUserDefinedPinRenamed(UK2Node* InNode, FName OldPinName, FName NewPinName)
+	{
+		FixupDynamicBindingPayloadParameterNameEvent.Broadcast(this, InNode, OldPinName, NewPinName);
+	}
+
+	DECLARE_MULTICAST_DELEGATE_FourParams(FFixupDynamicBindingPayloadParameterNameEvent, UMovieScene*, UK2Node*, FName, FName);
+
+	static FFixupDynamicBindingPayloadParameterNameEvent FixupDynamicBindingPayloadParameterNameEvent;
+
 #endif
 
 public:
diff --git a/Engine/Source/Runtime/MovieScene/Public/MovieSceneDynamicBinding.h b/Engine/Source/Runtime/MovieScene/Public/MovieSceneDynamicBinding.h
new file mode 100644
--- /dev/null
+++ b/Engine/Source/Runtime/MovieScene/Public/MovieSceneDynamicBinding.h
@@ -0,0 +1,112 @@
+// Copyright Epic Games, Inc. All Rights Reserved.
+
+#pragma once
+
+#include "CoreTypes.h"
+#include "EntitySystem/MovieSceneSequenceInstanceHandle.h"
+#include "Kismet/BlueprintFunctionLibrary.h"
+#include "MovieSceneSequenceID.h"
+#include "UObject/ObjectMacros.h"
+#include "UObject/ObjectPtr.h"
+
+#include "MovieSceneDynamicBinding.generated.h"
+
+class UMovieScene;
+class UMovieSceneSequence;
+class UMovieSceneEntitySystemLinker;
+
+/** Value definition for any type-agnostic variable (exported as text) */
+USTRUCT(BlueprintType)
+struct MOVIESCENE_API FMovieSceneDynamicBindingPayloadVariable
+{
+	GENERATED_BODY()
+
+	UPROPERTY(EditAnywhere, Category = "Sequencer|Dynamic Binding")
+	FString Value;
+};
+
+/**
+ * Data for a dynamic binding endpoint call.
+ */
+USTRUCT()
+struct MOVIESCENE_API FMovieSceneDynamicBinding
+{
+	GENERATED_BODY()
+
+	/** The function to call (normally a generated blueprint function on the sequence director) */
+	UPROPERTY()
+	TObjectPtr<UFunction> Function;
+
+	/** Property pointer for the function parameter that should receive the resolve params */
+	UPROPERTY()
+	TFieldPath<FProperty> ResolveParamsProperty;
+
+#if WITH_EDITORONLY_DATA
+
+	/** Array of payload variables to be added to the generated function */
+	UPROPERTY(EditAnywhere, Category = "Sequencer|Dynamic Binding")
+	TMap<FName, FMovieSceneDynamicBindingPayloadVariable> PayloadVariables;
+
+	/** Name of the generated blueprint function */
+	UPROPERTY(transient)
+	FName CompiledFunctionName;
+
+	/** Pin name for passing the resolve params */
+	UPROPERTY(EditAnywhere, Category="Sequencer|Dynamic Binding")
+	FName ResolveParamsPinName;
+
+	/** Endpoint node in the sequence director */
+	UPROPERTY(EditAnywhere, Category="Sequencer|Dynamic Binding")
+	TWeakObjectPtr<UObject> WeakEndpoint;
+
+#endif
+};
+
+/**
+ * Optional parameter struct for dynamic binding resolver functions.
+ */
+USTRUCT(BlueprintType)
+struct MOVIESCENE_API FMovieSceneDynamicBindingResolveParams
+{
+	GENERATED_BODY()
+
+	/** The sequence that contains the object binding being resolved */
+	UPROPERTY(EditAnywhere, Category="General")
+	TObjectPtr<UMovieSceneSequence> Sequence;
+
+	/** The ID of the object binding being resolved */
+	UPROPERTY(EditAnywhere, Category="General")
+	FGuid ObjectBindingID;
+
+	/** The root sequence */
+	UPROPERTY(EditAnywhere, Category="General")
+	TObjectPtr<UMovieSceneSequence> RootSequence;
+};
+
+/**
+ * Dummy structure for showing an FMovieSceneDynamicBinding inside a details view.
+ */
+USTRUCT()
+struct MOVIESCENE_API FMovieSceneDynamicBindingContainer
+{
+	GENERATED_BODY()
+
+	UPROPERTY(EditAnywhere, Category="Sequencer")
+	FMovieSceneDynamicBinding DynamicBinding;
+};
+
+/**
+ * Default dynamic binding resolver library, with several basic resolver functions.
+ */
+UCLASS(meta=(SequencerBindingResolverLibrary))
+class MOVIESCENE_API UBuiltInDynamicBindingResolverLibrary : public UBlueprintFunctionLibrary
+{
+	GENERATED_BODY()
+
+public:
+
+	/** Resolve the bound object to the player's pawn */
+	UFUNCTION(BlueprintPure, Category="Sequencer|Dynamic Binding", meta=(WorldContext="WorldContextObject"))
+	static UObject* ResolveToPlayerPawn(UObject* WorldContextObject, int32 PlayerControllerIndex = 0);
+};
+
diff --git a/Engine/Source/Runtime/MovieScene/Public/MovieScenePossessable.h b/Engine/Source/Runtime/MovieScene/Public/MovieScenePossessable.h
--- a/Engine/Source/Runtime/MovieScene/Public/MovieScenePossessable.h
+++ b/Engine/Source/Runtime/MovieScene/Public/MovieScenePossessable.h
@@ -7,6 +7,7 @@
 #include "CoreMinimal.h"
 #include "CoreTypes.h"
 #include "Misc/Guid.h"
+#include "MovieSceneDynamicBinding.h"
 #include "MovieSceneObjectBindingID.h"
 #include "MovieSceneSequenceID.h"
 #include "UObject/NameTypes.h"
@@ -147,6 +148,10 @@ public:
 	UPROPERTY(EditAnywhere, AdvancedDisplay, Category=Actor)
 	TArray<FName> Tags;
 
+	/** Optional user-defined possessable lookup information */
+	UPROPERTY(EditAnywhere, AdvancedDisplay, Category="Sequencer")
+	FMovieSceneDynamicBinding DynamicBinding;
+
 	/* Get the optional binding id for binding to a spawnable */
 	const FMovieSceneObjectBindingID& GetSpawnableObjectBindingID() const 
 	{
diff --git a/Engine/Source/Runtime/MovieScene/Public/MovieSceneSpawnable.h b/Engine/Source/Runtime/MovieScene/Public/MovieSceneSpawnable.h
--- a/Engine/Source/Runtime/MovieScene/Public/MovieSceneSpawnable.h
+++ b/Engine/Source/Runtime/MovieScene/Public/MovieSceneSpawnable.h
@@ -5,6 +5,8 @@
 #include "CoreMinimal.h"
 #include "UObject/ObjectMacros.h"
 #include "Misc/Guid.h"
+#include "MovieSceneDynamicBinding.h"
+
 #include "MovieSceneSpawnable.generated.h"
 
 struct FMovieSceneSequenceID;
@@ -247,6 +249,10 @@ public:
 	UPROPERTY(EditAnywhere, Category=Actor)
 	bool bNetAddressableName;
 
+	/** Optional user-defined spawning information */
+	UPROPERTY(EditAnywhere, Category="Sequencer")
+	FMovieSceneDynamicBinding DynamicBinding;
+
 private:
 
 	/** Unique identifier of the spawnable object. */
