Sequencer: Optimized bake keys algorithm by only calling ExpandToFrame once

diff --git a/Engine/Source/Editor/MovieSceneTools/Private/Channels/ChannelCurveModel.cpp b/Engine/Source/Editor/MovieSceneTools/Private/Channels/ChannelCurveModel.cpp
--- a/Engine/Source/Editor/MovieSceneTools/Private/Channels/ChannelCurveModel.cpp
+++ b/Engine/Source/Editor/MovieSceneTools/Private/Channels/ChannelCurveModel.cpp
@@ -356,12 +356,15 @@ void FChannelCurveModel<ChannelType, ChannelValue, KeyType>::AddKeys(TArrayView<
 		TArray<FKeyHandle> NewKeyHandles;
 		NewKeyHandles.SetNumUninitialized(InKeyPositions.Num());
 
+		FFrameNumber MinFrame = TNumericLimits<FFrameNumber>::Max();
+		FFrameNumber MaxFrame = TNumericLimits<FFrameNumber>::Min();
 		for (int32 Index = 0; Index < InKeyPositions.Num(); ++Index)
 		{
 			FKeyPosition Position = InKeyPositions[Index];
 
 			FFrameNumber Time = (Position.InputValue * TickResolution).RoundToFrame();
-			Section->ExpandToFrame(Time);
+			MinFrame = FMath::Min(MinFrame, Time);
+			MaxFrame = FMath::Max(MaxFrame, Time);
 
 			ChannelValue Value = (ChannelValue)(Position.OutputValue);
 
@@ -377,6 +380,12 @@ void FChannelCurveModel<ChannelType, ChannelValue, KeyType>::AddKeys(TArrayView<
 			}
 		}
 
+		if (InKeyPositions.Num() > 0)
+		{
+			Section->ExpandToFrame(MinFrame);
+			Section->ExpandToFrame(MaxFrame);
+		}
+
 		// We reuse SetKeyAttributes here as there is complex logic determining which parts of the attributes are valid to set.
 		// For now we need to duplicate the new key handle array due to API mismatch. This will auto calculate tangents if needed.
 		SetKeyAttributes(NewKeyHandles, InKeyAttributes);
