Sequencer: Remove pre-existing shots that are not present in an imported EDL

diff --git a/Engine/Source/Editor/MovieSceneTools/Private/MovieSceneTranslatorEDL.cpp b/Engine/Source/Editor/MovieSceneTools/Private/MovieSceneTranslatorEDL.cpp
--- a/Engine/Source/Editor/MovieSceneTools/Private/MovieSceneTranslatorEDL.cpp
+++ b/Engine/Source/Editor/MovieSceneTools/Private/MovieSceneTranslatorEDL.cpp
@@ -373,6 +373,8 @@ bool MovieSceneTranslatorEDL::ImportEDL(UMovieScene* InMovieScene, FFrameRate In
 		CinematicShotTrack = InMovieScene->AddTrack<UMovieSceneCinematicShotTrack>();
 	}
 
+	TSet<UMovieSceneSection*> RelevantSections;
+
 	for (FShotData ShotData : ShotDataArray)
 	{
 		if (ShotData.TrackType == FShotData::ETrackType::TT_Video)
@@ -422,6 +424,8 @@ bool MovieSceneTranslatorEDL::ImportEDL(UMovieScene* InMovieScene, FFrameRate In
 			// Conform this shot section
 			if (ShotSection)
 			{
+				RelevantSections.Add(ShotSection);
+
 				ShotSection->Modify();
 				ShotSection->Parameters.StartFrameOffset = ShotData.SourceInFrame;
 				ShotSection->SetRange(TRange<FFrameNumber>(ShotData.EditInFrame, ShotData.EditOutFrame));
@@ -429,6 +433,16 @@ bool MovieSceneTranslatorEDL::ImportEDL(UMovieScene* InMovieScene, FFrameRate In
 		}
 	}
 
+	// Remove irrelevant sections backwards to ensure prior indices remain valid
+	const TArray<UMovieSceneSection*>& Sections = CinematicShotTrack->GetAllSections();
+	for (int32 Index = Sections.Num()-1; Index >= 0; --Index) 
+	{
+		if (!RelevantSections.Contains(Sections[Index]))
+		{
+			CinematicShotTrack->RemoveSectionAt(Index);
+		}
+	}
+
 	return true;
 }
 
