Sequencer Import EDL tweaks.

- now supports KEY CLIP NAME
- parses the timecodes from the end of the line
- if the EDL contains a shot multiple times, subsequent shots will be created in Sequencer to match.

diff --git a/Engine/Source/Editor/MovieSceneTools/Private/MovieSceneTranslatorEDL.cpp b/Engine/Source/Editor/MovieSceneTools/Private/MovieSceneTranslatorEDL.cpp
--- a/Engine/Source/Editor/MovieSceneTools/Private/MovieSceneTranslatorEDL.cpp
+++ b/Engine/Source/Editor/MovieSceneTools/Private/MovieSceneTranslatorEDL.cpp
@@ -119,7 +119,7 @@ void ParseFromEDL(const FString& InputString, FFrameRate TickResolution, FFrameR
 		// 001 AX V C 00:00:00:00 00:00:12:02 00:00:07:20 00:00:12:03
 		if (!bFoundEventLine)
 		{
-			if (InputChars.Num() == 8)
+			if (InputChars.Num() >= 8)
 			{
 				EventName = InputChars[0];
 				AuxName = InputChars[1]; // Typically AX but unused in this case
@@ -166,10 +166,11 @@ void ParseFromEDL(const FString& InputString, FFrameRate TickResolution, FFrameR
 				if (TrackType != FShotData::ETrackType::TT_None &&
 					EditType != FShotData::EEditType::ET_None)
 				{
-					SourceInFrame  = SMPTEToFrame(InputChars[4], TickResolution, FrameRate);
-					SourceOutFrame = SMPTEToFrame(InputChars[5], TickResolution, FrameRate);
-					EditInFrame    = SMPTEToFrame(InputChars[6], TickResolution, FrameRate);
-					EditOutFrame   = SMPTEToFrame(InputChars[7], TickResolution, FrameRate);
+					// Look for timecodes at the end
+					SourceInFrame  = SMPTEToFrame(InputChars[InputChars.Num() - 4], TickResolution, FrameRate);
+					SourceOutFrame = SMPTEToFrame(InputChars[InputChars.Num() - 3], TickResolution, FrameRate);
+					EditInFrame    = SMPTEToFrame(InputChars[InputChars.Num() - 2], TickResolution, FrameRate);
+					EditOutFrame   = SMPTEToFrame(InputChars[InputChars.Num() - 1], TickResolution, FrameRate);
 
 					bFoundEventLine = true;
 					continue; // Go to the next line
@@ -179,18 +180,21 @@ void ParseFromEDL(const FString& InputString, FFrameRate TickResolution, FFrameR
 		
 		// Then look for:
 		// * FROM CLIP NAME: shot0010_001.avi
+		// * KEY CLIP NAME: shot0010_001.avi
 		else
 		{
 			if (InputChars.Num() == 5 &&
 				InputChars[0] == TEXT("*") &&
-				InputChars[1].ToUpper() == TEXT("FROM") &&
-				InputChars[2].ToUpper() == TEXT("CLIP") &&
+				( (InputChars[1].ToUpper() == TEXT("FROM") && InputChars[2].ToUpper() == TEXT("CLIP")) ||
+				  (InputChars[1].ToUpper() == TEXT("KEY") && InputChars[2].ToUpper() == TEXT("CLIP")) ) &&
 				InputChars[3].ToUpper() == TEXT("NAME:"))
 			{
 				ReelName = InputChars[4];
 
-				//@todo can't assume avi's written out
-				ReelName.LeftChopInline(4, false); // strip .avi
+				if (ReelName.EndsWith(TEXT(".avi")))
+				{
+					ReelName.LeftChopInline(4, false); // strip .avi
+				}
 
 				FString ElementName = MoveTemp(ReelName);
 				FString ElementPath = ElementName;
@@ -385,7 +389,7 @@ bool MovieSceneTranslatorEDL::ImportEDL(UMovieScene* InMovieScene, FFrameRate In
 			for (UMovieSceneSection* Section : CinematicShotTrack->GetAllSections())
 			{
 				UMovieSceneCinematicShotSection* CinematicShotSection = Cast<UMovieSceneCinematicShotSection>(Section);
-				if (CinematicShotSection != nullptr)
+				if (CinematicShotSection != nullptr && !RelevantSections.Contains(Section))
 				{
 					UMovieSceneSequence* ShotSequence = CinematicShotSection->GetSequence();
 				
