Sequencer: Fixed a crash where dragging an item from the Place Actor menu over a sub-scene section would cause a crash. Also fixed a UI issue where the preview for the section length shown was only as long as the first sub-sequence asset, not the longest asset in your drag/drop selection.

diff --git a/Engine/Source/Editor/MovieSceneTools/Private/TrackEditors/SubTrackEditor.cpp b/Engine/Source/Editor/MovieSceneTools/Private/TrackEditors/SubTrackEditor.cpp
--- a/Engine/Source/Editor/MovieSceneTools/Private/TrackEditors/SubTrackEditor.cpp
+++ b/Engine/Source/Editor/MovieSceneTools/Private/TrackEditors/SubTrackEditor.cpp
@@ -344,6 +344,7 @@ bool FSubTrackEditor::OnAllowDrop(const FDragDropEvent& DragDropEvent, FSequence
 
 	TSharedPtr<FAssetDragDropOp> DragDropOp = StaticCastSharedPtr<FAssetDragDropOp>( Operation );
 
+	TOptional<FFrameNumber> LongestLengthInFrames;
 	for (const FAssetData& AssetData : DragDropOp->GetAssets())
 	{
 		if (!MovieSceneToolHelpers::IsValidAsset(FocusedSequence, AssetData))
@@ -352,7 +353,7 @@ bool FSubTrackEditor::OnAllowDrop(const FDragDropEvent& DragDropEvent, FSequence
 		}
 
 		UMovieSceneSequence* Sequence = Cast<UMovieSceneSequence>(AssetData.GetAsset());
-		if (CanAddSubSequence(*Sequence))
+		if (Sequence && CanAddSubSequence(*Sequence))
 		{
 			FFrameRate TickResolution = SequencerPtr->GetFocusedTickResolution();
 
@@ -361,11 +362,18 @@ bool FSubTrackEditor::OnAllowDrop(const FDragDropEvent& DragDropEvent, FSequence
 				Sequence->GetMovieScene()->GetTickResolution());
 
 			FFrameNumber LengthInFrames = InnerDuration.ConvertTo(TickResolution).FrameNumber;
-			DragDropParams.FrameRange = TRange<FFrameNumber>(DragDropParams.FrameNumber, DragDropParams.FrameNumber + LengthInFrames);
-			return true;
+			
+			// Keep track of the longest sub-sequence asset we're trying to drop onto it for preview display purposes.
+			LongestLengthInFrames = FMath::Max(LongestLengthInFrames.Get(FFrameNumber(0)), LengthInFrames);
 		}
 	}
 
+	if (LongestLengthInFrames.IsSet())
+	{
+		DragDropParams.FrameRange = TRange<FFrameNumber>(DragDropParams.FrameNumber, DragDropParams.FrameNumber + LongestLengthInFrames.GetValue());
+		return true;
+	}
+
 	return false;
 }
 
