Sequencer: Compensate for transform origin when adding an object's transform to Sequencer.

Note, this already happens in TransformTrackEditor when adding keyframes.

diff --git a/Engine/Plugins/MovieScene/LevelSequenceEditor/Source/LevelSequenceEditor/Private/LevelSequenceEditorToolkit.cpp b/Engine/Plugins/MovieScene/LevelSequenceEditor/Source/LevelSequenceEditor/Private/LevelSequenceEditorToolkit.cpp
--- a/Engine/Plugins/MovieScene/LevelSequenceEditor/Source/LevelSequenceEditor/Private/LevelSequenceEditorToolkit.cpp
+++ b/Engine/Plugins/MovieScene/LevelSequenceEditor/Source/LevelSequenceEditor/Private/LevelSequenceEditorToolkit.cpp
@@ -24,6 +24,7 @@
 #include "Sections/MovieSceneSubSection.h"
 #include "ToolMenu.h"
 #include "Tracks/MovieSceneCinematicShotTrack.h"
+#include "Tracks/IMovieSceneTransformOrigin.h"
 #include "MovieSceneToolsProjectSettings.h"
 #include "MovieSceneToolHelpers.h"
 #include "ScopedTransaction.h"
@@ -336,6 +337,25 @@ void FLevelSequenceEditorToolkit::ExtendSequencerToolbar(FName InToolMenuName)
 	}
 }
 
+FTransform GetTransformOrigin(TSharedPtr<ISequencer> Sequencer)
+{
+	FTransform TransformOrigin;
+
+	const IMovieScenePlaybackClient* Client = Sequencer->GetPlaybackClient();
+	const UObject* InstanceData = Client ? Client->GetInstanceData() : nullptr;
+	const IMovieSceneTransformOrigin* RawInterface = Cast<const IMovieSceneTransformOrigin>(InstanceData);
+
+	const bool bHasInterface = RawInterface || (InstanceData && InstanceData->GetClass()->ImplementsInterface(UMovieSceneTransformOrigin::StaticClass()));
+	if (bHasInterface)
+	{
+		// Retrieve the current origin
+		TransformOrigin = RawInterface ? RawInterface->GetTransformOrigin() : IMovieSceneTransformOrigin::Execute_BP_GetTransformOrigin(InstanceData);
+	}
+
+	return TransformOrigin;
+}
+
+
 void FLevelSequenceEditorToolkit::AddDefaultTracksForActor(AActor& Actor, const FGuid Binding)
 {
 	// get focused movie scene
@@ -389,19 +409,22 @@ void FLevelSequenceEditorToolkit::AddDefaultTracksForActor(AActor& Actor, const
 		{
 			auto TransformSection = Cast<UMovieScene3DTransformSection>(NewSection);
 
-			FVector Location = Actor.GetActorLocation();
-			FRotator Rotation = Actor.GetActorRotation();
-			FVector Scale = Actor.GetActorScale();
+			FTransform Transform = Actor.GetTransform();
 
 			if (USceneComponent* SceneComponent = Cast<USceneComponent>(InComponent))
 			{
-				FTransform ActorRelativeTransform = SceneComponent->GetRelativeTransform();
+				Transform = SceneComponent->GetRelativeTransform();
 
-				Location = ActorRelativeTransform.GetTranslation();
-				Rotation = ActorRelativeTransform.GetRotation().Rotator();
-				Scale = ActorRelativeTransform.GetScale3D();
+				if (!SceneComponent->GetAttachParent())
+				{
+					Transform *= GetTransformOrigin(GetSequencer()).Inverse();
+				}
 			}
 
+			FVector Location = Transform.GetTranslation();
+			FRotator Rotation = Transform.GetRotation().Rotator();
+			FVector Scale = Transform.GetScale3D();
+
 			TArrayView<FMovieSceneDoubleChannel*> DoubleChannels = TransformSection->GetChannelProxy().GetChannels<FMovieSceneDoubleChannel>();
 			DoubleChannels[0]->SetDefault(Location.X);
 			DoubleChannels[1]->SetDefault(Location.Y);
diff --git a/Engine/Source/Editor/MovieSceneTools/Private/TrackEditors/TransformTrackEditor.cpp b/Engine/Source/Editor/MovieSceneTools/Private/TrackEditors/TransformTrackEditor.cpp
--- a/Engine/Source/Editor/MovieSceneTools/Private/TrackEditors/TransformTrackEditor.cpp
+++ b/Engine/Source/Editor/MovieSceneTools/Private/TrackEditors/TransformTrackEditor.cpp
@@ -941,6 +941,8 @@ FTransformData F3DTransformTrackEditor::RecomposeTransform(const FTransformData&
 		return InTransformData;
 	}
 
+	USceneComponent* SceneComponent = MovieSceneHelpers::SceneComponentFromRuntimeObject(AnimatedObject);
+
 	TGuardValue<FEntityManager*> DebugVizGuard(GEntityManagerForDebuggingVisualizers, &EntityLinker->EntityManager);
 
 	// We want the transform value contributed to by the given transform section.
@@ -954,6 +956,8 @@ FTransformData F3DTransformTrackEditor::RecomposeTransform(const FTransformData&
 	TArray<FMovieSceneEntityID> ImportedEntityIDs;
 	EvaluationTemplate.FindEntitiesFromOwner(Section, GetSequencer()->GetFocusedTemplateID(), ImportedEntityIDs);
 
+	FTransform CurrentTransform = SceneComponent->GetRelativeTransform();
+
 	if (ImportedEntityIDs.Num())
 	{
 		UMovieScenePropertyInstantiatorSystem* System = EntityLinker->FindSystem<UMovieScenePropertyInstantiatorSystem>();
@@ -961,7 +965,6 @@ FTransformData F3DTransformTrackEditor::RecomposeTransform(const FTransformData&
 		{
 			FBuiltInComponentTypes* BuiltInComponents = FBuiltInComponentTypes::Get();
 			FMovieSceneTracksComponentTypes* TrackComponents = FMovieSceneTracksComponentTypes::Get();
-			USceneComponent* SceneComponent = MovieSceneHelpers::SceneComponentFromRuntimeObject(AnimatedObject);
 
 			TArray<FMovieSceneEntityID> EntityIDs;
 			{
@@ -1014,23 +1017,22 @@ FTransformData F3DTransformTrackEditor::RecomposeTransform(const FTransformData&
 				}
 			}
 
-			FTransform CurrentTransform(
+			CurrentTransform = FTransform(
 				FRotator(CurrentTransformChannels[4], CurrentTransformChannels[5], CurrentTransformChannels[3]), // pitch yaw roll
 				FVector(CurrentTransformChannels[0], CurrentTransformChannels[1], CurrentTransformChannels[2]),
-				FVector(CurrentTransformChannels[6], CurrentTransformChannels[7], CurrentTransformChannels[8]));
-
-			// Account for the transform origin only if this is not parented because the transform origin is already being applied to the parent.
-			if (!SceneComponent->GetAttachParent())
-			{
-				CurrentTransform *= GetTransformOrigin().Inverse();
-			}
-
-			UpdateTransformBasedOnConstraint(CurrentTransform,SceneComponent);
-			return FTransformData(CurrentTransform.GetLocation(), CurrentTransform.GetRotation().Rotator(), CurrentTransform.GetScale3D());
+				FVector(CurrentTransformChannels[6], CurrentTransformChannels[7], CurrentTransformChannels[8]));			
 		}
 	}
 
-	return InTransformData;
+	// Account for the transform origin only if this is not parented because the transform origin is already being applied to the parent.
+	if (!SceneComponent->GetAttachParent())
+	{
+		CurrentTransform *= GetTransformOrigin().Inverse();
+	}
+
+	UpdateTransformBasedOnConstraint(CurrentTransform, SceneComponent);
+	
+	return FTransformData(CurrentTransform.GetLocation(), CurrentTransform.GetRotation().Rotator(), CurrentTransform.GetScale3D());
 }
 
 void F3DTransformTrackEditor::ProcessKeyOperation(FFrameNumber InKeyTime, const UE::Sequencer::FKeyOperation& Operation, ISequencer& InSequencer)
