Sequencer: Select camera even if hidden in the camera cuts menu

diff --git a/Engine/Source/Editor/MovieSceneTools/Private/Sections/CameraCutSection.cpp b/Engine/Source/Editor/MovieSceneTools/Private/Sections/CameraCutSection.cpp
--- a/Engine/Source/Editor/MovieSceneTools/Private/Sections/CameraCutSection.cpp
+++ b/Engine/Source/Editor/MovieSceneTools/Private/Sections/CameraCutSection.cpp
@@ -232,7 +232,10 @@ FText FCameraCutSection::HandleThumbnailTextBlockText() const
 
 void FCameraCutSection::HandleSelectCameraMenuEntryExecute(AActor* InCamera)
 {
-	GEditor->SelectActor(InCamera, true, true);
+	const bool bInSelected = true;
+	const bool bNotify = true;
+	const bool bSelectEventIfHidden = true;
+	GEditor->SelectActor(InCamera, bInSelected, bNotify, bSelectEventIfHidden);
 }
 
 void FCameraCutSection::HandleSetCameraMenuEntryExecute(AActor* InCamera)
