Sequencer: Delay regenerating waveforms during playback.

This delay also happens already for regenerating thumbnails

diff --git a/Engine/Source/Editor/MovieSceneTools/Private/TrackEditors/AudioTrackEditor.cpp b/Engine/Source/Editor/MovieSceneTools/Private/TrackEditors/AudioTrackEditor.cpp
--- a/Engine/Source/Editor/MovieSceneTools/Private/TrackEditors/AudioTrackEditor.cpp
+++ b/Engine/Source/Editor/MovieSceneTools/Private/TrackEditors/AudioTrackEditor.cpp
@@ -776,6 +776,20 @@ int32 FAudioSection::OnPaintSection( FSequencerSectionPainter& Painter ) const
 
 void FAudioSection::Tick( const FGeometry& AllottedGeometry, const FGeometry& ParentGeometry, const double InCurrentTime, const float InDeltaTime )
 {
+	// Defer regenerating waveforms if playing or scrubbing
+	TSharedPtr<ISequencer> SequencerPin = Sequencer.Pin();
+	if (!SequencerPin.IsValid())
+	{
+		return;
+	}
+	
+	EMovieScenePlayerStatus::Type PlaybackState = SequencerPin->GetPlaybackStatus();
+
+	if (PlaybackState == EMovieScenePlayerStatus::Playing || PlaybackState == EMovieScenePlayerStatus::Scrubbing)
+	{
+		return;
+	}
+
 	UMovieSceneAudioSection* AudioSection = Cast<UMovieSceneAudioSection>(&Section);
 	UMovieSceneTrack* Track = Section.GetTypedOuter<UMovieSceneTrack>();
 
@@ -816,11 +830,7 @@ void FAudioSection::Tick( const FGeometry& AllottedGeometry, const FGeometry& Pa
 			float DisplayScale = XSize / DrawRange.Size<float>();
 
 			// Use the view range if possible, as it's much more stable than using the texture size and draw range
-			TSharedPtr<ISequencer> SequencerPin = Sequencer.Pin();
-			if (SequencerPin.IsValid())
-			{
-				DisplayScale = SequencerPin->GetViewRange().Size<float>() / ParentGeometry.GetLocalSize().X;
-			}
+			DisplayScale = SequencerPin->GetViewRange().Size<float>() / ParentGeometry.GetLocalSize().X;	
 
 			RegenerateWaveforms(DrawRange, XOffset, XSize, Track->GetColorTint(), DisplayScale);
 			StoredSoundWave = SoundWave;
