CIS Static analysis error.
Fix for static analysis that showed up in UE5Main.

diff --git a/Engine/Source/Editor/MovieSceneTools/Private/Constraints/MovieSceneConstraintChannelHelper.cpp b/Engine/Source/Editor/MovieSceneTools/Private/Constraints/MovieSceneConstraintChannelHelper.cpp
--- a/Engine/Source/Editor/MovieSceneTools/Private/Constraints/MovieSceneConstraintChannelHelper.cpp
+++ b/Engine/Source/Editor/MovieSceneTools/Private/Constraints/MovieSceneConstraintChannelHelper.cpp
@@ -1047,7 +1047,7 @@ FConstraintSections FMovieSceneConstraintChannelHelper::GetConstraintSectionAndC
 	const FFrameNumber& EndFrame,
 	TArray<FFrameNumber>& OutFramesToBake)
 {
-	if (!InConstraint && !InConstraint->ChildTRSHandle && !InConstraint->ParentTRSHandle)
+	if ((InConstraint == nullptr) || (InConstraint->ChildTRSHandle == nullptr) || (InConstraint->ParentTRSHandle == nullptr))
 	{
 		return;
 	}
