Sequencer: Snapper: Changed default to use keys for snapping and python test doesn't have keys, also had to expose the enum to bp.
#trivial

diff --git a/Engine/Source/Editor/MovieSceneTools/Public/BakingAnimationKeySettings.h b/Engine/Source/Editor/MovieSceneTools/Public/BakingAnimationKeySettings.h
--- a/Engine/Source/Editor/MovieSceneTools/Public/BakingAnimationKeySettings.h
+++ b/Engine/Source/Editor/MovieSceneTools/Public/BakingAnimationKeySettings.h
@@ -5,7 +5,7 @@
 #include "BakingAnimationKeySettings.generated.h"
 
 
-UENUM()
+UENUM(Blueprintable)
 enum class EBakingKeySettings : uint8
 {
 	KeysOnly UMETA(DisplayName = "Keys Only"),
