Sequencer: CinematicShotTrackEditor now derives from SubTrackEditor

- Consolidated duplicate code, deprecated CinematicShotTrackEditor API in favor of SubTrackEditor
- SubTrackEditor now has takes functionality (ie. new take, switch take)
- Takes naming conventions have been expanded so that if there's no take number, it's considered take 0, ie. MyShot will have subsequence takes named as MyShot_01, MyShot_02
- Fixed Insert Shot so that it correctly finds the next and previous shot numbers and finds a shot number in between
- Extracted FindNextOrPreviousShot into new helper functions, FindNextSection and FindPreviousSection
- Remove obsolete Insert Filler

diff --git a/Engine/Source/Editor/MovieSceneTools/Private/MovieSceneToolHelpers.cpp b/Engine/Source/Editor/MovieSceneTools/Private/MovieSceneToolHelpers.cpp
--- a/Engine/Source/Editor/MovieSceneTools/Private/MovieSceneToolHelpers.cpp
+++ b/Engine/Source/Editor/MovieSceneTools/Private/MovieSceneToolHelpers.cpp
@@ -341,10 +341,14 @@ bool MovieSceneToolHelpers::ParseShotName(const FString& InShotName, FString& Sh
 				ShotPrefix = ShotName.Left(LastSlashPos);
 			}
 			
-			ShotNumber = INDEX_NONE; // Nullify the shot number since we only have a shot prefix
-			TakeNumber = FCString::Atoi(*ShotName.RightChop(LastSlashPos+1));
-			TakeNumberDigits = ShotName.Len() - (LastSlashPos+1);
-			return true;
+			FString TakeStr = ShotName.RightChop(LastSlashPos + 1);
+			if (TakeStr.IsNumeric())
+			{
+				ShotNumber = INDEX_NONE; // Nullify the shot number since we only have a shot prefix
+				TakeNumber = FCString::Atoi(*TakeStr);
+				TakeNumberDigits = ShotName.Len() - (LastSlashPos+1);
+				return true;
+			}
 		}
 	}
 
@@ -353,6 +357,15 @@ bool MovieSceneToolHelpers::ParseShotName(const FString& InShotName, FString& Sh
 		TakeNumber = ParsedTakeNumber.GetValue();
 	}
 
+	if (FirstShotNumberIndex == INDEX_NONE)
+	{
+		ShotPrefix = InShotName;
+		ShotNumber = INDEX_NONE; // Nullify the shot number since we only have a shot prefix
+		TakeNumber = 0;
+		TakeNumberDigits = ProjectSettings->TakeNumDigits;
+		return true;
+	}
+
 	return FirstShotNumberIndex != INDEX_NONE;
 }
 
@@ -451,46 +464,42 @@ FString MovieSceneToolHelpers::GenerateNewShotPath(UMovieScene* SequenceMovieSce
 	return NewShotPath;
 }
 
-
 FString MovieSceneToolHelpers::GenerateNewShotName(const TArray<UMovieSceneSection*>& AllSections, FFrameNumber Time)
 {
 	const UMovieSceneToolsProjectSettings* ProjectSettings = GetDefault<UMovieSceneToolsProjectSettings>();
 
-	UMovieSceneCinematicShotSection* BeforeShot = nullptr;
-	UMovieSceneCinematicShotSection* NextShot = nullptr;
+	UMovieSceneSubSection* CurrentSection = Cast<UMovieSceneSubSection>(MovieSceneHelpers::FindSectionAtTime(AllSections, Time));
+	UMovieSceneSubSection* NextSection = Cast<UMovieSceneSubSection>(MovieSceneHelpers::FindNextSection(AllSections, Time));
 
-	FFrameNumber MinEndDiff = TNumericLimits<int32>::Max();
-	FFrameNumber MinStartDiff = TNumericLimits<int32>::Max();
+	if (!CurrentSection)
+	{
+		CurrentSection = Cast<UMovieSceneSubSection>(MovieSceneHelpers::FindPreviousSection(AllSections, Time));
+	}
 
-	for (auto Section : AllSections)
+	if (!NextSection)
 	{
-		if (Section->HasEndFrame() && Section->GetExclusiveEndFrame() >= Time)
-		{
-			FFrameNumber EndDiff = Section->GetExclusiveEndFrame() - Time;
-			if (MinEndDiff > EndDiff)
-			{
-				MinEndDiff = EndDiff;
-				BeforeShot = Cast<UMovieSceneCinematicShotSection>(Section);
-			}
-		}
-		if (Section->HasStartFrame() && Section->GetInclusiveStartFrame() <= Time)
-		{
-			FFrameNumber StartDiff = Time - Section->GetInclusiveStartFrame();
-			if (MinStartDiff > StartDiff)
-			{
-				MinStartDiff = StartDiff;
-				NextShot = Cast<UMovieSceneCinematicShotSection>(Section);
-			}
-		}
+		NextSection = CurrentSection;
 	}
-	
-	// There aren't any shots, let's create the first shot name
-	if (BeforeShot == nullptr || NextShot == nullptr)
+
+	FString NextSectionName = NextSection && NextSection->GetSequence() ? NextSection->GetSequence()->GetName() : FString();
+	if (const UMovieSceneCinematicShotSection* ShotSection = Cast<UMovieSceneCinematicShotSection>(NextSection))
+	{
+		NextSectionName = ShotSection->GetShotDisplayName();
+	}
+
+	FString CurrentSectionName = CurrentSection && CurrentSection->GetSequence() ? CurrentSection->GetSequence()->GetName() : FString();
+	if (const UMovieSceneCinematicShotSection* ShotSection = Cast<UMovieSceneCinematicShotSection>(CurrentSection))
+	{
+		CurrentSectionName = ShotSection->GetShotDisplayName();
+	}
+
+	// There aren't any sections, let's create the first shot name
+	if (NextSection == nullptr || CurrentSection == nullptr)
 	{
 		// Default case
 	}
 	// This is the last shot
-	else if (BeforeShot == NextShot)
+	else if (CurrentSection == NextSection)
 	{
 		FString NextShotPrefix = ProjectSettings->ShotPrefix;
 		uint32 NextShotNumber = ProjectSettings->FirstShotNumber;
@@ -498,20 +507,20 @@ FString MovieSceneToolHelpers::GenerateNewShotName(const TArray<UMovieSceneSecti
 		uint32 ShotNumberDigits = ProjectSettings->ShotNumDigits;
 		uint32 TakeNumberDigits = ProjectSettings->TakeNumDigits;
 
-		if (ParseShotName(NextShot->GetShotDisplayName(), NextShotPrefix, NextShotNumber, NextTakeNumber, ShotNumberDigits, TakeNumberDigits))
+		if (ParseShotName(NextSectionName, NextShotPrefix, NextShotNumber, NextTakeNumber, ShotNumberDigits, TakeNumberDigits))
 		{
-			uint32 NewShotNumber = NextShotNumber + ProjectSettings->ShotIncrement;
+			uint32 NewShotNumber = NextShotNumber != INDEX_NONE ? NextShotNumber + ProjectSettings->ShotIncrement : INDEX_NONE;
 			return ComposeShotName(NextShotPrefix, NewShotNumber, ProjectSettings->FirstTakeNumber, ShotNumberDigits, TakeNumberDigits);
 		}
 	}
 	// This is in between two shots
 	else 
 	{
-		FString BeforeShotPrefix = ProjectSettings->ShotPrefix;
-		uint32 BeforeShotNumber = ProjectSettings->FirstShotNumber;
-		uint32 BeforeTakeNumber = ProjectSettings->FirstTakeNumber;
-		uint32 BeforeShotNumberDigits = ProjectSettings->ShotNumDigits;
-		uint32 BeforeTakeNumberDigits = ProjectSettings->TakeNumDigits;
+		FString CurrentShotPrefix = ProjectSettings->ShotPrefix;
+		uint32 CurrentShotNumber = ProjectSettings->FirstShotNumber;
+		uint32 CurrentTakeNumber = ProjectSettings->FirstTakeNumber;
+		uint32 CurrentShotNumberDigits = ProjectSettings->ShotNumDigits;
+		uint32 CurrentTakeNumberDigits = ProjectSettings->TakeNumDigits;
 
 		FString NextShotPrefix = ProjectSettings->ShotPrefix;
 		uint32 NextShotNumber = ProjectSettings->FirstShotNumber;
@@ -519,13 +528,13 @@ FString MovieSceneToolHelpers::GenerateNewShotName(const TArray<UMovieSceneSecti
 		uint32 NextShotNumberDigits = ProjectSettings->ShotNumDigits;
 		uint32 NextTakeNumberDigits = ProjectSettings->TakeNumDigits;
 
-		if (ParseShotName(BeforeShot->GetShotDisplayName(), BeforeShotPrefix, BeforeShotNumber, BeforeTakeNumber, BeforeShotNumberDigits, BeforeTakeNumberDigits) &&
-			ParseShotName(NextShot->GetShotDisplayName(), NextShotPrefix, NextShotNumber, NextTakeNumber, NextShotNumberDigits, NextTakeNumberDigits))
+		if (ParseShotName(CurrentSectionName, CurrentShotPrefix, CurrentShotNumber, CurrentTakeNumber, CurrentShotNumberDigits, CurrentTakeNumberDigits) &&
+			ParseShotName(NextSectionName, NextShotPrefix, NextShotNumber, NextTakeNumber, NextShotNumberDigits, NextTakeNumberDigits))
 		{
-			if (BeforeShotNumber < NextShotNumber)
+			if (CurrentShotNumber < NextShotNumber)
 			{
-				uint32 NewShotNumber = BeforeShotNumber + ( (NextShotNumber - BeforeShotNumber) / 2); // what if we can't find one? or conflicts with another?
-				return ComposeShotName(BeforeShotPrefix, NewShotNumber, ProjectSettings->FirstTakeNumber, BeforeShotNumberDigits, BeforeTakeNumberDigits); // use before or next shot?
+				uint32 NewShotNumber = NextShotNumber != INDEX_NONE && CurrentShotNumber != INDEX_NONE ? CurrentShotNumber + ( (NextShotNumber - CurrentShotNumber) / 2) : INDEX_NONE; // what if we can't find one? or conflicts with another?
+				return ComposeShotName(CurrentShotPrefix, NewShotNumber, ProjectSettings->FirstTakeNumber, CurrentShotNumberDigits, CurrentTakeNumberDigits); // use before or next shot?
 			}
 		}
 	}
@@ -634,36 +643,33 @@ void MovieSceneToolHelpers::GatherTakes(const UMovieSceneSection* Section, TArra
 	uint32 TakeNumberDigits = 0;
 
 	FString SubSectionName = SubSection->GetSequence()->GetName();
-	if (SubSection->IsA<UMovieSceneCinematicShotSection>())
+	if (const UMovieSceneCinematicShotSection* ShotSection = Cast<UMovieSceneCinematicShotSection>(SubSection))
 	{
-		const UMovieSceneCinematicShotSection* ShotSection = Cast<UMovieSceneCinematicShotSection>(SubSection);
 		SubSectionName = ShotSection->GetShotDisplayName();
 	}
 
-	if (ParseShotName(SubSectionName, ShotPrefix, ShotNumber, OutCurrentTakeNumber, ShotNumberDigits, TakeNumberDigits))
+	ParseShotName(SubSectionName, ShotPrefix, ShotNumber, OutCurrentTakeNumber, ShotNumberDigits, TakeNumberDigits);
+
+	// Gather up all level sequence assets
+	FAssetRegistryModule& AssetRegistryModule = FModuleManager::LoadModuleChecked<FAssetRegistryModule>(TEXT("AssetRegistry"));
+	TArray<FAssetData> ObjectList;
+	AssetRegistryModule.Get().GetAssetsByClass(ULevelSequence::StaticClass()->GetClassPathName(), ObjectList);
+
+	for (auto AssetObject : ObjectList)
 	{
-		// Gather up all level sequence assets
-		FAssetRegistryModule& AssetRegistryModule = FModuleManager::LoadModuleChecked<FAssetRegistryModule>(TEXT("AssetRegistry"));
-		TArray<FAssetData> ObjectList;
-		AssetRegistryModule.Get().GetAssetsByClass(ULevelSequence::StaticClass()->GetClassPathName(), ObjectList);
+		FString AssetPackagePath = AssetObject.PackagePath.ToString();
 
-		for (auto AssetObject : ObjectList)
+		if (AssetPackagePath == ShotPackagePath)
 		{
-			FString AssetPackagePath = AssetObject.PackagePath.ToString();
+			FString AssetShotPrefix;
+			uint32 AssetShotNumber = INDEX_NONE;
+			uint32 AssetTakeNumber = INDEX_NONE;
 
-			if (AssetPackagePath == ShotPackagePath)
+			ParseShotName(AssetObject.AssetName.ToString(), AssetShotPrefix, AssetShotNumber, AssetTakeNumber, ShotNumberDigits, TakeNumberDigits);
+			
+			if (AssetShotPrefix == ShotPrefix && AssetShotNumber == ShotNumber)
 			{
-				FString AssetShotPrefix;
-				uint32 AssetShotNumber = INDEX_NONE;
-				uint32 AssetTakeNumber = INDEX_NONE;
-
-				if (ParseShotName(AssetObject.AssetName.ToString(), AssetShotPrefix, AssetShotNumber, AssetTakeNumber, ShotNumberDigits, TakeNumberDigits))
-				{
-					if (AssetShotPrefix == ShotPrefix && AssetShotNumber == ShotNumber)
-					{
-						AssetData.Add(AssetObject);
-					}
-				}
+				AssetData.Add(AssetObject);
 			}
 		}
 	}
@@ -692,42 +698,39 @@ bool MovieSceneToolHelpers::GetTakeNumber(const UMovieSceneSection* Section, FAs
 	uint32 TakeNumberDigits = 0;
 
 	FString SubSectionName = SubSection->GetSequence()->GetName();
-	if (SubSection->IsA<UMovieSceneCinematicShotSection>())
+	if (const UMovieSceneCinematicShotSection* ShotSection = Cast<UMovieSceneCinematicShotSection>(SubSection))
 	{
-		const UMovieSceneCinematicShotSection* ShotSection = Cast<UMovieSceneCinematicShotSection>(SubSection);
 		SubSectionName = ShotSection->GetShotDisplayName();
 	}
 
-	if (ParseShotName(SubSectionName, ShotPrefix, ShotNumber, TakeNumberDummy, ShotNumberDigits, TakeNumberDigits))
-	{
-		// Gather up all level sequence assets
-		FAssetRegistryModule& AssetRegistryModule = FModuleManager::LoadModuleChecked<FAssetRegistryModule>(TEXT("AssetRegistry"));
-		TArray<FAssetData> ObjectList;
-		AssetRegistryModule.Get().GetAssetsByClass(ULevelSequence::StaticClass()->GetClassPathName(), ObjectList);
+	ParseShotName(SubSectionName, ShotPrefix, ShotNumber, TakeNumberDummy, ShotNumberDigits, TakeNumberDigits);
+	
+	// Gather up all level sequence assets
+	FAssetRegistryModule& AssetRegistryModule = FModuleManager::LoadModuleChecked<FAssetRegistryModule>(TEXT("AssetRegistry"));
+	TArray<FAssetData> ObjectList;
+	AssetRegistryModule.Get().GetAssetsByClass(ULevelSequence::StaticClass()->GetClassPathName(), ObjectList);
 
-		for (auto AssetObject : ObjectList)
+	for (auto AssetObject : ObjectList)
+	{
+		if (AssetObject == AssetData)
 		{
-			if (AssetObject == AssetData)
+			FString AssetPackagePath = AssetObject.PackagePath.ToString();
+			int32 AssetLastSlashPos = INDEX_NONE;
+			AssetPackagePath.FindLastChar(TCHAR('/'), AssetLastSlashPos);
+			AssetPackagePath.LeftInline(AssetLastSlashPos, false);
+
+			if (AssetPackagePath == ShotPackagePath)
 			{
-				FString AssetPackagePath = AssetObject.PackagePath.ToString();
-				int32 AssetLastSlashPos = INDEX_NONE;
-				AssetPackagePath.FindLastChar(TCHAR('/'), AssetLastSlashPos);
-				AssetPackagePath.LeftInline(AssetLastSlashPos, false);
+				FString AssetShotPrefix;
+				uint32 AssetShotNumber = INDEX_NONE;
+				uint32 AssetTakeNumber = INDEX_NONE;
 
-				if (AssetPackagePath == ShotPackagePath)
-				{
-					FString AssetShotPrefix;
-					uint32 AssetShotNumber = INDEX_NONE;
-					uint32 AssetTakeNumber = INDEX_NONE;
+				ParseShotName(AssetObject.AssetName.ToString(), AssetShotPrefix, AssetShotNumber, AssetTakeNumber, ShotNumberDigits, TakeNumberDigits);
 
-					if (ParseShotName(AssetObject.AssetName.ToString(), AssetShotPrefix, AssetShotNumber, AssetTakeNumber, ShotNumberDigits, TakeNumberDigits))
-					{
-						if (AssetShotPrefix == ShotPrefix && AssetShotNumber == ShotNumber)
-						{
-							OutTakeNumber = AssetTakeNumber;
-							return true;
-						}
-					}
+				if (AssetShotPrefix == ShotPrefix && AssetShotNumber == ShotNumber)
+				{
+					OutTakeNumber = AssetTakeNumber;
+					return true;
 				}
 			}
 		}
diff --git a/Engine/Source/Editor/MovieSceneTools/Private/Sections/CinematicShotSection.cpp b/Engine/Source/Editor/MovieSceneTools/Private/Sections/CinematicShotSection.cpp
--- a/Engine/Source/Editor/MovieSceneTools/Private/Sections/CinematicShotSection.cpp
+++ b/Engine/Source/Editor/MovieSceneTools/Private/Sections/CinematicShotSection.cpp
@@ -312,34 +312,34 @@ void FCinematicShotSection::BuildSectionContextMenu(FMenuBuilder& MenuBuilder, c
 		MenuBuilder.AddSubMenu(
 			LOCTEXT("TakesMenu", "Takes"),
 			LOCTEXT("TakesMenuTooltip", "Shot takes"),
-			FNewMenuDelegate::CreateLambda([=](FMenuBuilder& InMenuBuilder){ AddTakesMenu(InMenuBuilder); }));
+			FNewMenuDelegate::CreateLambda([this, &SectionObject](FMenuBuilder& InMenuBuilder) { CinematicShotTrackEditor.Pin()->AddTakesMenu(&SectionObject, InMenuBuilder); }));
 
 		MenuBuilder.AddMenuEntry(
 			LOCTEXT("NewTake", "New Take"),
 			FText::Format(LOCTEXT("NewTakeTooltip", "Create a new take for {0}"), FText::FromString(SectionObject.GetShotDisplayName())),
 			FSlateIcon(),
-			FUIAction(FExecuteAction::CreateSP(CinematicShotTrackEditor.Pin().ToSharedRef(), &FCinematicShotTrackEditor::NewTake, &SectionObject))
+			FUIAction(FExecuteAction::CreateSP(CinematicShotTrackEditor.Pin().ToSharedRef(), &FCinematicShotTrackEditor::CreateNewTake, Cast<UMovieSceneSubSection>(&SectionObject)))
 		);
 
 		MenuBuilder.AddMenuEntry(
 			LOCTEXT("InsertNewShot", "Insert Shot"),
 			LOCTEXT("InsertNewShotTooltip", "Insert a new shot at the current time"),
 			FSlateIcon(),
-			FUIAction(FExecuteAction::CreateSP(CinematicShotTrackEditor.Pin().ToSharedRef(), &FCinematicShotTrackEditor::InsertShot))
+			FUIAction(FExecuteAction::CreateSP(CinematicShotTrackEditor.Pin().ToSharedRef(), &FCinematicShotTrackEditor::InsertSection, Cast<UMovieSceneTrack>(SectionObject.GetOuter())))
 		);
 
 		MenuBuilder.AddMenuEntry(
 			LOCTEXT("DuplicateShot", "Duplicate Shot"),
 			FText::Format(LOCTEXT("DuplicateShotTooltip", "Duplicate {0} to create a new shot"), FText::FromString(SectionObject.GetShotDisplayName())),
 			FSlateIcon(),
-			FUIAction(FExecuteAction::CreateSP(CinematicShotTrackEditor.Pin().ToSharedRef(), &FCinematicShotTrackEditor::DuplicateShot, &SectionObject))
+			FUIAction(FExecuteAction::CreateSP(CinematicShotTrackEditor.Pin().ToSharedRef(), &FCinematicShotTrackEditor::DuplicateSection, Cast<UMovieSceneSubSection>(&SectionObject)))
 		);
 
 		MenuBuilder.AddMenuEntry(
 			LOCTEXT("RenderShot", "Render Shot"),
 			FText::Format(LOCTEXT("RenderShotTooltip", "Render shot movie"), FText::FromString(SectionObject.GetShotDisplayName())),
 			FSlateIcon(),
-			FUIAction(FExecuteAction::CreateLambda([this, &SectionObject]() 
+			FUIAction(FExecuteAction::CreateLambda([this, &SectionObject]()
 				{
 					TArray<UMovieSceneCinematicShotSection*> ShotSections;
 					TArray<UMovieSceneSection*> Sections;
@@ -371,43 +371,6 @@ void FCinematicShotSection::BuildSectionContextMenu(FMenuBuilder& MenuBuilder, c
 	MenuBuilder.EndSection();
 }
 
-void FCinematicShotSection::AddTakesMenu(FMenuBuilder& MenuBuilder)
-{
-	TArray<FAssetData> AssetData;
-	uint32 CurrentTakeNumber = INDEX_NONE;
-	const UMovieSceneCinematicShotSection& SectionObject = GetSectionObjectAs<UMovieSceneCinematicShotSection>();
-	MovieSceneToolHelpers::GatherTakes(&SectionObject, AssetData, CurrentTakeNumber);
-
-	AssetData.Sort([&SectionObject](const FAssetData &A, const FAssetData &B) {
-		uint32 TakeNumberA = INDEX_NONE;
-		uint32 TakeNumberB = INDEX_NONE;
-		if (MovieSceneToolHelpers::GetTakeNumber(&SectionObject, A, TakeNumberA) && MovieSceneToolHelpers::GetTakeNumber(&SectionObject, B, TakeNumberB))
-		{
-			return TakeNumberA < TakeNumberB;
-		}
-		return true;
-	});
-
-	for (auto ThisAssetData : AssetData)
-	{
-		uint32 TakeNumber = INDEX_NONE;
-		if (MovieSceneToolHelpers::GetTakeNumber(&SectionObject, ThisAssetData, TakeNumber))
-		{
-			UObject* TakeObject = ThisAssetData.GetAsset();
-			
-			if (TakeObject)
-			{
-				MenuBuilder.AddMenuEntry(
-					FText::Format(LOCTEXT("TakeNumber", "Take {0}"), FText::AsNumber(TakeNumber)),
-					FText::Format(LOCTEXT("TakeNumberTooltip", "Switch to {0}"), FText::FromString(TakeObject->GetPathName())),
-					TakeNumber == CurrentTakeNumber ? FSlateIcon(FAppStyle::GetAppStyleSetName(), "Sequencer.Star") : FSlateIcon(FAppStyle::GetAppStyleSetName(), "Sequencer.Empty"),
-					FUIAction(FExecuteAction::CreateSP(CinematicShotTrackEditor.Pin().ToSharedRef(), &FCinematicShotTrackEditor::SwitchTake, TakeObject))
-				);
-			}
-		}
-	}
-}
-
 /* FCinematicShotSection callbacks
  *****************************************************************************/
 
diff --git a/Engine/Source/Editor/MovieSceneTools/Private/Sections/CinematicShotSection.h b/Engine/Source/Editor/MovieSceneTools/Private/Sections/CinematicShotSection.h
--- a/Engine/Source/Editor/MovieSceneTools/Private/Sections/CinematicShotSection.h
+++ b/Engine/Source/Editor/MovieSceneTools/Private/Sections/CinematicShotSection.h
@@ -47,11 +47,6 @@ public:
 	virtual void HandleThumbnailTextBlockTextCommitted(const FText& NewThumbnailName, ETextCommit::Type CommitType) override;
 	virtual UCameraComponent* GetViewCamera() override;
 
-private:
-
-	/** Add shot takes menu */
-	void AddTakesMenu(FMenuBuilder& MenuBuilder);
-
 private:
 
 	/** The cinematic shot track editor that contains this section */
diff --git a/Engine/Source/Editor/MovieSceneTools/Private/TrackEditors/CinematicShotTrackEditor.cpp b/Engine/Source/Editor/MovieSceneTools/Private/TrackEditors/CinematicShotTrackEditor.cpp
--- a/Engine/Source/Editor/MovieSceneTools/Private/TrackEditors/CinematicShotTrackEditor.cpp
+++ b/Engine/Source/Editor/MovieSceneTools/Private/TrackEditors/CinematicShotTrackEditor.cpp
@@ -20,19 +20,14 @@
 #include "FCPXML/FCPXMLMovieSceneTranslator.h"
 #include "Sections/CinematicShotSection.h"
 #include "SequencerUtilities.h"
-#include "IAssetTools.h"
-#include "IContentBrowserSingleton.h"
-#include "ContentBrowserModule.h"
 #include "UObject/UObjectHash.h"
 #include "UObject/UObjectIterator.h"
 #include "LevelSequence.h"
 #include "AutomatedLevelSequenceCapture.h"
 #include "MovieSceneCaptureModule.h"
-#include "AssetToolsModule.h"
 #include "TrackEditorThumbnail/TrackEditorThumbnailPool.h"
 #include "MovieSceneToolsProjectSettings.h"
 #include "Editor.h"
-#include "DragAndDrop/AssetDragDropOp.h"
 #include "MovieSceneTimeHelpers.h"
 #include "Framework/Notifications/NotificationManager.h"
 #include "Widgets/Notifications/SNotificationList.h"
@@ -43,7 +38,7 @@
  *****************************************************************************/
 
 FCinematicShotTrackEditor::FCinematicShotTrackEditor(TSharedRef<ISequencer> InSequencer)
-	: FMovieSceneTrackEditor(InSequencer) 
+	: FSubTrackEditor(InSequencer) 
 {
 	ThumbnailPool = MakeShareable(new FTrackEditorThumbnailPool(InSequencer));
 }
@@ -73,20 +68,6 @@ void FCinematicShotTrackEditor::OnRelease()
 /* ISequencerTrackEditor interface
  *****************************************************************************/
 
-void FCinematicShotTrackEditor::BuildAddTrackMenu(FMenuBuilder& MenuBuilder)
-{
-	MenuBuilder.AddMenuEntry(
-		LOCTEXT("AddCinematicShotTrack", "Shot Track"),
-		LOCTEXT("AddCinematicShotTooltip", "Adds a shot track."),
-		FSlateIcon(FAppStyle::GetAppStyleSetName(), "Sequencer.Tracks.CinematicShot"),
-		FUIAction(
-			FExecuteAction::CreateRaw(this, &FCinematicShotTrackEditor::HandleAddCinematicShotTrackMenuEntryExecute),
-			FCanExecuteAction::CreateRaw(this, &FCinematicShotTrackEditor::HandleAddCinematicShotTrackMenuEntryCanExecute)
-		)
-	);
-}
-
-
 TSharedPtr<SWidget> FCinematicShotTrackEditor::BuildOutlinerEditWidget(const FGuid& ObjectBinding, UMovieSceneTrack* Track, const FBuildEditWidgetParams& Params)
 {
 	// Create a container edit box
@@ -97,7 +78,7 @@ TSharedPtr<SWidget> FCinematicShotTrackEditor::BuildOutlinerEditWidget(const FGu
 	.AutoWidth()
 	.VAlign(VAlign_Center)
 	[
-		FSequencerUtilities::MakeAddButton(LOCTEXT("CinematicShotText", "Shot"), FOnGetContent::CreateSP(this, &FCinematicShotTrackEditor::HandleAddCinematicShotComboButtonGetMenuContent), Params.NodeIsHovered, GetSequencer())
+		FSequencerUtilities::MakeAddButton(LOCTEXT("CinematicShotText", "Shot"), FOnGetContent::CreateSP(this, &FCinematicShotTrackEditor::HandleAddSubSequenceComboButtonGetMenuContent, Track), Params.NodeIsHovered, GetSequencer())
 	]
 
 	+ SHorizontalBox::Slot()
@@ -123,7 +104,6 @@ TSharedPtr<SWidget> FCinematicShotTrackEditor::BuildOutlinerEditWidget(const FGu
 	];
 }
 
-
 TSharedRef<ISequencerSection> FCinematicShotTrackEditor::MakeSectionInterface(UMovieSceneSection& SectionObject, UMovieSceneTrack& Track, FGuid ObjectBinding)
 {
 	check(SupportsType(SectionObject.GetOuter()->GetClass()));
@@ -132,66 +112,12 @@ TSharedRef<ISequencerSection> FCinematicShotTrackEditor::MakeSectionInterface(UM
 	return MakeShareable(new FCinematicShotSection(GetSequencer(), SectionObjectImpl, SharedThis(this), ThumbnailPool));
 }
 
-
-bool FCinematicShotTrackEditor::HandleAssetAdded(UObject* Asset, const FGuid& TargetObjectGuid)
-{
-	UMovieSceneSequence* Sequence = Cast<UMovieSceneSequence>(Asset);
-
-	if (Sequence == nullptr)
-	{
-		return false;
-	}
-
-	if (!SupportsSequence(Sequence))
-	{
-		return false;
-	}
-
-	// Only allow sequences with a camera cut track to be dropped as a shot. Otherwise, it'll be dropped as a subsequence.s
-	if (!Sequence->GetMovieScene()->GetCameraCutTrack())
-	{
-		return false;
-	}
-
-	if (Sequence->GetMovieScene()->GetPlaybackRange().IsEmpty())
-	{
-		FNotificationInfo Info(FText::Format(LOCTEXT("InvalidSequenceDuration", "Invalid level sequence {0}. The sequence has no duration."), Sequence->GetDisplayName()));
-		Info.bUseLargeFont = false;
-		FSlateNotificationManager::Get().AddNotification(Info);
-		return false;
-	}
-
-	if (CanAddSubSequence(*Sequence))
-	{
-		const FScopedTransaction Transaction(LOCTEXT("AddShot_Transaction", "Add Shot"));
-		
-		int32 RowIndex = INDEX_NONE;
-		AnimatablePropertyChanged(FOnKeyProperty::CreateRaw(this, &FCinematicShotTrackEditor::HandleSequenceAdded, Sequence, RowIndex));
-
-		return true;
-	}
-		
-	FNotificationInfo Info(FText::Format( LOCTEXT("InvalidSequence", "Invalid level sequence {0}. There could be a circular dependency."), Sequence->GetDisplayName()));	
-	Info.bUseLargeFont = false;
-	FSlateNotificationManager::Get().AddNotification(Info);
-
-	return false;
-}
-
-
 bool FCinematicShotTrackEditor::SupportsSequence(UMovieSceneSequence* InSequence) const
 {
 	ETrackSupport TrackSupported = InSequence ? InSequence->IsTrackSupported(UMovieSceneCinematicShotTrack::StaticClass()) : ETrackSupport::NotSupported;
 	return TrackSupported == ETrackSupport::Supported;
 }
 
-
-bool FCinematicShotTrackEditor::SupportsType(TSubclassOf<UMovieSceneTrack> Type) const
-{
-	return (Type == UMovieSceneCinematicShotTrack::StaticClass());
-}
-
-
 void FCinematicShotTrackEditor::Tick(float DeltaTime)
 {
 	TSharedPtr<ISequencer> SequencerPin = GetSequencer();
@@ -217,7 +143,6 @@ void FCinematicShotTrackEditor::Tick(float DeltaTime)
 	}
 }
 
-
 void FCinematicShotTrackEditor::BuildTrackContextMenu( FMenuBuilder& MenuBuilder, UMovieSceneTrack* Track )
 {
 	MenuBuilder.BeginSection("Import/Export", NSLOCTEXT("Sequencer", "ImportExportMenuSectionName", "Import/Export"));
@@ -253,517 +178,82 @@ void FCinematicShotTrackEditor::BuildTrackContextMenu( FMenuBuilder& MenuBuilder
 	MenuBuilder.EndSection();
 }
 
-
-const FSlateBrush* FCinematicShotTrackEditor::GetIconBrush() const
-{
-	return FAppStyle::GetBrush("Sequencer.Tracks.CinematicShot");
-}
-
-bool FCinematicShotTrackEditor::OnAllowDrop(const FDragDropEvent& DragDropEvent, FSequencerDragDropParams& DragDropParams)
-{
-	if (!DragDropParams.Track.IsValid() || !DragDropParams.Track.Get()->IsA(UMovieSceneCinematicShotTrack::StaticClass()))
-	{
-		return false;
-	}
-
-	TSharedPtr<FDragDropOperation> Operation = DragDropEvent.GetOperation();
-
-	if (!Operation.IsValid() || !Operation->IsOfType<FAssetDragDropOp>() )
-	{
-		return false;
-	}
-	
-	TSharedPtr<ISequencer> SequencerPtr = GetSequencer();
-	if (!SequencerPtr)
-	{
-		return false;
-	}
-
-	UMovieSceneSequence* FocusedSequence = SequencerPtr->GetFocusedMovieSceneSequence();
-	if (!FocusedSequence)
-	{
-		return false;
-	}
-
-	TSharedPtr<FAssetDragDropOp> DragDropOp = StaticCastSharedPtr<FAssetDragDropOp>( Operation );
-
-	for (const FAssetData& AssetData : DragDropOp->GetAssets())
-	{
-		if (!MovieSceneToolHelpers::IsValidAsset(FocusedSequence, AssetData))
-		{
-			continue;
-		}
-
-		if (UMovieSceneSequence* Sequence = Cast<UMovieSceneSequence>(AssetData.GetAsset()))
-		{
-			FFrameRate TickResolution = SequencerPtr->GetFocusedTickResolution();
-
-			const FQualifiedFrameTime InnerDuration = FQualifiedFrameTime(
-				UE::MovieScene::DiscreteSize(Sequence->GetMovieScene()->GetPlaybackRange()),
-				Sequence->GetMovieScene()->GetTickResolution());
-
-			FFrameNumber LengthInFrames = InnerDuration.ConvertTo(TickResolution).FrameNumber;
-			DragDropParams.FrameRange = TRange<FFrameNumber>(DragDropParams.FrameNumber, DragDropParams.FrameNumber + LengthInFrames);
-			return true;
-		}
-	}
-
-	return false;
-}
-
-
-FReply FCinematicShotTrackEditor::OnDrop(const FDragDropEvent& DragDropEvent, const FSequencerDragDropParams& DragDropParams)
-{
-	if (!DragDropParams.Track.IsValid() || !DragDropParams.Track.Get()->IsA(UMovieSceneCinematicShotTrack::StaticClass()))
-	{
-		return FReply::Unhandled();
-	}
-
-	TSharedPtr<FDragDropOperation> Operation = DragDropEvent.GetOperation();
-
-	if (!Operation.IsValid() || !Operation->IsOfType<FAssetDragDropOp>() )
-	{
-		return FReply::Unhandled();
-	}
-	
-	TSharedPtr<ISequencer> SequencerPtr = GetSequencer();
-	if (!SequencerPtr)
-	{
-		return FReply::Unhandled();
-	}
-
-	UMovieSceneSequence* FocusedSequence = SequencerPtr->GetFocusedMovieSceneSequence();
-	if (!FocusedSequence)
-	{
-		return FReply::Unhandled();
-	}
-
-	const FScopedTransaction Transaction(LOCTEXT("DropAssets", "Drop Assets"));
-
-	TSharedPtr<FAssetDragDropOp> DragDropOp = StaticCastSharedPtr<FAssetDragDropOp>( Operation );
-	
-	FMovieSceneTrackEditor::BeginKeying(DragDropParams.FrameNumber);
-
-	bool bAnyDropped = false;
-	for (const FAssetData& AssetData : DragDropOp->GetAssets())
-	{
-		if (!MovieSceneToolHelpers::IsValidAsset(FocusedSequence, AssetData))
-		{
-			continue;
-		}
-
-		UMovieSceneSequence* Sequence = Cast<UMovieSceneSequence>(AssetData.GetAsset());
-
-		if (Sequence)
-		{
-			AnimatablePropertyChanged( FOnKeyProperty::CreateRaw( this, &FCinematicShotTrackEditor::AddKeyInternal, Sequence, DragDropParams.RowIndex) );
-			
-			bAnyDropped = true;
-		}
-	}
-
-	FMovieSceneTrackEditor::EndKeying();
-
-	return bAnyDropped ? FReply::Handled() : FReply::Unhandled();
-}
-
 void FCinematicShotTrackEditor::InsertShot()
 {
-	FFrameTime NewShotStartTime = GetSequencer()->GetLocalTime().Time;
-
-	UMovieSceneCinematicShotTrack* CinematicShotTrack = FindOrCreateCinematicShotTrack();
-	FString NewShotName = MovieSceneToolHelpers::GenerateNewShotName(CinematicShotTrack->GetAllSections(), NewShotStartTime.FrameNumber);
-	FString NewSequencePath = MovieSceneToolHelpers::GenerateNewShotPath(GetSequencer()->GetFocusedMovieSceneSequence()->GetMovieScene(), NewShotName);
-
-	if (UMovieSceneSequence* NewSequence = MovieSceneToolHelpers::CreateSequence(NewShotName, NewSequencePath))
-	{
-		const FScopedTransaction Transaction(LOCTEXT("InsertShot_Transaction", "Insert Shot"));
-
-		int32 Duration = UE::MovieScene::DiscreteSize(NewSequence->GetMovieScene()->GetPlaybackRange());
-
-		if (UMovieSceneSubSection* NewShot = CinematicShotTrack->AddSequence(NewSequence, NewShotStartTime.FrameNumber, Duration))
-		{
-			NewShot->SetRowIndex(MovieSceneToolHelpers::FindAvailableRowIndex(CinematicShotTrack, NewShot));
-
-			GetSequencer()->NotifyMovieSceneDataChanged(EMovieSceneDataChangeType::MovieSceneStructureItemAdded);
-			GetSequencer()->EmptySelection();
-			GetSequencer()->SelectSection(NewShot);
-			GetSequencer()->ThrobSectionSelection();
-		}
-	}
-}
-
-
-void FCinematicShotTrackEditor::InsertFiller()
-{
-	const UMovieSceneToolsProjectSettings* ProjectSettings = GetDefault<UMovieSceneToolsProjectSettings>();
-
-	const FScopedTransaction Transaction(LOCTEXT("InsertFiller_Transaction", "Insert Filler"));
-
-	FQualifiedFrameTime CurrentTime = GetSequencer()->GetLocalTime();
-
-	UMovieSceneCinematicShotTrack* CinematicShotTrack = FindOrCreateCinematicShotTrack();
-
-	int32 Duration = (ProjectSettings->DefaultDuration * CurrentTime.Rate).FrameNumber.Value;
-
-	UMovieSceneSequence* NullSequence = nullptr;
-
-	UMovieSceneSubSection* NewSection = CinematicShotTrack->AddSequence(NullSequence, CurrentTime.Time.FrameNumber, Duration);
-
-	UMovieSceneCinematicShotSection* NewCinematicShotSection = Cast<UMovieSceneCinematicShotSection>(NewSection);
-
-	NewCinematicShotSection->SetShotDisplayName(FText(LOCTEXT("Filler", "Filler")).ToString());
-
-	GetSequencer()->NotifyMovieSceneDataChanged( EMovieSceneDataChangeType::MovieSceneStructureItemAdded );
-	GetSequencer()->EmptySelection();
-	GetSequencer()->SelectSection(NewSection);
-	GetSequencer()->ThrobSectionSelection();
+	UMovieScene* FocusedMovieScene = GetFocusedMovieScene();
+	UMovieSceneCinematicShotTrack* CinematicShotTrack = FocusedMovieScene ? FocusedMovieScene->FindTrack<UMovieSceneCinematicShotTrack>() : nullptr;
+	FSubTrackEditor::InsertSection(CinematicShotTrack);
 }
 
-
 void FCinematicShotTrackEditor::DuplicateShot(UMovieSceneCinematicShotSection* Section)
 {
-	UMovieSceneCinematicShotTrack* CinematicShotTrack = FindOrCreateCinematicShotTrack();
-
-	FFrameNumber StartTime = Section->HasStartFrame() ? Section->GetInclusiveStartFrame() : 0;
-	FString NewShotName = MovieSceneToolHelpers::GenerateNewShotName(CinematicShotTrack->GetAllSections(), StartTime);
-	FString NewSequencePath = FPaths::GetPath(Section->GetSequence()->GetPathName());
-
-
-
-	// Duplicate the shot and put it on the next available row
-	UMovieSceneSequence* NewSequence = MovieSceneToolHelpers::CreateSequence(NewShotName, NewSequencePath, Section);
-	if (NewSequence)
-	{
-		const FScopedTransaction Transaction(LOCTEXT("DuplicateShot_Transaction", "Duplicate Shot"));
-
-		int32 Duration = UE::MovieScene::DiscreteSize(Section->GetRange());
-
-		if (UMovieSceneSubSection* NewShot = CinematicShotTrack->AddSequence(NewSequence, StartTime, Duration))
-		{
-			NewShot->SetRange(Section->GetRange());
-			NewShot->SetRowIndex(MovieSceneToolHelpers::FindAvailableRowIndex(CinematicShotTrack, NewShot));
-			NewShot->Parameters.StartFrameOffset = Section->Parameters.StartFrameOffset;
-			NewShot->Parameters.TimeScale = Section->Parameters.TimeScale;
-			NewShot->SetPreRollFrames(Section->GetPreRollFrames());
-
-			GetSequencer()->NotifyMovieSceneDataChanged( EMovieSceneDataChangeType::MovieSceneStructureItemAdded );
-			GetSequencer()->EmptySelection();
-			GetSequencer()->SelectSection(NewShot);
-			GetSequencer()->ThrobSectionSelection();
-		}
-	}
+	FSubTrackEditor::DuplicateSection(Section);
 }
 
-
 void FCinematicShotTrackEditor::RenderShots(const TArray<UMovieSceneCinematicShotSection*>& Sections)
 {
 	GetSequencer()->RenderMovie(Sections);
 }
 
-
-void FCinematicShotTrackEditor::RenameShot(UMovieSceneCinematicShotSection* Section)
-{
-	//@todo
-}
-
-
 void FCinematicShotTrackEditor::NewTake(UMovieSceneCinematicShotSection* Section)
 {
-	FString ShotPrefix;
-	uint32 ShotNumber = INDEX_NONE;
-	uint32 TakeNumber = INDEX_NONE;
-	uint32 ShotNumberDigits = 0;
-	uint32 TakeNumberDigits = 0;
-	if (MovieSceneToolHelpers::ParseShotName(Section->GetShotDisplayName(), ShotPrefix, ShotNumber, TakeNumber, ShotNumberDigits, TakeNumberDigits))
-	{
-		TArray<FAssetData> AssetData;
-		uint32 CurrentTakeNumber = INDEX_NONE;
-		MovieSceneToolHelpers::GatherTakes(Section, AssetData, CurrentTakeNumber);
-		uint32 NewTakeNumber = CurrentTakeNumber;
-
-		for (auto ThisAssetData : AssetData)
-		{
-			uint32 ThisTakeNumber = INDEX_NONE;
-			if (MovieSceneToolHelpers::GetTakeNumber(Section, ThisAssetData, ThisTakeNumber))
-			{
-				if (ThisTakeNumber >= NewTakeNumber)
-				{
-					NewTakeNumber = ThisTakeNumber + 1;
-				}
-			}
-		}
-
-		FString NewShotName = MovieSceneToolHelpers::ComposeShotName(ShotPrefix, ShotNumber, NewTakeNumber, ShotNumberDigits, TakeNumberDigits);
-
-		TRange<FFrameNumber> NewShotRange         = Section->GetRange();
-		FFrameNumber         NewShotStartOffset   = Section->Parameters.StartFrameOffset;
-		float                NewShotTimeScale     = Section->Parameters.TimeScale;
-		int32                NewShotPrerollFrames = Section->GetPreRollFrames();
-		int32                NewRowIndex          = Section->GetRowIndex();
-		FFrameNumber         NewShotStartTime     = NewShotRange.GetLowerBound().IsClosed() ? UE::MovieScene::DiscreteInclusiveLower(NewShotRange) : 0;
-
-		UMovieSceneCinematicShotTrack* CinematicShotTrack = FindOrCreateCinematicShotTrack();
-		FString NewSequencePath = FPaths::GetPath(Section->GetSequence()->GetPathName());
-
-		if (UMovieSceneSequence* NewSequence = MovieSceneToolHelpers::CreateSequence(NewShotName, NewSequencePath, Section))
-		{
-			const FScopedTransaction Transaction(LOCTEXT("NewTake_Transaction", "New Take"));
-
-			int32 Duration = UE::MovieScene::DiscreteSize(Section->GetRange());
-
-			UMovieSceneSubSection* NewShot = CinematicShotTrack->AddSequence(NewSequence, NewShotStartTime, Duration);
-			CinematicShotTrack->RemoveSection(*Section);
-
-			NewShot->SetRange(NewShotRange);
-			NewShot->Parameters.StartFrameOffset = NewShotStartOffset;
-			NewShot->Parameters.TimeScale = NewShotTimeScale;
-			NewShot->SetPreRollFrames(NewShotPrerollFrames);
-			NewShot->SetRowIndex(NewRowIndex);
-
-			MovieSceneToolHelpers::SetTakeNumber(NewShot, NewTakeNumber);
-
-			GetSequencer()->NotifyMovieSceneDataChanged( EMovieSceneDataChangeType::MovieSceneStructureItemsChanged );
-			GetSequencer()->EmptySelection();
-			GetSequencer()->SelectSection(NewShot);
-			GetSequencer()->ThrobSectionSelection();
-		}
-	}
-}
-
-
-void FCinematicShotTrackEditor::SwitchTake(UObject* TakeObject)
-{
-	bool bSwitchedTake = false;
-
-	const FScopedTransaction Transaction(LOCTEXT("SwitchTake_Transaction", "Switch Take"));
-
-	TArray<UMovieSceneSection*> Sections;
-	GetSequencer()->GetSelectedSections(Sections);
-
-	for (int32 SectionIndex = 0; SectionIndex < Sections.Num(); ++SectionIndex)
-	{
-		if (!Sections[SectionIndex]->IsA<UMovieSceneSubSection>())
-		{
-			continue;
-		}
-
-		UMovieSceneSubSection* Section = Cast<UMovieSceneSubSection>(Sections[SectionIndex]);
-
-		if (TakeObject && TakeObject->IsA(UMovieSceneSequence::StaticClass()))
-		{
-			UMovieSceneSequence* MovieSceneSequence = CastChecked<UMovieSceneSequence>(TakeObject);
-			
-			UMovieSceneCinematicShotTrack* CinematicShotTrack = CastChecked<UMovieSceneCinematicShotTrack>(Section->GetOuter());
-
-			TRange<FFrameNumber> NewShotRange         = Section->GetRange();
-			FFrameNumber         NewShotStartOffset   = Section->Parameters.StartFrameOffset;
-			float                NewShotTimeScale     = Section->Parameters.TimeScale;
-			int32                NewShotPrerollFrames = Section->GetPreRollFrames();
-			int32                NewRowIndex          = Section->GetRowIndex();
-			FFrameNumber         NewShotStartTime     = NewShotRange.GetLowerBound().IsClosed() ? UE::MovieScene::DiscreteInclusiveLower(NewShotRange) : 0;
-			int32                NewShotRowIndex      = Section->GetRowIndex();
-
-			const int32 Duration = (NewShotRange.GetLowerBound().IsClosed() && NewShotRange.GetUpperBound().IsClosed() ) ? UE::MovieScene::DiscreteSize(NewShotRange) : 1;
-			UMovieSceneSubSection* NewShot = CinematicShotTrack->AddSequence(MovieSceneSequence, NewShotStartTime, Duration);
-
-			if (NewShot != nullptr)
-			{
-				CinematicShotTrack->RemoveSection(*Section);
-
-				NewShot->SetRange(NewShotRange);
-				NewShot->Parameters.StartFrameOffset = NewShotStartOffset;
-				NewShot->Parameters.TimeScale = NewShotTimeScale;
-				NewShot->SetPreRollFrames(NewShotPrerollFrames);
-				NewShot->SetRowIndex(NewShotRowIndex);
-				bSwitchedTake = true;
-			}
-		}
-	}
-
-	if (bSwitchedTake)
-	{
-		GetSequencer()->NotifyMovieSceneDataChanged( EMovieSceneDataChangeType::MovieSceneStructureItemsChanged );
-	}
+	FSubTrackEditor::CreateNewTake(Section);
 }
 
-
-/* FCinematicShotTrackEditor callbacks
+/* FSubTrackEditor
  *****************************************************************************/
 
-bool FCinematicShotTrackEditor::HandleAddCinematicShotTrackMenuEntryCanExecute() const
+FText FCinematicShotTrackEditor::GetSubTrackName() const
 {
-	UMovieScene* FocusedMovieScene = GetFocusedMovieScene();
-
-	return ((FocusedMovieScene != nullptr) && (FocusedMovieScene->FindTrack<UMovieSceneCinematicShotTrack>() == nullptr));
+	return LOCTEXT("ShotTrackName", "Shot");
 }
 
-
-void FCinematicShotTrackEditor::HandleAddCinematicShotTrackMenuEntryExecute()
+FText FCinematicShotTrackEditor::GetSubTrackToolTip() const
 {
-	UMovieSceneCinematicShotTrack* ShotTrack = FindOrCreateCinematicShotTrack();
-	if (ShotTrack)
-	{
-		if (GetSequencer().IsValid())
-		{
-			// Cinematic Shot Tracks can't be placed in folders, they're only allowed in the root.
-			GetSequencer()->OnAddTrack(ShotTrack, FGuid());
-		}
-	}
+	return LOCTEXT("ShotTrackToolTip", "A cinematic shot track.");
 }
 
-
-TSharedRef<SWidget> FCinematicShotTrackEditor::HandleAddCinematicShotComboButtonGetMenuContent()
+FName FCinematicShotTrackEditor::GetSubTrackBrushName() const
 {
-	FMenuBuilder MenuBuilder(true, nullptr);
-
-	MenuBuilder.AddMenuEntry(
-			LOCTEXT("InsertShot", "Insert Shot"),
-			LOCTEXT("InsertShotTooltip", "Insert new shot at current time"),
-			FSlateIcon(),
-			FUIAction(FExecuteAction::CreateSP(this, &FCinematicShotTrackEditor::InsertShot))
-	);
-
-	MenuBuilder.AddMenuEntry(
-			LOCTEXT("InsertFiller", "Insert Filler"),
-			LOCTEXT("InsertFillerTooltip", "Insert filler at current time"),
-			FSlateIcon(),
-			FUIAction(FExecuteAction::CreateSP(this, &FCinematicShotTrackEditor::InsertFiller))
-	);
-
-	UMovieSceneSequence* Sequence = GetSequencer() ? GetSequencer()->GetFocusedMovieSceneSequence() : nullptr;
-
-	FAssetPickerConfig AssetPickerConfig;
-	{
-		AssetPickerConfig.OnAssetSelected = FOnAssetSelected::CreateRaw( this, &FCinematicShotTrackEditor::HandleAddCinematicShotComboButtonMenuEntryExecute);
-		AssetPickerConfig.OnAssetEnterPressed = FOnAssetEnterPressed::CreateRaw( this, &FCinematicShotTrackEditor::HandleAddCinematicShotComboButtonMenuEntryEnterPressed);
-		AssetPickerConfig.bAllowNullSelection = false;
-		AssetPickerConfig.bAddFilterUI = true;
-		AssetPickerConfig.InitialAssetViewType = EAssetViewType::List;
-		AssetPickerConfig.Filter.ClassPaths.Add(FTopLevelAssetPath(TEXT("/Script/LevelSequence"), TEXT("LevelSequence")));
-		AssetPickerConfig.SaveSettingsName = TEXT("SequencerAssetPicker");
-		AssetPickerConfig.AdditionalReferencingAssets.Add(FAssetData(Sequence));
-	}
-
-	FContentBrowserModule& ContentBrowserModule = FModuleManager::Get().LoadModuleChecked<FContentBrowserModule>(TEXT("ContentBrowser"));
-
-	TSharedPtr<SBox> MenuEntry = SNew(SBox)
-		.WidthOverride(300.0f)
-		.HeightOverride(300.f)
-		[
-			ContentBrowserModule.Get().CreateAssetPicker(AssetPickerConfig)
-		];
-
-	MenuBuilder.AddWidget(MenuEntry.ToSharedRef(), FText::GetEmpty(), true);
-
-	return MenuBuilder.MakeWidget();
+	return TEXT("Sequencer.Tracks.CinematicShot");
 }
 
-
-void FCinematicShotTrackEditor::HandleAddCinematicShotComboButtonMenuEntryExecute(const FAssetData& AssetData)
+FString FCinematicShotTrackEditor::GetSubSectionDisplayName(const UMovieSceneSubSection* Section) const
 {
-	FSlateApplication::Get().DismissAllMenus();
-
-	UObject* SelectedObject = AssetData.GetAsset();
-
-	if (SelectedObject && SelectedObject->IsA(UMovieSceneSequence::StaticClass()))
-	{
-		UMovieSceneSequence* MovieSceneSequence = CastChecked<UMovieSceneSequence>(AssetData.GetAsset());
-
-		int32 RowIndex = INDEX_NONE;
-		AnimatablePropertyChanged( FOnKeyProperty::CreateRaw( this, &FCinematicShotTrackEditor::AddKeyInternal, MovieSceneSequence, RowIndex) );
-	}
+	return Cast<UMovieSceneCinematicShotSection>(Section)->GetShotDisplayName();
 }
 
-void FCinematicShotTrackEditor::HandleAddCinematicShotComboButtonMenuEntryEnterPressed(const TArray<FAssetData>& AssetData)
+TSubclassOf<UMovieSceneSubTrack> FCinematicShotTrackEditor::GetSubTrackClass() const
 {
-	if (AssetData.Num() > 0)
-	{
-		HandleAddCinematicShotComboButtonMenuEntryExecute(AssetData[0].GetAsset());
-	}
+	return UMovieSceneCinematicShotTrack::StaticClass();
 }
 
-FKeyPropertyResult FCinematicShotTrackEditor::AddKeyInternal(FFrameNumber KeyTime, UMovieSceneSequence* InMovieSceneSequence, int32 RowIndex)
-{	
-	FKeyPropertyResult KeyPropertyResult;
-
-	if (InMovieSceneSequence->GetMovieScene()->GetPlaybackRange().IsEmpty())
-	{
-		FNotificationInfo Info(FText::Format(LOCTEXT("InvalidSequenceDuration", "Invalid level sequence {0}. The sequence has no duration."), InMovieSceneSequence->GetDisplayName()));
-		Info.bUseLargeFont = false;
-		FSlateNotificationManager::Get().AddNotification(Info);
-		return KeyPropertyResult;
-	}
-
-	if (CanAddSubSequence(*InMovieSceneSequence))
-	{
-		UMovieSceneCinematicShotTrack* CinematicShotTrack = FindOrCreateCinematicShotTrack();
-		
-		const FFrameRate TickResolution = InMovieSceneSequence->GetMovieScene()->GetTickResolution();
-		const FQualifiedFrameTime InnerDuration = FQualifiedFrameTime(
-			UE::MovieScene::DiscreteSize(InMovieSceneSequence->GetMovieScene()->GetPlaybackRange()),
-			TickResolution);
-
-		const FFrameRate OuterFrameRate = CinematicShotTrack->GetTypedOuter<UMovieScene>()->GetTickResolution();
-		const int32      OuterDuration  = InnerDuration.ConvertTo(OuterFrameRate).FrameNumber.Value;
-
-		UMovieSceneSubSection* NewSection = CinematicShotTrack->AddSequenceOnRow(InMovieSceneSequence, KeyTime, OuterDuration, RowIndex);
-		KeyPropertyResult.bTrackModified = true;
-		KeyPropertyResult.SectionsCreated.Add(NewSection);
-
-		GetSequencer()->EmptySelection();
-		GetSequencer()->SelectSection(NewSection);
-		GetSequencer()->ThrobSectionSelection();
-
-		if (TickResolution != OuterFrameRate)
-		{
-			FNotificationInfo Info(FText::Format(LOCTEXT("TickResolutionMismatch", "The parent sequence has a different tick resolution {0} than the newly added sequence {1}"), OuterFrameRate.ToPrettyText(), TickResolution.ToPrettyText()));
-			Info.bUseLargeFont = false;
-			FSlateNotificationManager::Get().AddNotification(Info);
-		}
-
-		return KeyPropertyResult;
-	}
-
-	FNotificationInfo Info(FText::Format( LOCTEXT("InvalidSequence", "Invalid level sequence {0}. There could be a circular dependency."), InMovieSceneSequence->GetDisplayName()));	
-	Info.bUseLargeFont = false;
-	FSlateNotificationManager::Get().AddNotification(Info);
-
-	return KeyPropertyResult;
+TSharedRef<SWidget> FCinematicShotTrackEditor::HandleAddCinematicShotComboButtonGetMenuContent()
+{
+	UMovieScene* FocusedMovieScene = GetFocusedMovieScene();
+	UMovieSceneCinematicShotTrack* CinematicShotTrack = FocusedMovieScene ? FocusedMovieScene->FindTrack<UMovieSceneCinematicShotTrack>() : nullptr;
+	return HandleAddSubSequenceComboButtonGetMenuContent(CinematicShotTrack);
 }
 
-
 UMovieSceneCinematicShotTrack* FCinematicShotTrackEditor::FindOrCreateCinematicShotTrack()
 {
 	UMovieScene* FocusedMovieScene = GetFocusedMovieScene();
-	
-	if (FocusedMovieScene == nullptr)
-	{
-		return nullptr;
-	}
-
-	if (FocusedMovieScene->IsReadOnly())
-	{
-		return nullptr;
-	}
-
-	UMovieSceneCinematicShotTrack* CinematicShotTrack = FocusedMovieScene->FindTrack<UMovieSceneCinematicShotTrack>();
-	if (CinematicShotTrack != nullptr)
-	{
-		return CinematicShotTrack;
-	}
-
-	const FScopedTransaction Transaction(LOCTEXT("AddCinematicShotTrack_Transaction", "Add Cinematic Shot Track"));
-	FocusedMovieScene->Modify();
-
-	auto NewTrack = FocusedMovieScene->AddTrack<UMovieSceneCinematicShotTrack>();
-	ensure(NewTrack);
+	UMovieSceneCinematicShotTrack* CinematicShotTrack = FocusedMovieScene ? FocusedMovieScene->FindTrack<UMovieSceneCinematicShotTrack>() : nullptr;
+	return Cast<UMovieSceneCinematicShotTrack>(FindOrCreateSubTrack(FocusedMovieScene, CinematicShotTrack));
+}
 
-	GetSequencer()->NotifyMovieSceneDataChanged( EMovieSceneDataChangeType::MovieSceneStructureItemAdded );
+bool FCinematicShotTrackEditor::HandleAddSubTrackMenuEntryCanExecute() const
+{
+	UMovieScene* FocusedMovieScene = GetFocusedMovieScene();
 
-	return NewTrack;
+	return ((FocusedMovieScene != nullptr) && (FocusedMovieScene->FindTrack<UMovieSceneCinematicShotTrack>() == nullptr));
 }
 
+bool FCinematicShotTrackEditor::CanHandleAssetAdded(UMovieSceneSequence* Sequence) const
+{
+	// Only allow sequences with a camera cut track to be dropped as a shot. Otherwise, it'll be dropped as a subsequence.
+	return Sequence->GetMovieScene()->GetCameraCutTrack() != nullptr;
+}
 
 ECheckBoxState FCinematicShotTrackEditor::AreShotsLocked() const
 {
@@ -803,7 +293,6 @@ void FCinematicShotTrackEditor::OnLockShotsClicked(ECheckBoxState CheckBoxState)
 	GetSequencer()->ForceEvaluate();
 }
 
-
 FText FCinematicShotTrackEditor::GetLockShotsToolTip() const
 {
 	return AreShotsLocked() == ECheckBoxState::Checked ?
@@ -811,82 +300,12 @@ FText FCinematicShotTrackEditor::GetLockShotsToolTip() const
 		LOCTEXT("LockShots", "Lock Viewport to Shots");
 }
 
-
-bool FCinematicShotTrackEditor::CanAddSubSequence(const UMovieSceneSequence& Sequence) const
-{
-	// prevent adding ourselves and ensure we have a valid movie scene
-	UMovieSceneSequence* FocusedSequence = GetSequencer()->GetFocusedMovieSceneSequence();
-
-	if ((FocusedSequence == nullptr) || (FocusedSequence == &Sequence) || (FocusedSequence->GetMovieScene() == nullptr))
-	{
-		return false;
-	}
-
-	// ensure that the other sequence has a valid movie scene
-	UMovieScene* SequenceMovieScene = Sequence.GetMovieScene();
-
-	if (SequenceMovieScene == nullptr)
-	{
-		return false;
-	}
-
-	// make sure we are not contained in the other sequence (circular dependency)
-	// @todo sequencer: this check is not sufficient (does not prevent circular dependencies of 2+ levels)
-	UMovieSceneSubTrack* SequenceSubTrack = SequenceMovieScene->FindTrack<UMovieSceneSubTrack>();
-	if (SequenceSubTrack && SequenceSubTrack->ContainsSequence(*FocusedSequence, true))
-	{
-		return false;
-	}
-
-	UMovieSceneCinematicShotTrack* SequenceCinematicTrack = SequenceMovieScene->FindTrack<UMovieSceneCinematicShotTrack>();
-	if (SequenceCinematicTrack && SequenceCinematicTrack->ContainsSequence(*FocusedSequence, true))
-	{
-		return false;
-	}
-
-	return true;
-}
-
-
 void FCinematicShotTrackEditor::OnUpdateCameraCut(UObject* CameraObject, bool bJumpCut)
 {
 	// Keep track of the camera when it switches so that the thumbnail can be drawn with the correct camera
 	CinematicShotCamera = Cast<AActor>(CameraObject);
 }
 
-
-FKeyPropertyResult FCinematicShotTrackEditor::HandleSequenceAdded(FFrameNumber KeyTime, UMovieSceneSequence* Sequence, int32 RowIndex)
-{
-	FKeyPropertyResult KeyPropertyResult;
-
-	auto CinematicShotTrack = FindOrCreateCinematicShotTrack();
-
-	const FFrameRate TickResolution = Sequence->GetMovieScene()->GetTickResolution();
-	const FQualifiedFrameTime InnerDuration = FQualifiedFrameTime(
-		UE::MovieScene::DiscreteSize(Sequence->GetMovieScene()->GetPlaybackRange()),
-		TickResolution);
-
-	const FFrameRate OuterFrameRate = CinematicShotTrack->GetTypedOuter<UMovieScene>()->GetTickResolution();
-	const int32      OuterDuration  = InnerDuration.ConvertTo(OuterFrameRate).FrameNumber.Value;
-
-	UMovieSceneSubSection* NewSection = CinematicShotTrack->AddSequenceOnRow(Sequence, KeyTime, OuterDuration, RowIndex);
-	KeyPropertyResult.bTrackModified = true;
-	KeyPropertyResult.SectionsCreated.Add(NewSection);
-
-	GetSequencer()->EmptySelection();
-	GetSequencer()->SelectSection(NewSection);
-	GetSequencer()->ThrobSectionSelection();
-
-	if (TickResolution != OuterFrameRate)
-	{
-		FNotificationInfo Info(FText::Format(LOCTEXT("TickResolutionMismatch", "The parent sequence has a different tick resolution {0} than the newly added sequence {1}"), OuterFrameRate.ToPrettyText(), TickResolution.ToPrettyText()));
-		Info.bUseLargeFont = false;
-		FSlateNotificationManager::Get().AddNotification(Info);
-	}
-
-	return KeyPropertyResult;
-}
-
 UAutomatedLevelSequenceCapture* GetMovieSceneCapture()
 {
 	UAutomatedLevelSequenceCapture* MovieSceneCapture = Cast<UAutomatedLevelSequenceCapture>(IMovieSceneCaptureModule::Get().GetFirstActiveMovieSceneCapture());
diff --git a/Engine/Source/Editor/MovieSceneTools/Private/TrackEditors/SubTrackEditor.cpp b/Engine/Source/Editor/MovieSceneTools/Private/TrackEditors/SubTrackEditor.cpp
--- a/Engine/Source/Editor/MovieSceneTools/Private/TrackEditors/SubTrackEditor.cpp
+++ b/Engine/Source/Editor/MovieSceneTools/Private/TrackEditors/SubTrackEditor.cpp
@@ -49,9 +49,8 @@ class FSubSection
 {
 public:
 
-	FSubSection(TSharedPtr<ISequencer> InSequencer, UMovieSceneSection& InSection, const FText& InDisplayName, TSharedPtr<FSubTrackEditor> InSubTrackEditor)
+	FSubSection(TSharedPtr<ISequencer> InSequencer, UMovieSceneSection& InSection, TSharedPtr<FSubTrackEditor> InSubTrackEditor)
 		: TSubSectionMixin(InSequencer, *CastChecked<UMovieSceneSubSection>(&InSection))
-		, DisplayName(InDisplayName)
 		, SubTrackEditor(InSubTrackEditor)
 	{
 	}
@@ -69,23 +68,52 @@ public:
 	{
 		ISequencerSection::BuildSectionContextMenu(MenuBuilder, ObjectBinding);
 
-		MenuBuilder.AddSubMenu(
-			LOCTEXT("TakesMenu", "Takes"),
-			LOCTEXT("TakesMenuTooltip", "Sub section takes"),
-			FNewMenuDelegate::CreateLambda([=](FMenuBuilder& InMenuBuilder){ AddTakesMenu(InMenuBuilder); }));
-
-		MenuBuilder.AddMenuEntry(
-			LOCTEXT("PlayableDirectly_Label", "Playable Directly"),
-			LOCTEXT("PlayableDirectly_Tip", "When enabled, this sequence will also support being played directly outside of the root sequence. Disable this to save some memory on complex hierarchies of sequences."),
-			FSlateIcon(),
-			FUIAction(
-				FExecuteAction::CreateRaw(this, &FSubSection::TogglePlayableDirectly),
-				FCanExecuteAction::CreateLambda([]{ return true; }),
-				FGetActionCheckState::CreateRaw(this, &FSubSection::IsPlayableDirectly)
-			),
-			NAME_None,
-			EUserInterfaceActionType::ToggleButton
-		);
+		UMovieSceneSubSection* Section = &GetSubSectionObject();
+		
+		FString DisplayName = SubTrackEditor.Pin()->GetSubSectionDisplayName(Section);
+
+		MenuBuilder.BeginSection(NAME_None, LOCTEXT("SequenceMenuText", "Sequence"));
+		{
+			MenuBuilder.AddSubMenu(
+				LOCTEXT("TakesMenu", "Takes"),
+				LOCTEXT("TakesMenuTooltip", "Subsequence takes"),
+				FNewMenuDelegate::CreateLambda([=](FMenuBuilder& InMenuBuilder) { SubTrackEditor.Pin()->AddTakesMenu(Section, InMenuBuilder); }));
+
+			MenuBuilder.AddMenuEntry(
+				LOCTEXT("NewTake", "New Take"),
+				FText::Format(LOCTEXT("NewTakeTooltip", "Create a new take for {0}"), FText::FromString(DisplayName)),
+				FSlateIcon(),
+				FUIAction(FExecuteAction::CreateSP(SubTrackEditor.Pin().ToSharedRef(), &FSubTrackEditor::CreateNewTake, Section))
+			);
+
+			MenuBuilder.AddMenuEntry(
+				LOCTEXT("InsertNewSequence", "Insert Sequence"),
+				LOCTEXT("InsertNewSequenceTooltip", "Insert a new sequence at the current time"),
+				FSlateIcon(),
+				FUIAction(FExecuteAction::CreateSP(SubTrackEditor.Pin().ToSharedRef(), &FSubTrackEditor::InsertSection, Cast<UMovieSceneTrack>(Section->GetOuter())))
+			);
+
+			MenuBuilder.AddMenuEntry(
+				LOCTEXT("DuplicateSequence", "Duplicate Sequence"),
+				FText::Format(LOCTEXT("DuplicateSequenceTooltip", "Duplicate {0} to create a new sequence"), FText::FromString(DisplayName)),
+				FSlateIcon(),
+				FUIAction(FExecuteAction::CreateSP(SubTrackEditor.Pin().ToSharedRef(), &FSubTrackEditor::DuplicateSection, Section))
+			);
+		
+			MenuBuilder.AddMenuEntry(
+				LOCTEXT("PlayableDirectly_Label", "Playable Directly"),
+				LOCTEXT("PlayableDirectly_Tip", "When enabled, this sequence will also support being played directly outside of the root sequence. Disable this to save some memory on complex hierarchies of sequences."),
+				FSlateIcon(),
+				FUIAction(
+					FExecuteAction::CreateRaw(this, &FSubSection::TogglePlayableDirectly),
+					FCanExecuteAction::CreateLambda([]{ return true; }),
+					FGetActionCheckState::CreateRaw(this, &FSubSection::IsPlayableDirectly)
+				),
+				NAME_None,
+				EUserInterfaceActionType::ToggleButton
+			);
+		}
+		MenuBuilder.EndSection();
 	}
 
 	void TogglePlayableDirectly()
@@ -155,48 +183,6 @@ public:
 
 private:
 
-	void AddTakesMenu(FMenuBuilder& MenuBuilder)
-	{
-		TArray<FAssetData> AssetData;
-		uint32 CurrentTakeNumber = INDEX_NONE;
-		UMovieSceneSubSection& SectionObject = GetSubSectionObject();
-		MovieSceneToolHelpers::GatherTakes(&SectionObject, AssetData, CurrentTakeNumber);
-
-		AssetData.Sort([&SectionObject](const FAssetData &A, const FAssetData &B) {
-			uint32 TakeNumberA = INDEX_NONE;
-			uint32 TakeNumberB = INDEX_NONE;
-			if (MovieSceneToolHelpers::GetTakeNumber(&SectionObject, A, TakeNumberA) && MovieSceneToolHelpers::GetTakeNumber(&SectionObject, B, TakeNumberB))
-			{
-				return TakeNumberA < TakeNumberB;
-			}
-			return true;
-		});
-
-		for (auto ThisAssetData : AssetData)
-		{
-			uint32 TakeNumber = INDEX_NONE;
-			if (MovieSceneToolHelpers::GetTakeNumber(&SectionObject, ThisAssetData, TakeNumber))
-			{
-				UObject* TakeObject = ThisAssetData.GetAsset();
-
-				if (TakeObject)
-				{
-					MenuBuilder.AddMenuEntry(
-						FText::Format(LOCTEXT("TakeNumber", "Take {0}"), FText::AsNumber(TakeNumber)),
-						FText::Format(LOCTEXT("TakeNumberTooltip", "Switch to {0}"), FText::FromString(TakeObject->GetPathName())),
-						TakeNumber == CurrentTakeNumber ? FSlateIcon(FAppStyle::GetAppStyleSetName(), "Sequencer.Star") : FSlateIcon(FAppStyle::GetAppStyleSetName(), "Sequencer.Empty"),
-						FUIAction(FExecuteAction::CreateSP(SubTrackEditor.Pin().ToSharedRef(), &FSubTrackEditor::SwitchTake, TakeObject))
-					);
-				}
-			}
-		}
-	}
-
-private:
-
-	/** Display name of the section */
-	FText DisplayName;
-
 	/** The sub track editor that contains this section */
 	TWeakPtr<FSubTrackEditor> SubTrackEditor;
 };
@@ -216,11 +202,12 @@ FSubTrackEditor::FSubTrackEditor(TSharedRef<ISequencer> InSequencer)
 void FSubTrackEditor::BuildAddTrackMenu(FMenuBuilder& MenuBuilder)
 {
 	MenuBuilder.AddMenuEntry(
-		LOCTEXT("AddSubTrack", "Subsequences Track"),
-		LOCTEXT("AddSubTooltip", "Adds a new track that can contain other sequences."),
-		FSlateIcon(FAppStyle::GetAppStyleSetName(), "Sequencer.Tracks.Sub"),
+		FText::Join(FText::FromString(" "), GetSubTrackName(), LOCTEXT("TrackText", "Track")),
+		GetSubTrackToolTip(),
+		FSlateIcon(FAppStyle::GetAppStyleSetName(), GetSubTrackBrushName()),
 		FUIAction(
-			FExecuteAction::CreateRaw(this, &FSubTrackEditor::HandleAddSubTrackMenuEntryExecute)
+			FExecuteAction::CreateRaw(this, &FSubTrackEditor::HandleAddSubTrackMenuEntryExecute),
+			FCanExecuteAction::CreateRaw(this, &FSubTrackEditor::HandleAddSubTrackMenuEntryCanExecute)
 		)
 	);
 }
@@ -235,7 +222,7 @@ TSharedPtr<SWidget> FSubTrackEditor::BuildOutlinerEditWidget(const FGuid& Object
 	.AutoWidth()
 	.VAlign(VAlign_Center)
 	[
-		FSequencerUtilities::MakeAddButton(LOCTEXT("SubText", "Sequence"), FOnGetContent::CreateSP(this, &FSubTrackEditor::HandleAddSubSequenceComboButtonGetMenuContent, Track), Params.NodeIsHovered, GetSequencer())
+		FSequencerUtilities::MakeAddButton(GetSubTrackName(), FOnGetContent::CreateSP(this, &FSubTrackEditor::HandleAddSubSequenceComboButtonGetMenuContent, Track), Params.NodeIsHovered, GetSequencer())
 	];
 }
 
@@ -248,10 +235,16 @@ TSharedRef<ISequencerTrackEditor> FSubTrackEditor::CreateTrackEditor(TSharedRef<
 
 TSharedRef<ISequencerSection> FSubTrackEditor::MakeSectionInterface(UMovieSceneSection& SectionObject, UMovieSceneTrack& Track, FGuid ObjectBinding)
 {
-	return MakeShareable(new FSubSection(GetSequencer(), SectionObject, Track.GetDisplayName(), SharedThis(this)));
+	return MakeShareable(new FSubSection(GetSequencer(), SectionObject, SharedThis(this)));
 }
 
 
+bool FSubTrackEditor::CanHandleAssetAdded(UMovieSceneSequence* Sequence) const
+{
+	// Only allow sequences without a camera cut track to be dropped as a subsequence. Otherwise, it'll be dropped as a shot.
+	return Sequence->GetMovieScene()->GetCameraCutTrack() == nullptr;
+}
+
 bool FSubTrackEditor::HandleAssetAdded(UObject* Asset, const FGuid& TargetObjectGuid)
 {
 	UMovieSceneSequence* Sequence = Cast<UMovieSceneSequence>(Asset);
@@ -266,8 +259,7 @@ bool FSubTrackEditor::HandleAssetAdded(UObject* Asset, const FGuid& TargetObject
 		return false;
 	}
 
-	// Only allow sequences without a camera cut track to be dropped as a subsequence. Otherwise, it'll be dropped as a shot.
-	if (Sequence->GetMovieScene()->GetCameraCutTrack())
+	if (CanHandleAssetAdded(Sequence))
 	{
 		return false;
 	}
@@ -282,7 +274,7 @@ bool FSubTrackEditor::HandleAssetAdded(UObject* Asset, const FGuid& TargetObject
 
 	if (CanAddSubSequence(*Sequence))
 	{
-		const FScopedTransaction Transaction(LOCTEXT("AddSubSequence_Transaction", "Add Subsequence"));
+		const FScopedTransaction Transaction(FText::Join(FText::FromString(" "), LOCTEXT("AddText", "Add"), GetSubTrackName(), LOCTEXT("TrackText", "Track")));
 
 		int32 RowIndex = INDEX_NONE;
 		UMovieSceneTrack* Track = nullptr;
@@ -306,19 +298,22 @@ bool FSubTrackEditor::SupportsSequence(UMovieSceneSequence* InSequence) const
 
 bool FSubTrackEditor::SupportsType(TSubclassOf<UMovieSceneTrack> Type) const
 {
-	// We support sub movie scenes
-	return Type == UMovieSceneSubTrack::StaticClass();
+	return Type == GetSubTrackClass();
 }
 
 const FSlateBrush* FSubTrackEditor::GetIconBrush() const
 {
-	return FAppStyle::GetBrush("Sequencer.Tracks.Sub");
+	return FAppStyle::GetBrush(GetSubTrackBrushName());
 }
 
-
 bool FSubTrackEditor::OnAllowDrop(const FDragDropEvent& DragDropEvent, FSequencerDragDropParams& DragDropParams)
 {
-	if (!DragDropParams.Track.IsValid() || !DragDropParams.Track.Get()->IsA(UMovieSceneSubTrack::StaticClass()) || DragDropParams.Track.Get()->IsA(UMovieSceneCinematicShotTrack::StaticClass()))
+	if (!DragDropParams.Track.IsValid())
+	{
+		return false;
+	}
+
+	if (!DragDropParams.Track.Get()->IsA(GetSubTrackClass()))
 	{
 		return false;
 	}
@@ -377,10 +372,14 @@ bool FSubTrackEditor::OnAllowDrop(const FDragDropEvent& DragDropEvent, FSequence
 	return false;
 }
 
-
 FReply FSubTrackEditor::OnDrop(const FDragDropEvent& DragDropEvent, const FSequencerDragDropParams& DragDropParams)
 {
-	if (!DragDropParams.Track.IsValid() || !DragDropParams.Track.Get()->IsA(UMovieSceneSubTrack::StaticClass()) || DragDropParams.Track.Get()->IsA(UMovieSceneCinematicShotTrack::StaticClass()))
+	if (!DragDropParams.Track.IsValid())
+	{
+		return FReply::Unhandled();
+	}
+
+	if (!DragDropParams.Track.Get()->IsA(GetSubTrackClass()))
 	{
 		return FReply::Unhandled();
 	}
@@ -432,19 +431,262 @@ FReply FSubTrackEditor::OnDrop(const FDragDropEvent& DragDropEvent, const FSeque
 	return bAnyDropped ? FReply::Handled() : FReply::Unhandled();
 }
 
-/* FSubTrackEditor callbacks
+/* FSubTrackEditor
  *****************************************************************************/
 
+void FSubTrackEditor::InsertSection(UMovieSceneTrack* Track)
+{
+	FFrameTime NewSectionStartTime = GetSequencer()->GetLocalTime().Time;
+
+	UMovieScene* MovieScene = GetFocusedMovieScene();
+	if (!MovieScene)
+	{
+		return;
+	}
+
+	UMovieSceneSubTrack* SubTrack = FindOrCreateSubTrack(MovieScene, Track);
+
+	FString NewSequenceName = MovieSceneToolHelpers::GenerateNewShotName(SubTrack->GetAllSections(), NewSectionStartTime.FrameNumber);
+	FString NewSequencePath = MovieSceneToolHelpers::GenerateNewShotPath(GetSequencer()->GetFocusedMovieSceneSequence()->GetMovieScene(), NewSequenceName);
+
+	FAssetToolsModule& AssetToolsModule = FModuleManager::GetModuleChecked<FAssetToolsModule>("AssetTools");
+	AssetToolsModule.Get().CreateUniqueAssetName(NewSequencePath + TEXT("/") + NewSequenceName, TEXT(""), NewSequencePath, NewSequenceName);
+
+	if (UMovieSceneSequence* NewSequence = MovieSceneToolHelpers::CreateSequence(NewSequenceName, NewSequencePath))
+	{
+		const FScopedTransaction Transaction(FText::Join(FText::FromString(" "), LOCTEXT("InsertText", "Insert"), GetSubTrackName()));
+
+		int32 Duration = UE::MovieScene::DiscreteSize(NewSequence->GetMovieScene()->GetPlaybackRange());
+
+		if (UMovieSceneSubSection* NewSection = SubTrack->AddSequence(NewSequence, NewSectionStartTime.FrameNumber, Duration))
+		{
+			NewSection->SetRowIndex(MovieSceneToolHelpers::FindAvailableRowIndex(Track, NewSection));
+
+			GetSequencer()->NotifyMovieSceneDataChanged(EMovieSceneDataChangeType::MovieSceneStructureItemAdded);
+			GetSequencer()->EmptySelection();
+			GetSequencer()->SelectSection(NewSection);
+			GetSequencer()->ThrobSectionSelection();
+		}
+	}
+}
+
+void FSubTrackEditor::DuplicateSection(UMovieSceneSubSection* Section)
+{
+	UMovieSceneSubTrack* SubTrack = CastChecked<UMovieSceneSubTrack>(Section->GetOuter());
+
+	FFrameNumber StartTime = Section->HasStartFrame() ? Section->GetInclusiveStartFrame() : 0;
+	FString NewSectionName = MovieSceneToolHelpers::GenerateNewShotName(SubTrack->GetAllSections(), StartTime);
+	FString NewSequencePath = FPaths::GetPath(Section->GetSequence()->GetPathName());
+
+	// Duplicate the section and put it on the next available row
+	UMovieSceneSequence* NewSequence = MovieSceneToolHelpers::CreateSequence(NewSectionName, NewSequencePath, Section);
+	if (NewSequence)
+	{
+		const FScopedTransaction Transaction(FText::Join(FText::FromString(" "), LOCTEXT("DuplicateText", "Duplicate"), GetSubTrackName()));
+
+		int32 Duration = UE::MovieScene::DiscreteSize(Section->GetRange());
+
+		if (UMovieSceneSubSection* NewSection = SubTrack->AddSequence(NewSequence, StartTime, Duration))
+		{
+			NewSection->SetRange(Section->GetRange());
+			NewSection->SetRowIndex(MovieSceneToolHelpers::FindAvailableRowIndex(SubTrack, NewSection));
+			NewSection->Parameters.StartFrameOffset = Section->Parameters.StartFrameOffset;
+			NewSection->Parameters.TimeScale = Section->Parameters.TimeScale;
+			NewSection->SetPreRollFrames(Section->GetPreRollFrames());
+
+			GetSequencer()->NotifyMovieSceneDataChanged(EMovieSceneDataChangeType::MovieSceneStructureItemAdded);
+			GetSequencer()->EmptySelection();
+			GetSequencer()->SelectSection(NewSection);
+			GetSequencer()->ThrobSectionSelection();
+		}
+	}
+}
+
+void FSubTrackEditor::CreateNewTake(UMovieSceneSubSection* Section)
+{
+	FString ShotPrefix;
+	uint32 ShotNumber = INDEX_NONE;
+	uint32 TakeNumber = INDEX_NONE;
+	uint32 ShotNumberDigits = 0;
+	uint32 TakeNumberDigits = 0;
+	
+	if (MovieSceneToolHelpers::ParseShotName(GetSubSectionDisplayName(Section), ShotPrefix, ShotNumber, TakeNumber, ShotNumberDigits, TakeNumberDigits))
+	{
+		TArray<FAssetData> AssetData;
+		uint32 CurrentTakeNumber = INDEX_NONE;
+		MovieSceneToolHelpers::GatherTakes(Section, AssetData, CurrentTakeNumber);
+		uint32 NewTakeNumber = CurrentTakeNumber;
+
+		for (auto ThisAssetData : AssetData)
+		{
+			uint32 ThisTakeNumber = INDEX_NONE;
+			if (MovieSceneToolHelpers::GetTakeNumber(Section, ThisAssetData, ThisTakeNumber))
+			{
+				if (ThisTakeNumber >= NewTakeNumber)
+				{
+					NewTakeNumber = ThisTakeNumber + 1;
+				}
+			}
+		}
+
+		FString NewShotName = MovieSceneToolHelpers::ComposeShotName(ShotPrefix, ShotNumber, NewTakeNumber, ShotNumberDigits, TakeNumberDigits);
+
+		TRange<FFrameNumber> NewShotRange         = Section->GetRange();
+		FFrameNumber         NewShotStartOffset   = Section->Parameters.StartFrameOffset;
+		float                NewShotTimeScale     = Section->Parameters.TimeScale;
+		int32                NewShotPrerollFrames = Section->GetPreRollFrames();
+		int32                NewRowIndex          = Section->GetRowIndex();
+		FFrameNumber         NewShotStartTime     = NewShotRange.GetLowerBound().IsClosed() ? UE::MovieScene::DiscreteInclusiveLower(NewShotRange) : 0;
+
+		UMovieSceneSubTrack* SubTrack = CastChecked<UMovieSceneSubTrack>(Section->GetOuter());
+		FString NewSequencePath = FPaths::GetPath(Section->GetSequence()->GetPathName());
+
+		if (UMovieSceneSequence* NewSequence = MovieSceneToolHelpers::CreateSequence(NewShotName, NewSequencePath, Section))
+		{
+			const FScopedTransaction Transaction(LOCTEXT("NewTake_Transaction", "New Take"));
+
+			int32 Duration = UE::MovieScene::DiscreteSize(Section->GetRange());
+
+			UMovieSceneSubSection* NewShot = SubTrack->AddSequence(NewSequence, NewShotStartTime, Duration);
+			SubTrack->RemoveSection(*Section);
+
+			NewShot->SetRange(NewShotRange);
+			NewShot->Parameters.StartFrameOffset = NewShotStartOffset;
+			NewShot->Parameters.TimeScale = NewShotTimeScale;
+			NewShot->SetPreRollFrames(NewShotPrerollFrames);
+			NewShot->SetRowIndex(NewRowIndex);
+
+			MovieSceneToolHelpers::SetTakeNumber(NewShot, NewTakeNumber);
+
+			GetSequencer()->NotifyMovieSceneDataChanged( EMovieSceneDataChangeType::MovieSceneStructureItemsChanged );
+			GetSequencer()->EmptySelection();
+			GetSequencer()->SelectSection(NewShot);
+			GetSequencer()->ThrobSectionSelection();
+		}
+	}
+}
+
+void FSubTrackEditor::SwitchTake(UObject* TakeObject)
+{
+	bool bSwitchedTake = false;
+
+	const FScopedTransaction Transaction(LOCTEXT("SwitchTake_Transaction", "Switch Take"));
+
+	TArray<UMovieSceneSection*> Sections;
+	GetSequencer()->GetSelectedSections(Sections);
+
+	for (int32 SectionIndex = 0; SectionIndex < Sections.Num(); ++SectionIndex)
+	{
+		if (!Sections[SectionIndex]->IsA<UMovieSceneSubSection>())
+		{
+			continue;
+		}
+
+		UMovieSceneSubSection* Section = Cast<UMovieSceneSubSection>(Sections[SectionIndex]);
+
+		if (TakeObject && TakeObject->IsA(UMovieSceneSequence::StaticClass()))
+		{
+			UMovieSceneSequence* MovieSceneSequence = CastChecked<UMovieSceneSequence>(TakeObject);
+
+			UMovieSceneSubTrack* SubTrack = CastChecked<UMovieSceneSubTrack>(Section->GetOuter());
+
+			TRange<FFrameNumber> NewSectionRange = Section->GetRange();
+			FFrameNumber		 NewSectionStartOffset = Section->Parameters.StartFrameOffset;
+			float                NewSectionTimeScale = Section->Parameters.TimeScale;
+			int32                NewSectionPrerollFrames = Section->GetPreRollFrames();
+			int32                NewRowIndex = Section->GetRowIndex();
+			FFrameNumber         NewSectionStartTime = NewSectionRange.GetLowerBound().IsClosed() ? UE::MovieScene::DiscreteInclusiveLower(NewSectionRange) : 0;
+			int32                NewSectionRowIndex = Section->GetRowIndex();
+
+			const int32 Duration = (NewSectionRange.GetLowerBound().IsClosed() && NewSectionRange.GetUpperBound().IsClosed()) ? UE::MovieScene::DiscreteSize(NewSectionRange) : 1;
+			UMovieSceneSubSection* NewSection = SubTrack->AddSequence(MovieSceneSequence, NewSectionStartTime, Duration);
+
+			if (NewSection != nullptr)
+			{
+				SubTrack->RemoveSection(*Section);
+
+				NewSection->SetRange(NewSectionRange);
+				NewSection->Parameters.StartFrameOffset = NewSectionStartOffset;
+				NewSection->Parameters.TimeScale = NewSectionTimeScale;
+				NewSection->SetPreRollFrames(NewSectionPrerollFrames);
+				NewSection->SetRowIndex(NewSectionRowIndex);
+
+				bSwitchedTake = true;
+			}
+		}
+	}
+
+	if (bSwitchedTake)
+	{
+		GetSequencer()->NotifyMovieSceneDataChanged(EMovieSceneDataChangeType::MovieSceneStructureItemsChanged);
+	}
+}
+
+void FSubTrackEditor::AddTakesMenu(UMovieSceneSubSection* Section, FMenuBuilder& MenuBuilder)
+{
+	TArray<FAssetData> AssetData;
+	uint32 CurrentTakeNumber = INDEX_NONE;
+	MovieSceneToolHelpers::GatherTakes(Section, AssetData, CurrentTakeNumber);
+
+	AssetData.Sort([Section](const FAssetData& A, const FAssetData& B) {
+		uint32 TakeNumberA = INDEX_NONE;
+		uint32 TakeNumberB = INDEX_NONE;
+		if (MovieSceneToolHelpers::GetTakeNumber(Section, A, TakeNumberA) && MovieSceneToolHelpers::GetTakeNumber(Section, B, TakeNumberB))
+		{
+			return TakeNumberA < TakeNumberB;
+		}
+		return true;
+	});
+
+	for (auto ThisAssetData : AssetData)
+	{
+		uint32 TakeNumber = INDEX_NONE;
+		if (MovieSceneToolHelpers::GetTakeNumber(Section, ThisAssetData, TakeNumber))
+		{
+			UObject* TakeObject = ThisAssetData.GetAsset();
+
+			if (TakeObject)
+			{
+				MenuBuilder.AddMenuEntry(
+					FText::Format(LOCTEXT("TakeNumber", "Take {0}"), FText::AsNumber(TakeNumber)),
+					FText::Format(LOCTEXT("TakeNumberTooltip", "Switch to {0}"), FText::FromString(TakeObject->GetPathName())),
+					TakeNumber == CurrentTakeNumber ? FSlateIcon(FAppStyle::GetAppStyleSetName(), "Sequencer.Star") : FSlateIcon(FAppStyle::GetAppStyleSetName(), "Sequencer.Empty"),
+					FUIAction(FExecuteAction::CreateSP(this, &FSubTrackEditor::SwitchTake, TakeObject))
+				);
+			}
+		}
+	}
+}
+
 bool FSubTrackEditor::CanAddSubSequence(const UMovieSceneSequence& Sequence) const
 {
-	// prevent adding ourselves and ensure we have a valid movie scene
 	UMovieSceneSequence* FocusedSequence = GetSequencer()->GetFocusedMovieSceneSequence();
 	return FSubTrackEditorUtil::CanAddSubSequence(FocusedSequence, Sequence);
 }
 
-UMovieSceneSubTrack* FSubTrackEditor::CreateNewTrack(UMovieScene* MovieScene) const
+FText FSubTrackEditor::GetSubTrackName() const
+{ 
+	return LOCTEXT("SubTrackName", "Subsequence"); 
+}
+
+FText FSubTrackEditor::GetSubTrackToolTip() const
+{ 
+	return LOCTEXT("SubTrackToolTip", "A track that can contain other sequences.");
+}
+
+FName FSubTrackEditor::GetSubTrackBrushName() const
 {
-	return MovieScene->AddTrack<UMovieSceneSubTrack>();
+	return TEXT("Sequencer.Tracks.Sub");
+}
+
+FString FSubTrackEditor::GetSubSectionDisplayName(const UMovieSceneSubSection* Section) const
+{
+	return Section && Section->GetSequence() ? Section->GetSequence()->GetName() : FString();
+}
+
+TSubclassOf<UMovieSceneSubTrack> FSubTrackEditor::GetSubTrackClass() const
+{
+	return UMovieSceneSubTrack::StaticClass();
 }
 
 void FSubTrackEditor::GetSupportedSequenceClassPaths(TArray<FTopLevelAssetPath>& ClassPaths) const
@@ -452,6 +694,11 @@ void FSubTrackEditor::GetSupportedSequenceClassPaths(TArray<FTopLevelAssetPath>&
 	ClassPaths.Add(FTopLevelAssetPath(TEXT("/Script/LevelSequence"), TEXT("LevelSequence")));
 }
 
+UMovieSceneSubTrack* FSubTrackEditor::CreateNewTrack(UMovieScene* MovieScene) const
+{
+	return Cast<UMovieSceneSubTrack>(MovieScene->AddTrack(GetSubTrackClass()));
+}
+
 /* FSubTrackEditor callbacks
  *****************************************************************************/
 
@@ -469,10 +716,10 @@ void FSubTrackEditor::HandleAddSubTrackMenuEntryExecute()
 		return;
 	}
 
-	const FScopedTransaction Transaction(LOCTEXT("AddSubTrack_Transaction", "Add Sub Track"));
+	const FScopedTransaction Transaction(FText::Join(FText::FromString(" "), LOCTEXT("AddText", "Add"), GetSubTrackName(), LOCTEXT("TrackText", "Track")));
 	FocusedMovieScene->Modify();
 
-	UMovieSceneSubTrack* NewTrack = CreateNewTrack(FocusedMovieScene);
+	UMovieSceneSubTrack* NewTrack = FindOrCreateSubTrack(FocusedMovieScene, nullptr);
 	ensure(NewTrack);
 
 	if (GetSequencer().IsValid())
@@ -481,15 +728,25 @@ void FSubTrackEditor::HandleAddSubTrackMenuEntryExecute()
 	}
 }
 
+UMovieSceneSubTrack* FSubTrackEditor::FindOrCreateSubTrack(UMovieScene* MovieScene, UMovieSceneTrack* Track) const
+{
+	UMovieSceneSubTrack* SubTrack = Cast<UMovieSceneSubTrack>(Track);
+	if (!SubTrack)
+	{
+		SubTrack = Cast<UMovieSceneSubTrack>(MovieScene->AddTrack(GetSubTrackClass()));
+	}
+	return SubTrack;
+}
+
 TSharedRef<SWidget> FSubTrackEditor::HandleAddSubSequenceComboButtonGetMenuContent(UMovieSceneTrack* InTrack)
 {
 	FMenuBuilder MenuBuilder(true, nullptr);
 
 	MenuBuilder.AddMenuEntry(
-		LOCTEXT("InsertSequence", "Insert Sequence"),
-		LOCTEXT("InsertSequenceTooltip", "Insert new sequence at current time"),
+		FText::Join(FText::FromString(" "), LOCTEXT("InsertText", "Insert"), GetSubTrackName()),
+		LOCTEXT("InsertSectionTooltip", "Insert new sequence at current time"),
 		FSlateIcon(),
-		FUIAction(FExecuteAction::CreateSP(this, &FSubTrackEditor::InsertSequence, InTrack))
+		FUIAction(FExecuteAction::CreateSP(this, &FSubTrackEditor::InsertSection, InTrack))
 	);
 
 	MenuBuilder.BeginSection(TEXT("ChooseSequence"), LOCTEXT("ChooseSequence", "Choose Sequence"));
@@ -547,40 +804,6 @@ void FSubTrackEditor::HandleAddSubSequenceComboButtonMenuEntryEnterPressed(const
 	}
 }
 
-void FSubTrackEditor::InsertSequence(UMovieSceneTrack* Track)
-{
-	FFrameTime NewSectionStartTime = GetSequencer()->GetLocalTime().Time;
-
-	UMovieSceneSubTrack* SubTrack = Cast<UMovieSceneSubTrack>(Track);
-	if (!SubTrack)
-	{
-		SubTrack = FindOrCreateRootTrack<UMovieSceneSubTrack>().Track;
-	}
-
-	FString NewSequencePath = FPaths::GetPath(GetSequencer()->GetFocusedMovieSceneSequence()->GetPathName());
-	FString NewSequenceName = TEXT("NewSubSequence");
-
-	FAssetToolsModule& AssetToolsModule = FModuleManager::GetModuleChecked<FAssetToolsModule>("AssetTools");
-	AssetToolsModule.Get().CreateUniqueAssetName(NewSequencePath + TEXT("/") + NewSequenceName, TEXT(""), NewSequencePath, NewSequenceName);
-
-	if (UMovieSceneSequence* NewSequence = MovieSceneToolHelpers::CreateSequence(NewSequenceName, NewSequencePath))
-	{
-		const FScopedTransaction Transaction(LOCTEXT("InsertSequence_Transaction", "Insert Sequence"));
-
-		int32 Duration = UE::MovieScene::DiscreteSize(NewSequence->GetMovieScene()->GetPlaybackRange());
-
-		if (UMovieSceneSubSection* NewSection = SubTrack->AddSequence(NewSequence, NewSectionStartTime.FrameNumber, Duration))
-		{
-			NewSection->SetRowIndex(MovieSceneToolHelpers::FindAvailableRowIndex(Track, NewSection));
-
-			GetSequencer()->NotifyMovieSceneDataChanged(EMovieSceneDataChangeType::MovieSceneStructureItemAdded);
-			GetSequencer()->EmptySelection();
-			GetSequencer()->SelectSection(NewSection);
-			GetSequencer()->ThrobSectionSelection();
-		}
-	}
-}
-
 FKeyPropertyResult FSubTrackEditor::AddKeyInternal(FFrameNumber KeyTime, UMovieSceneSequence* InMovieSceneSequence, UMovieSceneTrack* InTrack, int32 RowIndex)
 {	
 	FKeyPropertyResult KeyPropertyResult;
@@ -595,7 +818,7 @@ FKeyPropertyResult FSubTrackEditor::AddKeyInternal(FFrameNumber KeyTime, UMovieS
 
 	if (CanAddSubSequence(*InMovieSceneSequence))
 	{
-		UMovieSceneSubTrack* SubTrack = Cast<UMovieSceneSubTrack>(InTrack);
+		UMovieSceneSubTrack* SubTrack = FindOrCreateSubTrack(InMovieSceneSequence->GetMovieScene(), InTrack);
 
 		const FFrameRate TickResolution = InMovieSceneSequence->GetMovieScene()->GetTickResolution();
 		const FQualifiedFrameTime InnerDuration = FQualifiedFrameTime(
@@ -634,11 +857,7 @@ FKeyPropertyResult FSubTrackEditor::HandleSequenceAdded(FFrameNumber KeyTime, UM
 {
 	FKeyPropertyResult KeyPropertyResult;
 
-	UMovieSceneSubTrack* SubTrack = Cast<UMovieSceneSubTrack>(Track);
-	if (!SubTrack)
-	{
-		SubTrack = FindOrCreateRootTrack<UMovieSceneSubTrack>().Track;
-	}
+	UMovieSceneSubTrack* SubTrack = FindOrCreateSubTrack(Sequence->GetMovieScene(), Track);
 
 	const FFrameRate TickResolution = Sequence->GetMovieScene()->GetTickResolution();
 	const FQualifiedFrameTime InnerDuration = FQualifiedFrameTime(
@@ -666,60 +885,4 @@ FKeyPropertyResult FSubTrackEditor::HandleSequenceAdded(FFrameNumber KeyTime, UM
 	return KeyPropertyResult;
 }
 
-void FSubTrackEditor::SwitchTake(UObject* TakeObject)
-{
-	bool bSwitchedTake = false;
-
-	const FScopedTransaction Transaction(LOCTEXT("SwitchTake_Transaction", "Switch Take"));
-
-	TArray<UMovieSceneSection*> Sections;
-	GetSequencer()->GetSelectedSections(Sections);
-
-	for (int32 SectionIndex = 0; SectionIndex < Sections.Num(); ++SectionIndex)
-	{
-		if (!Sections[SectionIndex]->IsA<UMovieSceneSubSection>())
-		{
-			continue;
-		}
-
-		UMovieSceneSubSection* Section = Cast<UMovieSceneSubSection>(Sections[SectionIndex]);
-
-		if (TakeObject && TakeObject->IsA(UMovieSceneSequence::StaticClass()))
-		{
-			UMovieSceneSequence* MovieSceneSequence = CastChecked<UMovieSceneSequence>(TakeObject);
-
-			UMovieSceneSubTrack* SubTrack = CastChecked<UMovieSceneSubTrack>(Section->GetOuter());
-
-			TRange<FFrameNumber> NewShotRange         = Section->GetRange();
-			FFrameNumber		 NewShotStartOffset   = Section->Parameters.StartFrameOffset;
-			float                NewShotTimeScale     = Section->Parameters.TimeScale;
-			int32                NewShotPrerollFrames = Section->GetPreRollFrames();
-			int32                NewRowIndex          = Section->GetRowIndex();
-			FFrameNumber         NewShotStartTime     = NewShotRange.GetLowerBound().IsClosed() ? UE::MovieScene::DiscreteInclusiveLower(NewShotRange) : 0;
-			int32                NewShotRowIndex      = Section->GetRowIndex();
-
-			const int32 Duration = (NewShotRange.GetLowerBound().IsClosed() && NewShotRange.GetUpperBound().IsClosed() ) ? UE::MovieScene::DiscreteSize(NewShotRange) : 1;
-			UMovieSceneSubSection* NewShot = SubTrack->AddSequence(MovieSceneSequence, NewShotStartTime, Duration);
-
-			if (NewShot != nullptr)
-			{
-				SubTrack->RemoveSection(*Section);
-
-				NewShot->SetRange(NewShotRange);
-				NewShot->Parameters.StartFrameOffset = NewShotStartOffset;
-				NewShot->Parameters.TimeScale = NewShotTimeScale;
-				NewShot->SetPreRollFrames(NewShotPrerollFrames);
-				NewShot->SetRowIndex(NewShotRowIndex);
-
-				bSwitchedTake = true;
-			}
-		}
-	}
-
-	if (bSwitchedTake)
-	{
-		GetSequencer()->NotifyMovieSceneDataChanged( EMovieSceneDataChangeType::MovieSceneStructureItemsChanged );
-	}
-}
-
 #undef LOCTEXT_NAMESPACE
diff --git a/Engine/Source/Editor/MovieSceneTools/Public/TrackEditors/CinematicShotTrackEditor.h b/Engine/Source/Editor/MovieSceneTools/Public/TrackEditors/CinematicShotTrackEditor.h
--- a/Engine/Source/Editor/MovieSceneTools/Public/TrackEditors/CinematicShotTrackEditor.h
+++ b/Engine/Source/Editor/MovieSceneTools/Public/TrackEditors/CinematicShotTrackEditor.h
@@ -9,11 +9,10 @@
 #include "ISequencer.h"
 #include "ISequencerSection.h"
 #include "MovieSceneTrack.h"
+#include "SubTrackEditor.h"
 #include "ISequencerTrackEditor.h"
-#include "MovieSceneTrackEditor.h"
 
 class AActor;
-struct FAssetData;
 class FMenuBuilder;
 class FTrackEditorThumbnailPool;
 class UMovieSceneCinematicShotSection;
@@ -23,7 +22,7 @@ class UMovieSceneSubSection;
 /**
  * Tools for cinematic shots.
  */
-class MOVIESCENETOOLS_API FCinematicShotTrackEditor : public FMovieSceneTrackEditor
+class MOVIESCENETOOLS_API FCinematicShotTrackEditor : public FSubTrackEditor
 {
 public:
 
@@ -52,34 +51,22 @@ public:
 	// ISequencerTrackEditor interface
 	virtual void OnInitialize() override;
 	virtual void OnRelease() override;
-	virtual void BuildAddTrackMenu(FMenuBuilder& MenuBuilder) override;
 	virtual TSharedPtr<SWidget> BuildOutlinerEditWidget(const FGuid& ObjectBinding, UMovieSceneTrack* Track, const FBuildEditWidgetParams& Params) override;
 	virtual TSharedRef<ISequencerSection> MakeSectionInterface(UMovieSceneSection& SectionObject, UMovieSceneTrack& Track, FGuid ObjectBinding) override;
-	virtual bool HandleAssetAdded(UObject* Asset, const FGuid& TargetObjectGuid) override;
 	virtual bool SupportsSequence(UMovieSceneSequence* InSequence) const override;
-	virtual bool SupportsType(TSubclassOf<UMovieSceneTrack> Type ) const override;
 	virtual void Tick(float DeltaTime) override;
 	virtual void BuildTrackContextMenu( FMenuBuilder& MenuBuilder, UMovieSceneTrack* Track ) override;
-	virtual const FSlateBrush* GetIconBrush() const override;
-	virtual bool OnAllowDrop(const FDragDropEvent& DragDropEvent, FSequencerDragDropParams& DragDropParams) override;
-	virtual FReply OnDrop(const FDragDropEvent& DragDropEvent, const FSequencerDragDropParams& DragDropParams) override;
 
-	/*
-	 * Insert shot. 
-	 */
+	UE_DEPRECATED(5.3, "InsertShot has been deprecated in favor of FSubTrackEditor::InsertSection(UMovieSceneTrack*)")
 	void InsertShot();
-
-	/*
-	 * Insert filler.
-	 */
-	void InsertFiller();
-
-	/*
-	 * Duplicate shot. 
-	 *
-	 * @param Section The section to duplicate
-	 */
+	UE_DEPRECATED(5.3, "DuplicateShot has been deprecated in favor of FSubTrackEditor::DuplicateSection(UMovieSceneSubSection*)")
 	void DuplicateShot(UMovieSceneCinematicShotSection* Section);
+	UE_DEPRECATED(5.3, "RenameShot has been removed because it was unused")
+	void RenameShot(UMovieSceneCinematicShotSection* Section) {}
+	UE_DEPRECATED(5.3, "NewTake has been deprecated in favor of FSubTrackEditor::CreateNewTake(UMovieSceneSubSection*)")
+	void NewTake(UMovieSceneCinematicShotSection* Section);
+	UE_DEPRECATED(5.3, "InsertFiller has been removed because it is obsolete")
+	void InsertFiller() {}
 
 	/*
 	 * Render shots. 
@@ -88,53 +75,26 @@ public:
 	 */
 	void RenderShots(const TArray<UMovieSceneCinematicShotSection*>& Sections);
 
-	/*
-	 * Rename shot. 
-	 *
-	 * @param Section The section to rename.
-	 */
-	void RenameShot(UMovieSceneCinematicShotSection* Section);
-
-	/*
-	 * New take. 
-	 *
-	 * @param Section The section to create a new take of.
-	 */
-	void NewTake(UMovieSceneCinematicShotSection* Section);
-
-	/*
-	* Switch take for the selected sections
-	*
-	* @param TakeObject The take object to switch to.
-	*/
-	void SwitchTake(UObject* TakeObject);
-
-private:
+public:
 
-	/** Callback for determining whether the "Add Shot" menu entry can execute. */
-	bool HandleAddCinematicShotTrackMenuEntryCanExecute() const;
+	// FSubTrackEditor interface
+	virtual FText GetSubTrackName() const override;
+	virtual FText GetSubTrackToolTip() const override;
+	virtual FName GetSubTrackBrushName() const override;
+	virtual FString GetSubSectionDisplayName(const UMovieSceneSubSection* Section) const override;
+	virtual TSubclassOf<UMovieSceneSubTrack> GetSubTrackClass() const;
 
-	/** Callback for executing the "Add Shot Track" menu entry. */
-	void HandleAddCinematicShotTrackMenuEntryExecute();
-	
 protected:
 
-	/** Callback for generating the menu of the "Add Shot" combo button. */
+	UE_DEPRECATED(5.3, "HandleAddCinematicShotComboButtonGetMenuContent has been deprecated. Please implement FSubTrackEditor::HandleAddSubSequenceComboButtonGetMenuContent(UMovieSceneTrack*) instead")
 	virtual TSharedRef<SWidget> HandleAddCinematicShotComboButtonGetMenuContent();
-	
-	/** Find or create a cinematic shot track in the currently focused movie scene. */
+	UE_DEPRECATED(5.3, "FindOrCreateCinematicShotTrack has been deprecated in favor of FSubTrackEditor::FindOrCreateSubTrack(UMovieScene* MovieScene, UMovieSceneTrack*)")
 	virtual UMovieSceneCinematicShotTrack* FindOrCreateCinematicShotTrack();
 
-private:
-
-	/** Callback for executing a menu entry in the "Add Shot" combo button. */
-	virtual void HandleAddCinematicShotComboButtonMenuEntryExecute(const FAssetData& AssetData);
-
-	/** Callback for executing a menu entry in the "Add Shot" combo button when enter pressed. */
-	virtual void HandleAddCinematicShotComboButtonMenuEntryEnterPressed(const TArray<FAssetData>& AssetData);
+	virtual bool HandleAddSubTrackMenuEntryCanExecute() const override;
+	virtual bool CanHandleAssetAdded(UMovieSceneSequence* Sequence) const override;
 
-	/** Delegate for AnimatablePropertyChanged in AddKey */
-	virtual FKeyPropertyResult AddKeyInternal(FFrameNumber KeyTime, UMovieSceneSequence* InMovieSceneSequence, int32 RowIndex);
+private:
 
 	/** Delegate for shots button lock state */
 	ECheckBoxState AreShotsLocked() const;
@@ -145,27 +105,9 @@ private:
 	/** Delegate for shots button lock tooltip */
 	FText GetLockShotsToolTip() const;
 
-protected:
-	
-	/**
-	 * Check whether the given sequence can be added as a sub-sequence.
-	 *
-	 * The purpose of this method is to disallow circular references
-	 * between sub-sequences in the focused movie scene.
-	 *
-	 * @param Sequence The sequence to check.
-	 * @return true if the sequence can be added as a sub-sequence, false otherwise.
-	 */
-	bool CanAddSubSequence(const UMovieSceneSequence& Sequence) const;
-	
-private:
-	
 	/** Called when our sequencer wants to switch cameras */
 	void OnUpdateCameraCut(UObject* CameraObject, bool bJumpCut);
 
-	/** Callback for AnimatablePropertyChanged in HandleAssetAdded. */
-	FKeyPropertyResult HandleSequenceAdded(FFrameNumber KeyTime, UMovieSceneSequence* Sequence, int32 RowIndex);
-
 	/** Callback for ImportEDL. */
 	void ImportEDL();
 	
diff --git a/Engine/Source/Editor/MovieSceneTools/Public/TrackEditors/SubTrackEditor.h b/Engine/Source/Editor/MovieSceneTools/Public/TrackEditors/SubTrackEditor.h
--- a/Engine/Source/Editor/MovieSceneTools/Public/TrackEditors/SubTrackEditor.h
+++ b/Engine/Source/Editor/MovieSceneTools/Public/TrackEditors/SubTrackEditor.h
@@ -18,7 +18,7 @@ class FMenuBuilder;
 class UMovieSceneSubTrack;
 
 /**
- * Tools for animatable property types such as floats ands vectors
+ * Tools for subsequences
  */
 class MOVIESCENETOOLS_API FSubTrackEditor
 	: public FMovieSceneTrackEditor
@@ -57,11 +57,23 @@ public:
 	virtual bool OnAllowDrop(const FDragDropEvent& DragDropEvent, FSequencerDragDropParams& DragDropParams) override;
 	virtual FReply OnDrop(const FDragDropEvent& DragDropEvent, const FSequencerDragDropParams& DragDropParams) override;
 	
-	/** Switch the selected sub section's take */
-	void SwitchTake(UObject* TakeObject);
-
 public:
 	
+	/** Insert sequence into this track */
+	virtual void InsertSection(UMovieSceneTrack* Track);
+
+	/** Duplicate the section into this track */
+	virtual void DuplicateSection(UMovieSceneSubSection* Section);
+
+	/** Create a new take of the given section */
+	virtual void CreateNewTake(UMovieSceneSubSection* Section);
+
+	/** Switch the selected section's take */
+	virtual void SwitchTake(UObject* TakeObject);
+
+	/** Generate a menu for takes for this section */
+	virtual void AddTakesMenu(UMovieSceneSubSection* Section, FMenuBuilder& MenuBuilder);
+
 	/**
 	 * Check whether the given sequence can be added as a sub-sequence.
 	 *
@@ -73,38 +85,54 @@ public:
 	 */
 	bool CanAddSubSequence(const UMovieSceneSequence& Sequence) const;
 
+public:
+
+	/** Get the name of the sub track */
+	virtual FText GetSubTrackName() const;
+
+	/** Get the tooltip for this sub track editor */
+	virtual FText GetSubTrackToolTip() const;
+
+	/** Get the brush used for the sub track editor */
+	virtual FName GetSubTrackBrushName() const;
+
+	/** Get the display name for the sub section */
+	virtual FString GetSubSectionDisplayName(const UMovieSceneSubSection* Section) const;
+
+	/** Get the UMovieSceneSubTrack class */
+	virtual TSubclassOf<UMovieSceneSubTrack> GetSubTrackClass() const;
+
 protected:
-	/**
-	 * Create a new UMovieSceneSubTrack.
-	 * 
-	 * @param MovieScene the MovieScene object to create a track.
-	 * @return the created MovieSceneSubTrack
-	 */
-	virtual UMovieSceneSubTrack* CreateNewTrack(UMovieScene* MovieScene) const;
 
-	/**
-	 * Return the list of supported sequence class paths.
-	 * 
-	 * @param OutClassPaths the list of supported sequence class paths.
-	 */
+	/** Get the list of supported sequence class paths */
 	virtual void GetSupportedSequenceClassPaths(TArray<FTopLevelAssetPath>& OutClassPaths) const;
 
-private:
+	/** Callback for executing the "Add Subsequence" menu entry. */
+	virtual void HandleAddSubTrackMenuEntryExecute();
+
+	/** Callback for determining whether the "Add Subsequence" menu entry can execute. */
+	virtual bool HandleAddSubTrackMenuEntryCanExecute() const { return true; }
 
-	/** Callback for executing the "Add Event Track" menu entry. */
-	void HandleAddSubTrackMenuEntryExecute();
+	/** Whether to handle this asset being dropped onto the sequence as opposed to a specific track. */
+	virtual bool CanHandleAssetAdded(UMovieSceneSequence* Sequence) const;
+
+	UE_DEPRECATED(5.3, "CreateNewTrack has been deprecated, please implement GetSubTrackClass")
+	virtual UMovieSceneSubTrack* CreateNewTrack(UMovieScene* MovieScene) const;
+
+	/** Find or create a sub track. If the given track is a subtrack, it will be returned. */
+	UMovieSceneSubTrack* FindOrCreateSubTrack(UMovieScene* MovieScene, UMovieSceneTrack* Track) const;
 
 	/** Callback for generating the menu of the "Add Sequence" combo button. */
 	TSharedRef<SWidget> HandleAddSubSequenceComboButtonGetMenuContent(UMovieSceneTrack* InTrack);
 
+private:
+
 	/** Callback for executing a menu entry in the "Add Sequence" combo button. */
 	void HandleAddSubSequenceComboButtonMenuEntryExecute(const FAssetData& AssetData, UMovieSceneTrack* InTrack);
 
 	/** Callback for executing a menu entry in the "Add Sequence" combo button when enter pressed. */
 	void HandleAddSubSequenceComboButtonMenuEntryEnterPressed(const TArray<FAssetData>& AssetData, UMovieSceneTrack* InTrack);
 
-	void InsertSequence(UMovieSceneTrack* Track);
-
 	/** Delegate for AnimatablePropertyChanged in AddKey */
 	FKeyPropertyResult AddKeyInternal(FFrameNumber KeyTime, UMovieSceneSequence* InMovieSceneSequence, UMovieSceneTrack* InTrack, int32 RowIndex);
 
diff --git a/Engine/Source/Editor/Sequencer/Private/Sequencer.cpp b/Engine/Source/Editor/Sequencer/Private/Sequencer.cpp
--- a/Engine/Source/Editor/Sequencer/Private/Sequencer.cpp
+++ b/Engine/Source/Editor/Sequencer/Private/Sequencer.cpp
@@ -2768,95 +2768,18 @@ void FSequencer::SetPlaybackRange(TRange<FFrameNumber> Range)
 	}
 }
 
-UMovieSceneSection* FSequencer::FindNextOrPreviousShot(UMovieSceneSequence* Sequence, FFrameNumber SearchFromTime, const bool bNextShot) const
+void FSequencer::SetSelectionRangeToShot(const bool bNextShot)
 {
-	UMovieScene* OwnerMovieScene = Sequence->GetMovieScene();
+	UMovieSceneSequence* Sequence = GetFocusedMovieSceneSequence();
+	UMovieScene* MovieScene = Sequence->GetMovieScene();
 
-	UMovieSceneTrack* CinematicShotTrack = OwnerMovieScene->FindTrack(UMovieSceneCinematicShotTrack::StaticClass());
+	UMovieSceneTrack* CinematicShotTrack = MovieScene->FindTrack(UMovieSceneCinematicShotTrack::StaticClass());
 	if (!CinematicShotTrack)
 	{
-		return nullptr;
-	}
-
-	FFrameNumber MinTime = TNumericLimits<FFrameNumber>::Max();
-
-	TMap<FFrameNumber, int32> StartTimeMap;
-	for (int32 SectionIndex = 0; SectionIndex < CinematicShotTrack->GetAllSections().Num(); ++SectionIndex)
-	{
-		UMovieSceneSection* ShotSection = CinematicShotTrack->GetAllSections()[SectionIndex];
-
-		if (ShotSection && ShotSection->HasStartFrame())
-		{
-			StartTimeMap.Add(ShotSection->GetInclusiveStartFrame(), SectionIndex);
-		}
-	}
-
-	StartTimeMap.KeySort(TLess<FFrameNumber>());
-
-	int32 MinShotIndex = -1;
-	for (auto StartTimeIt = StartTimeMap.CreateIterator(); StartTimeIt; ++StartTimeIt)
-	{
-		FFrameNumber StartTime = StartTimeIt->Key;
-		if (bNextShot)
-		{
-			if (StartTime > SearchFromTime)
-			{
-				FFrameNumber DiffTime = FMath::Abs(StartTime - SearchFromTime);
-				if (DiffTime < MinTime)
-				{
-					MinTime = DiffTime;
-					MinShotIndex = StartTimeIt->Value;
-				}
-			}
-		}
-		else
-		{
-			if (SearchFromTime >= StartTime)
-			{
-				FFrameNumber DiffTime = FMath::Abs(StartTime - SearchFromTime);
-				if (DiffTime < MinTime)
-				{
-					MinTime = DiffTime;
-					MinShotIndex = StartTimeIt->Value;
-				}
-			}
-		}
-	}
-
-	int32 TargetShotIndex = -1;
-
-	if (bNextShot)
-	{
-		TargetShotIndex = MinShotIndex;
-	}
-	else
-	{
-		int32 PreviousShotIndex = -1;
-		for (auto StartTimeIt = StartTimeMap.CreateIterator(); StartTimeIt; ++StartTimeIt)
-		{
-			if (StartTimeIt->Value == MinShotIndex)
-			{
-				if (PreviousShotIndex != -1)
-				{
-					TargetShotIndex = PreviousShotIndex;
-				}
-				break;
-			}
-			PreviousShotIndex = StartTimeIt->Value;
-		}
+		return;
 	}
 
-	if (TargetShotIndex == -1)
-	{
-		return nullptr;
-	}	
-
-	return CinematicShotTrack->GetAllSections()[TargetShotIndex];
-}
-
-void FSequencer::SetSelectionRangeToShot(const bool bNextShot)
-{
-	UMovieSceneSection* TargetShotSection = FindNextOrPreviousShot(GetFocusedMovieSceneSequence(), GetLocalTime().Time.FloorToFrame(), bNextShot);
+	UMovieSceneSection* TargetShotSection = MovieSceneHelpers::FindNextSection(CinematicShotTrack->GetAllSections(), GetLocalTime().Time.FloorToFrame());
 
 	TRange<FFrameNumber> NewSelectionRange = TargetShotSection ? TargetShotSection->GetRange() : TRange<FFrameNumber>::All();
 	if (NewSelectionRange.GetLowerBound().IsClosed() && NewSelectionRange.GetUpperBound().IsClosed())
@@ -9355,7 +9278,16 @@ void FSequencer::StepToNextShot()
 {
 	if (ActiveTemplateIDs.Num() < 2)
 	{
-		UMovieSceneSection* TargetShotSection = FindNextOrPreviousShot(GetFocusedMovieSceneSequence(), GetLocalTime().Time.FloorToFrame(), true);
+		UMovieSceneSequence* Sequence = GetFocusedMovieSceneSequence();
+		UMovieScene* MovieScene = Sequence->GetMovieScene();
+
+		UMovieSceneTrack* CinematicShotTrack = MovieScene->FindTrack(UMovieSceneCinematicShotTrack::StaticClass());
+		if (!CinematicShotTrack)
+		{
+			return;
+		}
+
+		UMovieSceneSection* TargetShotSection = MovieSceneHelpers::FindNextSection(CinematicShotTrack->GetAllSections(), GetLocalTime().Time.FloorToFrame());
 
 		if (TargetShotSection)
 		{
@@ -9376,9 +9308,16 @@ void FSequencer::StepToNextShot()
 		return;
 	}
 
-	FFrameTime CurrentTime = SubSequenceRange.GetLowerBoundValue() * SubData->OuterToInnerTransform.InverseFromWarp(RootToLocalLoopCounter);
+	UMovieScene* MovieScene = Sequence->GetMovieScene();
 
-	UMovieSceneSubSection* NextShot = Cast<UMovieSceneSubSection>(FindNextOrPreviousShot(Sequence, CurrentTime.FloorToFrame(), true));
+	UMovieSceneTrack* CinematicShotTrack = MovieScene->FindTrack(UMovieSceneCinematicShotTrack::StaticClass());
+	if (!CinematicShotTrack)
+	{
+		return;
+	}
+
+	FFrameTime CurrentTime = SubSequenceRange.GetLowerBoundValue() * SubData->OuterToInnerTransform.InverseFromWarp(RootToLocalLoopCounter);
+	UMovieSceneSubSection* NextShot = Cast<UMovieSceneSubSection>(MovieSceneHelpers::FindNextSection(CinematicShotTrack->GetAllSections(), CurrentTime.FloorToFrame()));
 	if (!NextShot)
 	{
 		return;
@@ -9397,7 +9336,16 @@ void FSequencer::StepToPreviousShot()
 {
 	if (ActiveTemplateIDs.Num() < 2)
 	{
-		UMovieSceneSection* TargetShotSection = FindNextOrPreviousShot(GetFocusedMovieSceneSequence(), GetLocalTime().Time.FloorToFrame(), false);
+		UMovieSceneSequence* Sequence = GetFocusedMovieSceneSequence();
+		UMovieScene* MovieScene = Sequence->GetMovieScene();
+
+		UMovieSceneTrack* CinematicShotTrack = MovieScene->FindTrack(UMovieSceneCinematicShotTrack::StaticClass());
+		if (!CinematicShotTrack)
+		{
+			return;
+		}
+
+		UMovieSceneSection* TargetShotSection = MovieSceneHelpers::FindPreviousSection(CinematicShotTrack->GetAllSections(), GetLocalTime().Time.FloorToFrame());
 
 		if (TargetShotSection)
 		{
@@ -9418,8 +9366,16 @@ void FSequencer::StepToPreviousShot()
 		return;
 	}
 
+	UMovieScene* MovieScene = Sequence->GetMovieScene();
+
+	UMovieSceneTrack* CinematicShotTrack = MovieScene->FindTrack(UMovieSceneCinematicShotTrack::StaticClass());
+	if (!CinematicShotTrack)
+	{
+		return;
+	}
+
 	FFrameTime CurrentTime = SubSequenceRange.GetLowerBoundValue() * SubData->OuterToInnerTransform.InverseFromWarp(RootToLocalLoopCounter);
-	UMovieSceneSubSection* PreviousShot = Cast<UMovieSceneSubSection>(FindNextOrPreviousShot(Sequence, CurrentTime.FloorToFrame(), false));
+	UMovieSceneSubSection* PreviousShot = Cast<UMovieSceneSubSection>(MovieSceneHelpers::FindPreviousSection(CinematicShotTrack->GetAllSections(), CurrentTime.FloorToFrame()));
 	if (!PreviousShot)
 	{
 		return;
diff --git a/Engine/Source/Editor/Sequencer/Private/Sequencer.h b/Engine/Source/Editor/Sequencer/Private/Sequencer.h
--- a/Engine/Source/Editor/Sequencer/Private/Sequencer.h
+++ b/Engine/Source/Editor/Sequencer/Private/Sequencer.h
@@ -962,8 +962,6 @@ protected:
 
 	void GetAllKeys(TUniquePtr<FSequencerKeyCollection>& KeyCollection, float DuplicateThresoldSeconds) const;
 
-	UMovieSceneSection* FindNextOrPreviousShot(UMovieSceneSequence* Sequence, FFrameNumber SearchFromTime, const bool bNext) const;
-
 protected:
 	
 	/** Called when a user executes the delete command to delete sections or keys */
diff --git a/Engine/Source/Runtime/MovieScene/Private/MovieSceneCommonHelpers.cpp b/Engine/Source/Runtime/MovieScene/Private/MovieSceneCommonHelpers.cpp
--- a/Engine/Source/Runtime/MovieScene/Private/MovieSceneCommonHelpers.cpp
+++ b/Engine/Source/Runtime/MovieScene/Private/MovieSceneCommonHelpers.cpp
@@ -93,6 +93,86 @@ UMovieSceneSection* MovieSceneHelpers::FindNearestSectionAtTime( TArrayView<UMov
 	return nullptr;
 }
 
+UMovieSceneSection* MovieSceneHelpers::FindNextSection(TArrayView<UMovieSceneSection* const> Sections, FFrameNumber Time)
+{
+	FFrameNumber MinTime = TNumericLimits<FFrameNumber>::Max();
+
+	TMap<FFrameNumber, int32> StartTimeMap;
+	for (int32 SectionIndex = 0; SectionIndex < Sections.Num(); ++SectionIndex)
+	{
+		UMovieSceneSection* ShotSection = Sections[SectionIndex];
+
+		if (ShotSection && ShotSection->HasStartFrame() && !ShotSection->GetRange().Contains(Time))
+		{
+			StartTimeMap.Add(ShotSection->GetInclusiveStartFrame(), SectionIndex);
+		}
+	}
+
+	StartTimeMap.KeySort(TLess<FFrameNumber>());
+
+	int32 NextSectionIndex = -1;
+	for (auto StartTimeIt = StartTimeMap.CreateIterator(); StartTimeIt; ++StartTimeIt)
+	{
+		FFrameNumber StartTime = StartTimeIt->Key;
+		if (StartTime > Time)
+		{
+			FFrameNumber DiffTime = FMath::Abs(StartTime - Time);
+			if (DiffTime < MinTime)
+			{
+				MinTime = DiffTime;
+				NextSectionIndex = StartTimeIt->Value;
+			}
+		}
+	}
+
+	if (NextSectionIndex == -1)
+	{
+		return nullptr;
+	}
+
+	return Sections[NextSectionIndex];
+}
+
+UMovieSceneSection* MovieSceneHelpers::FindPreviousSection(TArrayView<UMovieSceneSection* const> Sections, FFrameNumber Time)
+{
+	FFrameNumber MinTime = TNumericLimits<FFrameNumber>::Max();
+
+	TMap<FFrameNumber, int32> StartTimeMap;
+	for (int32 SectionIndex = 0; SectionIndex < Sections.Num(); ++SectionIndex)
+	{
+		UMovieSceneSection* ShotSection = Sections[SectionIndex];
+
+		if (ShotSection && ShotSection->HasStartFrame() && !ShotSection->GetRange().Contains(Time))
+		{
+			StartTimeMap.Add(ShotSection->GetInclusiveStartFrame(), SectionIndex);
+		}
+	}
+
+	StartTimeMap.KeySort(TLess<FFrameNumber>());
+
+	int32 PreviousSectionIndex = -1;
+	for (auto StartTimeIt = StartTimeMap.CreateIterator(); StartTimeIt; ++StartTimeIt)
+	{
+		FFrameNumber StartTime = StartTimeIt->Key;
+		if (Time >= StartTime)
+		{
+			FFrameNumber DiffTime = FMath::Abs(StartTime - Time);
+			if (DiffTime < MinTime)
+			{
+				MinTime = DiffTime;
+				PreviousSectionIndex = StartTimeIt->Value;
+			}
+		}
+	}
+
+	if (PreviousSectionIndex == -1)
+	{
+		return nullptr;
+	}
+
+	return Sections[PreviousSectionIndex];
+}
+
 bool MovieSceneHelpers::SortOverlappingSections(const UMovieSceneSection* A, const UMovieSceneSection* B)
 {
 	return A->GetRowIndex() == B->GetRowIndex()
diff --git a/Engine/Source/Runtime/MovieScene/Public/MovieSceneCommonHelpers.h b/Engine/Source/Runtime/MovieScene/Public/MovieSceneCommonHelpers.h
--- a/Engine/Source/Runtime/MovieScene/Public/MovieSceneCommonHelpers.h
+++ b/Engine/Source/Runtime/MovieScene/Public/MovieSceneCommonHelpers.h
@@ -63,6 +63,12 @@ public:
 	 */
 	static UMovieSceneSection* FindNearestSectionAtTime( TArrayView<UMovieSceneSection* const> Sections, FFrameNumber Time );
 
+	/** Find the next section that doesn't overlap - the section that has the next closest start time to the requested start time */
+	static UMovieSceneSection* FindNextSection(TArrayView<UMovieSceneSection* const> Sections, FFrameNumber Time);
+
+	/** Find the previous section that doesn't overlap - the section that has the previous closest start time to the requested start time */
+	static UMovieSceneSection* FindPreviousSection(TArrayView<UMovieSceneSection* const> Sections, FFrameNumber Time);
+
 	/*
 	 * Fix up consecutive sections so that there are no gaps
 	 * 
