Sequencer: Follow up fixes to CL #25167672 for drag and drop of level sequences

#rnx

diff --git a/Engine/Source/Editor/MovieSceneTools/Private/TrackEditors/SubTrackEditor.cpp b/Engine/Source/Editor/MovieSceneTools/Private/TrackEditors/SubTrackEditor.cpp
--- a/Engine/Source/Editor/MovieSceneTools/Private/TrackEditors/SubTrackEditor.cpp
+++ b/Engine/Source/Editor/MovieSceneTools/Private/TrackEditors/SubTrackEditor.cpp
@@ -259,7 +259,7 @@ bool FSubTrackEditor::HandleAssetAdded(UObject* Asset, const FGuid& TargetObject
 		return false;
 	}
 
-	if (CanHandleAssetAdded(Sequence))
+	if (!CanHandleAssetAdded(Sequence))
 	{
 		return false;
 	}
@@ -568,9 +568,14 @@ void FSubTrackEditor::CreateNewTake(UMovieSceneSubSection* Section)
 
 void FSubTrackEditor::SwitchTake(UObject* TakeObject)
 {
-	bool bSwitchedTake = false;
+	ChangeTake(Cast<UMovieSceneSequence>(TakeObject));
+}
+
+void FSubTrackEditor::ChangeTake(UMovieSceneSequence* Sequence)
+{
+	bool bChangedTake = false;
 
-	const FScopedTransaction Transaction(LOCTEXT("SwitchTake_Transaction", "Switch Take"));
+	const FScopedTransaction Transaction(LOCTEXT("ChangeTake_Transaction", "Change Take"));
 
 	TArray<UMovieSceneSection*> Sections;
 	GetSequencer()->GetSelectedSections(Sections);
@@ -583,40 +588,34 @@ void FSubTrackEditor::SwitchTake(UObject* TakeObject)
 		}
 
 		UMovieSceneSubSection* Section = Cast<UMovieSceneSubSection>(Sections[SectionIndex]);
+		UMovieSceneSubTrack* SubTrack = CastChecked<UMovieSceneSubTrack>(Section->GetOuter());
 
-		if (TakeObject && TakeObject->IsA(UMovieSceneSequence::StaticClass()))
-		{
-			UMovieSceneSequence* MovieSceneSequence = CastChecked<UMovieSceneSequence>(TakeObject);
-
-			UMovieSceneSubTrack* SubTrack = CastChecked<UMovieSceneSubTrack>(Section->GetOuter());
-
-			TRange<FFrameNumber> NewSectionRange = Section->GetRange();
-			FFrameNumber		 NewSectionStartOffset = Section->Parameters.StartFrameOffset;
-			float                NewSectionTimeScale = Section->Parameters.TimeScale;
-			int32                NewSectionPrerollFrames = Section->GetPreRollFrames();
-			int32                NewRowIndex = Section->GetRowIndex();
-			FFrameNumber         NewSectionStartTime = NewSectionRange.GetLowerBound().IsClosed() ? UE::MovieScene::DiscreteInclusiveLower(NewSectionRange) : 0;
-			int32                NewSectionRowIndex = Section->GetRowIndex();
+		TRange<FFrameNumber> NewSectionRange = Section->GetRange();
+		FFrameNumber		 NewSectionStartOffset = Section->Parameters.StartFrameOffset;
+		float                NewSectionTimeScale = Section->Parameters.TimeScale;
+		int32                NewSectionPrerollFrames = Section->GetPreRollFrames();
+		int32                NewRowIndex = Section->GetRowIndex();
+		FFrameNumber         NewSectionStartTime = NewSectionRange.GetLowerBound().IsClosed() ? UE::MovieScene::DiscreteInclusiveLower(NewSectionRange) : 0;
+		int32                NewSectionRowIndex = Section->GetRowIndex();
 
-			const int32 Duration = (NewSectionRange.GetLowerBound().IsClosed() && NewSectionRange.GetUpperBound().IsClosed()) ? UE::MovieScene::DiscreteSize(NewSectionRange) : 1;
-			UMovieSceneSubSection* NewSection = SubTrack->AddSequence(MovieSceneSequence, NewSectionStartTime, Duration);
+		const int32 Duration = (NewSectionRange.GetLowerBound().IsClosed() && NewSectionRange.GetUpperBound().IsClosed()) ? UE::MovieScene::DiscreteSize(NewSectionRange) : 1;
+		UMovieSceneSubSection* NewSection = SubTrack->AddSequence(Sequence, NewSectionStartTime, Duration);
 
-			if (NewSection != nullptr)
-			{
-				SubTrack->RemoveSection(*Section);
+		if (NewSection != nullptr)
+		{
+			SubTrack->RemoveSection(*Section);
 
-				NewSection->SetRange(NewSectionRange);
-				NewSection->Parameters.StartFrameOffset = NewSectionStartOffset;
-				NewSection->Parameters.TimeScale = NewSectionTimeScale;
-				NewSection->SetPreRollFrames(NewSectionPrerollFrames);
-				NewSection->SetRowIndex(NewSectionRowIndex);
+			NewSection->SetRange(NewSectionRange);
+			NewSection->Parameters.StartFrameOffset = NewSectionStartOffset;
+			NewSection->Parameters.TimeScale = NewSectionTimeScale;
+			NewSection->SetPreRollFrames(NewSectionPrerollFrames);
+			NewSection->SetRowIndex(NewSectionRowIndex);
 
-				bSwitchedTake = true;
-			}
+			bChangedTake = true;
 		}
 	}
 
-	if (bSwitchedTake)
+	if (bChangedTake)
 	{
 		GetSequencer()->NotifyMovieSceneDataChanged(EMovieSceneDataChangeType::MovieSceneStructureItemsChanged);
 	}
@@ -818,7 +817,9 @@ FKeyPropertyResult FSubTrackEditor::AddKeyInternal(FFrameNumber KeyTime, UMovieS
 
 	if (CanAddSubSequence(*InMovieSceneSequence))
 	{
-		UMovieSceneSubTrack* SubTrack = FindOrCreateSubTrack(InMovieSceneSequence->GetMovieScene(), InTrack);
+		UMovieScene* MovieScene = GetFocusedMovieScene();
+
+		UMovieSceneSubTrack* SubTrack = FindOrCreateSubTrack(MovieScene, InTrack);
 
 		const FFrameRate TickResolution = InMovieSceneSequence->GetMovieScene()->GetTickResolution();
 		const FQualifiedFrameTime InnerDuration = FQualifiedFrameTime(
@@ -857,7 +858,9 @@ FKeyPropertyResult FSubTrackEditor::HandleSequenceAdded(FFrameNumber KeyTime, UM
 {
 	FKeyPropertyResult KeyPropertyResult;
 
-	UMovieSceneSubTrack* SubTrack = FindOrCreateSubTrack(Sequence->GetMovieScene(), Track);
+	UMovieScene* MovieScene = GetFocusedMovieScene();
+
+	UMovieSceneSubTrack* SubTrack = FindOrCreateSubTrack(MovieScene, Track);
 
 	const FFrameRate TickResolution = Sequence->GetMovieScene()->GetTickResolution();
 	const FQualifiedFrameTime InnerDuration = FQualifiedFrameTime(
diff --git a/Engine/Source/Editor/MovieSceneTools/Public/TrackEditors/SubTrackEditor.h b/Engine/Source/Editor/MovieSceneTools/Public/TrackEditors/SubTrackEditor.h
--- a/Engine/Source/Editor/MovieSceneTools/Public/TrackEditors/SubTrackEditor.h
+++ b/Engine/Source/Editor/MovieSceneTools/Public/TrackEditors/SubTrackEditor.h
@@ -68,7 +68,10 @@ public:
 	/** Create a new take of the given section */
 	virtual void CreateNewTake(UMovieSceneSubSection* Section);
 
-	/** Switch the selected section's take */
+	/** Switch the selected section's take sequence */
+	virtual void ChangeTake(UMovieSceneSequence* Sequence);
+
+	UE_DEPRECATED(5.3, "SwitchTake has been deprecated, please use ChangeTake(UMovieSceneSequence*)")
 	virtual void SwitchTake(UObject* TakeObject);
 
 	/** Generate a menu for takes for this section */
