Sequencer: Follow up fixes to CL #25167672 for drag and drop of level sequences

#rnx

diff --git a/Engine/Source/Editor/MovieSceneTools/Private/TrackEditors/SubTrackEditor.cpp b/Engine/Source/Editor/MovieSceneTools/Private/TrackEditors/SubTrackEditor.cpp
--- a/Engine/Source/Editor/MovieSceneTools/Private/TrackEditors/SubTrackEditor.cpp
+++ b/Engine/Source/Editor/MovieSceneTools/Private/TrackEditors/SubTrackEditor.cpp
@@ -642,15 +642,14 @@ void FSubTrackEditor::AddTakesMenu(UMovieSceneSubSection* Section, FMenuBuilder&
 		uint32 TakeNumber = INDEX_NONE;
 		if (MovieSceneToolHelpers::GetTakeNumber(Section, ThisAssetData, TakeNumber))
 		{
-			UObject* TakeObject = ThisAssetData.GetAsset();
-
-			if (TakeObject)
+			UMovieSceneSequence* Sequence = Cast<UMovieSceneSequence>(ThisAssetData.GetAsset());
+			if (Sequence)
 			{
 				MenuBuilder.AddMenuEntry(
 					FText::Format(LOCTEXT("TakeNumber", "Take {0}"), FText::AsNumber(TakeNumber)),
-					FText::Format(LOCTEXT("TakeNumberTooltip", "Switch to {0}"), FText::FromString(TakeObject->GetPathName())),
+					FText::Format(LOCTEXT("TakeNumberTooltip", "Switch to {0}"), FText::FromString(Sequence->GetPathName())),
 					TakeNumber == CurrentTakeNumber ? FSlateIcon(FAppStyle::GetAppStyleSetName(), "Sequencer.Star") : FSlateIcon(FAppStyle::GetAppStyleSetName(), "Sequencer.Empty"),
-					FUIAction(FExecuteAction::CreateSP(this, &FSubTrackEditor::SwitchTake, TakeObject))
+					FUIAction(FExecuteAction::CreateSP(this, &FSubTrackEditor::ChangeTake, Sequence))
 				);
 			}
 		}
