Sequencer: Duplicate LOCTEXT fixes

#rnx

diff --git a/Engine/Source/Editor/MovieSceneTools/Private/TrackEditors/SubTrackEditor.cpp b/Engine/Source/Editor/MovieSceneTools/Private/TrackEditors/SubTrackEditor.cpp
--- a/Engine/Source/Editor/MovieSceneTools/Private/TrackEditors/SubTrackEditor.cpp
+++ b/Engine/Source/Editor/MovieSceneTools/Private/TrackEditors/SubTrackEditor.cpp
@@ -661,7 +661,9 @@ void FSubTrackEditor::AddTakesMenu(UMovieSceneSubSection* Section, FMenuBuilder&
 				FText MetaDataText = FSubTrackEditorUtil::GetMetaDataText(Sequence);
 				MenuBuilder.AddMenuEntry(
 					FText::Format(LOCTEXT("TakeNumber", "Take {0}"), FText::AsNumber(TakeNumber)),
-					MetaDataText.IsEmpty() ? FText::Format(LOCTEXT("TakeNumberTooltip", "Change to {0}"), FText::FromString(Sequence->GetPathName())) : FText::Format(LOCTEXT("TakeNumberTooltip", "Change to {0}\n\n{1}"), FText::FromString(Sequence->GetPathName()), MetaDataText),
+					MetaDataText.IsEmpty() ? 
+					FText::Format(LOCTEXT("TakeNumberTooltip", "Change to {0}"), FText::FromString(Sequence->GetPathName())) : 
+					FText::Format(LOCTEXT("TakeNumberWithMetaDataTooltip", "Change to {0}\n\n{1}"), FText::FromString(Sequence->GetPathName()), MetaDataText),
 					TakeNumber == CurrentTakeNumber ? FSlateIcon(FAppStyle::GetAppStyleSetName(), "Sequencer.Star") : FSlateIcon(FAppStyle::GetAppStyleSetName(), "Sequencer.Empty"),
 					FUIAction(FExecuteAction::CreateSP(this, &FSubTrackEditor::ChangeTake, Sequence))
 				);
diff --git a/Engine/Source/Editor/MovieSceneTools/Public/TrackEditors/SubTrackEditorBase.h b/Engine/Source/Editor/MovieSceneTools/Public/TrackEditors/SubTrackEditorBase.h
--- a/Engine/Source/Editor/MovieSceneTools/Public/TrackEditors/SubTrackEditorBase.h
+++ b/Engine/Source/Editor/MovieSceneTools/Public/TrackEditors/SubTrackEditorBase.h
@@ -244,7 +244,7 @@ FText TSubSectionMixin<ParentSectionClass>::GetSectionToolTip() const
 	}
 	else
 	{
-		return FText::Format(LOCTEXT("ToolTipContentFormat", "{0} - {1} ({2} frames @ {3})\n\n{4}"),
+		return FText::Format(LOCTEXT("ToolTipContentFormatWithMetaData", "{0} - {1} ({2} frames @ {3})\n\n{4}"),
 			InnerStartFrame,
 			InnerEndFrame,
 			InnerFrameLength,
diff --git a/Engine/Source/Editor/Sequencer/Private/SSequencer.cpp b/Engine/Source/Editor/Sequencer/Private/SSequencer.cpp
--- a/Engine/Source/Editor/Sequencer/Private/SSequencer.cpp
+++ b/Engine/Source/Editor/Sequencer/Private/SSequencer.cpp
@@ -2675,7 +2675,7 @@ TSharedRef<SWidget> SSequencer::MakeKeyGroupMenu()
 	{
 		MenuBuilder.AddMenuEntry(
 			LOCTEXT("SetKeyInterpolationSmartAuto", "Cubic (Smart Auto)"),
-			LOCTEXT("SetKeyInterpolationAutoTooltip", "Set key interpolation to smart auto"),
+			LOCTEXT("SetKeyInterpolationSmartAutoTooltip", "Set key interpolation to smart auto"),
 			FSlateIcon(FAppStyle::GetAppStyleSetName(), "Sequencer.IconKeySmartAuto"),
 			FUIAction(
 				FExecuteAction::CreateLambda([this] { SequencerPtr.Pin()->SetKeyInterpolation(EMovieSceneKeyInterpolation::SmartAuto); }),
