Sequencer: Duplicate LOCTEXT fixes

#rnx

diff --git a/Engine/Source/Editor/MovieSceneTools/Private/TrackEditors/SubTrackEditorBase.cpp b/Engine/Source/Editor/MovieSceneTools/Private/TrackEditors/SubTrackEditorBase.cpp
--- a/Engine/Source/Editor/MovieSceneTools/Private/TrackEditors/SubTrackEditorBase.cpp
+++ b/Engine/Source/Editor/MovieSceneTools/Private/TrackEditors/SubTrackEditorBase.cpp
@@ -481,7 +481,7 @@ FText FSubTrackEditorUtil::GetMetaDataText(const UMovieSceneSequence* Sequence)
 		return FText::GetEmpty();
 	}
 
-	return FText::Format(LOCTEXT("ToolTipContentFormat", "Author: {0}\nCreated: {1}\nNotes: {2}"),
+	return FText::Format(LOCTEXT("MetaDataContentFormat", "Author: {0}\nCreated: {1}\nNotes: {2}"),
 		FText::FromString(MetaData->GetAuthor()),
 		FText::AsDateTime(MetaData->GetCreated()),
 		FText::FromString(MetaData->GetNotes())
