Sequencer: Subsequence directory

diff --git a/Engine/Plugins/MovieScene/LevelSequenceEditor/Source/LevelSequenceEditor/Private/LevelSequenceEditorToolkit.cpp b/Engine/Plugins/MovieScene/LevelSequenceEditor/Source/LevelSequenceEditor/Private/LevelSequenceEditorToolkit.cpp
--- a/Engine/Plugins/MovieScene/LevelSequenceEditor/Source/LevelSequenceEditor/Private/LevelSequenceEditorToolkit.cpp
+++ b/Engine/Plugins/MovieScene/LevelSequenceEditor/Source/LevelSequenceEditor/Private/LevelSequenceEditorToolkit.cpp
@@ -735,7 +735,7 @@ void FLevelSequenceEditorToolkit::AddShot(UMovieSceneCinematicShotTrack* ShotTra
 		int32 RowIndex = 0;
 		for (auto SubSequenceName : LevelSequenceSettings->SubSequenceNames)
 		{
-			FString SubSequenceAssetName = ShotAssetName + ProjectSettings->SubSequenceSeparator + SubSequenceName.ToString();
+			FString SubSequenceAssetName = ShotAssetName + ProjectSettings->TakeSeparator + SubSequenceName.ToString();
 
 			UMovieSceneSequence* SubSequence = nullptr;
 			if (!LevelSequenceSettings->bInstanceSubSequences || ShotTrack->GetAllSections().Num() == 1)
@@ -750,7 +750,7 @@ void FLevelSequenceEditorToolkit::AddShot(UMovieSceneCinematicShotTrack* ShotTra
 				UMovieSceneSequence* FirstShotSequence = FirstShotSubSection->GetSequence();
 				UMovieSceneSubTrack* FirstShotSubTrack = Cast<UMovieSceneSubTrack>(FirstShotSequence->GetMovieScene()->FindTrack(UMovieSceneSubTrack::StaticClass()));
 			
-				FString FirstShotSubSequenceAssetName = FirstShotAssetName + ProjectSettings->SubSequenceSeparator + SubSequenceName.ToString();
+				FString FirstShotSubSequenceAssetName = FirstShotAssetName + ProjectSettings->TakeSeparator + SubSequenceName.ToString();
 
 				for (auto Section : FirstShotSubTrack->GetAllSections())
 				{
@@ -851,8 +851,8 @@ void FLevelSequenceEditorToolkit::HandleLevelSequenceWithShotsCreated(UObject* L
 	{
 		ShotEndTime += ShotDuration;
 
-		FString ShotName = MovieSceneToolHelpers::GenerateNewShotName(ShotTrack->GetAllSections(), ShotStartTime);
-		FString ShotPackagePath = MovieSceneToolHelpers::GenerateNewShotPath(LevelSequenceWithShots->GetMovieScene(), ShotName);
+		FString ShotName = MovieSceneToolHelpers::GenerateNewSubsequenceName(ShotTrack->GetAllSections(), ProjectSettings->ShotPrefix, ShotStartTime);
+		FString ShotPackagePath = MovieSceneToolHelpers::GenerateNewSubsequencePath(LevelSequenceWithShots->GetMovieScene(), ProjectSettings->ShotPrefix, ShotName);
 
 		if (ShotIndex == 0)
 		{
diff --git a/Engine/Source/Editor/MovieSceneTools/Private/MovieSceneToolHelpers.cpp b/Engine/Source/Editor/MovieSceneTools/Private/MovieSceneToolHelpers.cpp
--- a/Engine/Source/Editor/MovieSceneTools/Private/MovieSceneToolHelpers.cpp
+++ b/Engine/Source/Editor/MovieSceneTools/Private/MovieSceneToolHelpers.cpp
@@ -329,10 +329,11 @@ bool MovieSceneToolHelpers::ParseShotName(const FString& InShotName, FString& Sh
 		}
 	}
 
-	// If take number wasn't found, search backwards to find the first take separator and assume [shot prefix]_[take number]
+	// If take number wasn't found, start over with the original shot name, search backwards to find the first take separator and assume [shot prefix]_[take number]
 	//
 	if (!ParsedTakeNumber.IsSet())
 	{
+		ShotName = InShotName;
 		int32 LastSlashPos = ShotName.Find(ProjectSettings->TakeSeparator, ESearchCase::IgnoreCase, ESearchDir::FromEnd);
 		if (LastSlashPos != INDEX_NONE)
 		{
@@ -369,7 +370,6 @@ bool MovieSceneToolHelpers::ParseShotName(const FString& InShotName, FString& Sh
 	return FirstShotNumberIndex != INDEX_NONE;
 }
 
-
 FString MovieSceneToolHelpers::ComposeShotName(const FString& ShotPrefix, uint32 ShotNumber, uint32 TakeNumber, uint32 ShotNumberDigits, uint32 TakeNumberDigits)
 {
 	const UMovieSceneToolsProjectSettings* ProjectSettings = GetDefault<UMovieSceneToolsProjectSettings>();
@@ -404,6 +404,12 @@ bool IsPackageNameUnique(const TArray<FAssetData>& ObjectList, FString& NewPacka
 }
 
 FString MovieSceneToolHelpers::GenerateNewShotPath(UMovieScene* SequenceMovieScene, FString& NewShotName)
+{
+	const UMovieSceneToolsProjectSettings* ProjectSettings = GetDefault<UMovieSceneToolsProjectSettings>();
+	return GenerateNewSubsequencePath(SequenceMovieScene, ProjectSettings->ShotDirectory, NewShotName);
+}
+
+FString MovieSceneToolHelpers::GenerateNewSubsequencePath(UMovieScene * SequenceMovieScene, const FString& SubsequenceDirectory, FString &NewShotName)
 {
 	const UMovieSceneToolsProjectSettings* ProjectSettings = GetDefault<UMovieSceneToolsProjectSettings>();
 
@@ -425,18 +431,18 @@ FString MovieSceneToolHelpers::GenerateNewShotPath(UMovieScene* SequenceMovieSce
 	uint32 TakeNumberDigits = ProjectSettings->TakeNumDigits;
 	ParseShotName(NewShotName, NewShotPrefix, NewShotNumber, NewTakeNumber, ShotNumberDigits, TakeNumberDigits);
 
-	FString NewShotDirectory = ComposeShotName(NewShotPrefix, NewShotNumber, INDEX_NONE, ShotNumberDigits, TakeNumberDigits);
-	FString NewShotPath = SequencePath;
+	FString NewDirectory = ComposeShotName(NewShotPrefix, NewShotNumber, INDEX_NONE, ShotNumberDigits, TakeNumberDigits);
+	FString NewPath = SequencePath;
 
-	FString ShotDirectory = ProjectSettings->ShotDirectory;
-	if (!ShotDirectory.IsEmpty())
+	FString Directory = SubsequenceDirectory;
+	if (!Directory.IsEmpty())
 	{
-		NewShotPath /= ShotDirectory;
+		NewPath /= Directory;
 	}
-	NewShotPath /= NewShotDirectory; // put this in the shot directory, ie. /Game/cine/max/shots/shot0010
+	NewPath /= NewDirectory; // put this in the shot directory, ie. /Game/cine/max/shots/shot0010
 
 	// Make sure this shot path is unique
-	FString NewPackageName = NewShotPath;
+	FString NewPackageName = NewPath;
 	NewPackageName /= NewShotName; // ie. /Game/cine/max/shots/shot0010/shot0010_001
 	if (!IsPackageNameUnique(ObjectList, NewPackageName))
 	{
@@ -444,15 +450,15 @@ FString MovieSceneToolHelpers::GenerateNewShotPath(UMovieScene* SequenceMovieSce
 		{
 			NewShotNumber += ProjectSettings->ShotIncrement;
 			NewShotName = ComposeShotName(NewShotPrefix, NewShotNumber, NewTakeNumber, ShotNumberDigits, TakeNumberDigits);
-			NewShotDirectory = ComposeShotName(NewShotPrefix, NewShotNumber, INDEX_NONE, ShotNumberDigits, TakeNumberDigits);
-			NewShotPath = SequencePath;
-			if (!ShotDirectory.IsEmpty())
+			NewDirectory = ComposeShotName(NewShotPrefix, NewShotNumber, INDEX_NONE, ShotNumberDigits, TakeNumberDigits);
+			NewPath = SequencePath;
+			if (!Directory.IsEmpty())
 			{
-				NewShotPath /= ShotDirectory;
+				NewPath /= Directory;
 			}
-			NewShotPath /= NewShotDirectory;
+			NewPath /= NewDirectory;
 
-			NewPackageName = NewShotPath;
+			NewPackageName = NewPath;
 			NewPackageName /= NewShotName;
 			if (IsPackageNameUnique(ObjectList, NewPackageName))
 			{
@@ -461,13 +467,20 @@ FString MovieSceneToolHelpers::GenerateNewShotPath(UMovieScene* SequenceMovieSce
 		}
 	}
 
-	return NewShotPath;
+	return NewPath;
 }
 
 FString MovieSceneToolHelpers::GenerateNewShotName(const TArray<UMovieSceneSection*>& AllSections, FFrameNumber Time)
 {
 	const UMovieSceneToolsProjectSettings* ProjectSettings = GetDefault<UMovieSceneToolsProjectSettings>();
 
+	return GenerateNewSubsequenceName(AllSections, ProjectSettings->ShotPrefix, Time);
+}
+
+FString MovieSceneToolHelpers::GenerateNewSubsequenceName(const TArray<UMovieSceneSection*>&AllSections, const FString& SubsequencePrefix, FFrameNumber Time)
+{
+	const UMovieSceneToolsProjectSettings* ProjectSettings = GetDefault<UMovieSceneToolsProjectSettings>();
+
 	UMovieSceneSubSection* CurrentSection = Cast<UMovieSceneSubSection>(MovieSceneHelpers::FindSectionAtTime(AllSections, Time));
 	UMovieSceneSubSection* NextSection = Cast<UMovieSceneSubSection>(MovieSceneHelpers::FindNextSection(AllSections, Time));
 
@@ -493,15 +506,10 @@ FString MovieSceneToolHelpers::GenerateNewShotName(const TArray<UMovieSceneSecti
 		CurrentSectionName = ShotSection->GetShotDisplayName();
 	}
 
-	// There aren't any sections, let's create the first shot name
-	if (NextSection == nullptr || CurrentSection == nullptr)
-	{
-		// Default case
-	}
-	// This is the last shot
-	else if (CurrentSection == NextSection)
+	// This is the first or last shot
+	if ((CurrentSection == nullptr && NextSection) || (CurrentSection != nullptr && CurrentSection == NextSection))
 	{
-		FString NextShotPrefix = ProjectSettings->ShotPrefix;
+		FString NextShotPrefix = SubsequencePrefix;
 		uint32 NextShotNumber = ProjectSettings->FirstShotNumber;
 		uint32 NextTakeNumber = ProjectSettings->FirstTakeNumber;
 		uint32 ShotNumberDigits = ProjectSettings->ShotNumDigits;
@@ -509,20 +517,30 @@ FString MovieSceneToolHelpers::GenerateNewShotName(const TArray<UMovieSceneSecti
 
 		if (ParseShotName(NextSectionName, NextShotPrefix, NextShotNumber, NextTakeNumber, ShotNumberDigits, TakeNumberDigits))
 		{
-			uint32 NewShotNumber = NextShotNumber != INDEX_NONE ? NextShotNumber + ProjectSettings->ShotIncrement : INDEX_NONE;
-			return ComposeShotName(NextShotPrefix, NewShotNumber, ProjectSettings->FirstTakeNumber, ShotNumberDigits, TakeNumberDigits);
+			// Valid shot number
+			if (NextShotNumber != INDEX_NONE)
+			{
+				uint32 NewShotNumber = NextShotNumber + ProjectSettings->ShotIncrement;
+				return ComposeShotName(NextShotPrefix, NewShotNumber, ProjectSettings->FirstTakeNumber, ShotNumberDigits, TakeNumberDigits);
+			}
+			// No shot number, but valid take number
+			else if (NextTakeNumber != INDEX_NONE)
+			{
+				uint32 NewTakeNumber = NextTakeNumber + ProjectSettings->ShotIncrement;
+				return ComposeShotName(NextShotPrefix, INDEX_NONE, NewTakeNumber, ShotNumberDigits, TakeNumberDigits);
+			}
 		}
 	}
 	// This is in between two shots
-	else 
+	else if (CurrentSection && NextSection)
 	{
-		FString CurrentShotPrefix = ProjectSettings->ShotPrefix;
+		FString CurrentShotPrefix = SubsequencePrefix;
 		uint32 CurrentShotNumber = ProjectSettings->FirstShotNumber;
 		uint32 CurrentTakeNumber = ProjectSettings->FirstTakeNumber;
 		uint32 CurrentShotNumberDigits = ProjectSettings->ShotNumDigits;
 		uint32 CurrentTakeNumberDigits = ProjectSettings->TakeNumDigits;
 
-		FString NextShotPrefix = ProjectSettings->ShotPrefix;
+		FString NextShotPrefix = SubsequencePrefix;
 		uint32 NextShotNumber = ProjectSettings->FirstShotNumber;
 		uint32 NextTakeNumber = ProjectSettings->FirstTakeNumber;
 		uint32 NextShotNumberDigits = ProjectSettings->ShotNumDigits;
@@ -531,16 +549,31 @@ FString MovieSceneToolHelpers::GenerateNewShotName(const TArray<UMovieSceneSecti
 		if (ParseShotName(CurrentSectionName, CurrentShotPrefix, CurrentShotNumber, CurrentTakeNumber, CurrentShotNumberDigits, CurrentTakeNumberDigits) &&
 			ParseShotName(NextSectionName, NextShotPrefix, NextShotNumber, NextTakeNumber, NextShotNumberDigits, NextTakeNumberDigits))
 		{
-			if (CurrentShotNumber < NextShotNumber)
+			// Valid shot numbers
+			if (NextShotNumber != INDEX_NONE && CurrentShotNumber != INDEX_NONE)
 			{
-				uint32 NewShotNumber = NextShotNumber != INDEX_NONE && CurrentShotNumber != INDEX_NONE ? CurrentShotNumber + ( (NextShotNumber - CurrentShotNumber) / 2) : INDEX_NONE; // what if we can't find one? or conflicts with another?
+				uint32 NewShotNumber = CurrentShotNumber + ProjectSettings->ShotIncrement;
+				if (NextShotNumber - CurrentShotNumber > 1)
+				{
+					NewShotNumber = CurrentShotNumber + ( (NextShotNumber - CurrentShotNumber) / 2); // what if we can't find one? or conflicts with another?
+				}
 				return ComposeShotName(CurrentShotPrefix, NewShotNumber, ProjectSettings->FirstTakeNumber, CurrentShotNumberDigits, CurrentTakeNumberDigits); // use before or next shot?
 			}
+			// No shot numbers, but valid take numbers
+			else if (NextTakeNumber != INDEX_NONE && CurrentTakeNumber != INDEX_NONE)
+			{
+				uint32 NewTakeNumber = CurrentTakeNumber + ProjectSettings->ShotIncrement;
+				if (NextTakeNumber - CurrentTakeNumber > 1)
+				{
+					NewTakeNumber = CurrentTakeNumber + ( (NextTakeNumber - CurrentTakeNumber) / 2); // what if we can't find one? or conflicts with another?
+				}
+				return ComposeShotName(CurrentShotPrefix, INDEX_NONE, NewTakeNumber, CurrentShotNumberDigits, CurrentTakeNumberDigits); // use before or next shot?
+			}
 		}
 	}
 
 	// Default case
-	return ComposeShotName(ProjectSettings->ShotPrefix, ProjectSettings->FirstShotNumber, ProjectSettings->FirstTakeNumber, ProjectSettings->ShotNumDigits, ProjectSettings->TakeNumDigits);
+	return ComposeShotName(SubsequencePrefix, ProjectSettings->FirstShotNumber, ProjectSettings->FirstTakeNumber, ProjectSettings->ShotNumDigits, ProjectSettings->TakeNumDigits);
 }
 
 UMovieSceneSequence* MovieSceneToolHelpers::CreateSequence(FString& NewSequenceName, FString& NewSequencePath, UMovieSceneSubSection* SectionToDuplicate)
diff --git a/Engine/Source/Editor/MovieSceneTools/Private/MovieSceneToolsProjectSettings.cpp b/Engine/Source/Editor/MovieSceneTools/Private/MovieSceneToolsProjectSettings.cpp
--- a/Engine/Source/Editor/MovieSceneTools/Private/MovieSceneToolsProjectSettings.cpp
+++ b/Engine/Source/Editor/MovieSceneTools/Private/MovieSceneToolsProjectSettings.cpp
@@ -6,7 +6,9 @@
 UMovieSceneToolsProjectSettings::UMovieSceneToolsProjectSettings()
 	: DefaultStartTime(0.f)
 	, DefaultDuration(5.f)
+	, SubsequenceDirectory(TEXT("subsequences"))
 	, ShotDirectory(TEXT("shots"))
+	, SubsequencePrefix(TEXT("subsequence"))
 	, ShotPrefix(TEXT("shot"))
 	, FirstShotNumber(10)
 	, ShotIncrement(10)
@@ -14,5 +16,4 @@ UMovieSceneToolsProjectSettings::UMovieSceneToolsProjectSettings()
 	, TakeNumDigits(2)
 	, FirstTakeNumber(1)
 	, TakeSeparator(TEXT("_"))
-	, SubSequenceSeparator(TEXT("_"))
 { }
diff --git a/Engine/Source/Editor/MovieSceneTools/Private/MovieSceneToolsProjectSettingsCustomization.cpp b/Engine/Source/Editor/MovieSceneTools/Private/MovieSceneToolsProjectSettingsCustomization.cpp
--- a/Engine/Source/Editor/MovieSceneTools/Private/MovieSceneToolsProjectSettingsCustomization.cpp
+++ b/Engine/Source/Editor/MovieSceneTools/Private/MovieSceneToolsProjectSettingsCustomization.cpp
@@ -23,9 +23,6 @@ void FMovieSceneToolsProjectSettingsCustomization::CustomizeDetails(IDetailLayou
 {
 	TSharedRef<IPropertyHandle> TakeSeparatorProperty = DetailBuilder.GetProperty(GET_MEMBER_NAME_CHECKED(UMovieSceneToolsProjectSettings, TakeSeparator));
 	TakeSeparatorProperty->SetOnPropertyValueChanged(FSimpleDelegate::CreateRaw(this, &FMovieSceneToolsProjectSettingsCustomization::OnTakeSeparatorUpdated));
-
-	TSharedRef<IPropertyHandle> SubSequenceSeparatorProperty = DetailBuilder.GetProperty(GET_MEMBER_NAME_CHECKED(UMovieSceneToolsProjectSettings, SubSequenceSeparator));
-	SubSequenceSeparatorProperty->SetOnPropertyValueChanged(FSimpleDelegate::CreateRaw(this, &FMovieSceneToolsProjectSettingsCustomization::OnSubSequenceSeparatorUpdated));
 }
 
 void FMovieSceneToolsProjectSettingsCustomization::OnTakeSeparatorUpdated()
@@ -42,18 +39,3 @@ void FMovieSceneToolsProjectSettingsCustomization::OnTakeSeparatorUpdated()
 		ProjectSettings->TakeSeparator = TEXT("_");
 	}
 }
-
-void FMovieSceneToolsProjectSettingsCustomization::OnSubSequenceSeparatorUpdated()
-{
-	UMovieSceneToolsProjectSettings* ProjectSettings = GetMutableDefault<UMovieSceneToolsProjectSettings>();
-
-	FString SubSequenceSeparator = ProjectSettings->SubSequenceSeparator;
-
-	// Make sure the subsequence separator is a valid character
-	FText OutErrorMessage;
-	if (!FName(*SubSequenceSeparator).IsValidXName(INVALID_OBJECTNAME_CHARACTERS INVALID_LONGPACKAGE_CHARACTERS, &OutErrorMessage))
-	{
-		UE_LOG(LogMovieScene, Warning, TEXT("Invalid separator: %s"), *OutErrorMessage.ToString());
-		ProjectSettings->SubSequenceSeparator = TEXT("_");
-	}
-}
diff --git a/Engine/Source/Editor/MovieSceneTools/Private/MovieSceneToolsProjectSettingsCustomization.h b/Engine/Source/Editor/MovieSceneTools/Private/MovieSceneToolsProjectSettingsCustomization.h
--- a/Engine/Source/Editor/MovieSceneTools/Private/MovieSceneToolsProjectSettingsCustomization.h
+++ b/Engine/Source/Editor/MovieSceneTools/Private/MovieSceneToolsProjectSettingsCustomization.h
@@ -15,5 +15,4 @@ public:
 
 private:
 	void OnTakeSeparatorUpdated();
-	void OnSubSequenceSeparatorUpdated();
 };
diff --git a/Engine/Source/Editor/MovieSceneTools/Private/TrackEditors/CinematicShotTrackEditor.cpp b/Engine/Source/Editor/MovieSceneTools/Private/TrackEditors/CinematicShotTrackEditor.cpp
--- a/Engine/Source/Editor/MovieSceneTools/Private/TrackEditors/CinematicShotTrackEditor.cpp
+++ b/Engine/Source/Editor/MovieSceneTools/Private/TrackEditors/CinematicShotTrackEditor.cpp
@@ -223,6 +223,18 @@ FString FCinematicShotTrackEditor::GetSubSectionDisplayName(const UMovieSceneSub
 	return Cast<UMovieSceneCinematicShotSection>(Section)->GetShotDisplayName();
 }
 
+FString FCinematicShotTrackEditor::GetDefaultSubsequenceName() const
+{
+	const UMovieSceneToolsProjectSettings* ProjectSettings = GetDefault<UMovieSceneToolsProjectSettings>();
+	return ProjectSettings->ShotPrefix;
+}
+
+FString FCinematicShotTrackEditor::GetDefaultSubsequenceDirectory() const
+{
+	const UMovieSceneToolsProjectSettings* ProjectSettings = GetDefault<UMovieSceneToolsProjectSettings>();
+	return ProjectSettings->ShotDirectory;
+}
+
 TSubclassOf<UMovieSceneSubTrack> FCinematicShotTrackEditor::GetSubTrackClass() const
 {
 	return UMovieSceneCinematicShotTrack::StaticClass();
diff --git a/Engine/Source/Editor/MovieSceneTools/Private/TrackEditors/SubTrackEditor.cpp b/Engine/Source/Editor/MovieSceneTools/Private/TrackEditors/SubTrackEditor.cpp
--- a/Engine/Source/Editor/MovieSceneTools/Private/TrackEditors/SubTrackEditor.cpp
+++ b/Engine/Source/Editor/MovieSceneTools/Private/TrackEditors/SubTrackEditor.cpp
@@ -25,6 +25,7 @@
 #include "MovieSceneMetaData.h"
 #include "MovieSceneSequence.h"
 #include "MovieSceneToolHelpers.h"
+#include "MovieSceneToolsProjectSettings.h"
 #include "Misc/QualifiedFrameTime.h"
 #include "MovieSceneTimeHelpers.h"
 #include "EngineAnalytics.h"
@@ -459,8 +460,8 @@ void FSubTrackEditor::InsertSection(UMovieSceneTrack* Track)
 
 	UMovieSceneSubTrack* SubTrack = FindOrCreateSubTrack(MovieScene, Track);
 
-	FString NewSequenceName = MovieSceneToolHelpers::GenerateNewShotName(SubTrack->GetAllSections(), NewSectionStartTime.FrameNumber);
-	FString NewSequencePath = MovieSceneToolHelpers::GenerateNewShotPath(GetSequencer()->GetFocusedMovieSceneSequence()->GetMovieScene(), NewSequenceName);
+	FString NewSequenceName = MovieSceneToolHelpers::GenerateNewSubsequenceName(SubTrack->GetAllSections(), GetDefaultSubsequenceName(), NewSectionStartTime.FrameNumber);
+	FString NewSequencePath = MovieSceneToolHelpers::GenerateNewSubsequencePath(GetSequencer()->GetFocusedMovieSceneSequence()->GetMovieScene(), GetDefaultSubsequenceDirectory(), NewSequenceName);
 
 	FAssetToolsModule& AssetToolsModule = FModuleManager::GetModuleChecked<FAssetToolsModule>("AssetTools");
 	AssetToolsModule.Get().CreateUniqueAssetName(NewSequencePath + TEXT("/") + NewSequenceName, TEXT(""), NewSequencePath, NewSequenceName);
@@ -488,7 +489,7 @@ void FSubTrackEditor::DuplicateSection(UMovieSceneSubSection* Section)
 	UMovieSceneSubTrack* SubTrack = CastChecked<UMovieSceneSubTrack>(Section->GetOuter());
 
 	FFrameNumber StartTime = Section->HasStartFrame() ? Section->GetInclusiveStartFrame() : 0;
-	FString NewSectionName = MovieSceneToolHelpers::GenerateNewShotName(SubTrack->GetAllSections(), StartTime);
+	FString NewSectionName = MovieSceneToolHelpers::GenerateNewSubsequenceName(SubTrack->GetAllSections(), GetDefaultSubsequenceName(), StartTime);
 	FString NewSequencePath = FPaths::GetPath(Section->GetSequence()->GetPathName());
 
 	// Duplicate the section and put it on the next available row
@@ -745,8 +746,8 @@ bool FSubTrackEditor::CanAddSubSequence(const UMovieSceneSequence& Sequence) con
 }
 
 FText FSubTrackEditor::GetSubTrackName() const
-{ 
-	return LOCTEXT("SubTrackName", "Subsequence"); 
+{
+	return LOCTEXT("SubTrackName", "Subsequence");
 }
 
 FText FSubTrackEditor::GetSubTrackToolTip() const
@@ -764,6 +765,18 @@ FString FSubTrackEditor::GetSubSectionDisplayName(const UMovieSceneSubSection* S
 	return Section && Section->GetSequence() ? Section->GetSequence()->GetName() : FString();
 }
 
+FString FSubTrackEditor::GetDefaultSubsequenceName() const
+{
+	const UMovieSceneToolsProjectSettings* ProjectSettings = GetDefault<UMovieSceneToolsProjectSettings>();
+	return ProjectSettings->SubsequencePrefix;
+}
+
+FString FSubTrackEditor::GetDefaultSubsequenceDirectory() const
+{
+	const UMovieSceneToolsProjectSettings* ProjectSettings = GetDefault<UMovieSceneToolsProjectSettings>();
+	return ProjectSettings->SubsequenceDirectory;
+}
+
 TSubclassOf<UMovieSceneSubTrack> FSubTrackEditor::GetSubTrackClass() const
 {
 	return UMovieSceneSubTrack::StaticClass();
diff --git a/Engine/Source/Editor/MovieSceneTools/Public/MovieSceneToolHelpers.h b/Engine/Source/Editor/MovieSceneTools/Public/MovieSceneToolHelpers.h
--- a/Engine/Source/Editor/MovieSceneTools/Public/MovieSceneToolHelpers.h
+++ b/Engine/Source/Editor/MovieSceneTools/Public/MovieSceneToolHelpers.h
@@ -154,14 +154,19 @@ public:
 	static FString ComposeShotName(const FString& ShotPrefix, uint32 ShotNumber, uint32 TakeNumber, uint32 ShotNumberDigits, uint32 TakeNumberDigits);
 
 	/**
-	 * Generate a new shot package
+	 * Generate a new subsequence package
 	 *
-	 * @param SequenceMovieScene The sequence movie scene for the new shot
+	 * @param SequenceMovieScene The sequence movie scene for the new subsequence
+	 * @param SubsequenceDirectory The directory for the new subsequence
 	 * @param NewShotName The new shot name
-	 * @return The new shot path
+	 * @return The new subsequence path
 	 */
+	static FString GenerateNewSubsequencePath(UMovieScene* SequenceMovieScene, const FString& SubsequenceDirectory, FString& NewShotName);
+	
+	UE_DEPRECATED(5.3, "GenerateNewShotPath has been deprecated in favor of GenerateNewSubsequencePath that takes a given directory")
 	static FString GenerateNewShotPath(UMovieScene* SequenceMovieScene, FString& NewShotName);
 
+
 	/**
 	 * Generate a new shot name
 	 *
@@ -169,6 +174,9 @@ public:
 	 * @param Time The time to generate the new shot name at
 	 * @return The new shot name
 	 */
+	static FString GenerateNewSubsequenceName(const TArray<UMovieSceneSection*>& AllSections, const FString& SubsequencePrefix, FFrameNumber Time);
+
+	UE_DEPRECATED(5.3, "GenerateNewShotName has been deprecated in favor of GenerateNewSubsequenceName that takes a given prefix")
 	static FString GenerateNewShotName(const TArray<UMovieSceneSection*>& AllSections, FFrameNumber Time);
 
 	/*
diff --git a/Engine/Source/Editor/MovieSceneTools/Public/MovieSceneToolsProjectSettings.h b/Engine/Source/Editor/MovieSceneTools/Public/MovieSceneToolsProjectSettings.h
--- a/Engine/Source/Editor/MovieSceneTools/Public/MovieSceneToolsProjectSettings.h
+++ b/Engine/Source/Editor/MovieSceneTools/Public/MovieSceneToolsProjectSettings.h
@@ -68,10 +68,18 @@ public:
 	UPROPERTY(config, EditAnywhere, Category=Timeline, meta=(ClampMin=0.00001f, Units=s))
 	float DefaultDuration;
 
+	/** The default directory for the subsequences. */
+	UPROPERTY(config, EditAnywhere, Category=Shots)
+	FString SubsequenceDirectory;
+
 	/** The default directory for the shots. */
 	UPROPERTY(config, EditAnywhere, Category=Shots)
 	FString ShotDirectory;
 
+	/** The default prefix for subsequence names. */
+	UPROPERTY(config, EditAnywhere, Category=Shots)
+	FString SubsequencePrefix;
+
 	/** The default prefix for shot names. */
 	UPROPERTY(config, EditAnywhere, Category=Shots)
 	FString ShotPrefix;
@@ -100,10 +108,6 @@ public:
 	UPROPERTY(config, EditAnywhere, Category=Shots)
 	FString TakeSeparator;
 
-	/** Character separator between the take number and the sub sequence name. */
-	UPROPERTY(config, EditAnywhere, Category=Shots)
-	FString SubSequenceSeparator;
-
 	/** Mapping between fbx property name and property track path */
 	UPROPERTY(config, EditAnywhere, Category=TrackSettings)
 	TArray<FMovieSceneToolsFbxSettings> FbxSettings;
diff --git a/Engine/Source/Editor/MovieSceneTools/Public/TrackEditors/CinematicShotTrackEditor.h b/Engine/Source/Editor/MovieSceneTools/Public/TrackEditors/CinematicShotTrackEditor.h
--- a/Engine/Source/Editor/MovieSceneTools/Public/TrackEditors/CinematicShotTrackEditor.h
+++ b/Engine/Source/Editor/MovieSceneTools/Public/TrackEditors/CinematicShotTrackEditor.h
@@ -82,6 +82,8 @@ public:
 	virtual FText GetSubTrackToolTip() const override;
 	virtual FName GetSubTrackBrushName() const override;
 	virtual FString GetSubSectionDisplayName(const UMovieSceneSubSection* Section) const override;
+	virtual FString GetDefaultSubsequenceName() const override;
+	virtual FString GetDefaultSubsequenceDirectory() const override;
 	virtual TSubclassOf<UMovieSceneSubTrack> GetSubTrackClass() const;
 
 protected:
diff --git a/Engine/Source/Editor/MovieSceneTools/Public/TrackEditors/SubTrackEditor.h b/Engine/Source/Editor/MovieSceneTools/Public/TrackEditors/SubTrackEditor.h
--- a/Engine/Source/Editor/MovieSceneTools/Public/TrackEditors/SubTrackEditor.h
+++ b/Engine/Source/Editor/MovieSceneTools/Public/TrackEditors/SubTrackEditor.h
@@ -105,6 +105,12 @@ public:
 	/** Get the display name for the sub section */
 	virtual FString GetSubSectionDisplayName(const UMovieSceneSubSection* Section) const;
 
+	/** Get the default sub sequence name */
+	virtual FString GetDefaultSubsequenceName() const;
+
+	/** Get the sub sequence directory */
+	virtual FString GetDefaultSubsequenceDirectory() const;
+
 	/** Get the UMovieSceneSubTrack class */
 	virtual TSubclassOf<UMovieSceneSubTrack> GetSubTrackClass() const;
 
