Constraints: Baking, we change the tangent type to the default for the keys at compensation to avoid inserting keys of the wrong type.

diff --git a/Engine/Plugins/Animation/ControlRig/Source/ControlRigEditor/Private/Tools/ConstraintBaker.cpp b/Engine/Plugins/Animation/ControlRig/Source/ControlRigEditor/Private/Tools/ConstraintBaker.cpp
--- a/Engine/Plugins/Animation/ControlRig/Source/ControlRigEditor/Private/Tools/ConstraintBaker.cpp
+++ b/Engine/Plugins/Animation/ControlRig/Source/ControlRigEditor/Private/Tools/ConstraintBaker.cpp
@@ -271,7 +271,7 @@ void FConstraintBaker::Bake(UWorld* InWorld,
 	
 	FCompensationEvaluator Evaluator(InConstraint);
 	Evaluator.ComputeLocalTransformsForBaking(InWorld, InSequencer, FramesToBake);
-	const TArray<FTransform>& Transforms = Evaluator.ChildLocals;
+	TArray<FTransform> Transforms = Evaluator.ChildLocals;
 	if (FramesToBake.Num() != Transforms.Num())
 	{
 		return;
@@ -287,6 +287,8 @@ void FConstraintBaker::Bake(UWorld* InWorld,
 	const TArrayView<FMovieSceneFloatChannel*> FloatTransformChannels = Handle->GetFloatChannels(ConstraintSections.ChildTransformSection);
 	const TArrayView<FMovieSceneDoubleChannel*> DoubleTransformChannels = Handle->GetDoubleChannels(ConstraintSections.ChildTransformSection);
 
+	EMovieSceneKeyInterpolation KeyType = InSequencer->GetKeyInterpolation();
+
 	for (int32 Index = 0; Index < ConstraintFrames.Num(); ++Index)
 	{
 		const FFrameNumber Frame = ConstraintFrames[Index];
@@ -301,11 +303,34 @@ void FConstraintBaker::Bake(UWorld* InWorld,
 			if (FloatTransformChannels.Num() > 0)
 			{
 				FMovieSceneConstraintChannelHelper::DeleteTransformKeys(FloatTransformChannels, FrameMinusOne);
+				//we also set the tangent at the break to the default type
+				FMovieSceneConstraintChannelHelper::ChangeKeyInterpolation(FloatTransformChannels, Frame, KeyType);
 			}
 			else if (DoubleTransformChannels.Num() > 0)
 			{
 				FMovieSceneConstraintChannelHelper::DeleteTransformKeys(DoubleTransformChannels, FrameMinusOne);
+				//we also set the tangent at the break to the default type
+				FMovieSceneConstraintChannelHelper::ChangeKeyInterpolation(DoubleTransformChannels, Frame, KeyType);
+			}
+			if (int32 RemoveIndex = Algo::BinarySearch(FramesToBake, FrameMinusOne) != INDEX_NONE)
+			{
+				FramesToBake.RemoveAt(RemoveIndex);
+				Transforms.RemoveAt(RemoveIndex);
 			}
+			
+		}
+	}
+	//remove transform keys if baking frame
+	if (InSettings.IsSet() && InSettings.GetValue().BakingKeySettings == EBakingKeySettings::AllFrames)
+	{
+		if (FloatTransformChannels.Num() > 0)
+		{
+			FMovieSceneConstraintChannelHelper::DeleteTransformTimes(FloatTransformChannels, FramesToBake[0] +1, FramesToBake[FramesToBake.Num() -1]);
+		}
+		else if (DoubleTransformChannels.Num() > 0)
+		{
+			FMovieSceneConstraintChannelHelper::DeleteTransformTimes(DoubleTransformChannels, FramesToBake[0] + 1, FramesToBake[FramesToBake.Num() - 1]);
+
 		}
 	}
 
diff --git a/Engine/Source/Editor/MovieSceneTools/Public/Constraints/MovieSceneConstraintChannelHelper.h b/Engine/Source/Editor/MovieSceneTools/Public/Constraints/MovieSceneConstraintChannelHelper.h
--- a/Engine/Source/Editor/MovieSceneTools/Public/Constraints/MovieSceneConstraintChannelHelper.h
+++ b/Engine/Source/Editor/MovieSceneTools/Public/Constraints/MovieSceneConstraintChannelHelper.h
@@ -138,12 +138,19 @@ public:
 		const FFrameNumber& InCurrentTime,
 		const FFrameNumber& InNextTime);
 
-	/** @todo documentation. */
+	/** delete transform keys at that time */
 	template< typename ChannelType >
 	static void DeleteTransformKeys(
 		const TArrayView<ChannelType*>& InChannels,
 		const FFrameNumber& InTime);
 
+	/** Change key interpolation at the specified time*/
+	template< typename ChannelType >
+	static void ChangeKeyInterpolation(
+		const TArrayView<ChannelType*>& InChannels,
+		const FFrameNumber& InTime,
+		EMovieSceneKeyInterpolation KeyInterpolation);
+
 	static void HandleConstraintPropertyChanged(
 		UTickableTransformConstraint* InConstraint,
 		const FMovieSceneConstraintChannel& InActiveChannel,
diff --git a/Engine/Source/Editor/MovieSceneTools/Public/Constraints/MovieSceneConstraintChannelHelper.inl b/Engine/Source/Editor/MovieSceneTools/Public/Constraints/MovieSceneConstraintChannelHelper.inl
--- a/Engine/Source/Editor/MovieSceneTools/Public/Constraints/MovieSceneConstraintChannelHelper.inl
+++ b/Engine/Source/Editor/MovieSceneTools/Public/Constraints/MovieSceneConstraintChannelHelper.inl
@@ -262,6 +262,68 @@ void FMovieSceneConstraintChannelHelper::DeleteTransformKeys(
 	}
 }
 
+
+template< typename ChannelType >
+void FMovieSceneConstraintChannelHelper::ChangeKeyInterpolation(
+	const TArrayView<ChannelType*>& InChannels,
+	const FFrameNumber& InTime,
+	EMovieSceneKeyInterpolation KeyInterpolation)
+{
+	using ChannelValueType = typename ChannelType::ChannelValueType;
+	TEnumAsByte<ERichCurveInterpMode> InterpMode = RCIM_Cubic;
+	TEnumAsByte<ERichCurveTangentMode> TangentMode = RCTM_Auto;
+	
+	switch (KeyInterpolation)
+	{
+		case EMovieSceneKeyInterpolation::SmartAuto:
+		{
+			InterpMode = RCIM_Cubic;
+			TangentMode = RCTM_SmartAuto;
+			break;
+		}
+		case EMovieSceneKeyInterpolation::Auto:
+		{
+			InterpMode = RCIM_Cubic;
+			TangentMode = RCTM_Auto;
+			break;
+		}
+		case EMovieSceneKeyInterpolation::User:
+		{
+			InterpMode = RCIM_Cubic;
+			TangentMode = RCTM_User;
+			break;
+		}
+		case EMovieSceneKeyInterpolation::Break:
+		{
+			InterpMode = RCIM_Cubic;
+			TangentMode = RCTM_Break;
+			break;
+		}
+		case EMovieSceneKeyInterpolation::Linear:
+		{
+			InterpMode = RCIM_Linear;
+			break;
+		}
+		case EMovieSceneKeyInterpolation::Constant:
+		{
+			InterpMode = RCIM_Constant;
+			break;
+		}
+	};
+	for (ChannelType* Channel : InChannels)
+	{
+		TMovieSceneChannelData<ChannelValueType> ChannelInterface = Channel->GetData();
+		const TArrayView<const FFrameNumber> Times = ChannelInterface.GetTimes();
+		const int32 KeyIndex = Algo::LowerBound(Times, InTime);
+		TArrayView<ChannelValueType> Values = ChannelInterface.GetValues();
+		if (Times.IsValidIndex(KeyIndex) && Times[KeyIndex] == InTime && Values.IsValidIndex(KeyIndex))
+		{
+			Values[KeyIndex].InterpMode = InterpMode;
+			Values[KeyIndex].TangentMode = TangentMode;
+		}
+	}
+}
+
 template< typename ChannelType >
 TArray<FFrameNumber> FMovieSceneConstraintChannelHelper::GetTransformTimes(
 	const TArrayView<ChannelType*>& InChannels,
@@ -311,33 +373,6 @@ template< typename ChannelType >
 	 }
 }
 
-/** this will only set the values on channels with keys at the specified time, reusing tangents  etc. */
-template< typename ChannelType >
-void FMovieSceneConstraintChannelHelper::SetTransformTimes(
-	const TArrayView<ChannelType*>& InChannels,
-	const TArray<FFrameNumber>& Frames,
-	const TArray<FTransform>& Transforms)
-{
-	if (Frames.Num() != Transforms.Num())
-	{
-		return;
-	}
-	for (ChannelType* Channel : InChannels)
-	{
-		TMovieSceneChannelData<typename ChannelType::ChannelValueType> Data = Channel->GetData();
-		const TArrayView<const FFrameNumber> Times = Data.GetTimes();
-		for (int32 Index = 0; Index < Frames.Num(); ++Index)
-		{
-			FFrameNumber Time = Frames[0];
-			const int32 KeyIndex = Algo::LowerBound(Times, Time);
-			if (Times.IsValidIndex(KeyIndex) && Times[KeyIndex] == Time)
-			{
-				Data.GetValues()[KeyIndex] = 1.0;//todo 
-			}
-		}
-	}
-}
-
 template<typename ChannelType>
 void EvaluateTangentAtThisTime(int32 ChannelIndex, int32 NumChannels,
 	UMovieSceneSection* Section,
