CIS NonUnity Fix
#trivial

diff --git a/Engine/Source/Editor/MovieSceneTools/Public/Constraints/MovieSceneConstraintChannelHelper.inl b/Engine/Source/Editor/MovieSceneTools/Public/Constraints/MovieSceneConstraintChannelHelper.inl
--- a/Engine/Source/Editor/MovieSceneTools/Public/Constraints/MovieSceneConstraintChannelHelper.inl
+++ b/Engine/Source/Editor/MovieSceneTools/Public/Constraints/MovieSceneConstraintChannelHelper.inl
@@ -10,6 +10,7 @@
 #include "MovieSceneSection.h"
 #include "Algo/Unique.h"
 #include "Containers/SortedMap.h"
+#include "KeyParams.h"
 
 template<typename ChannelType>
 void FMovieSceneConstraintChannelHelper::GetFramesToCompensate(
