CIS FIx NonUnity
#trivial

diff --git a/Engine/Source/Editor/MovieSceneTools/Public/Constraints/MovieSceneConstraintChannelHelper.h b/Engine/Source/Editor/MovieSceneTools/Public/Constraints/MovieSceneConstraintChannelHelper.h
--- a/Engine/Source/Editor/MovieSceneTools/Public/Constraints/MovieSceneConstraintChannelHelper.h
+++ b/Engine/Source/Editor/MovieSceneTools/Public/Constraints/MovieSceneConstraintChannelHelper.h
@@ -17,7 +17,7 @@ class UTransformableHandle;
 struct FMovieSceneConstraintChannel;
 struct FFrameNumber;
 struct FConstraintAndActiveChannel;
-
+enum class EMovieSceneKeyInterpolation : uint8;
 
 struct MOVIESCENETOOLS_API FCompensationEvaluator
 {
