Sequencer: Key editor now edits selected keys if they are selected. Otherwise, it edits at the current time as usual.

diff --git a/Engine/Source/Editor/MovieSceneTools/Private/CurveKeyEditors/SequencerKeyEditor.h b/Engine/Source/Editor/MovieSceneTools/Private/CurveKeyEditors/SequencerKeyEditor.h
--- a/Engine/Source/Editor/MovieSceneTools/Private/CurveKeyEditors/SequencerKeyEditor.h
+++ b/Engine/Source/Editor/MovieSceneTools/Private/CurveKeyEditors/SequencerKeyEditor.h
@@ -4,6 +4,7 @@
 
 #include "Templates/SharedPointer.h"
 #include "MovieSceneSection.h"
+#include "IKeyArea.h"
 #include "ISequencer.h"
 #include "MovieSceneCommonHelpers.h"
 #include "Channels/MovieSceneChannelTraits.h"
@@ -11,7 +12,8 @@
 #include "Channels/MovieSceneChannelHandle.h"
 #include "MovieSceneTimeHelpers.h"
 #include "Channels/MovieSceneFloatChannel.h"
-
+#include "MVVM/Selection/Selection.h"
+#include "MVVM/ViewModels/SequencerEditorViewModel.h"
 
 template<typename ChannelType, typename ValueType>
 struct TSequencerKeyEditor
@@ -95,6 +97,7 @@ struct TSequencerKeyEditor
 	{
 		using namespace UE::MovieScene;
 		using namespace Sequencer;
+		using namespace UE::Sequencer;
 
 		UMovieSceneSection* OwningSection = WeakSection.Get();
 		if (!OwningSection)
@@ -112,38 +115,59 @@ struct TSequencerKeyEditor
 			return;
 		}
 
-		const FFrameNumber CurrentTime = Sequencer->GetLocalTime().Time.FloorToFrame();
 		const bool  bAutoSetTrackDefaults = Sequencer->GetAutoSetTrackDefaults();
 
-		EMovieSceneKeyInterpolation Interpolation = GetInterpolationMode(Channel,CurrentTime,Sequencer->GetKeyInterpolation());
-
-		TArray<FKeyHandle> KeysAtCurrentTime;
-		Channel->GetKeys(TRange<FFrameNumber>(CurrentTime), nullptr, &KeysAtCurrentTime);
-
-		if (KeysAtCurrentTime.Num() > 0)
+		const FKeySelection& KeySelection = Sequencer->GetViewModel()->GetSelection()->KeySelection;
+		if (KeySelection.Num())
 		{
-			AssignValue(Channel, KeysAtCurrentTime[0], InValue);
+			for (FKeyHandle Key : KeySelection)
+			{
+				// Make sure we only manipulate the values of the channel with the same channel type we're editing
+				TSharedPtr<FChannelModel> ChannelModel = KeySelection.GetModelForKey(Key);
+				if (ChannelModel && ChannelModel->GetKeyArea() && ChannelModel->GetKeyArea()->GetChannel().GetChannelTypeName() == ChannelHandle.GetChannelTypeName())
+				{
+					UMovieSceneSection* Section = ChannelModel->GetSection();
+					if (Section && Section->TryModify())
+					{
+						AssignValue(reinterpret_cast<ChannelType*>(ChannelModel->GetChannel()), Key, InValue);
+					}
+				}
+			}
 		}
 		else
 		{
-			bool bHasAnyKeys = Channel->GetNumKeys() != 0;
+			const FFrameNumber CurrentTime = Sequencer->GetLocalTime().Time.FloorToFrame();
 
-			if (bHasAnyKeys || bAutoSetTrackDefaults == false)
+			EMovieSceneKeyInterpolation Interpolation = GetInterpolationMode(Channel, CurrentTime, Sequencer->GetKeyInterpolation());
+
+			TArray<FKeyHandle> KeysAtCurrentTime;
+			Channel->GetKeys(TRange<FFrameNumber>(CurrentTime), nullptr, &KeysAtCurrentTime);
+
+			if (KeysAtCurrentTime.Num() > 0)
 			{
-				// When auto setting track defaults are disabled, add a key even when it's empty so that the changed
-				// value is saved and is propagated to the property.
-				AddKeyToChannel(Channel, CurrentTime, InValue, Interpolation);
-				bHasAnyKeys = Channel->GetNumKeys() != 0;
+				AssignValue(Channel, KeysAtCurrentTime[0], InValue);
 			}
-
-			if (bHasAnyKeys)
+			else
 			{
-				TRange<FFrameNumber> KeyRange     = TRange<FFrameNumber>(CurrentTime);
-				TRange<FFrameNumber> SectionRange = OwningSection->GetRange();
+				bool bHasAnyKeys = Channel->GetNumKeys() != 0;
 
-				if (!SectionRange.Contains(KeyRange))
+				if (bHasAnyKeys || bAutoSetTrackDefaults == false)
 				{
-					OwningSection->SetRange(TRange<FFrameNumber>::Hull(KeyRange, SectionRange));
+					// When auto setting track defaults are disabled, add a key even when it's empty so that the changed
+					// value is saved and is propagated to the property.
+					AddKeyToChannel(Channel, CurrentTime, InValue, Interpolation);
+					bHasAnyKeys = Channel->GetNumKeys() != 0;
+				}
+
+				if (bHasAnyKeys)
+				{
+					TRange<FFrameNumber> KeyRange     = TRange<FFrameNumber>(CurrentTime);
+					TRange<FFrameNumber> SectionRange = OwningSection->GetRange();
+
+					if (!SectionRange.Contains(KeyRange))
+					{
+						OwningSection->SetRange(TRange<FFrameNumber>::Hull(KeyRange, SectionRange));
+					}
 				}
 			}
 		}
