Sequencer: Nullptr fix for audio track editor crash

#rnx

diff --git a/Engine/Source/Editor/MovieSceneTools/Private/TrackEditors/AudioTrackEditor.cpp b/Engine/Source/Editor/MovieSceneTools/Private/TrackEditors/AudioTrackEditor.cpp
--- a/Engine/Source/Editor/MovieSceneTools/Private/TrackEditors/AudioTrackEditor.cpp
+++ b/Engine/Source/Editor/MovieSceneTools/Private/TrackEditors/AudioTrackEditor.cpp
@@ -698,7 +698,11 @@ FText FAudioSection::GetSectionToolTip() const
 
 float FAudioSection::GetSectionHeight() const
 {
-	return Section.GetTypedOuter<UMovieSceneAudioTrack>()->GetRowHeight();
+	if (UMovieSceneAudioTrack* Track = Section.GetTypedOuter<UMovieSceneAudioTrack>())
+	{
+		return Track->GetRowHeight();
+	}
+	return ISequencerSection::GetSectionHeight();
 }
 
 int32 FAudioSection::OnPaintSection( FSequencerSectionPainter& Painter ) const
