Curve Editor: Fixed BOOL and INT curves not drawing past left or right-most key value in Curve Editor.

diff --git a/Engine/Source/Editor/MovieSceneTools/Private/Channels/ChannelCurveModel.cpp b/Engine/Source/Editor/MovieSceneTools/Private/Channels/ChannelCurveModel.cpp
--- a/Engine/Source/Editor/MovieSceneTools/Private/Channels/ChannelCurveModel.cpp
+++ b/Engine/Source/Editor/MovieSceneTools/Private/Channels/ChannelCurveModel.cpp
@@ -98,6 +98,13 @@ void FChannelCurveModel<ChannelType, ChannelValue, KeyType>::DrawCurve(const FCu
 		const int32 StartingIndex = Algo::UpperBound(Times, StartFrame);
 		const int32 EndingIndex = Algo::LowerBound(Times, EndFrame);
 
+		// Add the lower bound of the visible space
+		const bool bValidRange = StartingIndex < EndingIndex;
+		if (bValidRange)
+		{
+			OutInterpolatingPoints.Add(MakeTuple(StartFrame / TickResolution, GetKeyValue(Values, StartingIndex)));
+		}
+
 		TOptional<double> PreviousValue;
 		for (int32 KeyIndex = StartingIndex; KeyIndex < EndingIndex; ++KeyIndex)
 		{
@@ -110,6 +117,12 @@ void FChannelCurveModel<ChannelType, ChannelValue, KeyType>::DrawCurve(const FCu
 			OutInterpolatingPoints.Add(MakeTuple(Times[KeyIndex] / TickResolution, Value));
 			PreviousValue = Value;
 		}
+
+		// Add the upper bound of the visible space
+		if (bValidRange)
+		{
+			OutInterpolatingPoints.Add(MakeTuple(EndFrame / TickResolution, GetKeyValue(Values, EndingIndex - 1)));
+		}
 	}
 }
 
