Sequencer: Tweak key editor's behavior when editing selected keys to only allow this if the channel of the key editor is one of the selected keys

#rnx

diff --git a/Engine/Source/Editor/MovieSceneTools/Private/CurveKeyEditors/SequencerKeyEditor.h b/Engine/Source/Editor/MovieSceneTools/Private/CurveKeyEditors/SequencerKeyEditor.h
--- a/Engine/Source/Editor/MovieSceneTools/Private/CurveKeyEditors/SequencerKeyEditor.h
+++ b/Engine/Source/Editor/MovieSceneTools/Private/CurveKeyEditors/SequencerKeyEditor.h
@@ -118,7 +118,21 @@ struct TSequencerKeyEditor
 		const bool  bAutoSetTrackDefaults = Sequencer->GetAutoSetTrackDefaults();
 
 		const FKeySelection& KeySelection = Sequencer->GetViewModel()->GetSelection()->KeySelection;
-		if (KeySelection.Num())
+
+		// Allow editing the key selection if the key editor's channel is one of the selected key's channels
+		bool bAllowEditingKeySelection = false;
+		for (FKeyHandle Key : KeySelection)
+		{
+			// Make sure we only manipulate the values of the channel with the same channel type we're editing
+			TSharedPtr<FChannelModel> ChannelModel = KeySelection.GetModelForKey(Key);
+			if (ChannelModel && ChannelModel->GetChannel() == Channel)
+			{
+				bAllowEditingKeySelection = true;	
+				break;
+			}
+		}
+
+		if (bAllowEditingKeySelection)
 		{
 			for (FKeyHandle Key : KeySelection)
 			{
