Sequencer: Key the material parameter collection instance's current value if it exists, rather than always just the default value.

diff --git a/Engine/Source/Editor/MovieSceneTools/Private/TrackEditors/MaterialParameterCollectionTrackEditor.cpp b/Engine/Source/Editor/MovieSceneTools/Private/TrackEditors/MaterialParameterCollectionTrackEditor.cpp
--- a/Engine/Source/Editor/MovieSceneTools/Private/TrackEditors/MaterialParameterCollectionTrackEditor.cpp
+++ b/Engine/Source/Editor/MovieSceneTools/Private/TrackEditors/MaterialParameterCollectionTrackEditor.cpp
@@ -23,6 +23,7 @@
 #include "Internationalization/Internationalization.h"
 #include "Internationalization/Text.h"
 #include "Materials/MaterialParameterCollection.h"
+#include "Materials/MaterialParameterCollectionInstance.h"
 #include "Misc/AssertionMacros.h"
 #include "Misc/Attribute.h"
 #include "Misc/FrameNumber.h"
@@ -321,11 +322,20 @@ void FMaterialParameterCollectionTrackEditor::AddScalarParameter(UMovieSceneMate
 		return;
 	}
 
+	float Value = Parameter.DefaultValue;
+	if (UWorld* World = Cast<UWorld>(GetSequencer()->GetPlaybackContext()))
+	{
+		if (UMaterialParameterCollectionInstance* Instance = World->GetParameterCollectionInstance(Track->MPC))
+		{
+			Instance->GetScalarParameterValue(Parameter.ParameterName, Value);
+		}
+	}
+
 	FFrameNumber KeyTime = GetTimeForKey();
 
 	const FScopedTransaction Transaction(LOCTEXT("AddScalarParameter", "Add scalar parameter"));
 	Track->Modify();
-	Track->AddScalarParameterKey(Parameter.ParameterName, KeyTime, RowIndex, Parameter.DefaultValue);
+	Track->AddScalarParameterKey(Parameter.ParameterName, KeyTime, RowIndex, Value);
 	GetSequencer()->NotifyMovieSceneDataChanged(EMovieSceneDataChangeType::MovieSceneStructureItemAdded);
 }
 
@@ -337,11 +347,20 @@ void FMaterialParameterCollectionTrackEditor::AddVectorParameter(UMovieSceneMate
 		return;
 	}
 
+	FLinearColor Value = Parameter.DefaultValue;
+	if (UWorld* World = Cast<UWorld>(GetSequencer()->GetPlaybackContext()))
+	{
+		if (UMaterialParameterCollectionInstance* Instance = World->GetParameterCollectionInstance(Track->MPC))
+		{
+			Instance->GetVectorParameterValue(Parameter.ParameterName, Value);
+		}
+	}
+
 	FFrameNumber KeyTime = GetTimeForKey();
 
 	const FScopedTransaction Transaction(LOCTEXT("AddVectorParameter", "Add vector parameter"));
 	Track->Modify();
-	Track->AddColorParameterKey(Parameter.ParameterName, KeyTime, RowIndex, Parameter.DefaultValue);
+	Track->AddColorParameterKey(Parameter.ParameterName, KeyTime, RowIndex, Value);
 	GetSequencer()->NotifyMovieSceneDataChanged(EMovieSceneDataChangeType::MovieSceneStructureItemAdded);
 }
 
