Control Rig Sequencer I/O Improvements: Add a Control Rig FBX Exporter + Improve & Fix the existing Control Rig Importer.

diff --git a/Engine/Plugins/Animation/ControlRig/Source/ControlRig/Private/Sequencer/MovieSceneControlRigParameterTrack.cpp b/Engine/Plugins/Animation/ControlRig/Source/ControlRig/Private/Sequencer/MovieSceneControlRigParameterTrack.cpp
--- a/Engine/Plugins/Animation/ControlRig/Source/ControlRig/Private/Sequencer/MovieSceneControlRigParameterTrack.cpp
+++ b/Engine/Plugins/Animation/ControlRig/Source/ControlRig/Private/Sequencer/MovieSceneControlRigParameterTrack.cpp
@@ -665,7 +665,61 @@ void UMovieSceneControlRigParameterTrack::GetSelectedNodes(TArray<FName>& Select
 	}
 }
 
-TArray<FFBXNodeAndChannels>* UMovieSceneControlRigParameterTrack::GetNodeAndChannelMappings(UMovieSceneSection* InSection )
+#if WITH_EDITOR
+bool UMovieSceneControlRigParameterTrack::GetFbxCurveDataFromChannelMetadata(const FMovieSceneChannelMetaData& MetaData, FControlRigFbxCurveData& OutCurveData)
+{
+	const FString ChannelName = MetaData.Name.ToString();
+	TArray<FString> ChannelParts;
+
+	// The channel has an attribute
+	if (ChannelName.ParseIntoArray(ChannelParts, TEXT(".")) > 1)
+	{
+		// Retrieve the attribute
+		OutCurveData.AttributeName = ChannelParts.Last();
+		
+		// The control name (left part) will be used as the node name
+		OutCurveData.NodeName = ChannelParts[0];
+		OutCurveData.ControlName = *OutCurveData.NodeName;
+
+		// The channel has 3 parts, the middle one (i.e. Location) will be treated as the property name
+		if (ChannelParts.Num() > 2)
+		{
+			OutCurveData.AttributePropertyName = ChannelParts[1];
+		}
+	}
+	// The channel does not have an attribute
+	else
+	{
+		// Thus no property above the attribute
+		OutCurveData.AttributePropertyName.Empty();
+		
+		// The channel group will be used as the node name (name of the control this channel is grouped under - i.e. for animation channels)
+		OutCurveData.NodeName = MetaData.Group.ToString();
+
+		// The channel name will be used as the control name and attribute name (i.e. Weight)
+		OutCurveData.ControlName = *ChannelName;
+		OutCurveData.AttributeName = OutCurveData.ControlName.ToString();
+	}
+
+	if (OutCurveData.NodeName.IsEmpty() || OutCurveData.AttributeName.IsEmpty())
+	{
+		return false;
+	}
+	
+	// Retrieve the control type
+	if (GetControlRig())
+	{
+		if (FRigControlElement* Control = GetControlRig()->FindControl(OutCurveData.ControlName))
+		{
+			OutCurveData.ControlType = (FFBXControlRigTypeProxyEnum)(uint8)Control->Settings.ControlType;
+			return true;
+		}
+	}
+	return false;
+}
+#endif
+
+TArray<FRigControlFBXNodeAndChannels>* UMovieSceneControlRigParameterTrack::GetNodeAndChannelMappings(UMovieSceneSection* InSection )
 {
 #if WITH_EDITOR
 	if (GetControlRig() == nullptr)
@@ -690,11 +744,10 @@ TArray<FFBXNodeAndChannels>* UMovieSceneControlRigParameterTrack::GetNodeAndChan
 	const FName EnumChannelTypeName = FMovieSceneByteChannel::StaticStruct()->GetFName();
 	const FName IntegerChannelTypeName = FMovieSceneIntegerChannel::StaticStruct()->GetFName();
 
-
+	// Our resulting mapping containing FBX node & UE channels data for each control
+	TArray<FRigControlFBXNodeAndChannels>* NodeAndChannels = new TArray<FRigControlFBXNodeAndChannels>();
+	
 	FMovieSceneChannelProxy& ChannelProxy = CurrentSectionToKey->GetChannelProxy();
-	TArray<FFBXNodeAndChannels>* NodeAndChannels = new TArray<FFBXNodeAndChannels>();
-	TArray<FString> StringArray;
-
 	for (const FMovieSceneChannelEntry& Entry : CurrentSectionToKey->GetChannelProxy().GetAllEntries())
 	{
 		const FName ChannelTypeName = Entry.GetChannelTypeName();
@@ -704,58 +757,67 @@ TArray<FFBXNodeAndChannels>* UMovieSceneControlRigParameterTrack::GetNodeAndChan
 			continue;
 		}
 
-		TArrayView<FMovieSceneChannel* const>        Channels = Entry.GetChannels();
+		TArrayView<FMovieSceneChannel* const> Channels = Entry.GetChannels();
 		TArrayView<const FMovieSceneChannelMetaData> AllMetaData = Entry.GetMetaData();
 
 		for (int32 Index = 0; Index < Channels.Num(); ++Index)
 		{
 			FMovieSceneChannelHandle Channel = ChannelProxy.MakeHandle(ChannelTypeName, Index);
-
 			const FMovieSceneChannelMetaData& MetaData = AllMetaData[Index];
-			StringArray.SetNum(0);
-			FString String = MetaData.Name.ToString();
-			String.ParseIntoArray(StringArray, TEXT("."));
-			if (StringArray.Num() > 0)
+
+			FControlRigFbxCurveData FbxCurveData;
+			if (!GetFbxCurveDataFromChannelMetadata(MetaData, FbxCurveData))
 			{
-				FString NodeName = StringArray[0];
-				FRigControlElement* ControlElement = GetControlRig() ? GetControlRig()->FindControl(FName(*StringArray[0])) : nullptr;
-				if (ControlElement)
-				{
-					NodeName = NodeName.ToUpper();
-					if (NodeAndChannels->Num() == 0 || (*NodeAndChannels)[NodeAndChannels->Num() - 1].NodeName != NodeName)
-					{
-						FFBXNodeAndChannels NodeAndChannel;
-						NodeAndChannel.MovieSceneTrack = this;
-						NodeAndChannel.ControlType = (FFBXControlRigTypeProxyEnum)(uint8)ControlElement->Settings.ControlType;
-						NodeAndChannel.NodeName = NodeName;
-						NodeAndChannels->Add(NodeAndChannel);
-					}
-					if (ChannelTypeName == DoubleChannelTypeName)
-					{
-						FMovieSceneDoubleChannel* DoubleChannel = Channel.Cast<FMovieSceneDoubleChannel>().Get();
-						(*NodeAndChannels)[NodeAndChannels->Num() - 1].DoubleChannels.Add(DoubleChannel);
-					}
-					else if (ChannelTypeName == FloatChannelTypeName)
-					{
-						FMovieSceneFloatChannel* FloatChannel = Channel.Cast<FMovieSceneFloatChannel>().Get();
-						(*NodeAndChannels)[NodeAndChannels->Num() - 1].FloatChannels.Add(FloatChannel);
-					}
-					else if (ChannelTypeName == BoolChannelTypeName)
-					{
-						FMovieSceneBoolChannel* BoolChannel = Channel.Cast<FMovieSceneBoolChannel>().Get();
-						(*NodeAndChannels)[NodeAndChannels->Num() - 1].BoolChannels.Add(BoolChannel);
-					}
-					else if (ChannelTypeName == EnumChannelTypeName)
-					{
-						FMovieSceneByteChannel* EnumChannel = Channel.Cast<FMovieSceneByteChannel>().Get();
-						(*NodeAndChannels)[NodeAndChannels->Num() - 1].EnumChannels.Add(EnumChannel);
-					}
-					else if (ChannelTypeName == IntegerChannelTypeName)
-					{
-						FMovieSceneIntegerChannel* IntegerChannel = Channel.Cast<FMovieSceneIntegerChannel>().Get();
-						(*NodeAndChannels)[NodeAndChannels->Num() - 1].IntegerChannels.Add(IntegerChannel);
-					}
-				}
+				continue;
+			}
+
+			// Retrieve the current control node, usually the last one but not given
+			FRigControlFBXNodeAndChannels* CurrentNodeAndChannel;
+
+			const int i = NodeAndChannels->FindLastByPredicate([FbxCurveData](const FRigControlFBXNodeAndChannels& A)
+				{ return A.NodeName == FbxCurveData.NodeName && A.ControlName == FbxCurveData.ControlName; }
+			);
+			if (i != INDEX_NONE)
+			{
+				CurrentNodeAndChannel = &(*NodeAndChannels)[i];
+			}
+			// Create the node if not created yet
+			else
+			{
+				NodeAndChannels->Add(FRigControlFBXNodeAndChannels());
+				
+				CurrentNodeAndChannel = &NodeAndChannels->Last();
+
+				CurrentNodeAndChannel->MovieSceneTrack = this;
+				CurrentNodeAndChannel->ControlType = FbxCurveData.ControlType;
+				CurrentNodeAndChannel->NodeName = FbxCurveData.NodeName;
+				CurrentNodeAndChannel->ControlName = FbxCurveData.ControlName;
+			}
+
+			if (ChannelTypeName == DoubleChannelTypeName)
+			{
+				FMovieSceneDoubleChannel* DoubleChannel = Channel.Cast<FMovieSceneDoubleChannel>().Get();
+				CurrentNodeAndChannel->DoubleChannels.Add(DoubleChannel);
+			}
+			else if (ChannelTypeName == FloatChannelTypeName)
+			{
+				FMovieSceneFloatChannel* FloatChannel = Channel.Cast<FMovieSceneFloatChannel>().Get();
+				CurrentNodeAndChannel->FloatChannels.Add(FloatChannel);
+			}
+			else if (ChannelTypeName == BoolChannelTypeName)
+			{
+				FMovieSceneBoolChannel* BoolChannel = Channel.Cast<FMovieSceneBoolChannel>().Get();
+				CurrentNodeAndChannel->BoolChannels.Add(BoolChannel);
+			}
+			else if (ChannelTypeName == EnumChannelTypeName)
+			{
+				FMovieSceneByteChannel* EnumChannel = Channel.Cast<FMovieSceneByteChannel>().Get();
+				CurrentNodeAndChannel->EnumChannels.Add(EnumChannel);
+			}
+			else if (ChannelTypeName == IntegerChannelTypeName)
+			{
+				FMovieSceneIntegerChannel* IntegerChannel = Channel.Cast<FMovieSceneIntegerChannel>().Get();
+				CurrentNodeAndChannel->IntegerChannels.Add(IntegerChannel);
 			}
 		}
 	}
diff --git a/Engine/Plugins/Animation/ControlRig/Source/ControlRig/Public/Sequencer/MovieSceneControlRigParameterTrack.h b/Engine/Plugins/Animation/ControlRig/Source/ControlRig/Public/Sequencer/MovieSceneControlRigParameterTrack.h
--- a/Engine/Plugins/Animation/ControlRig/Source/ControlRig/Public/Sequencer/MovieSceneControlRigParameterTrack.h
+++ b/Engine/Plugins/Animation/ControlRig/Source/ControlRig/Public/Sequencer/MovieSceneControlRigParameterTrack.h
@@ -51,8 +51,11 @@ public:
 #endif
 
 	//INodeAndMappingsInterface
-	virtual TArray<FFBXNodeAndChannels>* GetNodeAndChannelMappings(UMovieSceneSection* InSection)  override;
+	virtual TArray<FRigControlFBXNodeAndChannels>* GetNodeAndChannelMappings(UMovieSceneSection* InSection)  override;
 	virtual void GetSelectedNodes(TArray<FName>& OutSelectedNodes) override;
+#if WITH_EDITOR
+	virtual bool GetFbxCurveDataFromChannelMetadata(const FMovieSceneChannelMetaData& MetaData, FControlRigFbxCurveData& OutCurveData) override;
+#endif // WITH_EDITOR
 
 	//UControlRig Delegates
 	void HandleOnInitialized_GameThread();
diff --git a/Engine/Plugins/Animation/ControlRig/Source/ControlRigEditor/Private/ControlRigSequencerEditorLibrary.cpp b/Engine/Plugins/Animation/ControlRig/Source/ControlRigEditor/Private/ControlRigSequencerEditorLibrary.cpp
--- a/Engine/Plugins/Animation/ControlRig/Source/ControlRigEditor/Private/ControlRigSequencerEditorLibrary.cpp
+++ b/Engine/Plugins/Animation/ControlRig/Source/ControlRigEditor/Private/ControlRigSequencerEditorLibrary.cpp
@@ -2569,7 +2569,7 @@ bool UControlRigSequencerEditorLibrary::ImportFBXToControlRigTrack(UWorld* World
 	INodeAndChannelMappings* ChannelMapping = Cast<INodeAndChannelMappings>(InTrack);
 	if (ChannelMapping)
 	{
-		TArray<FFBXNodeAndChannels>* NodeAndChannels = ChannelMapping->GetNodeAndChannelMappings(InSection);
+		TArray<FRigControlFBXNodeAndChannels>* NodeAndChannels = ChannelMapping->GetNodeAndChannelMappings(InSection);
 		TArray<FName> SelectedControls;
 		for (const FString& StringName : ControlRigNames)
 		{
@@ -2587,6 +2587,36 @@ bool UControlRigSequencerEditorLibrary::ImportFBXToControlRigTrack(UWorld* World
 	return bValid;
 }
 
+bool UControlRigSequencerEditorLibrary::ExportFBXFromControlRigSection(ULevelSequence* Sequence, const UMovieSceneControlRigParameterSection* Section, const UMovieSceneUserExportFBXControlRigSettings* ExportFBXControlRigSettings)
+{
+	if (!Sequence || !ExportFBXControlRigSettings ||
+		!Section || !Section->GetControlRig() ||
+		!Sequence->GetMovieScene() || Sequence->GetMovieScene()->IsReadOnly())
+	{
+		return false;
+	}
+
+	TArray<FName> SelectedControls;
+	if (UMovieSceneControlRigParameterTrack* Track = Section->GetTypedOuter<UMovieSceneControlRigParameterTrack>())
+	{
+		Track->GetSelectedNodes(SelectedControls);
+	}
+
+	FMovieSceneSequenceTransform RootToLocalTransform;
+	if (IAssetEditorInstance* AssetEditor = GEditor->GetEditorSubsystem<UAssetEditorSubsystem>()->FindEditorForAsset(Sequence, false))
+	{
+		if (const ILevelSequenceEditorToolkit* LevelSequenceEditor = static_cast<ILevelSequenceEditorToolkit*>(AssetEditor))
+		{
+			if (const TSharedPtr<ISequencer> Sequencer = LevelSequenceEditor->GetSequencer())
+			{
+				RootToLocalTransform = Sequencer->GetFocusedMovieSceneSequenceTransform();
+			}
+		}
+	}
+	
+	return  MovieSceneToolHelpers::ExportFBXFromControlRigChannels(Section, ExportFBXControlRigSettings, SelectedControls, RootToLocalTransform);
+}
+
 bool UControlRigSequencerEditorLibrary::CollapseControlRigAnimLayersWithSettings(ULevelSequence* InSequence, UMovieSceneControlRigParameterTrack* InTrack, const FBakingAnimationKeySettings& InSettings)
 {
 	TWeakPtr<ISequencer> WeakSequencer = GetSequencerFromAsset();
diff --git a/Engine/Plugins/Animation/ControlRig/Source/ControlRigEditor/Private/Sequencer/ControlRigParameterTrackEditor.cpp b/Engine/Plugins/Animation/ControlRig/Source/ControlRigEditor/Private/Sequencer/ControlRigParameterTrackEditor.cpp
--- a/Engine/Plugins/Animation/ControlRig/Source/ControlRigEditor/Private/Sequencer/ControlRigParameterTrackEditor.cpp
+++ b/Engine/Plugins/Animation/ControlRig/Source/ControlRigEditor/Private/Sequencer/ControlRigParameterTrackEditor.cpp
@@ -3828,16 +3828,23 @@ void FControlRigParameterTrackEditor::BuildTrackContextMenu(FMenuBuilder& MenuBu
 		return;
 	}
 
-	TArray<FFBXNodeAndChannels>* NodeAndChannels = Track->GetNodeAndChannelMappings(SectionToKey);
+	TArray<FRigControlFBXNodeAndChannels>* NodeAndChannels = Track->GetNodeAndChannelMappings(SectionToKey);
 
-	MenuBuilder.BeginSection("Import To Control Rig", LOCTEXT("ImportToControlRig", "Import To Control Rig"));
+	MenuBuilder.BeginSection("Control Rig IO", LOCTEXT("ControlRigIO", "Control Rig I/O"));
 	{
 		MenuBuilder.AddMenuEntry(
 			LOCTEXT("ImportControlRigFBX", "Import Control Rig FBX"),
-			LOCTEXT("ImportControlRigFBXTooltip", "Import Control Rig FBX"),
+			LOCTEXT("ImportControlRigFBXTooltip", "Import Control Rig animation from FBX"),
 			FSlateIcon(),
 			FUIAction(
 				FExecuteAction::CreateRaw(this, &FControlRigParameterTrackEditor::ImportFBX, Track, SectionToKey, NodeAndChannels)));
+
+		MenuBuilder.AddMenuEntry(
+			LOCTEXT("ExportControlRigFBX", "Export Control Rig FBX"),
+			LOCTEXT("ExportControlRigFBXTooltip", "Export Control Rig animation to FBX"),
+			FSlateIcon(),
+			FUIAction(
+				FExecuteAction::CreateRaw(this, &FControlRigParameterTrackEditor::ExportFBX, Track, SectionToKey)));
 	}
 	MenuBuilder.EndSection();
 
@@ -3962,12 +3969,21 @@ void FControlRigParameterTrackEditor::ToggleFKControlRig(UMovieSceneControlRigPa
 
 
 void FControlRigParameterTrackEditor::ImportFBX(UMovieSceneControlRigParameterTrack* InTrack, UMovieSceneControlRigParameterSection* InSection,
-	TArray<FFBXNodeAndChannels>* NodeAndChannels)
+	TArray<FRigControlFBXNodeAndChannels>* NodeAndChannels)
 {
 	if (NodeAndChannels)
 	{
-		//NodeAndChannels will be deleted later
-		MovieSceneToolHelpers::ImportFBXIntoChannelsWithDialog(GetSequencer().ToSharedRef(), NodeAndChannels);
+		// NodeAndChannels will be deleted later
+		MovieSceneToolHelpers::ImportFBXIntoControlRigChannelsWithDialog(GetSequencer().ToSharedRef(), NodeAndChannels);
+	}
+}
+
+void FControlRigParameterTrackEditor::ExportFBX(UMovieSceneControlRigParameterTrack* InTrack, UMovieSceneControlRigParameterSection* InSection)
+{
+	if (InTrack && InTrack->GetControlRig())
+	{
+		// ControlComponentTransformsMapping will be deleted later
+		MovieSceneToolHelpers::ExportFBXFromControlRigChannelsWithDialog(GetSequencer().ToSharedRef(), InTrack);
 	}
 }
 
diff --git a/Engine/Plugins/Animation/ControlRig/Source/ControlRigEditor/Private/Sequencer/ControlRigParameterTrackEditor.h b/Engine/Plugins/Animation/ControlRig/Source/ControlRigEditor/Private/Sequencer/ControlRigParameterTrackEditor.h
--- a/Engine/Plugins/Animation/ControlRig/Source/ControlRigEditor/Private/Sequencer/ControlRigParameterTrackEditor.h
+++ b/Engine/Plugins/Animation/ControlRig/Source/ControlRigEditor/Private/Sequencer/ControlRigParameterTrackEditor.h
@@ -174,7 +174,11 @@ private:
 
 	/** Import FBX*/
 	void ImportFBX(UMovieSceneControlRigParameterTrack* InTrack, UMovieSceneControlRigParameterSection* InSection, 
-		TArray<FFBXNodeAndChannels>* NodeAndChannels);
+		TArray<FRigControlFBXNodeAndChannels>* NodeAndChannels);
+
+	/** Export FBX*/
+	void ExportFBX(UMovieSceneControlRigParameterTrack* InTrack, UMovieSceneControlRigParameterSection* InSection);
+	
 	/** Find Track for given ControlRig*/
 	UMovieSceneControlRigParameterTrack* FindTrack(const UControlRig* InControlRig) const;
 
diff --git a/Engine/Plugins/Animation/ControlRig/Source/ControlRigEditor/Public/ControlRigSequencerEditorLibrary.h b/Engine/Plugins/Animation/ControlRig/Source/ControlRigEditor/Public/ControlRigSequencerEditorLibrary.h
--- a/Engine/Plugins/Animation/ControlRig/Source/ControlRigEditor/Public/ControlRigSequencerEditorLibrary.h
+++ b/Engine/Plugins/Animation/ControlRig/Source/ControlRigEditor/Public/ControlRigSequencerEditorLibrary.h
@@ -14,6 +14,7 @@
 #include "Tools/ControlRigSnapper.h"
 #include "TransformNoScale.h"
 #include "EulerTransform.h"
+#include "MovieSceneToolsUserSettings.h"
 #include "Tools/ControlRigSnapSettings.h"
 #include "SequenceTimeUnit.h"
 #include "RigSpacePickerBakeSettings.h"
@@ -925,6 +926,11 @@ public:
 			UMovieSceneUserImportFBXControlRigSettings* ImportFBXControlRigSettings,
 			const FString& ImportFilename);
 
+	/** Exports an FBX from the given control rig section. */
+	UFUNCTION(BlueprintCallable, Category = "Editor Scripting | Sequencer Tools | FBX")
+	static bool ExportFBXFromControlRigSection(ULevelSequence* Sequence, const UMovieSceneControlRigParameterSection* Section,
+	                                    const UMovieSceneUserExportFBXControlRigSettings* ExportFBXControlRigSettings);
+
 	/*
 	 * Collapse and bake all sections and layers on a control rig track to just one section.
 	 *
diff --git a/Engine/Plugins/MovieScene/SequencerScripting/Source/SequencerScriptingEditor/Private/SequencerTools.cpp b/Engine/Plugins/MovieScene/SequencerScripting/Source/SequencerScriptingEditor/Private/SequencerTools.cpp
--- a/Engine/Plugins/MovieScene/SequencerScripting/Source/SequencerScriptingEditor/Private/SequencerTools.cpp
+++ b/Engine/Plugins/MovieScene/SequencerScripting/Source/SequencerScriptingEditor/Private/SequencerTools.cpp
@@ -736,7 +736,7 @@ bool USequencerToolsFunctionLibrary::ImportFBXToControlRig(UWorld* World, ULevel
 				INodeAndChannelMappings* ChannelMapping = Cast<INodeAndChannelMappings>(Track); 
 				if (ChannelMapping)
 				{
-					TArray<FFBXNodeAndChannels>* NodeAndChannels = ChannelMapping->GetNodeAndChannelMappings(nullptr);
+					TArray<FRigControlFBXNodeAndChannels>* NodeAndChannels = ChannelMapping->GetNodeAndChannelMappings(nullptr);
 					//use passed in controls for selected, actually selected controls should almost be empty anyway since we just loaded/set everything up.
 					for (const FString& StringName : ControlRigNames)
 					{
@@ -762,6 +762,51 @@ bool USequencerToolsFunctionLibrary::ImportFBXToControlRig(UWorld* World, ULevel
 
 }
 
+bool USequencerToolsFunctionLibrary::ExportFBXFromControlRig(ULevelSequence* Sequence,
+	const FString& ActorWithControlRigTrack,
+	const UMovieSceneUserExportFBXControlRigSettings* ExportFBXControlRigSettings)
+{
+	bool bValid = false;
+
+	if (!Sequence || !ExportFBXControlRigSettings)
+	{
+		return false;
+	}
+	
+	UMovieScene* MovieScene = Sequence->GetMovieScene();
+	if (!MovieScene || MovieScene->IsReadOnly())
+	{
+		return false;
+	}
+
+	const TArray<FMovieSceneBinding>& Bindings =  MovieScene->GetBindings();
+	for (const FMovieSceneBinding& Binding : Bindings)
+	{
+		if (Binding.GetName() != ActorWithControlRigTrack)
+		{
+			continue;
+		}
+
+		for (UMovieSceneTrack* Track : Binding.GetTracks())
+		{
+			INodeAndChannelMappings* ChannelMapping = Cast<INodeAndChannelMappings>(Track);
+			const UMovieSceneSection* Section = Track->GetSectionToKey();
+			if (!ChannelMapping || !Section)
+			{
+				continue;
+			}
+
+			TArray<FName> SelectedControls;
+			ChannelMapping->GetSelectedNodes(SelectedControls);
+
+			const FMovieSceneSequenceTransform RootToLocalTransform;
+			return MovieSceneToolHelpers::ExportFBXFromControlRigChannels(Section, ExportFBXControlRigSettings, SelectedControls, RootToLocalTransform);
+		}
+	}
+
+	return false;
+}
+
 
 FMovieSceneEvent USequencerToolsFunctionLibrary::CreateEvent(UMovieSceneSequence* InSequence, UMovieSceneEventSectionBase* InSection, const FSequencerQuickBindingResult& InEndpoint, const TArray<FString>& InPayload)
 {
diff --git a/Engine/Plugins/MovieScene/SequencerScripting/Source/SequencerScriptingEditor/Public/SequencerTools.h b/Engine/Plugins/MovieScene/SequencerScripting/Source/SequencerScriptingEditor/Public/SequencerTools.h
--- a/Engine/Plugins/MovieScene/SequencerScripting/Source/SequencerScriptingEditor/Public/SequencerTools.h
+++ b/Engine/Plugins/MovieScene/SequencerScripting/Source/SequencerScriptingEditor/Public/SequencerTools.h
@@ -282,6 +282,10 @@ public:
 	static bool ImportFBXToControlRig(UWorld* World, ULevelSequence* InSequence, const FString& ActorWithControlRigTrack, const TArray<FString>& SelectedControlRigNames,
 		UMovieSceneUserImportFBXControlRigSettings* ImportFBXControlRigSettings,
 		const FString& ImportFilename);
+
+	/** Exports an FBX from the section to key of the control rig track of the actor with the given name. */
+	UFUNCTION(BlueprintCallable, Category = "Editor Scripting | Sequencer Tools | FBX")
+	static bool ExportFBXFromControlRig(ULevelSequence* Sequence, const FString& ActorWithControlRigTrack, const UMovieSceneUserExportFBXControlRigSettings* ExportFBXControlRigSettings);
 };
 
 #if UE_ENABLE_INCLUDE_ORDER_DEPRECATED_IN_5_2
diff --git a/Engine/Source/Editor/MovieSceneTools/Private/MovieSceneToolHelpers.cpp b/Engine/Source/Editor/MovieSceneTools/Private/MovieSceneToolHelpers.cpp
--- a/Engine/Source/Editor/MovieSceneTools/Private/MovieSceneToolHelpers.cpp
+++ b/Engine/Source/Editor/MovieSceneTools/Private/MovieSceneToolHelpers.cpp
@@ -87,6 +87,7 @@
 #include "EntitySystem/Interrogation/MovieSceneInterrogationLinker.h"
 #include "ConstraintsManager.h"
 #include "Constraints/MovieSceneConstraintChannelHelper.inl"
+#include "Exporters/FbxExportOption.h"
 
 /* FSkelMeshRecorder
  ***********/
@@ -1495,455 +1496,193 @@ void ImportTransformChannelToInteger(const FRichCurve& Source, FMovieSceneIntege
 }
 
 void SetChannelValue(FMovieSceneDoubleChannel* DoubleChannel, FMovieSceneFloatChannel* FloatChannel, FMovieSceneBoolChannel *BoolChannel, FMovieSceneByteChannel *EnumChannel, FMovieSceneIntegerChannel* IntegerChannel,
-	FFrameRate FrameRate, FFrameNumber StartFrame,
+	const FFrameRate& FrameRate, const FFrameNumber& StartFrame,
 	FControlRigChannelEnum ChannelEnum, UMovieSceneUserImportFBXControlRigSettings* ImportFBXControlRigSettings,
-	FTransform& DefaultTransform,
-	FRichCurve& TranslationX, FRichCurve& TranslationY, FRichCurve& TranslationZ,
-	FRichCurve& EulerRotationX, FRichCurve& EulerRotationY, FRichCurve& EulerRotationZ,
-	FRichCurve& ScaleX, FRichCurve& ScaleY, FRichCurve& ScaleZ)
+	const FTransform& DefaultTransform,
+	const FRichCurve& TranslationX, const FRichCurve& TranslationY, const FRichCurve& TranslationZ,
+	const FRichCurve& EulerRotationX, const FRichCurve& EulerRotationY, const FRichCurve& EulerRotationZ,
+	const FRichCurve& ScaleX, const FRichCurve& ScaleY, const FRichCurve& ScaleZ)
 {
+	const FRichCurve* TransformCurves[9] = {&TranslationX, &TranslationY, &TranslationZ, &EulerRotationX, &EulerRotationY, &EulerRotationZ, &ScaleX, &ScaleY, &ScaleZ};
+	
 	FVector Location = DefaultTransform.GetLocation(), Rotation = DefaultTransform.GetRotation().Euler(), Scale3D = DefaultTransform.GetScale3D();
+	const double DefaultTransformValues[9] = {Location[0], Location[1], Location[2], Rotation[0], Rotation[1], Rotation[2], Scale3D[0], Scale3D[1], Scale3D[2]};
+
+	const FRichCurve *RichCurve = &TranslationX;
+	float DefaultValue = Location.X;
+	bool bNegate = false;
 
-	for (FControlToTransformMappings& Mapping: ImportFBXControlRigSettings->ControlChannelMappings)
+	uint8 ChannelIndex = 0; // Default Channel is TX
+	
+	if (const FControlToTransformMappings* ChannelMapping = ImportFBXControlRigSettings->ControlChannelMappings.FindByPredicate(
+		[ChannelEnum](const FControlToTransformMappings& Mapping)	{ return ChannelEnum == Mapping.ControlChannel; }))
 	{
-		if (ChannelEnum == Mapping.ControlChannel)
-		{
-			bool bNegate = Mapping.bNegate;
-			if (Mapping.FBXChannel == FTransformChannelEnum::TranslateX)
-			{
-				if (ChannelEnum == FControlRigChannelEnum::Bool && BoolChannel)
-				{
-					bool bDefault = Location.X == 0.0 ? false : true;
-					BoolChannel->SetDefault(bDefault);
-					ImportTransformChannelToBool(TranslationX, BoolChannel, FrameRate, false, StartFrame);
-				}
-				else if (ChannelEnum == FControlRigChannelEnum::Enum && EnumChannel)
-				{
-					uint8 bDefault = (uint8)bNegate ? -Location.X : Location.X;
-					EnumChannel->SetDefault(bDefault);
-					ImportTransformChannelToEnum(TranslationX, EnumChannel, FrameRate, false, StartFrame);
-				}
-				else if (ChannelEnum == FControlRigChannelEnum::Integer && IntegerChannel)
-				{
-					int32 bDefault = (int32)bNegate ? -Location.X : Location.X;
-					IntegerChannel->SetDefault(bDefault);
-					ImportTransformChannelToInteger(TranslationX, IntegerChannel, FrameRate, false, StartFrame);
-				}
-				else if (FloatChannel)
-				{
-					float Default = bNegate ? -Location.X : Location.X;
-					FloatChannel->SetDefault(Default);
-					ImportTransformChannelToFloat(TranslationX, FloatChannel, FrameRate, false, false, StartFrame, bNegate);
+		bNegate = ChannelMapping->bNegate;
+		
+		ChannelIndex = (uint8)ChannelMapping->FBXChannel;
+	}
+	else
+	{
+		const uint8 ChannelEnumIndex = (uint8)ChannelEnum;
 
-				}
-				else if (DoubleChannel)
-				{
-					float Default = bNegate ? -Location.X : Location.X;
-					DoubleChannel->SetDefault(Default);
-					ImportTransformChannelToDouble(TranslationX, DoubleChannel, FrameRate, false, false, StartFrame, bNegate);
-				}
-			}
-			else if (Mapping.FBXChannel == FTransformChannelEnum::TranslateY)
-			{
-				if (ChannelEnum == FControlRigChannelEnum::Bool && BoolChannel)
-				{
-					bool bDefault = Location.Y == 0.0 ? false : true;
-					BoolChannel->SetDefault(bDefault);
-					ImportTransformChannelToBool(TranslationY, BoolChannel, FrameRate, false, StartFrame);
-				}
-				else if (ChannelEnum == FControlRigChannelEnum::Enum && EnumChannel)
-				{
-					uint8 bDefault = (uint8)bNegate ? -Location.Y : Location.Y;
-					EnumChannel->SetDefault(bDefault);
-					ImportTransformChannelToEnum(TranslationY, EnumChannel, FrameRate, false, StartFrame);
-				}
-				else if (ChannelEnum == FControlRigChannelEnum::Integer && IntegerChannel)
-				{
-					int32 bDefault = (int32)bNegate ? -Location.Y : Location.Y;
-					IntegerChannel->SetDefault(bDefault);
-					ImportTransformChannelToInteger(TranslationY, IntegerChannel, FrameRate, false, StartFrame);
-				}
-				else if (FloatChannel)
-				{
-					bNegate = !bNegate;
-					float Default = bNegate ? -Location.Y : Location.Y;
-					FloatChannel->SetDefault(Default);
-					ImportTransformChannelToFloat(TranslationY, FloatChannel, FrameRate, false, false, StartFrame, bNegate);
-				}
-				else if (DoubleChannel)
-				{
-					bNegate = !bNegate;
-					float Default = bNegate ? -Location.Y : Location.Y;
-					DoubleChannel->SetDefault(Default);
-					ImportTransformChannelToDouble(TranslationY, DoubleChannel, FrameRate, false, false, StartFrame, bNegate);
-				}
-			}
-			else if (Mapping.FBXChannel == FTransformChannelEnum::TranslateZ)
-			{
-				if (ChannelEnum == FControlRigChannelEnum::Bool && BoolChannel)
-				{
-					bool bDefault = Location.Z == 0.0 ? false : true;
-					BoolChannel->SetDefault(bDefault);
-					ImportTransformChannelToBool(TranslationZ, BoolChannel, FrameRate, false, StartFrame);
-				}
-				else if (ChannelEnum == FControlRigChannelEnum::Enum && EnumChannel)
-				{
-					uint8 bDefault = (uint8)bNegate ? -Location.Z : Location.Z;
-					EnumChannel->SetDefault(bDefault);
-					ImportTransformChannelToEnum(TranslationZ, EnumChannel, FrameRate, false, StartFrame);
-				}
-				else if (ChannelEnum == FControlRigChannelEnum::Integer && IntegerChannel)
-				{
-					int32 bDefault = (int32)bNegate ? -Location.Z : Location.Z;
-					IntegerChannel->SetDefault(bDefault);
-					ImportTransformChannelToInteger(TranslationZ, IntegerChannel, FrameRate, false, StartFrame);
-				}
-				else if (FloatChannel)
-				{
-					float Default = bNegate ? -Location.Z : Location.Z;
-					FloatChannel->SetDefault(Default);
-					ImportTransformChannelToFloat(TranslationZ, FloatChannel, FrameRate, false, false, StartFrame, bNegate);
-				}
-				else if (DoubleChannel)
-				{
-					float Default = bNegate ? -Location.Z : Location.Z;
-					DoubleChannel->SetDefault(Default);
-					ImportTransformChannelToDouble(TranslationZ, DoubleChannel, FrameRate, false, false, StartFrame, bNegate);
-				}
-			}
-			else if (Mapping.FBXChannel == FTransformChannelEnum::RotateX)
-			{
-				if (ChannelEnum == FControlRigChannelEnum::Bool && BoolChannel)
-				{
-					bool bDefault = Rotation.X == 0.0 ? false : true;
-					BoolChannel->SetDefault(bDefault);
-					ImportTransformChannelToBool(EulerRotationX, BoolChannel, FrameRate, false, StartFrame);
-				}
-				else if (ChannelEnum == FControlRigChannelEnum::Enum && EnumChannel)
-				{
-					uint8 bDefault = (uint8)bNegate ? -Rotation.X : Rotation.X;
-					EnumChannel->SetDefault(bDefault);
-					ImportTransformChannelToEnum(EulerRotationX, EnumChannel, FrameRate, false, StartFrame);
-				}
-				else if (ChannelEnum == FControlRigChannelEnum::Integer && IntegerChannel)
-				{
-					int32 bDefault = (int32)bNegate ? -Rotation.X : Rotation.X;
-					IntegerChannel->SetDefault(bDefault);
-					ImportTransformChannelToInteger(EulerRotationX, IntegerChannel, FrameRate, false, StartFrame);
-				}
-				else if (FloatChannel)
-				{
-					float Default = bNegate ? -Rotation.X : Rotation.X;
-					FloatChannel->SetDefault(Default);
-					ImportTransformChannelToFloat(EulerRotationX, FloatChannel, FrameRate, false, false, StartFrame, bNegate);
-				}
-				else if (DoubleChannel)
-				{
-					float Default = bNegate ? -Rotation.X : Rotation.X;
-					DoubleChannel->SetDefault(Default);
-					ImportTransformChannelToDouble(EulerRotationX, DoubleChannel, FrameRate, false, false, StartFrame, bNegate);
-				}
-			}
-			else if (Mapping.FBXChannel == FTransformChannelEnum::RotateY)
-			{
-				if (ChannelEnum == FControlRigChannelEnum::Bool && BoolChannel)
-				{
-					bool bDefault = Rotation.Y == 0.0 ? false : true;
-					BoolChannel->SetDefault(bDefault);
-					ImportTransformChannelToBool(EulerRotationY, BoolChannel, FrameRate, false, StartFrame);
-				}
-				else if (ChannelEnum == FControlRigChannelEnum::Enum && EnumChannel)
-				{
-					uint8 bDefault = (uint8)bNegate ? -Rotation.Y : Rotation.Y;
-					EnumChannel->SetDefault(bDefault);
-					ImportTransformChannelToEnum(EulerRotationY, EnumChannel, FrameRate, false, StartFrame);
-				}
-				else if (ChannelEnum == FControlRigChannelEnum::Integer && IntegerChannel)
-				{
-					int32 bDefault = (int32)bNegate ? -Rotation.Y : Rotation.Y;
-					IntegerChannel->SetDefault(bDefault);
-					ImportTransformChannelToInteger(EulerRotationY, IntegerChannel, FrameRate, false, StartFrame);
-				}
-				else if (FloatChannel)
-				{
-					float Default = bNegate ? -Rotation.Y : Rotation.Y;
-					FloatChannel->SetDefault(Default);
-					ImportTransformChannelToFloat(EulerRotationY, FloatChannel, FrameRate, false, false, StartFrame, bNegate);
-				}
-				else if (DoubleChannel)
-				{
-					float Default = bNegate ? -Rotation.Y : Rotation.Y;
-					DoubleChannel->SetDefault(Default);
-					ImportTransformChannelToDouble(EulerRotationY, DoubleChannel, FrameRate, false, false, StartFrame, bNegate);
-				}
-			}
-			else if (Mapping.FBXChannel == FTransformChannelEnum::RotateZ)
-			{
-				if (ChannelEnum == FControlRigChannelEnum::Bool && BoolChannel)
-				{
-					bool bDefault = Rotation.Z == 0.0 ? false : true;
-					BoolChannel->SetDefault(bDefault);
-					ImportTransformChannelToBool(EulerRotationZ, BoolChannel, FrameRate, false, StartFrame);
-				}
-				else if (ChannelEnum == FControlRigChannelEnum::Enum && EnumChannel)
-				{
-					uint8 bDefault = (uint8)bNegate ? -Rotation.Z : Rotation.Z;
-					EnumChannel->SetDefault(bDefault);
-					ImportTransformChannelToEnum(EulerRotationZ, EnumChannel, FrameRate, false, StartFrame);
-				}
-				else if (ChannelEnum == FControlRigChannelEnum::Integer && IntegerChannel)
-				{
-					int32 bDefault = (int32)bNegate ? -Rotation.Z : Rotation.Z;
-					IntegerChannel->SetDefault(bDefault);
-					ImportTransformChannelToInteger(EulerRotationZ, IntegerChannel, FrameRate, false, StartFrame);
-				}
-				else if (FloatChannel)
-				{
-					float Default = bNegate ? -Rotation.Z : Rotation.Z;
-					FloatChannel->SetDefault(Default);
-					ImportTransformChannelToFloat(EulerRotationZ, FloatChannel, FrameRate, false, false, StartFrame, bNegate);
-				}
-				else if (DoubleChannel)
-				{
-					float Default = bNegate ? -Rotation.Z : Rotation.Z;
-					DoubleChannel->SetDefault(Default);
-					ImportTransformChannelToDouble(EulerRotationZ, DoubleChannel, FrameRate, false, false, StartFrame, bNegate);
-				}
-			}
-			else if (Mapping.FBXChannel == FTransformChannelEnum::ScaleX)
-			{
-				if (ChannelEnum == FControlRigChannelEnum::Bool && BoolChannel)
-				{
-					bool bDefault = Scale3D.X == 0.0 ? false : true;
-					BoolChannel->SetDefault(bDefault);
-					ImportTransformChannelToBool(ScaleX, BoolChannel, FrameRate, false, StartFrame);
-				}
-				else if (ChannelEnum == FControlRigChannelEnum::Enum && EnumChannel)
-				{
-					uint8 bDefault = (uint8)bNegate ? -Scale3D.X : Scale3D.X;
-					EnumChannel->SetDefault(bDefault);
-					ImportTransformChannelToEnum(ScaleX, EnumChannel, FrameRate, false, StartFrame);
-				}
-				else if (ChannelEnum == FControlRigChannelEnum::Integer && IntegerChannel)
-				{
-					int32 bDefault = (int32)bNegate ? -Scale3D.X : Scale3D.X;
-					IntegerChannel->SetDefault(bDefault);
-					ImportTransformChannelToInteger(ScaleX, IntegerChannel, FrameRate, false, StartFrame);
-				}
-				else if (FloatChannel)
-				{
-					float Default = bNegate ? -Scale3D.X : Scale3D.X;
-					FloatChannel->SetDefault(Default);
-					ImportTransformChannelToFloat(ScaleX, FloatChannel, FrameRate, false, false, StartFrame, bNegate);
-				}
-				else if (DoubleChannel)
-				{
-					float Default = bNegate ? -Scale3D.X : Scale3D.X;
-					DoubleChannel->SetDefault(Default);
-					ImportTransformChannelToDouble(ScaleX, DoubleChannel, FrameRate, false, false, StartFrame, bNegate);
-				}
-			}
-			else if (Mapping.FBXChannel == FTransformChannelEnum::ScaleY)
-			{
-				if (ChannelEnum == FControlRigChannelEnum::Bool && BoolChannel)
-				{
-					bool bDefault = Scale3D.Y == 0.0 ? false : true;
-					BoolChannel->SetDefault(bDefault);
-					ImportTransformChannelToBool(ScaleY, BoolChannel, FrameRate, false, StartFrame);
-				}
-				else if (ChannelEnum == FControlRigChannelEnum::Enum && EnumChannel)
-				{
-					uint8 bDefault = (uint8)bNegate ? -Scale3D.Y : Scale3D.Y;
-					EnumChannel->SetDefault(bDefault);
-					ImportTransformChannelToEnum(ScaleY, EnumChannel, FrameRate, false, StartFrame);
-				}
-				else if (ChannelEnum == FControlRigChannelEnum::Integer && IntegerChannel)
-				{
-					int32 bDefault = (int32)bNegate ? -Scale3D.Y : Scale3D.Y;
-					IntegerChannel->SetDefault(bDefault);
-					ImportTransformChannelToInteger(ScaleY, IntegerChannel, FrameRate, false, StartFrame);
-				}
-				else if (FloatChannel)
-				{
-					float Default = bNegate ? -Scale3D.Y : Scale3D.Y;
-					FloatChannel->SetDefault(Default);
-					ImportTransformChannelToFloat(ScaleY, FloatChannel, FrameRate, false, false, StartFrame, bNegate);
-				}
-				else if (DoubleChannel)
-				{
-					float Default = bNegate ? -Scale3D.Y : Scale3D.Y;
-					DoubleChannel->SetDefault(Default);
-					ImportTransformChannelToDouble(ScaleY, DoubleChannel, FrameRate, false, false, StartFrame, bNegate);
-				}
-			}
-			else if (Mapping.FBXChannel == FTransformChannelEnum::ScaleZ)
-			{
-				if (ChannelEnum == FControlRigChannelEnum::Bool && BoolChannel)
-				{
-					bool bDefault = Scale3D.Z == 0.0 ? false : true;
-					BoolChannel->SetDefault(bDefault);
-					ImportTransformChannelToBool(ScaleZ, BoolChannel, FrameRate, false, StartFrame);
-				}
-				else if (ChannelEnum == FControlRigChannelEnum::Enum && EnumChannel)
-				{
-					uint8 bDefault = (uint8)bNegate ? -Scale3D.Z : Scale3D.Z;
-					EnumChannel->SetDefault(bDefault);
-					ImportTransformChannelToEnum(ScaleZ, EnumChannel, FrameRate, false, StartFrame);
-				}
-				else if (ChannelEnum == FControlRigChannelEnum::Integer && IntegerChannel)
-				{
-					int32 bDefault = (int32)bNegate ? -Scale3D.Z : Scale3D.Z;
-					IntegerChannel->SetDefault(bDefault);
-					ImportTransformChannelToInteger(ScaleZ, IntegerChannel, FrameRate, false, StartFrame);
-				}
-				else if (FloatChannel)
-				{
-					float Default = bNegate ? -Scale3D.Z : Scale3D.Z;
-					FloatChannel->SetDefault(Default);
-					ImportTransformChannelToFloat(ScaleZ, FloatChannel, FrameRate, false, false, StartFrame, bNegate);
-				}
-				else if (DoubleChannel)
-				{
-					float Default = bNegate ? -Scale3D.Z : Scale3D.Z;
-					DoubleChannel->SetDefault(Default);
-					ImportTransformChannelToDouble(ScaleZ, DoubleChannel, FrameRate, false, false, StartFrame, bNegate);
-				}
-			}
+		// 2DX, Bool and other default to TX
+		if (ChannelEnum == FControlRigChannelEnum::Vector2DY)
+		{
+			ChannelIndex = 1;
+		}
+		else if (ChannelEnumIndex >= (uint8)FControlRigChannelEnum::PositionX && ChannelEnumIndex <= (uint8)FControlRigChannelEnum::ScaleZ)
+		{
+			ChannelIndex = ChannelEnumIndex - (uint8)FControlRigChannelEnum::PositionX;
 		}
 	}
 
+	if (ChannelIndex < 9)
+	{
+		DefaultValue = DefaultTransformValues[ChannelIndex];
+		RichCurve = TransformCurves[ChannelIndex];
+	}
+	
 	if (ChannelEnum == FControlRigChannelEnum::Bool && BoolChannel)
 	{
-		bool bDefault = Location.X == 0.0 ? false : true;
-		BoolChannel->SetDefault(bDefault);
-		ImportTransformChannelToBool(TranslationX, BoolChannel, FrameRate, false, StartFrame);
+		const bool Default = DefaultValue != 0.0 ? !bNegate : bNegate;
 
+		BoolChannel->SetDefault(Default);
+		ImportTransformChannelToBool(*RichCurve, BoolChannel, FrameRate, false, StartFrame);
 	}
 	else if (ChannelEnum == FControlRigChannelEnum::Enum && EnumChannel)
 	{
-		bool bNegate = false;
-		uint8 Default = (uint8) Location.X;
+		const uint8 Default = (uint8)(bNegate ? -DefaultValue : DefaultValue);
 		EnumChannel->SetDefault(Default);
-		ImportTransformChannelToEnum(TranslationX, EnumChannel, FrameRate, false, StartFrame);
+		ImportTransformChannelToEnum(*RichCurve, EnumChannel, FrameRate, false, StartFrame);
 	}
 	else if (ChannelEnum == FControlRigChannelEnum::Integer && IntegerChannel)
 	{
-		bool bNegate = false;
-		int32  Default = (int32 )Location.X;
+		const int32 Default = (int32)(bNegate ? -DefaultValue : DefaultValue);
 		IntegerChannel->SetDefault(Default);
-		ImportTransformChannelToInteger(TranslationX, IntegerChannel, FrameRate, false, StartFrame);
-	}
-	else if (ChannelEnum == FControlRigChannelEnum::Float)
-	{
-		bool bNegate = false;
-		float Default = Location.X;
-		FloatChannel->SetDefault(Default);
-		ImportTransformChannelToFloat(TranslationX, FloatChannel, FrameRate, false,false, StartFrame, bNegate);
-	}
-	else if (ChannelEnum == FControlRigChannelEnum::Vector2DX)
-	{
-		bool bNegate = false;
-		float Default = Location.X;
-		FloatChannel->SetDefault(Default);
-		ImportTransformChannelToFloat(TranslationX, FloatChannel, FrameRate, false, false,StartFrame, bNegate);
+		ImportTransformChannelToInteger(*RichCurve, IntegerChannel, FrameRate, false, StartFrame);
 	}
-	else if (ChannelEnum == FControlRigChannelEnum::Vector2DY)
+	else if (FloatChannel || DoubleChannel)
 	{
-		bool bNegate = true;
-		float Default = -Location.Y;
-		FloatChannel->SetDefault(Default);
-		ImportTransformChannelToFloat(TranslationY, FloatChannel, FrameRate, false, false,StartFrame, bNegate);
+		if (*RichCurve == TranslationY)
+		{
+			bNegate = !bNegate;
+		}
+		
+		const float Default = bNegate ? -DefaultValue : DefaultValue;
+		if (FloatChannel)
+		{
+			FloatChannel->SetDefault(Default);
+			ImportTransformChannelToFloat(*RichCurve, FloatChannel, FrameRate, bNegate, false, StartFrame, bNegate);
+		}
+		else if (DoubleChannel)
+		{
+			DoubleChannel->SetDefault(Default);
+			ImportTransformChannelToDouble(*RichCurve, DoubleChannel, FrameRate, bNegate, false, StartFrame, bNegate);
+		}
 	}
-	else if (ChannelEnum == FControlRigChannelEnum::PositionX)
+}
+
+static bool ImportFBXNonTransformCurvesToChannels(const FString& NodeName, const UMovieSceneUserImportFBXSettings* ImportFBXSettings,UMovieSceneUserImportFBXControlRigSettings* ImportFBXControlRigSettings,
+                                                  const FFrameNumber StartFrame, const FFrameRate FrameRate, FRigControlFBXNodeAndChannels& NodeAndChannels, UnFbx::FFbxCurvesAPI& CurveAPI)
+{
+	FbxNode* Node = CurveAPI.Scene->GetRootNode()->FindChild(TCHAR_TO_UTF8(*NodeName));
+	if (!Node)
 	{
-		bool bNegate = false;
-		float Default = Location.X;
-		FloatChannel->SetDefault(Default);
-		ImportTransformChannelToFloat(TranslationX, FloatChannel, FrameRate, false, false, StartFrame, bNegate);
+		return false;
 	}
-	else if (ChannelEnum == FControlRigChannelEnum::PositionY)
+	FRichCurve EmptyDefaultCurve;
+
+	FbxProperty Property = Node->FindProperty(TCHAR_TO_UTF8(*NodeAndChannels.ControlName.ToString()));
+	if (!Property.IsValid())
 	{
-		bool bNegate = true;
-		float Default = -Location.Y;
-		FloatChannel->SetDefault(Default);
-		ImportTransformChannelToFloat(TranslationY, FloatChannel, FrameRate, false, false, StartFrame, bNegate);
-	}
-	else if (ChannelEnum == FControlRigChannelEnum::PositionZ)
+		return false;
+	}	
+	EmptyDefaultCurve.SetDefaultValue(Property.Get<float>());
+
+	TMap<FName, FRichCurve> Curves;
+	CurveAPI.GetConvertedNonTransformCurveData(NodeName, true, ImportFBXSettings->ImportUniformScale, Curves);
+
+	FRichCurve* Curve = Curves.Find(NodeAndChannels.ControlName);
+	if (!Curve)
 	{
-		bool bNegate = false;
-		float Default = Location.Z;
-		FloatChannel->SetDefault(Default);
-		ImportTransformChannelToFloat(TranslationZ, FloatChannel, FrameRate, false, false, StartFrame, bNegate);
+		Curve = &EmptyDefaultCurve;
 	}
-	else if (ChannelEnum == FControlRigChannelEnum::RotatorX)
+	
+	const float CurveDefault = Curve->GetDefaultValue();
+
+	auto ShouldNegateChannel = [ImportFBXControlRigSettings](FControlRigChannelEnum ChannelType)
 	{
-		bool bNegate = false;
-		float Default = Rotation.X;
-		FloatChannel->SetDefault(Default);
-		ImportTransformChannelToFloat(EulerRotationX, FloatChannel, FrameRate, false, false, StartFrame, bNegate);
-	}
-	else if (ChannelEnum == FControlRigChannelEnum::RotatorY)
+		const FControlToTransformMappings* ControlMapping = ImportFBXControlRigSettings->ControlChannelMappings.FindByPredicate(
+		[ChannelType](const FControlToTransformMappings& Mapping)
+			{
+				return Mapping.ControlChannel == ChannelType;
+			}
+		);
+		return ControlMapping && ControlMapping->bNegate;
+	};
+	
+	if (!NodeAndChannels.BoolChannels.IsEmpty())
 	{
-		bool bNegate = false;
-		float Default = Rotation.Y;
-		FloatChannel->SetDefault(Default);
-		ImportTransformChannelToFloat(EulerRotationY, FloatChannel, FrameRate, false, false, StartFrame, bNegate);
+		const bool bNegate = ShouldNegateChannel(FControlRigChannelEnum::Bool);
+		const bool ChannelDefault = bNegate ? CurveDefault == 0.0 : CurveDefault != 0.0;
+		
+		NodeAndChannels.BoolChannels[0]->SetDefault(ChannelDefault);
+		ImportTransformChannelToBool(*Curve, NodeAndChannels.BoolChannels[0], FrameRate, false, StartFrame);
 	}
-	else if (ChannelEnum == FControlRigChannelEnum::RotatorZ)
+	else if (!NodeAndChannels.EnumChannels.IsEmpty())
 	{
-		bool bNegate = false;
-		float Default =  Rotation.Z;
-		FloatChannel->SetDefault(Default);
-		ImportTransformChannelToFloat(EulerRotationZ, FloatChannel, FrameRate, false, false, StartFrame, bNegate);
+		const bool bNegate = ShouldNegateChannel(FControlRigChannelEnum::Enum);
+		const uint8 ChannelDefault = bNegate ? -(uint8)CurveDefault : (uint8)CurveDefault;
+		
+		NodeAndChannels.EnumChannels[0]->SetDefault(ChannelDefault);
+		ImportTransformChannelToEnum(*Curve, NodeAndChannels.EnumChannels[0], FrameRate, false, StartFrame);
 	}
-	else if (ChannelEnum == FControlRigChannelEnum::ScaleX)
+	else if (!NodeAndChannels.IntegerChannels.IsEmpty())
 	{
-		bool bNegate = false;
-		float Default = Scale3D.X;
-		FloatChannel->SetDefault(Default);
-		ImportTransformChannelToFloat(ScaleX, FloatChannel, FrameRate, false, false, StartFrame, bNegate);
+		const bool bNegate = ShouldNegateChannel(FControlRigChannelEnum::Integer);
+		const int32 ChannelDefault = bNegate ? -(int32)CurveDefault : (int32)CurveDefault;
+		NodeAndChannels.IntegerChannels[0]->SetDefault(ChannelDefault);
+		
+		ImportTransformChannelToInteger(*Curve, NodeAndChannels.IntegerChannels[0], FrameRate, false, StartFrame);
 	}
-	else if (ChannelEnum == FControlRigChannelEnum::ScaleY)
+	else if (!NodeAndChannels.FloatChannels.IsEmpty())
 	{
-		bool bNegate = false;
-		float Default = Scale3D.Y;
-		FloatChannel->SetDefault(Default);
-		ImportTransformChannelToFloat(ScaleY, FloatChannel, FrameRate, false, false, StartFrame, bNegate);
+		const bool bNegate = ShouldNegateChannel(FControlRigChannelEnum::Integer);
+		const float ChannelDefault = bNegate ? -CurveDefault : CurveDefault;
+		NodeAndChannels.FloatChannels[0]->SetDefault(ChannelDefault);
+		
+		ImportTransformChannelToFloat(*Curve, NodeAndChannels.FloatChannels[0], FrameRate, false, false, StartFrame, bNegate);
 	}
-	else if (ChannelEnum == FControlRigChannelEnum::ScaleZ)
+	else
 	{
-		bool bNegate = false;
-		float Default = Scale3D.Z;
-		FloatChannel->SetDefault(Default);
-		ImportTransformChannelToFloat(ScaleZ, FloatChannel, FrameRate, false, false, StartFrame, bNegate);
+		return false;	
 	}
-
+	return true;
 }
+
 //if one channel goes to Y
 //if two channel go to X Y
 //if three channel to to x y z
 // if 9 due full
-static bool ImportFBXTransformToChannels(FString NodeName, const UMovieSceneUserImportFBXSettings* ImportFBXSettings,UMovieSceneUserImportFBXControlRigSettings* ImportFBXControlRigSettings,  FFrameNumber StartFrame, FFrameRate FrameRate, FFBXNodeAndChannels& NodeAndChannels,
-	 UnFbx::FFbxCurvesAPI& CurveAPI)
+static bool ImportFBXTransformToChannels(FString NodeName, const UMovieSceneUserImportFBXSettings* ImportFBXSettings,UMovieSceneUserImportFBXControlRigSettings* ImportFBXControlRigSettings,
+                                         const FFrameNumber StartFrame, const FFrameRate FrameRate, FRigControlFBXNodeAndChannels& NodeAndChannels, UnFbx::FFbxCurvesAPI& CurveAPI)
 {
-
 	TArray<FMovieSceneDoubleChannel*>& DoubleChannels = NodeAndChannels.DoubleChannels;
 	TArray<FMovieSceneFloatChannel*>& FloatChannels = NodeAndChannels.FloatChannels;
 	TArray<FMovieSceneBoolChannel*>& BoolChannels = NodeAndChannels.BoolChannels;
 	TArray<FMovieSceneByteChannel*>& EnumChannels = NodeAndChannels.EnumChannels;
 	TArray<FMovieSceneIntegerChannel*>& IntegerChannels = NodeAndChannels.IntegerChannels;
-
-
+	
 	// Look for transforms explicitly
-	FRichCurve Translation[3];
-	FRichCurve EulerRotation[3];
-	FRichCurve Scale[3];
+	FRichCurve Transform[9];
 	FTransform DefaultTransform;
 	const bool bUseSequencerCurve = true;
-	CurveAPI.GetConvertedTransformCurveData(NodeName, Translation[0], Translation[1], Translation[2], EulerRotation[0], EulerRotation[1], EulerRotation[2], Scale[0], Scale[1], Scale[2], DefaultTransform, true, ImportFBXSettings->ImportUniformScale);
-
-
+	CurveAPI.GetConvertedTransformCurveData(NodeName, Transform[0], Transform[1], Transform[2], Transform[3], Transform[4], Transform[5], Transform[6], Transform[7], Transform[8], DefaultTransform, bUseSequencerCurve, ImportFBXSettings->ImportUniformScale);
+	
 	FVector Location = DefaultTransform.GetLocation(), Rotation = DefaultTransform.GetRotation().Euler(), Scale3D = DefaultTransform.GetScale3D();
+	const double DefaultTransformValues[] = {Location[0], Location[1], Location[2], Rotation[0], Rotation[1], Rotation[2], Scale3D[0], Scale3D[1], Scale3D[2]};
+	
+		
 	//For non-transforms we need to re-negate the Y since it happens automatically(todo double check.).
 	//But then if we negate we need to re-re-negate... so leave it alone.
 
@@ -1952,8 +1691,8 @@ static bool ImportFBXTransformToChannels(FString NodeName, const UMovieSceneUser
 		FControlRigChannelEnum Channel = FControlRigChannelEnum::Bool;
 		SetChannelValue(nullptr, nullptr, BoolChannels[0], nullptr,nullptr, FrameRate, StartFrame,
 			Channel, ImportFBXControlRigSettings, DefaultTransform,
-			Translation[0], Translation[1], Translation[2], EulerRotation[0],
-			EulerRotation[1], EulerRotation[2], Scale[0], Scale[1], Scale[2]);
+			Transform[0], Transform[1], Transform[2], Transform[3],
+			Transform[4], Transform[5], Transform[6], Transform[7], Transform[8]);
 	}
 
 	if (EnumChannels.Num() == 1)
@@ -1961,8 +1700,8 @@ static bool ImportFBXTransformToChannels(FString NodeName, const UMovieSceneUser
 		FControlRigChannelEnum Channel = FControlRigChannelEnum::Enum;
 		SetChannelValue(nullptr, nullptr, nullptr, EnumChannels[0], nullptr, FrameRate, StartFrame,
 			Channel, ImportFBXControlRigSettings, DefaultTransform,
-			Translation[0], Translation[1], Translation[2], EulerRotation[0],
-			EulerRotation[1], EulerRotation[2], Scale[0], Scale[1], Scale[2]);
+			Transform[0], Transform[1], Transform[2], Transform[3],
+			Transform[4], Transform[5], Transform[6], Transform[7], Transform[8]);
 	}
 
 	if (IntegerChannels.Num() == 1)
@@ -1970,243 +1709,99 @@ static bool ImportFBXTransformToChannels(FString NodeName, const UMovieSceneUser
 		FControlRigChannelEnum Channel = FControlRigChannelEnum::Integer;
 		SetChannelValue(nullptr, nullptr, nullptr, nullptr, IntegerChannels[0], FrameRate, StartFrame,
 			Channel,  ImportFBXControlRigSettings,DefaultTransform,
-			Translation[0], Translation[1], Translation[2], EulerRotation[0], 
-			EulerRotation[1], EulerRotation[2], Scale[0], Scale[1], Scale[2]);
+			Transform[0], Transform[1], Transform[2], Transform[3], 
+			Transform[4], Transform[5], Transform[6], Transform[7], Transform[8]);
 	}
 
+	int ChannelEnumIndex = INDEX_NONE;
+
 	if (FloatChannels.Num() == 1)
 	{
-		FControlRigChannelEnum Channel = FControlRigChannelEnum::Float;
-		SetChannelValue(nullptr, FloatChannels[0], nullptr, nullptr, nullptr, FrameRate, StartFrame,
-			Channel, ImportFBXControlRigSettings, DefaultTransform,
-			Translation[0], Translation[1], Translation[2], EulerRotation[0],
-			EulerRotation[1], EulerRotation[2], Scale[0], Scale[1], Scale[2]);
+		ChannelEnumIndex = (int)FControlRigChannelEnum::Float;
 	}
 	else if (FloatChannels.Num() == 2)
 	{
-		FControlRigChannelEnum Channel = FControlRigChannelEnum::Vector2DX;
-
-		SetChannelValue(nullptr, FloatChannels[0], nullptr, nullptr, nullptr, FrameRate, StartFrame,
-			Channel, ImportFBXControlRigSettings, DefaultTransform,
-			Translation[0], Translation[1], Translation[2], EulerRotation[0],
-			EulerRotation[1], EulerRotation[2], Scale[0], Scale[1], Scale[2]);
-
-		Channel = FControlRigChannelEnum::Vector2DY;
-		SetChannelValue(nullptr, FloatChannels[1], nullptr, nullptr, nullptr, FrameRate, StartFrame,
-			Channel, ImportFBXControlRigSettings, DefaultTransform,
-			Translation[0], Translation[1], Translation[2], EulerRotation[0],
-			EulerRotation[1], EulerRotation[2], Scale[0], Scale[1], Scale[2]);
+		ChannelEnumIndex = (int)FControlRigChannelEnum::Vector2DX;
 	}
 	else if (FloatChannels.Num() == 3)
 	{
 		if (NodeAndChannels.ControlType == FFBXControlRigTypeProxyEnum::Position)
 		{
-			FControlRigChannelEnum Channel = FControlRigChannelEnum::PositionX;
-
-			SetChannelValue(nullptr, FloatChannels[0], nullptr, nullptr, nullptr, FrameRate, StartFrame,
-				Channel, ImportFBXControlRigSettings, DefaultTransform,
-				Translation[0], Translation[1], Translation[2], EulerRotation[0],
-				EulerRotation[1], EulerRotation[2], Scale[0], Scale[1], Scale[2]);
-
-			Channel = FControlRigChannelEnum::PositionY;
-			SetChannelValue(nullptr, FloatChannels[1], nullptr, nullptr, nullptr, FrameRate, StartFrame,
-				Channel, ImportFBXControlRigSettings, DefaultTransform,
-				Translation[0], Translation[1], Translation[2], EulerRotation[0],
-				EulerRotation[1], EulerRotation[2], Scale[0], Scale[1], Scale[2]);
-
-			Channel = FControlRigChannelEnum::PositionZ;
-			SetChannelValue(nullptr, FloatChannels[2], nullptr, nullptr, nullptr, FrameRate, StartFrame,
-				Channel, ImportFBXControlRigSettings, DefaultTransform,
-				Translation[0], Translation[1], Translation[2], EulerRotation[0],
-				EulerRotation[1], EulerRotation[2], Scale[0], Scale[1], Scale[2]);
+			ChannelEnumIndex = (int)FControlRigChannelEnum::PositionX;
 		}
 		else if (NodeAndChannels.ControlType == FFBXControlRigTypeProxyEnum::Rotator)
 		{
-			FControlRigChannelEnum Channel = FControlRigChannelEnum::RotatorX;
-			SetChannelValue(nullptr, FloatChannels[0], nullptr, nullptr, nullptr, FrameRate, StartFrame,
-				Channel, ImportFBXControlRigSettings, DefaultTransform,
-				Translation[0], Translation[1], Translation[2], EulerRotation[0],
-				EulerRotation[1], EulerRotation[2], Scale[0], Scale[1], Scale[2]);
-
-			Channel = FControlRigChannelEnum::RotatorY;
-			SetChannelValue(nullptr, FloatChannels[1], nullptr, nullptr, nullptr, FrameRate, StartFrame,
-				Channel, ImportFBXControlRigSettings, DefaultTransform,
-				Translation[0], Translation[1], Translation[2], EulerRotation[0],
-				EulerRotation[1], EulerRotation[2], Scale[0], Scale[1], Scale[2]);
-
-			Channel = FControlRigChannelEnum::RotatorZ;
-			SetChannelValue(nullptr, FloatChannels[2], nullptr, nullptr, nullptr, FrameRate, StartFrame,
-				Channel, ImportFBXControlRigSettings, DefaultTransform,
-				Translation[0], Translation[1], Translation[2], EulerRotation[0],
-				EulerRotation[1], EulerRotation[2], Scale[0], Scale[1], Scale[2]);
+			ChannelEnumIndex = (int)FControlRigChannelEnum::RotatorX;
 		}
 		else if (NodeAndChannels.ControlType == FFBXControlRigTypeProxyEnum::Scale)
 		{
-			FControlRigChannelEnum Channel = FControlRigChannelEnum::ScaleX;
-			SetChannelValue(nullptr, FloatChannels[0], nullptr, nullptr, nullptr, FrameRate, StartFrame,
-				Channel, ImportFBXControlRigSettings, DefaultTransform,
-				Translation[0], Translation[1], Translation[2], EulerRotation[0],
-				EulerRotation[1], EulerRotation[2], Scale[0], Scale[1], Scale[2]);
-			Channel = FControlRigChannelEnum::ScaleY;
-			SetChannelValue(nullptr, FloatChannels[1], nullptr, nullptr, nullptr, FrameRate, StartFrame,
-				Channel, ImportFBXControlRigSettings, DefaultTransform,
-				Translation[0], Translation[1], Translation[2], EulerRotation[0],
-				EulerRotation[1], EulerRotation[2], Scale[0], Scale[1], Scale[2]);
-			Channel = FControlRigChannelEnum::ScaleZ;
-			SetChannelValue(nullptr, FloatChannels[2], nullptr, nullptr, nullptr, FrameRate, StartFrame,
-				Channel, ImportFBXControlRigSettings, DefaultTransform,
-				Translation[0], Translation[1], Translation[2], EulerRotation[0],
-				EulerRotation[1], EulerRotation[2], Scale[0], Scale[1], Scale[2]);
+			ChannelEnumIndex = (int)FControlRigChannelEnum::ScaleX;
 		}
 	}
-	else if (FloatChannels.Num() == 9 || FloatChannels.Num() == 6)
+	if (ChannelEnumIndex != INDEX_NONE)
 	{
-		FloatChannels[0]->SetDefault(Location.X);
-		FloatChannels[1]->SetDefault(Location.Y);
-		FloatChannels[2]->SetDefault(Location.Z);
-
-		FloatChannels[3]->SetDefault(Rotation.X);
-		FloatChannels[4]->SetDefault(Rotation.Y);
-		FloatChannels[5]->SetDefault(Rotation.Z);
-
-		if (FloatChannels.Num() > 6) //noscale
+		for (int i = 0; i < FloatChannels.Num(); i++)
 		{
-			FloatChannels[6]->SetDefault(Scale3D.X);
-			FloatChannels[7]->SetDefault(Scale3D.Y);
-			FloatChannels[8]->SetDefault(Scale3D.Z);
+			const FControlRigChannelEnum ChannelEnum = (FControlRigChannelEnum)(ChannelEnumIndex + i);
+			SetChannelValue(nullptr, FloatChannels[i], nullptr, nullptr, nullptr,
+				FrameRate, StartFrame, ChannelEnum, ImportFBXControlRigSettings, DefaultTransform,
+				Transform[0], Transform[1], Transform[2], Transform[3],
+				Transform[4], Transform[5], Transform[6], Transform[7], Transform[8]);	
 		}
-
-		ImportTransformChannelToFloat(Translation[0], FloatChannels[0], FrameRate, false, false, StartFrame);
-		ImportTransformChannelToFloat(Translation[1], FloatChannels[1], FrameRate, true, false, StartFrame);
-		ImportTransformChannelToFloat(Translation[2], FloatChannels[2], FrameRate, false, false, StartFrame);
-
-		ImportTransformChannelToFloat(EulerRotation[0], FloatChannels[3], FrameRate, false, false, StartFrame);
-		ImportTransformChannelToFloat(EulerRotation[1], FloatChannels[4], FrameRate, true, false, StartFrame);
-		ImportTransformChannelToFloat(EulerRotation[2], FloatChannels[5], FrameRate, true, false, StartFrame);
-
-		if (FloatChannels.Num() > 6) //noscale
+	}
+	else if (FloatChannels.Num() == 9 || FloatChannels.Num() == 6)
+	{
+		for (int i = 0; i < FloatChannels.Num(); i++)
 		{
-			ImportTransformChannelToFloat(Scale[0], FloatChannels[6], FrameRate, false, false, StartFrame);
-			ImportTransformChannelToFloat(Scale[1], FloatChannels[7], FrameRate, false, false, StartFrame);
-			ImportTransformChannelToFloat(Scale[2], FloatChannels[8], FrameRate, false, false, StartFrame);
+			FloatChannels[i]->SetDefault(DefaultTransformValues[i]);
+
+			const bool bNegate = i == 1 || i == 4 || i == 5;
+			ImportTransformChannelToFloat(Transform[i], FloatChannels[i], FrameRate, bNegate, false, StartFrame);
 		}
 	}
 
+	ChannelEnumIndex = INDEX_NONE;
 
 	if (DoubleChannels.Num() == 1)
 	{
-		FControlRigChannelEnum Channel = FControlRigChannelEnum::Float; //todo control rig doesn't support double but may
-		SetChannelValue(DoubleChannels[0], nullptr, nullptr, nullptr, nullptr, FrameRate, StartFrame,
-			Channel, ImportFBXControlRigSettings, DefaultTransform,
-			Translation[0], Translation[1], Translation[2], EulerRotation[0],
-			EulerRotation[1], EulerRotation[2], Scale[0], Scale[1], Scale[2]);
+		ChannelEnumIndex = (int)FControlRigChannelEnum::Float;
 	}
 	else if (DoubleChannels.Num() == 2)
 	{
-		FControlRigChannelEnum Channel = FControlRigChannelEnum::Vector2DX;
-
-		SetChannelValue(DoubleChannels[0], nullptr, nullptr, nullptr, nullptr, FrameRate, StartFrame,
-			Channel, ImportFBXControlRigSettings, DefaultTransform,
-			Translation[0], Translation[1], Translation[2], EulerRotation[0],
-			EulerRotation[1], EulerRotation[2], Scale[0], Scale[1], Scale[2]);
-
-		Channel = FControlRigChannelEnum::Vector2DY;
-		SetChannelValue(DoubleChannels[1], nullptr, nullptr, nullptr, nullptr, FrameRate, StartFrame,
-			Channel, ImportFBXControlRigSettings, DefaultTransform,
-			Translation[0], Translation[1], Translation[2], EulerRotation[0],
-			EulerRotation[1], EulerRotation[2], Scale[0], Scale[1], Scale[2]);
+		ChannelEnumIndex = (int)FControlRigChannelEnum::Vector2DX;
 	}
 	else if (DoubleChannels.Num() == 3)
 	{
 		if (NodeAndChannels.ControlType == FFBXControlRigTypeProxyEnum::Position)
 		{
-			FControlRigChannelEnum Channel = FControlRigChannelEnum::PositionX;
-			SetChannelValue(DoubleChannels[0], nullptr, nullptr, nullptr, nullptr, FrameRate, StartFrame,
-				Channel, ImportFBXControlRigSettings, DefaultTransform,
-				Translation[0], Translation[1], Translation[2], EulerRotation[0],
-				EulerRotation[1], EulerRotation[2], Scale[0], Scale[1], Scale[2]);
-
-			Channel = FControlRigChannelEnum::PositionY;
-			SetChannelValue(DoubleChannels[1], nullptr, nullptr, nullptr, nullptr, FrameRate, StartFrame,
-				Channel, ImportFBXControlRigSettings, DefaultTransform,
-				Translation[0], Translation[1], Translation[2], EulerRotation[0],
-				EulerRotation[1], EulerRotation[2], Scale[0], Scale[1], Scale[2]);
-
-			Channel = FControlRigChannelEnum::PositionZ;
-			SetChannelValue(DoubleChannels[2], nullptr,nullptr, nullptr, nullptr, FrameRate, StartFrame,
-				Channel, ImportFBXControlRigSettings, DefaultTransform,
-				Translation[0], Translation[1], Translation[2], EulerRotation[0],
-				EulerRotation[1], EulerRotation[2], Scale[0], Scale[1], Scale[2]);
+			ChannelEnumIndex = (int)FControlRigChannelEnum::PositionX;
 		}
 		else if (NodeAndChannels.ControlType == FFBXControlRigTypeProxyEnum::Rotator)
 		{
-			FControlRigChannelEnum Channel = FControlRigChannelEnum::RotatorX;
-			SetChannelValue(DoubleChannels[0], nullptr,nullptr, nullptr, nullptr, FrameRate, StartFrame,
-				Channel, ImportFBXControlRigSettings, DefaultTransform,
-				Translation[0], Translation[1], Translation[2], EulerRotation[0],
-				EulerRotation[1], EulerRotation[2], Scale[0], Scale[1], Scale[2]);
-
-			Channel = FControlRigChannelEnum::RotatorY;
-			SetChannelValue(DoubleChannels[1], nullptr, nullptr, nullptr, nullptr, FrameRate, StartFrame,
-				Channel, ImportFBXControlRigSettings, DefaultTransform,
-				Translation[0], Translation[1], Translation[2], EulerRotation[0],
-				EulerRotation[1], EulerRotation[2], Scale[0], Scale[1], Scale[2]);
-
-			Channel = FControlRigChannelEnum::RotatorZ;
-			SetChannelValue(DoubleChannels[2], nullptr, nullptr, nullptr, nullptr, FrameRate, StartFrame,
-				Channel, ImportFBXControlRigSettings, DefaultTransform,
-				Translation[0], Translation[1], Translation[2], EulerRotation[0],
-				EulerRotation[1], EulerRotation[2], Scale[0], Scale[1], Scale[2]);
+			ChannelEnumIndex = (int)FControlRigChannelEnum::RotatorX;
 		}
 		else if (NodeAndChannels.ControlType == FFBXControlRigTypeProxyEnum::Scale)
 		{
-			FControlRigChannelEnum Channel = FControlRigChannelEnum::ScaleX;
-			SetChannelValue(DoubleChannels[0], nullptr, nullptr, nullptr, nullptr, FrameRate, StartFrame,
-				Channel, ImportFBXControlRigSettings, DefaultTransform,
-				Translation[0], Translation[1], Translation[2], EulerRotation[0],
-				EulerRotation[1], EulerRotation[2], Scale[0], Scale[1], Scale[2]);
-			Channel = FControlRigChannelEnum::ScaleY;
-			SetChannelValue(DoubleChannels[1], nullptr, nullptr, nullptr, nullptr, FrameRate, StartFrame,
-				Channel, ImportFBXControlRigSettings, DefaultTransform,
-				Translation[0], Translation[1], Translation[2], EulerRotation[0],
-				EulerRotation[1], EulerRotation[2], Scale[0], Scale[1], Scale[2]);
-			Channel = FControlRigChannelEnum::ScaleZ;
-			SetChannelValue(DoubleChannels[2], nullptr, nullptr, nullptr, nullptr, FrameRate, StartFrame,
-				Channel, ImportFBXControlRigSettings, DefaultTransform,
-				Translation[0], Translation[1], Translation[2], EulerRotation[0],
-				EulerRotation[1], EulerRotation[2], Scale[0], Scale[1], Scale[2]);
+			ChannelEnumIndex = (int)FControlRigChannelEnum::ScaleX;
 		}
 	}
-	else if (DoubleChannels.Num() == 9 || DoubleChannels.Num() == 6)
+	if (ChannelEnumIndex != INDEX_NONE)
 	{
-		DoubleChannels[0]->SetDefault(Location.X);
-		DoubleChannels[1]->SetDefault(Location.Y);
-		DoubleChannels[2]->SetDefault(Location.Z);
-
-		DoubleChannels[3]->SetDefault(Rotation.X);
-		DoubleChannels[4]->SetDefault(Rotation.Y);
-		DoubleChannels[5]->SetDefault(Rotation.Z);
-
-		if (DoubleChannels.Num() > 6) //noscale
+		for (int i = 0; i < DoubleChannels.Num(); i++)
 		{
-			DoubleChannels[6]->SetDefault(Scale3D.X);
-			DoubleChannels[7]->SetDefault(Scale3D.Y);
-			DoubleChannels[8]->SetDefault(Scale3D.Z);
+			const FControlRigChannelEnum ChannelEnum = (FControlRigChannelEnum)(ChannelEnumIndex + i);
+			SetChannelValue(DoubleChannels[i], nullptr, nullptr, nullptr, nullptr,
+				FrameRate, StartFrame, ChannelEnum, ImportFBXControlRigSettings, DefaultTransform,
+				Transform[0], Transform[1], Transform[2], Transform[3],
+				Transform[4], Transform[5], Transform[6], Transform[7], Transform[8]);	
 		}
-
-		ImportTransformChannelToDouble(Translation[0], DoubleChannels[0], FrameRate, false, false, StartFrame);
-		ImportTransformChannelToDouble(Translation[1], DoubleChannels[1], FrameRate, false, false, StartFrame);
-		ImportTransformChannelToDouble(Translation[2], DoubleChannels[2], FrameRate, false, false, StartFrame);
-
-		ImportTransformChannelToDouble(EulerRotation[0], DoubleChannels[3], FrameRate, false,false, StartFrame);
-		ImportTransformChannelToDouble(EulerRotation[1], DoubleChannels[4], FrameRate, false, false, StartFrame);
-		ImportTransformChannelToDouble(EulerRotation[2], DoubleChannels[5], FrameRate, false, false, StartFrame);
-
-		if (DoubleChannels.Num() > 6) //noscale
+	}
+	else if (DoubleChannels.Num() == 9 || DoubleChannels.Num() == 6)
+	{
+		for (int i = 0; i < DoubleChannels.Num(); i++)
 		{
-			ImportTransformChannelToDouble(Scale[0], DoubleChannels[6], FrameRate, false, false, StartFrame);
-			ImportTransformChannelToDouble(Scale[1], DoubleChannels[7], FrameRate, false, false, StartFrame);
-			ImportTransformChannelToDouble(Scale[2], DoubleChannels[8], FrameRate, false, false, StartFrame);
+			DoubleChannels[i]->SetDefault(DefaultTransformValues[i]);
+			ImportTransformChannelToDouble(Transform[i], DoubleChannels[i], FrameRate, false, false, StartFrame);
 		}
 	}
 	return true;
@@ -2214,14 +1809,20 @@ static bool ImportFBXTransformToChannels(FString NodeName, const UMovieSceneUser
 static FString GetNewString(const FString& InString, UMovieSceneUserImportFBXControlRigSettings* ImportFBXControlRigSettings)
 {
 	FString NewString = InString;
+	int Index = INDEX_NONE;
+	
+	if (ImportFBXControlRigSettings->bStripNamespace && InString.FindLastChar(':', Index))
+	{
+		NewString.RightChopInline(Index + 1);
+	}
 	for (const FControlFindReplaceString& FindReplace : ImportFBXControlRigSettings->FindAndReplaceStrings)
 	{
-		NewString = NewString.Replace(*FindReplace.Find, *FindReplace.Replace); //ignores tupe
+		NewString.ReplaceInline(*FindReplace.Find, *FindReplace.Replace); //ignores tupe
 	}
 	return NewString;
 }
 
-static void PrepForInsertReplaceAnimation(bool bInsert, const FFBXNodeAndChannels& NodeAndChannel,
+static void PrepForInsertReplaceAnimation(bool bInsert, const FRigControlFBXNodeAndChannels& NodeAndChannel,
 	FFrameNumber  FrameToInsertOrReplace, FFrameNumber  StartFrame, FFrameNumber  EndFrame)
 {
 
@@ -2378,8 +1979,7 @@ class SControlRigImportFBXSettings : public SCompoundWidget
 			NAME_None,
 			EUserInterfaceActionType::Button
 		);
-
-
+		
 		MenuBuilder.AddMenuEntry(
 			NSLOCTEXT("MovieSceneTools", "MetaHumanControlMappings", "MetaHuman Control Mappings"),
 			NSLOCTEXT("MovieSceneTools", "MetaHumanControlMappings_Tooltip", "Use MetaHuman Control Mappings Preset"),
@@ -2391,7 +1991,6 @@ class SControlRigImportFBXSettings : public SCompoundWidget
 			EUserInterfaceActionType::Button
 		);
 
-
 		return MenuBuilder.MakeWidget();
 	}
 
@@ -2437,20 +2036,16 @@ class SControlRigImportFBXSettings : public SCompoundWidget
 			ImportFBXSettings->ImportedFileName = FileName;
 		}
 	}
-	void SetNodeAndChannels(TArray<FFBXNodeAndChannels>* InNodeAndChannels)
+	void SetNodeAndChannels(TArray<FRigControlFBXNodeAndChannels>* InNodeAndChannels)
 	{
 		NodeAndChannels = InNodeAndChannels;
 	}
 
-
-
-
 private:
 
 	FReply OnImportFBXClicked()
 	{
-
-		if (Sequencer.IsValid() == false)
+		if (!Sequencer.IsValid() || !NodeAndChannels)
 		{
 			return  FReply::Unhandled();
 		}
@@ -2458,7 +2053,7 @@ private:
 		UMovieSceneUserImportFBXControlRigSettings* ImportFBXControlRigSettings = GetMutableDefault<UMovieSceneUserImportFBXControlRigSettings>();
 		
 		TArray<FName> SelectedControlNames;
-		for (FFBXNodeAndChannels& NodeAndChannel : *NodeAndChannels)
+		for (const FRigControlFBXNodeAndChannels& NodeAndChannel : *NodeAndChannels)
 		{
 			if (NodeAndChannel.MovieSceneTrack)
 			{
@@ -2470,14 +2065,13 @@ private:
 					SelectedControlNames.Append(LocalControls);
 				}
 			}
-
-
 		}
-		bool bValid = MovieSceneToolHelpers::ImportFBXIntoControlRigChannels(Sequencer.Pin()->GetFocusedMovieSceneSequence()->GetMovieScene(), ImportFilename, ImportFBXControlRigSettings, 
-			NodeAndChannels,SelectedControlNames, Sequencer.Pin()->GetFocusedTickResolution());
-		
 
-		TSharedPtr<SWindow> Window = FSlateApplication::Get().FindWidgetWindow(AsShared());
+		const bool bValid = MovieSceneToolHelpers::ImportFBXIntoControlRigChannels(
+			Sequencer.Pin()->GetFocusedMovieSceneSequence()->GetMovieScene(), ImportFilename, ImportFBXControlRigSettings, 
+		    NodeAndChannels,SelectedControlNames, Sequencer.Pin()->GetFocusedTickResolution());
+		
+		const TSharedPtr<SWindow> Window = FSlateApplication::Get().FindWidgetWindow(AsShared());
 
 		if (Window.IsValid())
 		{
@@ -2488,109 +2082,196 @@ private:
 			Sequencer.Pin()->NotifyMovieSceneDataChanged(EMovieSceneDataChangeType::MovieSceneStructureItemAdded);
 		}
 		return bValid ? FReply::Handled() : FReply::Unhandled();
-
 	}
 
-	void SetPresets(bool bMetaHuman)
+	void SetPresets(const bool bMetaHuman) const
 	{
-		//since we can't change the API unfortunately need to do this here.
 		UMovieSceneUserImportFBXControlRigSettings* ImportFBXControlRigSettings = GetMutableDefault<UMovieSceneUserImportFBXControlRigSettings>();
-		ImportFBXControlRigSettings->ControlChannelMappings.SetNum(0); //clear and reset
-		FControlToTransformMappings Bool;
-		Bool.bNegate = false;
-		Bool.ControlChannel = FControlRigChannelEnum::Bool;
-		Bool.FBXChannel = FTransformChannelEnum::TranslateX;
-		ImportFBXControlRigSettings->ControlChannelMappings.Add(Bool);
+		ImportFBXControlRigSettings->LoadControlMappingsFromPreset(bMetaHuman);
+	}
+
+	TSharedPtr<IDetailsView> DetailView;
+	FString ImportFilename;
+	TArray<FRigControlFBXNodeAndChannels>* NodeAndChannels = nullptr;
+	TWeakPtr<ISequencer> Sequencer;
+};
+
+class SControlRigExportFBXSettings : public SCompoundWidget
+{
+	SLATE_BEGIN_ARGS(SControlRigExportFBXSettings) {}
+	SLATE_ARGUMENT(FString, ExportFilename)
+	SLATE_END_ARGS()
+
+	void Construct(const FArguments& InArgs,  const TSharedRef<ISequencer>& InSequencer, TWeakObjectPtr<UMovieSceneTrack> InTrack)
+	{
+		Sequencer = InSequencer;
+		Track = InTrack;
+		
+		UMovieSceneUserExportFBXControlRigSettings* ExportFBXSettings = GetMutableDefault<UMovieSceneUserExportFBXControlRigSettings>();
+		ExportFBXSettings->ExportFileName = InArgs._ExportFilename;
+
+		FPropertyEditorModule& PropertyEditor = FModuleManager::LoadModuleChecked<FPropertyEditorModule>("PropertyEditor");
+		
+		FDetailsViewArgs DetailsViewArgs;
+		DetailsViewArgs.bShowOptions = false;
+		DetailsViewArgs.bAllowSearch = false;
+		DetailsViewArgs.bShowPropertyMatrixButton = false;
+		DetailsViewArgs.bUpdatesFromSelection = false;
+		DetailsViewArgs.bLockable = false;
+		DetailsViewArgs.NameAreaSettings = FDetailsViewArgs::HideNameArea;
+		DetailsViewArgs.ViewIdentifier = "Export FBX Settings";
+
+		DetailView = PropertyEditor.CreateDetailView(DetailsViewArgs);
+		DetailView->SetObject(ExportFBXSettings);
+
+		ChildSlot
+		[
+			SNew(SVerticalBox)
+
+			+ SVerticalBox::Slot()
+			[
+				DetailView.ToSharedRef()
+			]
+
+			+ SVerticalBox::Slot()
+			.AutoHeight()
+			[
+				SNew(SComboButton)
+				.HasDownArrow(true)
+				.OnGetMenuContent(this, &SControlRigExportFBXSettings::HandlePresetMenuContent)
+				.ButtonContent()
+				[
+					SNew(STextBlock)
+					.Text(NSLOCTEXT("MovieSceneTools", "ControlMappingPresets", "Control Mapping Presets"))
+					.ToolTipText(NSLOCTEXT("MovieSceneTools", "SetControlMappingFromAPreset", "Set Control Mappings From A Preset"))
+				]
+			]
 
-		FControlToTransformMappings Float;
-		Float.bNegate = false;
-		Float.ControlChannel = FControlRigChannelEnum::Float;
-		if (bMetaHuman)
+			+ SVerticalBox::Slot()
+			.AutoHeight()
+			.HAlign(HAlign_Right)
+			.Padding(5.f)
+			[
+				SNew(SButton)
+				.ContentPadding(FMargin(10, 5))
+				.Text(NSLOCTEXT("MovieSceneTools", "ExportFBXButtonText", "Export"))
+				.OnClicked(this, &SControlRigExportFBXSettings::OnExportFBXClicked)
+			]
+		];
+	}
+	TSharedRef<SWidget> HandlePresetMenuContent()
+	{
+		FMenuBuilder MenuBuilder(/*bInShouldCloseWindowAfterMenuSelection=*/true, nullptr);
+
+		MenuBuilder.AddMenuEntry(
+			NSLOCTEXT("MovieSceneTools", "EmptyControlMappings", "Empty Control Mappings"),
+			NSLOCTEXT("MovieSceneTools", "EmptyControlMappings_Tooltip",
+				"Do not use Control Mappings. Bool, Enum, Int and Float channels will be exported as attributes of the corresponding data type"),
+			FSlateIcon(),
+			FUIAction(
+				FExecuteAction::CreateSP(this, &SControlRigExportFBXSettings::ClearPresets)
+			),
+			NAME_None,
+			EUserInterfaceActionType::Button
+		);
+		
+		MenuBuilder.AddMenuEntry(
+			NSLOCTEXT("MovieSceneTools", "DefaultControlMappings", "Default Control Mappings"),
+			NSLOCTEXT("MovieSceneTools", "DefaultControlMappings_Tooltip", "Use Default Control Mappings Preset"),
+			FSlateIcon(),
+			FUIAction(
+				FExecuteAction::CreateSP(this, &SControlRigExportFBXSettings::SetPresets, false)
+			),
+			NAME_None,
+			EUserInterfaceActionType::Button
+		);
+
+		MenuBuilder.AddMenuEntry(
+			NSLOCTEXT("MovieSceneTools", "MetaHumanControlMappings", "MetaHuman Control Mappings"),
+			NSLOCTEXT("MovieSceneTools", "MetaHumanControlMappings_Tooltip", "Use MetaHuman Control Mappings Preset"),
+			FSlateIcon(),
+			FUIAction(
+				FExecuteAction::CreateSP(this, &SControlRigExportFBXSettings::SetPresets, true)
+			),
+			NAME_None,
+			EUserInterfaceActionType::Button
+		);
+
+		return MenuBuilder.MakeWidget();
+	}
+
+private:
+
+	FReply OnExportFBXClicked()
+	{
+		if (!Sequencer.IsValid() || !Track.IsValid())
 		{
-			Float.FBXChannel = FTransformChannelEnum::TranslateY;  //use Y for metahuman
+			return  FReply::Unhandled();
+		}
+
+		UMovieSceneUserExportFBXControlRigSettings* ExportFBXControlRigSettings = GetMutableDefault<UMovieSceneUserExportFBXControlRigSettings>();
+		
+		TArray<FName> SelectedControlNames;
+		if (Track.IsValid())
+		{
+			if (INodeAndChannelMappings* ChannelMapping = Cast<INodeAndChannelMappings>(Track.Get()))
+			{
+				ChannelMapping->GetSelectedNodes(SelectedControlNames);
+			}
+		}
+
+		FMovieSceneSequenceTransform RootToLocalTransform;
+		if (Sequencer.IsValid())
+		{
+			RootToLocalTransform = Sequencer.Pin()->GetFocusedMovieSceneSequenceTransform();
+		}
+
+		const bool bValid = MovieSceneToolHelpers::ExportFBXFromControlRigChannels(
+			Track.Get()->GetSectionToKey(), ExportFBXControlRigSettings, SelectedControlNames, RootToLocalTransform);
+
+		if (bValid)
+		{
+			FNotificationInfo Info(NSLOCTEXT("MovieSceneTools", "ExportControlRigFBXSucceeded", "FBX Export Succeeded."));
+			Info.Hyperlink = FSimpleDelegate::CreateStatic([](FString InFilename) { FPlatformProcess::ExploreFolder(*InFilename); }, ExportFBXControlRigSettings->ExportFileName);
+			Info.HyperlinkText = FText::FromString(ExportFBXControlRigSettings->ExportFileName);
+			Info.ExpireDuration = 5.0f;
+			FSlateNotificationManager::Get().AddNotification(Info)->SetCompletionState(SNotificationItem::CS_Success);
 		}
 		else
 		{
-			Float.FBXChannel = FTransformChannelEnum::TranslateX;
-		}
-		ImportFBXControlRigSettings->ControlChannelMappings.Add(Float);
-
-		FControlToTransformMappings Vector2DX;
-		Vector2DX.bNegate = false;
-		Vector2DX.ControlChannel = FControlRigChannelEnum::Vector2DX;
-		Vector2DX.FBXChannel = FTransformChannelEnum::TranslateX;
-		ImportFBXControlRigSettings->ControlChannelMappings.Add(Vector2DX);
-
-		FControlToTransformMappings Vector2DY;
-		Vector2DY.bNegate = false;
-		Vector2DY.ControlChannel = FControlRigChannelEnum::Vector2DY;
-		Vector2DY.FBXChannel = FTransformChannelEnum::TranslateY;
-		ImportFBXControlRigSettings->ControlChannelMappings.Add(Vector2DY);
-
-		FControlToTransformMappings PositionX;
-		PositionX.bNegate = false;
-		PositionX.ControlChannel = FControlRigChannelEnum::PositionX;
-		PositionX.FBXChannel = FTransformChannelEnum::TranslateX;
-		ImportFBXControlRigSettings->ControlChannelMappings.Add(PositionX);
-
-		FControlToTransformMappings PositionY;
-		PositionY.bNegate = false;
-		PositionY.ControlChannel = FControlRigChannelEnum::PositionY;
-		PositionY.FBXChannel = FTransformChannelEnum::TranslateY;
-		ImportFBXControlRigSettings->ControlChannelMappings.Add(PositionY);
-
-		FControlToTransformMappings PositionZ;
-		PositionZ.bNegate = false;
-		PositionZ.ControlChannel = FControlRigChannelEnum::PositionZ;
-		PositionZ.FBXChannel = FTransformChannelEnum::TranslateZ;
-		ImportFBXControlRigSettings->ControlChannelMappings.Add(PositionZ);
-
-		FControlToTransformMappings RotatorX;
-		RotatorX.bNegate = false;
-		RotatorX.ControlChannel = FControlRigChannelEnum::RotatorX;
-		RotatorX.FBXChannel = FTransformChannelEnum::RotateX;
-		ImportFBXControlRigSettings->ControlChannelMappings.Add(RotatorX);
-
-		FControlToTransformMappings RotatorY;
-		RotatorY.bNegate = false;
-		RotatorY.ControlChannel = FControlRigChannelEnum::RotatorY;
-		RotatorY.FBXChannel = FTransformChannelEnum::RotateY;
-		ImportFBXControlRigSettings->ControlChannelMappings.Add(RotatorY);
-
-		FControlToTransformMappings RotatorZ;
-		RotatorZ.bNegate = false;
-		RotatorZ.ControlChannel = FControlRigChannelEnum::RotatorZ;
-		RotatorZ.FBXChannel = FTransformChannelEnum::RotateZ;
-		ImportFBXControlRigSettings->ControlChannelMappings.Add(RotatorZ);
-
-		FControlToTransformMappings ScaleX;
-		ScaleX.bNegate = false;
-		ScaleX.ControlChannel = FControlRigChannelEnum::ScaleX;
-		ScaleX.FBXChannel = FTransformChannelEnum::ScaleX;
-		ImportFBXControlRigSettings->ControlChannelMappings.Add(ScaleX);
-
-		FControlToTransformMappings ScaleY;
-		ScaleY.bNegate = false;
-		ScaleY.ControlChannel = FControlRigChannelEnum::ScaleY;
-		ScaleY.FBXChannel = FTransformChannelEnum::ScaleY;
-		ImportFBXControlRigSettings->ControlChannelMappings.Add(ScaleY);
-
-		FControlToTransformMappings ScaleZ;
-		ScaleZ.bNegate = false;
-		ScaleZ.ControlChannel = FControlRigChannelEnum::ScaleZ;
-		ScaleZ.FBXChannel = FTransformChannelEnum::ScaleZ;
-		ImportFBXControlRigSettings->ControlChannelMappings.Add(ScaleZ);
+			FNotificationInfo Info(NSLOCTEXT("MovieSceneTools", "ExportControlRigFBXFailed", "FBX Export Failed."));
+			Info.ExpireDuration = 5.0f;
+			FSlateNotificationManager::Get().AddNotification(Info)->SetCompletionState(SNotificationItem::CS_Fail);
+		}
+		
+		const TSharedPtr<SWindow> Window = FSlateApplication::Get().FindWidgetWindow(AsShared());
+
+		if (Window.IsValid())
+		{
+			Window->RequestDestroyWindow();
+		}
+		return bValid ? FReply::Handled() : FReply::Unhandled();
+	}
+
+	void SetPresets(const bool bMetaHuman) const
+	{
+		UMovieSceneUserExportFBXControlRigSettings* ExportFBXControlRigSettings = GetMutableDefault<UMovieSceneUserExportFBXControlRigSettings>();
+		ExportFBXControlRigSettings->LoadControlMappingsFromPreset(bMetaHuman);
 	}
 
+	void ClearPresets() const
+	{
+		UMovieSceneUserExportFBXControlRigSettings* ExportFBXControlRigSettings = GetMutableDefault<UMovieSceneUserExportFBXControlRigSettings>();
+		ExportFBXControlRigSettings->ControlChannelMappings.Empty();
+	}
+	
 	TSharedPtr<IDetailsView> DetailView;
-	FString ImportFilename;
-	TArray<FFBXNodeAndChannels>* NodeAndChannels;
 	TWeakPtr<ISequencer> Sequencer;
-
+	TWeakObjectPtr<UMovieSceneTrack> Track;
 };
 
 bool MovieSceneToolHelpers::ImportFBXIntoControlRigChannels(UMovieScene* MovieScene,const FString& ImportFilename, UMovieSceneUserImportFBXControlRigSettings* ImportFBXControlRigSettings,
-	TArray<FFBXNodeAndChannels>* NodeAndChannels, const TArray<FName>& SelectedControlNames, FFrameRate FrameRate)
+	TArray<FRigControlFBXNodeAndChannels>* NodeAndChannels, const TArray<FName>& SelectedControlNames, FFrameRate FrameRate)
 {
 	UnFbx::FFbxImporter* FbxImporter = UnFbx::FFbxImporter::GetInstance();
 
@@ -2641,21 +2322,9 @@ bool MovieSceneToolHelpers::ImportFBXIntoControlRigChannels(UMovieScene* MovieSc
 		//if matching selected remove out the non-selected
 		if (ImportFBXControlRigSettings->bImportOntoSelectedControls)
 		{
-			for (int32 Index = NodeAndChannels->Num() - 1; Index >= 0; --Index)
-			{
-				bool bHasOneMatch = false;
-				for (const FName& SelectedName : SelectedControlNames)
-				{
-					if (FCString::Strcmp(*SelectedName.ToString().ToUpper(), *((*NodeAndChannels)[Index].NodeName).ToUpper()) == 0)
-					{
-						bHasOneMatch = true;
-					}
-				}
-				if (!bHasOneMatch)
-				{
-					NodeAndChannels->RemoveAt(Index);
-				}
-			}
+			NodeAndChannels->RemoveAll([SelectedControlNames](const FRigControlFBXNodeAndChannels& A)
+				{ return !SelectedControlNames.Contains(A.ControlName); }
+			);
 		}
 
 		FFrameNumber  FrameToInsertOrReplace = ImportFBXControlRigSettings->TimeToInsertOrReplaceAnimation;
@@ -2691,21 +2360,14 @@ bool MovieSceneToolHelpers::ImportFBXIntoControlRigChannels(UMovieScene* MovieSc
 
 		TSet<UMovieSceneSection*> AllModifiedSections;
 
-		for (int32 NodeIndex = 0; NodeIndex < AllNodeNames.Num(); ++NodeIndex)
+		for (const FString& NodeName : AllNodeNames)
 		{
-			FString NodeName = AllNodeNames[NodeIndex];
-			/** Why was this here I think due to speeed....
-			if (NodeName[0] != 'C')
-			{
-				continue;
-			}
-			*/
-			FString NewNodeName = GetNewString(*(NodeName).ToUpper(), ImportFBXControlRigSettings);
+			const FString NewNodeName = GetNewString(*NodeName, ImportFBXControlRigSettings);
 
 			TSet<UMovieSceneSection*> ModifiedSections;
-			for (FFBXNodeAndChannels& NodeAndChannel : *NodeAndChannels)
+			for (FRigControlFBXNodeAndChannels& NodeAndChannel : *NodeAndChannels)
 			{
-				if (FCString::Strcmp(*(NodeAndChannel.NodeName).ToUpper(), *NewNodeName.ToUpper()) == 0)
+				if (NodeAndChannel.NodeName.Equals(NewNodeName, ESearchCase::IgnoreCase))
 				{
 					if (NodeAndChannel.MovieSceneTrack)
 					{
@@ -2728,7 +2390,16 @@ bool MovieSceneToolHelpers::ImportFBXIntoControlRigChannels(UMovieScene* MovieSc
 						FrameToInsertOrReplace,
 						StartFrame, EndFrame);
 
-					ImportFBXTransformToChannels(NodeName, CurrentImportFBXSettings, ImportFBXControlRigSettings, FrameToInsertOrReplace, FrameRate, NodeAndChannel, CurveAPI);
+					FFBXControlRigTypeProxyEnum ControlType = NodeAndChannel.ControlType;
+					bool bTryImportNonTransformCurves = ControlType == FFBXControlRigTypeProxyEnum::Bool || ControlType == FFBXControlRigTypeProxyEnum::Float || ControlType == FFBXControlRigTypeProxyEnum::Integer;
+
+					// If the node channels could not be imported from custom FBX attribute curves, try importing it them from the FBX transform channels, using the corresponding mapping 
+					if (!bTryImportNonTransformCurves || !ImportFBXNonTransformCurvesToChannels(
+						NodeName, CurrentImportFBXSettings, ImportFBXControlRigSettings, FrameToInsertOrReplace,
+						FrameRate, NodeAndChannel, CurveAPI))
+					{
+						ImportFBXTransformToChannels(NodeName, CurrentImportFBXSettings, ImportFBXControlRigSettings, FrameToInsertOrReplace, FrameRate, NodeAndChannel, CurveAPI);
+					}
 				}
 			}
 
@@ -2801,7 +2472,102 @@ bool MovieSceneToolHelpers::ImportFBXIntoControlRigChannels(UMovieScene* MovieSc
 	return bValid;
 }
 
-bool MovieSceneToolHelpers::ImportFBXIntoChannelsWithDialog(const TSharedRef<ISequencer>& InSequencer,TArray<FFBXNodeAndChannels>* NodeAndChannels)
+bool MovieSceneToolHelpers::ExportFBXFromControlRigChannels(const UMovieSceneSection* Section,
+	const UMovieSceneUserExportFBXControlRigSettings* ExportFBXControlRigSettings, const TArray<FName>& SelectedControlNames,
+	const FMovieSceneSequenceTransform& RootToLocalTransform)
+{
+	if (!Section)
+	{
+		return false;
+	}
+
+	UnFbx::FFbxExporter* Exporter = UnFbx::FFbxExporter::GetInstance();
+
+	UFbxExportOption* ExportOptions = Exporter->GetExportOptions();
+	ExportOptions->FbxExportCompatibility = ExportFBXControlRigSettings->FbxExportCompatibility;
+	ExportOptions->bASCII = ExportFBXControlRigSettings->bASCII;
+	ExportOptions->bForceFrontXAxis = ExportFBXControlRigSettings->bForceFrontXAxis;
+	ExportOptions->bExportLocalTime = ExportFBXControlRigSettings->bExportLocalTime;
+	
+	Exporter->CreateDocument();
+	Exporter->SetTransformBaking(true);
+	Exporter->SetKeepHierarchy(true);
+
+	TArray<FControlRigFbxNodeMapping> ChannelsMapping;
+
+	// Map to make it easier to convert from Control Rig Channel Enum to Control Rig Type Proxy + Transform Channel Index
+	const TMap<FControlRigChannelEnum, TPair<FFBXControlRigTypeProxyEnum, uint8>> ProxyMapping = {
+		{ FControlRigChannelEnum::Vector2DX, { FFBXControlRigTypeProxyEnum::Vector2D, 0 } },
+		{ FControlRigChannelEnum::Vector2DY, { FFBXControlRigTypeProxyEnum::Vector2D, 1 } },
+		{ FControlRigChannelEnum::PositionX, { FFBXControlRigTypeProxyEnum::Position, 0 } },
+		{ FControlRigChannelEnum::PositionY, { FFBXControlRigTypeProxyEnum::Position, 1 } },
+		{ FControlRigChannelEnum::PositionZ, { FFBXControlRigTypeProxyEnum::Position, 2 } },
+		{ FControlRigChannelEnum::RotatorX, { FFBXControlRigTypeProxyEnum::Rotator, 0 } },
+		{ FControlRigChannelEnum::RotatorY, { FFBXControlRigTypeProxyEnum::Rotator, 1 } },
+		{ FControlRigChannelEnum::RotatorZ, { FFBXControlRigTypeProxyEnum::Rotator, 2 } },
+		{ FControlRigChannelEnum::ScaleX, { FFBXControlRigTypeProxyEnum::Scale, 0 } },
+		{ FControlRigChannelEnum::ScaleY, { FFBXControlRigTypeProxyEnum::Scale, 1 } },
+		{ FControlRigChannelEnum::ScaleZ, { FFBXControlRigTypeProxyEnum::Scale, 2 } }
+	}; 
+
+	for (const FControlToTransformMappings& ControlMapping : ExportFBXControlRigSettings->ControlChannelMappings)
+	{
+		FControlRigFbxNodeMapping ChannelMapping;
+		
+		if (ControlMapping.ControlChannel == FControlRigChannelEnum::Bool)
+		{
+			ChannelMapping.ChannelType = FMovieSceneBoolChannel::StaticStruct()->GetFName();
+			ChannelMapping.ControlType = FFBXControlRigTypeProxyEnum::Bool;
+		}
+		else if (ControlMapping.ControlChannel == FControlRigChannelEnum::Enum)
+		{
+			ChannelMapping.ChannelType = FMovieSceneByteChannel::StaticStruct()->GetFName();
+			ChannelMapping.ControlType = FFBXControlRigTypeProxyEnum::Integer;
+		}
+		else if (ControlMapping.ControlChannel == FControlRigChannelEnum::Integer)
+		{
+			ChannelMapping.ChannelType = FMovieSceneIntegerChannel::StaticStruct()->GetFName();
+			ChannelMapping.ControlType = FFBXControlRigTypeProxyEnum::Integer;
+		}
+		else
+		{
+			ChannelMapping.ChannelType = FMovieSceneFloatChannel::StaticStruct()->GetFName();
+
+			// Convert Control Channel Type to Control Type and Transform Component Index (0-2)
+			if (const TPair<FFBXControlRigTypeProxyEnum, uint8>* Mapping = ProxyMapping.Find(ControlMapping.ControlChannel))
+			{
+				ChannelMapping.ControlType = Mapping->Key;
+				ChannelMapping.ChannelAttrIndex = Mapping->Value;
+			}
+			else
+			{
+				ChannelMapping.ControlType = FFBXControlRigTypeProxyEnum::Float;
+			}
+		}
+
+		// Convert FBXChannel Type to Transform Component Index (0-8)
+		ChannelMapping.FbxAttrIndex = (uint8)ControlMapping.FBXChannel;
+		
+		ChannelMapping.bNegate = ControlMapping.bNegate;
+		ChannelsMapping.Add(ChannelMapping);
+	}
+
+	// Filter only selected controls if required
+	TArray<FName> ControlsFilter;
+	if (ExportFBXControlRigSettings->bExportOnlySelectedControls)
+	{
+		ControlsFilter = SelectedControlNames;
+	}
+
+	Exporter->ExportControlRigSection(Section, ChannelsMapping, ControlsFilter, RootToLocalTransform);
+
+	// Save to disk
+	Exporter->WriteToFile(*ExportFBXControlRigSettings->ExportFileName);
+	
+	return true;
+}
+
+bool MovieSceneToolHelpers::ImportFBXIntoControlRigChannelsWithDialog(const TSharedRef<ISequencer>& InSequencer,TArray<FRigControlFBXNodeAndChannels>* NodeAndChannels)
 {
 	TArray<FString> OpenFilenames;
 	IDesktopPlatform* DesktopPlatform = FDesktopPlatformModule::Get();
@@ -2830,16 +2596,14 @@ bool MovieSceneToolHelpers::ImportFBXIntoChannelsWithDialog(const TSharedRef<ISe
 	{
 		return false;
 	}
-
-
+	
 	const FText TitleText = NSLOCTEXT("MovieSceneTools", "ImportFBXTitleOnToControlRig", "Import FBX Onto Control Rig");
 
 	// Create the window to choose our options
 	TSharedRef<SWindow> Window = SNew(SWindow)
 		.Title(TitleText)
 		.HasCloseButton(true)
-		.SizingRule(ESizingRule::UserSized)
-		.ClientSize(FVector2D(400.0f, 200.0f))
+		.SizingRule(ESizingRule::Autosized)
 		.AutoCenter(EAutoCenter::PreferredWorkArea)
 		.SupportsMinimize(false);
 
@@ -2890,6 +2654,52 @@ bool MovieSceneToolHelpers::ImportFBXIntoChannelsWithDialog(const TSharedRef<ISe
 	return true;
 
 }
+
+bool MovieSceneToolHelpers::ExportFBXFromControlRigChannelsWithDialog(const TSharedRef<ISequencer>& InSequencer, UMovieSceneTrack* Track)
+{
+	FString ExportFilename;
+		
+	if (IDesktopPlatform* DesktopPlatform = FDesktopPlatformModule::Get())
+	{
+		TArray<FString> SavedFiles;
+
+		const bool bFilePicked = DesktopPlatform->SaveFileDialog(
+			FSlateApplication::Get().FindBestParentWindowHandleForDialogs(nullptr),
+			NSLOCTEXT("MovieSceneTools", "ExportControlRigFBX", "Export Control Rig FBX" ).ToString(),
+			*(FEditorDirectories::Get().GetLastDirectory(ELastDirectory::FBX)),
+			TEXT(""),
+			TEXT("FBX document|*.fbx"),
+			EFileDialogFlags::None,
+			SavedFiles);
+
+		if (!bFilePicked || SavedFiles.IsEmpty())
+		{
+			return false;
+		}
+		ExportFilename = SavedFiles[0];
+		FEditorDirectories::Get().SetLastDirectory(ELastDirectory::FBX, FPaths::GetPath(ExportFilename)); 
+	}
+
+	const FText TitleText = NSLOCTEXT("MovieSceneTools", "ExportFBXFromControlRigTitle", "Export FBX from Control Rig");
+
+	// Create the window to choose our options
+	const TSharedRef<SWindow> Window = SNew(SWindow)
+		.Title(TitleText)
+		.HasCloseButton(true)
+		.SizingRule(ESizingRule::Autosized)
+		.AutoCenter(EAutoCenter::PreferredWorkArea)
+		.SupportsMinimize(false);
+
+	const TSharedRef<SControlRigExportFBXSettings> DialogWidget = SNew(SControlRigExportFBXSettings, InSequencer, Track)
+		.ExportFilename(ExportFilename);
+
+	Window->SetContent(DialogWidget);
+
+	FSlateApplication::Get().AddWindow(Window);
+
+	return true;
+}
+
 bool ImportFBXTransform(FString NodeName, FGuid ObjectBinding, UnFbx::FFbxCurvesAPI& CurveAPI, UMovieSceneSequence* InSequence)
 {
 	UMovieScene* MovieScene = InSequence->GetMovieScene();
diff --git a/Engine/Source/Editor/MovieSceneTools/Private/MovieSceneToolsUserSettings.cpp b/Engine/Source/Editor/MovieSceneTools/Private/MovieSceneToolsUserSettings.cpp
--- a/Engine/Source/Editor/MovieSceneTools/Private/MovieSceneToolsUserSettings.cpp
+++ b/Engine/Source/Editor/MovieSceneTools/Private/MovieSceneToolsUserSettings.cpp
@@ -42,24 +42,10 @@ UMovieSceneUserImportFBXSettings::UMovieSceneUserImportFBXSettings(const FObject
 }
 
 
-UMovieSceneUserImportFBXControlRigSettings::UMovieSceneUserImportFBXControlRigSettings(const FObjectInitializer& Initializer)
-	: Super(Initializer)
+void SetControlRigIOChannelMappingPresets(TArray<FControlToTransformMappings>& ControlChannelMappings, const bool bMetaHuman)
 {
-	bForceFrontXAxis = false;
-	bConvertSceneUnit = true;
-	ImportUniformScale = 1.0f;
-	bSpecifyTimeRange = false;
-	StartTimeRange = 0;
-	EndTimeRange = 0;
-	TimeToInsertOrReplaceAnimation = 0;
-	bInsertAnimation = true;
-	bImportOntoSelectedControls = false;
-
-	FControlFindReplaceString FirstFindReplace;
-	FirstFindReplace.Find = FString("");
-	FirstFindReplace.Replace = FString("");
-	FindAndReplaceStrings.Add(FirstFindReplace);
-
+	// Since we can't change the API unfortunately need to do this here.
+	ControlChannelMappings.SetNum(0); //clear and reset
 	FControlToTransformMappings Bool;
 	Bool.bNegate = false;
 	Bool.ControlChannel = FControlRigChannelEnum::Bool;
@@ -69,7 +55,14 @@ UMovieSceneUserImportFBXControlRigSettings::UMovieSceneUserImportFBXControlRigSe
 	FControlToTransformMappings Float;
 	Float.bNegate = false;
 	Float.ControlChannel = FControlRigChannelEnum::Float;
-	Float.FBXChannel = FTransformChannelEnum::TranslateX;
+	if (bMetaHuman)
+	{
+		Float.FBXChannel = FTransformChannelEnum::TranslateY;  // use Y for metahuman
+	}
+	else
+	{
+		Float.FBXChannel = FTransformChannelEnum::TranslateX;
+	}
 	ControlChannelMappings.Add(Float);
 
 	FControlToTransformMappings Vector2DX;
@@ -137,8 +130,46 @@ UMovieSceneUserImportFBXControlRigSettings::UMovieSceneUserImportFBXControlRigSe
 	ScaleZ.ControlChannel = FControlRigChannelEnum::ScaleZ;
 	ScaleZ.FBXChannel = FTransformChannelEnum::ScaleZ;
 	ControlChannelMappings.Add(ScaleZ);
+}
+
+
+UMovieSceneUserImportFBXControlRigSettings::UMovieSceneUserImportFBXControlRigSettings(const FObjectInitializer& Initializer)
+	: Super(Initializer)
+{
+	bForceFrontXAxis = false;
+	bConvertSceneUnit = true;
+	ImportUniformScale = 1.0f;
+	bSpecifyTimeRange = false;
+	StartTimeRange = 0;
+	EndTimeRange = 0;
+	TimeToInsertOrReplaceAnimation = 0;
+	bInsertAnimation = true;
+	bImportOntoSelectedControls = false;
+
+	FControlFindReplaceString FirstFindReplace;
+	FirstFindReplace.Find = FString("");
+	FirstFindReplace.Replace = FString("");
+	FindAndReplaceStrings.Add(FirstFindReplace);
 
+	SetControlRigIOChannelMappingPresets(ControlChannelMappings, false);
 }
+
+void UMovieSceneUserImportFBXControlRigSettings::LoadControlMappingsFromPreset(bool bMetaHumanPreset)
+{
+	SetControlRigIOChannelMappingPresets(ControlChannelMappings, bMetaHumanPreset);
+}
+
+UMovieSceneUserExportFBXControlRigSettings::UMovieSceneUserExportFBXControlRigSettings(
+	const FObjectInitializer& Initializer)
+{
+}
+
+void UMovieSceneUserExportFBXControlRigSettings::LoadControlMappingsFromPreset(bool bMetaHumanPreset)
+{
+	SetControlRigIOChannelMappingPresets(ControlChannelMappings, bMetaHumanPreset);
+}
+
+
 /*
 enum class ERigControlType : uint8
 {
diff --git a/Engine/Source/Editor/MovieSceneTools/Public/MovieSceneToolHelpers.h b/Engine/Source/Editor/MovieSceneTools/Public/MovieSceneToolHelpers.h
--- a/Engine/Source/Editor/MovieSceneTools/Public/MovieSceneToolHelpers.h
+++ b/Engine/Source/Editor/MovieSceneTools/Public/MovieSceneToolHelpers.h
@@ -29,6 +29,7 @@ class UMovieSceneTrack;
 struct FMovieSceneEvaluationTrack;
 class UMovieSceneUserImportFBXSettings;
 class UMovieSceneUserImportFBXControlRigSettings;
+class UMovieSceneUserExportFBXControlRigSettings;
 struct FMovieSceneDoubleValue;
 struct FMovieSceneFloatValue;
 class INodeNameAdapter;
@@ -517,15 +518,26 @@ public:
 	static FString GetCameraName(fbxsdk::FbxCamera* InCamera);
 
 	/*
-	 * Import FBX into Channels With Dialog
+	 * Import FBX into Control Rig Channels With Dialog
 	 */
-	static bool ImportFBXIntoChannelsWithDialog(const TSharedRef<ISequencer>& InSequencer, TArray<FFBXNodeAndChannels>* NodeAndChannels);
+	static bool ImportFBXIntoControlRigChannelsWithDialog(const TSharedRef<ISequencer>& InSequencer, TArray<FRigControlFBXNodeAndChannels>* NodeAndChannels);
 
 	/*
-	* Import FBX into Channels
+	** Export FBX from Control Rig Channels With Dialog
+	*/
+	static bool ExportFBXFromControlRigChannelsWithDialog(const TSharedRef<ISequencer>& InSequencer, UMovieSceneTrack* Track);
+
+	/*
+	* Import FBX into Control Rig Channels
 	*/	
 	static bool ImportFBXIntoControlRigChannels(UMovieScene* MovieScene, const FString& ImportFilename,  UMovieSceneUserImportFBXControlRigSettings *ControlRigSettings,
-		TArray<FFBXNodeAndChannels>* NodeAndChannels, const TArray<FName>& SelectedControlNames, FFrameRate FrameRate);
+		TArray<FRigControlFBXNodeAndChannels>* NodeAndChannels, const TArray<FName>& SelectedControlNames, FFrameRate FrameRate);
+
+	/*
+	** Export FBX from Control Rig Channels
+	*/	
+	static bool ExportFBXFromControlRigChannels(const UMovieSceneSection* Section, const UMovieSceneUserExportFBXControlRigSettings* ExportFBXControlRigSettings,
+	                                            const TArray<FName>& SelectedControlNames, const FMovieSceneSequenceTransform& RootToLocalTransform);
 
 	/*
 	* Acquire first SkeletalMeshComponent from the Object
diff --git a/Engine/Source/Editor/MovieSceneTools/Public/MovieSceneToolsUserSettings.h b/Engine/Source/Editor/MovieSceneTools/Public/MovieSceneToolsUserSettings.h
--- a/Engine/Source/Editor/MovieSceneTools/Public/MovieSceneToolsUserSettings.h
+++ b/Engine/Source/Editor/MovieSceneTools/Public/MovieSceneToolsUserSettings.h
@@ -8,6 +8,8 @@
 #include "Misc/FrameNumber.h"
 #include "MovieSceneToolHelpers.h"
 #include "SceneTypes.h"
+#include "Exporters/FbxExportOption.h"
+
 #include "MovieSceneToolsUserSettings.generated.h"
 
 UENUM()
@@ -217,6 +219,10 @@ public:
 	UPROPERTY(EditAnywhere, config, Category = "String Matching Options", meta = (ToolTip = "Strings In Imported Node To Find And Replace"))
 	TArray< FControlFindReplaceString> FindAndReplaceStrings;
 
+	/** Strip namespaces from FBX node names */
+	UPROPERTY(EditAnywhere, config, Category = "Import Options", meta = (ToolTip = "Will strip any namespace from the FBX node names"))
+	bool bStripNamespace = true;
+	
 	/** Whether to force the front axis to be align with X instead of -Y. */
 	UPROPERTY(EditAnywhere, config, Category = "Import Options", meta = (ToolTip = "Convert the scene from FBX coordinate system to UE coordinate system with front X axis instead of -Y"))
 	bool bForceFrontXAxis;
@@ -256,4 +262,50 @@ public:
 	/** Mappings for how Control Rig Control Attributes Map to the incoming Transforms*/
 	UPROPERTY(Config, EditAnywhere, Category = "Control Attribute Mappings")
 	TArray<FControlToTransformMappings> ControlChannelMappings;
-};
\ No newline at end of file
+
+public:
+	/** Load the default or metahuman preset into the current mappings */
+	UFUNCTION(BlueprintCallable, Category = "Control Rig")
+	void LoadControlMappingsFromPreset(bool bMetaHumanPreset);
+};
+
+UCLASS(BlueprintType, config = EditorSettings)
+class MOVIESCENETOOLS_API UMovieSceneUserExportFBXControlRigSettings : public UObject
+{
+public:
+	UMovieSceneUserExportFBXControlRigSettings(const FObjectInitializer& Initializer);
+
+	GENERATED_BODY()
+	/** Imported File Name */
+	UPROPERTY(EditAnywhere, Category = "Imported File Information")
+	FString ExportFileName;
+	
+	/** This will set the fbx sdk compatibility when exporting to fbx file. The default value is 2013 */
+	UPROPERTY(EditAnywhere, BlueprintReadWrite, config, Category = Exporter)
+	EFbxExportCompatibility FbxExportCompatibility = EFbxExportCompatibility::FBX_2018;
+
+	/** If enabled, save as ascii instead of binary */
+	UPROPERTY(EditAnywhere, BlueprintReadWrite, AdvancedDisplay, config, Category = Exporter)
+	uint32 bASCII : 1;
+
+	/** Whether to force the front axis to be align with X instead of -Y. */
+	UPROPERTY(EditAnywhere, config, Category = Exporter, meta = (ToolTip = "Convert the scene from FBX coordinate system to UE coordinate system with front X axis instead of -Y"))
+	bool bForceFrontXAxis = false;
+
+	/** Whether or not import onto selected controls or all controls*/
+	UPROPERTY(EditAnywhere, config, Category = "Control Rig")
+	bool bExportOnlySelectedControls = false;
+
+	/** Mappings for how Control Rig Control Attributes Map to the incoming Transforms*/
+	UPROPERTY(Config, EditAnywhere, Category = "Control Rig")
+	TArray<FControlToTransformMappings> ControlChannelMappings;
+
+	/** If enabled, export sequencer animation in its local time, relative to its sequence. */
+	UPROPERTY(EditAnywhere, BlueprintReadWrite, config, category = Animation)
+	uint32 bExportLocalTime : 1;
+
+public:
+	/** Load the default or metahuman preset into the current mappings */
+	UFUNCTION(BlueprintCallable, Category = "Control Rig")
+	void LoadControlMappingsFromPreset(bool bMetaHumanPreset);
+};
diff --git a/Engine/Source/Editor/UnrealEd/Private/Fbx/FbxCurveDataImport.cpp b/Engine/Source/Editor/UnrealEd/Private/Fbx/FbxCurveDataImport.cpp
--- a/Engine/Source/Editor/UnrealEd/Private/Fbx/FbxCurveDataImport.cpp
+++ b/Engine/Source/Editor/UnrealEd/Private/Fbx/FbxCurveDataImport.cpp
@@ -529,7 +529,41 @@ namespace UnFbx {
 		}
 	}
 
+	void FFbxCurvesAPI::GetConvertedNonTransformCurveData(const FString& NodeName, bool bUseSequencerCurve, float UniformScale, TMap<FName, FRichCurve>& OutCurves)
+	{
+		for (TPair<uint64, FFbxAnimNodeHandle> AnimNodeKvp : CurvesData)
+		{
+			const FFbxAnimNodeHandle& AnimNodeHandle = AnimNodeKvp.Value;
+			if (AnimNodeHandle.Name.Compare(NodeName))
+			{
+				continue;
+			}
+			for (auto NodePropertyKvp : AnimNodeHandle.NodeProperties)
+			{
+				FFbxAnimPropertyHandle& AnimPropertyHandle = NodePropertyKvp.Value;
+				for (FFbxAnimCurveHandle& CurveHandle : AnimPropertyHandle.CurveHandles)
+				{
+					if (CurveHandle.CurveType != FFbxAnimCurveHandle::NotTransform)
+					{
+						continue;
+					}
+					FRichCurve RichCurve;
 
+					if (bUseSequencerCurve)
+					{
+						GetCurveDataForSequencer(CurveHandle, RichCurve, false, UniformScale);
+					}
+					else
+					{
+						GetCurveData(CurveHandle, RichCurve, false, UniformScale);
+					}
+
+					OutCurves.Add(*AnimPropertyHandle.Name, RichCurve);
+				}
+			}
+		}
+	}
+	
 	void FFbxCurvesAPI::GetConvertedTransformCurveData(const FString& NodeName, FRichCurve& TranslationX, FRichCurve& TranslationY, FRichCurve& TranslationZ,
 		FRichCurve& EulerRotationX, FRichCurve& EulerRotationY, FRichCurve& EulerRotationZ,
 		FRichCurve& ScaleX, FRichCurve& ScaleY, FRichCurve& ScaleZ,
diff --git a/Engine/Source/Editor/UnrealEd/Private/Fbx/FbxMainExport.cpp b/Engine/Source/Editor/UnrealEd/Private/Fbx/FbxMainExport.cpp
--- a/Engine/Source/Editor/UnrealEd/Private/Fbx/FbxMainExport.cpp
+++ b/Engine/Source/Editor/UnrealEd/Private/Fbx/FbxMainExport.cpp
@@ -19,10 +19,13 @@
 #include "Components/StaticMeshComponent.h"
 #include "Components/LightComponent.h"
 #include "Model.h"
+
 #include "Channels/MovieSceneDoubleChannel.h"
 #include "Channels/MovieSceneFloatChannel.h"
 #include "Channels/MovieSceneIntegerChannel.h"
 #include "Channels/MovieSceneStringChannel.h"
+#include "Channels/MovieSceneByteChannel.h"
+
 #include "Animation/AnimTypes.h"
 #include "Animation/AnimSequenceBase.h"
 #include "Animation/AnimSequence.h"
@@ -63,6 +66,10 @@
 #include "StaticMeshResources.h"
 
 #include "FbxExporter.h"
+#include "Exporters/FbxExportOption.h"
+#include "FbxExportOptionsWindow.h"
+#include "FbxAnimUtils.h"
+#include "INodeAndChannelMappings.h"
 
 #include "StaticMeshAttributes.h"
 #include "StaticMeshOperations.h"
@@ -96,10 +103,9 @@
 #include "Chaos/Particles.h"
 #include "Chaos/Plane.h"
 #include "ChaosCheck.h"
+#include "Animation/AnimationSettings.h"
 #include "Chaos/Convex.h"
 
-#include "Exporters/FbxExportOption.h"
-#include "FbxExportOptionsWindow.h"
 #include "Widgets/SWindow.h"
 #include "Framework/Application/SlateApplication.h"
 #include "Interfaces/IMainFrameModule.h"
@@ -2003,6 +2009,398 @@ bool FFbxExporter::ExportLevelSequence(UMovieScene* MovieScene, const TArray<FGu
 	return true;
 }
 
+void FFbxExporter::AddTimecodeAttributesAndSetKey(const UMovieSceneSection* InSection, FbxNode* InFbxNode, const FMovieSceneSequenceTransform& RootToLocalTransform)
+{
+	FbxAnimLayer* BaseLayer = AnimStack->GetMember<FbxAnimLayer>(0);
+	const FFrameRate TickResolution = InSection->GetTypedOuter<UMovieScene>()->GetTickResolution();
+	
+	FName TCHourAttrName(TEXT("TCHour"));
+	FName TCMinuteAttrName(TEXT("TCMinute"));
+	FName TCSecondAttrName(TEXT("TCSecond"));
+	FName TCFrameAttrName(TEXT("TCFrame"));
+
+	if (const UAnimationSettings* AnimationSettings = UAnimationSettings::Get())
+	{
+		TCHourAttrName = AnimationSettings->BoneTimecodeCustomAttributeNameSettings.HourAttributeName;
+		TCMinuteAttrName = AnimationSettings->BoneTimecodeCustomAttributeNameSettings.MinuteAttributeName;
+		TCSecondAttrName = AnimationSettings->BoneTimecodeCustomAttributeNameSettings.SecondAttributeName;
+		TCFrameAttrName = AnimationSettings->BoneTimecodeCustomAttributeNameSettings.FrameAttributeName;
+	}
+
+	const TArray<FString> TimecodePropertyNames = {
+		TCHourAttrName.ToString(),
+		TCMinuteAttrName.ToString(),
+		TCSecondAttrName.ToString(),
+		TCFrameAttrName.ToString()
+	};
+
+	const TArray<int32> TimecodeValues = {
+		InSection->TimecodeSource.Timecode.Hours,
+		InSection->TimecodeSource.Timecode.Minutes,
+		InSection->TimecodeSource.Timecode.Seconds,
+		InSection->TimecodeSource.Timecode.Frames
+	};
+
+	const FFrameNumber KeyTime = InSection->GetTypedOuter<UMovieScene>()->GetPlaybackRange().GetLowerBoundValue();
+		
+	for (int i = 0; i < TimecodePropertyNames.Num(); i++)
+	{
+		const FString PropertyName = TimecodePropertyNames[i];
+		CreateAnimatableUserProperty(InFbxNode, 0.0, StringCast<char>(*PropertyName).Get(), StringCast<char>(*PropertyName).Get(), FbxFloatDT);
+		FbxProperty Property = InFbxNode->FindProperty(StringCast<char>(*PropertyName).Get());
+
+		if (Property.IsValid())
+		{
+			if (FbxAnimCurve* AnimCurve = Property.GetCurve(BaseLayer, nullptr, true))
+			{
+				const float KeyValue = (float)TimecodeValues[i];
+					
+				AnimCurve->KeyModifyBegin();
+					
+				FbxTime FbxTime;
+				const double KeyTimeSeconds = GetExportOptions()->bExportLocalTime ? KeyTime / TickResolution : (KeyTime * RootToLocalTransform.InverseLinearOnly()) / TickResolution;
+
+				FbxTime.SetSecondDouble(KeyTimeSeconds);
+
+				const int FbxKeyIndex = AnimCurve->KeyAdd(FbxTime);
+				AnimCurve->KeySet(FbxKeyIndex, FbxTime, KeyValue, FbxAnimCurveDef::eInterpolationConstant);
+				AnimCurve->KeySetConstantMode(FbxKeyIndex, FbxAnimCurveDef::EConstantMode::eConstantStandard);
+
+				AnimCurve->KeyModifyEnd();
+			}
+		}
+	}
+}
+
+bool FFbxExporter::ExportControlRigSection(const UMovieSceneSection* Section, const TArray<FControlRigFbxNodeMapping>& ChannelsMapping,
+	const TArray<FName>& FilterControls, const FMovieSceneSequenceTransform& RootToLocalTransform)
+{
+	if (!Section)
+	{
+		return false;
+	}
+	
+	UMovieSceneTrack* Track = Section->GetTypedOuter<UMovieSceneTrack>();
+
+	FbxAnimLayer* BaseLayer = AnimStack->GetMember<FbxAnimLayer>(0);
+
+	const FFrameRate TickResolution = Track->GetTypedOuter<UMovieScene>()->GetTickResolution();
+
+	// Helpers to convert channel/property names to indices and vice versa
+	const TArray<FString> TransformComponentsArray = {"X", "Y", "Z"};
+	const TArray<FString> TransformPropertyArray = {"Location", "Rotation", "Scale"};
+
+	// Cached transform channels for a single property to bake and export the 3 channels in one go
+	FMovieSceneFloatChannel* CachedTransformChannels[3] = {nullptr, nullptr, nullptr};
+	// Cached transform property index keep track of which property we are currently caching to bake & export Transforms & TransformNoScale
+	int CurrentTmPropertyIndex = 0;
+	
+	// The control rig parent group node
+	FbxNode* RootFbxNode = CreateNode(Track->GetDisplayName().ToString());
+
+	// Export timecode
+	AddTimecodeAttributesAndSetKey(Section, RootFbxNode, RootToLocalTransform);
+	
+	const FName BoolChannelTypeName = FMovieSceneBoolChannel::StaticStruct()->GetFName();
+	const FName ByteChannelTypeName = FMovieSceneByteChannel::StaticStruct()->GetFName();
+	const FName DoubleChannelTypeName = FMovieSceneDoubleChannel::StaticStruct()->GetFName();
+	const FName FloatChannelTypeName = FMovieSceneFloatChannel::StaticStruct()->GetFName();
+	const FName IntegerChannelTypeName = FMovieSceneIntegerChannel::StaticStruct()->GetFName();
+	
+	FMovieSceneChannelProxy& ChannelProxy = Section->GetChannelProxy();
+	for (const FMovieSceneChannelEntry& Entry : ChannelProxy.GetAllEntries())
+	{
+		TArrayView<FMovieSceneChannel* const> Channels = Entry.GetChannels();
+		TArrayView<const FMovieSceneChannelMetaData> AllMetaData = Entry.GetMetaData();
+
+		for (int32 Index = 0; Index < Channels.Num(); ++Index)
+		{
+			const FName ChannelTypeName = Entry.GetChannelTypeName();
+			FMovieSceneChannelHandle ChannelHandle = ChannelProxy.MakeHandle(ChannelTypeName, Index);
+
+			FMovieSceneBoolChannel* BoolChannel = ChannelTypeName == BoolChannelTypeName ? ChannelHandle.Cast<FMovieSceneBoolChannel>().Get() : nullptr;
+			FMovieSceneByteChannel* ByteChannel = ChannelTypeName == ByteChannelTypeName ? ChannelHandle.Cast<FMovieSceneByteChannel>().Get() : nullptr;
+			FMovieSceneDoubleChannel* DoubleChannel = ChannelTypeName == DoubleChannelTypeName ? ChannelHandle.Cast<FMovieSceneDoubleChannel>().Get() : nullptr;
+			FMovieSceneFloatChannel* FloatChannel = ChannelTypeName == FloatChannelTypeName ? ChannelHandle.Cast<FMovieSceneFloatChannel>().Get() : nullptr;
+			FMovieSceneIntegerChannel* IntegerChannel = ChannelTypeName == IntegerChannelTypeName ? ChannelHandle.Cast<FMovieSceneIntegerChannel>().Get() : nullptr;
+
+			if (!BoolChannel && !ByteChannel && !DoubleChannel && !FloatChannel && !IntegerChannel)
+			{
+				continue;
+			}
+
+			// Find the node and attribute names
+			const FMovieSceneChannelMetaData& MetaData = AllMetaData[Index];
+			const FString ChannelName = MetaData.Name.ToString();
+
+			FControlRigFbxCurveData FbxCurveData;
+			if (INodeAndChannelMappings* ChannelMapping = Cast<INodeAndChannelMappings>(Track))
+			{
+				if (!ChannelMapping->GetFbxCurveDataFromChannelMetadata(MetaData, FbxCurveData))
+				{
+					continue;
+				}
+			}
+
+			// Skip any control not in the filter if the filter isn't empty
+			if (!FilterControls.IsEmpty() && !FilterControls.Contains(FbxCurveData.ControlName))
+			{
+				continue;
+			}
+
+			// Find or create the node
+			FbxNode* ControlFbxNode = RootFbxNode->FindChild(StringCast<char>(*FbxCurveData.NodeName).Get());
+			if (ControlFbxNode == nullptr)
+			{
+				ControlFbxNode = FbxNode::Create(Scene, StringCast<char>(*FbxCurveData.NodeName).Get());
+				RootFbxNode->AddChild(ControlFbxNode);
+			}
+
+			// Remap channel, optionally baking & exporting it for transforms & rotations 
+			bool bNegate = false;
+			int TmComponentIndex = TransformComponentsArray.Find(FbxCurveData.AttributeName);
+
+			// Find the attribute to remap onto and whether to negate the channel if required
+			const FControlRigFbxNodeMapping* NodeMapping = ChannelsMapping.FindByPredicate(
+				[ChannelTypeName, FbxCurveData, TmComponentIndex](const FControlRigFbxNodeMapping& AMapping)
+			{
+				// If the control is treated as an attribute of another control, we don't need to check the attribute to remap
+				const bool bRemapControl = AMapping.ChannelType == ChannelTypeName && AMapping.ControlType == FbxCurveData.ControlType;
+				return bRemapControl && (!FbxCurveData.IsControlNode() || AMapping.ChannelAttrIndex == TmComponentIndex);
+			});
+			
+			if (TmComponentIndex != INDEX_NONE)
+			{
+				// Retrieve property name for Vector2D, Position, Scale
+				if (FbxCurveData.ControlType == FFBXControlRigTypeProxyEnum::Vector2D || FbxCurveData.ControlType == FFBXControlRigTypeProxyEnum::Position)
+				{
+					FbxCurveData.AttributePropertyName = "Location";
+				}
+				if (FbxCurveData.ControlType == FFBXControlRigTypeProxyEnum::Scale)
+				{
+					FbxCurveData.AttributePropertyName = "Scale";
+				}
+				
+				// Export converted & baked for Rotations & Transforms
+				if ((uint8)FbxCurveData.ControlType >= (uint8)FFBXControlRigTypeProxyEnum::Rotator)
+				{
+					const bool bExportCachedTransforms = TmComponentIndex == 2;
+					const bool bRotator = FbxCurveData.ControlType == FFBXControlRigTypeProxyEnum::Rotator;
+					
+					int TmPropertyIndex = bRotator ? 1 : CurrentTmPropertyIndex;
+
+					// todo: am: does not handle remapping to another property (i.e. rotation to location)
+					// Remap transform channel
+					if (NodeMapping)
+					{
+						TmPropertyIndex = NodeMapping->FbxAttrIndex / 3;
+						TmComponentIndex = NodeMapping->FbxAttrIndex % 3;
+						bNegate = NodeMapping->bNegate;
+					}
+
+					// Cache a single transform property channel
+					CachedTransformChannels[TmComponentIndex] = FloatChannel;
+
+					// Once all 3 channels in the current property have been cached, export baked
+					if (bExportCachedTransforms)
+					{
+						for (int i = 0; i < 3; i++)
+						{
+							if (CachedTransformChannels[i] && CachedTransformChannels[i]->GetNumKeys() > 0)
+							{
+								// Only export baked if at least one curve has keys
+								ExportTransformChannelsToFbxCurve(ControlFbxNode, CachedTransformChannels[0],
+									CachedTransformChannels[1], CachedTransformChannels[2],
+									TmPropertyIndex, Track, RootToLocalTransform);
+								
+								break;
+							}
+						}
+
+						// Increase or reset CurrentTmPropertyIndex to keep track of which property in a Transform we are currently caching
+						const bool bNoScale = FbxCurveData.ControlType == FFBXControlRigTypeProxyEnum::TransformNoScale;
+						const bool bTransformProcessed = bRotator || CurrentTmPropertyIndex == 2 || (bNoScale && CurrentTmPropertyIndex == 1);
+						CurrentTmPropertyIndex = bTransformProcessed ? 0 : CurrentTmPropertyIndex + 1;
+					}
+
+					continue;
+				}
+			}
+			// Remap non channels if required (apart from rotations & transforms, remapped separately)
+			if (NodeMapping)
+			{
+				FbxCurveData.AttributePropertyName = TransformPropertyArray[NodeMapping->FbxAttrIndex / 3];
+				FbxCurveData.AttributeName = TransformComponentsArray[NodeMapping->FbxAttrIndex % 3];
+				bNegate = NodeMapping->bNegate;
+			}
+			
+			FbxProperty Property;
+			
+			// Use AttributeName (i.e. Weight) as the property if no AttributePropertyName (i.e. Scale). Use NodeName if AttributeName was also empty...
+			const FString FbxPropertyName = FbxCurveData.AttributePropertyName.IsEmpty() ? (!FbxCurveData.AttributeName.IsEmpty() ? FbxCurveData.AttributeName : FbxCurveData.NodeName) : FbxCurveData.AttributePropertyName;
+
+			// Try and find the existing property
+			if (FbxCurveData.AttributePropertyName == "Location")
+			{
+				Property = ControlFbxNode->LclTranslation;
+			}
+			else if (FbxCurveData.AttributePropertyName == "Rotation")
+			{
+				Property = ControlFbxNode->LclRotation;
+			}
+			else if (FbxCurveData.AttributePropertyName == "Scale")
+			{
+				Property = ControlFbxNode->LclScaling;
+			}
+			else
+			{
+				Property = ControlFbxNode->FindProperty(StringCast<char>(*FbxPropertyName).Get());
+			}
+
+			// Or create it otherwise
+			if (!Property.IsValid())
+			{
+				if (DoubleChannel)
+				{
+					double Default = DoubleChannel->GetDefault().Get(0);
+					CreateAnimatableUserProperty(ControlFbxNode, Default, StringCast<char>(*FbxPropertyName).Get(), StringCast<char>(*FbxPropertyName).Get(), FbxDoubleDT);
+				}
+				else if (FloatChannel)
+				{
+					float Default = FloatChannel->GetDefault().Get(0);
+					CreateAnimatableUserProperty(ControlFbxNode, Default, StringCast<char>(*FbxPropertyName).Get(), StringCast<char>(*FbxPropertyName).Get(), FbxFloatDT);
+				}
+				else if (IntegerChannel)
+				{
+					int32 Default = IntegerChannel->GetDefault().Get(0);
+					CreateAnimatableUserProperty(ControlFbxNode, Default, StringCast<char>(*FbxPropertyName).Get(), StringCast<char>(*FbxPropertyName).Get(), FbxIntDT);
+				}
+				else if (BoolChannel)
+				{
+					bool Default = BoolChannel->GetDefault().Get(false);
+					CreateAnimatableUserProperty(ControlFbxNode, Default, StringCast<char>(*FbxPropertyName).Get(), StringCast<char>(*FbxPropertyName).Get(), FbxBoolDT);
+				}
+				else if (ByteChannel)
+				{
+					uint8 Default = ByteChannel->GetDefault().Get(0);
+					CreateAnimatableUserProperty(ControlFbxNode, Default, StringCast<char>(*FbxPropertyName).Get(), StringCast<char>(*FbxPropertyName).Get(), FbxEnumDT);
+				}
+
+				Property = ControlFbxNode->FindProperty(StringCast<char>(*FbxPropertyName).Get());
+			}
+
+			if (!Property.IsValid())
+			{
+				continue;
+			}
+
+			// If AttributePropertyName was found (i.e Scale), use AttributeName as the property channel (i.e. X) 
+			const FString FbxCurveChannel = FbxCurveData.AttributePropertyName.IsEmpty() ? "" : FbxCurveData.AttributeName;
+
+			// Create the anim curve, optionally of the given name if the attribute had a component 
+			FbxAnimCurve* AnimCurve = Property.GetCurve(BaseLayer, StringCast<char>(*FbxCurveChannel).Get(), true);
+			if (!AnimCurve)
+			{
+				continue;
+			}
+
+			// Export the curve
+			if (BoolChannel)
+			{
+				ExportChannelToFbxCurve(*AnimCurve, *BoolChannel, TickResolution, RootToLocalTransform);
+			}
+			else if (ByteChannel)
+			{
+				ExportChannelToFbxCurve(*AnimCurve, *ByteChannel, TickResolution, RootToLocalTransform);
+			}
+			else if (DoubleChannel)
+			{
+				ExportChannelToFbxCurve(*AnimCurve, *DoubleChannel, TickResolution, ERichCurveValueMode::Default, bNegate, RootToLocalTransform);
+			}
+			else if (FloatChannel)
+			{
+				ExportChannelToFbxCurve(*AnimCurve, *FloatChannel, TickResolution, ERichCurveValueMode::Default, bNegate, RootToLocalTransform);
+			}
+			else if (IntegerChannel)
+			{
+				ExportChannelToFbxCurve(*AnimCurve, *IntegerChannel, TickResolution, RootToLocalTransform);
+			}
+		}
+	}
+	return true;
+}
+
+void FFbxExporter::ExportTransformChannelsToFbxCurve(FbxNode* InFbxNode, FMovieSceneFloatChannel* ChannelX, FMovieSceneFloatChannel* ChannelY, FMovieSceneFloatChannel* ChannelZ, int TmPropertyIndex, const UMovieSceneTrack* Track, const FMovieSceneSequenceTransform& RootToLocalTransform)
+{
+	FbxAnimLayer* BaseLayer = AnimStack->GetMember<FbxAnimLayer>(0);
+
+	FbxProperty Property = TmPropertyIndex == 0 ? InFbxNode->LclTranslation : TmPropertyIndex == 1 ? InFbxNode->LclRotation : InFbxNode->LclScaling;
+	
+	FbxAnimCurve* FbxCurveX = Property.GetCurve(BaseLayer, FBXSDK_CURVENODE_COMPONENT_X, true);
+	FbxAnimCurve* FbxCurveY = Property.GetCurve(BaseLayer, FBXSDK_CURVENODE_COMPONENT_Y, true);
+	FbxAnimCurve* FbxCurveZ = Property.GetCurve(BaseLayer, FBXSDK_CURVENODE_COMPONENT_Z, true);
+
+	FTransform RotationDirectionConvert;
+
+	FbxCurveX->KeyModifyBegin();
+	FbxCurveY->KeyModifyBegin();
+	FbxCurveZ->KeyModifyBegin();
+
+	FFrameRate TickResolution = Track->GetTypedOuter<UMovieScene>()->GetTickResolution();
+	FFrameRate DisplayRate = Track->GetTypedOuter<UMovieScene>()->GetDisplayRate();
+	TRange<FFrameNumber> PlaybackRange = Track->GetTypedOuter<UMovieScene>()->GetPlaybackRange();
+	
+	int32 LocalStartFrame = FFrameRate::TransformTime(FFrameTime(UE::MovieScene::DiscreteInclusiveLower(PlaybackRange)), TickResolution, DisplayRate).RoundToFrame().Value;
+	int32 StartFrame = FFrameRate::TransformTime(FFrameTime(UE::MovieScene::DiscreteInclusiveLower(PlaybackRange) * RootToLocalTransform.InverseLinearOnly()), TickResolution, DisplayRate).RoundToFrame().Value;
+	int32 AnimationLength = FFrameRate::TransformTime(FFrameTime(FFrameNumber(UE::MovieScene::DiscreteSize(PlaybackRange))), TickResolution, DisplayRate).RoundToFrame().Value;
+
+	for (int32 FrameCount = 0; FrameCount <= AnimationLength; ++FrameCount)
+	{
+		int32 LocalFrame = LocalStartFrame + FrameCount;
+
+		FFrameTime LocalTime = FFrameRate::TransformTime(FFrameTime(LocalFrame), DisplayRate, TickResolution);
+
+		FVector3f Vec = FVector3f::ZeroVector;
+		if (ChannelX)
+		{
+			ChannelX->Evaluate(LocalTime, Vec.X);
+		}
+		if (ChannelY)
+		{
+			ChannelY->Evaluate(LocalTime, Vec.Y);
+		}
+		if (ChannelZ)
+		{
+			ChannelZ->Evaluate(LocalTime, Vec.Z);
+		}
+
+		FbxVector4 KeyVec;
+
+		if (TmPropertyIndex == 0)
+		{
+			KeyVec = Converter.ConvertToFbxPos((FVector)Vec);
+		}
+		else if (TmPropertyIndex == 1)
+		{
+			KeyVec = Converter.ConvertToFbxRot((FVector)Vec);
+		}
+		else
+		{
+			KeyVec = Converter.ConvertToFbxScale((FVector)Vec);
+		}
+
+		FbxTime FbxTime;
+		FbxTime.SetSecondDouble(GetExportOptions()->bExportLocalTime ? DisplayRate.AsSeconds(LocalFrame) : DisplayRate.AsSeconds(StartFrame + FrameCount));
+
+		FbxCurveX->KeySet(FbxCurveX->KeyAdd(FbxTime), FbxTime, KeyVec[0]);
+		FbxCurveY->KeySet(FbxCurveY->KeyAdd(FbxTime), FbxTime, KeyVec[1]);
+		FbxCurveZ->KeySet(FbxCurveZ->KeyAdd(FbxTime), FbxTime, KeyVec[2]);
+	}
+
+	FbxCurveX->KeyModifyEnd();
+	FbxCurveY->KeyModifyEnd();
+	FbxCurveZ->KeyModifyEnd();
+}
 
 /**
  * Exports a scene node with the placement indicated by a given actor.
@@ -2421,19 +2819,34 @@ void FFbxExporter::ExportChannelToFbxCurve(FbxAnimCurve& InFbxCurve, const FMovi
 }
 
 void FFbxExporter::ExportChannelToFbxCurve(FbxAnimCurve& InFbxCurve, const FMovieSceneIntegerChannel& InChannel, FFrameRate TickResolution, const FMovieSceneSequenceTransform& RootToLocalTransform)
+{
+	ExportConstantChannelToFbxCurve<FMovieSceneIntegerChannel, int32>(InFbxCurve, InChannel, TickResolution, RootToLocalTransform);
+}
+
+void FFbxExporter::ExportChannelToFbxCurve(FbxAnimCurve& InFbxCurve, const FMovieSceneBoolChannel& InChannel, FFrameRate TickResolution, const FMovieSceneSequenceTransform& RootToLocalTransform)
+{
+	ExportConstantChannelToFbxCurve<FMovieSceneBoolChannel, bool>(InFbxCurve, InChannel, TickResolution, RootToLocalTransform);
+}
+
+void FFbxExporter::ExportChannelToFbxCurve(FbxAnimCurve& InFbxCurve, const FMovieSceneByteChannel& InChannel, FFrameRate TickResolution, const FMovieSceneSequenceTransform& RootToLocalTransform)
+{
+	ExportConstantChannelToFbxCurve<FMovieSceneByteChannel, uint8>(InFbxCurve, InChannel, TickResolution, RootToLocalTransform);
+}
+
+template<typename ChannelType, typename T>
+void FFbxExporter::ExportConstantChannelToFbxCurve(FbxAnimCurve& InFbxCurve, const ChannelType& InChannel, FFrameRate TickResolution, const FMovieSceneSequenceTransform& RootToLocalTransform)
 {
 	InFbxCurve.KeyModifyBegin();
 
-	TArrayView<const FFrameNumber> Times  = InChannel.GetTimes();
-	TArrayView<const int32> Values = InChannel.GetValues();
+	const TArrayView<const FFrameNumber> Times  = InChannel.GetTimes();
+	const TArrayView<const T> Values = InChannel.GetValues();
 
 	for (int32 Index = 0; Index < Times.Num(); ++Index)
 	{
 		const FFrameNumber KeyTime = Times[Index];
-		int32 KeyValue = Values[Index];
+		const T KeyValue = Values[Index];
 
 		FbxTime FbxTime;
-		FbxAnimCurveKey FbxKey;
 		const double KeyTimeSeconds = GetExportOptions()->bExportLocalTime ? KeyTime / TickResolution : (KeyTime * RootToLocalTransform.InverseLinearOnly()) / TickResolution;
 
 		FbxTime.SetSecondDouble(KeyTimeSeconds);
@@ -2441,6 +2854,7 @@ void FFbxExporter::ExportChannelToFbxCurve(FbxAnimCurve& InFbxCurve, const FMovi
 		const int FbxKeyIndex = InFbxCurve.KeyAdd(FbxTime);
 
 		InFbxCurve.KeySet(FbxKeyIndex, FbxTime, KeyValue, FbxAnimCurveDef::eInterpolationConstant);
+		InFbxCurve.KeySetConstantMode(FbxKeyIndex, FbxAnimCurveDef::EConstantMode::eConstantStandard);
 	}
 	InFbxCurve.KeyModifyEnd();
 }
diff --git a/Engine/Source/Editor/UnrealEd/Private/FbxExporter.h b/Engine/Source/Editor/UnrealEd/Private/FbxExporter.h
--- a/Engine/Source/Editor/UnrealEd/Private/FbxExporter.h
+++ b/Engine/Source/Editor/UnrealEd/Private/FbxExporter.h
@@ -11,6 +11,7 @@
 #include "Evaluation/MovieSceneSequenceTransform.h"
 #include "MovieSceneFwd.h"
 #include "FbxImporter.h"
+#include "INodeAndChannelMappings.h"
 #include "UObject/GCObject.h"
 #include "Animation/AnimTypes.h"
 
@@ -165,6 +166,12 @@ public:
 	 */
 	bool ExportLevelSequence(UMovieScene* MovieScene, const TArray<FGuid>& InBindings, IMovieScenePlayer* MovieScenePlayer, INodeNameAdapter& NodeNameAdapter, FMovieSceneSequenceIDRef SequenceID, const FMovieSceneSequenceTransform& RootToLocalTransform);
 
+	/** Add timecode attributes to the given fbx node and add a key at the sequence playback start, using the timecode of the source section */
+	void AddTimecodeAttributesAndSetKey(const UMovieSceneSection* InSection, FbxNode* InFbxNode, const FMovieSceneSequenceTransform& RootToLocalTransform);
+
+	/** Export an FBX from the given control rig section. Optionally remapping and filtering controls. */
+	bool ExportControlRigSection(const UMovieSceneSection* Section, const TArray<FControlRigFbxNodeMapping>& ChannelsMapping, const TArray<FName>& FilterControls, const FMovieSceneSequenceTransform& RootToLocalTransform);
+	
 	/**
 	 * Exports the given level sequence track information into a FBX document.
 	 *
@@ -457,6 +464,15 @@ private:
 	/** Exports a movie scene integer channel to an fbx animation curve. */
 	void ExportChannelToFbxCurve(FbxAnimCurve& InFbxCurve, const FMovieSceneIntegerChannel& InChannel, FFrameRate TickResolution, const FMovieSceneSequenceTransform& RootToLocalTransform = FMovieSceneSequenceTransform());
 
+	void ExportChannelToFbxCurve(FbxAnimCurve& InFbxCurve, const FMovieSceneBoolChannel& InChannel, FFrameRate TickResolution, const FMovieSceneSequenceTransform& RootToLocalTransform);
+	
+	void ExportChannelToFbxCurve(FbxAnimCurve& InFbxCurve, const FMovieSceneByteChannel& InChannel, FFrameRate TickResolution, const FMovieSceneSequenceTransform& RootToLocalTransform);
+
+	template <class ChannelType, typename T>
+	void ExportConstantChannelToFbxCurve(FbxAnimCurve& InFbxCurve, const ChannelType& InChannel, FFrameRate TickResolution, const FMovieSceneSequenceTransform& RootToLocalTransform);
+
+	void ExportTransformChannelsToFbxCurve(FbxNode* InFbxNode, FMovieSceneFloatChannel* ChannelX, FMovieSceneFloatChannel* ChannelY, FMovieSceneFloatChannel* ChannelZ, int TmPropertyIndex, const UMovieSceneTrack* Track, const FMovieSceneSequenceTransform& RootToLocalTransform);
+
 	/**
 	 * Finds the given actor in the already-exported list of structures
 	 * @return FbxNode* the FBX node created from the UE actor
diff --git a/Engine/Source/Editor/UnrealEd/Public/FbxImporter.h b/Engine/Source/Editor/UnrealEd/Public/FbxImporter.h
--- a/Engine/Source/Editor/UnrealEd/Public/FbxImporter.h
+++ b/Engine/Source/Editor/UnrealEd/Public/FbxImporter.h
@@ -375,6 +375,8 @@ public:
 
 	UNREALED_API void GetBakeCurveData(const FFbxAnimCurveHandle &CurveHandle, TArray<float>& CurveData, float PeriodTime, float StartTime = 0.0f, float StopTime = -1.0f, bool bNegative = false, float Scale = 1.0f) const;
 
+	UNREALED_API void GetConvertedNonTransformCurveData(const FString& NodeName, bool bUseSequencerCurve, float UniformScale, TMap<FName, FRichCurve>& OutCurves);
+	
 	//Conversion API
 	UE_DEPRECATED(4.21, "Please use FRichCurve version instead to get tangent weight support")
 	UNREALED_API void GetConvertedTransformCurveData(const FString& NodeName, FInterpCurveFloat& TranslationX, FInterpCurveFloat& TranslationY, FInterpCurveFloat& TranslationZ,
diff --git a/Engine/Source/Runtime/MovieScene/Public/INodeAndChannelMappings.h b/Engine/Source/Runtime/MovieScene/Public/INodeAndChannelMappings.h
--- a/Engine/Source/Runtime/MovieScene/Public/INodeAndChannelMappings.h
+++ b/Engine/Source/Runtime/MovieScene/Public/INodeAndChannelMappings.h
@@ -8,6 +8,7 @@
 #include "Containers/Array.h"
 #include "Containers/UnrealString.h"
 #include "CoreMinimal.h"
+#include "Channels/MovieSceneChannelEditorData.h"
 #include "HAL/Platform.h"
 #include "UObject/Interface.h"
 #include "UObject/Object.h"
@@ -20,19 +21,19 @@ class FName;
 class UMovieSceneSection;
 class UObject;
 
-/** Enumeration specifying the control type */
+/** Enumeration specifying the control type. Must match ERigControlType in order */
 enum class FFBXControlRigTypeProxyEnum : uint8
 {
 	Bool,
 	Float,
+	Integer,
 	Vector2D,
 	Position,
 	Scale,
 	Rotator,
 	Transform,
 	TransformNoScale,
-	EulerTransform,
-	Integer
+	EulerTransform
 };
 
 class UMovieSceneTrack;
@@ -43,11 +44,13 @@ struct FMovieSceneFloatChannel;
 struct FMovieSceneIntegerChannel;
 
 // For import data onto channels directly
-struct FFBXNodeAndChannels
+struct FRigControlFBXNodeAndChannels
 {
 	UMovieSceneTrack* MovieSceneTrack;
 	FFBXControlRigTypeProxyEnum ControlType;
 	FString NodeName;
+	FName ControlName;
+	
 	TArray< FMovieSceneDoubleChannel*> DoubleChannels;
 	TArray< FMovieSceneFloatChannel*> FloatChannels;
 	//will really only have one ever.
@@ -57,6 +60,44 @@ struct FFBXNodeAndChannels
 
 };
 
+struct FControlRigFbxNodeMapping
+{
+	// The UE channel type name
+	FName ChannelType;
+
+	// The UE control type proxy
+	FFBXControlRigTypeProxyEnum ControlType;
+
+	// The UE channel transform attribute index (0-9, invalid if non-transform) 
+	int8 ChannelAttrIndex = INDEX_NONE;
+
+	// The FBX transform attribute index (0-9)
+	uint8 FbxAttrIndex = 0;
+
+	bool bNegate = false;
+};
+
+struct FControlRigFbxCurveData
+{
+	// The name of the control associated to this curve
+	FName ControlName;
+
+	// The type of the control associated to this curve
+	FFBXControlRigTypeProxyEnum ControlType;
+
+	// The name of the node holding this curve
+	FString NodeName;
+
+	// the name of the FBX node's attribute (i.e. 'X', 'Weight'...) holding this curve
+	FString AttributeName;
+
+	// the name of the FBX node's property (i.e. 'Rotation', 'Scale'...) holding the channel attribute. Usually empty for non transform curves
+	FString AttributePropertyName;
+
+	// @returns whether the control is a node in itself instead of being grouped as an attribute of another existing control's node 
+	bool IsControlNode() const { return NodeName == ControlName.ToString(); }
+};
+
 
 UINTERFACE(meta = (CannotImplementInterfaceInBlueprint))
 class MOVIESCENE_API UNodeAndChannelMappings : public UInterface
@@ -71,8 +112,18 @@ class MOVIESCENE_API INodeAndChannelMappings
 	/** Get The Node And Mappings for this Track. Note Callee is responsible for deleting 
 	*  InSection The section from which to get the nodes and the channels from. If empty, up to implementor to decide which section to use, usually the Section To Key.
 	*/
-	virtual TArray<FFBXNodeAndChannels>*  GetNodeAndChannelMappings(UMovieSceneSection* InSection)  = 0;
+	virtual TArray<FRigControlFBXNodeAndChannels>* GetNodeAndChannelMappings(UMovieSceneSection* InSection)  = 0;
 	
 	/** Get Selected Nodes */
 	virtual void GetSelectedNodes(TArray<FName>& OutSelectedNodes) = 0;
+
+#if WITH_EDITOR
+	/**
+	 * @param MetaData The channel metadata
+	 * @param OutCurveData The data to be associated to the FBX curve corresponding to the given channel 
+	 */
+	virtual bool GetFbxCurveDataFromChannelMetadata(const FMovieSceneChannelMetaData& MetaData, FControlRigFbxCurveData& OutCurveData) = 0;
+
+#endif // WITH_EDITOR
+
 };
