Sequencer:  Fix crash when setting keys with cuve editor on non-bezier channels due to new check for interpolation type.

diff --git a/Engine/Source/Editor/MovieSceneTools/Private/Channels/BezierChannelCurveModel.cpp b/Engine/Source/Editor/MovieSceneTools/Private/Channels/BezierChannelCurveModel.cpp
--- a/Engine/Source/Editor/MovieSceneTools/Private/Channels/BezierChannelCurveModel.cpp
+++ b/Engine/Source/Editor/MovieSceneTools/Private/Channels/BezierChannelCurveModel.cpp
@@ -152,6 +152,40 @@ void FBezierChannelCurveModel<ChannelType, ChannelValue, KeyType>::GetKeyDrawInf
 	}
 }
 
+template <class ChannelType, class ChannelValue, class KeyType>
+ERichCurveInterpMode FBezierChannelCurveModel<ChannelType, ChannelValue, KeyType>::GetInterpolationMode(const double& InTime, ERichCurveInterpMode DefaultInterpolationMode) const
+{
+	ChannelType* Channel = this->GetChannelHandle().Get();
+	UMovieSceneSection* Section = Cast<UMovieSceneSection>(this->GetOwningObject());
+	if (Channel && Section)
+	{
+		FFrameRate TickResolution = Section->GetTypedOuter<UMovieScene>()->GetTickResolution();
+
+		TMovieSceneChannelData<ChannelValue> ChannelData = Channel->GetData();
+		TArrayView<const FFrameNumber> Times = ChannelData.GetTimes();
+		TArrayView<const ChannelValue> Values = ChannelData.GetValues();
+
+		const FFrameNumber InFrame = (InTime * TickResolution).RoundToFrame();
+
+		if (Times.Num() > 0)
+		{
+			int32 InterpolationIndex = Algo::LowerBound(Times, InFrame) - 1;
+			if (InterpolationIndex < 0)
+			{
+				InterpolationIndex = 0;
+			}
+			const FKeyHandle KeyHandle = ChannelData.GetHandle(InterpolationIndex);
+			TArrayView<const FKeyHandle> InKey(&KeyHandle, 1);
+			TArray<FKeyAttributes> KeyAttributes;
+			KeyAttributes.SetNum(1);
+			GetKeyAttributes(InKey, KeyAttributes);
+			return KeyAttributes[0].GetInterpMode();
+		}
+	}
+
+	return DefaultInterpolationMode;
+}
+
 template<typename ChannelType, typename ChannelValue, typename KeyType> 
 void FBezierChannelCurveModel<ChannelType, ChannelValue, KeyType>::GetKeyAttributes(TArrayView<const FKeyHandle> InKeys, TArrayView<FKeyAttributes> OutAttributes) const
 {
diff --git a/Engine/Source/Editor/MovieSceneTools/Private/Channels/ChannelCurveModel.cpp b/Engine/Source/Editor/MovieSceneTools/Private/Channels/ChannelCurveModel.cpp
--- a/Engine/Source/Editor/MovieSceneTools/Private/Channels/ChannelCurveModel.cpp
+++ b/Engine/Source/Editor/MovieSceneTools/Private/Channels/ChannelCurveModel.cpp
@@ -331,40 +331,6 @@ void FChannelCurveModel<ChannelType, ChannelValue, KeyType>::GetNeighboringKeys(
 	}
 }
 
-template <class ChannelType, class ChannelValue, class KeyType>
-ERichCurveInterpMode FChannelCurveModel<ChannelType, ChannelValue, KeyType>::GetInterpolationMode(const double& InTime, ERichCurveInterpMode DefaultInterpolationMode) const
-{
-	ChannelType* Channel = ChannelHandle.Get();
-	UMovieSceneSection* Section = WeakSection.Get();
-	if (Channel && Section)
-	{
-		FFrameRate TickResolution = Section->GetTypedOuter<UMovieScene>()->GetTickResolution();
-
-		TMovieSceneChannelData<ChannelValue> ChannelData = Channel->GetData();
-		TArrayView<const FFrameNumber> Times = ChannelData.GetTimes();
-		TArrayView<const ChannelValue> Values = ChannelData.GetValues();
-
-		const FFrameNumber InFrame = (InTime * TickResolution).RoundToFrame();
-
-		if (Times.Num() > 0)
-		{
- 			int32 InterpolationIndex = Algo::LowerBound(Times, InFrame) - 1;
- 			if (InterpolationIndex < 0)
- 			{
- 				InterpolationIndex = 0;
- 			}
-			const FKeyHandle KeyHandle = ChannelData.GetHandle(InterpolationIndex);
-			TArrayView<const FKeyHandle> InKey(&KeyHandle, 1);
-			TArray<FKeyAttributes> KeyAttributes;
-			KeyAttributes.SetNum(1);
-			GetKeyAttributes(InKey, KeyAttributes);
-			return KeyAttributes[0].GetInterpMode();
-		}
-	}
-
-	return DefaultInterpolationMode;
-}
-
 template <class ChannelType, class ChannelValue, class KeyType>
 bool FChannelCurveModel<ChannelType, ChannelValue, KeyType>::Evaluate(double Time, double& OutValue) const
 {
diff --git a/Engine/Source/Editor/MovieSceneTools/Public/Channels/BezierChannelCurveModel.h b/Engine/Source/Editor/MovieSceneTools/Public/Channels/BezierChannelCurveModel.h
--- a/Engine/Source/Editor/MovieSceneTools/Public/Channels/BezierChannelCurveModel.h
+++ b/Engine/Source/Editor/MovieSceneTools/Public/Channels/BezierChannelCurveModel.h
@@ -63,6 +63,7 @@ public:
 	virtual void GetKeyAttributes(TArrayView<const FKeyHandle> InKeys, TArrayView<FKeyAttributes> OutAttributes) const override;
 	virtual void SetKeyAttributes(TArrayView<const FKeyHandle> InKeys, TArrayView<const FKeyAttributes> InAttributes, EPropertyChangeType::Type ChangeType = EPropertyChangeType::Unspecified) override;
 
+	virtual ERichCurveInterpMode GetInterpolationMode(const double& InTime, ERichCurveInterpMode DefaultInterpolationMode) const override;
 	virtual void GetCurveAttributes(FCurveAttributes& OutCurveAttributes) const override;
 	virtual void SetCurveAttributes(const FCurveAttributes& InCurveAttributes) override;
 	virtual void GetValueRange(double& MinValue, double& MaxValue) const override;
diff --git a/Engine/Source/Editor/MovieSceneTools/Public/Channels/ChannelCurveModel.h b/Engine/Source/Editor/MovieSceneTools/Public/Channels/ChannelCurveModel.h
--- a/Engine/Source/Editor/MovieSceneTools/Public/Channels/ChannelCurveModel.h
+++ b/Engine/Source/Editor/MovieSceneTools/Public/Channels/ChannelCurveModel.h
@@ -58,7 +58,6 @@ public:
 	virtual void GetValueRange(double& MinValue, double& MaxValue) const override;
 	virtual int32 GetNumKeys() const override;
 	virtual void GetNeighboringKeys(const FKeyHandle InKeyHandle, TOptional<FKeyHandle>& OutPreviousKeyHandle, TOptional<FKeyHandle>& OutNextKeyHandle) const override;
-	virtual ERichCurveInterpMode GetInterpolationMode(const double& InTime, ERichCurveInterpMode DefaultInterpolationMode) const override;
 	virtual bool Evaluate(double ProspectiveTime, double& OutValue) const override;
 	virtual void AddKeys(TArrayView<const FKeyPosition> InKeyPositions, TArrayView<const FKeyAttributes> InAttributes, TArrayView<TOptional<FKeyHandle>>* OutKeyHandles) override;
 	virtual void RemoveKeys(TArrayView<const FKeyHandle> InKeys) override;
