Sequencer: Fixed a crash where you could use the context menu for the Camera Cuts track to try and select an unspawned spawnable. Now selection is blocked and the tooltip indicates that it cannot be selected as it is currently not spawned.

diff --git a/Engine/Source/Editor/MovieSceneTools/Private/Sections/CameraCutSection.cpp b/Engine/Source/Editor/MovieSceneTools/Private/Sections/CameraCutSection.cpp
--- a/Engine/Source/Editor/MovieSceneTools/Private/Sections/CameraCutSection.cpp
+++ b/Engine/Source/Editor/MovieSceneTools/Private/Sections/CameraCutSection.cpp
@@ -82,12 +82,20 @@ void FCameraCutSection::BuildSectionContextMenu(FMenuBuilder& MenuBuilder, const
 	if (CameraActor)
 	{
 		MenuBuilder.AddMenuSeparator();
+		
+		// GetCameraForFrame will return the Spawnable Template (for names) but we can't select those.
+		const bool bCanSelect = CameraActor->GetWorld() != nullptr;
+		const FText CameraNameLabel = FText::FromString(CameraActor->GetActorLabel());
+		const FText Tooltip = bCanSelect ?
+			FText::Format(LOCTEXT("SelectCameraTooltipFormat", "Select {0}"), CameraNameLabel) :
+			FText::Format(LOCTEXT("SelectCameraInvalidTooltipFormat", "Cannot Select {0} (Currently Unspawned)"), CameraNameLabel);
 
 		MenuBuilder.AddMenuEntry(
-			FText::Format(LOCTEXT("SelectCameraTextFormat", "Select {0}"), FText::FromString(CameraActor->GetActorLabel())),
-			FText::Format(LOCTEXT("SelectCameraTooltipFormat", "Select {0}"), FText::FromString(CameraActor->GetActorLabel())),
+			FText::Format(LOCTEXT("SelectCameraTextFormat", "Select {0}"), CameraNameLabel),
+			Tooltip,
 			FSlateIcon(),
-			FUIAction(FExecuteAction::CreateRaw(this, &FCameraCutSection::HandleSelectCameraMenuEntryExecute, CameraActor))
+			FUIAction(FExecuteAction::CreateRaw(this, &FCameraCutSection::HandleSelectCameraMenuEntryExecute, CameraActor),
+				FCanExecuteAction::CreateRaw(this, & FCameraCutSection::CanSelectCameraActor, CameraActor))
 		);
 	}
 
@@ -230,6 +238,11 @@ FText FCameraCutSection::HandleThumbnailTextBlockText() const
 /* FCameraCutSection callbacks
  *****************************************************************************/
 
+bool FCameraCutSection::CanSelectCameraActor(AActor* InCamera) const
+{
+	return InCamera && InCamera->GetWorld();
+}
+
 void FCameraCutSection::HandleSelectCameraMenuEntryExecute(AActor* InCamera)
 {
 	const bool bInSelected = true;
diff --git a/Engine/Source/Editor/MovieSceneTools/Private/Sections/CameraCutSection.h b/Engine/Source/Editor/MovieSceneTools/Private/Sections/CameraCutSection.h
--- a/Engine/Source/Editor/MovieSceneTools/Private/Sections/CameraCutSection.h
+++ b/Engine/Source/Editor/MovieSceneTools/Private/Sections/CameraCutSection.h
@@ -48,6 +48,10 @@ private:
 	/** Callback for executing a "Select Camera" menu entry in the context menu. */
 	void HandleSelectCameraMenuEntryExecute(AActor* InCamera);
 
+	/** Callback for determining if the specified camera can actually be selected in the World. */
+	bool CanSelectCameraActor(AActor* InCamera) const;
+
+
 	/** Callback for executing a "Set Camera" menu entry in the context menu. */
 	void HandleSetCameraMenuEntryExecute(AActor* InCamera);
 };
