Constraints: propagate offset change so that local transform animations are synchronized

diff --git a/Engine/Source/Editor/MovieSceneTools/Private/Constraints/MovieSceneConstraintChannelHelper.cpp b/Engine/Source/Editor/MovieSceneTools/Private/Constraints/MovieSceneConstraintChannelHelper.cpp
--- a/Engine/Source/Editor/MovieSceneTools/Private/Constraints/MovieSceneConstraintChannelHelper.cpp
+++ b/Engine/Source/Editor/MovieSceneTools/Private/Constraints/MovieSceneConstraintChannelHelper.cpp
@@ -191,6 +191,7 @@ void FCompensationEvaluator::ComputeLocalTransforms(
 		}
 	}
 }
+
 void FCompensationEvaluator::ComputeLocalTransformsForBaking(UWorld* InWorld, const TSharedPtr<ISequencer>& InSequencer, const TArray<FFrameNumber>& InFrames)
 {
 	if (InFrames.IsEmpty())
@@ -620,6 +621,46 @@ void FCompensationEvaluator::CacheTransforms(UWorld* InWorld, const TSharedPtr<I
 	}
 }
 
+void FCompensationEvaluator::ComputeCurrentTransforms(UWorld* InWorld)
+{
+	ChildLocals = ChildGlobals = SpaceGlobals = {FTransform::Identity};
+
+	using ConstraintPtr = TObjectPtr<UTickableConstraint>;
+	const TArray< ConstraintPtr > Constraints = GetHandleTransformConstraints(InWorld);
+	if (Constraints.IsEmpty())
+	{
+		return;
+	}
+
+	for (const UTickableConstraint* InConstraint : Constraints)
+	{
+		InConstraint->Evaluate();
+	}
+
+	ChildLocals[0] = Handle->GetLocalTransform();
+	ChildGlobals[0] = Handle->GetGlobalTransform();
+	
+	auto GetLastActiveConstraint = [Constraints]()
+	{
+		// find last active constraint in the list that is different than the one we want to compensate for
+		const int32 LastActiveIndex = FTransformConstraintUtils::GetLastActiveConstraintIndex(Constraints);
+
+		// if found, return its parent global transform
+		return LastActiveIndex > INDEX_NONE ? Cast<UTickableTransformConstraint>(Constraints[LastActiveIndex]) : nullptr;
+	};
+	
+	if (const UTickableTransformConstraint* LastConstraint = GetLastActiveConstraint())
+	{
+		SpaceGlobals[0] = LastConstraint->GetParentGlobalTransform();
+		TOptional<FTransform> Relative =
+			FTransformConstraintUtils::GetConstraintsRelativeTransform(Constraints, ChildLocals[0], ChildGlobals[0]);
+		if (Relative)
+		{
+			ChildLocals[0] = *Relative;
+		}
+	}
+}
+
 TArray< TObjectPtr<UTickableConstraint> > FCompensationEvaluator::GetHandleTransformConstraints(UWorld* InWorld) const
 {
 	using ConstraintPtr = TObjectPtr<UTickableConstraint>;
@@ -1158,6 +1199,14 @@ void FMovieSceneConstraintChannelHelper::HandleConstraintPropertyChanged(
 	{
 		return CompensateScale(Cast<UTickableParentConstraint>(InConstraint), InActiveChannel, InSequencer, InSection);
 	}
+	
+	if (PropertyName == GET_MEMBER_NAME_CHECKED(UTickableTranslationConstraint, OffsetTranslation) ||
+		PropertyName == GET_MEMBER_NAME_CHECKED(UTickableRotationConstraint, OffsetRotation) ||
+		PropertyName == GET_MEMBER_NAME_CHECKED(UTickableScaleConstraint, OffsetScale) ||
+		PropertyName == GET_MEMBER_NAME_CHECKED(UTickableParentConstraint, OffsetTransform))
+	{
+		return HandleOffsetChanged(InConstraint, InActiveChannel, InSequencer);
+	}
 }
 
 void FMovieSceneConstraintChannelHelper::CompensateScale(
@@ -1245,3 +1294,48 @@ void FMovieSceneConstraintChannelHelper::CompensateScale(
 	// reset scaling to reference value
 	InParentConstraint->SetScaling(bRefScalingValue);
 }
+
+void FMovieSceneConstraintChannelHelper::HandleOffsetChanged(
+	UTickableTransformConstraint* InConstraint,
+	const FMovieSceneConstraintChannel& InActiveChannel,
+	const TSharedPtr<ISequencer>& InSequencer)
+{
+	if (!InConstraint || !InSequencer.IsValid())
+	{
+		return;
+	}
+	
+	TObjectPtr<UTransformableHandle> Handle = InConstraint->ChildTRSHandle;
+	ITransformConstraintChannelInterface* Interface = GetHandleInterface(Handle);
+	if (!Interface)
+	{
+		return;
+	}
+
+	const TArrayView<const FFrameNumber> Times = InActiveChannel.GetTimes();
+	if (Times.IsEmpty())
+	{
+		return;
+	}
+
+	const FFrameRate TickResolution = InSequencer->GetFocusedTickResolution();
+	const FFrameTime FrameTime = InSequencer->GetLocalTime().ConvertTo(TickResolution);
+	const FFrameNumber Time = FrameTime.GetFrame();
+		
+	bool bIsActive = false; InActiveChannel.Evaluate(Time, bIsActive);
+	if (bIsActive)
+	{
+		const EMovieSceneTransformChannel Channels = InConstraint->GetChannelsToKey();
+	
+		// compute the current local value
+		FCompensationEvaluator Evaluator(InConstraint);
+		UWorld* World = GCurrentLevelEditingViewportClient ? GCurrentLevelEditingViewportClient->GetWorld() : nullptr;
+		Evaluator.ComputeCurrentTransforms(World);
+
+		// update key
+		Interface->AddHandleTransformKeys(InSequencer, Handle, {Time}, {Evaluator.ChildLocals[0]}, Channels);
+		
+		// force evaluation so that new local values are evaluated before the constraint 
+		InSequencer->ForceEvaluate();
+	}
+}
\ No newline at end of file
diff --git a/Engine/Source/Editor/MovieSceneTools/Public/Constraints/MovieSceneConstraintChannelHelper.h b/Engine/Source/Editor/MovieSceneTools/Public/Constraints/MovieSceneConstraintChannelHelper.h
--- a/Engine/Source/Editor/MovieSceneTools/Public/Constraints/MovieSceneConstraintChannelHelper.h
+++ b/Engine/Source/Editor/MovieSceneTools/Public/Constraints/MovieSceneConstraintChannelHelper.h
@@ -33,6 +33,7 @@ public:
 	void ComputeCompensation(UWorld* InWorld, const TSharedPtr<ISequencer>& InSequencer, const FFrameNumber& InTime);
 	void ComputeLocalTransformsForBaking(UWorld* InWorld, const TSharedPtr<ISequencer>& InSequencer, const TArray<FFrameNumber>& InFrames);
 	void CacheTransforms(UWorld* InWorld, const TSharedPtr<ISequencer>& InSequencer, const TArray<FFrameNumber>& InFrames);
+	void ComputeCurrentTransforms(UWorld* InWorld);
 	
 private:
 
@@ -186,9 +187,16 @@ private:
 	/** For the given handle create any movie scene binding for it based upon the current sequencer that's open*/
 	static void CreateBindingIDForHandle(const TSharedPtr<ISequencer>& InSequencer, UTransformableHandle* InHandle);
 
+	/** Compensate scale keys when enabling/disabling scaling for parent constraints. */
 	static void CompensateScale(
 		UTickableParentConstraint* InParentConstraint,
 		const FMovieSceneConstraintChannel& InActiveChannel,
 		const TSharedPtr<ISequencer>& InSequencer,
 		UMovieSceneSection* InSection);
+
+	/** Handle offset modifications so that the child's transform channels are synced. */
+	static void HandleOffsetChanged(
+		UTickableTransformConstraint* InConstraint,
+		const FMovieSceneConstraintChannel& InActiveChannel,
+		const TSharedPtr<ISequencer>& InSequencer);
 };
diff --git a/Engine/Source/Runtime/Experimental/Animation/Constraints/Private/TransformConstraint.cpp b/Engine/Source/Runtime/Experimental/Animation/Constraints/Private/TransformConstraint.cpp
--- a/Engine/Source/Runtime/Experimental/Animation/Constraints/Private/TransformConstraint.cpp
+++ b/Engine/Source/Runtime/Experimental/Animation/Constraints/Private/TransformConstraint.cpp
@@ -301,9 +301,13 @@ void UTickableTransformConstraint::PostEditChangeProperty(FPropertyChangedEvent&
 
 	if (const FProperty* MemberProperty = PropertyChangedEvent.MemberProperty)
 	{
-		static const FString OffsetStr("Offset");
-		if (MemberProperty->GetFName().ToString().Contains(OffsetStr) )
+		const FName MemberPropertyName = MemberProperty->GetFName();
+		if (MemberPropertyName == GET_MEMBER_NAME_CHECKED(UTickableTranslationConstraint, OffsetTranslation) ||
+			MemberPropertyName == GET_MEMBER_NAME_CHECKED(UTickableRotationConstraint, OffsetRotation) ||
+			MemberPropertyName == GET_MEMBER_NAME_CHECKED(UTickableScaleConstraint, OffsetScale) ||
+			MemberPropertyName == GET_MEMBER_NAME_CHECKED(UTickableParentConstraint, OffsetTransform))
 		{
+			OnConstraintChanged.Broadcast(this, PropertyChangedEvent);
 			Evaluate();
 		}
 	}
diff --git a/Engine/Source/Runtime/Experimental/Animation/Constraints/Public/TransformConstraint.h b/Engine/Source/Runtime/Experimental/Animation/Constraints/Public/TransformConstraint.h
--- a/Engine/Source/Runtime/Experimental/Animation/Constraints/Public/TransformConstraint.h
+++ b/Engine/Source/Runtime/Experimental/Animation/Constraints/Public/TransformConstraint.h
@@ -211,10 +211,12 @@ protected:
 	/** Computes the child's local translation offset in the parent space. */
 	virtual void ComputeOffset() override;
 
+public:
 	/** Defines the local child's translation offset in the parent space. */
 	UPROPERTY(EditAnywhere, BlueprintReadWrite, Category="Offset", meta=(EditCondition="bMaintainOffset"))
 	FVector OffsetTranslation = FVector::ZeroVector;
 
+protected:
 	/** Defines which translation axis is constrained. */
 	UPROPERTY(BlueprintReadWrite, Category = "Axis Filter")
 	FFilterOptionPerAxis AxisFilter;
@@ -257,11 +259,13 @@ protected:
 	
 	/** Computes the child's local rotation offset in the parent space. */
 	virtual void ComputeOffset() override;
-	
+
+public:
 	/** Defines the local child's rotation offset in the parent space. */
 	UPROPERTY(EditAnywhere, BlueprintReadWrite, Category="Offset", meta=(EditCondition="bMaintainOffset"))
 	FQuat OffsetRotation = FQuat::Identity;
 
+protected:
 	/** Defines which rotation axis is constrained. */
 	UPROPERTY(BlueprintReadWrite, Category = "Axis Filter")
 	FFilterOptionPerAxis AxisFilter;
@@ -303,11 +307,13 @@ protected:
 	};
 	mutable FDynamicCache Cache;
 	uint32 CalculateInputHash() const;
-	
+
+public:
 	/** Defines the local child's scale offset in the parent space. */
 	UPROPERTY(EditAnywhere, BlueprintReadWrite, Category="Offset", meta=(EditCondition="bMaintainOffset"))
 	FVector OffsetScale = FVector::OneVector;
 
+protected:
 	/** Defines which scale axis is constrained. */
 	UPROPERTY(BlueprintReadWrite, Category = "Axis Filter")
 	FFilterOptionPerAxis AxisFilter;
@@ -362,10 +368,12 @@ protected:
 	/** Computes the child's local transform offset in the parent space. */
 	virtual void ComputeOffset() override;
 
+public:
 	/** Defines the local child's transform offset in the parent space. */
 	UPROPERTY(EditAnywhere, BlueprintReadWrite, Category="Offset", meta=(EditCondition="bMaintainOffset"))
 	FTransform OffsetTransform = FTransform::Identity;
 
+protected:
 	/** Defines whether we propagate the parent scale. */
 	UPROPERTY(EditAnywhere, BlueprintReadWrite, Category="Properties")
 	bool bScaling = false;
