PR #10408: Fix "trim and preserve" calculating incorrect transforms when attach component isn't the root component

diff --git a/Engine/Source/Editor/MovieSceneTools/Private/TrackEditors/AttachTrackEditor.cpp b/Engine/Source/Editor/MovieSceneTools/Private/TrackEditors/AttachTrackEditor.cpp
--- a/Engine/Source/Editor/MovieSceneTools/Private/TrackEditors/AttachTrackEditor.cpp
+++ b/Engine/Source/Editor/MovieSceneTools/Private/TrackEditors/AttachTrackEditor.cpp
@@ -614,7 +614,7 @@ struct FWorldTransformEvaluator : ITransformEvaluator
 	 * Creates a new evaluator for a given object
 	 * @param InSocketName is the socket to evaluate for if this is a skeletal mesh
 	 */
-	FWorldTransformEvaluator(TWeakPtr<F3DAttachTrackEditor> InWeakAttachTrackEditor, UObject* InObject, const FName InSocketName = NAME_None)
+	FWorldTransformEvaluator(TWeakPtr<F3DAttachTrackEditor> InWeakAttachTrackEditor, UObject* InObject, const FName InSocketName = NAME_None, const FName InComponentName = NAME_None)
 		: WeakAttachTrackEditor(InWeakAttachTrackEditor)
 	{
 		TSharedPtr<F3DAttachTrackEditor> AttachTrackEditor = InWeakAttachTrackEditor.Pin();
@@ -632,11 +632,16 @@ struct FWorldTransformEvaluator : ITransformEvaluator
 		UMovieScene* MovieScene = Sequencer->GetFocusedMovieSceneSequence()->GetMovieScene();
 
 		FName SocketName = InSocketName;
+		FName ComponentName = InComponentName;
 
 		USceneComponent* SceneComponent = Cast<USceneComponent>(InObject);
 		if (AActor* Actor = Cast<AActor>(InObject))
 		{
-			SceneComponent = Actor->GetRootComponent();
+			UE::MovieScene::FComponentAttachParamsDestination AttachParams;
+			AttachParams.SocketName = SocketName;
+			AttachParams.ComponentName = ComponentName;
+
+			SceneComponent = AttachParams.ResolveAttachment(Actor);
 		}
 
 		if (!SceneComponent)
@@ -772,14 +777,14 @@ struct FAttachRevertModifier
 	 *                               modifying the object's movements to how they were before the attach.
 	 *                      If false: Parent's movement is kept and transforms are simply converted to world space.
 	 */
-	FAttachRevertModifier(TSharedPtr<F3DAttachTrackEditor> InWeakAttachTrackEditor, const TRange<FFrameNumber>& InRevertRange, UMovieScene3DAttachSection* InAttachSection, const FName InSocketName, bool bInFullRevert)
+	FAttachRevertModifier(TSharedPtr<F3DAttachTrackEditor> InWeakAttachTrackEditor, const TRange<FFrameNumber>& InRevertRange, UMovieScene3DAttachSection* InAttachSection, const FName InSocketName, const FName InComponentName, bool bInFullRevert)
 		: bFullRevert(bInFullRevert)
 		, RevertRange(InRevertRange)
 	{
 		FMovieSceneObjectBindingID ConstraintID = InAttachSection->GetConstraintBindingID();
 		UObject* ConstraintObject = GetConstraintObject(InWeakAttachTrackEditor->GetSequencer(), ConstraintID);
 
-		TransformEvaluator = FWorldTransformEvaluator(InWeakAttachTrackEditor, ConstraintObject, InSocketName);
+		TransformEvaluator = FWorldTransformEvaluator(InWeakAttachTrackEditor, ConstraintObject, InSocketName, InComponentName);
 
 		BeginConstraintTransform = TransformEvaluator(InRevertRange.GetLowerBoundValue());
 
@@ -937,7 +942,7 @@ void F3DAttachTrackEditor::TrimAndPreserve(FGuid InObjectBinding, UMovieSceneSec
 		check(AttachSection);
 
 		// Create a revert modifier with the range and section as parameters
-		FAttachRevertModifier RevertModifier(SharedThis(this), ExcludedRange, AttachSection, AttachSection->AttachSocketName, AttachSection->bFullRevertOnDetach);
+		FAttachRevertModifier RevertModifier(SharedThis(this), ExcludedRange, AttachSection, AttachSection->AttachSocketName, AttachSection->AttachComponentName, AttachSection->bFullRevertOnDetach);
 
 		// Find the transform section associated with the track, so far we only support modifying transform tracks with one section
 		UMovieScene3DTransformSection* TransformSection = Cast<UMovieScene3DTransformSection>(TransformTrack->GetAllSections()[0]);
@@ -1333,7 +1338,7 @@ FKeyPropertyResult F3DAttachTrackEditor::AddKeyInternal( FFrameNumber KeyTime, c
 				continue;
 			}
 
-			RevertModifier = FAttachRevertModifier(SharedThis(this), RevertRange, IntersectingAttachSection, SocketName, PreserveType == ETransformPreserveType::CurrentKey);
+			RevertModifier = FAttachRevertModifier(SharedThis(this), RevertRange, IntersectingAttachSection, SocketName, ComponentName, PreserveType == ETransformPreserveType::CurrentKey);
 		}
 		// Existing parent that's not an attach track
 		else if (WorldChildTransformEval.GetTransformEvalsView().Num() > 1)
