Movie Scene Tools:
- Fix string curve editor widget not being noticeable when value is blank

diff --git a/Engine/Source/Editor/MovieSceneTools/Private/CurveKeyEditors/SStringCurveKeyEditor.cpp b/Engine/Source/Editor/MovieSceneTools/Private/CurveKeyEditors/SStringCurveKeyEditor.cpp
--- a/Engine/Source/Editor/MovieSceneTools/Private/CurveKeyEditors/SStringCurveKeyEditor.cpp
+++ b/Engine/Source/Editor/MovieSceneTools/Private/CurveKeyEditors/SStringCurveKeyEditor.cpp
@@ -9,7 +9,7 @@
 #include "Layout/Children.h"
 #include "Misc/Attribute.h"
 #include "ScopedTransaction.h"
-#include "Widgets/Input/SEditableText.h"
+#include "Widgets/Input/SEditableTextBox.h"
 
 #define LOCTEXT_NAMESPACE "StringCurveKeyEditor"
 
@@ -19,7 +19,8 @@ void SStringCurveKeyEditor::Construct(const FArguments& InArgs, const TSequencer
 
 	ChildSlot
 	[
-		SNew(SEditableText)
+		SNew(SEditableTextBox)
+		.MinDesiredWidth(10.f)
 		.SelectAllTextWhenFocused(true)
 		.Text(this, &SStringCurveKeyEditor::GetText)
 		.OnTextCommitted(this, &SStringCurveKeyEditor::OnTextCommitted)
