Sequencer: support multiple sequences for the dynamic binding BP extension

Multiple wdgets animations can share the same BP extension since the BP is the UserWidget itself.

diff --git a/Engine/Source/Editor/MovieSceneTools/Private/MovieSceneDynamicBindingBlueprintExtension.cpp b/Engine/Source/Editor/MovieSceneTools/Private/MovieSceneDynamicBindingBlueprintExtension.cpp
--- a/Engine/Source/Editor/MovieSceneTools/Private/MovieSceneDynamicBindingBlueprintExtension.cpp
+++ b/Engine/Source/Editor/MovieSceneTools/Private/MovieSceneDynamicBindingBlueprintExtension.cpp
@@ -27,25 +27,55 @@ class UEdGraphNode;
 
 void UMovieSceneDynamicBindingBlueprintExtension::BindTo(TWeakObjectPtr<UMovieSceneSequence> InMovieSceneSequence)
 {
-	check(!WeakMovieSceneSequence.IsValid() || WeakMovieSceneSequence == InMovieSceneSequence);
-	WeakMovieSceneSequence = InMovieSceneSequence;
+	WeakMovieSceneSequences.AddUnique(InMovieSceneSequence);
+}
+
+void UMovieSceneDynamicBindingBlueprintExtension::Serialize(FArchive& Ar)
+{
+	Super::Serialize(Ar);
+
+	if (Ar.IsLoading())
+	{
+		if (WeakMovieSceneSequence_DEPRECATED.IsValid())
+		{
+			WeakMovieSceneSequences.Add(WeakMovieSceneSequence_DEPRECATED);
+			WeakMovieSceneSequence_DEPRECATED.Reset();
+		}
+	}
+}
+
+void UMovieSceneDynamicBindingBlueprintExtension::PostLoad()
+{
+	WeakMovieSceneSequences.Remove(nullptr);
+	Super::PostLoad();
 }
 
 void UMovieSceneDynamicBindingBlueprintExtension::HandlePreloadObjectsForCompilation(UBlueprint* OwningBlueprint)
 {
-	if (UMovieSceneSequence* MovieSceneSequence = WeakMovieSceneSequence.Get())
+	for (TWeakObjectPtr<UMovieSceneSequence> WeakMovieSceneSequence : WeakMovieSceneSequences)
 	{
-		UBlueprint::ForceLoad(MovieSceneSequence);
-		if (UMovieScene* MovieScene = MovieSceneSequence->GetMovieScene())
+		if (UMovieSceneSequence* MovieSceneSequence = WeakMovieSceneSequence.Get())
 		{
-			UBlueprint::ForceLoad(MovieScene);
+			UBlueprint::ForceLoad(MovieSceneSequence);
+			if (UMovieScene* MovieScene = MovieSceneSequence->GetMovieScene())
+			{
+				UBlueprint::ForceLoad(MovieScene);
+			}
 		}
 	}
 }
 
 void UMovieSceneDynamicBindingBlueprintExtension::HandleGenerateFunctionGraphs(FKismetCompilerContext* CompilerContext)
 {
-	UMovieSceneSequence* MovieSceneSequence = WeakMovieSceneSequence.Get();
+	for (TWeakObjectPtr<UMovieSceneSequence> WeakMovieSceneSequence : WeakMovieSceneSequences)
+	{
+		UMovieSceneSequence* MovieSceneSequence = WeakMovieSceneSequence.Get();
+		HandleGenerateFunctionGraphs(CompilerContext, MovieSceneSequence);
+	}
+}
+
+void UMovieSceneDynamicBindingBlueprintExtension::HandleGenerateFunctionGraphs(FKismetCompilerContext* CompilerContext, UMovieSceneSequence* MovieSceneSequence)
+{
 	UMovieScene* MovieScene = MovieSceneSequence ? MovieSceneSequence->GetMovieScene() : nullptr;
 	if (!MovieScene)
 	{
@@ -86,10 +116,10 @@ void UMovieSceneDynamicBindingBlueprintExtension::HandleGenerateFunctionGraphs(F
 	// This callback sets the generated function calls back onto the spawnables/possessables in the sequence,
 	// and keeps a pointer to any needed arguments on this function, so that we can pass special values later
 	// at runtime.
-	TWeakObjectPtr<UMovieSceneSequence> WeakMovieSceneSequenceCopy(WeakMovieSceneSequence);
-	auto OnFunctionListGenerated = [WeakMovieSceneSequenceCopy](FKismetCompilerContext* CompilerContext)
+	TWeakObjectPtr<UMovieSceneSequence> WeakMovieSceneSequence(MovieSceneSequence);
+	auto OnFunctionListGenerated = [WeakMovieSceneSequence](FKismetCompilerContext* CompilerContext)
 	{
-		UMovieSceneSequence* MovieSceneSequence = WeakMovieSceneSequenceCopy.Get();
+		UMovieSceneSequence* MovieSceneSequence = WeakMovieSceneSequence.Get();
 		UMovieScene* MovieScene = MovieSceneSequence ? MovieSceneSequence->GetMovieScene() : nullptr;
 		if (!ensureMsgf(MovieSceneSequence, TEXT("A movie scene was garbage-collected while its director blueprint was being compiled!")))
 		{
diff --git a/Engine/Source/Editor/MovieSceneTools/Private/MovieSceneDynamicBindingBlueprintExtension.h b/Engine/Source/Editor/MovieSceneTools/Private/MovieSceneDynamicBindingBlueprintExtension.h
--- a/Engine/Source/Editor/MovieSceneTools/Private/MovieSceneDynamicBindingBlueprintExtension.h
+++ b/Engine/Source/Editor/MovieSceneTools/Private/MovieSceneDynamicBindingBlueprintExtension.h
@@ -34,11 +34,22 @@ public:
 
 private:
 
+	virtual void Serialize(FArchive& Ar) override;
+	virtual void PostLoad() override;
+
 	virtual void HandlePreloadObjectsForCompilation(UBlueprint* OwningBlueprint) override final;
 	virtual void HandleGenerateFunctionGraphs(FKismetCompilerContext* CompilerContext) override final;
 
+	void HandleGenerateFunctionGraphs(FKismetCompilerContext* CompilerContext, UMovieSceneSequence* MovieSceneSequence);
+
+private:
+
+	/** List of movie scenes that are bound to the director blueprint */
+	UPROPERTY()
+	TArray<TWeakObjectPtr<UMovieSceneSequence>> WeakMovieSceneSequences;
+
 	/** The movie scene that is bound to the director blueprint */
 	UPROPERTY()
-	TWeakObjectPtr<UMovieSceneSequence> WeakMovieSceneSequence;
+	TWeakObjectPtr<UMovieSceneSequence> WeakMovieSceneSequence_DEPRECATED;
 };
 
