Sequencer: Compensate to root component of attach actor if it's an actor we're attaching to

#rnx

diff --git a/Engine/Source/Editor/MovieSceneTools/Private/TrackEditors/AttachTrackEditor.cpp b/Engine/Source/Editor/MovieSceneTools/Private/TrackEditors/AttachTrackEditor.cpp
--- a/Engine/Source/Editor/MovieSceneTools/Private/TrackEditors/AttachTrackEditor.cpp
+++ b/Engine/Source/Editor/MovieSceneTools/Private/TrackEditors/AttachTrackEditor.cpp
@@ -491,6 +491,10 @@ UObject* GetConstraintObject(TSharedPtr<ISequencer> InSequencer, const FMovieSce
 
 	if (RuntimeObjects.Num() >= 1 && RuntimeObjects[0].IsValid())
 	{
+		if (AActor* Actor = Cast<AActor>(RuntimeObjects[0].Get()))
+		{
+			return Actor->GetRootComponent();
+		}
 		return RuntimeObjects[0].Get();
 	}
 
