Sequencer: Use the same color tint as the shot that's being duplicated or new take

diff --git a/Engine/Source/Editor/MovieSceneTools/Private/TrackEditors/SubTrackEditor.cpp b/Engine/Source/Editor/MovieSceneTools/Private/TrackEditors/SubTrackEditor.cpp
--- a/Engine/Source/Editor/MovieSceneTools/Private/TrackEditors/SubTrackEditor.cpp
+++ b/Engine/Source/Editor/MovieSceneTools/Private/TrackEditors/SubTrackEditor.cpp
@@ -507,6 +507,7 @@ void FSubTrackEditor::DuplicateSection(UMovieSceneSubSection* Section)
 			NewSection->Parameters.StartFrameOffset = Section->Parameters.StartFrameOffset;
 			NewSection->Parameters.TimeScale = Section->Parameters.TimeScale;
 			NewSection->SetPreRollFrames(Section->GetPreRollFrames());
+			NewSection->SetColorTint(Section->GetColorTint());
 
 			GetSequencer()->NotifyMovieSceneDataChanged(EMovieSceneDataChangeType::MovieSceneStructureItemAdded);
 			GetSequencer()->EmptySelection();
@@ -551,7 +552,7 @@ void FSubTrackEditor::CreateNewTake(UMovieSceneSubSection* Section)
 		int32                NewShotPrerollFrames = Section->GetPreRollFrames();
 		int32                NewRowIndex          = Section->GetRowIndex();
 		FFrameNumber         NewShotStartTime     = NewShotRange.GetLowerBound().IsClosed() ? UE::MovieScene::DiscreteInclusiveLower(NewShotRange) : 0;
-
+		FColor               NewShotColorTint     = Section->GetColorTint();
 		UMovieSceneSubTrack* SubTrack = CastChecked<UMovieSceneSubTrack>(Section->GetOuter());
 		FString NewSequencePath = FPaths::GetPath(Section->GetSequence()->GetPathName());
 
@@ -569,6 +570,7 @@ void FSubTrackEditor::CreateNewTake(UMovieSceneSubSection* Section)
 			NewShot->Parameters.TimeScale = NewShotTimeScale;
 			NewShot->SetPreRollFrames(NewShotPrerollFrames);
 			NewShot->SetRowIndex(NewRowIndex);
+			NewShot->SetColorTint(NewShotColorTint);
 
 			MovieSceneToolHelpers::SetTakeNumber(NewShot, NewTakeNumber);
 
@@ -611,6 +613,7 @@ void FSubTrackEditor::ChangeTake(UMovieSceneSequence* Sequence)
 		int32                NewRowIndex = Section->GetRowIndex();
 		FFrameNumber         NewSectionStartTime = NewSectionRange.GetLowerBound().IsClosed() ? UE::MovieScene::DiscreteInclusiveLower(NewSectionRange) : 0;
 		int32                NewSectionRowIndex = Section->GetRowIndex();
+		FColor               NewSectionColorTint = Section->GetColorTint();
 
 		const int32 Duration = (NewSectionRange.GetLowerBound().IsClosed() && NewSectionRange.GetUpperBound().IsClosed()) ? UE::MovieScene::DiscreteSize(NewSectionRange) : 1;
 		UMovieSceneSubSection* NewSection = SubTrack->AddSequence(Sequence, NewSectionStartTime, Duration);
@@ -624,6 +627,7 @@ void FSubTrackEditor::ChangeTake(UMovieSceneSequence* Sequence)
 			NewSection->Parameters.TimeScale = NewSectionTimeScale;
 			NewSection->SetPreRollFrames(NewSectionPrerollFrames);
 			NewSection->SetRowIndex(NewSectionRowIndex);
+			NewSection->SetColorTint(NewSectionColorTint);
 
 			bChangedTake = true;
 		}
