Sequencer: Linked Anim Sequence: Only dirty if the data is changed not if refreshing or changing sequences.

diff --git a/Engine/Source/Editor/MovieSceneTools/Private/TrackEditors/SkeletalAnimationTrackEditor.cpp b/Engine/Source/Editor/MovieSceneTools/Private/TrackEditors/SkeletalAnimationTrackEditor.cpp
--- a/Engine/Source/Editor/MovieSceneTools/Private/TrackEditors/SkeletalAnimationTrackEditor.cpp
+++ b/Engine/Source/Editor/MovieSceneTools/Private/TrackEditors/SkeletalAnimationTrackEditor.cpp
@@ -1266,6 +1266,13 @@ void FSkeletalAnimationTrackEditor::OnSequencerSaved(ISequencer& )
 //dirty anim sequence when the sequencer changes, to make sure it get's checked out etc..
 void FSkeletalAnimationTrackEditor::OnSequencerDataChanged(EMovieSceneDataChangeType DataChangeType)
 {
+	//only return if data really changed
+	if(DataChangeType ==  EMovieSceneDataChangeType::RefreshTree ||
+		DataChangeType == EMovieSceneDataChangeType::ActiveMovieSceneChanged ||
+		DataChangeType == EMovieSceneDataChangeType::RefreshAllImmediately)
+	{
+		return;
+	}
 	TSharedPtr<ISequencer> SequencerPtr = GetSequencer();
 	if (!SequencerPtr.IsValid())
 	{
