Sequencer: Constraints: When baking only remove transform keys from the channels we are baking when setting up the bake

diff --git a/Engine/Plugins/Animation/ControlRig/Source/ControlRigEditor/Private/Tools/ConstraintBaker.cpp b/Engine/Plugins/Animation/ControlRig/Source/ControlRigEditor/Private/Tools/ConstraintBaker.cpp
--- a/Engine/Plugins/Animation/ControlRig/Source/ControlRigEditor/Private/Tools/ConstraintBaker.cpp
+++ b/Engine/Plugins/Animation/ControlRig/Source/ControlRigEditor/Private/Tools/ConstraintBaker.cpp
@@ -331,16 +331,19 @@ void FConstraintBaker::Bake(UWorld* InWorld,
 		}
 	}
 	//remove transform keys if baking frame
+	const EMovieSceneTransformChannel ChannelsToKey = InConstraint->GetChannelsToKey();
+
 	if (InSettings.IsSet() && InSettings.GetValue().BakingKeySettings == EBakingKeySettings::AllFrames)
 	{
 		if (FloatTransformChannels.Num() > 0)
 		{
-			FMovieSceneConstraintChannelHelper::DeleteTransformTimes(FloatTransformChannels, FramesToBake[0] +1, FramesToBake[FramesToBake.Num() -1]);
+			FMovieSceneConstraintChannelHelper::DeleteTransformTimes(FloatTransformChannels, FramesToBake[0] +1, FramesToBake[FramesToBake.Num() -1],
+				ChannelsToKey);
 		}
 		else if (DoubleTransformChannels.Num() > 0)
 		{
-			FMovieSceneConstraintChannelHelper::DeleteTransformTimes(DoubleTransformChannels, FramesToBake[0] + 1, FramesToBake[FramesToBake.Num() - 1]);
-
+			FMovieSceneConstraintChannelHelper::DeleteTransformTimes(DoubleTransformChannels, FramesToBake[0] + 1, FramesToBake[FramesToBake.Num() - 1],
+				ChannelsToKey);
 		}
 	}
 
diff --git a/Engine/Source/Editor/MovieSceneTools/Public/Constraints/MovieSceneConstraintChannelHelper.h b/Engine/Source/Editor/MovieSceneTools/Public/Constraints/MovieSceneConstraintChannelHelper.h
--- a/Engine/Source/Editor/MovieSceneTools/Public/Constraints/MovieSceneConstraintChannelHelper.h
+++ b/Engine/Source/Editor/MovieSceneTools/Public/Constraints/MovieSceneConstraintChannelHelper.h
@@ -18,6 +18,7 @@ struct FMovieSceneConstraintChannel;
 struct FFrameNumber;
 struct FConstraintAndActiveChannel;
 enum class EMovieSceneKeyInterpolation : uint8;
+enum class EMovieSceneTransformChannel : uint32;
 
 struct MOVIESCENETOOLS_API FCompensationEvaluator
 {
@@ -169,7 +170,8 @@ public:
 	static  void DeleteTransformTimes(
 		const TArrayView<ChannelType*>& InChannels,
 		const FFrameNumber& StartTime,
-		const FFrameNumber& EndTime);
+		const FFrameNumber& EndTime,
+		EMovieSceneTransformChannel Channels = EMovieSceneTransformChannel::AllTransform);
 
 	/** this will only set the value son channels with keys at the specified time, reusing tangent time etc. */
 	template< typename ChannelType >
diff --git a/Engine/Source/Editor/MovieSceneTools/Public/Constraints/MovieSceneConstraintChannelHelper.inl b/Engine/Source/Editor/MovieSceneTools/Public/Constraints/MovieSceneConstraintChannelHelper.inl
--- a/Engine/Source/Editor/MovieSceneTools/Public/Constraints/MovieSceneConstraintChannelHelper.inl
+++ b/Engine/Source/Editor/MovieSceneTools/Public/Constraints/MovieSceneConstraintChannelHelper.inl
@@ -357,15 +357,35 @@ template< typename ChannelType >
  void FMovieSceneConstraintChannelHelper::DeleteTransformTimes(
 	 const TArrayView<ChannelType*>& InChannels,
 	 const FFrameNumber& StartTime,
-	 const FFrameNumber& EndTime)
+	 const FFrameNumber& EndTime,
+	 EMovieSceneTransformChannel InChannelsToKey)
 {
 	 TRange<FFrameNumber> WithinRange(0, 0);
 	 WithinRange.SetLowerBoundValue(StartTime);
 	 WithinRange.SetUpperBoundValue(EndTime);
 	 TArray<FFrameNumber> KeyTimes;
 	 TArray<FKeyHandle> KeyHandles;
-	 for (ChannelType* Channel : InChannels)
+
+	 const bool bKeyTranslation = EnumHasAllFlags(InChannelsToKey, EMovieSceneTransformChannel::Translation);
+	 const bool bKeyRotation = EnumHasAllFlags(InChannelsToKey, EMovieSceneTransformChannel::Rotation);
+	 const bool bKeyScale = EnumHasAllFlags(InChannelsToKey, EMovieSceneTransformChannel::Scale);
+
+	 TArray<uint32> ChannelsIndexToKey;
+	 if (bKeyTranslation)
+	 {
+		 ChannelsIndexToKey.Append({ 0,1,2 });
+	 }
+	 if (bKeyRotation)
+	 {
+		 ChannelsIndexToKey.Append({ 3,4,5 });
+	 }
+	 if (bKeyScale)
+	 {
+		 ChannelsIndexToKey.Append({ 6,7,8 });
+	 }
+	 for (const int32 ChannelIndex : ChannelsIndexToKey)
 	 {
+		 ChannelType* Channel = InChannels[ChannelIndex];
 		 TMovieSceneChannelData<typename ChannelType::ChannelValueType> Data = Channel->GetData();
 		 KeyTimes.SetNum(0);
 		 KeyHandles.SetNum(0);
