Nonunity build error
#trivial

diff --git a/Engine/Source/Editor/MovieSceneTools/Public/Constraints/MovieSceneConstraintChannelHelper.h b/Engine/Source/Editor/MovieSceneTools/Public/Constraints/MovieSceneConstraintChannelHelper.h
--- a/Engine/Source/Editor/MovieSceneTools/Public/Constraints/MovieSceneConstraintChannelHelper.h
+++ b/Engine/Source/Editor/MovieSceneTools/Public/Constraints/MovieSceneConstraintChannelHelper.h
@@ -4,6 +4,8 @@
 
 #include "Templates/SharedPointer.h"
 #include "UObject/ObjectPtr.h"
+#include "MovieSceneSection.h"
+
 
 class IMovieSceneConstrainedSection;
 struct ITransformConstraintChannelInterface;
@@ -18,7 +20,6 @@ struct FMovieSceneConstraintChannel;
 struct FFrameNumber;
 struct FConstraintAndActiveChannel;
 enum class EMovieSceneKeyInterpolation : uint8;
-enum class EMovieSceneTransformChannel : uint32;
 
 struct MOVIESCENETOOLS_API FCompensationEvaluator
 {
