Sequencer: Control Rig: FBX Import was not taking into account negation values when importing onto transforms anymore

diff --git a/Engine/Source/Editor/MovieSceneTools/Private/MovieSceneToolHelpers.cpp b/Engine/Source/Editor/MovieSceneTools/Private/MovieSceneToolHelpers.cpp
--- a/Engine/Source/Editor/MovieSceneTools/Private/MovieSceneToolHelpers.cpp
+++ b/Engine/Source/Editor/MovieSceneTools/Private/MovieSceneToolHelpers.cpp
@@ -1754,9 +1754,15 @@ static bool ImportFBXTransformToChannels(FString NodeName, const UMovieSceneUser
 		for (int i = 0; i < FloatChannels.Num(); i++)
 		{
 			FloatChannels[i]->SetDefault(DefaultTransformValues[i]);
-
-			const bool bNegate = i == 1 || i == 4 || i == 5;
-			ImportTransformChannelToFloat(Transform[i], FloatChannels[i], FrameRate, bNegate, false, StartFrame);
+			ChannelEnumIndex = (int)FControlRigChannelEnum::PositionX;
+			bool bNegate = i == 1 || i == 4 || i == 5;
+			const FControlRigChannelEnum ChannelEnum = (FControlRigChannelEnum)(ChannelEnumIndex + i);
+			if (const FControlToTransformMappings* ChannelMapping = ImportFBXControlRigSettings->ControlChannelMappings.FindByPredicate(
+				[ChannelEnum](const FControlToTransformMappings& Mapping) { return ChannelEnum == Mapping.ControlChannel; }))
+			{
+				bNegate = ChannelMapping->bNegate ? !bNegate : bNegate;
+			}
+			ImportTransformChannelToFloat(Transform[i], FloatChannels[i], FrameRate, bNegate, false, StartFrame, bNegate);
 		}
 	}
 
