Sequencer: Curve Editor Also Keep tangent type when setting keys if a previous one is set

diff --git a/Engine/Source/Editor/CurveEditor/Private/Views/SInteractiveCurveEditorView.cpp b/Engine/Source/Editor/CurveEditor/Private/Views/SInteractiveCurveEditorView.cpp
--- a/Engine/Source/Editor/CurveEditor/Private/Views/SInteractiveCurveEditorView.cpp
+++ b/Engine/Source/Editor/CurveEditor/Private/Views/SInteractiveCurveEditorView.cpp
@@ -1077,8 +1077,10 @@ FReply SInteractiveCurveEditorView::OnMouseButtonDown(const FGeometry& MyGeometr
 					FKeyAttributes KeyAttributes = CurveEditor->GetDefaultKeyAttributes().Get();
 					if (KeyAttributes.HasInterpMode())
 					{
-						// Set interpolation mode based on surrounding keys, if any
-						KeyAttributes.SetInterpMode(CurveToAddTo->GetInterpolationMode(MouseTime, KeyAttributes.GetInterpMode()));
+						// Set interpolation and tangent mode based on surrounding keys, if any
+						TPair<ERichCurveInterpMode, ERichCurveTangentMode> Modes = CurveToAddTo->GetInterpolationMode(MouseTime, KeyAttributes.GetInterpMode(), KeyAttributes.GetTangentMode());
+						KeyAttributes.SetInterpMode(Modes.Key);
+						KeyAttributes.SetTangentMode(Modes.Value);
 					}
 
 					FCurveSnapMetrics SnapMetrics = CurveEditor->GetCurveSnapMetrics(HoveredCurve.GetValue());
@@ -1563,8 +1565,10 @@ void SInteractiveCurveEditorView::AddKeyAtTime(const TSet<FCurveModelID>& ToCurv
 				FKeyAttributes KeyAttributes = CurveEditor->GetDefaultKeyAttribute().Get();
 				if (KeyAttributes.HasInterpMode())
 				{
-					// Set interpolation mode based on surrounding keys, if any
-					KeyAttributes.SetInterpMode(CurveModel->GetInterpolationMode(EvalTime, KeyAttributes.GetInterpMode()));
+					// Set interpolation and tangent mode based on surrounding keys, if any
+					TPair<ERichCurveInterpMode, ERichCurveTangentMode> Modes = CurveModel->GetInterpolationMode(EvalTime, KeyAttributes.GetInterpMode(), KeyAttributes.GetTangentMode());
+					KeyAttributes.SetInterpMode(Modes.Key);
+					KeyAttributes.SetTangentMode(Modes.Value);
 				}
 
 				// Add a key on this curve
diff --git a/Engine/Source/Editor/CurveEditor/Public/CurveModel.h b/Engine/Source/Editor/CurveEditor/Public/CurveModel.h
--- a/Engine/Source/Editor/CurveEditor/Public/CurveModel.h
+++ b/Engine/Source/Editor/CurveEditor/Public/CurveModel.h
@@ -5,6 +5,7 @@
 #include "Containers/Array.h"
 #include "Containers/ArrayView.h"
 #include "Containers/UnrealString.h"
+#include "Containers/Map.h"
 #include "CoreTypes.h"
 #include "CurveEditorTypes.h"
 #include "Curves/RichCurve.h"
@@ -157,7 +158,10 @@ public:
 	 * @param DefaultInterpolationMode		Current default interpolation mode, returned if other keys not found or interpolation not supported
 	 * @return Interpolation mode to use at that frame
 	 */
-	virtual ERichCurveInterpMode GetInterpolationMode(const double& InTime, ERichCurveInterpMode DefaultInterpolationMode) const { return DefaultInterpolationMode; }
+	virtual TPair<ERichCurveInterpMode, ERichCurveTangentMode> GetInterpolationMode(const double& InTime, ERichCurveInterpMode DefaultInterpolationMode, ERichCurveTangentMode DefaultTangentMode) const 
+	{
+		return TPair<ERichCurveInterpMode, ERichCurveTangentMode>(DefaultInterpolationMode, DefaultTangentMode);
+	}
 
 	/**
 	 * Evaluate this curve at the specified time
diff --git a/Engine/Source/Editor/MovieSceneTools/Private/Channels/BezierChannelCurveModel.cpp b/Engine/Source/Editor/MovieSceneTools/Private/Channels/BezierChannelCurveModel.cpp
--- a/Engine/Source/Editor/MovieSceneTools/Private/Channels/BezierChannelCurveModel.cpp
+++ b/Engine/Source/Editor/MovieSceneTools/Private/Channels/BezierChannelCurveModel.cpp
@@ -153,7 +153,7 @@ void FBezierChannelCurveModel<ChannelType, ChannelValue, KeyType>::GetKeyDrawInf
 }
 
 template <class ChannelType, class ChannelValue, class KeyType>
-ERichCurveInterpMode FBezierChannelCurveModel<ChannelType, ChannelValue, KeyType>::GetInterpolationMode(const double& InTime, ERichCurveInterpMode DefaultInterpolationMode) const
+TPair<ERichCurveInterpMode, ERichCurveTangentMode> FBezierChannelCurveModel<ChannelType, ChannelValue, KeyType>::GetInterpolationMode(const double& InTime, ERichCurveInterpMode DefaultInterpolationMode, ERichCurveTangentMode DefaultTangentMode) const
 {
 	ChannelType* Channel = this->GetChannelHandle().Get();
 	UMovieSceneSection* Section = Cast<UMovieSceneSection>(this->GetOwningObject());
@@ -179,11 +179,19 @@ ERichCurveInterpMode FBezierChannelCurveModel<ChannelType, ChannelValue, KeyType
 			TArray<FKeyAttributes> KeyAttributes;
 			KeyAttributes.SetNum(1);
 			GetKeyAttributes(InKey, KeyAttributes);
-			return KeyAttributes[0].GetInterpMode();
+			ERichCurveInterpMode InterpMode = KeyAttributes[0].GetInterpMode();
+			ERichCurveTangentMode TangentMode = KeyAttributes[0].HasTangentMode() ? KeyAttributes[0].GetTangentMode() : DefaultTangentMode;
+			//if we are cubic, with anything but auto tangents we use the default instead, since they will give us flat tangents which aren't good
+			if (InterpMode == ERichCurveInterpMode::RCIM_Cubic &&
+				(TangentMode != ERichCurveTangentMode::RCTM_Auto && TangentMode != ERichCurveTangentMode::RCTM_SmartAuto))
+			{
+				TangentMode = DefaultTangentMode;
+			}
+			return TPair<ERichCurveInterpMode, ERichCurveTangentMode>(InterpMode, TangentMode);
 		}
 	}
 
-	return DefaultInterpolationMode;
+	return TPair<ERichCurveInterpMode, ERichCurveTangentMode>(DefaultInterpolationMode, DefaultTangentMode);
 }
 
 template<typename ChannelType, typename ChannelValue, typename KeyType> 
diff --git a/Engine/Source/Editor/MovieSceneTools/Public/Channels/BezierChannelCurveModel.h b/Engine/Source/Editor/MovieSceneTools/Public/Channels/BezierChannelCurveModel.h
--- a/Engine/Source/Editor/MovieSceneTools/Public/Channels/BezierChannelCurveModel.h
+++ b/Engine/Source/Editor/MovieSceneTools/Public/Channels/BezierChannelCurveModel.h
@@ -63,7 +63,7 @@ public:
 	virtual void GetKeyAttributes(TArrayView<const FKeyHandle> InKeys, TArrayView<FKeyAttributes> OutAttributes) const override;
 	virtual void SetKeyAttributes(TArrayView<const FKeyHandle> InKeys, TArrayView<const FKeyAttributes> InAttributes, EPropertyChangeType::Type ChangeType = EPropertyChangeType::Unspecified) override;
 
-	virtual ERichCurveInterpMode GetInterpolationMode(const double& InTime, ERichCurveInterpMode DefaultInterpolationMode) const override;
+	virtual TPair<ERichCurveInterpMode, ERichCurveTangentMode> GetInterpolationMode(const double& InTime, ERichCurveInterpMode DefaultInterpolationMode, ERichCurveTangentMode DefaultTangentMode) const override;
 	virtual void GetCurveAttributes(FCurveAttributes& OutCurveAttributes) const override;
 	virtual void SetCurveAttributes(const FCurveAttributes& InCurveAttributes) override;
 	virtual void GetValueRange(double& MinValue, double& MaxValue) const override;
