Sequencer: Allow keying when defaults have been cleared and there are no valid entities to recompose from. The current value will just be keyed.

diff --git a/Engine/Source/Editor/MovieSceneTools/Private/TrackEditors/TransformTrackEditor.cpp b/Engine/Source/Editor/MovieSceneTools/Private/TrackEditors/TransformTrackEditor.cpp
--- a/Engine/Source/Editor/MovieSceneTools/Private/TrackEditors/TransformTrackEditor.cpp
+++ b/Engine/Source/Editor/MovieSceneTools/Private/TrackEditors/TransformTrackEditor.cpp
@@ -1134,7 +1134,7 @@ void F3DTransformTrackEditor::ProcessKeyOperation(UObject* ObjectToKey, TArrayVi
 
 	UMovieSceneInterrogatedPropertyInstantiatorSystem* System = Interrogator.GetLinker()->FindSystem<UMovieSceneInterrogatedPropertyInstantiatorSystem>();
 
-	if (ensure(System && ValidEntities.Num() != 0))
+	if (ensure(System))
 	{
 		FDecompositionQuery Query;
 		Query.Entities = ValidEntities;
