Sequencer: Allow for picking class properties which are assigned to an object property track.

If a property is a ClassProperty, the MetaClass is assigned as the PropertyClass to animate and an additional bClassProperty is set so that the channel editor can display a SClassPropertyEntryBox instead of an SObjectPropertyEntryBox

diff --git a/Engine/Source/Editor/MovieSceneTools/Private/Channels/BuiltInChannelEditors.cpp b/Engine/Source/Editor/MovieSceneTools/Private/Channels/BuiltInChannelEditors.cpp
--- a/Engine/Source/Editor/MovieSceneTools/Private/Channels/BuiltInChannelEditors.cpp
+++ b/Engine/Source/Editor/MovieSceneTools/Private/Channels/BuiltInChannelEditors.cpp
@@ -46,6 +46,7 @@
 #include "EntitySystem/Interrogation/MovieSceneInterrogationLinker.h"
 #include "EntitySystem/Interrogation/MovieSceneInterrogatedPropertyInstantiator.h"
 #include "Systems/MovieScenePropertyInstantiator.h"
+#include "Tracks/MovieSceneObjectPropertyTrack.h"
 #include "Tracks/MovieScenePropertyTrack.h"
 #include "Widgets/Input/SComboButton.h"
 #include "MovieSceneSpawnableAnnotation.h"
@@ -401,49 +402,87 @@ TSharedRef<SWidget> CreateKeyEditor(const TMovieSceneChannelHandle<FMovieSceneOb
 {
 	const TMovieSceneExternalValue<UObject*>* ExternalValue = Channel.GetExtendedEditorData();
 	const FMovieSceneObjectPathChannel*       RawChannel    = Channel.Get();
+	UMovieSceneObjectPropertyTrack* ObjectPathTrack = Cast<UMovieSceneObjectPropertyTrack>(Section->GetOuter());
+
 	if (ExternalValue && RawChannel)
 	{
 		TSequencerKeyEditor<FMovieSceneObjectPathChannel, UObject*> KeyEditor(InObjectBindingID, Channel, Section, InSequencer, PropertyBindings, ExternalValue->OnGetExternalValue);
 
-		auto OnSetObjectLambda = [KeyEditor](const FAssetData& Asset) mutable
+		UClass* PropertyClass = ObjectPathTrack ? ObjectPathTrack->PropertyClass : nullptr;
+		const bool bClassPicker = ObjectPathTrack ? ObjectPathTrack->bClassProperty : false;
+		if (bClassPicker)
 		{
-			FScopedTransaction Transaction(LOCTEXT("SetEnumKey", "Set Enum Key Value"));
-			KeyEditor.SetValueWithNotify(Asset.GetAsset(), EMovieSceneDataChangeType::TrackValueChangedRefreshImmediately);
-		};
+			auto OnSetClassLambda = [KeyEditor](const UClass* Class) mutable
+			{
+				FScopedTransaction Transaction(LOCTEXT("SetObjectPathKey", "Set Object Path Key Value"));
+				KeyEditor.SetValueWithNotify(const_cast<UClass*>(Class), EMovieSceneDataChangeType::TrackValueChangedRefreshImmediately);
+			};
 
-		auto GetObjectPathLambda = [KeyEditor]() -> FString
-		{
-			UObject* Obj = KeyEditor.GetCurrentValue();
-			return Obj ? Obj->GetPathName() : FString();
-		};
+			auto GetSelectedClassLambda = [KeyEditor]() -> const UClass*
+			{
+				return Cast<UClass>(KeyEditor.GetCurrentValue());
+			};
 
-		TArray<FAssetData> AssetDataArray;
-		if (InSequencer.IsValid())
-		{
-			UMovieSceneSequence* Sequence = InSequencer.Pin()->GetFocusedMovieSceneSequence();
-			AssetDataArray.Add((FAssetData)Sequence);
+			return SNew(SHorizontalBox)
+
+				+ SHorizontalBox::Slot()
+				.AutoWidth()
+				[
+					SNew(SClassPropertyEntryBox)
+					.MetaClass(PropertyClass)
+					.SelectedClass_Lambda(GetSelectedClassLambda)
+					.OnSetClass_Lambda(OnSetClassLambda)
+				]
+
+				+ SHorizontalBox::Slot()
+				.AutoWidth()
+				.Padding(8.0f, 0.0f)
+				[
+					SNew(SSpacer)
+				];
 		}
+		else
+		{
+			auto OnSetObjectLambda = [KeyEditor](const FAssetData& Asset) mutable
+			{
+				FScopedTransaction Transaction(LOCTEXT("SetObjectPathKey", "Set Object Path Key Value"));
+				KeyEditor.SetValueWithNotify(Asset.GetAsset(), EMovieSceneDataChangeType::TrackValueChangedRefreshImmediately);
+			};
 
-		return SNew(SHorizontalBox)
+			auto GetObjectPathLambda = [KeyEditor]() -> FString
+			{
+				UObject* Obj = KeyEditor.GetCurrentValue();
+				return Obj ? Obj->GetPathName() : FString();
+			};
 
-			+ SHorizontalBox::Slot()
-			.AutoWidth()
-			[
-				SNew(SObjectPropertyEntryBox)
-				.DisplayBrowse(true)
-				.DisplayUseSelected(false)
-				.ObjectPath_Lambda(GetObjectPathLambda)
-				.AllowedClass(RawChannel->GetPropertyClass())
-				.OnObjectChanged_Lambda(OnSetObjectLambda)
-				.OwnerAssetDataArray(AssetDataArray)
-			]
+			TArray<FAssetData> AssetDataArray;
+			if (InSequencer.IsValid())
+			{
+				UMovieSceneSequence* Sequence = InSequencer.Pin()->GetFocusedMovieSceneSequence();
+				AssetDataArray.Add((FAssetData)Sequence);
+			}
 
-			+ SHorizontalBox::Slot()
-			.AutoWidth()
-			.Padding(8.0f, 0.0f)
-			[
-				SNew(SSpacer)
-			];
+			return SNew(SHorizontalBox)
+
+				+ SHorizontalBox::Slot()
+				.AutoWidth()
+				[
+					SNew(SObjectPropertyEntryBox)
+					.DisplayBrowse(true)
+					.DisplayUseSelected(false)
+					.ObjectPath_Lambda(GetObjectPathLambda)
+					.AllowedClass(RawChannel->GetPropertyClass())
+					.OnObjectChanged_Lambda(OnSetObjectLambda)
+					.OwnerAssetDataArray(AssetDataArray)
+				]
+
+				+ SHorizontalBox::Slot()
+				.AutoWidth()
+				.Padding(8.0f, 0.0f)
+				[
+					SNew(SSpacer)
+				];
+		}
 	}
 
 	return SNullWidget::NullWidget;
diff --git a/Engine/Source/Editor/MovieSceneTools/Private/TrackEditors/ObjectPropertyTrackEditor.cpp b/Engine/Source/Editor/MovieSceneTools/Private/TrackEditors/ObjectPropertyTrackEditor.cpp
--- a/Engine/Source/Editor/MovieSceneTools/Private/TrackEditors/ObjectPropertyTrackEditor.cpp
+++ b/Engine/Source/Editor/MovieSceneTools/Private/TrackEditors/ObjectPropertyTrackEditor.cpp
@@ -50,6 +50,16 @@ void FObjectPropertyTrackEditor::InitializeNewTrack(UMovieSceneObjectPropertyTra
 	{
 		NewTrack->PropertyClass = KeyedProperty->PropertyClass;
 
+		if (KeyedProperty->HasAllPropertyFlags(CPF_UObjectWrapper))
+		{
+			FClassProperty* ClassProperty = CastField<FClassProperty>(KeyedProperty);
+			if (ClassProperty)
+			{
+				NewTrack->PropertyClass = ClassProperty->MetaClass;
+				NewTrack->bClassProperty = true;
+			}
+		}
+
 		FAssetToolsModule& AssetToolsModule = FModuleManager::LoadModuleChecked<FAssetToolsModule>("AssetTools");
 		TArray< TWeakPtr<IAssetTypeActions> > AssetTypeActions;
 		AssetToolsModule.Get().GetAssetTypeActionsList(AssetTypeActions);
diff --git a/Engine/Source/Runtime/MovieSceneTracks/Public/Tracks/MovieSceneObjectPropertyTrack.h b/Engine/Source/Runtime/MovieSceneTracks/Public/Tracks/MovieSceneObjectPropertyTrack.h
--- a/Engine/Source/Runtime/MovieSceneTracks/Public/Tracks/MovieSceneObjectPropertyTrack.h
+++ b/Engine/Source/Runtime/MovieSceneTracks/Public/Tracks/MovieSceneObjectPropertyTrack.h
@@ -17,6 +17,11 @@ public:
 	UPROPERTY()
 	TObjectPtr<UClass> PropertyClass;
 
+#if WITH_EDITORONLY_DATA
+	UPROPERTY()
+	bool bClassProperty = false;
+#endif
+
 	UMovieSceneObjectPropertyTrack(const FObjectInitializer& ObjInit);
 
 	/*~ UMovieSceneTrack interface */
