Constraints: make sure that the delegate is unbound when closing the track editor

#rnx

diff --git a/Engine/Plugins/Animation/ControlRig/Source/ControlRig/Private/Sequencer/MovieSceneControlRigParameterTrack.cpp b/Engine/Plugins/Animation/ControlRig/Source/ControlRig/Private/Sequencer/MovieSceneControlRigParameterTrack.cpp
--- a/Engine/Plugins/Animation/ControlRig/Source/ControlRig/Private/Sequencer/MovieSceneControlRigParameterTrack.cpp
+++ b/Engine/Plugins/Animation/ControlRig/Source/ControlRig/Private/Sequencer/MovieSceneControlRigParameterTrack.cpp
@@ -74,7 +74,10 @@ UMovieSceneSection* UMovieSceneControlRigParameterTrack::CreateNewSection()
 	}
 
 	NewSection->SpaceChannelAdded().AddUObject(this, &UMovieSceneControlRigParameterTrack::HandleOnSpaceAdded);
-	NewSection->ConstraintChannelAdded().AddUObject(this, &UMovieSceneControlRigParameterTrack::HandleOnConstraintAdded);
+	if (!NewSection->ConstraintChannelAdded().IsBoundToObject(this))
+	{
+		NewSection->ConstraintChannelAdded().AddUObject(this, &UMovieSceneControlRigParameterTrack::HandleOnConstraintAdded);
+	}
 
 	if (ControlRig)
 	{
diff --git a/Engine/Source/Editor/MovieSceneTools/Private/Constraints/MovieSceneConstraintChannelHelper.cpp b/Engine/Source/Editor/MovieSceneTools/Private/Constraints/MovieSceneConstraintChannelHelper.cpp
--- a/Engine/Source/Editor/MovieSceneTools/Private/Constraints/MovieSceneConstraintChannelHelper.cpp
+++ b/Engine/Source/Editor/MovieSceneTools/Private/Constraints/MovieSceneConstraintChannelHelper.cpp
@@ -1200,10 +1200,15 @@ void FMovieSceneConstraintChannelHelper::HandleConstraintPropertyChanged(
 		return CompensateScale(Cast<UTickableParentConstraint>(InConstraint), InActiveChannel, InSequencer, InSection);
 	}
 	
-	if (PropertyName == GET_MEMBER_NAME_CHECKED(UTickableTranslationConstraint, OffsetTranslation) ||
-		PropertyName == GET_MEMBER_NAME_CHECKED(UTickableRotationConstraint, OffsetRotation) ||
-		PropertyName == GET_MEMBER_NAME_CHECKED(UTickableScaleConstraint, OffsetScale) ||
-		PropertyName == GET_MEMBER_NAME_CHECKED(UTickableParentConstraint, OffsetTransform))
+	auto IsOffsetProperty = [](const FName InPropertyName)
+	{
+		return InPropertyName == GET_MEMBER_NAME_CHECKED(UTickableTranslationConstraint, OffsetTranslation) ||
+			InPropertyName == GET_MEMBER_NAME_CHECKED(UTickableRotationConstraint, OffsetRotation) ||
+			InPropertyName == GET_MEMBER_NAME_CHECKED(UTickableScaleConstraint, OffsetScale) ||
+			InPropertyName == GET_MEMBER_NAME_CHECKED(UTickableParentConstraint, OffsetTransform);
+	};
+	
+	if (IsOffsetProperty(PropertyName) || IsOffsetProperty(InPropertyChangedEvent.GetMemberPropertyName()))
 	{
 		return HandleOffsetChanged(InConstraint, InActiveChannel, InSequencer);
 	}
diff --git a/Engine/Source/Editor/MovieSceneTools/Private/TrackEditors/TransformTrackEditor.cpp b/Engine/Source/Editor/MovieSceneTools/Private/TrackEditors/TransformTrackEditor.cpp
--- a/Engine/Source/Editor/MovieSceneTools/Private/TrackEditors/TransformTrackEditor.cpp
+++ b/Engine/Source/Editor/MovieSceneTools/Private/TrackEditors/TransformTrackEditor.cpp
@@ -171,7 +171,10 @@ TSharedRef<ISequencerSection> F3DTransformTrackEditor::MakeSectionInterface(UMov
 	UMovieScene3DTransformSection* Section = Cast<UMovieScene3DTransformSection>(&SectionObject);
 	if (Section)
 	{
-		Section->ConstraintChannelAdded().AddRaw(this, &F3DTransformTrackEditor::HandleOnConstraintAdded);
+		if (!Section->ConstraintChannelAdded().IsBoundToObject(this))
+		{
+			Section->ConstraintChannelAdded().AddRaw(this, &F3DTransformTrackEditor::HandleOnConstraintAdded);
+		}
 	}
 	//if there are channels already we need to act like they were added
 	TArray<FConstraintAndActiveChannel>& Channels = Section->GetConstraintsChannels();
@@ -1685,6 +1688,8 @@ void F3DTransformTrackEditor::ClearOutConstraintDelegates()
 			{
 				CRSection->OnConstraintRemovedHandle.Reset();
 			}
+
+			CRSection->ConstraintChannelAdded().RemoveAll(this);
 		}
 	}
 	SectionsToClear.Reset();
