Sequencer: NLA: We were using the actual location of the bone not the modified widget location of the bone which is dropped so Z = 0 when manipulating which can cause double transforms

diff --git a/Engine/Source/Editor/MovieSceneTools/Private/EditModes/SkeletalAnimationTrackEditMode.cpp b/Engine/Source/Editor/MovieSceneTools/Private/EditModes/SkeletalAnimationTrackEditMode.cpp
--- a/Engine/Source/Editor/MovieSceneTools/Private/EditModes/SkeletalAnimationTrackEditMode.cpp
+++ b/Engine/Source/Editor/MovieSceneTools/Private/EditModes/SkeletalAnimationTrackEditMode.cpp
@@ -532,18 +532,47 @@ bool FSkeletalAnimationTrackEditMode::HandleClick(FEditorViewportClient* InViewp
 
 			if (Sequencer.IsValid())
 			{
+				if (FSlateApplication::Get().GetModifierKeys().IsShiftDown() == false &&
+					FSlateApplication::Get().GetModifierKeys().IsControlDown() == false)
+				{
+					SelectedRootData.Empty();
+					if (GEditor)
+					{
+						GEditor->Exec(GetWorld(), TEXT("SELECT NONE"));
+					}
+				}
 				UMovieSceneSkeletalAnimationSection* Section = Proxy->AnimSection.Get();
 				USkeletalMeshComponent* Comp = Proxy->SkelMeshComp.Get();
+
 				if (Section && Comp)
 				{
 					FSelectedRootData RootData(Section, Comp);
-					SelectedRootData.Add(RootData);
+
+					if (FSlateApplication::Get().GetModifierKeys().IsControlDown() == true)
+					{
+						if (SelectedRootData.Contains(RootData))
+						{
+							SelectedRootData.Remove(RootData);
+						}
+						else
+						{
+							SelectedRootData.Add(RootData);
+						}
+					}
+					else
+					{
+						if (SelectedRootData.Contains(RootData) == false)
+						{
+							SelectedRootData.Add(RootData);
+						}
+					}
 				}
 			}
 		}
 	}
 	return false;
 }
+
 FSelectedRootData::FSelectedRootData(UMovieSceneSkeletalAnimationSection* InSection, USkeletalMeshComponent* InComp) : 
 	SelectedSection(InSection), SelectedMeshComp(InComp)
 {
@@ -592,25 +621,8 @@ void FSkeletalAnimationTrackEditMode::SelectNone()
 
 void FSkeletalAnimationTrackEditMode::OnKeySelected(FViewport* Viewport, HMovieSceneSkeletalAnimationRootHitProxy* Proxy)
 {
-	if (!Proxy)
-	{
-		SelectedRootData.Empty();
-		return;
-	}
-
-	TSharedPtr<ISequencer> Sequencer = WeakSequencer.Pin();
-	if (Sequencer.IsValid())
-	{
-		UMovieSceneSkeletalAnimationSection* Section = Proxy->AnimSection.Get();
-		USkeletalMeshComponent* Comp = Proxy->SkelMeshComp.Get();
-		SelectedRootData.Empty();
-
-		if (Section && Comp)
-		{
-			FSelectedRootData RootData(Section, Comp);
-			SelectedRootData.Add(RootData);
-		}
-	}
+	//something changed but this is no longer needed in 5.3
+	return;
 }
 
 bool FSkeletalAnimationTrackEditMode::InputDelta(FEditorViewportClient* InViewportClient, FViewport* InViewport, FVector& InDrag, FRotator& InRot, FVector& InScale)
@@ -645,7 +657,7 @@ bool FSkeletalAnimationTrackEditMode::InputDelta(FEditorViewportClient* InViewpo
 
 				if (bDoTranslation)
 				{
-					FVector CurrentLocation = CurrentTransform.GetLocation();
+					FVector CurrentLocation = GetWidgetLocation();
 					CurrentLocation = CurrentLocation + InDrag;
 					CurrentTransform.SetLocation(CurrentLocation);
 				}
diff --git a/Engine/Source/Editor/MovieSceneTools/Private/EditModes/SkeletalAnimationTrackEditMode.h b/Engine/Source/Editor/MovieSceneTools/Private/EditModes/SkeletalAnimationTrackEditMode.h
--- a/Engine/Source/Editor/MovieSceneTools/Private/EditModes/SkeletalAnimationTrackEditMode.h
+++ b/Engine/Source/Editor/MovieSceneTools/Private/EditModes/SkeletalAnimationTrackEditMode.h
@@ -16,8 +16,9 @@ struct HMovieSceneSkeletalAnimationRootHitProxy;
 struct FSelectedRootData
 {
 	FSelectedRootData(UMovieSceneSkeletalAnimationSection* InSection, USkeletalMeshComponent* InComp);
-
+	bool operator == (const FSelectedRootData& InData) const { return (InData.SelectedSection == SelectedSection && InData.SelectedMeshComp == SelectedMeshComp); }
 	void CalcTransform(const FFrameTime& Frametime, FTransform& OutTransform, FTransform& OutParentTransform);
+
 	TWeakObjectPtr<UMovieSceneSkeletalAnimationSection> SelectedSection;
 	TWeakObjectPtr<USkeletalMeshComponent> SelectedMeshComp;
 };
