Sequencer: Curve Editor: Normalized/Stack view take into account time range when figuring out min, max values for positioning. We added a new CurveModel function to handle this instead of breaking the other one since that's used in other places  (Rich curve, plugins), and didn't want to regress any expected behaviors.

diff --git a/Engine/Source/Editor/CurveEditor/Private/SCurveEditorView.cpp b/Engine/Source/Editor/CurveEditor/Private/SCurveEditorView.cpp
--- a/Engine/Source/Editor/CurveEditor/Private/SCurveEditorView.cpp
+++ b/Engine/Source/Editor/CurveEditor/Private/SCurveEditorView.cpp
@@ -158,8 +158,6 @@ void SCurveEditorView::ZoomAround(const FVector2D& Amount, double InputOrigin, d
 
 void SCurveEditorView::GetCurveDrawParams(TArray<FCurveDrawParams>& OutDrawParams) 
 {
-	//make sure the transform is set up
-	UpdateViewToTransformCurves();
 
 	TSharedPtr<FCurveEditor> CurveEditor = WeakCurveEditor.Pin();
 	if (!CurveEditor)
@@ -171,6 +169,9 @@ void SCurveEditorView::GetCurveDrawParams(TArray<FCurveDrawParams>& OutDrawParam
 	double InputMin = 0, InputMax = 1;
 	GetInputBounds(InputMin, InputMax);
 
+	//make sure the transform is set up
+	UpdateViewToTransformCurves(InputMin,InputMax);
+
 	OutDrawParams.Reset(CurveInfoByID.Num());
 
 	for (const TTuple<FCurveModelID, FCurveInfo>& Pair : CurveInfoByID)
diff --git a/Engine/Source/Editor/CurveEditor/Private/Views/SCurveEditorViewNormalized.cpp b/Engine/Source/Editor/CurveEditor/Private/Views/SCurveEditorViewNormalized.cpp
--- a/Engine/Source/Editor/CurveEditor/Private/Views/SCurveEditorViewNormalized.cpp
+++ b/Engine/Source/Editor/CurveEditor/Private/Views/SCurveEditorViewNormalized.cpp
@@ -153,7 +153,7 @@ void SCurveEditorViewNormalized::FrameVertical(double InOutputMin, double InOutp
 	OutputMax = 1.1;;
 }
 
-void SCurveEditorViewNormalized::UpdateViewToTransformCurves() 
+void SCurveEditorViewNormalized::UpdateViewToTransformCurves(double InputMin, double InputMax) 
 {
 	TSharedPtr<FCurveEditor> CurveEditor = WeakCurveEditor.Pin();
 	if (!CurveEditor)
@@ -170,7 +170,7 @@ void SCurveEditorViewNormalized::UpdateViewToTransformCurves()
 		}
 
 		double CurveOutputMin = 0, CurveOutputMax = 1;
-		Curve->GetValueRange(CurveOutputMin, CurveOutputMax);
+		Curve->GetValueRange(InputMin, InputMax, CurveOutputMin, CurveOutputMax);
 
 		It->Value.ViewToCurveTransform = CalculateViewToCurveTransform(OutputMin, OutputMax, CurveOutputMin, CurveOutputMax);
 	}
@@ -188,7 +188,9 @@ void SCurveEditorViewNormalized::Tick(const FGeometry& AllottedGeometry, const d
 	if (!CurveEditor->AreBoundTransformUpdatesSuppressed())
 	{
 		bFixedOutputBounds = false;
-		UpdateViewToTransformCurves();
+		double InputMin = 0, InputMax = 1;
+		GetInputBounds(InputMin, InputMax);
+		UpdateViewToTransformCurves(InputMin, InputMax);
 	}
 
 	SInteractiveCurveEditorView::Tick(AllottedGeometry, InCurrentTime, InDeltaTime);
diff --git a/Engine/Source/Editor/CurveEditor/Private/Views/SCurveEditorViewStacked.cpp b/Engine/Source/Editor/CurveEditor/Private/Views/SCurveEditorViewStacked.cpp
--- a/Engine/Source/Editor/CurveEditor/Private/Views/SCurveEditorViewStacked.cpp
+++ b/Engine/Source/Editor/CurveEditor/Private/Views/SCurveEditorViewStacked.cpp
@@ -379,9 +379,8 @@ void SCurveEditorViewStacked::DrawBufferedCurves(const FGeometry& AllottedGeomet
 	}
 }
 
-void SCurveEditorViewStacked::UpdateViewToTransformCurves()
+void SCurveEditorViewStacked::UpdateViewToTransformCurves(double InputMin, double InputMax)
 {
-
 	TSharedPtr<FCurveEditor> CurveEditor = WeakCurveEditor.Pin();
 	if (!CurveEditor)
 	{
@@ -403,7 +402,7 @@ void SCurveEditorViewStacked::UpdateViewToTransformCurves()
 		const double ValueOffset = -CurveIndexFromBottom - PaddingToBottomOfView;
 
 		double CurveOutputMin = 0, CurveOutputMax = 1;
-		Curve->GetValueRange(CurveOutputMin, CurveOutputMax);
+		Curve->GetValueRange(InputMin, InputMax, CurveOutputMin, CurveOutputMax);
 
 		It->Value.ViewToCurveTransform = CalculateViewToCurveTransform(CurveOutputMin, CurveOutputMax, ValueOffset);
 	}
@@ -422,7 +421,10 @@ void SCurveEditorViewStacked::Tick(const FGeometry& AllottedGeometry, const doub
 
 	if (!CurveEditor->AreBoundTransformUpdatesSuppressed())
 	{
-		UpdateViewToTransformCurves();
+		// Get the Min/Max values on the X axis, for Time
+		double InputMin = 0, InputMax = 1;
+		GetInputBounds(InputMin, InputMax);
+		UpdateViewToTransformCurves(InputMin, InputMax);
 	}
 
 	SInteractiveCurveEditorView::Tick(AllottedGeometry, InCurrentTime, InDeltaTime);
diff --git a/Engine/Source/Editor/CurveEditor/Public/CurveModel.h b/Engine/Source/Editor/CurveEditor/Public/CurveModel.h
--- a/Engine/Source/Editor/CurveEditor/Public/CurveModel.h
+++ b/Engine/Source/Editor/CurveEditor/Public/CurveModel.h
@@ -138,6 +138,16 @@ public:
 	*/
 	virtual void GetValueRange(double& MinValue, double& MaxValue) const = 0;
 
+	/** Get range of output value based on specified input times. By default will just get the range
+	* without a specified time
+	* @param MinTime Minimum Time
+	* @param MaxTime Maximium Time
+	* @param MinValue Minimum Value
+	* @param MaxValue Minimum Value
+	*/
+	virtual void GetValueRange(double InMinTime, double InMaxTime, double& MinValue, double& MaxValue) const { GetValueRange(MinValue, MaxValue); }
+
+
 	/** Get the number of keys
 	* @param The number of keys
 	*/
diff --git a/Engine/Source/Editor/CurveEditor/Public/SCurveEditorView.h b/Engine/Source/Editor/CurveEditor/Public/SCurveEditorView.h
--- a/Engine/Source/Editor/CurveEditor/Public/SCurveEditorView.h
+++ b/Engine/Source/Editor/CurveEditor/Public/SCurveEditorView.h
@@ -149,8 +149,11 @@ public:
 	/** This should be called every tick by an owning widget, to see if the cache is valid, which will then recreate it and invalidate widget*/
 	virtual void CheckCacheAndInvalidateIfNeeded();
 
+	UE_DEPRECATED(5.3, "Use UpdateViewToTransformCurves(double InputMin, double InputMax) instead.")
+	virtual void UpdateViewToTransformCurves() {}
+
 	/** Function to make sure to update the view to the transform curves, we need to do this before we cache*/
-	virtual void UpdateViewToTransformCurves() {};
+	virtual void UpdateViewToTransformCurves(double InputMin, double InputMax) {};
 
 	/** Frame the view vertially by the input and output bounds, peformaing any custom clipping as needed */
 	virtual void FrameVertical(double InOutputMin, double InOutputMax);
diff --git a/Engine/Source/Editor/CurveEditor/Public/Views/SCurveEditorViewNormalized.h b/Engine/Source/Editor/CurveEditor/Public/Views/SCurveEditorViewNormalized.h
--- a/Engine/Source/Editor/CurveEditor/Public/Views/SCurveEditorViewNormalized.h
+++ b/Engine/Source/Editor/CurveEditor/Public/Views/SCurveEditorViewNormalized.h
@@ -28,7 +28,7 @@ public:
 	/** Tools should use vertical snapping since grid lines to snap to will usually be visible */
 	virtual bool IsValueSnapEnabled() const override { return true; }
 
-	virtual void UpdateViewToTransformCurves() override;
+	virtual void UpdateViewToTransformCurves(double InputMin, double InputMax) override;
 
 	void FrameVertical(double InOutputMin, double InOutputMax) override;
 
diff --git a/Engine/Source/Editor/CurveEditor/Public/Views/SCurveEditorViewStacked.h b/Engine/Source/Editor/CurveEditor/Public/Views/SCurveEditorViewStacked.h
--- a/Engine/Source/Editor/CurveEditor/Public/Views/SCurveEditorViewStacked.h
+++ b/Engine/Source/Editor/CurveEditor/Public/Views/SCurveEditorViewStacked.h
@@ -29,7 +29,7 @@ public:
 	/** Tools should use vertical snapping since grid lines to snap to will usually be visible */
 	virtual bool IsValueSnapEnabled() const override { return true; }
 
-	virtual void UpdateViewToTransformCurves() override;
+	virtual void UpdateViewToTransformCurves(double InputMin, double InputMax) override;
 
 protected:
 
diff --git a/Engine/Source/Editor/MovieSceneTools/Private/Channels/BezierChannelCurveModel.cpp b/Engine/Source/Editor/MovieSceneTools/Private/Channels/BezierChannelCurveModel.cpp
--- a/Engine/Source/Editor/MovieSceneTools/Private/Channels/BezierChannelCurveModel.cpp
+++ b/Engine/Source/Editor/MovieSceneTools/Private/Channels/BezierChannelCurveModel.cpp
@@ -464,10 +464,10 @@ void FBezierChannelCurveModel<ChannelType, ChannelValue, KeyType>::FeaturePointM
 	}
 }
 
-template<typename ChannelType, typename ChannelValue, typename KeyType> 
-void FBezierChannelCurveModel<ChannelType, ChannelValue, KeyType>::GetValueRange(double& MinValue, double& MaxValue) const
+template<typename ChannelType, typename ChannelValue, typename KeyType>
+void FBezierChannelCurveModel<ChannelType, ChannelValue, KeyType>::GetValueRange(double InMinTime, double InMaxTime, double& MinValue, double& MaxValue) const
 {
-	ChannelType* Channel    	= this->GetChannelHandle().Get();
+	ChannelType* Channel = this->GetChannelHandle().Get();
 	UMovieSceneSection* Section = Cast<UMovieSceneSection>(this->GetOwningObject());
 
 	if (Channel && Section && Section->GetTypedOuter<UMovieScene>())
@@ -485,26 +485,53 @@ void FBezierChannelCurveModel<ChannelType, ChannelValue, KeyType>::GetValueRange
 			FFrameRate TickResolution = Section->GetTypedOuter<UMovieScene>()->GetTickResolution();
 			double ToTime = TickResolution.AsInterval();
 			int32 LastKeyIndex = Values.Num() - 1;
-			MinValue = MaxValue = Values[0].Value;
+			MinValue = DBL_MAX;
+			MaxValue = DBL_MIN;
 
 			for (int32 i = 0; i < Values.Num(); i++)
 			{
+				double KeyTime = static_cast<double>(Times[i].Value) * ToTime;
+				if (KeyTime < InMinTime)
+				{
+					continue;
+				}
+				else if (KeyTime > InMaxTime)
+				{
+					break;
+				}
 				const ChannelValue& Key = Values[i];
 
-				MinValue = FMath::Min(MinValue,(double) Key.Value);
+				MinValue = FMath::Min(MinValue, (double)Key.Value);
 				MaxValue = FMath::Max(MaxValue, (double)Key.Value);
 
 				if (Key.InterpMode == RCIM_Cubic && i != LastKeyIndex)
 				{
 					const ChannelValue& NextKey = Values[i + 1];
-					double KeyTime = static_cast<double>(Times[i].Value) *ToTime;
-					double NextTime = static_cast<double>(Times[i +1].Value) *ToTime;
+					double NextTime = static_cast<double>(Times[i + 1].Value) * ToTime;
 					double TimeStep = (NextTime - KeyTime) * 0.2f;
 					FeaturePointMethod(KeyTime, NextTime, Key.Value, TimeStep, 0, 3, MaxValue, MinValue);
 				}
 			}
 		}
 	}
+	//if nothing found just set to zero
+	if (MinValue == DBL_MAX)
+	{
+		MinValue = 0.0;
+	}
+	if (MaxValue == DBL_MIN)
+	{
+		MaxValue = 0.0;
+	}
+}
+
+
+template<typename ChannelType, typename ChannelValue, typename KeyType> 
+void FBezierChannelCurveModel<ChannelType, ChannelValue, KeyType>::GetValueRange(double& MinValue, double& MaxValue) const
+{
+	const double InMinTime = DBL_MIN;
+	const double InMaxtime = DBL_MAX;
+	GetValueRange(InMinTime, InMaxtime, MinValue, MaxValue);
 }
 
 template<typename ChannelType, typename ChannelValue, typename KeyType> 
diff --git a/Engine/Source/Editor/MovieSceneTools/Public/Channels/BezierChannelCurveModel.h b/Engine/Source/Editor/MovieSceneTools/Public/Channels/BezierChannelCurveModel.h
--- a/Engine/Source/Editor/MovieSceneTools/Public/Channels/BezierChannelCurveModel.h
+++ b/Engine/Source/Editor/MovieSceneTools/Public/Channels/BezierChannelCurveModel.h
@@ -67,6 +67,7 @@ public:
 	virtual void GetCurveAttributes(FCurveAttributes& OutCurveAttributes) const override;
 	virtual void SetCurveAttributes(const FCurveAttributes& InCurveAttributes) override;
 	virtual void GetValueRange(double& MinValue, double& MaxValue) const override;
+	virtual void GetValueRange(double InMinTime, double InMaxTime, double& MinValue, double& MaxValue) const override;
 
 protected:
 	// FChannelCurveModel
