Sequencer: NLA revert previous submit that made things a little worse when moving root bone

diff --git a/Engine/Source/Editor/MovieSceneTools/Private/EditModes/SkeletalAnimationTrackEditMode.cpp b/Engine/Source/Editor/MovieSceneTools/Private/EditModes/SkeletalAnimationTrackEditMode.cpp
--- a/Engine/Source/Editor/MovieSceneTools/Private/EditModes/SkeletalAnimationTrackEditMode.cpp
+++ b/Engine/Source/Editor/MovieSceneTools/Private/EditModes/SkeletalAnimationTrackEditMode.cpp
@@ -654,10 +654,9 @@ bool FSkeletalAnimationTrackEditMode::InputDelta(FEditorViewportClient* InViewpo
 					CurrentRotation = (InRot.Quaternion() * CurrentRotation);
 					CurrentTransform.SetRotation(CurrentRotation);
 				}
-
 				if (bDoTranslation)
 				{
-					FVector CurrentLocation = GetWidgetLocation();
+					FVector CurrentLocation = CurrentTransform.GetLocation();
 					CurrentLocation = CurrentLocation + InDrag;
 					CurrentTransform.SetLocation(CurrentLocation);
 				}
@@ -722,3 +721,5 @@ bool FSkeletalAnimationTrackEditMode::IsCompatibleWith(FEditorModeID OtherModeID
 
 
 #undef LOCTEXT_NAMESPACE
+
+
