Sequencer- Refactor referencing of materials for component material tracks to allow for more ways to resolve materials, including:
* Adding the ability to reference component material slots based on material slot name rather than just by index. This prevents breaking resolution of material parameter tracks when materials indices get reordered within a mesh, as long as the slot names remain consistent. This type of reordering can happen when reimporting meshes for example, as material order may not be consistent depending on your mesh authoring package. It can also happen when creating meshes dynamically, which certain customers do. For now, we try to match by slot name first, and then by index if we cannot find a material by its slot name (or one doesn't exist in the binding for backwards compatibility). In future we will introduce ways for the user to update or fix track bindings manually, for example in cases that material slot names are renamed.
* Adding the ability to reference the OverlayMaterial in MeshComponents, and animate parameters in it, or change which material is assigned to the overlay material property using a material switcher track.
* Formalizing binding to decal material- previously this was just done by default if the component the track was added to was a DecalComponent.

The above refactor also allows us to more easily add animation of other types of material tracks in other components.

In addition to the above, I've exposed the Material Slot Name (read-only) in the outliner/details panel on PrimitiveComponent. Previously you could only see these slot names by opening the mesh asset.

See design docs for this here and here:
https://docs.google.com/document/d/137PB3QN-dBdESJGfrfscNsiQPZmhg9FxIFerU1M00uU/edit
https://docs.google.com/document/d/18dlhZcJ8wE-b_Q9ulTpVirmmY_Na8w-5ZuZjeAA8oEI/edit

diff --git a/Engine/Plugins/MovieScene/SequencerScripting/Source/SequencerScripting/Private/ExtensionLibraries/MovieSceneMaterialTrackExtensions.cpp b/Engine/Plugins/MovieScene/SequencerScripting/Source/SequencerScripting/Private/ExtensionLibraries/MovieSceneMaterialTrackExtensions.cpp
--- a/Engine/Plugins/MovieScene/SequencerScripting/Source/SequencerScripting/Private/ExtensionLibraries/MovieSceneMaterialTrackExtensions.cpp
+++ b/Engine/Plugins/MovieScene/SequencerScripting/Source/SequencerScripting/Private/ExtensionLibraries/MovieSceneMaterialTrackExtensions.cpp
@@ -14,7 +14,7 @@ void UMovieSceneMaterialTrackExtensions::SetMaterialIndex(UMovieSceneComponentMa
 	}
 
 	Track->Modify();
-	Track->SetMaterialIndex(MaterialIndex);
+	Track->SetMaterialInfo(FComponentMaterialInfo{FName(), MaterialIndex, EComponentMaterialType::IndexedMaterial});
 }
 
 int32 UMovieSceneMaterialTrackExtensions::GetMaterialIndex(UMovieSceneComponentMaterialTrack* Track)
@@ -25,8 +25,33 @@ int32 UMovieSceneMaterialTrackExtensions::GetMaterialIndex(UMovieSceneComponentM
 		return -1;
 	}
 
-	return Track->GetMaterialIndex();
+	return Track->GetMaterialInfo().MaterialSlotIndex;
 }
 
+void UMovieSceneMaterialTrackExtensions::SetMaterialInfo(UMovieSceneComponentMaterialTrack* Track, const FComponentMaterialInfo& MaterialInfo)
+{
+	if (!Track)
+	{
+		FFrame::KismetExecutionMessage(TEXT("Cannot call SetMaterialInfo on a null track"), ELogVerbosity::Error);
+		return;
+	}
+
+	Track->Modify();
+	Track->SetMaterialInfo(MaterialInfo);
+}
+
+FComponentMaterialInfo UMovieSceneMaterialTrackExtensions::GetMaterialInfo(UMovieSceneComponentMaterialTrack* Track)
+{
+	if (!Track)
+	{
+		FFrame::KismetExecutionMessage(TEXT("Cannot call GetMaterialInfo on a null track"), ELogVerbosity::Error);
+		return FComponentMaterialInfo();
+	}
+
+	return Track->GetMaterialInfo();
+}
+
+
+
 
 
diff --git a/Engine/Plugins/MovieScene/SequencerScripting/Source/SequencerScripting/Private/ExtensionLibraries/MovieScenePrimitiveMaterialTrackExtensions.cpp b/Engine/Plugins/MovieScene/SequencerScripting/Source/SequencerScripting/Private/ExtensionLibraries/MovieScenePrimitiveMaterialTrackExtensions.cpp
--- a/Engine/Plugins/MovieScene/SequencerScripting/Source/SequencerScripting/Private/ExtensionLibraries/MovieScenePrimitiveMaterialTrackExtensions.cpp
+++ b/Engine/Plugins/MovieScene/SequencerScripting/Source/SequencerScripting/Private/ExtensionLibraries/MovieScenePrimitiveMaterialTrackExtensions.cpp
@@ -1,6 +1,7 @@
 // Copyright Epic Games, Inc. All Rights Reserved.
 
 #include "ExtensionLibraries/MovieScenePrimitiveMaterialTrackExtensions.h"
+#include "Tracks/MovieSceneMaterialTrack.h"
 #include "Tracks/MovieScenePrimitiveMaterialTrack.h"
 
 #include UE_INLINE_GENERATED_CPP_BY_NAME(MovieScenePrimitiveMaterialTrackExtensions)
@@ -14,7 +15,7 @@ void UMovieScenePrimitiveMaterialTrackExtensions::SetMaterialIndex(UMovieScenePr
 	}
 
 	Track->Modify();
-	Track->SetMaterialIndex(MaterialIndex);
+	Track->SetMaterialInfo(FComponentMaterialInfo{ FName(), MaterialIndex, EComponentMaterialType::IndexedMaterial });
 }
 
 int32 UMovieScenePrimitiveMaterialTrackExtensions::GetMaterialIndex(UMovieScenePrimitiveMaterialTrack* Track)
@@ -25,8 +26,36 @@ int32 UMovieScenePrimitiveMaterialTrackExtensions::GetMaterialIndex(UMovieSceneP
 		return -1;
 	}
 
-	return Track->GetMaterialIndex();
+	return Track->GetMaterialInfo().MaterialSlotIndex;
 }
 
+void UMovieScenePrimitiveMaterialTrackExtensions::SetMaterialInfo(UMovieScenePrimitiveMaterialTrack* Track, const FComponentMaterialInfo& MaterialInfo)
+{
+	if (!Track)
+	{
+		FFrame::KismetExecutionMessage(TEXT("Cannot call SetMaterialInfo on a null track"), ELogVerbosity::Error);
+		return;
+	}
+
+	Track->Modify();
+	Track->SetMaterialInfo(MaterialInfo);
+}
+
+FComponentMaterialInfo UMovieScenePrimitiveMaterialTrackExtensions::GetMaterialInfo(UMovieScenePrimitiveMaterialTrack* Track)
+{
+	if (!Track)
+	{
+		FFrame::KismetExecutionMessage(TEXT("Cannot call GetMaterialInfo on a null track"), ELogVerbosity::Error);
+		return FComponentMaterialInfo();
+	}
+
+	return Track->GetMaterialInfo();
+}
+
+
+
+
+
+
 
 
diff --git a/Engine/Plugins/MovieScene/SequencerScripting/Source/SequencerScripting/Public/ExtensionLibraries/MovieSceneMaterialTrackExtensions.h b/Engine/Plugins/MovieScene/SequencerScripting/Source/SequencerScripting/Public/ExtensionLibraries/MovieSceneMaterialTrackExtensions.h
--- a/Engine/Plugins/MovieScene/SequencerScripting/Source/SequencerScripting/Public/ExtensionLibraries/MovieSceneMaterialTrackExtensions.h
+++ b/Engine/Plugins/MovieScene/SequencerScripting/Source/SequencerScripting/Public/ExtensionLibraries/MovieSceneMaterialTrackExtensions.h
@@ -10,6 +10,8 @@ class UMovieSceneComponentMaterialTrack;
 
 class UMovieSceneMaterialTrack;
 
+struct FComponentMaterialInfo;
+
 /**
  * Function library containing methods that should be hoisted onto UMovieSceneMaterialTrack for scripting
  */
@@ -25,7 +27,8 @@ public:
 	 * @param Track The track to use
 	 * @param MaterialIndex The desired material index to animate. Values of < 0 or >= NumMaterials will be silently ignored and evaluation will fail.
 	 */
-	UFUNCTION(BlueprintCallable, Category = "Sequencer|Track", meta = (ScriptMethod))
+	UE_DEPRECATED(5.4, "Use SetMaterialInfo instead")
+	UFUNCTION(BlueprintCallable, Category = "Sequencer|Track", meta = (ScriptMethod, DeprecatedFunction, DeprecationMessage = "Use SetMaterialInfo instead"))
 	static void SetMaterialIndex(UMovieSceneComponentMaterialTrack* Track, const int32 MaterialIndex);
 
 	/**
@@ -33,6 +36,23 @@ public:
 	 * @param Track The track to use
 	 * @return The material index.
 	 */
-	UFUNCTION(BlueprintCallable, Category = "Sequencer|Track", meta=(ScriptMethod))
+	UE_DEPRECATED(5.4, "Use SetMaterialInfo instead")
+	UFUNCTION(BlueprintCallable, Category = "Sequencer|Track", meta=(ScriptMethod, DeprecatedFunction, DeprecationMessage = "Use SetMaterialInfo instead"))
 	static int32 GetMaterialIndex(UMovieSceneComponentMaterialTrack* Track);
+
+		/**
+	 * Set material info of the component that is animated by the material track.
+	 * @param Track The track to use
+	 * @param MaterialInfo The desired material info to animate.
+	 */
+	UFUNCTION(BlueprintCallable, Category = "Sequencer|Track", meta = (ScriptMethod))
+	static void SetMaterialInfo(UMovieSceneComponentMaterialTrack* Track, const FComponentMaterialInfo& MaterialInfo);
+
+	/**
+	 * Get material info of the component that is animated by the material track.
+	 * @param Track The track to use
+	 * @return The material info.
+	 */
+	UFUNCTION(BlueprintCallable, Category = "Sequencer|Track", meta=(ScriptMethod, DeprecatedFunction, DeprecationMessage = "Use SetMaterialInfo instead"))
+	static FComponentMaterialInfo GetMaterialInfo(UMovieSceneComponentMaterialTrack* Track);
 };
diff --git a/Engine/Plugins/MovieScene/SequencerScripting/Source/SequencerScripting/Public/ExtensionLibraries/MovieScenePrimitiveMaterialTrackExtensions.h b/Engine/Plugins/MovieScene/SequencerScripting/Source/SequencerScripting/Public/ExtensionLibraries/MovieScenePrimitiveMaterialTrackExtensions.h
--- a/Engine/Plugins/MovieScene/SequencerScripting/Source/SequencerScripting/Public/ExtensionLibraries/MovieScenePrimitiveMaterialTrackExtensions.h
+++ b/Engine/Plugins/MovieScene/SequencerScripting/Source/SequencerScripting/Public/ExtensionLibraries/MovieScenePrimitiveMaterialTrackExtensions.h
@@ -23,7 +23,8 @@ public:
 	 * @param Track The track to use
 	 * @param MaterialIndex The desired material index to animate. Values of < 0 or >= NumMaterials will be silently ignored and evaluation will fail.
 	 */
-	UFUNCTION(BlueprintCallable, Category = "Sequencer|Track", meta = (ScriptMethod))
+	UE_DEPRECATED(5.4, "Use SetMaterialInfo instead")
+	UFUNCTION(BlueprintCallable, Category = "Sequencer|Track", meta = (ScriptMethod, DeprecatedFunction, DeprecationMessage = "Use SetMaterialInfo instead"))
 	static void SetMaterialIndex(UMovieScenePrimitiveMaterialTrack* Track, const int32 MaterialIndex);
 
 	/**
@@ -31,6 +32,23 @@ public:
 	 * @param Track The track to use
 	 * @return The material index.
 	 */
-	UFUNCTION(BlueprintCallable, Category = "Sequencer|Track", meta=(ScriptMethod))
+	UE_DEPRECATED(5.4, "Use SetMaterialInfo instead")
+	UFUNCTION(BlueprintCallable, Category = "Sequencer|Track", meta=(ScriptMethod, DeprecatedFunction, DeprecationMessage = "Use SetMaterialInfo instead"))
 	static int32 GetMaterialIndex(UMovieScenePrimitiveMaterialTrack* Track);
+
+		/**
+	 * Set material info of the component that is animated by the material track.
+	 * @param Track The track to use
+	 * @param MaterialInfo The desired material info to animate.
+	 */
+	UFUNCTION(BlueprintCallable, Category = "Sequencer|Track", meta = (ScriptMethod))
+	static void SetMaterialInfo(UMovieScenePrimitiveMaterialTrack* Track, const FComponentMaterialInfo& MaterialInfo);
+
+	/**
+	 * Get material info of the component that is animated by the material track.
+	 * @param Track The track to use
+	 * @return The material info.
+	 */
+	UFUNCTION(BlueprintCallable, Category = "Sequencer|Track", meta=(ScriptMethod, DeprecatedFunction, DeprecationMessage = "Use SetMaterialInfo instead"))
+	static FComponentMaterialInfo GetMaterialInfo(UMovieScenePrimitiveMaterialTrack* Track);
 };
diff --git a/Engine/Source/Editor/DetailCustomizations/Private/ComponentMaterialCategory.cpp b/Engine/Source/Editor/DetailCustomizations/Private/ComponentMaterialCategory.cpp
--- a/Engine/Source/Editor/DetailCustomizations/Private/ComponentMaterialCategory.cpp
+++ b/Engine/Source/Editor/DetailCustomizations/Private/ComponentMaterialCategory.cpp
@@ -35,6 +35,7 @@
 #include "UObject/UObjectHash.h"
 #include "UObject/UnrealType.h"
 #include "UnrealEdGlobals.h"
+#include "PropertyCustomizationHelpers.h"
 
 class AActor;
 class UWorld;
@@ -81,10 +82,12 @@ public:
 				// Primitive components and some actor components have materials
 				UPrimitiveComponent* PrimitiveComp = Cast<UPrimitiveComponent>( CurComponent );
 				UDecalComponent* DecalComponent = (PrimitiveComp ? NULL : Cast<UDecalComponent>( CurComponent ));
+				TArray<FName> MaterialSlotNames;
 
 				if( PrimitiveComp )
 				{
 					NumMaterials = PrimitiveComp->GetNumMaterials();
+					MaterialSlotNames = PrimitiveComp->GetMaterialSlotNames();
 				}
 				else if( DecalComponent )
 				{
@@ -100,10 +103,19 @@ public:
 					if( PrimitiveComp )
 					{
 						Material = PrimitiveComp->GetEditorMaterial(CurMaterialIndex);
+						if (MaterialSlotNames.IsValidIndex(CurMaterialIndex))
+						{
+							CurMaterialSlotName = MaterialSlotNames[CurMaterialIndex];
+						}
+						else
+						{
+							CurMaterialSlotName = FName();
+						}
 					}
 					else if( DecalComponent )
 					{
 						Material = DecalComponent->GetMaterial(CurMaterialIndex);
+						CurMaterialSlotName = FName();
 					}
 
 					CurMaterial = Material;
@@ -117,6 +129,7 @@ public:
 				}
 				// Out of materials on this component, reset for the next component
 				CurMaterialIndex = 0;
+				CurMaterialSlotName = FName();
 			}
 			// Advance to the next compoment
 			++CurComponentIndex;
@@ -129,6 +142,7 @@ public:
 		CurComponent = NULL;
 		CurMaterial = NULL;
 		CurMaterialIndex = INDEX_NONE;
+		CurMaterialSlotName = FName();
 	}
 
 	/**
@@ -149,6 +163,11 @@ public:
 	 */
 	int32 GetMaterialIndex() const { return CurMaterialIndex; }
 
+	/**
+	 * @return The material slot name of the material in the current component, if applicable.
+	 */
+	FName GetMaterialSlotName() const { return CurMaterialSlotName; }
+
 	/**
 	 * @return The current component using the current material
 	 */
@@ -165,6 +184,8 @@ private:
 	int32 CurComponentIndex;
 	/** The index of the material we are stopped on */
 	int32 CurMaterialIndex;
+	/** The slot name of the material we are stopped on, if applicable */
+	FName CurMaterialSlotName;
 	/** Whether or not we've reached the end of the components */
 	uint32 bReachedEnd:1;
 };
@@ -186,6 +207,7 @@ void FComponentMaterialCategory::Create( IDetailLayoutBuilder& DetailBuilder )
 	MaterialListDelegates.OnCanCopyMaterialItem.BindSP( this, &FComponentMaterialCategory::OnCanCopyMaterialItem );
 	MaterialListDelegates.OnCopyMaterialItem.BindSP( this, &FComponentMaterialCategory::OnCopyMaterialItem );
 	MaterialListDelegates.OnPasteMaterialItem.BindSP(this, &FComponentMaterialCategory::OnPasteMaterialItem);
+	MaterialListDelegates.OnGenerateCustomMaterialWidgets.BindSP(this, &FComponentMaterialCategory::OnGenerateWidgetsForMaterial);
 	
 	//Pass an empty material list owner (owner can be use by the asset picker filter. In this case we do not need it)
 	TArray<FAssetData> MaterialListOwner;
@@ -559,4 +581,31 @@ void FComponentMaterialCategory::OnPasteMaterialItem(int32 CurrentSlot)
 	}
 }
 
+TSharedRef<SWidget> FComponentMaterialCategory::OnGenerateWidgetsForMaterial(UMaterialInterface* Material, int32 SlotIndex)
+{
+	return
+		SNew(SMaterialSlotWidget, SlotIndex, true)
+		.MaterialName(this, &FComponentMaterialCategory::GetMaterialNameText, SlotIndex)
+		.IsMaterialSlotNameReadOnly(true)
+		.CanDeleteMaterialSlot(false);
+}
+
+FText FComponentMaterialCategory::GetMaterialNameText(int32 MaterialIndex) const
+{
+	// For now we don't support multi-select
+	if (SelectedComponents.Num() != 1)
+	{
+		return FText::FromName(NAME_None);
+	}
+	for (FMaterialIterator It(SelectedComponents); It; ++It)
+	{
+		if (It.GetMaterialIndex() == MaterialIndex)
+		{
+			return FText::FromName(It.GetMaterialSlotName());
+		}
+	}
+
+	return FText::FromName(NAME_None);
+}
+
 #undef LOCTEXT_NAMESPACE
diff --git a/Engine/Source/Editor/DetailCustomizations/Private/ComponentMaterialCategory.h b/Engine/Source/Editor/DetailCustomizations/Private/ComponentMaterialCategory.h
--- a/Engine/Source/Editor/DetailCustomizations/Private/ComponentMaterialCategory.h
+++ b/Engine/Source/Editor/DetailCustomizations/Private/ComponentMaterialCategory.h
@@ -13,6 +13,7 @@ class IDetailLayoutBuilder;
 class UActorComponent;
 class UMaterialInterface;
 class USceneComponent;
+class SWidget;
 
 /**
  * Encapsulates functionality for the ActorDetails material category
@@ -69,6 +70,14 @@ private:
 	 */
 	void OnPasteMaterialItem(int32 CurrentSlot);
 
+	/**
+	 * Called by the material list widget to customize the material slot widget
+	 *
+	 * @param Material	The material for the slot widget to be customized
+	 * @param SlotIndex The slot index for the slot widget to be customized
+	 */
+	TSharedRef<SWidget> OnGenerateWidgetsForMaterial(UMaterialInterface* Material, int32 SlotIndex);
+
 	/** 
 	 * @return true if a component is editable (and visible in the view)
 	 */
@@ -82,4 +91,6 @@ private:
 	FNotifyHook* NotifyHook;
 
 	IDetailCategoryBuilder* MaterialCategory;
+
+	FText GetMaterialNameText(int32 MaterialIndex) const;
 };
diff --git a/Engine/Source/Editor/MovieSceneTools/Private/TrackEditors/MaterialTrackEditor.cpp b/Engine/Source/Editor/MovieSceneTools/Private/TrackEditors/MaterialTrackEditor.cpp
--- a/Engine/Source/Editor/MovieSceneTools/Private/TrackEditors/MaterialTrackEditor.cpp
+++ b/Engine/Source/Editor/MovieSceneTools/Private/TrackEditors/MaterialTrackEditor.cpp
@@ -15,6 +15,7 @@
 #include "MaterialEditorModule.h"
 #include "Engine/Selection.h"
 #include "ISequencerModule.h"
+#include "Components/MeshComponent.h"
 
 
 #define LOCTEXT_NAMESPACE "MaterialTrackEditor"
@@ -317,15 +318,41 @@ UMaterialInterface* FComponentMaterialTrackEditor::GetMaterialInterfaceForTrack(
 		return nullptr;
 	}
 
-	if (UPrimitiveComponent* Component = Cast<UPrimitiveComponent>(Object))
+	const FComponentMaterialInfo& MaterialInfo = ComponentMaterialTrack->GetMaterialInfo();
+	switch (MaterialInfo.MaterialType)
 	{
-		return Component->GetMaterial( ComponentMaterialTrack->GetMaterialIndex() );
-	}
-	else if (UDecalComponent* DecalComponent = Cast<UDecalComponent>(Object))
-	{
-		return DecalComponent->GetDecalMaterial();
+	case EComponentMaterialType::Empty:
+		break;
+	case EComponentMaterialType::IndexedMaterial:
+		if (UPrimitiveComponent* Component = Cast<UPrimitiveComponent>(Object))
+		{
+			UMaterialInterface* Material = nullptr;
+			if (!MaterialInfo.MaterialSlotName.IsNone())
+			{
+				Material = Component->GetMaterialByName(MaterialInfo.MaterialSlotName);
+				if (!Material)
+				{
+					Material = Component->GetMaterial(MaterialInfo.MaterialSlotIndex);
+				}
+				return Material;
+			}
+		}
+		break;
+	case EComponentMaterialType::OverlayMaterial:
+		if (UMeshComponent* Component = Cast<UMeshComponent>(Object))
+		{
+			return Component->GetOverlayMaterial();
+		}
+		break;
+	case EComponentMaterialType::DecalMaterial:
+		if (UDecalComponent* DecalComponent = Cast<UDecalComponent>(Object))
+		{
+			return DecalComponent->GetDecalMaterial();
+		}
+		break;
+	default:
+		break;
 	}
-
 	return nullptr;
 }
 
@@ -339,6 +366,11 @@ void FComponentMaterialTrackEditor::ExtendObjectBindingTrackMenu(TSharedRef<FExt
 
 void FComponentMaterialTrackEditor::ConstructObjectBindingTrackMenu(FMenuBuilder& MenuBuilder, TArray<FGuid> ObjectBindings)
 {
+	auto GetMaterialInfoForTrack = [](UMovieSceneTrack* InTrack)
+	{
+		UMovieSceneComponentMaterialTrack* MaterialTrack = Cast<UMovieSceneComponentMaterialTrack>(InTrack);
+		return MaterialTrack ? MaterialTrack->GetMaterialInfo() : FComponentMaterialInfo();
+	};
 	UObject* Object = GetSequencer()->FindSpawnedObjectOrTemplate(ObjectBindings[0]);
 	if (!Object)
 	{
@@ -351,20 +383,51 @@ void FComponentMaterialTrackEditor::ConstructObjectBindingTrackMenu(FMenuBuilder
 		return;
 	}
 
+	const UMovieScene* MovieScene = GetFocusedMovieScene();
+	const FMovieSceneBinding* Binding = Algo::FindBy(MovieScene->GetBindings(), ObjectBindings[0], &FMovieSceneBinding::GetObjectGuid);
+
 	if (UPrimitiveComponent* PrimitiveComponent = Cast<UPrimitiveComponent>(SceneComponent))
 	{
 		int32 NumMaterials = PrimitiveComponent->GetNumMaterials();
-		if (NumMaterials > 0)
+		TArray<FName> MaterialSlotNames = PrimitiveComponent->GetMaterialSlotNames();
+		UMeshComponent* MeshComponent = Cast<UMeshComponent>(SceneComponent);
+		if (NumMaterials > 0 || MeshComponent)
 		{
 			MenuBuilder.BeginSection("Materials", LOCTEXT("MaterialSection", "Material Parameters"));
 			{
 				for (int32 MaterialIndex = 0; MaterialIndex < NumMaterials; MaterialIndex++)
 				{
-					FUIAction AddComponentMaterialAction(FExecuteAction::CreateRaw(this, &FComponentMaterialTrackEditor::HandleAddComponentMaterialActionExecute, SceneComponent, MaterialIndex));
-					FText AddComponentMaterialLabel = FText::Format(LOCTEXT("ComponentMaterialIndexLabelFormat", "Element {0}"), FText::AsNumber(MaterialIndex));
-					FText AddComponentMaterialToolTip = FText::Format(LOCTEXT("ComponentMaterialIndexToolTipFormat", "Add material element {0}"), FText::AsNumber(MaterialIndex));
+					FName MaterialSlotName = MaterialSlotNames.IsValidIndex(MaterialIndex) ? MaterialSlotNames[MaterialIndex] : FName();
+					FComponentMaterialInfo MaterialInfo{ MaterialSlotName, MaterialIndex, EComponentMaterialType::IndexedMaterial };
+					const bool bAlreadyExists = Algo::FindBy(Binding->GetTracks(), MaterialInfo, GetMaterialInfoForTrack) != nullptr;
+					if (bAlreadyExists)
+					{
+						continue;
+					}
+					FUIAction AddComponentMaterialAction(FExecuteAction::CreateRaw(this, &FComponentMaterialTrackEditor::HandleAddComponentMaterialActionExecute, SceneComponent, MaterialInfo));
+					FText AddComponentMaterialLabel = !MaterialSlotName.IsNone() ? 
+						FText::Format(LOCTEXT("ComponentMaterialSlotNameLabelFormat", "Slot: {0}"), FText::FromName(MaterialSlotName)) :
+						FText::Format(LOCTEXT("ComponentMaterialIndexLabelFormat", "Element {0}"), FText::AsNumber(MaterialIndex));
+					FText AddComponentMaterialToolTip = !MaterialSlotName.IsNone() ? 
+						FText::Format(LOCTEXT("ComponentMaterialSlotNameToolTipFormat", "Add material slot {0}, index {1}"), FText::FromName(MaterialSlotName), FText::AsNumber(MaterialIndex)) :
+						FText::Format(LOCTEXT("ComponentMaterialIndexToolTipFormat", "Add material element {0}"), FText::AsNumber(MaterialIndex));
 					MenuBuilder.AddMenuEntry(AddComponentMaterialLabel, AddComponentMaterialToolTip, FSlateIcon(), AddComponentMaterialAction);
 				}
+				if (MeshComponent)
+				{
+					if (UMaterialInterface* OverlayMaterial = MeshComponent->GetOverlayMaterial())
+					{
+						FComponentMaterialInfo MaterialInfo{ FName(), 0, EComponentMaterialType::OverlayMaterial };
+						const bool bAlreadyExists = Algo::FindBy(Binding->GetTracks(), MaterialInfo, GetMaterialInfoForTrack) != nullptr;
+						if (!bAlreadyExists)
+						{
+							FUIAction AddComponentMaterialAction(FExecuteAction::CreateRaw(this, &FComponentMaterialTrackEditor::HandleAddComponentMaterialActionExecute, SceneComponent, MaterialInfo));
+							FText AddOverlayMaterialLabel = FText::Format(LOCTEXT("AddOverlayMaterialLabelFormat", "Overlay: {0}"), FText::FromString(OverlayMaterial->GetName()));
+							FText AddOverlayMaterialToolTip = FText::Format(LOCTEXT("AddOverlayMaterialToolTipFormat", "Add overlay material {0}"), FText::FromString(OverlayMaterial->GetName()));
+							MenuBuilder.AddMenuEntry(AddOverlayMaterialLabel, AddOverlayMaterialToolTip, FSlateIcon(), AddComponentMaterialAction);
+						}
+					}
+				}
 			}
 			MenuBuilder.EndSection();
 		}
@@ -375,17 +438,28 @@ void FComponentMaterialTrackEditor::ConstructObjectBindingTrackMenu(FMenuBuilder
 		{
 			MenuBuilder.BeginSection("Materials", LOCTEXT("MaterialSection", "Material Parameters"));
 			{
-				FUIAction AddComponentMaterialAction(FExecuteAction::CreateRaw(this, &FComponentMaterialTrackEditor::HandleAddComponentMaterialActionExecute, SceneComponent, 0));
-				FText AddDecalMaterialToolTip = FText::Format(LOCTEXT("AddDecalMaterialToolTipFormat", "Add decal material {0}"), FText::FromString(DecalMaterial->GetName()));
-				MenuBuilder.AddMenuEntry(FText::FromString(DecalMaterial->GetName()), AddDecalMaterialToolTip, FSlateIcon(), AddComponentMaterialAction);
+				FComponentMaterialInfo MaterialInfo{ FName(), 0, EComponentMaterialType::DecalMaterial};
+				const bool bAlreadyExists = Algo::FindBy(Binding->GetTracks(), MaterialInfo, GetMaterialInfoForTrack) != nullptr;
+				if (!bAlreadyExists)
+				{
+					FUIAction AddComponentMaterialAction(FExecuteAction::CreateRaw(this, &FComponentMaterialTrackEditor::HandleAddComponentMaterialActionExecute, SceneComponent, MaterialInfo));
+					FText AddDecalMaterialLabel = FText::Format(LOCTEXT("AddDecalMaterialLabelFormat", "Decal: {0}"), FText::FromString(DecalMaterial->GetName()));
+					FText AddDecalMaterialToolTip = FText::Format(LOCTEXT("AddDecalMaterialToolTipFormat", "Add decal material {0}"), FText::FromString(DecalMaterial->GetName()));
+					MenuBuilder.AddMenuEntry(AddDecalMaterialLabel, AddDecalMaterialToolTip, FSlateIcon(), AddComponentMaterialAction);
+				}
 			}
 			MenuBuilder.EndSection();
 		}
 	}
 }
 
-void FComponentMaterialTrackEditor::HandleAddComponentMaterialActionExecute(USceneComponent* Component, int32 MaterialIndex)
+void FComponentMaterialTrackEditor::HandleAddComponentMaterialActionExecute(USceneComponent* Component, FComponentMaterialInfo MaterialInfo)
 {
+	auto GetMaterialInfoForTrack = [](UMovieSceneTrack* InTrack)
+	{
+		UMovieSceneComponentMaterialTrack* MaterialTrack = Cast<UMovieSceneComponentMaterialTrack>(InTrack);
+		return MaterialTrack ? MaterialTrack->GetMaterialInfo() : FComponentMaterialInfo();
+	};
 	TSharedPtr<ISequencer> SequencerPtr = GetSequencer();
 	UMovieScene* MovieScene = SequencerPtr->GetFocusedMovieSceneSequence()->GetMovieScene();
 	if (MovieScene->IsReadOnly())
@@ -424,12 +498,48 @@ void FComponentMaterialTrackEditor::HandleAddComponentMaterialActionExecute(USce
 	for (UActorComponent* ActorComponent : ActorComponents)
 	{
 		FGuid ObjectHandle = SequencerPtr->GetHandleToObject(ActorComponent);
-		FName IndexName(*FString::FromInt(MaterialIndex));
-		if (MovieScene->FindTrack(UMovieSceneComponentMaterialTrack::StaticClass(), ObjectHandle, IndexName) == nullptr)
+		const FMovieSceneBinding* Binding = Algo::FindBy(MovieScene->GetBindings(), ObjectHandle, &FMovieSceneBinding::GetObjectGuid);
+
+		const bool bAlreadyExists = Algo::FindBy(Binding->GetTracks(), MaterialInfo, GetMaterialInfoForTrack) != nullptr;
+		if (!bAlreadyExists)
 		{
 			UMovieSceneComponentMaterialTrack* MaterialTrack = MovieScene->AddTrack<UMovieSceneComponentMaterialTrack>(ObjectHandle);
 			MaterialTrack->Modify();
-			MaterialTrack->SetMaterialIndex(MaterialIndex);
+			MaterialTrack->SetMaterialInfo(MaterialInfo);
+			// Construct display name from MaterialInfo
+			UMaterialInterface* MaterialInterface = GetMaterialInterfaceForTrack(ObjectHandle, MaterialTrack);
+			FText TrackDisplayName;
+			FText TrackTooltipText;
+			switch (MaterialInfo.MaterialType)
+			{
+			case EComponentMaterialType::Empty:
+				break;
+			case EComponentMaterialType::IndexedMaterial:
+				TrackDisplayName = !MaterialInfo.MaterialSlotName.IsNone() ? FText::Format(LOCTEXT("SlotMaterialTrackName", "Material Slot: {0}"), FText::FromName(MaterialInfo.MaterialSlotName))
+					: FText::Format(LOCTEXT("IndexedMaterialTrackName", "Material Element {0}"), FText::AsNumber(MaterialInfo.MaterialSlotIndex));
+				TrackTooltipText = !MaterialInfo.MaterialSlotName.IsNone() ? FText::Format(LOCTEXT("SlotMaterialTrackTooltip", "Material parameter track for slot {0} at index {1}, asset {2}"), FText::FromName(MaterialInfo.MaterialSlotName), FText::AsNumber(MaterialInfo.MaterialSlotIndex), MaterialInterface ? FText::FromString(MaterialInterface->GetName()) : FText())
+					: FText::Format(LOCTEXT("IndexedMaterialTrackTooltip", "Material parameter track for element at index {0}, asset {1}"), FText::AsNumber(MaterialInfo.MaterialSlotIndex), MaterialInterface ? FText::FromString(MaterialInterface->GetName()) : FText());
+				break;
+			case EComponentMaterialType::OverlayMaterial:
+				TrackDisplayName = FText::Format(LOCTEXT("OverlayMaterialTrackName", "Overlay Material {0}"), MaterialInterface ? FText::FromString(MaterialInterface->GetName()) : FText());
+				TrackTooltipText = FText::Format(LOCTEXT("OverlayMaterialTrackTooltip", "Material parameter track for overlay material {0}"), MaterialInterface ? FText::FromString(MaterialInterface->GetName()) : FText());
+				break;
+			case EComponentMaterialType::DecalMaterial:
+				TrackDisplayName = FText::Format(LOCTEXT("DecalMaterialTrackName", "Decal Material {0}"), MaterialInterface ? FText::FromString(MaterialInterface->GetName()) : FText());
+				TrackTooltipText = FText::Format(LOCTEXT("DecalMaterialTrackTooltip", "Material parameter track for decal material {0}"), MaterialInterface ? FText::FromString(MaterialInterface->GetName()) : FText());
+				break;
+			default:
+				break;
+
+			}
+			if (!TrackDisplayName.IsEmpty())
+			{
+				MaterialTrack->SetDisplayName(TrackDisplayName);
+			}
+			if (!TrackTooltipText.IsEmpty())
+			{
+				MaterialTrack->SetDisplayNameTooltipText(TrackTooltipText);
+			}
 		}
 	}
 
diff --git a/Engine/Source/Editor/MovieSceneTools/Private/TrackEditors/PrimitiveMaterialTrackEditor.cpp b/Engine/Source/Editor/MovieSceneTools/Private/TrackEditors/PrimitiveMaterialTrackEditor.cpp
--- a/Engine/Source/Editor/MovieSceneTools/Private/TrackEditors/PrimitiveMaterialTrackEditor.cpp
+++ b/Engine/Source/Editor/MovieSceneTools/Private/TrackEditors/PrimitiveMaterialTrackEditor.cpp
@@ -5,8 +5,10 @@
 #include "Sections/MovieScenePrimitiveMaterialSection.h"
 #include "ISequencerModule.h"
 #include "Components/PrimitiveComponent.h"
+#include "Components/DecalComponent.h"
 #include "Materials/MaterialInterface.h"
 #include "Algo/Find.h"
+#include "Components/MeshComponent.h"
 
 
 #define LOCTEXT_NAMESPACE "PrimitiveMaterialTrackEditor"
@@ -31,53 +33,90 @@ void FPrimitiveMaterialTrackEditor::ExtendObjectBindingTrackMenu(TSharedRef<FExt
 
 void FPrimitiveMaterialTrackEditor::ConstructObjectBindingTrackMenu(FMenuBuilder& MenuBuilder, TArray<FGuid> ObjectBindings)
 {
-	auto GetMaterialIndexForTrack = [](UMovieSceneTrack* InTrack)
+	auto GetMaterialInfoForTrack = [](UMovieSceneTrack* InTrack)
 	{
 		UMovieScenePrimitiveMaterialTrack* MaterialTrack = Cast<UMovieScenePrimitiveMaterialTrack>(InTrack);
-		return MaterialTrack ? MaterialTrack->GetMaterialIndex() : INDEX_NONE;
+		return MaterialTrack ? MaterialTrack->GetMaterialInfo() : FComponentMaterialInfo();
 	};
 
-	int32 MinNumMaterials = TNumericLimits<int32>::Max();
-
-	for (TWeakObjectPtr<> WeakObject : GetSequencer()->FindObjectsInCurrentSequence(ObjectBindings[0]))
+	UObject* Object = GetSequencer()->FindSpawnedObjectOrTemplate(ObjectBindings[0]);
+	if (!Object)
 	{
-		UPrimitiveComponent* PrimitiveComponent = Cast<UPrimitiveComponent>(WeakObject.Get());
-		if (PrimitiveComponent)
-		{
-			MinNumMaterials = FMath::Min(MinNumMaterials, PrimitiveComponent->GetNumMaterials());
-		}
+		return;
 	}
 
-	if (MinNumMaterials > 0 && MinNumMaterials < TNumericLimits<int32>::Max())
+	USceneComponent* SceneComponent = Cast<USceneComponent>(Object);
+	if (!SceneComponent)
 	{
-		MenuBuilder.BeginSection(NAME_None, LOCTEXT("MaterialSwitcherTitle", "Material Switchers"));
+		return;
+	}	
 
-		const UMovieScene*        MovieScene = GetFocusedMovieScene();
-		const FMovieSceneBinding* Binding    = Algo::FindBy(MovieScene->GetBindings(), ObjectBindings[0], &FMovieSceneBinding::GetObjectGuid);
+	const UMovieScene* MovieScene = GetFocusedMovieScene();
+	const FMovieSceneBinding* Binding = Algo::FindBy(MovieScene->GetBindings(), ObjectBindings[0], &FMovieSceneBinding::GetObjectGuid);
 
-		check(Binding);
-
-		for (int32 Index = 0; Index < MinNumMaterials; ++Index)
+	if (UPrimitiveComponent* PrimitiveComponent = Cast<UPrimitiveComponent>(SceneComponent))
+	{
+		int32 NumMaterials = PrimitiveComponent->GetNumMaterials();
+		TArray<FName> MaterialSlotNames = PrimitiveComponent->GetMaterialSlotNames();
+		UMeshComponent* MeshComponent = Cast<UMeshComponent>(SceneComponent);
+		if (NumMaterials > 0 || MeshComponent)
 		{
-			const bool bAlreadyExists = Algo::FindBy(Binding->GetTracks(), Index, GetMaterialIndexForTrack) != nullptr;
+			MenuBuilder.BeginSection(NAME_None, LOCTEXT("MaterialSwitcherTitle", "Material Switchers"));
+			{
+				for (int32 MaterialIndex = 0; MaterialIndex < NumMaterials; MaterialIndex++)
+				{
+					FName MaterialSlotName = MaterialSlotNames.IsValidIndex(MaterialIndex) ? MaterialSlotNames[MaterialIndex] : FName();
+					FComponentMaterialInfo MaterialInfo{ MaterialSlotName, MaterialIndex, EComponentMaterialType::IndexedMaterial };
+
+					const bool bAlreadyExists = Algo::FindBy(Binding->GetTracks(), MaterialInfo, GetMaterialInfoForTrack) != nullptr;
+					if (bAlreadyExists)
+					{
+						continue;
+					}
+					FUIAction AddMaterialSwitcherAction(FExecuteAction::CreateSP(this, &FPrimitiveMaterialTrackEditor::CreateTrackForElement, ObjectBindings, MaterialInfo));
+					FText MaterialSwitcherLabel = !MaterialSlotName.IsNone() ?
+						FText::Format(LOCTEXT("MaterialSlot_Format", "Material Slot {0} Switcher"), FText::FromName(MaterialSlotName)) :
+						FText::Format(LOCTEXT("MaterialID_Format", "Material Element {0} Switcher"), FText::AsNumber(MaterialIndex));
+					FText MaterialSwitcherTooltip = !MaterialSlotName.IsNone() ?
+						FText::Format(LOCTEXT("MaterialSlotTooltip_Format", "Add material switcher for slot {0}, index {1}"), FText::FromName(MaterialSlotName), FText::AsNumber(MaterialIndex)) :
+						FText::Format(LOCTEXT("MaterialIDTooltip_Format", "Add material switcher for element {0}"), FText::AsNumber(MaterialIndex));
+					MenuBuilder.AddMenuEntry(MaterialSwitcherLabel, MaterialSwitcherTooltip, FSlateIcon(), AddMaterialSwitcherAction);
+				}
+				if (MeshComponent)
+				{
+					FComponentMaterialInfo MaterialInfo{ FName(), 0, EComponentMaterialType::OverlayMaterial };
+					const bool bAlreadyExists = Algo::FindBy(Binding->GetTracks(), MaterialInfo, GetMaterialInfoForTrack) != nullptr;
+					if (!bAlreadyExists)
+					{
+						FUIAction AddMaterialSwitcherAction(FExecuteAction::CreateSP(this, &FPrimitiveMaterialTrackEditor::CreateTrackForElement, ObjectBindings, MaterialInfo));
+						FText OverlayMaterialSwitcherLabel = LOCTEXT("OverlayMaterialSwitcher_Format", "Overlay Material Switcher");
+						FText OverlayMaterialSwitcherTooltip = LOCTEXT("OverlayMaterialSwitcherTooltip_Format", "Add overlay material switcher");
+						MenuBuilder.AddMenuEntry(OverlayMaterialSwitcherLabel, OverlayMaterialSwitcherTooltip, FSlateIcon(), AddMaterialSwitcherAction);
+					}
+				}
+			}
+			MenuBuilder.EndSection();
+		}
+	}
+	else if (UDecalComponent* DecalComponent = Cast<UDecalComponent>(SceneComponent))
+	{
+		MenuBuilder.BeginSection("Materials", LOCTEXT("MaterialSection", "Material Parameters"));
+		{
+			FComponentMaterialInfo MaterialInfo{ FName(), 0, EComponentMaterialType::DecalMaterial };
+			const bool bAlreadyExists = Algo::FindBy(Binding->GetTracks(), MaterialInfo, GetMaterialInfoForTrack) != nullptr;
 			if (!bAlreadyExists)
 			{
-				MenuBuilder.AddMenuEntry(
-					FText::Format(LOCTEXT("MaterialID_Format", "Material Element {0} Switcher"), FText::AsNumber(Index)),
-					FText(),
-					FSlateIcon(),
-					FUIAction(
-						FExecuteAction::CreateSP(this, &FPrimitiveMaterialTrackEditor::CreateTrackForElement, ObjectBindings, Index)
-					)
-				);
+				FUIAction AddMaterialSwitcherAction(FExecuteAction::CreateSP(this, &FPrimitiveMaterialTrackEditor::CreateTrackForElement, ObjectBindings, MaterialInfo));
+				FText DecalMaterialSwitcherLabel = LOCTEXT("DecalMaterialSwitcher_Format", "Decal Material Switcher");
+				FText DecalMaterialSwitcherTooltip = LOCTEXT("DecalMaterialSwitcherTooltip_Format", "Add decal material switcher");
+				MenuBuilder.AddMenuEntry(DecalMaterialSwitcherLabel, DecalMaterialSwitcherTooltip, FSlateIcon(), AddMaterialSwitcherAction);
 			}
 		}
-
 		MenuBuilder.EndSection();
 	}
 }
 
-void FPrimitiveMaterialTrackEditor::CreateTrackForElement(TArray<FGuid> ObjectBindingIDs, int32 MaterialIndex)
+void FPrimitiveMaterialTrackEditor::CreateTrackForElement(TArray<FGuid> ObjectBindingIDs, FComponentMaterialInfo MaterialInfo)
 {
 	UMovieScene* MovieScene = GetFocusedMovieScene();
 
@@ -87,8 +126,41 @@ void FPrimitiveMaterialTrackEditor::CreateTrackForElement(TArray<FGuid> ObjectBi
 	for (FGuid ObjectBindingID : ObjectBindingIDs)
 	{
 		UMovieScenePrimitiveMaterialTrack* NewTrack = MovieScene->AddTrack<UMovieScenePrimitiveMaterialTrack>(ObjectBindingID);
-		NewTrack->SetMaterialIndex(MaterialIndex);
-		NewTrack->SetDisplayName(FText::Format(LOCTEXT("MaterialTrackName_Format", "Material Element {0}"), FText::AsNumber(MaterialIndex)));
+		NewTrack->SetMaterialInfo(MaterialInfo);
+		// Construct display names from MaterialInfo
+		FText TrackDisplayName;
+		FText TrackTooltipText;
+		switch (MaterialInfo.MaterialType)
+		{
+		case EComponentMaterialType::Empty:
+			break;
+		case EComponentMaterialType::IndexedMaterial:
+			TrackDisplayName = !MaterialInfo.MaterialSlotName.IsNone() ? FText::Format(LOCTEXT("SlotMaterialSwitcherTrackName", "Material Slot: {0}"), FText::FromName(MaterialInfo.MaterialSlotName))
+				: FText::Format(LOCTEXT("IndexedMaterialSwitcherTrackName", "Material Element {0}"), FText::AsNumber(MaterialInfo.MaterialSlotIndex));
+
+			TrackTooltipText = !MaterialInfo.MaterialSlotName.IsNone() ? FText::Format(LOCTEXT("SlotMaterialSwitcherTrackTooltip", "Material switcher for {0} at index {1}"), FText::FromName(MaterialInfo.MaterialSlotName), FText::AsNumber(MaterialInfo.MaterialSlotIndex))
+				: FText::Format(LOCTEXT("IndexedMaterialSwitcherTrackTooltip", "Material switcher for element at index {0}"), FText::AsNumber(MaterialInfo.MaterialSlotIndex));
+			break;
+		case EComponentMaterialType::OverlayMaterial:
+			TrackDisplayName = LOCTEXT("OverlayMaterialSwitcherTrackName", "Overlay Material");
+			TrackTooltipText = LOCTEXT("OverlayMaterialSwitcherTrackTooltip", "Material switcher for overlay material");
+			break;
+		case EComponentMaterialType::DecalMaterial:
+			TrackDisplayName = LOCTEXT("DecalMaterialSwitcherTrackName", "Decal Material");
+			TrackTooltipText = LOCTEXT("DecalMaterialSwitcherTrackTooltip", "Material switcher for decal material");
+			break;
+		default:
+			break;
+
+		}
+		if (!TrackDisplayName.IsEmpty())
+		{
+			NewTrack->SetDisplayName(TrackDisplayName);
+		}
+		if (!TrackTooltipText.IsEmpty())
+		{
+			NewTrack->SetDisplayNameTooltipText(TrackTooltipText);
+		}
 
 		NewTrack->AddSection(*NewTrack->CreateNewSection());
 	}
diff --git a/Engine/Source/Editor/MovieSceneTools/Private/TrackEditors/PrimitiveMaterialTrackEditor.h b/Engine/Source/Editor/MovieSceneTools/Private/TrackEditors/PrimitiveMaterialTrackEditor.h
--- a/Engine/Source/Editor/MovieSceneTools/Private/TrackEditors/PrimitiveMaterialTrackEditor.h
+++ b/Engine/Source/Editor/MovieSceneTools/Private/TrackEditors/PrimitiveMaterialTrackEditor.h
@@ -3,6 +3,7 @@
 #pragma once
 
 #include "KeyframeTrackEditor.h"
+#include "Tracks/MovieSceneMaterialTrack.h"
 #include "Tracks/MovieScenePrimitiveMaterialTrack.h"
 #include "Materials/MaterialInterface.h"
 
@@ -22,5 +23,5 @@ public:
 	virtual void ExtendObjectBindingTrackMenu(TSharedRef<FExtender> Extender, const TArray<FGuid>& ObjectBindings, const UClass* ObjectClass) override;
 
 	void ConstructObjectBindingTrackMenu(FMenuBuilder& MenuBuilder, TArray<FGuid> ObjectBindings);
-	void CreateTrackForElement(TArray<FGuid> ObjectBindingIDs, int32 MaterialIndex);
+	void CreateTrackForElement(TArray<FGuid> ObjectBindingIDs, FComponentMaterialInfo MaterialInfo);
 };
diff --git a/Engine/Source/Editor/MovieSceneTools/Public/TrackEditors/MaterialTrackEditor.h b/Engine/Source/Editor/MovieSceneTools/Public/TrackEditors/MaterialTrackEditor.h
--- a/Engine/Source/Editor/MovieSceneTools/Public/TrackEditors/MaterialTrackEditor.h
+++ b/Engine/Source/Editor/MovieSceneTools/Public/TrackEditors/MaterialTrackEditor.h
@@ -17,6 +17,7 @@ class UMaterial;
 class UMaterialInterface;
 class UMovieSceneMaterialTrack;
 class USceneComponent;
+struct FComponentMaterialInfo;
 
 /**
  * Track editor for material parameters.
@@ -103,6 +104,6 @@ private:
 
 	void ConstructObjectBindingTrackMenu(FMenuBuilder& MenuBuilder, TArray<FGuid> ObjectBindings);
 	/** Callback for executing the add component material track. */
-	void HandleAddComponentMaterialActionExecute(USceneComponent* Component, int32 MaterialIndex);
+	void HandleAddComponentMaterialActionExecute(USceneComponent* Component, FComponentMaterialInfo MaterialInfo);
 
 };
diff --git a/Engine/Source/Editor/PropertyEditor/Private/MaterialList.cpp b/Engine/Source/Editor/PropertyEditor/Private/MaterialList.cpp
--- a/Engine/Source/Editor/PropertyEditor/Private/MaterialList.cpp
+++ b/Engine/Source/Editor/PropertyEditor/Private/MaterialList.cpp
@@ -38,11 +38,11 @@ public:
 	 * @param Material		The material being used
 	 * @param bCanBeReplced	Whether or not the material can be replaced by a user
 	 */
-	virtual void AddMaterial( uint32 SlotIndex, UMaterialInterface* Material, bool bCanBeReplaced, UActorComponent* InCurrentComponent ) override
+	virtual void AddMaterial( uint32 SlotIndex, UMaterialInterface* Material, bool bCanBeReplaced, UActorComponent* InCurrentComponent, FName SlotName ) override
 	{
 		int32 NumMaterials = MaterialSlots.Num();
 
-		FMaterialListItem MaterialItem( Material, SlotIndex, bCanBeReplaced, InCurrentComponent ); 
+		FMaterialListItem MaterialItem( Material, SlotIndex, bCanBeReplaced, InCurrentComponent, SlotName ); 
 		if( !UniqueMaterials.Contains( MaterialItem ) ) 
 		{
 			MaterialSlots.Add( MaterialItem );
diff --git a/Engine/Source/Editor/PropertyEditor/Private/PropertyCustomizationHelpers.cpp b/Engine/Source/Editor/PropertyEditor/Private/PropertyCustomizationHelpers.cpp
--- a/Engine/Source/Editor/PropertyEditor/Private/PropertyCustomizationHelpers.cpp
+++ b/Engine/Source/Editor/PropertyEditor/Private/PropertyCustomizationHelpers.cpp
@@ -1473,6 +1473,7 @@ void SMaterialSlotWidget::Construct(const FArguments& InArgs, int32 SlotIndex, b
 			[
 				SNew(SEditableTextBox)
 				.Text(InArgs._MaterialName)
+				.IsReadOnly(InArgs._IsMaterialSlotNameReadOnly)
 				.OnTextChanged(InArgs._OnMaterialNameChanged)
 				.OnTextCommitted(InArgs._OnMaterialNameCommitted)
 				.Font(IDetailLayoutBuilder::GetDetailFont())
diff --git a/Engine/Source/Editor/PropertyEditor/Public/MaterialList.h b/Engine/Source/Editor/PropertyEditor/Public/MaterialList.h
--- a/Engine/Source/Editor/PropertyEditor/Public/MaterialList.h
+++ b/Engine/Source/Editor/PropertyEditor/Public/MaterialList.h
@@ -96,7 +96,7 @@ public:
 	 * @param bCanBeReplced Whether or not the material can be replaced by a user.
 	 * @param InCurrentComponent The current component using the current material
 	 */
-	virtual void AddMaterial( uint32 SlotIndex, UMaterialInterface* Material, bool bCanBeReplaced, UActorComponent* InCurrentComponent = nullptr) = 0;
+	virtual void AddMaterial( uint32 SlotIndex, UMaterialInterface* Material, bool bCanBeReplaced, UActorComponent* InCurrentComponent = nullptr, FName SlotName = FName()) = 0;
 };
 
 /**
@@ -116,21 +116,25 @@ struct FMaterialListItem
 	/** The current component using the current material */
 	TWeakObjectPtr<UActorComponent> CurrentComponent;
 
-	FMaterialListItem( UMaterialInterface* InMaterial = NULL, uint32 InSlotIndex = 0, bool bInCanBeReplaced = false, UActorComponent* InCurrentComponent = nullptr)
+	/** User-specified name for the slot on a component where this material is at (mesh element). May be empty. */
+	FName SlotName;
+
+	FMaterialListItem( UMaterialInterface* InMaterial = NULL, uint32 InSlotIndex = 0, bool bInCanBeReplaced = false, UActorComponent* InCurrentComponent = nullptr, FName InSlotName = FName())
 		: Material( InMaterial )
 		, SlotIndex( InSlotIndex )
 		, bCanBeReplaced( bInCanBeReplaced )
 		, CurrentComponent(InCurrentComponent)
+		, SlotName(InSlotName)
 	{}
 
 	friend uint32 GetTypeHash( const FMaterialListItem& InItem )
 	{
-		return GetTypeHash( InItem.Material ) + InItem.SlotIndex ;
+		return HashCombine(HashCombine(GetTypeHash(InItem.Material), GetTypeHash(InItem.SlotIndex)), GetTypeHash(InItem.SlotName));
 	}
 
 	bool operator==( const FMaterialListItem& Other ) const
 	{
-		return Material == Other.Material && SlotIndex == Other.SlotIndex ;
+		return Material == Other.Material && SlotIndex == Other.SlotIndex && SlotName == SlotName;
 	}
 
 	bool operator!=( const FMaterialListItem& Other ) const
diff --git a/Engine/Source/Editor/PropertyEditor/Public/PropertyCustomizationHelpers.h b/Engine/Source/Editor/PropertyEditor/Public/PropertyCustomizationHelpers.h
--- a/Engine/Source/Editor/PropertyEditor/Public/PropertyCustomizationHelpers.h
+++ b/Engine/Source/Editor/PropertyEditor/Public/PropertyCustomizationHelpers.h
@@ -555,6 +555,7 @@ class SMaterialSlotWidget : public SCompoundWidget
 	SLATE_BEGIN_ARGS(SMaterialSlotWidget)
 	{}
 		SLATE_ATTRIBUTE(FText, MaterialName)
+		SLATE_ATTRIBUTE(bool, IsMaterialSlotNameReadOnly)
 		SLATE_EVENT(FOnTextChanged, OnMaterialNameChanged)
 		SLATE_EVENT(FOnTextCommitted, OnMaterialNameCommitted)
 		SLATE_ATTRIBUTE(bool, CanDeleteMaterialSlot)
diff --git a/Engine/Source/Runtime/Engine/Classes/Components/MeshComponent.h b/Engine/Source/Runtime/Engine/Classes/Components/MeshComponent.h
--- a/Engine/Source/Runtime/Engine/Classes/Components/MeshComponent.h
+++ b/Engine/Source/Runtime/Engine/Classes/Components/MeshComponent.h
@@ -35,15 +35,6 @@ public:
 	UFUNCTION(BlueprintCallable, Category="Rendering|Material")
 	ENGINE_API virtual TArray<class UMaterialInterface*> GetMaterials() const;
 
-	UFUNCTION(BlueprintCallable, Category = "Rendering|Material")
-	ENGINE_API virtual int32 GetMaterialIndex(FName MaterialSlotName) const;
-
-	UFUNCTION(BlueprintCallable, Category = "Rendering|Material")
-	ENGINE_API virtual TArray<FName> GetMaterialSlotNames() const;
-
-	UFUNCTION(BlueprintCallable, Category = "Rendering|Material")
-	ENGINE_API virtual bool IsMaterialSlotNameValid(FName MaterialSlotName) const;
-
 	/** Determines if we use the nanite overrides from any materials */
 	virtual bool UseNaniteOverrideMaterials() const { return false; }
 
@@ -100,6 +91,7 @@ public:
 	//~ Begin UPrimitiveComponent Interface
 	ENGINE_API virtual int32 GetNumMaterials() const override;
 	ENGINE_API virtual UMaterialInterface* GetMaterial(int32 ElementIndex) const override;
+	ENGINE_API virtual UMaterialInterface* GetMaterialByName(FName MaterialSlotName) const override;
 	ENGINE_API virtual void SetMaterial(int32 ElementIndex, UMaterialInterface* Material) override;
 	ENGINE_API virtual void SetMaterialByName(FName MaterialSlotName, class UMaterialInterface* Material) override;
 	ENGINE_API virtual void GetUsedMaterials(TArray<UMaterialInterface*>& OutMaterials, bool bGetDebugMaterials = false) const override;	
diff --git a/Engine/Source/Runtime/Engine/Classes/Components/PrimitiveComponent.h b/Engine/Source/Runtime/Engine/Classes/Components/PrimitiveComponent.h
--- a/Engine/Source/Runtime/Engine/Classes/Components/PrimitiveComponent.h
+++ b/Engine/Source/Runtime/Engine/Classes/Components/PrimitiveComponent.h
@@ -1433,6 +1433,23 @@ public:
 		return GetMaterial(ElementIndex);
 	}
 
+	UFUNCTION(BlueprintCallable, Category = "Rendering|Material")
+	ENGINE_API virtual int32 GetMaterialIndex(FName MaterialSlotName) const;
+
+	UFUNCTION(BlueprintCallable, Category = "Rendering|Material")
+	ENGINE_API virtual TArray<FName> GetMaterialSlotNames() const;
+
+	UFUNCTION(BlueprintCallable, Category = "Rendering|Material")
+	ENGINE_API virtual bool IsMaterialSlotNameValid(FName MaterialSlotName) const;
+
+	/**
+	* Returns the material used by the element in the slot with the specified name.
+	* @param MaterialSlotName - The slot name to access the material of.
+	* @return the material used in the slot specified, or null if none exists or the slot name is not found.
+	*/
+	UFUNCTION(BlueprintCallable, Category = "Rendering|Material")
+	ENGINE_API virtual class UMaterialInterface* GetMaterialByName(FName MaterialSlotName) const;
+
 	/**
 	 * Changes the material applied to an element of the mesh.
 	 * @param ElementIndex - The element to access the material of.
diff --git a/Engine/Source/Runtime/Engine/Classes/Particles/ParticleSystemComponent.h b/Engine/Source/Runtime/Engine/Classes/Particles/ParticleSystemComponent.h
--- a/Engine/Source/Runtime/Engine/Classes/Particles/ParticleSystemComponent.h
+++ b/Engine/Source/Runtime/Engine/Classes/Particles/ParticleSystemComponent.h
@@ -1476,6 +1476,7 @@ public:
 	UFUNCTION(BlueprintCallable, Category = "Rendering|Material")
 	ENGINE_API virtual class UMaterialInstanceDynamic* CreateNamedDynamicMaterialInstance(FName InName, class UMaterialInterface* SourceMaterial = NULL);
 
+	ENGINE_API virtual UMaterialInterface* GetMaterialByName(FName MaterialSlotName) const override;
 	ENGINE_API virtual void SetMaterialByName(FName MaterialSlotName, class UMaterialInterface* SourceMaterial) override;
 
 	/** Returns a named material. If this named material is not found, returns NULL. */
diff --git a/Engine/Source/Runtime/Engine/Private/Components/MeshComponent.cpp b/Engine/Source/Runtime/Engine/Private/Components/MeshComponent.cpp
--- a/Engine/Source/Runtime/Engine/Private/Components/MeshComponent.cpp
+++ b/Engine/Source/Runtime/Engine/Private/Components/MeshComponent.cpp
@@ -49,6 +49,14 @@ UMaterialInterface* UMeshComponent::GetMaterial(int32 ElementIndex) const
 	return OutMaterial;
 }
 
+UMaterialInterface* UMeshComponent::GetMaterialByName(FName MaterialSlotName) const
+{
+	int32 MaterialIndex = GetMaterialIndex(MaterialSlotName);
+	if (MaterialIndex < 0)
+		return nullptr;
+	return GetMaterial(MaterialIndex);
+}
+
 void UMeshComponent::SetMaterial(int32 ElementIndex, UMaterialInterface* Material)
 {
 	if (ElementIndex >= 0)
@@ -341,24 +349,6 @@ TArray<class UMaterialInterface*> UMeshComponent::GetMaterials() const
 	return OutMaterials;
 }
 
-int32 UMeshComponent::GetMaterialIndex(FName MaterialSlotName) const
-{
-	// This function should be overridden
-	return INDEX_NONE;
-}
-
-TArray<FName> UMeshComponent::GetMaterialSlotNames() const
-{
-	// This function should be overridden
-	return TArray<FName>();
-}
-
-bool UMeshComponent::IsMaterialSlotNameValid(FName MaterialSlotName) const
-{
-	// This function should be overridden
-	return false;
-}
-
 void UMeshComponent::SetScalarParameterValueOnMaterials(const FName ParameterName, const float ParameterValue)
 {
 	if (!bEnableMaterialParameterCaching)
diff --git a/Engine/Source/Runtime/Engine/Private/Components/PrimitiveComponent.cpp b/Engine/Source/Runtime/Engine/Private/Components/PrimitiveComponent.cpp
--- a/Engine/Source/Runtime/Engine/Private/Components/PrimitiveComponent.cpp
+++ b/Engine/Source/Runtime/Engine/Private/Components/PrimitiveComponent.cpp
@@ -1996,19 +1996,46 @@ void UPrimitiveComponent::SetBoundsScale(float NewBoundsScale)
 
 UMaterialInterface* UPrimitiveComponent::GetMaterial(int32 Index) const
 {
-	return NULL;
+	// This function should be overridden
+	return nullptr;
+}
+
+int32 UPrimitiveComponent::GetMaterialIndex(FName MaterialSlotName) const
+{
+	// This function should be overridden
+	return INDEX_NONE;
+}
+
+TArray<FName> UPrimitiveComponent::GetMaterialSlotNames() const
+{
+	// This function should be overridden
+	return TArray<FName>();
+}
+
+bool UPrimitiveComponent::IsMaterialSlotNameValid(FName MaterialSlotName) const
+{
+	// This function should be overridden
+	return false;
+}
+
+UMaterialInterface* UPrimitiveComponent::GetMaterialByName(FName MaterialSlotName) const
+{
+	return nullptr;
 }
 
 void UPrimitiveComponent::SetMaterial(int32 Index, UMaterialInterface* InMaterial)
 {
+	// This function should be overridden
 }
 
 void UPrimitiveComponent::SetMaterialByName(FName MaterialSlotName, class UMaterialInterface* Material)
 {
+	// This function should be overridden
 }
 
 int32 UPrimitiveComponent::GetNumMaterials() const
 {
+	// This function should be overridden
 	return 0;
 }
 
diff --git a/Engine/Source/Runtime/Engine/Private/Particles/ParticleComponents.cpp b/Engine/Source/Runtime/Engine/Private/Particles/ParticleComponents.cpp
--- a/Engine/Source/Runtime/Engine/Private/Particles/ParticleComponents.cpp
+++ b/Engine/Source/Runtime/Engine/Private/Particles/ParticleComponents.cpp
@@ -8280,6 +8280,13 @@ UMaterialInstanceDynamic* UParticleSystemComponent::CreateNamedDynamicMaterialIn
 	return MID;
 }
 
+
+UMaterialInterface* UParticleSystemComponent::GetMaterialByName(FName MaterialSlotName) const
+{
+	return GetNamedMaterial(MaterialSlotName);
+}
+
+
 void UParticleSystemComponent::SetMaterialByName(FName MaterialSlotName, class UMaterialInterface* SourceMaterial)
 {
 	int32 Index = GetNamedMaterialIndex(MaterialSlotName);
diff --git a/Engine/Source/Runtime/MovieSceneTracks/Private/MovieSceneTracksComponentTypes.cpp b/Engine/Source/Runtime/MovieSceneTracks/Private/MovieSceneTracksComponentTypes.cpp
--- a/Engine/Source/Runtime/MovieSceneTracks/Private/MovieSceneTracksComponentTypes.cpp
+++ b/Engine/Source/Runtime/MovieSceneTracks/Private/MovieSceneTracksComponentTypes.cpp
@@ -540,9 +540,9 @@ FMovieSceneTracksComponentTypes::FMovieSceneTracksComponentTypes()
 	ComponentRegistry->NewComponentType(&LevelVisibility, TEXT("Level Visibility"));
 	ComponentRegistry->NewComponentType(&DataLayer, TEXT("Data Layer"));
 
-	ComponentRegistry->NewComponentType(&ComponentMaterialIndex, TEXT("Component Material Index"), EComponentTypeFlags::CopyToChildren | EComponentTypeFlags::CopyToOutput);
-	ComponentRegistry->NewComponentType(&BoundMaterial,          TEXT("Bound Material"), EComponentTypeFlags::CopyToChildren | EComponentTypeFlags::CopyToOutput);
-	ComponentRegistry->NewComponentType(&MPC,                    TEXT("Material Parameter Collection"), EComponentTypeFlags::CopyToChildren | EComponentTypeFlags::CopyToOutput);
+	ComponentRegistry->NewComponentType(&ComponentMaterialInfo,		TEXT("Component Material Info"), EComponentTypeFlags::CopyToChildren | EComponentTypeFlags::CopyToOutput);
+	ComponentRegistry->NewComponentType(&BoundMaterial,				TEXT("Bound Material"), EComponentTypeFlags::CopyToChildren | EComponentTypeFlags::CopyToOutput);
+	ComponentRegistry->NewComponentType(&MPC,						TEXT("Material Parameter Collection"), EComponentTypeFlags::CopyToChildren | EComponentTypeFlags::CopyToOutput);
 
 	ComponentRegistry->NewComponentType(&BoolParameterName,      TEXT("Bool Parameter Name"), EComponentTypeFlags::CopyToChildren | EComponentTypeFlags::CopyToOutput);
 	ComponentRegistry->NewComponentType(&ScalarParameterName,    TEXT("Scalar Parameter Name"), EComponentTypeFlags::CopyToChildren | EComponentTypeFlags::CopyToOutput);
diff --git a/Engine/Source/Runtime/MovieSceneTracks/Private/Sections/MovieSceneComponentMaterialParameterSection.cpp b/Engine/Source/Runtime/MovieSceneTracks/Private/Sections/MovieSceneComponentMaterialParameterSection.cpp
--- a/Engine/Source/Runtime/MovieSceneTracks/Private/Sections/MovieSceneComponentMaterialParameterSection.cpp
+++ b/Engine/Source/Runtime/MovieSceneTracks/Private/Sections/MovieSceneComponentMaterialParameterSection.cpp
@@ -323,11 +323,11 @@ void UMovieSceneComponentMaterialParameterSection::ImportEntityImpl(UMovieSceneE
 
 	FGuid ObjectBindingID = Params.GetObjectBindingID();
 	
-	// Find material index from the outer track
-	int32 MaterialIndex = 0;
+	// Find material info from the outer track
+	FComponentMaterialInfo MaterialInfo;
 	if (UMovieSceneComponentMaterialTrack* MaterialTrack = GetTypedOuter<UMovieSceneComponentMaterialTrack>())
 	{
-		MaterialIndex = MaterialTrack->GetMaterialIndex();
+		MaterialInfo = MaterialTrack->GetMaterialInfo();
 	}
 
 	TEntityBuilder<TAddConditional<FGuid>> BaseBuilder = FEntityBuilder()
@@ -345,7 +345,7 @@ void UMovieSceneComponentMaterialParameterSection::ImportEntityImpl(UMovieSceneE
 				BaseBuilder
 				.Add(TracksComponentTypes->ScalarMaterialParameterInfo, Scalar.ParameterInfo)
 				.Add(BuiltInComponentTypes->FloatChannel[0], &Scalar.ParameterCurve)
-				.Add(TracksComponentTypes->ComponentMaterialIndex, MaterialIndex)
+				.Add(TracksComponentTypes->ComponentMaterialInfo, MaterialInfo)
 				// If the section has no valid blend type (legacy data), make it use absolute blending.
 				// Otherwise, the base section class will add the appropriate blend type tag in BuildDefaultComponents.
 				.AddTagConditional(BuiltInComponentTypes->Tags.AbsoluteBlend, !GetBlendType().IsValid())
@@ -365,7 +365,7 @@ void UMovieSceneComponentMaterialParameterSection::ImportEntityImpl(UMovieSceneE
 				.AddConditional(BuiltInComponentTypes->FloatChannel[1], &Color.GreenCurve, Color.GreenCurve.HasAnyData())
 				.AddConditional(BuiltInComponentTypes->FloatChannel[2], &Color.BlueCurve, Color.BlueCurve.HasAnyData())
 				.AddConditional(BuiltInComponentTypes->FloatChannel[3], &Color.AlphaCurve, Color.AlphaCurve.HasAnyData())
-				.Add(TracksComponentTypes->ComponentMaterialIndex, MaterialIndex)
+				.Add(TracksComponentTypes->ComponentMaterialInfo, MaterialInfo)
 				// If the section has no valid blend type (legacy data), make it use absolute blending.
 				// Otherwise, the base section class will add the appropriate blend type tag in BuildDefaultComponents.
 				.AddTagConditional(BuiltInComponentTypes->Tags.AbsoluteBlend, !GetBlendType().IsValid())
diff --git a/Engine/Source/Runtime/MovieSceneTracks/Private/Sections/MovieScenePrimitiveMaterialSection.cpp b/Engine/Source/Runtime/MovieSceneTracks/Private/Sections/MovieScenePrimitiveMaterialSection.cpp
--- a/Engine/Source/Runtime/MovieSceneTracks/Private/Sections/MovieScenePrimitiveMaterialSection.cpp
+++ b/Engine/Source/Runtime/MovieSceneTracks/Private/Sections/MovieScenePrimitiveMaterialSection.cpp
@@ -2,6 +2,7 @@
 
 #include "Sections/MovieScenePrimitiveMaterialSection.h"
 #include "Tracks/MovieScenePrimitiveMaterialTrack.h"
+#include "Tracks/MovieSceneMaterialTrack.h"
 #include "Channels/MovieSceneChannelProxy.h"
 #include "Materials/MaterialInterface.h"
 #include "EntitySystem/MovieSceneEntityBuilder.h"
@@ -95,7 +96,7 @@ void UMovieScenePrimitiveMaterialSection::ImportEntityImpl(UMovieSceneEntitySyst
 	UMovieScenePrimitiveMaterialTrack* Track = GetTypedOuter<UMovieScenePrimitiveMaterialTrack>();
 	check(Track);
 
-	const int32 MaterialIndex = Track->GetMaterialIndex();
+	const FComponentMaterialInfo& MaterialInfo = Track->GetMaterialInfo();
 
 	FBuiltInComponentTypes* BuiltInComponentTypes = FBuiltInComponentTypes::Get();
 	FMovieSceneTracksComponentTypes* TracksComponentTypes = FMovieSceneTracksComponentTypes::Get();
@@ -119,6 +120,6 @@ void UMovieScenePrimitiveMaterialSection::ImportEntityImpl(UMovieSceneEntitySyst
 		FEntityBuilder()
 		.AddConditional(BuiltInComponentTypes->GenericObjectBinding, ObjectBindingID, ObjectBindingID.IsValid())
 		.Add(BuiltInComponentTypes->ObjectResult, FObjectComponent::Strong(ValueToImport.Get()))
-		.Add(TracksComponentTypes->ComponentMaterialIndex, MaterialIndex)
+		.Add(TracksComponentTypes->ComponentMaterialInfo, MaterialInfo)
 	);
 }
diff --git a/Engine/Source/Runtime/MovieSceneTracks/Private/Systems/MovieSceneComponentMaterialSystem.cpp b/Engine/Source/Runtime/MovieSceneTracks/Private/Systems/MovieSceneComponentMaterialSystem.cpp
--- a/Engine/Source/Runtime/MovieSceneTracks/Private/Systems/MovieSceneComponentMaterialSystem.cpp
+++ b/Engine/Source/Runtime/MovieSceneTracks/Private/Systems/MovieSceneComponentMaterialSystem.cpp
@@ -11,9 +11,11 @@
 #include "Systems/MovieScenePiecewiseDoubleBlenderSystem.h"
 
 #include "Materials/MaterialInstanceDynamic.h"
+#include "Tracks/MovieSceneMaterialTrack.h"
 
 #include "Components/PrimitiveComponent.h"
 #include "Components/DecalComponent.h"
+#include "Components/MeshComponent.h"
 
 #include UE_INLINE_GENERATED_CPP_BY_NAME(MovieSceneComponentMaterialSystem)
 
@@ -22,12 +24,12 @@ namespace UE::MovieScene
 
 FComponentMaterialAccessor::FComponentMaterialAccessor(const FComponentMaterialKey& InKey)
 	: Object(InKey.Object.ResolveObjectPtr())
-	, MaterialIndex(InKey.MaterialIndex)
+	, MaterialInfo(InKey.MaterialInfo)
 {}
 
-FComponentMaterialAccessor::FComponentMaterialAccessor(UObject* InObject, int32 InMaterialIndex)
+FComponentMaterialAccessor::FComponentMaterialAccessor(UObject* InObject, const FComponentMaterialInfo& InMaterialInfo)
 	: Object(InObject)
-	, MaterialIndex(InMaterialIndex)
+	, MaterialInfo(InMaterialInfo)
 {}
 
 FComponentMaterialAccessor::operator bool() const
@@ -37,52 +39,134 @@ FComponentMaterialAccessor::operator bool() const
 
 FString FComponentMaterialAccessor::ToString() const
 {
-	return FString::Printf(TEXT("object %s (element %d)"), *Object->GetPathName(), MaterialIndex);
+	return FString::Printf(TEXT("object %s (material info: %s)"), *Object->GetPathName(), *MaterialInfo.ToString());
 }
 
 UMaterialInterface* FComponentMaterialAccessor::GetMaterial() const
 {
-	if (UPrimitiveComponent* PrimitiveComponent = Cast<UPrimitiveComponent>(Object))
+	switch (MaterialInfo.MaterialType)
 	{
-		return PrimitiveComponent->GetMaterial(MaterialIndex);
-	}
-	else if (UDecalComponent* DecalComponent = Cast<UDecalComponent>(Object))
-	{
-		return DecalComponent->GetDecalMaterial();
+	case EComponentMaterialType::Empty:
+		break;
+	case EComponentMaterialType::IndexedMaterial:
+		if (UPrimitiveComponent* Component = Cast<UPrimitiveComponent>(Object))
+		{
+			UMaterialInterface* Material = nullptr;
+			if (!MaterialInfo.MaterialSlotName.IsNone())
+			{
+				Material = Component->GetMaterialByName(MaterialInfo.MaterialSlotName);
+			}
+			if (!Material)
+			{
+				Material = Component->GetMaterial(MaterialInfo.MaterialSlotIndex);
+			}
+			return Material;
+		}
+		break;
+	case EComponentMaterialType::OverlayMaterial:
+		if (UMeshComponent* Component = Cast<UMeshComponent>(Object))
+		{
+			return Component->GetOverlayMaterial();
+		}
+		break;
+	case EComponentMaterialType::DecalMaterial:
+		if (UDecalComponent* DecalComponent = Cast<UDecalComponent>(Object))
+		{
+			return DecalComponent->GetDecalMaterial();
+		}
+		break;
+	default:
+		break;
 	}
 	return nullptr;
 }
 
 void FComponentMaterialAccessor::SetMaterial(UMaterialInterface* InMaterial) const
 {
-	if (UPrimitiveComponent* PrimitiveComponent = Cast<UPrimitiveComponent>(Object))
+	switch (MaterialInfo.MaterialType)
 	{
-		PrimitiveComponent->SetMaterial(MaterialIndex, InMaterial);
-	}
-	else if (UDecalComponent* DecalComponent = Cast<UDecalComponent>(Object))
-	{
-		DecalComponent->SetDecalMaterial(InMaterial);
+	case EComponentMaterialType::Empty:
+		break;
+	case EComponentMaterialType::IndexedMaterial:
+		if (UPrimitiveComponent* Component = Cast<UPrimitiveComponent>(Object))
+		{
+			TArray<FName> MaterialSlots = Component->GetMaterialSlotNames();
+			if (!MaterialInfo.MaterialSlotName.IsNone() && MaterialSlots.Contains(MaterialInfo.MaterialSlotName))
+			{
+				Component->SetMaterialByName(MaterialInfo.MaterialSlotName, InMaterial);
+			}
+			else
+			{
+				Component->SetMaterial(MaterialInfo.MaterialSlotIndex, InMaterial);
+			}
+		}
+		break;
+	case EComponentMaterialType::OverlayMaterial:
+		if (UMeshComponent* Component = Cast<UMeshComponent>(Object))
+		{
+			 Component->SetOverlayMaterial(InMaterial);
+		}
+		break;
+	case EComponentMaterialType::DecalMaterial:
+		if (UDecalComponent* DecalComponent = Cast<UDecalComponent>(Object))
+		{
+			DecalComponent->SetDecalMaterial(InMaterial);
+		}
+		break;
+	default:
+		break;
 	}
 }
 
 UMaterialInstanceDynamic* FComponentMaterialAccessor::CreateDynamicMaterial(UMaterialInterface* InMaterial)
 {
-	// Need to create a new MID, either because the parent has changed, or because one doesn't already exist
-	if (UPrimitiveComponent* PrimitiveComponent = Cast<UPrimitiveComponent>(Object))
+	auto MakeDynamicMaterial = [this, InMaterial]()
 	{
 		TStringBuilder<128> DynamicName;
 		InMaterial->GetFName().ToString(DynamicName);
 		DynamicName.Append(TEXT("_Animated"));
-		FName UniqueDynamicName = MakeUniqueObjectName(Object, UMaterialInstanceDynamic::StaticClass() , DynamicName.ToString());
-
-		UMaterialInstanceDynamic* Result = UMaterialInstanceDynamic::Create(InMaterial, Object, UniqueDynamicName );
-		PrimitiveComponent->SetMaterial(MaterialIndex, Result);
-		return Result;
-	}
-	else if (UDecalComponent* DecalComponent = Cast<UDecalComponent>(Object))
+		FName UniqueDynamicName = MakeUniqueObjectName(Object, UMaterialInstanceDynamic::StaticClass(), DynamicName.ToString());
+		return UMaterialInstanceDynamic::Create(InMaterial, Object, UniqueDynamicName);
+	};
+	// Need to create a new MID, either because the parent has changed, or because one doesn't already exist
+	switch (MaterialInfo.MaterialType)
 	{
-		return DecalComponent->CreateDynamicMaterialInstance();
+	case EComponentMaterialType::Empty:
+		break;
+	case EComponentMaterialType::IndexedMaterial:
+		if (UPrimitiveComponent* Component = Cast<UPrimitiveComponent>(Object))
+		{
+			TArray<FName> MaterialSlots = Component->GetMaterialSlotNames();
+			UMaterialInstanceDynamic* Result = MakeDynamicMaterial();
+			if (!MaterialInfo.MaterialSlotName.IsNone() && MaterialSlots.Contains(MaterialInfo.MaterialSlotName))
+			{
+				Component->SetMaterialByName(MaterialInfo.MaterialSlotName, Result);
+			}
+			else
+			{
+				Component->SetMaterial(MaterialInfo.MaterialSlotIndex, Result);
+			}
+			return Result;
+		}
+		break;
+	case EComponentMaterialType::OverlayMaterial:
+		if (UMeshComponent* Component = Cast<UMeshComponent>(Object))
+		{
+			UMaterialInstanceDynamic* Result = MakeDynamicMaterial();
+			Component->SetOverlayMaterial(Result);
+			return Result;
+		}
+		break;
+	case EComponentMaterialType::DecalMaterial:
+		if (UDecalComponent* DecalComponent = Cast<UDecalComponent>(Object))
+		{
+			return DecalComponent->CreateDynamicMaterialInstance();
+		}
+		break;
+	default:
+		break;
 	}
+
 	return nullptr;
 }
 
@@ -99,9 +183,10 @@ UMovieSceneComponentMaterialSystem::UMovieSceneComponentMaterialSystem(const FOb
 	FBuiltInComponentTypes* BuiltInComponents = FBuiltInComponentTypes::Get();
 	FMovieSceneTracksComponentTypes* TracksComponents = FMovieSceneTracksComponentTypes::Get();
 
-	RelevantComponent = TracksComponents->ComponentMaterialIndex;
+	RelevantComponent = TracksComponents->ComponentMaterialInfo;
 	Phase = ESystemPhase::Instantiation;
 
+
 	if (HasAnyFlags(RF_ClassDefaultObject))
 	{
 		DefineComponentConsumer(GetClass(), BuiltInComponents->ObjectResult);
@@ -118,7 +203,7 @@ void UMovieSceneComponentMaterialSystem::OnLink()
 	SystemImpl.MaterialSwitcherStorage = Linker->PreAnimatedState.GetOrCreateStorage<FPreAnimatedComponentMaterialSwitcherStorage>();
 	SystemImpl.MaterialParameterStorage = Linker->PreAnimatedState.GetOrCreateStorage<FPreAnimatedComponentMaterialParameterStorage>();
 
-	SystemImpl.OnLink(Linker, FBuiltInComponentTypes::Get()->BoundObject, FMovieSceneTracksComponentTypes::Get()->ComponentMaterialIndex);
+	SystemImpl.OnLink(Linker, FBuiltInComponentTypes::Get()->BoundObject, FMovieSceneTracksComponentTypes::Get()->ComponentMaterialInfo);
 }
 
 void UMovieSceneComponentMaterialSystem::OnUnlink()
@@ -130,12 +215,12 @@ void UMovieSceneComponentMaterialSystem::OnRun(FSystemTaskPrerequisites& InPrere
 {
 	using namespace UE::MovieScene;
 
-	SystemImpl.OnRun(Linker, FBuiltInComponentTypes::Get()->BoundObject, FMovieSceneTracksComponentTypes::Get()->ComponentMaterialIndex, InPrerequisites, Subsequents);
+	SystemImpl.OnRun(Linker, FBuiltInComponentTypes::Get()->BoundObject, FMovieSceneTracksComponentTypes::Get()->ComponentMaterialInfo, InPrerequisites, Subsequents);
 }
 
 void UMovieSceneComponentMaterialSystem::SavePreAnimatedState(const FPreAnimationParameters& InParameters)
 {
 	using namespace UE::MovieScene;
 
-	SystemImpl.SavePreAnimatedState(Linker, FBuiltInComponentTypes::Get()->BoundObject, FMovieSceneTracksComponentTypes::Get()->ComponentMaterialIndex, InParameters);
+	SystemImpl.SavePreAnimatedState(Linker, FBuiltInComponentTypes::Get()->BoundObject, FMovieSceneTracksComponentTypes::Get()->ComponentMaterialInfo, InParameters);
 }
diff --git a/Engine/Source/Runtime/MovieSceneTracks/Private/Tracks/MovieSceneMaterialTrack.cpp b/Engine/Source/Runtime/MovieSceneTracks/Private/Tracks/MovieSceneMaterialTrack.cpp
--- a/Engine/Source/Runtime/MovieSceneTracks/Private/Tracks/MovieSceneMaterialTrack.cpp
+++ b/Engine/Source/Runtime/MovieSceneTracks/Private/Tracks/MovieSceneMaterialTrack.cpp
@@ -236,7 +236,7 @@ void UMovieSceneComponentMaterialTrack::ExtendEntityImpl(UMovieSceneParameterSec
 	// Material parameters are always absolute blends for the time being
 	OutImportedEntity->AddBuilder(
 		FEntityBuilder()
-		.Add(TracksComponents->ComponentMaterialIndex, MaterialIndex)
+		.Add(TracksComponents->ComponentMaterialInfo, MaterialInfo)
 		// If the section has no valid blend type (legacy data), make it use absolute blending.
 		// Otherwise, the base section class will add the appropriate blend type tag in BuildDefaultComponents.
 		.AddTagConditional(BuiltInComponents->Tags.AbsoluteBlend, !Section->GetBlendType().IsValid())
@@ -279,10 +279,24 @@ bool UMovieSceneComponentMaterialTrack::PopulateEvaluationFieldImpl(const TRange
 	return true;
 }
 
+
 #if WITH_EDITORONLY_DATA
+
 FText UMovieSceneComponentMaterialTrack::GetDefaultDisplayName() const
 {
-	return FText::FromString(FString::Printf(TEXT("Material Element %i"), MaterialIndex));
+	// Old track name before we started naming directly from editor
+	return FText::FromString(FString::Printf(TEXT("Material Element %i"), MaterialIndex_DEPRECATED));
+}
+
+void UMovieSceneComponentMaterialTrack::PostLoad()
+{
+	Super::PostLoad();
+	// Backwards compatibility with MaterialIndex alone as a way to reference materials.
+	if (MaterialInfo.MaterialType == EComponentMaterialType::Empty)
+	{
+		MaterialInfo.MaterialType = EComponentMaterialType::IndexedMaterial;
+		MaterialInfo.MaterialSlotIndex = MaterialIndex_DEPRECATED;
+	}
 }
 #endif
 
diff --git a/Engine/Source/Runtime/MovieSceneTracks/Private/Tracks/MovieScenePrimitiveMaterialTrack.cpp b/Engine/Source/Runtime/MovieSceneTracks/Private/Tracks/MovieScenePrimitiveMaterialTrack.cpp
--- a/Engine/Source/Runtime/MovieSceneTracks/Private/Tracks/MovieScenePrimitiveMaterialTrack.cpp
+++ b/Engine/Source/Runtime/MovieSceneTracks/Private/Tracks/MovieScenePrimitiveMaterialTrack.cpp
@@ -11,7 +11,6 @@
 UMovieScenePrimitiveMaterialTrack::UMovieScenePrimitiveMaterialTrack(const FObjectInitializer& ObjInit)
 	: Super(ObjInit)
 {
-	MaterialIndex = 0;
 	SupportedBlendTypes.Add(EMovieSceneBlendType::Absolute);
 	SupportedBlendTypes.Add(EMovieSceneBlendType::Additive);
 
@@ -30,12 +29,37 @@ bool UMovieScenePrimitiveMaterialTrack::SupportsType(TSubclassOf<UMovieSceneSect
 	return SectionClass == UMovieScenePrimitiveMaterialSection::StaticClass();
 }
 
+#if WITH_EDITORONLY_DATA
+void UMovieScenePrimitiveMaterialTrack::PostLoad()
+{
+	Super::PostLoad();
+	// Backwards compatibility with MaterialIndex alone as a way to reference materials.
+	if (MaterialInfo.MaterialType == EComponentMaterialType::Empty)
+	{
+		MaterialInfo.MaterialType = EComponentMaterialType::IndexedMaterial;
+		MaterialInfo.MaterialSlotIndex = MaterialIndex_DEPRECATED;
+	}
+}
+#endif
+
 int32 UMovieScenePrimitiveMaterialTrack::GetMaterialIndex() const
 {
-	return MaterialIndex;
+	return MaterialInfo.MaterialSlotIndex;
 }
 
-void UMovieScenePrimitiveMaterialTrack::SetMaterialIndex(int32 InMaterialIndex)
+void UMovieScenePrimitiveMaterialTrack::SetMaterialInfo(const FComponentMaterialInfo& InMaterialInfo)
 {
-	MaterialIndex = InMaterialIndex;
+	MaterialInfo = InMaterialInfo;
 }
+
+const FComponentMaterialInfo& UMovieScenePrimitiveMaterialTrack::GetMaterialInfo() const
+{
+	return MaterialInfo;
+}
+
+void UMovieScenePrimitiveMaterialTrack::SetMaterialIndex(int32 InMaterialIndex)
+{
+	MaterialInfo.MaterialSlotIndex = InMaterialIndex;
+	// Assumption is if this is being called by old code, we should be indexed
+	MaterialInfo.MaterialType = EComponentMaterialType::IndexedMaterial;
+}
\ No newline at end of file
diff --git a/Engine/Source/Runtime/MovieSceneTracks/Public/MovieSceneTracksComponentTypes.h b/Engine/Source/Runtime/MovieSceneTracks/Public/MovieSceneTracksComponentTypes.h
--- a/Engine/Source/Runtime/MovieSceneTracks/Public/MovieSceneTracksComponentTypes.h
+++ b/Engine/Source/Runtime/MovieSceneTracks/Public/MovieSceneTracksComponentTypes.h
@@ -15,6 +15,7 @@
 #include "Styling/SlateColor.h"
 #include "ConstraintChannel.h"
 #include "MaterialTypes.h"
+#include "Tracks/MovieSceneMaterialTrack.h"
 #include "MovieSceneTracksComponentTypes.generated.h"
 
 class UMaterialParameterCollection;
@@ -557,7 +558,7 @@ struct FMovieSceneTracksComponentTypes
 
 	TComponentTypeID<FMovieSceneSkeletalAnimationComponentData> SkeletalAnimation;
 
-	TComponentTypeID<int32> ComponentMaterialIndex;
+	TComponentTypeID<FComponentMaterialInfo> ComponentMaterialInfo;
 
 	TComponentTypeID<FName> BoolParameterName;
 	TComponentTypeID<FName> ScalarParameterName;
diff --git a/Engine/Source/Runtime/MovieSceneTracks/Public/Systems/MovieSceneComponentMaterialSystem.h b/Engine/Source/Runtime/MovieSceneTracks/Public/Systems/MovieSceneComponentMaterialSystem.h
--- a/Engine/Source/Runtime/MovieSceneTracks/Public/Systems/MovieSceneComponentMaterialSystem.h
+++ b/Engine/Source/Runtime/MovieSceneTracks/Public/Systems/MovieSceneComponentMaterialSystem.h
@@ -7,6 +7,7 @@
 #include "Evaluation/PreAnimatedState/MovieScenePreAnimatedStateStorage.h"
 
 #include "Systems/MovieSceneMaterialSystem.h"
+#include "Tracks/MovieSceneMaterialTrack.h"
 
 #include "MovieSceneComponentMaterialSystem.generated.h"
 
@@ -18,15 +19,15 @@ namespace UE::MovieScene
 struct FComponentMaterialKey
 {
 	FObjectKey Object;
-	int32 MaterialIndex;
+	FComponentMaterialInfo MaterialInfo;
 
 	friend uint32 GetTypeHash(const FComponentMaterialKey& In)
 	{
-		return GetTypeHash(In.Object) ^ ::GetTypeHash(In.MaterialIndex);
+		return HashCombine(GetTypeHash(In.Object), GetTypeHash(In.MaterialInfo));
 	}
 	friend bool operator==(const FComponentMaterialKey& A, const FComponentMaterialKey& B)
 	{
-		return A.Object == B.Object && A.MaterialIndex == B.MaterialIndex;
+		return A.Object == B.Object && A.MaterialInfo == B.MaterialInfo;
 	}
 };
 
@@ -35,10 +36,10 @@ struct FComponentMaterialAccessor
 	using KeyType = FComponentMaterialKey;
 
 	UObject* Object;
-	int32 MaterialIndex;
+	FComponentMaterialInfo MaterialInfo;
 
 	FComponentMaterialAccessor(const FComponentMaterialKey& InKey);
-	FComponentMaterialAccessor(UObject* InObject, int32 InMaterialIndex);
+	FComponentMaterialAccessor(UObject* InObject, const FComponentMaterialInfo& InMaterialInfo);
 
 	explicit operator bool() const;
 
@@ -48,17 +49,17 @@ struct FComponentMaterialAccessor
 	FString ToString() const;
 };
 
-using FPreAnimatedComponentMaterialTraits          = TPreAnimatedMaterialTraits<FComponentMaterialAccessor, UObject*, int32>;
-using FPreAnimatedComponentMaterialParameterTraits = TPreAnimatedMaterialParameterTraits<FComponentMaterialAccessor, UObject*, int32>;
+using FPreAnimatedComponentMaterialTraits          = TPreAnimatedMaterialTraits<FComponentMaterialAccessor, UObject*, FComponentMaterialInfo>;
+using FPreAnimatedComponentMaterialParameterTraits = TPreAnimatedMaterialParameterTraits<FComponentMaterialAccessor, UObject*, FComponentMaterialInfo>;
 
 struct FPreAnimatedComponentMaterialSwitcherStorage
-	: public TPreAnimatedStateStorage<TPreAnimatedMaterialTraits<FComponentMaterialAccessor, UObject*, int32>>
+	: public TPreAnimatedStateStorage<TPreAnimatedMaterialTraits<FComponentMaterialAccessor, UObject*, FComponentMaterialInfo>>
 {
 	static TAutoRegisterPreAnimatedStorageID<FPreAnimatedComponentMaterialSwitcherStorage> StorageID;
 };
 
 struct FPreAnimatedComponentMaterialParameterStorage
-	: public TPreAnimatedStateStorage<TPreAnimatedMaterialParameterTraits<FComponentMaterialAccessor, UObject*, int32>>
+	: public TPreAnimatedStateStorage<TPreAnimatedMaterialParameterTraits<FComponentMaterialAccessor, UObject*, FComponentMaterialInfo>>
 {
 	static TAutoRegisterPreAnimatedStorageID<FPreAnimatedComponentMaterialParameterStorage> StorageID;
 };
@@ -87,5 +88,5 @@ private:
 
 private:
 
-	UE::MovieScene::TMovieSceneMaterialSystem<UE::MovieScene::FComponentMaterialAccessor, UObject*, int32> SystemImpl;
+	UE::MovieScene::TMovieSceneMaterialSystem<UE::MovieScene::FComponentMaterialAccessor, UObject*, FComponentMaterialInfo> SystemImpl;
 };
diff --git a/Engine/Source/Runtime/MovieSceneTracks/Public/Tracks/MovieSceneMaterialTrack.h b/Engine/Source/Runtime/MovieSceneTracks/Public/Tracks/MovieSceneMaterialTrack.h
--- a/Engine/Source/Runtime/MovieSceneTracks/Public/Tracks/MovieSceneMaterialTrack.h
+++ b/Engine/Source/Runtime/MovieSceneTracks/Public/Tracks/MovieSceneMaterialTrack.h
@@ -10,6 +10,70 @@
 #include "MaterialTypes.h"
 #include "MovieSceneMaterialTrack.generated.h"
 
+UENUM()
+enum class EComponentMaterialType
+{
+	/* Empty/Uninitialized*/
+	Empty,
+	/* A material in one of the indexed slots on a primitive component*/
+	IndexedMaterial,
+	/* An overlay material on a mesh component*/
+	OverlayMaterial,
+	/* A decal material*/
+	DecalMaterial
+};
+
+/**
+ * Contains what is necessary to uniquely identify a material on a component, whether that be an indexed material, one with a slot name, or an overlay material.
+ */
+USTRUCT(BlueprintType)
+struct FComponentMaterialInfo
+{
+	GENERATED_BODY()
+
+	UPROPERTY()
+	FName MaterialSlotName;
+
+	UPROPERTY()
+	int MaterialSlotIndex = 0;
+
+	UPROPERTY()
+	EComponentMaterialType MaterialType = EComponentMaterialType::Empty;
+
+	friend uint32 GetTypeHash(const FComponentMaterialInfo& In)
+	{
+		return GetTypeHash(In.MaterialSlotName) ^ ::GetTypeHash(In.MaterialSlotIndex) ^ ::GetTypeHash(In.MaterialType);
+	}
+	friend bool operator==(const FComponentMaterialInfo& A, const FComponentMaterialInfo& B)
+	{
+		return A.MaterialSlotName == B.MaterialSlotName && A.MaterialSlotIndex == B.MaterialSlotIndex && A.MaterialSlotName == B.MaterialSlotName && A.MaterialType == B.MaterialType;
+	}
+
+	FString ToString() const 
+	{
+		switch (MaterialType)
+		{
+		case EComponentMaterialType::IndexedMaterial:
+			if (!MaterialSlotName.IsNone())
+			{
+				return FString::Printf(TEXT("Material Slot: %s"), *MaterialSlotName.ToString());
+			}
+			else
+			{
+				return FString::Printf(TEXT("Material Element %i"), MaterialSlotIndex);
+			}
+			break;
+		case EComponentMaterialType::OverlayMaterial:
+			return TEXT("Overlay Material");
+			break;
+		case EComponentMaterialType::DecalMaterial:
+			return TEXT("Decal Material");
+			break;
+		}
+		return FString();
+	}
+};
+
 /**
  * Handles manipulation of material parameters in a movie scene.
  */
@@ -148,26 +212,44 @@ public:
 	/*~ IMovieSceneParameterSectionExtender */
 	virtual void ExtendEntityImpl(UMovieSceneParameterSection* Section, UMovieSceneEntitySystemLinker* EntityLinker, const UE::MovieScene::FEntityImportParams& Params, UE::MovieScene::FImportedEntity* OutImportedEntity) override;
 
-	virtual FName GetTrackName() const { return FName( *FString::FromInt(MaterialIndex) ); }
 
 #if WITH_EDITORONLY_DATA
 	virtual FText GetDefaultDisplayName() const override;
 #endif
-
 public:
 
 	/** Gets the index of the material in the component. */
-	int32 GetMaterialIndex() const { return MaterialIndex; }
+	const FComponentMaterialInfo& GetMaterialInfo() const { return MaterialInfo; }
 
 	/** Sets the index of the material in the component. */
-	void SetMaterialIndex(int32 InMaterialIndex) 
+	void SetMaterialInfo(const FComponentMaterialInfo& InMaterialInfo)
 	{
-		MaterialIndex = InMaterialIndex;
+		MaterialInfo = InMaterialInfo;
 	}
 
+#if WITH_EDITOR
+	virtual FText GetDisplayNameToolTipText() const override { return TooltipText; }
+	void SetDisplayNameTooltipText(FText InTooltipText) { TooltipText = InTooltipText; }
+#endif
+
+
+#if WITH_EDITORONLY_DATA
+
+protected:
+	void PostLoad() override;
+
 private:
+	/** The index of this material this track is animating. Has been deprecated in favor of MaterialInfo*/
+	UPROPERTY(meta = (DeprecatedProperty, DeprecationMessage = "Use MaterialInfo instead."))
+	int32 MaterialIndex_DEPRECATED;
+
+	UPROPERTY()
+	FText TooltipText;
 
-	/** The index of this material this track is animating. */
+#endif
+private:
+	/** The info on the material this track is animating.*/
 	UPROPERTY()
-	int32 MaterialIndex;
+	FComponentMaterialInfo MaterialInfo;
+
 };
diff --git a/Engine/Source/Runtime/MovieSceneTracks/Public/Tracks/MovieScenePrimitiveMaterialTrack.h b/Engine/Source/Runtime/MovieSceneTracks/Public/Tracks/MovieScenePrimitiveMaterialTrack.h
--- a/Engine/Source/Runtime/MovieSceneTracks/Public/Tracks/MovieScenePrimitiveMaterialTrack.h
+++ b/Engine/Source/Runtime/MovieSceneTracks/Public/Tracks/MovieScenePrimitiveMaterialTrack.h
@@ -4,6 +4,7 @@
 
 #include "MovieSceneTrack.h"
 #include "Tracks/MovieScenePropertyTrack.h"
+#include "MovieSceneMaterialTrack.h"
 #include "MovieScenePrimitiveMaterialTrack.generated.h"
 
 UCLASS(MinimalAPI)
@@ -16,16 +17,42 @@ public:
 	UMovieScenePrimitiveMaterialTrack(const FObjectInitializer& ObjInit);
 
 	/* Set the material index that this track is assigned to */
+	UE_DEPRECATED(5.4, "SetMaterialIndex is deprecated, use SetMaterialInfo instead")
 	MOVIESCENETRACKS_API void SetMaterialIndex(int32 MaterialIndex);
 
 	/* Get the material index that this track is assigned to */
+	UE_DEPRECATED(5.4, "GetMaterialIndex is deprecated, use GetMaterialInfo instead")
 	MOVIESCENETRACKS_API int32 GetMaterialIndex() const;
 
+	/* Set the material info for the material that this track is assigned to */
+	MOVIESCENETRACKS_API void SetMaterialInfo(const FComponentMaterialInfo& InMaterialInfo);
+
+	/* Get the material info for the material that this track is assigned to */
+	MOVIESCENETRACKS_API const FComponentMaterialInfo& GetMaterialInfo() const;
+
 	/*~ UMovieSceneTrack interface */
 	virtual UMovieSceneSection* CreateNewSection() override;
 	virtual bool SupportsType(TSubclassOf<UMovieSceneSection> SectionClass) const override;
 
+#if WITH_EDITOR
+	virtual FText GetDisplayNameToolTipText() const override { return TooltipText; }
+	void SetDisplayNameTooltipText(FText InTooltipText) { TooltipText = InTooltipText; }
+#endif
+
+protected:
+#if WITH_EDITORONLY_DATA
+	void PostLoad() override;
+
+private:
+	UPROPERTY()
+	int32 MaterialIndex_DEPRECATED;
+		
+	UPROPERTY()
+	FText TooltipText;
+
+#endif
+
 private:
 	UPROPERTY()
-	int32 MaterialIndex;
+	FComponentMaterialInfo MaterialInfo;
 };
