Sequencer: Unify ProcessKeyOperation with GetTransformKeys so that it uses RecomposeTransform

diff --git a/Engine/Source/Editor/MovieSceneTools/Private/TrackEditors/TransformTrackEditor.cpp b/Engine/Source/Editor/MovieSceneTools/Private/TrackEditors/TransformTrackEditor.cpp
--- a/Engine/Source/Editor/MovieSceneTools/Private/TrackEditors/TransformTrackEditor.cpp
+++ b/Engine/Source/Editor/MovieSceneTools/Private/TrackEditors/TransformTrackEditor.cpp
@@ -1086,115 +1086,57 @@ int32 GetPreviousKey(FMovieSceneDoubleChannel& Channel, FFrameNumber Time)
 
 void F3DTransformTrackEditor::ProcessKeyOperation(UObject* ObjectToKey, TArrayView<const UE::Sequencer::FKeySectionOperation> SectionsToKey, ISequencer& InSequencer, FFrameNumber KeyTime)
 {
-	USceneComponent* Component = MovieSceneHelpers::SceneComponentFromRuntimeObject(ObjectToKey);
-	if (!Component)
-	{
-		return;
-	}
-
 	using namespace UE::MovieScene;
 	using namespace UE::Sequencer;
 
-	FSystemInterrogator Interrogator;
-	Interrogator.TrackImportedEntities(true);
-
-	TGuardValue<FEntityManager*> DebugVizGuard(GEntityManagerForDebuggingVisualizers, &Interrogator.GetLinker()->EntityManager);
-
-	TArray<FInterrogationChannel> InterrogationChannelsPerOperations;
-	for (const FKeySectionOperation& Operation : SectionsToKey)
+	USceneComponent* Component = MovieSceneHelpers::SceneComponentFromRuntimeObject(ObjectToKey);
+	if (!Component)
 	{
-		if (UMovieScenePropertyTrack* Track = Operation.Section->GetSectionObject()->GetTypedOuter<UMovieScenePropertyTrack>())
-		{
-			const FMovieScenePropertyBinding PropertyBinding = Track->GetPropertyBinding();
-			const FInterrogationChannel InterrogationChannel = Interrogator.AllocateChannel(Component, PropertyBinding);
-			InterrogationChannelsPerOperations.Add(InterrogationChannel);
-			Interrogator.ImportTrack(Track, InterrogationChannel);
-		}
-		else
-		{
-			InterrogationChannelsPerOperations.Add(FInterrogationChannel::Invalid());
-		}
+		return;
 	}
 
-	Interrogator.AddInterrogation(KeyTime);
-
-	Interrogator.Update();
+	FTransform CurrentTransform(Component->GetRelativeRotation(), Component->GetRelativeLocation(), Component->GetRelativeScale3D());
 
-	TArray<FMovieSceneEntityID> EntitiesPerSection, ValidEntities;
 	for (int32 Index = 0; Index < SectionsToKey.Num(); ++Index)
 	{
-		const FKeySectionOperation& Operation = SectionsToKey[Index];
-		const FInterrogationChannel InterrogationChannel = InterrogationChannelsPerOperations[Index];
-		const FInterrogationKey InterrogationKey(InterrogationChannel, 0);
-		FMovieSceneEntityID EntityID = Interrogator.FindEntityFromOwner(InterrogationKey, Operation.Section->GetSectionObject(), 0);
+		FTransformData RecomposedTransform = RecomposeTransform(CurrentTransform, ObjectToKey, SectionsToKey[Index].Section->GetSectionObject());
 
-		EntitiesPerSection.Add(EntityID);
-		if (EntityID)
+		for (TSharedPtr<IKeyArea> KeyArea : SectionsToKey[Index].KeyAreas)
 		{
-			ValidEntities.Add(EntityID);
-		}
-	}
-
-	UMovieSceneInterrogatedPropertyInstantiatorSystem* System = Interrogator.GetLinker()->FindSystem<UMovieSceneInterrogatedPropertyInstantiatorSystem>();
-
-	if (ensure(System))
-	{
-		FDecompositionQuery Query;
-		Query.Entities = ValidEntities;
-		Query.bConvertFromSourceEntityIDs = false;
-		Query.Object   = Component;
-
-		FTransform CurrentTransform(Component->GetRelativeRotation(), Component->GetRelativeLocation(), Component->GetRelativeScale3D());
-
-		// Account for the transform origin only if this is not parented because the transform origin is already being applied to the parent.
-		if (!Component->GetAttachParent())
-		{
-			CurrentTransform *= GetTransformOrigin().Inverse();
-		}
-
-		UpdateTransformBasedOnConstraint(CurrentTransform, Component);
-
-		FIntermediate3DTransform CurrentValue(CurrentTransform.GetTranslation(), CurrentTransform.GetRotation().Rotator(), CurrentTransform.GetScale3D());
-		TRecompositionResult<FIntermediate3DTransform> TransformData = System->RecomposeBlendOperational(FMovieSceneTracksComponentTypes::Get()->ComponentTransform, Query, CurrentValue);
-
-		for (int32 Index = 0; Index < SectionsToKey.Num(); ++Index)
-		{
-			FMovieSceneEntityID EntityID = EntitiesPerSection[Index];
-			if (!EntityID)
-			{
-				continue;
-			}
-
-			const FIntermediate3DTransform& RecomposedTransform = TransformData.Values[Index];
-
-			for (TSharedPtr<IKeyArea> KeyArea : SectionsToKey[Index].KeyAreas)
+			FMovieSceneChannelHandle Handle  = KeyArea->GetChannel();
+			if (Handle.GetChannelTypeName() == FMovieSceneDoubleChannel::StaticStruct()->GetFName())
 			{
-				FMovieSceneChannelHandle Handle  = KeyArea->GetChannel();
-				if (Handle.GetChannelTypeName() == FMovieSceneDoubleChannel::StaticStruct()->GetFName())
+				FMovieSceneDoubleChannel* Channel = static_cast<FMovieSceneDoubleChannel*>(Handle.Get());
+
+				double Value =
+					Handle.GetChannelIndex() == 0 ? RecomposedTransform.Translation[0] :
+					Handle.GetChannelIndex() == 1 ? RecomposedTransform.Translation[1] :
+					Handle.GetChannelIndex() == 2 ? RecomposedTransform.Translation[2] :
+					Handle.GetChannelIndex() == 3 ? RecomposedTransform.Rotation.Roll :
+					Handle.GetChannelIndex() == 4 ? RecomposedTransform.Rotation.Pitch :
+					Handle.GetChannelIndex() == 5 ? RecomposedTransform.Rotation.Yaw :
+					Handle.GetChannelIndex() == 6 ? RecomposedTransform.Scale[0] :
+					Handle.GetChannelIndex() == 7 ? RecomposedTransform.Scale[1] :
+					Handle.GetChannelIndex() == 8 ? RecomposedTransform.Scale[2] : 0.f;
+
+				if (KeyArea->GetName() == "Rotation.X" ||
+					KeyArea->GetName() == "Rotation.Y" ||
+					KeyArea->GetName() == "Rotation.Z")
 				{
-					FMovieSceneDoubleChannel* Channel = static_cast<FMovieSceneDoubleChannel*>(Handle.Get());
-
-					double Value = RecomposedTransform[Handle.GetChannelIndex()];
-
-					if (KeyArea->GetName() == "Rotation.X" ||
-						KeyArea->GetName() == "Rotation.Y" ||
-						KeyArea->GetName() == "Rotation.Z")
+					int32 PreviousKey = GetPreviousKey(*Channel, KeyTime);
+					if (PreviousKey != INDEX_NONE && PreviousKey < Channel->GetData().GetValues().Num())
 					{
-						int32 PreviousKey = GetPreviousKey(*Channel, KeyTime);
-						if (PreviousKey != INDEX_NONE && PreviousKey < Channel->GetData().GetValues().Num())
-						{
-							double OldValue = Channel->GetData().GetValues()[PreviousKey].Value;
-							Value = UnwindChannel(OldValue, Value);
-						}
+						double OldValue = Channel->GetData().GetValues()[PreviousKey].Value;
+						Value = UnwindChannel(OldValue, Value);
 					}
-
-					EMovieSceneKeyInterpolation Interpolation = GetInterpolationMode(Channel, KeyTime, InSequencer.GetKeyInterpolation());
-					AddKeyToChannel(Channel, KeyTime, Value, Interpolation);
-				}
-				else
-				{
-					KeyArea->AddOrUpdateKey(KeyTime, FGuid(), InSequencer);
 				}
+
+				EMovieSceneKeyInterpolation Interpolation = GetInterpolationMode(Channel, KeyTime, InSequencer.GetKeyInterpolation());
+				AddKeyToChannel(Channel, KeyTime, Value, Interpolation);
+			}
+			else
+			{
+				KeyArea->AddOrUpdateKey(KeyTime, FGuid(), InSequencer);
 			}
 		}
 	}
