Sequencer- Adding a track for animating custom primitive data on PrimitiveComponents.
Custom primitive data is an alternative way of driving material parameters. It allows you to map any of 36 numbered float parameters on a PrimitiveComponent to one or more material parameters across multiple meshes and materials. You can apportion between 1-4 slots in a row to handle anything from scalar to color parameters. Previously these could only be set via code or blueprint, but now these can be animated in Sequencer.

A custom primitive data track can be added to a Primitive Component binding track in Sequencer.
From there, you can choose to add a float, Vector2D, Vector, or Color parameter, choosing the custom primitive data start index the parameter starts from. The UI will also scan the primitive component for material parameters currently mapped to any index to give you helpful suggestions for which parameters you may want to add. These parameters can then be animated and blended just like any parameter of that type.

As this was a new track, we had to build a new track editor, track, and section. The section inherits from MovieSceneParameterSection, as this was already capable of animating parameters of these types. We use the unique start index of the parameter as the name for these parameter sections.

To create the dynamic tooltips that show which material parameters the custom primitive data entries are mapped to, a change was made to the tooltiptext channel metadata to take a lambda instead of just an FText. This allows the tooltips to be generated live, so any changes to materials will be taken into account in the tooltip.

A new screenshot autotest using the feature was also created.

diff --git a/Engine/Source/Editor/MovieSceneTools/Private/MovieSceneToolsModule.cpp b/Engine/Source/Editor/MovieSceneTools/Private/MovieSceneToolsModule.cpp
--- a/Engine/Source/Editor/MovieSceneTools/Private/MovieSceneToolsModule.cpp
+++ b/Engine/Source/Editor/MovieSceneTools/Private/MovieSceneToolsModule.cpp
@@ -53,6 +53,7 @@
 #include "TrackEditors/PrimitiveMaterialTrackEditor.h"
 #include "TrackEditors/CameraShakeSourceShakeTrackEditor.h"
 #include "TrackEditors/CVarTrackEditor.h"
+#include "TrackEditors/CustomPrimitiveDataTrackEditor.h"
 
 #include "Channels/PerlinNoiseChannelInterface.h"
 
@@ -179,6 +180,7 @@ void FMovieSceneToolsModule::StartupModule()
 		PrimitiveMaterialCreateEditorHandle = SequencerModule.RegisterTrackEditor(FOnCreateTrackEditor::CreateStatic(&FPrimitiveMaterialTrackEditor::CreateTrackEditor));
 		CameraShakeSourceShakeCreateEditorHandle = SequencerModule.RegisterTrackEditor(FOnCreateTrackEditor::CreateStatic(&FCameraShakeSourceShakeTrackEditor::CreateTrackEditor));
 		CVarTrackCreateEditorHandle = SequencerModule.RegisterTrackEditor(FOnCreateTrackEditor::CreateStatic(&FCVarTrackEditor::CreateTrackEditor));
+		CustomPrimitiveDataTrackCreateEditorHandle = SequencerModule.RegisterTrackEditor(FOnCreateTrackEditor::CreateStatic(&FCustomPrimitiveDataTrackEditor::CreateTrackEditor));
 
 		// register track models
 		CameraCutTrackModelHandle = SequencerModule.RegisterTrackModel(FOnCreateTrackModel::CreateStatic(&FCameraCutTrackModel::CreateTrackModel));
diff --git a/Engine/Source/Editor/MovieSceneTools/Private/Sections/ParameterSection.cpp b/Engine/Source/Editor/MovieSceneTools/Private/Sections/ParameterSection.cpp
--- a/Engine/Source/Editor/MovieSceneTools/Private/Sections/ParameterSection.cpp
+++ b/Engine/Source/Editor/MovieSceneTools/Private/Sections/ParameterSection.cpp
@@ -78,23 +78,26 @@ namespace UE::Sequencer
 			{
 				FMovieSceneChannelHandle ChannelHandle = ChildChannel->GetKeyArea()->GetChannel();
 
-				FText DisplayText = ChannelHandle.GetMetaData()->DisplayText;
-
-				if (ChannelHandle.GetMetaData()->DisplayText.EqualTo(FCommonChannelData::ChannelR))
-				{
-					ColorChannels[0] = ChannelHandle.Cast<FMovieSceneFloatChannel>().Get();
-				}
-				if (ChannelHandle.GetMetaData()->DisplayText.EqualTo(FCommonChannelData::ChannelG))
-				{
-					ColorChannels[1] = ChannelHandle.Cast<FMovieSceneFloatChannel>().Get();
-				}
-				if (ChannelHandle.GetMetaData()->DisplayText.EqualTo(FCommonChannelData::ChannelB))
-				{
-					ColorChannels[2] = ChannelHandle.Cast<FMovieSceneFloatChannel>().Get();
-				}
-				if (ChannelHandle.GetMetaData()->DisplayText.EqualTo(FCommonChannelData::ChannelA))
+				if (const FMovieSceneChannelMetaData* Metadata = ChannelHandle.GetMetaData())
 				{
-					ColorChannels[3] = ChannelHandle.Cast<FMovieSceneFloatChannel>().Get();
+					FText DisplayText = Metadata->DisplayText;
+
+					if (DisplayText.EqualTo(FCommonChannelData::ChannelR))
+					{
+						ColorChannels[0] = ChannelHandle.Cast<FMovieSceneFloatChannel>().Get();
+					}
+					else if (DisplayText.EqualTo(FCommonChannelData::ChannelG))
+					{
+						ColorChannels[1] = ChannelHandle.Cast<FMovieSceneFloatChannel>().Get();
+					}
+					else if (DisplayText.EqualTo(FCommonChannelData::ChannelB))
+					{
+						ColorChannels[2] = ChannelHandle.Cast<FMovieSceneFloatChannel>().Get();
+					}
+					else if (DisplayText.EqualTo(FCommonChannelData::ChannelA))
+					{
+						ColorChannels[3] = ChannelHandle.Cast<FMovieSceneFloatChannel>().Get();
+					}
 				}
 			}
 
diff --git a/Engine/Source/Editor/MovieSceneTools/Private/TrackEditors/CustomPrimitiveDataTrackEditor.cpp b/Engine/Source/Editor/MovieSceneTools/Private/TrackEditors/CustomPrimitiveDataTrackEditor.cpp
new file mode 100644
--- /dev/null
+++ b/Engine/Source/Editor/MovieSceneTools/Private/TrackEditors/CustomPrimitiveDataTrackEditor.cpp
@@ -0,0 +1,532 @@
+// Copyright Epic Games, Inc. All Rights Reserved.
+
+#include "TrackEditors/CustomPrimitiveDataTrackEditor.h"
+#include "Framework/MultiBox/MultiBoxBuilder.h"
+#include "Components/DecalComponent.h"
+#include "Components/PrimitiveComponent.h"
+#include "Materials/Material.h"
+#include "Materials/MaterialInstance.h"
+#include "Sections/MovieSceneCustomPrimitiveDataSection.h"
+#include "Tracks/MovieSceneCustomPrimitiveDataTrack.h"
+#include "Sections/ParameterSection.h"
+#include "SequencerUtilities.h"
+#include "Modules/ModuleManager.h"
+#include "MaterialEditorModule.h"
+#include "Engine/Selection.h"
+#include "ISequencerModule.h"
+#include "Components/MeshComponent.h"
+#include "MaterialTypes.h"
+#include "Widgets/Input/SSpinBox.h"
+#include "Widgets/Layout/SSpacer.h"
+
+
+#define LOCTEXT_NAMESPACE "CustomPrimitiveDataTrackEditor"
+
+
+FCustomPrimitiveDataTrackEditor::FCustomPrimitiveDataTrackEditor(TSharedRef<ISequencer> InSequencer)
+	: FMovieSceneTrackEditor(InSequencer)
+{
+}
+
+TSharedRef<ISequencerSection> FCustomPrimitiveDataTrackEditor::MakeSectionInterface(UMovieSceneSection& SectionObject, UMovieSceneTrack& Track, FGuid ObjectBinding)
+{
+	UMovieSceneCustomPrimitiveDataSection* CustomPrimitiveDataSection = Cast<UMovieSceneCustomPrimitiveDataSection>(&SectionObject);
+	checkf(CustomPrimitiveDataSection != nullptr, TEXT("Unsupported section type."));
+
+	return MakeShareable(new FParameterSection(*CustomPrimitiveDataSection));
+}
+
+TSharedPtr<SWidget> FCustomPrimitiveDataTrackEditor::BuildOutlinerEditWidget(const FGuid& ObjectBinding, UMovieSceneTrack* Track, const FBuildEditWidgetParams& Params)
+{
+	UMovieSceneCustomPrimitiveDataTrack* CPDTrack = Cast<UMovieSceneCustomPrimitiveDataTrack>(Track);
+	FOnGetContent MenuContent = FOnGetContent::CreateSP(this, &FCustomPrimitiveDataTrackEditor::OnGetAddMenuContent, ObjectBinding, CPDTrack, Params.TrackInsertRowIndex);
+
+	return FSequencerUtilities::MakeAddButton(LOCTEXT("AddParameterButton", "Parameter"), MenuContent, Params.NodeIsHovered, GetSequencer());
+}
+
+
+TSharedRef<SWidget> FCustomPrimitiveDataTrackEditor::OnGetAddMenuContent(FGuid ObjectBinding, UMovieSceneCustomPrimitiveDataTrack* CPDTrack, int32 TrackInsertRowIndex)
+{
+	// IF this is supported, allow creating other sections with different blend types, and put
+	// the material parameters after a separator. Otherwise, just show the parameters menu.
+	const FMovieSceneBlendTypeField SupportedBlendTypes = CPDTrack->GetSupportedBlendTypes();
+
+	if (SupportedBlendTypes.Num() > 1)
+	{
+		FMenuBuilder MenuBuilder(true, nullptr);
+
+		TWeakPtr<ISequencer> WeakSequencer = GetSequencer();
+		FSequencerUtilities::PopulateMenu_CreateNewSection(MenuBuilder, TrackInsertRowIndex, CPDTrack, WeakSequencer);
+
+		MenuBuilder.AddSeparator();
+
+		OnBuildAddParameterMenu(MenuBuilder, ObjectBinding, CPDTrack);
+
+		return MenuBuilder.MakeWidget();
+	}
+	else
+	{
+		return OnGetAddParameterMenuContent(ObjectBinding, CPDTrack);
+	}
+}
+
+TSharedRef<SWidget> FCustomPrimitiveDataTrackEditor::OnGetAddParameterMenuContent(FGuid ObjectBinding, UMovieSceneCustomPrimitiveDataTrack* CPDTrack)
+{
+	FMenuBuilder AddParameterMenuBuilder(true, nullptr);
+
+	OnBuildAddParameterMenu(AddParameterMenuBuilder, ObjectBinding, CPDTrack);
+	return AddParameterMenuBuilder.MakeWidget();
+}
+
+void FCustomPrimitiveDataTrackEditor::OnBuildAddParameterMenu(FMenuBuilder& MenuBuilder, FGuid ObjectBinding, UMovieSceneCustomPrimitiveDataTrack* CPDTrack)
+{
+	check(CPDTrack);
+
+	UObject* Object = GetSequencer()->FindSpawnedObjectOrTemplate(ObjectBinding);
+	if (!Object)
+	{
+		return;
+	}
+
+	UPrimitiveComponent* Component = Cast<UPrimitiveComponent>(Object);
+	if (!Component)
+	{
+		return;
+	}
+
+	auto OnStartIndexChanged = [this](int32 NewValue) 
+	{ 
+		StartIndex = NewValue; 
+	};
+	FGetStartIndexDelegate GetStartIndex = [this]() { return StartIndex; };
+
+	TSortedMap<uint8, TArray<FCustomPrimitiveDataMaterialParametersData>> CPDData;
+	CPDTrack->GetCPDMaterialData(*GetSequencer(), ObjectBinding, GetSequencer()->GetFocusedTemplateID(), CPDData);
+	if (!CPDData.IsEmpty())
+	{
+		auto GetMaterialParameterLayerName = [](const FMaterialLayersFunctions& Layers, const FMaterialParameterInfo& InParameterInfo)
+		{
+			FString LayerName;
+			if (Layers.EditorOnly.LayerNames.IsValidIndex(InParameterInfo.Index))
+			{
+				LayerName = Layers.GetLayerName(InParameterInfo.Index).ToString();
+			}
+			return LayerName;
+		};
+		auto GetMaterialParameterAssetName = [](const FMaterialLayersFunctions& Layers, const FMaterialParameterInfo& InParameterInfo)
+		{
+			FString AssetName;
+			if (InParameterInfo.Association == EMaterialParameterAssociation::LayerParameter && Layers.Layers.IsValidIndex(InParameterInfo.Index))
+			{
+				AssetName = Layers.Layers[InParameterInfo.Index]->GetName();
+			}
+			else if (InParameterInfo.Association == EMaterialParameterAssociation::BlendParameter && Layers.Blends.IsValidIndex(InParameterInfo.Index))
+			{
+				AssetName = Layers.Blends[InParameterInfo.Index]->GetName();
+			}
+			return AssetName;
+		};
+
+		auto GetMaterialParameterDisplayName = [](const FMaterialParameterInfo& InParameterInfo, const FString& InMaterialName, const FString& InLayerName, const FString& InAssetName, uint8 CPDIndex)
+		{
+			FText DisplayName = FText::Format(LOCTEXT("MaterialAndParameterDisplayName", "{0} (Index: {1}, Asset: {2})"), FText::FromName(InParameterInfo.Name), FText::AsNumber(CPDIndex), FText::FromString(InMaterialName));
+			if (!InLayerName.IsEmpty() && !InAssetName.IsEmpty())
+			{
+				DisplayName = FText::Format(LOCTEXT("MaterialParameterDisplayName", "{0} (Index: {1}, Asset: {2}.{3}.{4})"), 
+					FText::FromName(InParameterInfo.Name),
+					FText::AsNumber(CPDIndex),
+					FText::FromString(InMaterialName),
+					FText::FromString(InLayerName), 
+					FText::FromString(InAssetName));
+			}
+			return DisplayName;
+		};
+
+		MenuBuilder.BeginSection("UsedInMaterials", LOCTEXT("CustomPrimitiveDataFromMaterialParameters", "Used In Material Parameters"));
+		{
+			for (const auto& CPDDataMapEntry : CPDData)
+			{
+				uint8 CPDIndex = CPDDataMapEntry.Key;
+
+				FGetStartIndexDelegate GetCPDIndex = [CPDIndex] { return CPDIndex; };
+				for (const FCustomPrimitiveDataMaterialParametersData& Data : CPDDataMapEntry.Value)
+				{
+					FString MaterialName = Data.MaterialAsset.IsValid() ? Data.MaterialAsset->GetName() : FString();
+					FMaterialLayersFunctions Layers;
+					if (Data.MaterialAsset.IsValid())
+					{
+						Data.MaterialAsset->GetMaterialLayers(Layers);
+					}
+					FString LayerName = GetMaterialParameterLayerName(Layers, Data.ParameterInfo);
+					FString AssetName = GetMaterialParameterAssetName(Layers, Data.ParameterInfo);
+					FText DisplayName = GetMaterialParameterDisplayName(Data.ParameterInfo, MaterialName, LayerName, AssetName, CPDIndex);
+
+					if (Data.MaterialParameterType == EMaterialParameterType::Vector)
+					{
+						MenuBuilder.AddMenuEntry(
+							DisplayName,
+							TAttribute<FText>::CreateLambda([this, CPDTrack, GetCPDIndex]()
+								{
+									if (CanAddParameter(CPDTrack, GetCPDIndex, 4))
+									{
+										return FText::Format(LOCTEXT("CustomPrimitiveDataMaterialTooltipVector", "Add a track for this vector/color parameter starting at index {0}"), FText::AsNumber(GetCPDIndex()));
+									}
+									else
+									{
+										return FText::Format(LOCTEXT("CustomPrimitiveDataMaterialTooltipVectorInvalid", "Can't add a vector/color parameter starting at index {0} as it overlaps with other parameters using that index."), FText::AsNumber(GetCPDIndex()));
+									}
+								}),
+							FSlateIcon(),
+							FUIAction(
+								FExecuteAction::CreateRaw(this, &FCustomPrimitiveDataTrackEditor::AddColorParameter, Component, CPDTrack, GetCPDIndex),
+								FCanExecuteAction::CreateRaw(this, &FCustomPrimitiveDataTrackEditor::CanAddParameter, CPDTrack, GetCPDIndex, 4))
+						);
+					}
+					else
+					{
+						MenuBuilder.AddMenuEntry(
+							DisplayName,
+							TAttribute<FText>::CreateLambda([this, CPDTrack, GetCPDIndex]()
+								{
+									if (CanAddParameter(CPDTrack, GetCPDIndex, 1))
+									{
+										return FText::Format(LOCTEXT("CustomPrimitiveDataMaterialTooltipScalar", "Add a track for this scalar parameter starting at index {0}"), FText::AsNumber(GetCPDIndex()));
+									}
+									else
+									{
+										return FText::Format(LOCTEXT("CustomPrimitiveDataMaterialTooltipScalarInvalid", "Can't add a scalar starting at index {0} as it overlaps with other parameters using that index."), FText::AsNumber(GetCPDIndex()));
+									}
+								}),
+							FSlateIcon(),
+							FUIAction(
+								FExecuteAction::CreateRaw(this, &FCustomPrimitiveDataTrackEditor::AddScalarParameter, Component, CPDTrack, GetCPDIndex),
+								FCanExecuteAction::CreateRaw(this, &FCustomPrimitiveDataTrackEditor::CanAddParameter, CPDTrack, GetCPDIndex, 1))
+						);
+					}
+				}
+			}
+		}
+		MenuBuilder.EndSection();
+	}
+	MenuBuilder.BeginSection("FromStartIndex", LOCTEXT("CustomPrimitiveDataFromStartIndex", "From Start Index"));
+	{
+		MenuBuilder.AddWidget(
+			SNew(SHorizontalBox)
+			+ SHorizontalBox::Slot()
+			[
+				SNew(SSpacer)
+			]
+		+ SHorizontalBox::Slot()
+			.AutoWidth()
+			[
+				SNew(SBox)
+				.WidthOverride(50.f)
+				[
+					SNew(SSpinBox<int32>)
+					.Style(&FAppStyle::GetWidgetStyle<FSpinBoxStyle>("Sequencer.HyperlinkSpinBox"))
+					.OnValueCommitted_Lambda([=](int32 Value, ETextCommit::Type CommitType) { OnStartIndexChanged(Value); })
+					.OnValueChanged_Lambda(OnStartIndexChanged)
+					.MinValue(0)
+					.MaxValue(35)
+					.Value_Lambda([=]() -> int32 { return GetStartIndex(); })
+				]
+			],
+			LOCTEXT("StartIndexText", "Custom Primitive Data Start Index")
+			);
+
+
+		MenuBuilder.AddMenuEntry(
+			LOCTEXT("CustomPrimitiveDataScalar", "Scalar Parameter"),
+			TAttribute<FText>::CreateLambda([this, CPDTrack, GetStartIndex]()
+				{
+					if (CanAddParameter(CPDTrack, GetStartIndex, 1))
+					{
+						return FText::Format(LOCTEXT("CustomPrimitiveDataMaterialTooltipScalar", "Add a track for a new scalar parameter starting at index {0}"), FText::AsNumber(StartIndex));
+					}
+					else
+					{
+						return FText::Format(LOCTEXT("CustomPrimitiveDataMaterialTooltipScalarInvalid", "Can't add a scalar parameter starting at index {0} as it overlaps with other parameters using that index."), FText::AsNumber(StartIndex));
+					}
+				}), 
+			FSlateIcon(),
+			FUIAction(
+				FExecuteAction::CreateRaw(this, &FCustomPrimitiveDataTrackEditor::AddScalarParameter, Component, CPDTrack, GetStartIndex),
+				FCanExecuteAction::CreateRaw(this, &FCustomPrimitiveDataTrackEditor::CanAddParameter, CPDTrack, GetStartIndex, 1))
+		);
+
+		MenuBuilder.AddMenuEntry(
+			LOCTEXT("CustomPrimitiveDataVector2D", "Vector2D Parameter"),
+			TAttribute<FText>::CreateLambda([this, CPDTrack, GetStartIndex]()
+				{
+					if (CanAddParameter(CPDTrack, GetStartIndex, 2))
+					{
+						return FText::Format(LOCTEXT("CustomPrimitiveDataMaterialTooltipVector2D", "Add a track for a new Vector2D parameter starting at index {0}"), FText::AsNumber(StartIndex));
+					}
+					else
+					{
+						return FText::Format(LOCTEXT("CustomPrimitiveDataMaterialTooltipVector2DInvalid", "Can't add a Vector2D parameter starting at index {0} as it overlaps with other parameters using that index."), FText::AsNumber(StartIndex));
+					}
+				}), 
+			FSlateIcon(),
+			FUIAction(
+				FExecuteAction::CreateRaw(this, &FCustomPrimitiveDataTrackEditor::AddVector2DParameter, Component, CPDTrack, GetStartIndex),
+				FCanExecuteAction::CreateRaw(this, &FCustomPrimitiveDataTrackEditor::CanAddParameter, CPDTrack, GetStartIndex, 2))
+		);
+
+		MenuBuilder.AddMenuEntry(
+			LOCTEXT("CustomPrimitiveDataVector", "Vector Parameter"),
+			TAttribute<FText>::CreateLambda([this, CPDTrack, GetStartIndex]()
+				{
+					if (CanAddParameter(CPDTrack, GetStartIndex, 3))
+					{
+						return FText::Format(LOCTEXT("CustomPrimitiveDataMaterialTooltipVector", "Add a track for a new Vector parameter starting at index {0}"), FText::AsNumber(StartIndex));
+					}
+					else
+					{
+						return FText::Format(LOCTEXT("CustomPrimitiveDataMaterialTooltipVectorInvalid", "Can't add a Vector parameter starting at index {0} as it overlaps with other parameters using that index."), FText::AsNumber(StartIndex));
+					}
+				}),
+			FSlateIcon(),
+			FUIAction(
+				FExecuteAction::CreateRaw(this, &FCustomPrimitiveDataTrackEditor::AddVectorParameter, Component, CPDTrack, GetStartIndex),
+				FCanExecuteAction::CreateRaw(this, &FCustomPrimitiveDataTrackEditor::CanAddParameter, CPDTrack, GetStartIndex, 3))
+		);
+
+		MenuBuilder.AddMenuEntry(
+			LOCTEXT("CustomPrimitiveDataColor", "Color Parameter"),
+			TAttribute<FText>::CreateLambda([this, CPDTrack, GetStartIndex]()
+				{
+					if (CanAddParameter(CPDTrack, GetStartIndex, 4))
+					{
+						return FText::Format(LOCTEXT("CustomPrimitiveDataMaterialTooltipColor", "Add a track for a new color parameter starting at index {0}"), FText::AsNumber(StartIndex));
+					}
+					else
+					{
+						return FText::Format(LOCTEXT("CustomPrimitiveDataMaterialTooltipVectorColorInvalid", "Can't add a color parameter starting at index {0} as it overlaps with other parameters using that index."), FText::AsNumber(StartIndex));
+					}
+				}), 
+			FSlateIcon(),
+			FUIAction(
+				FExecuteAction::CreateRaw(this, &FCustomPrimitiveDataTrackEditor::AddColorParameter, Component, CPDTrack, GetStartIndex),
+				FCanExecuteAction::CreateRaw(this, &FCustomPrimitiveDataTrackEditor::CanAddParameter, CPDTrack, GetStartIndex, 4))
+		);
+	}
+	MenuBuilder.EndSection();
+}
+
+
+void FCustomPrimitiveDataTrackEditor::AddScalarParameter(UPrimitiveComponent* Component, UMovieSceneCustomPrimitiveDataTrack* CPDTrack, FGetStartIndexDelegate GetStartIndexDelegate)
+{
+	FFrameNumber KeyTime = GetTimeForKey();
+
+	const FScopedTransaction Transaction(LOCTEXT("AddScalarParameter", "Add scalar parameter"));
+	float ParameterValue = 0.0f;
+	const TArray<float>& CustomPrimitiveData = Component->GetCustomPrimitiveData().Data;
+	uint8 PrimitiveDataIndex = GetStartIndexDelegate();
+	if (CustomPrimitiveData.IsValidIndex(PrimitiveDataIndex))
+	{
+		ParameterValue = CustomPrimitiveData[PrimitiveDataIndex];
+	}
+	CPDTrack->Modify();
+	CPDTrack->AddScalarParameterKey(PrimitiveDataIndex, KeyTime, ParameterValue);
+
+	GetSequencer()->NotifyMovieSceneDataChanged(EMovieSceneDataChangeType::MovieSceneStructureItemAdded);
+}
+
+void FCustomPrimitiveDataTrackEditor::AddVector2DParameter(UPrimitiveComponent* Component, UMovieSceneCustomPrimitiveDataTrack* CPDTrack, FGetStartIndexDelegate GetStartIndexDelegate)
+{
+	FFrameNumber KeyTime = GetTimeForKey();
+
+	const FScopedTransaction Transaction(LOCTEXT("AddVector2DParameter", "Add Vector2D parameter"));
+	FVector2D ParameterValue = FVector2D(0.0f, 0.0f);
+	const TArray<float>& CustomPrimitiveData = Component->GetCustomPrimitiveData().Data;
+	uint8 PrimitiveDataStartIndex = GetStartIndexDelegate();
+	if (CustomPrimitiveData.IsValidIndex(PrimitiveDataStartIndex) && CustomPrimitiveData.IsValidIndex(PrimitiveDataStartIndex + 1))
+	{
+		ParameterValue.X = CustomPrimitiveData[PrimitiveDataStartIndex];
+		ParameterValue.Y = CustomPrimitiveData[PrimitiveDataStartIndex + 1];
+	}
+	CPDTrack->Modify();
+	CPDTrack->AddVector2DParameterKey(PrimitiveDataStartIndex, KeyTime, ParameterValue);
+
+	GetSequencer()->NotifyMovieSceneDataChanged(EMovieSceneDataChangeType::MovieSceneStructureItemAdded);
+}
+
+void FCustomPrimitiveDataTrackEditor::AddVectorParameter(UPrimitiveComponent* Component, UMovieSceneCustomPrimitiveDataTrack* CPDTrack, FGetStartIndexDelegate GetStartIndexDelegate)
+{
+	FFrameNumber KeyTime = GetTimeForKey();
+
+	const FScopedTransaction Transaction(LOCTEXT("AddVectorParameter", "Add Vector parameter"));
+	FVector ParameterValue = FVector(0.0f, 0.0f, 0.0f);
+	const TArray<float>& CustomPrimitiveData = Component->GetCustomPrimitiveData().Data;
+	uint8 PrimitiveDataStartIndex = GetStartIndexDelegate();
+	if (CustomPrimitiveData.IsValidIndex(PrimitiveDataStartIndex) && CustomPrimitiveData.IsValidIndex(PrimitiveDataStartIndex + 2))
+	{
+		ParameterValue.X = CustomPrimitiveData[PrimitiveDataStartIndex];
+		ParameterValue.Y = CustomPrimitiveData[PrimitiveDataStartIndex + 1];
+		ParameterValue.Z = CustomPrimitiveData[PrimitiveDataStartIndex + 2];
+	}
+	CPDTrack->Modify();
+	CPDTrack->AddVectorParameterKey(PrimitiveDataStartIndex, KeyTime, ParameterValue);
+
+	GetSequencer()->NotifyMovieSceneDataChanged(EMovieSceneDataChangeType::MovieSceneStructureItemAdded);
+}
+
+void FCustomPrimitiveDataTrackEditor::AddColorParameter(UPrimitiveComponent* Component, UMovieSceneCustomPrimitiveDataTrack* CPDTrack, FGetStartIndexDelegate GetStartIndexDelegate)
+{
+	FFrameNumber KeyTime = GetTimeForKey();
+
+	const FScopedTransaction Transaction(LOCTEXT("AddColorParameter", "Add Color parameter"));
+	FLinearColor ParameterValue = FLinearColor(1.0f, 1.0f, 1.0f, 1.0f);
+	const TArray<float>& CustomPrimitiveData = Component->GetCustomPrimitiveData().Data;
+	uint8 PrimitiveDataStartIndex = GetStartIndexDelegate();
+	if (CustomPrimitiveData.IsValidIndex(PrimitiveDataStartIndex) && CustomPrimitiveData.IsValidIndex(PrimitiveDataStartIndex + 3))
+	{
+		ParameterValue.R = CustomPrimitiveData[PrimitiveDataStartIndex];
+		ParameterValue.G = CustomPrimitiveData[PrimitiveDataStartIndex + 1];
+		ParameterValue.B = CustomPrimitiveData[PrimitiveDataStartIndex + 2];
+		ParameterValue.A = CustomPrimitiveData[PrimitiveDataStartIndex + 3];
+	}
+	CPDTrack->Modify();
+	CPDTrack->AddColorParameterKey(PrimitiveDataStartIndex, KeyTime, ParameterValue);
+
+	GetSequencer()->NotifyMovieSceneDataChanged(EMovieSceneDataChangeType::MovieSceneStructureItemAdded);
+}
+
+bool FCustomPrimitiveDataTrackEditor::CanAddParameter(UMovieSceneCustomPrimitiveDataTrack* CPDTrack, FGetStartIndexDelegate GetStartIndexDelegate, int ParameterSize)
+{
+	check(CPDTrack);
+	check(ParameterSize >= 1 && ParameterSize <= 4);
+	UMovieSceneCustomPrimitiveDataSection* NearestSection = Cast< UMovieSceneCustomPrimitiveDataSection>(CPDTrack->GetSectionToKey());
+	if (NearestSection == nullptr)
+	{
+		NearestSection = Cast<UMovieSceneCustomPrimitiveDataSection>(MovieSceneHelpers::FindNearestSectionAtTime(CPDTrack->GetAllSections(), GetTimeForKey(), INDEX_NONE));
+	}
+	if (NearestSection)
+	{
+		uint8 CPDStartIndex = GetStartIndexDelegate();
+		uint64 ChannelsUsedBitmap = NearestSection->GetChannelsUsed();
+		uint8 TestBit = 0;
+		for (int i = CPDStartIndex; i < ParameterSize + CPDStartIndex; ++i)
+		{
+			TestBit |= ((uint64)1 << i);
+		}
+		return (CPDStartIndex >= 0 && CPDStartIndex < FCustomPrimitiveData::NumCustomPrimitiveDataFloats - ParameterSize + 1)
+			&& ((ChannelsUsedBitmap & TestBit) == 0);
+	}
+	return true;
+}
+
+TSharedRef<ISequencerTrackEditor> FCustomPrimitiveDataTrackEditor::CreateTrackEditor(TSharedRef<ISequencer> OwningSequencer)
+{
+	return MakeShareable(new FCustomPrimitiveDataTrackEditor(OwningSequencer));
+}
+
+bool FCustomPrimitiveDataTrackEditor::SupportsType(TSubclassOf<UMovieSceneTrack> Type) const
+{
+	return Type == UMovieSceneCustomPrimitiveDataTrack::StaticClass();
+}
+
+bool FCustomPrimitiveDataTrackEditor::GetDefaultExpansionState(UMovieSceneTrack* InTrack) const
+{
+	return true;
+}
+
+void FCustomPrimitiveDataTrackEditor::ExtendObjectBindingTrackMenu(TSharedRef<FExtender> Extender, const TArray<FGuid>& ObjectBindings, const UClass* ObjectClass)
+{
+	if (ObjectClass->IsChildOf(UPrimitiveComponent::StaticClass()))
+	{
+		Extender->AddMenuExtension(SequencerMenuExtensionPoints::AddTrackMenu_PropertiesSection, EExtensionHook::Before, nullptr, FMenuExtensionDelegate::CreateSP(this, &FCustomPrimitiveDataTrackEditor::ConstructObjectBindingTrackMenu, ObjectBindings));
+	}
+}
+
+void FCustomPrimitiveDataTrackEditor::ConstructObjectBindingTrackMenu(FMenuBuilder& MenuBuilder, TArray<FGuid> ObjectBindings)
+{
+	UObject* Object = GetSequencer()->FindSpawnedObjectOrTemplate(ObjectBindings[0]);
+	if (!Object)
+	{
+		return;
+	}
+
+	UPrimitiveComponent* PrimitiveComponent = Cast<UPrimitiveComponent>(Object);
+	if (!PrimitiveComponent)
+	{
+		return;
+	}
+
+	TSharedPtr<ISequencer> SequencerPtr = GetSequencer();
+	UMovieScene* MovieScene = SequencerPtr->GetFocusedMovieSceneSequence()->GetMovieScene();
+
+	const bool bAlreadyExists = MovieScene->FindTrack<UMovieSceneCustomPrimitiveDataTrack>(ObjectBindings[0]) != nullptr;
+	if (!bAlreadyExists)
+	{
+		MenuBuilder.BeginSection("CustomPrimitiveData", LOCTEXT("CustomPrimitiveDataSection", "Custom Primitive Data"));
+		{
+			MenuBuilder.AddMenuEntry(
+				NSLOCTEXT("Sequencer", "AddCustomPrimitiveDataTrack", "Custom Primitive Data"),
+				NSLOCTEXT("Sequencer", "AddCustomPrimitiveDataTrackTooltip", "Adds a track to animate custom primitive data"),
+				FSlateIcon(),
+				FUIAction(
+					FExecuteAction::CreateSP(this, &FCustomPrimitiveDataTrackEditor::HandleAddCustomPrimitiveDataTrackExecute, PrimitiveComponent)
+				)
+			);
+		}
+		MenuBuilder.EndSection();
+	}
+}
+
+void FCustomPrimitiveDataTrackEditor::HandleAddCustomPrimitiveDataTrackExecute(UPrimitiveComponent* Component)
+{
+	TSharedPtr<ISequencer> SequencerPtr = GetSequencer();
+	UMovieScene* MovieScene = SequencerPtr->GetFocusedMovieSceneSequence()->GetMovieScene();
+	if (MovieScene->IsReadOnly())
+	{
+		return;
+	}
+
+	const FScopedTransaction Transaction(LOCTEXT("AddCPDTrack", "Add custom primitive data track"));
+
+	MovieScene->Modify();
+
+	FString ComponentName = Component->GetName();
+
+	TArray<UActorComponent*> ActorComponents;
+	ActorComponents.Add(Component);
+
+	USelection* SelectedActors = GEditor->GetSelectedActors();
+	if (SelectedActors && SelectedActors->Num() > 0)
+	{
+		for (FSelectionIterator Iter(*SelectedActors); Iter; ++Iter)
+		{
+			AActor* Actor = CastChecked<AActor>(*Iter);
+
+			TArray<UActorComponent*> OutActorComponents;
+			Actor->GetComponents(OutActorComponents);
+			for (UActorComponent* ActorComponent : OutActorComponents)
+			{
+				if (ActorComponent->GetName() == ComponentName)
+				{
+					ActorComponents.AddUnique(ActorComponent);
+				}
+			}
+		}
+	}
+
+	for (UActorComponent* ActorComponent : ActorComponents)
+	{
+		FGuid ObjectHandle = SequencerPtr->GetHandleToObject(ActorComponent);
+		const FMovieSceneBinding* Binding = Algo::FindBy(MovieScene->GetBindings(), ObjectHandle, &FMovieSceneBinding::GetObjectGuid);
+		
+		const bool bAlreadyExists = MovieScene->FindTrack<UMovieSceneCustomPrimitiveDataTrack>(ObjectHandle) != nullptr;
+		if (!bAlreadyExists)
+		{
+			MovieScene->AddTrack<UMovieSceneCustomPrimitiveDataTrack>(ObjectHandle);
+		}
+	}
+
+	SequencerPtr->NotifyMovieSceneDataChanged(EMovieSceneDataChangeType::MovieSceneStructureItemAdded);
+}
+
+
+
+#undef LOCTEXT_NAMESPACE
diff --git a/Engine/Source/Editor/MovieSceneTools/Public/MovieSceneToolsModule.h b/Engine/Source/Editor/MovieSceneTools/Public/MovieSceneToolsModule.h
--- a/Engine/Source/Editor/MovieSceneTools/Public/MovieSceneToolsModule.h
+++ b/Engine/Source/Editor/MovieSceneTools/Public/MovieSceneToolsModule.h
@@ -123,6 +123,7 @@ private:
 	FDelegateHandle PrimitiveMaterialCreateEditorHandle;
 	FDelegateHandle CameraShakeSourceShakeCreateEditorHandle;
 	FDelegateHandle CVarTrackCreateEditorHandle;
+	FDelegateHandle CustomPrimitiveDataTrackCreateEditorHandle;
 
 	FDelegateHandle CameraCutTrackModelHandle;
 	FDelegateHandle CinematicShotTrackModelHandle;
diff --git a/Engine/Source/Editor/MovieSceneTools/Public/TrackEditors/CustomPrimitiveDataTrackEditor.h b/Engine/Source/Editor/MovieSceneTools/Public/TrackEditors/CustomPrimitiveDataTrackEditor.h
new file mode 100644
--- /dev/null
+++ b/Engine/Source/Editor/MovieSceneTools/Public/TrackEditors/CustomPrimitiveDataTrackEditor.h
@@ -0,0 +1,98 @@
+// Copyright Epic Games, Inc. All Rights Reserved.
+
+#pragma once
+
+#include "CoreMinimal.h"
+#include "Misc/Guid.h"
+#include "Templates/SubclassOf.h"
+#include "Widgets/SWidget.h"
+#include "ISequencer.h"
+#include "MovieSceneTrack.h"
+#include "ISequencerSection.h"
+#include "ISequencerTrackEditor.h"
+#include "MaterialTypes.h"
+#include "MovieSceneTrackEditor.h"
+
+class UMovieSceneCustomPrimitiveDataTrack;
+
+using FGetStartIndexDelegate = TFunction<uint8()>;
+
+/**
+ * Track editor for custom primitive data tracks
+ */
+class MOVIESCENETOOLS_API FCustomPrimitiveDataTrackEditor
+	: public FMovieSceneTrackEditor
+{
+public:
+
+	/** Constructor. */
+	FCustomPrimitiveDataTrackEditor(TSharedRef<ISequencer> InSequencer);
+
+	/** Virtual destructor. */
+	virtual ~FCustomPrimitiveDataTrackEditor() { }
+
+	static TSharedRef<ISequencerTrackEditor> CreateTrackEditor(TSharedRef<ISequencer> OwningSequencer);
+
+public:
+
+	// ISequencerTrackEditor interface
+
+	virtual TSharedPtr<SWidget> BuildOutlinerEditWidget(const FGuid& ObjectBinding, UMovieSceneTrack* Track, const FBuildEditWidgetParams& Params) override;
+	virtual TSharedRef<ISequencerSection> MakeSectionInterface(UMovieSceneSection& SectionObject, UMovieSceneTrack& Track, FGuid ObjectBinding) override;
+	virtual bool SupportsType(TSubclassOf<UMovieSceneTrack> Type) const override;
+	virtual bool GetDefaultExpansionState(UMovieSceneTrack* InTrack) const override; 
+	virtual void ExtendObjectBindingTrackMenu(TSharedRef<FExtender> Extender, const TArray<FGuid>& ObjectBindings, const UClass* ObjectClass) override;
+
+protected:
+
+private:
+
+
+	/** Provides the contents of the outliner edit widget */
+	TSharedRef<SWidget> OnGetAddMenuContent(FGuid ObjectBinding, UMovieSceneCustomPrimitiveDataTrack* CPDTrack, int32 TrackInsertRowIndex);
+
+	/** Provides the contents of the add parameter menu. */
+	TSharedRef<SWidget> OnGetAddParameterMenuContent(FGuid ObjectBinding, UMovieSceneCustomPrimitiveDataTrack* CPDTrack);
+
+	/** Provides the contents of the add parameter menu. */
+	void OnBuildAddParameterMenu(FMenuBuilder& MenuBuilder, FGuid ObjectBinding, UMovieSceneCustomPrimitiveDataTrack* CPDTrack);
+
+	/** Adds a scalar parameter and initial key to a custom primitive data track.
+	 * @param ObjectBinding The object binding which owns the custom primitive data track.
+	 * @param CPDTrack The track in which to look for sections to add the parameter to.
+	 * @param PrimitiveDataIndex The index for the primitive data to animate
+	 */
+	void AddScalarParameter(UPrimitiveComponent* Component, UMovieSceneCustomPrimitiveDataTrack* CPDTrack, FGetStartIndexDelegate GetStartIndexDelegate);
+
+	/** Adds a Vector2D parameter and initial key to a custom primitive data track.
+	 * @param ObjectBinding The object binding which owns the custom primitive data track.
+	 * @param CPDTrack The track in which to look for sections to add the parameter to.
+	 * @param PrimitiveDataStartIndex The start index in primitive data where we are animating a Vector2D
+	 */
+	void AddVector2DParameter(UPrimitiveComponent* Component, UMovieSceneCustomPrimitiveDataTrack* CPDTrack, FGetStartIndexDelegate GetStartIndexDelegate);
+
+	/** Adds a Vector parameter and initial key to a custom primitive data track.
+	 * @param ObjectBinding The object binding which owns the custom primitive data track.
+	 * @param CPDTrack The track in which to look for sections to add the parameter to.
+	 * @param PrimitiveDataStartIndex The start index in primitive data where we are animating a Vector
+	 */
+	void AddVectorParameter(UPrimitiveComponent* Component, UMovieSceneCustomPrimitiveDataTrack* CPDTrack, FGetStartIndexDelegate GetStartIndexDelegate);
+
+	/** Adds a color parameter and initial key to a custom primitive data track.
+	 * @param ObjectBinding The object binding which owns the custom primitive data track.
+	 * @param CPDTrack The track in which to look for sections to add the parameter to.
+	 * @param PrimitiveDataStartIndex The start index in primitive data where we are animating a Color
+	 */
+	void AddColorParameter(UPrimitiveComponent* Component, UMovieSceneCustomPrimitiveDataTrack* CPDTrack, FGetStartIndexDelegate GetStartIndexDelegate);
+
+	/** Returns whether we can add a parameter starting at this index of size ParameterSize. 
+	 * Will return false if another parameter added overlaps that index.
+	 */
+	bool CanAddParameter(UMovieSceneCustomPrimitiveDataTrack* CPDTrack, FGetStartIndexDelegate GetStartIndexDelegate, int ParameterSize);
+
+	void ConstructObjectBindingTrackMenu(FMenuBuilder& MenuBuilder, TArray<FGuid> ObjectBindings);
+
+	void HandleAddCustomPrimitiveDataTrackExecute(UPrimitiveComponent* Component);
+
+	uint8 StartIndex = 0;
+};
\ No newline at end of file
diff --git a/Engine/Source/Editor/Sequencer/Private/ISequencerSection.cpp b/Engine/Source/Editor/Sequencer/Private/ISequencerSection.cpp
--- a/Engine/Source/Editor/Sequencer/Private/ISequencerSection.cpp
+++ b/Engine/Source/Editor/Sequencer/Private/ISequencerSection.cpp
@@ -29,9 +29,9 @@ struct FMovieSceneChannel;
 /** Data pertaining to a group of channels */
 struct FGroupData
 {
-	FGroupData(FText InGroupText, FText InGroupTooltipText)
+	FGroupData(FText InGroupText, FGetMovieSceneTooltipText InGetGroupTooltipTextDelegate)
 		: GroupText(InGroupText)
-		, GroupTooltipText(InGroupTooltipText)
+		, GetGroupTooltipTextDelegate(InGetGroupTooltipTextDelegate)
 		, SortOrder(-1)
 	{}
 
@@ -48,8 +48,8 @@ struct FGroupData
 	/** Text to display for the group */
 	FText GroupText;
 	
-	/** Text to display for the group tooltip */
-	FText GroupTooltipText;
+	/** Getter for text to display for the group tooltip */
+	FGetMovieSceneTooltipText GetGroupTooltipTextDelegate;
 
 	/** Sort order of the group */
 	uint32 SortOrder;
@@ -97,7 +97,7 @@ void ISequencerSection::GenerateSectionLayout( ISectionLayoutBuilder& LayoutBuil
 					{
 						GroupDisplayName = FText::FromName(GroupName);
 					}
-					ExistingGroup = &GroupToChannelsMap.Add(GroupName, FGroupData(GroupDisplayName, FText::FromString(MetaData.GetPropertyMetaData(FCommonChannelData::GroupTooltipText))));
+					ExistingGroup = &GroupToChannelsMap.Add(GroupName, FGroupData(GroupDisplayName, MetaData.GetGroupTooltipTextDelegate));
 				}
 
 				ExistingGroup->AddChannel(FChannelData{ Channel, MetaData });
@@ -185,7 +185,7 @@ void ISequencerSection::GenerateSectionLayout( ISectionLayoutBuilder& LayoutBuil
 				return this->ConstructCategoryModel(InCategoryName, InDisplayText, ChannelData.Channels);
 			};
 
-			LayoutBuilder.PushCategory(GroupName, ChannelData.GroupText, ChannelData.GroupTooltipText, Factory);
+			LayoutBuilder.PushCategory(GroupName, ChannelData.GroupText, ChannelData.GetGroupTooltipTextDelegate, Factory);
 		}
 
 		for (const FChannelData& ChannelAndData : ChannelData.Channels)
diff --git a/Engine/Source/Editor/Sequencer/Private/MVVM/ViewModels/CategoryModel.cpp b/Engine/Source/Editor/Sequencer/Private/MVVM/ViewModels/CategoryModel.cpp
--- a/Engine/Source/Editor/Sequencer/Private/MVVM/ViewModels/CategoryModel.cpp
+++ b/Engine/Source/Editor/Sequencer/Private/MVVM/ViewModels/CategoryModel.cpp
@@ -6,10 +6,11 @@
 #include "ISequencerSection.h"
 #include "MVVM/Extensions/ITrackExtension.h"
 #include "MVVM/ViewModels/ChannelModel.h"
+#include "MVVM/ViewModels/ObjectBindingModel.h"
 #include "MVVM/ViewModels/SectionModel.h"
 #include "MVVM/ViewModels/SequencerEditorViewModel.h"
+#include "MVVM/ViewModels/SequenceModel.h"
 #include "MVVM/ViewModels/SequencerModelUtils.h"
-#include "MVVM/ViewModels/ViewModelIterators.h"
 #include "MVVM/Views/SChannelView.h"
 #include "MVVM/Views/SOutlinerItemViewBase.h"
 #include "MVVM/Views/SSequencerKeyNavigationButtons.h"
@@ -85,10 +86,10 @@ FLinearColor FCategoryModel::GetKeyBarColor() const
 	return FColor(160, 160, 160);
 }
 
-FCategoryGroupModel::FCategoryGroupModel(FName InCategoryName, const FText& InDisplayText, const FText& InTooltipText)
+FCategoryGroupModel::FCategoryGroupModel(FName InCategoryName, const FText& InDisplayText, FGetMovieSceneTooltipText InGetGroupTooltipTextDelegate)
 	: CategoryName(InCategoryName)
 	, DisplayText(InDisplayText)
-	, TooltipText(InTooltipText)
+	, GetGroupTooltipTextDelegate(InGetGroupTooltipTextDelegate)
 {
 	SetIdentifier(InCategoryName);
 }
@@ -168,7 +169,27 @@ FText FCategoryGroupModel::GetLabel() const
 
 FText FCategoryGroupModel::GetLabelToolTipText() const
 {
-	return GetTooltipText();
+	if (GetGroupTooltipTextDelegate.IsBound())
+	{
+		if (TViewModelPtr<FSequenceModel> SequenceModel = FindAncestorOfType<FSequenceModel>())
+		{
+			if (TSharedPtr<FSequencerEditorViewModel> SequencerModel = SequenceModel->GetEditor())
+			{
+				FMovieSceneSequenceID SequenceID = SequenceModel->GetSequenceID();
+				IMovieScenePlayer* Player = SequencerModel->GetSequencer().Get();
+				if (Player)
+				{
+					if (TViewModelPtr<FObjectBindingModel> ObjectBindingModel = FindAncestorOfType<FObjectBindingModel>())
+					{
+						FGuid ObjectBindingID = ObjectBindingModel->GetObjectGuid();
+
+						return GetGroupTooltipTextDelegate.Execute(Player, ObjectBindingID, SequenceID);
+					}
+				}
+			}
+		}
+	}
+	return FText();
 }
 
 FSlateFontInfo FCategoryGroupModel::GetLabelFont() const
diff --git a/Engine/Source/Editor/Sequencer/Private/MVVM/ViewModels/ChannelModel.cpp b/Engine/Source/Editor/Sequencer/Private/MVVM/ViewModels/ChannelModel.cpp
--- a/Engine/Source/Editor/Sequencer/Private/MVVM/ViewModels/ChannelModel.cpp
+++ b/Engine/Source/Editor/Sequencer/Private/MVVM/ViewModels/ChannelModel.cpp
@@ -13,6 +13,7 @@
 #include "ISequencerChannelInterface.h"
 #include "ISequencerSection.h"
 #include "MVVM/Extensions/ITrackExtension.h"
+#include "MVVM/ViewModels/ObjectBindingModel.h"
 #include "MVVM/ViewModels/SectionModel.h"
 #include "MVVM/ViewModels/SequenceModel.h"
 #include "MVVM/ViewModels/SequencerEditorViewModel.h"
@@ -261,12 +262,25 @@ FLinearColor FChannelModel::GetKeyBarColor() const
 	return FColor(160, 160, 160);
 }
 
-
+FChannelGroupModel::FChannelGroupModel(FName InChannelName, const FText& InDisplayText)
+	: ChannelsSerialNumber(0)
+	, ChannelName(InChannelName)
+	, DisplayText(InDisplayText)
+{
+}
 FChannelGroupModel::FChannelGroupModel(FName InChannelName, const FText& InDisplayText, const FText& InTooltipText)
 	: ChannelsSerialNumber(0)
 	, ChannelName(InChannelName)
 	, DisplayText(InDisplayText)
-	, TooltipText(InTooltipText)
+	, GetTooltipTextDelegate(FGetMovieSceneTooltipText::CreateLambda([InTooltipText](...) {return InTooltipText; }))
+{
+}
+
+FChannelGroupModel::FChannelGroupModel(FName InChannelName, const FText& InDisplayText, FGetMovieSceneTooltipText InGetTooltipTextDelegate)
+	: ChannelsSerialNumber(0)
+	, ChannelName(InChannelName)
+	, DisplayText(InDisplayText)
+	, GetTooltipTextDelegate(InGetTooltipTextDelegate)
 {
 }
 
@@ -318,6 +332,31 @@ void FChannelGroupModel::CleanupChannels()
 	}
 }
 
+FText FChannelGroupModel::GetTooltipText() const
+{
+	if (GetTooltipTextDelegate.IsBound())
+	{
+		if (TViewModelPtr<FSequenceModel> SequenceModel = FindAncestorOfType<FSequenceModel>())
+		{
+			if (TSharedPtr<FSequencerEditorViewModel> SequencerModel = SequenceModel->GetEditor())
+			{
+				FMovieSceneSequenceID SequenceID = SequenceModel->GetSequenceID();
+				IMovieScenePlayer* Player = SequencerModel->GetSequencer().Get();
+				if (Player)
+				{
+					if (TViewModelPtr<FObjectBindingModel> ObjectBindingModel = FindAncestorOfType<FObjectBindingModel>())
+					{
+						FGuid ObjectBindingID = ObjectBindingModel->GetObjectGuid();
+
+						return GetTooltipTextDelegate.Execute(Player, ObjectBindingID, SequenceID);
+					}
+				}
+			}
+		}
+	}
+	return FText();
+}
+
 TArrayView<const TWeakViewModelPtr<FChannelModel>> FChannelGroupModel::GetChannels() const
 {
 	return Channels;
@@ -720,8 +759,8 @@ void FChannelGroupOverrideHelper::RemoveChannelOverrides(TSharedPtr<FSequenceMod
 	SequenceModel->GetSequencer()->NotifyMovieSceneDataChanged(EMovieSceneDataChangeType::MovieSceneStructureItemsChanged);
 }
 
-FChannelGroupOutlinerModel::FChannelGroupOutlinerModel(FName InChannelName, const FText& InDisplayText, const FText& InTooltipText)
-	: TOutlinerModelMixin<FChannelGroupModel>(InChannelName, InDisplayText, InTooltipText)
+FChannelGroupOutlinerModel::FChannelGroupOutlinerModel(FName InChannelName, const FText& InDisplayText, FGetMovieSceneTooltipText InGetTooltipTextDelegate)
+	: TOutlinerModelMixin<FChannelGroupModel>(InChannelName, InDisplayText, InGetTooltipTextDelegate)
 {
 	SetIdentifier(InChannelName);
 }
diff --git a/Engine/Source/Editor/Sequencer/Private/MVVM/ViewModels/TrackModelLayoutBuilder.cpp b/Engine/Source/Editor/Sequencer/Private/MVVM/ViewModels/TrackModelLayoutBuilder.cpp
--- a/Engine/Source/Editor/Sequencer/Private/MVVM/ViewModels/TrackModelLayoutBuilder.cpp
+++ b/Engine/Source/Editor/Sequencer/Private/MVVM/ViewModels/TrackModelLayoutBuilder.cpp
@@ -164,7 +164,7 @@ void FTrackModelLayoutBuilder::RefreshLayout(TSharedPtr<FSectionModel> InSection
 	TrackAreaList = FHierarchicalModelListRefresher();
 }
 
-void FTrackModelLayoutBuilder::PushCategory(FName CategoryName, const FText& DisplayLabel, const FText& TooltipText, TFunction<TSharedPtr<FCategoryModel>(FName, const FText&)> OptionalFactory)
+void FTrackModelLayoutBuilder::PushCategory(FName CategoryName, const FText& DisplayLabel, FGetMovieSceneTooltipText GetGroupTooltipTextDelegate, TFunction<TSharedPtr<FCategoryModel>(FName, const FText&)> OptionalFactory)
 {
 	check(TrackAreaList.IsValid() && SequencerSection);
 
@@ -179,7 +179,7 @@ void FTrackModelLayoutBuilder::PushCategory(FName CategoryName, const FText& Dis
 	TSharedPtr<FCategoryGroupModel> OutlinerModel = OutlinerList.FindItem<FCategoryGroupModel>(CategoryNamePredicate);
 	if (!OutlinerModel)
 	{
-		OutlinerModel = MakeShared<FCategoryGroupModel>(CategoryName, DisplayLabel, TooltipText);
+		OutlinerModel = MakeShared<FCategoryGroupModel>(CategoryName, DisplayLabel, GetGroupTooltipTextDelegate);
 	}
 
 	OutlinerList.Link(OutlinerModel);
@@ -280,14 +280,13 @@ void FTrackModelLayoutBuilder::AddChannel(const FMovieSceneChannelHandle& Channe
 
 	if (!OutlinerModel)
 	{
-		FString TooltipText = MetaData->PropertyMetaData.FindRef(FCommonChannelData::TooltipText, FString());
 		if (bIsTopLevel)
 		{
-			OutlinerModel = MakeShared<FChannelGroupModel>(ChannelName, MetaData->DisplayText, FText::FromString(TooltipText));
+			OutlinerModel = MakeShared<FChannelGroupModel>(ChannelName, MetaData->DisplayText, MetaData->GetTooltipTextDelegate);
 		}
 		else
 		{
-			OutlinerModel = MakeShared<FChannelGroupOutlinerModel>(ChannelName, MetaData->DisplayText, FText::FromString(TooltipText));
+			OutlinerModel = MakeShared<FChannelGroupOutlinerModel>(ChannelName, MetaData->DisplayText, MetaData->GetTooltipTextDelegate);
 		}
 	}
 
diff --git a/Engine/Source/Editor/Sequencer/Private/MVVM/ViewModels/TrackModelLayoutBuilder.h b/Engine/Source/Editor/Sequencer/Private/MVVM/ViewModels/TrackModelLayoutBuilder.h
--- a/Engine/Source/Editor/Sequencer/Private/MVVM/ViewModels/TrackModelLayoutBuilder.h
+++ b/Engine/Source/Editor/Sequencer/Private/MVVM/ViewModels/TrackModelLayoutBuilder.h
@@ -112,7 +112,7 @@ public:
 
 	// ISectionLayoutBuilder interface
 
-	virtual void PushCategory( FName CategoryName, const FText& DisplayLabel, const FText& TooltipText, TFunction<TSharedPtr<UE::Sequencer::FCategoryModel>(FName, const FText&)> OptionalFactory ) override;
+	virtual void PushCategory( FName CategoryName, const FText& DisplayLabel, FGetMovieSceneTooltipText GetGroupTooltipTextDelegate, TFunction<TSharedPtr<UE::Sequencer::FCategoryModel>(FName, const FText&)> OptionalFactory ) override;
 	virtual void SetTopLevelChannel( const FMovieSceneChannelHandle& Channel, TFunction<TSharedPtr<UE::Sequencer::FChannelModel>(FName, const FMovieSceneChannelHandle&)> OptionalFactory ) override;
 	virtual void AddChannel( const FMovieSceneChannelHandle& Channel, TFunction<TSharedPtr<UE::Sequencer::FChannelModel>(FName, const FMovieSceneChannelHandle&)> OptionalFactory ) override;
 	virtual void PopCategory() override;
diff --git a/Engine/Source/Editor/Sequencer/Public/ISectionLayoutBuilder.h b/Engine/Source/Editor/Sequencer/Public/ISectionLayoutBuilder.h
--- a/Engine/Source/Editor/Sequencer/Public/ISectionLayoutBuilder.h
+++ b/Engine/Source/Editor/Sequencer/Public/ISectionLayoutBuilder.h
@@ -4,6 +4,7 @@
 
 #include "CoreMinimal.h"
 #include "Templates/Function.h"
+#include "Channels/MovieSceneChannelEditorData.h"
 
 class IKeyArea;
 
@@ -29,7 +30,7 @@ public:
 	 * @param CategoryName	The name of the category
 	 * @param DisplayLabel	The localized display label for the category
 	 */
-	virtual void PushCategory( FName CategoryName, const FText& DisplayLabel, const FText& TooltipText, TFunction<TSharedPtr<UE::Sequencer::FCategoryModel>(FName, const FText&)> OptionalFactory) = 0;
+	virtual void PushCategory( FName CategoryName, const FText& DisplayLabel, FGetMovieSceneTooltipText GetGroupTooltipTextDelegate, TFunction<TSharedPtr<UE::Sequencer::FCategoryModel>(FName, const FText&)> OptionalFactory) = 0;
 
 	/**
 	 * Sets the section as a key area itself
diff --git a/Engine/Source/Editor/Sequencer/Public/MVVM/ViewModels/CategoryModel.h b/Engine/Source/Editor/Sequencer/Public/MVVM/ViewModels/CategoryModel.h
--- a/Engine/Source/Editor/Sequencer/Public/MVVM/ViewModels/CategoryModel.h
+++ b/Engine/Source/Editor/Sequencer/Public/MVVM/ViewModels/CategoryModel.h
@@ -2,6 +2,7 @@
 
 #pragma once
 
+#include "Channels/MovieSceneChannelEditorData.h"
 #include "Containers/Array.h"
 #include "Containers/ArrayView.h"
 #include "Fonts/SlateFontInfo.h"
@@ -87,7 +88,7 @@ class SEQUENCER_API FCategoryGroupModel
 public:
 	UE_SEQUENCER_DECLARE_CASTABLE(FCategoryGroupModel, FOutlinerItemModel, ITrackAreaExtension, ICompoundOutlinerExtension, IDeletableExtension, IRecyclableExtension);
 
-	explicit FCategoryGroupModel(FName InCategoryName, const FText& InDisplayText, const FText& InTooltipText);
+	explicit FCategoryGroupModel(FName InCategoryName, const FText& InDisplayText, FGetMovieSceneTooltipText InGetGroupTooltipTextDelegate);
 
 	~FCategoryGroupModel();
 
@@ -106,11 +107,6 @@ public:
 		return DisplayText;
 	}
 
-	FText GetTooltipText() const
-	{
-		return TooltipText;
-	}
-
 	void AddCategory(TWeakViewModelPtr<FCategoryModel> InCategory);
 
 	TArrayView<const TWeakViewModelPtr<FCategoryModel>> GetCategories() const;
@@ -143,7 +139,7 @@ private:
 	TArray<TWeakViewModelPtr<FCategoryModel>> Categories;
 	FName CategoryName;
 	FText DisplayText;
-	FText TooltipText;
+	FGetMovieSceneTooltipText GetGroupTooltipTextDelegate;
 	FOutlinerSizing ComputedSizing;
 };
 
diff --git a/Engine/Source/Editor/Sequencer/Public/MVVM/ViewModels/ChannelModel.h b/Engine/Source/Editor/Sequencer/Public/MVVM/ViewModels/ChannelModel.h
--- a/Engine/Source/Editor/Sequencer/Public/MVVM/ViewModels/ChannelModel.h
+++ b/Engine/Source/Editor/Sequencer/Public/MVVM/ViewModels/ChannelModel.h
@@ -2,6 +2,7 @@
 
 #pragma once
 
+#include "Channels/MovieSceneChannelEditorData.h"
 #include "Channels/MovieSceneChannelOverrideContainer.h"
 #include "Containers/Array.h"
 #include "Containers/ArrayView.h"
@@ -127,7 +128,9 @@ public:
 
 	UE_SEQUENCER_DECLARE_CASTABLE(FChannelGroupModel, FViewModel, ITrackAreaExtension, IRecyclableExtension);
 
+	FChannelGroupModel(FName InChannelName, const FText& InDisplayText);
 	FChannelGroupModel(FName InChannelName, const FText& InDisplayText, const FText& InTooltipText);
+	FChannelGroupModel(FName InChannelName, const FText& InDisplayText, FGetMovieSceneTooltipText InGetTooltipTextDelegate);
 	~FChannelGroupModel();
 
 	/** Returns whether any of the channels within this group have any keyframes on them */
@@ -140,7 +143,7 @@ public:
 	FText GetDisplayText() const { return DisplayText; }
 
 	/** Returns the tooltip for this group */
-	FText GetTooltipText() const { return TooltipText; }
+	FText GetTooltipText() const;
 
 	/** Gets all the channel models in this group */
 	TArrayView<const TWeakViewModelPtr<FChannelModel>> GetChannels() const;
@@ -193,7 +196,7 @@ protected:
 	uint32 ChannelsSerialNumber;
 	FName ChannelName;
 	FText DisplayText;
-	FText TooltipText;
+	FGetMovieSceneTooltipText GetTooltipTextDelegate;
 };
 
 
@@ -210,7 +213,7 @@ public:
 
 	UE_SEQUENCER_DECLARE_CASTABLE(FChannelGroupOutlinerModel, FChannelGroupModel, FOutlinerItemModelMixin, ICompoundOutlinerExtension, IDeletableExtension);
 
-	FChannelGroupOutlinerModel(FName InChannelName, const FText& InDisplayText, const FText& InTooltipText);
+	FChannelGroupOutlinerModel(FName InChannelName, const FText& InDisplayText, FGetMovieSceneTooltipText InGetTooltipTextDelegate);
 	~FChannelGroupOutlinerModel();
 
 public:
diff --git a/Engine/Source/Runtime/MovieScene/Private/Channels/MovieSceneChannelEditorData.cpp b/Engine/Source/Runtime/MovieScene/Private/Channels/MovieSceneChannelEditorData.cpp
--- a/Engine/Source/Runtime/MovieScene/Private/Channels/MovieSceneChannelEditorData.cpp
+++ b/Engine/Source/Runtime/MovieScene/Private/Channels/MovieSceneChannelEditorData.cpp
@@ -18,9 +18,7 @@ const FLinearColor FCommonChannelData::RedChannelColor(1.0f, 0.05f, 0.05f, 0.9f)
 const FLinearColor FCommonChannelData::GreenChannelColor(0.05f, 1.0f, 0.05f, 0.9f);
 const FLinearColor FCommonChannelData::BlueChannelColor(0.1f, 0.2f, 1.0f, 0.9f);
 
-const FName FCommonChannelData::TooltipText = TEXT("TooltipText");
 const FName FCommonChannelData::GroupDisplayName = TEXT("GroupDisplayName");
-const FName FCommonChannelData::GroupTooltipText = TEXT("GroupTooltipText");
 
 
 FMovieSceneChannelMetaData::FMovieSceneChannelMetaData()
diff --git a/Engine/Source/Runtime/MovieScene/Public/Channels/MovieSceneChannelEditorData.h b/Engine/Source/Runtime/MovieScene/Public/Channels/MovieSceneChannelEditorData.h
--- a/Engine/Source/Runtime/MovieScene/Public/Channels/MovieSceneChannelEditorData.h
+++ b/Engine/Source/Runtime/MovieScene/Public/Channels/MovieSceneChannelEditorData.h
@@ -20,6 +20,7 @@ struct FMovieSceneRootEvaluationTemplateInstance;
 
 #if WITH_EDITOR
 
+DECLARE_DELEGATE_RetVal_ThreeParams(FText, FGetMovieSceneTooltipText, IMovieScenePlayer*, FGuid, FMovieSceneSequenceID);
 /**
  * Editor meta data for a channel of data within a movie scene section
  */
@@ -66,8 +67,12 @@ struct FMovieSceneChannelMetaData
 	FName Name;
 	/** Text to display on this channel's key area node */
 	FText DisplayText;
+	/** Delegate to get a dynamic tooltip for the key area node */
+	FGetMovieSceneTooltipText GetTooltipTextDelegate;
 	/** Name to group this channel with others of the same group name */
 	FText Group;
+	/** Delegate to get a dynamic tooltip for the group */
+	FGetMovieSceneTooltipText GetGroupTooltipTextDelegate;
 	/** Intent name */
 	FText IntentName;
 	/* Optional. If unspecified IKeyArea::CreateCurveEditorModel will create a fallback. */
@@ -145,7 +150,6 @@ struct FCommonChannelData
 
 	static MOVIESCENE_API const FName TooltipText;
 	static MOVIESCENE_API const FName GroupDisplayName;
-	static MOVIESCENE_API const FName GroupTooltipText;
 };
 
 
diff --git a/Engine/Source/Runtime/MovieSceneTracks/Private/MovieSceneTracksComponentTypes.cpp b/Engine/Source/Runtime/MovieSceneTracks/Private/MovieSceneTracksComponentTypes.cpp
--- a/Engine/Source/Runtime/MovieSceneTracks/Private/MovieSceneTracksComponentTypes.cpp
+++ b/Engine/Source/Runtime/MovieSceneTracks/Private/MovieSceneTracksComponentTypes.cpp
@@ -564,6 +564,10 @@ FMovieSceneTracksComponentTypes::FMovieSceneTracksComponentTypes()
 	Tags.BoundMaterialChanged = ComponentRegistry->NewTag(TEXT("Bound Material Changed"));
 	FBuiltInComponentTypes::Get()->RequiresInstantiationMask.Set(Tags.BoundMaterialChanged);
 
+	// Used to indicate the ParameterName component for certain parameter types (scalar, vector2d, vector, color)
+	// should be interpreted as an index for custom primitive data.
+	Tags.CustomPrimitiveData = ComponentRegistry->NewTag(TEXT("Custom Primitive Data"));
+
 	FBuiltInComponentTypes* BuiltInComponents = FBuiltInComponentTypes::Get();
 
 	// --------------------------------------------------------------------------------------------
@@ -755,6 +759,10 @@ FMovieSceneTracksComponentTypes::FMovieSceneTracksComponentTypes()
 	// Set up SkeletalAnimation components
 	ComponentRegistry->Factories.DuplicateChildComponent(SkeletalAnimation);
 
+	// --------------------------------------------------------------------------------------------
+	// Set up custom primitive data components
+	ComponentRegistry->Factories.DefineChildComponent(Tags.CustomPrimitiveData, Tags.CustomPrimitiveData);
+
 	InitializeMovieSceneTracksAccessors(this);
 }
 
diff --git a/Engine/Source/Runtime/MovieSceneTracks/Private/Sections/MovieSceneComponentMaterialParameterSection.cpp b/Engine/Source/Runtime/MovieSceneTracks/Private/Sections/MovieSceneComponentMaterialParameterSection.cpp
--- a/Engine/Source/Runtime/MovieSceneTracks/Private/Sections/MovieSceneComponentMaterialParameterSection.cpp
+++ b/Engine/Source/Runtime/MovieSceneTracks/Private/Sections/MovieSceneComponentMaterialParameterSection.cpp
@@ -73,9 +73,9 @@ EMovieSceneChannelProxyType UMovieSceneComponentMaterialParameterSection::CacheC
 		}
 		return Path;
 	};
-	auto GetMaterialParameterTooltipText = [](const FString& ParameterPath)
+	auto GetMaterialParameterTooltipText = [](IMovieScenePlayer* Player, FGuid BindingID, FMovieSceneSequenceID SequenceID, FString ParameterPath)
 	{
-		return FString::Printf(TEXT("Path: %s"), *ParameterPath);
+		return FText::Format(LOCTEXT("MaterialParameterPath", "Path: {0}"), FText::FromString(ParameterPath));
 	};
 	int32 SortOrder = 0;
 	for (FScalarMaterialParameterInfoAndCurve& Scalar : ScalarParameterInfosAndCurves)
@@ -84,7 +84,7 @@ EMovieSceneChannelProxyType UMovieSceneComponentMaterialParameterSection::CacheC
 		FText ParameterDisplayName = GetMaterialParameterDisplayText(FText::FromName(Scalar.ParameterInfo.Name), Scalar.ParameterLayerName, Scalar.ParameterAssetName);
 		FMovieSceneChannelMetaData MetaData(*ParameterPath, ParameterDisplayName);
 		// Prevent single channels from collapsing to the track node
-		MetaData.PropertyMetaData.Add(FCommonChannelData::TooltipText, GetMaterialParameterTooltipText(ParameterPath));
+		MetaData.GetTooltipTextDelegate.BindLambda(GetMaterialParameterTooltipText, ParameterPath);
 		MetaData.bCanCollapseToTrack = false;
 		MetaData.SortOrder = SortOrder++;
 		Channels.Add(Scalar.ParameterCurve, MetaData, TMovieSceneExternalValue<float>());
@@ -92,36 +92,35 @@ EMovieSceneChannelProxyType UMovieSceneComponentMaterialParameterSection::CacheC
 	for (FColorMaterialParameterInfoAndCurves& Color : ColorParameterInfosAndCurves)
 	{
 		FString ParameterPath = GetMaterialParameterPath(Color.ParameterInfo.Name, Color.ParameterLayerName, Color.ParameterAssetName);
-		FString TooltipText = GetMaterialParameterTooltipText(ParameterPath);
 		FText ParameterDisplayName = GetMaterialParameterDisplayText(FText::FromName(Color.ParameterInfo.Name), Color.ParameterLayerName, Color.ParameterAssetName);
 		FText Group = FText::FromString(ParameterPath);
 
 		FMovieSceneChannelMetaData MetaData_R(*(ParameterPath + TEXT("R")), FCommonChannelData::ChannelR, Group);
-		MetaData_R.PropertyMetaData.Add(FCommonChannelData::TooltipText, FString::Printf(TEXT("%s.%s"), *TooltipText, TEXT("R")));
+		MetaData_R.GetTooltipTextDelegate.BindLambda(GetMaterialParameterTooltipText, ParameterPath + TEXT(".R"));
 		MetaData_R.SortOrder = SortOrder++;
 		MetaData_R.Color = FCommonChannelData::RedChannelColor;
 		MetaData_R.PropertyMetaData.Add(FCommonChannelData::GroupDisplayName, ParameterDisplayName.ToString());
-		MetaData_R.PropertyMetaData.Add(FCommonChannelData::GroupTooltipText, TooltipText);
+		MetaData_R.GetGroupTooltipTextDelegate.BindLambda(GetMaterialParameterTooltipText, ParameterPath);
 
 		FMovieSceneChannelMetaData MetaData_G(*(ParameterPath + TEXT("G")), FCommonChannelData::ChannelG, Group);
-		MetaData_G.PropertyMetaData.Add(FCommonChannelData::TooltipText, FString::Printf(TEXT("%s.%s"), *TooltipText, TEXT("G")));
+		MetaData_G.GetTooltipTextDelegate.BindLambda(GetMaterialParameterTooltipText, ParameterPath + TEXT(".G"));
 		MetaData_G.SortOrder = SortOrder++;
 		MetaData_G.Color = FCommonChannelData::GreenChannelColor;
 		MetaData_G.PropertyMetaData.Add(FCommonChannelData::GroupDisplayName, ParameterDisplayName.ToString());
-		MetaData_G.PropertyMetaData.Add(FCommonChannelData::GroupTooltipText, TooltipText);
+		MetaData_G.GetGroupTooltipTextDelegate.BindLambda(GetMaterialParameterTooltipText, ParameterPath);
 
 		FMovieSceneChannelMetaData MetaData_B(*(ParameterPath + TEXT("B")), FCommonChannelData::ChannelB, Group);
-		MetaData_B.PropertyMetaData.Add(FCommonChannelData::TooltipText, FString::Printf(TEXT("%s.%s"), *TooltipText, TEXT("B")));
+		MetaData_B.GetTooltipTextDelegate.BindLambda(GetMaterialParameterTooltipText, ParameterPath + TEXT(".B"));
 		MetaData_B.SortOrder = SortOrder++;
 		MetaData_B.Color = FCommonChannelData::BlueChannelColor;
 		MetaData_B.PropertyMetaData.Add(FCommonChannelData::GroupDisplayName, ParameterDisplayName.ToString());
-		MetaData_B.PropertyMetaData.Add(FCommonChannelData::GroupTooltipText, TooltipText);
+		MetaData_B.GetGroupTooltipTextDelegate.BindLambda(GetMaterialParameterTooltipText, ParameterPath);
 
 		FMovieSceneChannelMetaData MetaData_A(*(ParameterPath + TEXT("A")), FCommonChannelData::ChannelA, Group);
-		MetaData_A.PropertyMetaData.Add(FCommonChannelData::TooltipText, FString::Printf(TEXT("%s.%s"), *TooltipText, TEXT("A")));
+		MetaData_A.GetTooltipTextDelegate.BindLambda(GetMaterialParameterTooltipText, ParameterPath + TEXT(".A"));
 		MetaData_A.SortOrder = SortOrder++;
 		MetaData_A.PropertyMetaData.Add(FCommonChannelData::GroupDisplayName, ParameterDisplayName.ToString());
-		MetaData_A.PropertyMetaData.Add(FCommonChannelData::GroupTooltipText, TooltipText);
+		MetaData_A.GetGroupTooltipTextDelegate.BindLambda(GetMaterialParameterTooltipText, ParameterPath);
 
 		Channels.Add(Color.RedCurve, MetaData_R, TMovieSceneExternalValue<float>());
 		Channels.Add(Color.GreenCurve, MetaData_G, TMovieSceneExternalValue<float>());
diff --git a/Engine/Source/Runtime/MovieSceneTracks/Private/Sections/MovieSceneCustomPrimitiveDataSection.cpp b/Engine/Source/Runtime/MovieSceneTracks/Private/Sections/MovieSceneCustomPrimitiveDataSection.cpp
new file mode 100644
--- /dev/null
+++ b/Engine/Source/Runtime/MovieSceneTracks/Private/Sections/MovieSceneCustomPrimitiveDataSection.cpp
@@ -0,0 +1,406 @@
+// Copyright Epic Games, Inc. All Rights Reserved.
+
+#include "Sections/MovieSceneCustomPrimitiveDataSection.h"
+#include "Channels/MovieSceneChannelProxy.h"
+#include "Tracks/MovieSceneCustomPrimitiveDataTrack.h"
+#include "MovieSceneTracksComponentTypes.h"
+#include "EntitySystem/BuiltInComponentTypes.h"
+#include "SceneTypes.h"
+#include "MaterialTypes.h"
+#include "Materials/MaterialInterface.h"
+#include "Materials/MaterialFunctionInterface.h"
+#include "Materials/MaterialLayersFunctions.h"
+
+#include UE_INLINE_GENERATED_CPP_BY_NAME(MovieSceneCustomPrimitiveDataSection)
+
+#define LOCTEXT_NAMESPACE "CustomPrimitiveDataTrack"
+
+namespace UE::MovieScene
+{
+	/* Entity IDs are an encoded type, index, and offset, with the upper 8 bits being the type, the next 8 bits being the offset, and the lower 16 bits as the index */
+	uint32 EncodeEntityID(int16 InIndex, uint8 InType, uint8 InOffset)
+	{
+		check(InIndex >= 0 && InIndex < int32(0x0000FFFF));
+		return static_cast<uint32>(InIndex) | (uint32(InType) << 24) | uint32(InOffset) << 16;
+	}
+	void DecodeEntityID(uint32 InEntityID, int16& OutIndex, uint8& OutType, uint8& OutOffset)
+	{
+		// Mask out the type to get the index
+		OutIndex = static_cast<int16>(InEntityID & 0x0000FFFF);
+		OutType = InEntityID >> 24;
+		OutOffset = static_cast<uint8>((InEntityID >> 16) & 0x000000FF);
+	}
+
+
+}// namespace UE::MovieScene
+
+void UMovieSceneCustomPrimitiveDataSection::ReconstructChannelProxy()
+{
+	FMovieSceneChannelProxyData Channels;
+
+#if WITH_EDITOR
+
+	UMovieSceneCustomPrimitiveDataTrack* CPDTrack = GetTypedOuter<UMovieSceneCustomPrimitiveDataTrack>();
+	auto GetCPDTooltipText = [CPDTrack](IMovieScenePlayer* Player, FGuid ObjectBindingID, FMovieSceneSequenceID SequenceID, int32 StartIndex, EMaterialParameterType ParameterType, FText ParameterTypeText)
+	{
+		auto GetMaterialParameterLayerName = [](const FMaterialLayersFunctions& Layers, const FMaterialParameterInfo& InParameterInfo)
+		{
+			FString LayerName;
+			if (Layers.EditorOnly.LayerNames.IsValidIndex(InParameterInfo.Index))
+			{
+				LayerName = Layers.GetLayerName(InParameterInfo.Index).ToString();
+			}
+			return LayerName;
+		};
+		auto GetMaterialParameterAssetName = [](const FMaterialLayersFunctions& Layers, const FMaterialParameterInfo& InParameterInfo)
+		{
+			FString AssetName;
+			if (InParameterInfo.Association == EMaterialParameterAssociation::LayerParameter && Layers.Layers.IsValidIndex(InParameterInfo.Index))
+			{
+				AssetName = Layers.Layers[InParameterInfo.Index]->GetName();
+			}
+			else if (InParameterInfo.Association == EMaterialParameterAssociation::BlendParameter && Layers.Blends.IsValidIndex(InParameterInfo.Index))
+			{
+				AssetName = Layers.Blends[InParameterInfo.Index]->GetName();
+			}
+			return AssetName;
+		};
+
+		auto GetMaterialParameterDisplayName = [](const FMaterialParameterInfo& InParameterInfo, const FString& InMaterialName, const FString& InLayerName, const FString& InAssetName, uint8 CPDIndex)
+		{
+			FText DisplayName = FText::Format(LOCTEXT("MaterialAndParameterDisplayName", "{0} (Index: {1}, Asset: {2})"), FText::FromName(InParameterInfo.Name), FText::AsNumber(CPDIndex), FText::FromString(InMaterialName));
+			if (!InLayerName.IsEmpty() && !InAssetName.IsEmpty())
+			{
+				DisplayName = FText::Format(LOCTEXT("MaterialParameterDisplayName", "{0} (Index: {1}, Asset: {2}.{3}.{4})"),
+					FText::FromName(InParameterInfo.Name),
+					FText::AsNumber(CPDIndex),
+					FText::FromString(InMaterialName),
+					FText::FromString(InLayerName),
+					FText::FromString(InAssetName));
+			}
+			return DisplayName;
+		};
+
+		TSortedMap<uint8, TArray<FCustomPrimitiveDataMaterialParametersData>> CPDMetadata;
+		CPDTrack->GetCPDMaterialData(*Player, ObjectBindingID, SequenceID, CPDMetadata);
+		if (const TArray<FCustomPrimitiveDataMaterialParametersData>* DataArray = CPDMetadata.Find(StartIndex))
+		{
+			FTextBuilder UsedInMaterialTextBuilder;
+			UsedInMaterialTextBuilder.AppendLine(FText::Format(LOCTEXT("CustomPrimitiveData_Tooltip_HasMaterials", "Custom Primitize Data of type {0} starting at data index {1}. Used for material parameter(s):"), ParameterTypeText, FText::AsNumber(StartIndex)));
+			for (const FCustomPrimitiveDataMaterialParametersData& Data : *DataArray)
+			{
+				if (Data.MaterialParameterType == ParameterType && Data.MaterialAsset.IsValid())
+				{
+					FMaterialLayersFunctions Layers;
+					Data.MaterialAsset->GetMaterialLayers(Layers);
+					FString LayerName = GetMaterialParameterLayerName(Layers, Data.ParameterInfo);
+					FString AssetName = GetMaterialParameterAssetName(Layers, Data.ParameterInfo);
+					FText MaterialDisplayName = GetMaterialParameterDisplayName(Data.ParameterInfo, Data.MaterialAsset->GetName(), LayerName, AssetName, StartIndex);
+
+					UsedInMaterialTextBuilder.AppendLine(MaterialDisplayName);
+				}
+			}
+			return UsedInMaterialTextBuilder.ToText();
+		}
+		else
+		{
+			return FText::Format(LOCTEXT("CustomPrimitiveData_Tooltip_NoMaterials", "Custom Primitive Data of type {0} starting at data index {1}. Not used in materials."), ParameterTypeText, FText::AsNumber(StartIndex));
+		}
+	};
+
+	ChannelsUsedBitmap = 0;
+	check(CPDTrack);
+	for (FScalarParameterNameAndCurve& Scalar : GetScalarParameterNamesAndCurves())
+	{
+		FMovieSceneChannelMetaData MetaData(Scalar.ParameterName, FText::Format(LOCTEXT("CustomPrimitiveData_Scalar_DisplayName", "Scalar (Index {0})"), FText::FromName(Scalar.ParameterName)));
+		// Should be the start index
+		check(Scalar.ParameterName.ToString().IsNumeric());
+		int32 StartIndex = FCString::Atoi(*Scalar.ParameterName.ToString());
+		check(StartIndex >= 0 && StartIndex < FCustomPrimitiveData::NumCustomPrimitiveDataFloats);
+		
+		// Prevent single channels from collapsing to the track node
+		MetaData.bCanCollapseToTrack = false;
+		MetaData.SortOrder = StartIndex;
+		ChannelsUsedBitmap |= ((uint64)1 << StartIndex);
+
+		MetaData.GetTooltipTextDelegate.BindLambda(GetCPDTooltipText, StartIndex, EMaterialParameterType::Scalar, LOCTEXT("CustomPrimitiveData_ParameterType_Scalar", "Scalar"));
+
+		Channels.Add(Scalar.ParameterCurve, MetaData, TMovieSceneExternalValue<float>());
+	}
+
+	for (FVector2DParameterNameAndCurves& Vector2D : GetVector2DParameterNamesAndCurves())
+	{
+		FString ParameterString = Vector2D.ParameterName.ToString(); 
+		check(ParameterString.IsNumeric()); 
+		int32 StartIndex = FCString::Atoi(*ParameterString);
+		check(StartIndex >= 0 && StartIndex < FCustomPrimitiveData::NumCustomPrimitiveDataFloats - 1);
+		ChannelsUsedBitmap |= (0b11 << StartIndex);
+		FText Group = FText::FromString(ParameterString);
+		FString GroupDisplayName = FText::Format(LOCTEXT("CustomPrimitiveData_Vector2D_DisplayName", "Vector2D (Index {0})"), FText::FromName(Vector2D.ParameterName)).ToString();
+
+		FMovieSceneChannelMetaData MetaData_X = FMovieSceneChannelMetaData(*(ParameterString + TEXT(".X")), FCommonChannelData::ChannelX, Group);
+		MetaData_X.SortOrder = StartIndex;
+		MetaData_X.PropertyMetaData.Add(FCommonChannelData::GroupDisplayName, GroupDisplayName);
+		MetaData_X.GetGroupTooltipTextDelegate.BindLambda(GetCPDTooltipText, StartIndex, EMaterialParameterType::None, LOCTEXT("CustomPrimitiveData_ParameterType_Vector2D", "Vector2D"));
+
+		FMovieSceneChannelMetaData MetaData_Y = FMovieSceneChannelMetaData(*(ParameterString + TEXT(".Y")), FCommonChannelData::ChannelY, Group);
+		MetaData_Y.SortOrder = StartIndex + 1;
+		MetaData_Y.PropertyMetaData.Add(FCommonChannelData::GroupDisplayName, GroupDisplayName);
+		MetaData_Y.GetGroupTooltipTextDelegate.BindLambda(GetCPDTooltipText, StartIndex, EMaterialParameterType::None, LOCTEXT("CustomPrimitiveData_ParameterType_Vector2D", "Vector2D"));
+
+		Channels.Add(Vector2D.XCurve, MetaData_X, TMovieSceneExternalValue<float>());
+		Channels.Add(Vector2D.YCurve, MetaData_Y, TMovieSceneExternalValue<float>());
+	}
+
+	for (FVectorParameterNameAndCurves& Vector : GetVectorParameterNamesAndCurves())
+	{
+		FString ParameterString = Vector.ParameterName.ToString();
+		check(ParameterString.IsNumeric());
+		int32 StartIndex = FCString::Atoi(*ParameterString);
+		check(StartIndex >= 0 && StartIndex < FCustomPrimitiveData::NumCustomPrimitiveDataFloats - 2);
+		ChannelsUsedBitmap |= (0b111 << StartIndex);
+		FText Group = FText::FromString(ParameterString);
+		FString GroupDisplayName = FText::Format(LOCTEXT("CustomPrimitiveData_Vector_DisplayName", "Vector (Index {0})"), FText::FromName(Vector.ParameterName)).ToString();
+		
+		FMovieSceneChannelMetaData MetaData_X = FMovieSceneChannelMetaData(*(ParameterString + TEXT(".X")), FCommonChannelData::ChannelX, Group);
+		MetaData_X.SortOrder = StartIndex;
+		MetaData_X.PropertyMetaData.Add(FCommonChannelData::GroupDisplayName, GroupDisplayName);
+		MetaData_X.GetGroupTooltipTextDelegate.BindLambda(GetCPDTooltipText, StartIndex, EMaterialParameterType::Vector, LOCTEXT("CustomPrimitiveData_ParameterType_Vector", "Vector"));
+
+		FMovieSceneChannelMetaData MetaData_Y = FMovieSceneChannelMetaData(*(ParameterString + TEXT(".Y")), FCommonChannelData::ChannelY, Group);
+		MetaData_Y.SortOrder = StartIndex + 1;
+		MetaData_Y.PropertyMetaData.Add(FCommonChannelData::GroupDisplayName, GroupDisplayName);
+		MetaData_Y.GetGroupTooltipTextDelegate.BindLambda(GetCPDTooltipText, StartIndex, EMaterialParameterType::Vector, LOCTEXT("CustomPrimitiveData_ParameterType_Vector", "Vector"));
+
+		FMovieSceneChannelMetaData MetaData_Z = FMovieSceneChannelMetaData(*(ParameterString + TEXT(".Z")), FCommonChannelData::ChannelZ, Group);
+		MetaData_Z.SortOrder = StartIndex + 2;
+		MetaData_Z.PropertyMetaData.Add(FCommonChannelData::GroupDisplayName, GroupDisplayName);
+		MetaData_Z.GetGroupTooltipTextDelegate.BindLambda(GetCPDTooltipText, StartIndex, EMaterialParameterType::Vector, LOCTEXT("CustomPrimitiveData_ParameterType_Vector", "Vector"));
+
+		Channels.Add(Vector.XCurve, MetaData_X, TMovieSceneExternalValue<float>());
+		Channels.Add(Vector.YCurve, MetaData_Y, TMovieSceneExternalValue<float>());
+		Channels.Add(Vector.ZCurve, MetaData_Z, TMovieSceneExternalValue<float>());
+	}
+
+	for (FColorParameterNameAndCurves& Color : GetColorParameterNamesAndCurves())
+	{
+		FString ParameterString = Color.ParameterName.ToString();
+		check(ParameterString.IsNumeric());
+		int32 StartIndex = FCString::Atoi(*ParameterString);
+		check(StartIndex >= 0 && StartIndex < FCustomPrimitiveData::NumCustomPrimitiveDataFloats - 3);
+		ChannelsUsedBitmap |= (0b1111 << StartIndex);
+		FText Group = FText::FromString(ParameterString);
+		FString GroupDisplayName = FText::Format(LOCTEXT("CustomPrimitiveData_Color_DisplayName", "Color (Index {0})"), FText::FromName(Color.ParameterName)).ToString();
+		
+		FMovieSceneChannelMetaData MetaData_R(*(ParameterString + TEXT("R")), FCommonChannelData::ChannelR, Group);
+		MetaData_R.SortOrder = StartIndex;
+		MetaData_R.Color = FCommonChannelData::RedChannelColor;
+		MetaData_R.PropertyMetaData.Add(FCommonChannelData::GroupDisplayName, GroupDisplayName);
+		MetaData_R.GetGroupTooltipTextDelegate.BindLambda(GetCPDTooltipText, StartIndex, EMaterialParameterType::Vector, LOCTEXT("CustomPrimitiveData_ParameterType_Color", "Color"));
+
+		FMovieSceneChannelMetaData MetaData_G(*(ParameterString + TEXT("G")), FCommonChannelData::ChannelG, Group);
+		MetaData_G.SortOrder = StartIndex + 1;
+		MetaData_G.Color = FCommonChannelData::GreenChannelColor;
+		MetaData_G.PropertyMetaData.Add(FCommonChannelData::GroupDisplayName, GroupDisplayName);
+		MetaData_G.GetGroupTooltipTextDelegate.BindLambda(GetCPDTooltipText, StartIndex, EMaterialParameterType::Vector, LOCTEXT("CustomPrimitiveData_ParameterType_Color", "Color"));
+
+		FMovieSceneChannelMetaData MetaData_B(*(ParameterString + TEXT("B")), FCommonChannelData::ChannelB, Group);
+		MetaData_B.SortOrder = StartIndex + 2;
+		MetaData_B.Color = FCommonChannelData::BlueChannelColor;
+		MetaData_B.PropertyMetaData.Add(FCommonChannelData::GroupDisplayName, GroupDisplayName);
+		MetaData_B.GetGroupTooltipTextDelegate.BindLambda(GetCPDTooltipText, StartIndex, EMaterialParameterType::Vector, LOCTEXT("CustomPrimitiveData_ParameterType_Color", "Color"));
+
+		FMovieSceneChannelMetaData MetaData_A(*(ParameterString + TEXT("A")), FCommonChannelData::ChannelA, Group);
+		MetaData_A.PropertyMetaData.Add(FCommonChannelData::GroupDisplayName, GroupDisplayName);
+		MetaData_A.GetGroupTooltipTextDelegate.BindLambda(GetCPDTooltipText, StartIndex, EMaterialParameterType::Vector, LOCTEXT("CustomPrimitiveData_ParameterType_Color", "Color"));
+		MetaData_A.SortOrder = StartIndex + 3;
+
+		Channels.Add(Color.RedCurve, MetaData_R, TMovieSceneExternalValue<float>());
+		Channels.Add(Color.GreenCurve, MetaData_G, TMovieSceneExternalValue<float>());
+		Channels.Add(Color.BlueCurve, MetaData_B, TMovieSceneExternalValue<float>());
+		Channels.Add(Color.AlphaCurve, MetaData_A, TMovieSceneExternalValue<float>());
+	}
+#else
+
+	for (FScalarParameterNameAndCurve& Scalar : GetScalarParameterNamesAndCurves())
+	{
+		Channels.Add(Scalar.ParameterCurve);
+	}
+	for (FVector2DParameterNameAndCurves& Vector2D : GetVector2DParameterNamesAndCurves())
+	{
+		Channels.Add(Vector2D.XCurve);
+		Channels.Add(Vector2D.YCurve);
+	}
+
+	for (FVectorParameterNameAndCurves& Vector : GetVectorParameterNamesAndCurves())
+	{
+		Channels.Add(Vector.XCurve);
+		Channels.Add(Vector.YCurve);
+		Channels.Add(Vector.ZCurve);
+	}
+
+	for (FColorParameterNameAndCurves& Color : GetColorParameterNamesAndCurves())
+	{
+		Channels.Add(Color.RedCurve);
+		Channels.Add(Color.GreenCurve);
+		Channels.Add(Color.BlueCurve);
+		Channels.Add(Color.AlphaCurve);
+	}
+
+#endif
+
+	ChannelProxy = MakeShared<FMovieSceneChannelProxy>(MoveTemp(Channels));
+}
+
+void UMovieSceneCustomPrimitiveDataSection::ImportEntityImpl(UMovieSceneEntitySystemLinker* EntityLinker, const FEntityImportParams& Params, FImportedEntity* OutImportedEntity)
+{
+	using namespace UE::MovieScene;
+
+	IMovieSceneParameterSectionExtender* Extender = GetImplementingOuter<IMovieSceneParameterSectionExtender>();
+	if (!ensureMsgf(Extender, TEXT("It is not valid for a UMovieSceneParameterSection to be used for importing entities outside of an outer chain that implements IMovieSceneParameterSectionExtender")))
+	{
+		return;
+	}
+
+	uint8 ParameterType = 0;
+	int16 EntityIndex = 0;
+	uint8 Offset = 0;
+	DecodeEntityID(Params.EntityID, EntityIndex, ParameterType, Offset);
+
+	FBuiltInComponentTypes* BuiltInComponentTypes = FBuiltInComponentTypes::Get();
+	FMovieSceneTracksComponentTypes* TracksComponentTypes = FMovieSceneTracksComponentTypes::Get();
+
+	FGuid ObjectBindingID = Params.GetObjectBindingID();
+
+	TEntityBuilder<TAddConditional<FGuid>> BaseBuilder = FEntityBuilder()
+		.AddConditional(BuiltInComponentTypes->GenericObjectBinding, ObjectBindingID, ObjectBindingID.IsValid());
+
+	auto GetParameterName = [](FName BaseOffsetName, uint8 Offset)
+	{
+		if (Offset == 0)
+		{
+			return BaseOffsetName;
+		}
+		FString BaseOffsetString = BaseOffsetName.ToString();
+		check(BaseOffsetString.IsNumeric());
+		int BaseOffset = FCString::Atoi(*BaseOffsetString);
+		return FName(FString::FromInt(BaseOffset + Offset));
+	};
+
+	switch (ParameterType)
+	{
+	case 0:
+	{
+		const FScalarParameterNameAndCurve& Scalar = ScalarParameterNamesAndCurves[EntityIndex];
+
+		if (Scalar.ParameterCurve.HasAnyData())
+		{
+			OutImportedEntity->AddBuilder(
+				BaseBuilder
+				.Add(TracksComponentTypes->ScalarParameterName, Scalar.ParameterName)
+				.Add(BuiltInComponentTypes->FloatChannel[0], &Scalar.ParameterCurve)
+			);
+		}
+		break;
+	}
+	case 1:
+	{
+		const FVector2DParameterNameAndCurves& Vector2D = Vector2DParameterNamesAndCurves[EntityIndex];
+		check(Offset >= 0 && Offset <= 1);
+		const FMovieSceneFloatChannel& Vector2DCurve = Offset == 0 ? Vector2D.XCurve : Vector2D.YCurve;
+		if (Vector2DCurve.HasAnyData())
+		{
+			OutImportedEntity->AddBuilder(
+				BaseBuilder
+				.Add(TracksComponentTypes->ScalarParameterName, GetParameterName(Vector2D.ParameterName, Offset))
+				.AddConditional(BuiltInComponentTypes->FloatChannel[0], &Vector2DCurve, Vector2DCurve.HasAnyData())
+			);
+		}
+		break;
+	}
+	case 2:
+	{
+		const FVectorParameterNameAndCurves& Vector = VectorParameterNamesAndCurves[EntityIndex];
+
+		check(Offset >= 0 && Offset <= 2);
+		const FMovieSceneFloatChannel& VectorCurve = Offset == 0 ? Vector.XCurve : Offset == 1 ? Vector.YCurve : Vector.ZCurve;
+		if (VectorCurve.HasAnyData())
+		{
+			OutImportedEntity->AddBuilder(
+				BaseBuilder
+				.Add(TracksComponentTypes->ScalarParameterName, GetParameterName(Vector.ParameterName, Offset))
+				.AddConditional(BuiltInComponentTypes->FloatChannel[0], &VectorCurve, VectorCurve.HasAnyData())
+			);
+		}
+		break;
+	}
+	case 3:
+	{
+		const FColorParameterNameAndCurves& Color = ColorParameterNamesAndCurves[EntityIndex];
+		check(Offset >= 0 && Offset <= 3);
+		const FMovieSceneFloatChannel& ColorCurve = Offset == 0 ? Color.RedCurve : Offset == 1 ? Color.GreenCurve : Offset == 2 ? Color.BlueCurve : Color.AlphaCurve;
+		if (ColorCurve.HasAnyData())
+		{
+			OutImportedEntity->AddBuilder(
+				BaseBuilder
+				.Add(TracksComponentTypes->ScalarParameterName, GetParameterName(Color.ParameterName, Offset))
+				.AddTag(TracksComponentTypes->Tags.CustomPrimitiveData)
+				.AddConditional(BuiltInComponentTypes->FloatChannel[0], &ColorCurve, ColorCurve.HasAnyData())
+			);
+		}
+		break;
+	}
+	}
+
+	Extender->ExtendEntity(this, EntityLinker, Params, OutImportedEntity);
+}
+
+void UMovieSceneCustomPrimitiveDataSection::ExternalPopulateEvaluationField(const TRange<FFrameNumber>& EffectiveRange, const FMovieSceneEvaluationFieldEntityMetaData& InMetaData, FMovieSceneEntityComponentFieldBuilder* OutFieldBuilder)
+{
+	using namespace UE::MovieScene;
+
+	const int32 MetaDataIndex = OutFieldBuilder->AddMetaData(InMetaData);
+
+	// We add a separate entity for each float channel, and encode in the name the type, the index, and the offset (for each channel)
+
+	const int32 NumScalarID = ScalarParameterNamesAndCurves.Num();
+	const int32 NumVector2DID = Vector2DParameterNamesAndCurves.Num();
+	const int32 NumVectorID = VectorParameterNamesAndCurves.Num();
+	const int32 NumColorID = ColorParameterNamesAndCurves.Num();
+
+	for (int32 Index = 0; Index < NumScalarID; ++Index)
+	{
+		const int32 EntityIndex = OutFieldBuilder->FindOrAddEntity(this, EncodeEntityID(Index, 0, 0));
+		OutFieldBuilder->AddPersistentEntity(EffectiveRange, EntityIndex, MetaDataIndex);
+	}
+	for (int32 Index = 0; Index < NumVector2DID; ++Index)
+	{
+		const int32 EntityIndex0 = OutFieldBuilder->FindOrAddEntity(this, EncodeEntityID(Index, 1, 0));
+		const int32 EntityIndex1 = OutFieldBuilder->FindOrAddEntity(this, EncodeEntityID(Index, 1, 1));
+		OutFieldBuilder->AddPersistentEntity(EffectiveRange, EntityIndex0, MetaDataIndex);
+		OutFieldBuilder->AddPersistentEntity(EffectiveRange, EntityIndex1, MetaDataIndex);
+	}
+	for (int32 Index = 0; Index < NumVectorID; ++Index)
+	{
+		const int32 EntityIndex0 = OutFieldBuilder->FindOrAddEntity(this, EncodeEntityID(Index, 2, 0));
+		const int32 EntityIndex1 = OutFieldBuilder->FindOrAddEntity(this, EncodeEntityID(Index, 2, 1));
+		const int32 EntityIndex2 = OutFieldBuilder->FindOrAddEntity(this, EncodeEntityID(Index, 2, 2));
+		OutFieldBuilder->AddPersistentEntity(EffectiveRange, EntityIndex0, MetaDataIndex);
+		OutFieldBuilder->AddPersistentEntity(EffectiveRange, EntityIndex1, MetaDataIndex);
+		OutFieldBuilder->AddPersistentEntity(EffectiveRange, EntityIndex2, MetaDataIndex);
+	}
+	for (int32 Index = 0; Index < NumColorID; ++Index)
+	{
+		const int32 EntityIndex0 = OutFieldBuilder->FindOrAddEntity(this, EncodeEntityID(Index, 3, 0));
+		const int32 EntityIndex1 = OutFieldBuilder->FindOrAddEntity(this, EncodeEntityID(Index, 3, 1));
+		const int32 EntityIndex2 = OutFieldBuilder->FindOrAddEntity(this, EncodeEntityID(Index, 3, 2));
+		const int32 EntityIndex3 = OutFieldBuilder->FindOrAddEntity(this, EncodeEntityID(Index, 3, 3));
+		OutFieldBuilder->AddPersistentEntity(EffectiveRange, EntityIndex0, MetaDataIndex);
+		OutFieldBuilder->AddPersistentEntity(EffectiveRange, EntityIndex1, MetaDataIndex);
+		OutFieldBuilder->AddPersistentEntity(EffectiveRange, EntityIndex2, MetaDataIndex);
+		OutFieldBuilder->AddPersistentEntity(EffectiveRange, EntityIndex3, MetaDataIndex);
+	}
+}
+
+
+#undef LOCTEXT_NAMESPACE
\ No newline at end of file
diff --git a/Engine/Source/Runtime/MovieSceneTracks/Private/Systems/MovieSceneCustomPrimitiveDataSystem.cpp b/Engine/Source/Runtime/MovieSceneTracks/Private/Systems/MovieSceneCustomPrimitiveDataSystem.cpp
new file mode 100644
--- /dev/null
+++ b/Engine/Source/Runtime/MovieSceneTracks/Private/Systems/MovieSceneCustomPrimitiveDataSystem.cpp
@@ -0,0 +1,524 @@
+// Copyright Epic Games, Inc. All Rights Reserved.
+
+#include "Systems/MovieSceneCustomPrimitiveDataSystem.h"
+
+#include "Systems/FloatChannelEvaluatorSystem.h"
+#include "Systems/MovieScenePiecewiseDoubleBlenderSystem.h"
+#include "Systems/MovieSceneHierarchicalBiasSystem.h"
+#include "Systems/MovieSceneInitialValueSystem.h"
+#include "Systems/MovieScenePreAnimatedMaterialParameters.h"
+#include "Systems/WeightAndEasingEvaluatorSystem.h"
+
+#include "Evaluation/PreAnimatedState/MovieScenePreAnimatedObjectStorage.h"
+#include "Evaluation/PreAnimatedState/MovieScenePreAnimatedStorageID.inl"
+#include "MovieSceneTracksComponentTypes.h"
+#include "Systems/DoubleChannelEvaluatorSystem.h"
+#include "EntitySystem/MovieSceneEntitySystemRunner.h"
+#include "SceneTypes.h"
+#include "Components/PrimitiveComponent.h"
+
+#include UE_INLINE_GENERATED_CPP_BY_NAME(MovieSceneCustomPrimitiveDataSystem)
+
+namespace UE::MovieScene
+{
+	void CollectGarbageForOutput(FAnimatedCustomPrimitiveDataInfo* Output)
+	{
+		// This should only happen during garbage collection
+		if (Output->OutputEntityID.IsValid())
+		{
+			UMovieSceneEntitySystemLinker* Linker = Output->WeakLinker.Get();
+			if (Linker)
+			{
+				Linker->EntityManager.AddComponent(Output->OutputEntityID, FBuiltInComponentTypes::Get()->Tags.NeedsUnlink);
+			}
+
+			Output->OutputEntityID = FMovieSceneEntityID();
+		}
+
+		if (Output->BlendChannelID.IsValid())
+		{
+			UMovieSceneBlenderSystem* BlenderSystem = Output->WeakBlenderSystem.Get();
+			if (BlenderSystem)
+			{
+				BlenderSystem->ReleaseBlendChannel(Output->BlendChannelID);
+			}
+			Output->BlendChannelID = FMovieSceneBlendChannelID();
+		}
+	}
+
+	FAnimatedCustomPrimitiveDataInfo::~FAnimatedCustomPrimitiveDataInfo()
+	{
+	}
+
+	/** Apply scalar material parameters */
+	struct FApplyCustomPrimitiveDataParameters
+	{
+		static void ForEachAllocation(FEntityAllocationIteratorItem Item,
+			TRead<UObject*> BoundObjects,
+			TRead<FName> ParameterNames,
+			TRead<double> ScalarValues)
+		{
+			const int32 Num = Item.GetAllocation()->Num();
+			for (int32 Index = 0; Index < Num; ++Index)
+			{
+				UPrimitiveComponent* PrimitiveComponent = CastChecked<UPrimitiveComponent>(BoundObjects[Index]);
+				FString ParameterNameString = ParameterNames[Index].ToString();
+				check(ParameterNameString.IsNumeric());
+				int32 DataIndex = FCString::Atoi(*ParameterNameString);
+				check(DataIndex >= 0 && DataIndex < FCustomPrimitiveData::NumCustomPrimitiveDataFloats);
+				PrimitiveComponent->SetCustomPrimitiveDataFloat(DataIndex, (float)ScalarValues[Index]);
+			}
+		}
+	};
+
+	struct FCustomPrimitiveDataMixin
+	{
+		void CreateEntity(UMovieSceneEntitySystemLinker* Linker, FObjectKey BoundObject, FName ParameterName, FComponentTypeID BlenderTypeTag, TArrayView<const FMovieSceneEntityID> Inputs, FAnimatedCustomPrimitiveDataInfo* Output)
+		{
+			FBuiltInComponentTypes* BuiltInComponents = FBuiltInComponentTypes::Get();
+			FMovieSceneTracksComponentTypes* TracksComponents = FMovieSceneTracksComponentTypes::Get();
+
+			bool bHasInitialValue = false;
+			float InitialValue = 0.f;
+
+			for (FMovieSceneEntityID Input : Inputs)
+			{
+				if (TOptionalComponentReader<double> ExistingInitialValue = Linker->EntityManager.ReadComponent(Input, TracksComponents->FloatParameter.InitialValue))
+				{
+					InitialValue = static_cast<float>(*ExistingInitialValue);
+					bHasInitialValue = true;
+					break;
+				}
+			}
+
+			if (!bHasInitialValue)
+			{
+				if (UPrimitiveComponent* PrimitiveComponent = Cast<UPrimitiveComponent>(BoundObject.ResolveObjectPtr()))
+				{
+					FString ParameterNameString = ParameterName.ToString();
+					check(ParameterNameString.IsNumeric());
+					int32 DataIndex = FCString::Atoi(*ParameterNameString);
+					check(DataIndex >= 0 && DataIndex < FCustomPrimitiveData::NumCustomPrimitiveDataFloats);
+					const FCustomPrimitiveData& CustomPrimitiveData = PrimitiveComponent->GetCustomPrimitiveData();
+					if (CustomPrimitiveData.Data.IsValidIndex(DataIndex))
+					{
+						bHasInitialValue = true;
+						InitialValue = CustomPrimitiveData.Data[DataIndex];
+					}
+				}
+			}
+
+			Output->OutputEntityID = FEntityBuilder()
+				.Add(BuiltInComponents->BoundObject, BoundObject.ResolveObjectPtr())
+				.Add(BuiltInComponents->BlendChannelOutput, Output->BlendChannelID)
+				.Add(BuiltInComponents->DoubleResult[0], 0.0)
+				.AddConditional(TracksComponents->FloatParameter.InitialValue, InitialValue, bHasInitialValue)
+				.AddTag(TracksComponents->FloatParameter.PropertyTag)
+				.AddTag(TracksComponents->Tags.CustomPrimitiveData)
+				.AddTag(BlenderTypeTag)
+				.AddTag(BuiltInComponents->Tags.NeedsLink)
+				.AddMutualComponents()
+				.CreateEntity(&Linker->EntityManager);
+
+			Linker->EntityManager.CopyComponents(Inputs[0], Output->OutputEntityID, Linker->EntityManager.GetComponents()->GetCopyAndMigrationMask());
+		}
+
+		void InitializeSoleInput(UMovieSceneEntitySystemLinker* Linker, FObjectKey BoundObject, FName ParameterName, FMovieSceneEntityID SoleContributor, FAnimatedCustomPrimitiveDataInfo* Output)
+		{
+			FBuiltInComponentTypes* BuiltInComponents = FBuiltInComponentTypes::Get();
+			FMovieSceneTracksComponentTypes* TracksComponents = FMovieSceneTracksComponentTypes::Get();
+
+			const FComponentMask& SoleContributorType = Linker->EntityManager.GetEntityType(SoleContributor);
+			const bool bNeedsInitialValue = SoleContributorType.ContainsAny({ BuiltInComponents->Tags.RestoreState, BuiltInComponents->Tags.AlwaysCacheInitialValue });
+			const bool bAlreadyHasInitialValue = SoleContributorType.Contains(TracksComponents->FloatParameter.InitialValue);
+
+			if (bNeedsInitialValue && !bAlreadyHasInitialValue)
+			{
+				float InitialValue = 0.0;
+				if (UPrimitiveComponent* PrimitiveComponent = Cast<UPrimitiveComponent>(BoundObject.ResolveObjectPtr()))
+				{
+					FString ParameterNameString = ParameterName.ToString();
+					check(ParameterNameString.IsNumeric());
+					int32 DataIndex = FCString::Atoi(*ParameterNameString);
+					check(DataIndex >= 0 && DataIndex < FCustomPrimitiveData::NumCustomPrimitiveDataFloats);
+					const FCustomPrimitiveData& CustomPrimitiveData = PrimitiveComponent->GetCustomPrimitiveData();
+					if (CustomPrimitiveData.Data.IsValidIndex(DataIndex))
+					{
+						InitialValue = CustomPrimitiveData.Data[DataIndex];
+						Linker->EntityManager.AddComponent(SoleContributor, TracksComponents->FloatParameter.InitialValue, InitialValue);
+					}
+				}
+			}
+
+			Linker->EntityManager.RemoveComponent(SoleContributor, FBuiltInComponentTypes::Get()->HierarchicalBlendTarget);
+		}
+	};
+
+	/** Handler that manages creation of blend outputs where there are multiple contributors for the same material parameter */
+	template<typename Mixin>
+	struct TOverlappingCustomPrimitiveDataHandler : Mixin
+	{
+		UMovieSceneEntitySystemLinker* Linker;
+		UMovieSceneCustomPrimitiveDataSystem* System;
+
+		TOverlappingCustomPrimitiveDataHandler(UMovieSceneCustomPrimitiveDataSystem* InSystem)
+			: Linker(InSystem->GetLinker())
+			, System(InSystem)
+		{}
+
+		void InitializeOutput(FObjectKey BoundObject, FName ParameterName, TArrayView<const FMovieSceneEntityID> Inputs, FAnimatedCustomPrimitiveDataInfo* Output, FEntityOutputAggregate Aggregate)
+		{
+			UpdateOutput(BoundObject, ParameterName, Inputs, Output, Aggregate);
+		}
+
+		void UpdateOutput(FObjectKey BoundObject, FName ParameterName, TArrayView<const FMovieSceneEntityID> Inputs, FAnimatedCustomPrimitiveDataInfo* Output, FEntityOutputAggregate Aggregate)
+		{
+			using namespace UE::MovieScene;
+
+			FBuiltInComponentTypes* BuiltInComponents = FBuiltInComponentTypes::Get();
+			FMovieSceneTracksComponentTypes* TracksComponents = FMovieSceneTracksComponentTypes::Get();
+
+			const int32 NumContributors = Inputs.Num();
+			if (!ensure(NumContributors != 0))
+			{
+				return;
+			}
+
+			const bool bUseBlending = NumContributors > 1 || !Linker->EntityManager.HasComponent(Inputs[0], BuiltInComponents->Tags.AbsoluteBlend) || Linker->EntityManager.HasComponent(Inputs[0], BuiltInComponents->WeightAndEasingResult);
+			if (bUseBlending)
+			{
+				if (!Output->OutputEntityID)
+				{
+					if (!System->DoubleBlenderSystem)
+					{
+						System->DoubleBlenderSystem = Linker->LinkSystem<UMovieScenePiecewiseDoubleBlenderSystem>();
+						Linker->SystemGraph.AddReference(System, System->DoubleBlenderSystem);
+					}
+
+					Output->WeakLinker = Linker;
+					Output->WeakBlenderSystem = System->DoubleBlenderSystem;
+
+					// Initialize the blend channel ID
+					Output->BlendChannelID = System->DoubleBlenderSystem->AllocateBlendChannel();
+
+					// Needs blending
+					FComponentTypeID BlenderTypeTag = System->DoubleBlenderSystem->GetBlenderTypeTag();
+					Mixin::CreateEntity(Linker, BoundObject, ParameterName, BlenderTypeTag, Inputs, Output);
+				}
+
+				const FComponentTypeID BlenderTypeTag = System->DoubleBlenderSystem->GetBlenderTypeTag();
+
+				struct FBlendInfo
+				{
+					int16 HBias = TNumericLimits<int16>::Min();
+					bool bBlendHierarchicalBias = true;
+				};
+				FBlendInfo IgnoredBlendInfo;
+				FBlendInfo BlendInfo;
+				for (FMovieSceneEntityID Input : Inputs)
+				{
+					FBlendInfo& BlendInfoToUpdate = Linker->EntityManager.HasComponent(Input, BuiltInComponents->Tags.Ignored) ? IgnoredBlendInfo : BlendInfo;
+
+					TOptionalComponentReader<int16> HBiasComponent = Linker->EntityManager.ReadComponent(Input, BuiltInComponents->HierarchicalBias);
+					const int16 HBias = HBiasComponent ? *HBiasComponent : 0;
+
+					if (HBias > BlendInfoToUpdate.HBias)
+					{
+						BlendInfoToUpdate.HBias = HBias;
+						BlendInfoToUpdate.bBlendHierarchicalBias = Linker->EntityManager.HasComponent(Input, BuiltInComponents->Tags.BlendHierarchicalBias);
+					}
+					else if (HBias == BlendInfoToUpdate.HBias && !BlendInfoToUpdate.bBlendHierarchicalBias)
+					{
+						BlendInfoToUpdate.bBlendHierarchicalBias = Linker->EntityManager.HasComponent(Input, BuiltInComponents->Tags.BlendHierarchicalBias);
+					}
+				}
+
+				if (BlendInfo.HBias == TNumericLimits<int16>::Min())
+				{
+					BlendInfo = IgnoredBlendInfo;
+				}
+				else if (IgnoredBlendInfo.HBias != TNumericLimits<int16>::Min())
+				{
+					BlendInfo.bBlendHierarchicalBias |= IgnoredBlendInfo.bBlendHierarchicalBias;
+				}
+
+				for (FMovieSceneEntityID Input : Inputs)
+				{
+					if (!Linker->EntityManager.HasComponent(Input, BuiltInComponents->BlendChannelInput))
+					{
+						Linker->EntityManager.AddComponent(Input, BuiltInComponents->BlendChannelInput, Output->BlendChannelID);
+					}
+					else
+					{
+						// If the bound object changed, we might have been re-assigned a different blend channel so make sure it's up to date
+						Linker->EntityManager.WriteComponentChecked(Input, BuiltInComponents->BlendChannelInput, Output->BlendChannelID);
+					}
+
+					if (BlendInfo.bBlendHierarchicalBias)
+					{
+						if (!Linker->EntityManager.HasComponent(Input, BuiltInComponents->HierarchicalBlendTarget))
+						{
+							Linker->EntityManager.AddComponent(Input, BuiltInComponents->HierarchicalBlendTarget, BlendInfo.HBias);
+						}
+						else
+						{
+							Linker->EntityManager.WriteComponentChecked(Input, BuiltInComponents->HierarchicalBlendTarget, BlendInfo.HBias);
+						}
+					}
+
+					// Ensure we have the blender type tag on the inputs.
+					Linker->EntityManager.AddComponent(Input, BlenderTypeTag);
+				}
+			}
+			else if (!Output->OutputEntityID && Inputs.Num() == 1)
+			{
+				Linker->EntityManager.RemoveComponent(Inputs[0], BuiltInComponents->BlendChannelInput);
+
+				Mixin::InitializeSoleInput(Linker, BoundObject, ParameterName, Inputs[0], Output);
+			}
+
+			Output->NumContributors = NumContributors;
+		}
+
+		void DestroyOutput(FObjectKey BoundObject, FName ParameterName, FAnimatedCustomPrimitiveDataInfo* Output, FEntityOutputAggregate Aggregate)
+		{
+			if (Output->OutputEntityID)
+			{
+				FBuiltInComponentTypes* BuiltInComponents = FBuiltInComponentTypes::Get();
+
+				Linker->EntityManager.AddComponent(Output->OutputEntityID, BuiltInComponents->Tags.NeedsUnlink);
+				Output->OutputEntityID = FMovieSceneEntityID();
+
+				if (System->DoubleBlenderSystem)
+				{
+					System->DoubleBlenderSystem->ReleaseBlendChannel(Output->BlendChannelID);
+				}
+				Output->BlendChannelID = FMovieSceneBlendChannelID();
+			}
+		}
+
+	};
+
+	struct FCustomPrimitiveDataEntryTraits : FBoundObjectPreAnimatedStateTraits
+	{
+		struct FKeyType
+		{
+			FObjectKey Object;
+			FName IndexName;
+
+			/** Constructor that takes a BoundObject (the component) and the FName for the data index */
+			FKeyType(FObjectKey InObject, FName InIndexName)
+				: Object(InObject)
+				, IndexName(InIndexName)
+			{}
+
+			/** Hashing and equality required for storage within a map */
+			friend uint32 GetTypeHash(const FKeyType& InKey)
+			{
+				return HashCombine(GetTypeHash(InKey.Object), GetTypeHash(InKey.IndexName));
+			}
+			friend bool operator==(const FKeyType& A, const FKeyType& B)
+			{
+				return A.Object == B.Object && A.IndexName == B.IndexName;
+			}
+		};
+
+		using KeyType = FKeyType;
+		using StorageType = float;
+
+		static float CachePreAnimatedValue(FObjectKey InObject, FName IndexName)
+		{
+			if (UPrimitiveComponent* PrimitiveComponent = Cast<UPrimitiveComponent>(InObject.ResolveObjectPtr()))
+			{
+				FString ParameterNameString = IndexName.ToString();
+				check(ParameterNameString.IsNumeric());
+				int32 DataIndex = FCString::Atoi(*ParameterNameString);
+				check(DataIndex >= 0 && DataIndex < FCustomPrimitiveData::NumCustomPrimitiveDataFloats);
+				const FCustomPrimitiveData& CustomPrimitiveData = PrimitiveComponent->GetCustomPrimitiveData();
+				if (CustomPrimitiveData.Data.IsValidIndex(DataIndex))
+				{
+					return CustomPrimitiveData.Data[DataIndex];
+				}
+			}
+			return 0.0f;
+		}
+
+		static void RestorePreAnimatedValue(const FKeyType& InKey, float OldValue, const FRestoreStateParams& Params)
+		{
+			if (UPrimitiveComponent* PrimitiveComponent = Cast<UPrimitiveComponent>(InKey.Object.ResolveObjectPtr()))
+			{
+				FString ParameterNameString = InKey.IndexName.ToString();
+				check(ParameterNameString.IsNumeric());
+				int32 DataIndex = FCString::Atoi(*ParameterNameString);
+				check(DataIndex >= 0 && DataIndex < FCustomPrimitiveData::NumCustomPrimitiveDataFloats);
+				PrimitiveComponent->SetCustomPrimitiveDataFloat(DataIndex, OldValue);
+			}
+		}
+	};
+
+	struct FPreAnimatedCustomPrimitiveDataEntryStorage
+		: public TPreAnimatedStateStorage<FCustomPrimitiveDataEntryTraits>
+	{
+		static TAutoRegisterPreAnimatedStorageID<FPreAnimatedCustomPrimitiveDataEntryStorage> StorageID;
+	};
+
+	TAutoRegisterPreAnimatedStorageID<FPreAnimatedCustomPrimitiveDataEntryStorage> FPreAnimatedCustomPrimitiveDataEntryStorage::StorageID;
+
+} // namespace UE::MovieScene
+
+UMovieSceneCustomPrimitiveDataSystem::UMovieSceneCustomPrimitiveDataSystem(const FObjectInitializer& ObjInit)
+	: Super(ObjInit)
+{
+	using namespace UE::MovieScene;
+
+	FBuiltInComponentTypes* BuiltInComponents = FBuiltInComponentTypes::Get();
+	FMovieSceneTracksComponentTypes* TracksComponents = FMovieSceneTracksComponentTypes::Get();
+
+	RelevantComponent = TracksComponents->Tags.CustomPrimitiveData;
+	Phase = ESystemPhase::Instantiation | ESystemPhase::Scheduling;
+
+	if (HasAnyFlags(RF_ClassDefaultObject))
+	{
+		DefineComponentConsumer(GetClass(), BuiltInComponents->BoundObject);
+		DefineComponentProducer(GetClass(), BuiltInComponents->HierarchicalBlendTarget);
+
+		DefineImplicitPrerequisite(UFloatChannelEvaluatorSystem::StaticClass(), GetClass());
+		DefineImplicitPrerequisite(UDoubleChannelEvaluatorSystem::StaticClass(), GetClass());
+
+		for (int32 Index = 0; Index < UE_ARRAY_COUNT(BuiltInComponents->DoubleResult); ++Index)
+		{
+			DefineComponentConsumer(GetClass(), BuiltInComponents->DoubleResult[Index]);
+		}
+
+		DefineComponentConsumer(GetClass(), TracksComponents->Tags.CustomPrimitiveData);
+
+		DefineImplicitPrerequisite(UMovieSceneHierarchicalEasingInstantiatorSystem::StaticClass(), GetClass());
+		DefineImplicitPrerequisite(UMovieScenePiecewiseDoubleBlenderSystem::StaticClass(), GetClass());
+		DefineImplicitPrerequisite(GetClass(), UMovieSceneHierarchicalBiasSystem::StaticClass());
+		DefineImplicitPrerequisite(GetClass(), UMovieSceneInitialValueSystem::StaticClass());
+	}
+}
+
+void UMovieSceneCustomPrimitiveDataSystem::OnLink()
+{
+	using namespace UE::MovieScene;
+
+	ScalarParameterTracker.Initialize(this);
+
+	ScalarParameterStorage = Linker->PreAnimatedState.GetOrCreateStorage<FPreAnimatedCustomPrimitiveDataEntryStorage>();
+}
+
+void UMovieSceneCustomPrimitiveDataSystem::OnUnlink()
+{
+	using namespace UE::MovieScene;
+
+	// Always reset the float blender system on link to ensure that recycled systems are correctly initialized.
+	DoubleBlenderSystem = nullptr;
+
+	ScalarParameterTracker.Destroy(TOverlappingCustomPrimitiveDataHandler<FCustomPrimitiveDataMixin>(this));
+}
+
+void UMovieSceneCustomPrimitiveDataSystem::OnRun(FSystemTaskPrerequisites& InPrerequisites, FSystemSubsequentTasks& Subsequents)
+{
+	using namespace UE::MovieScene;
+
+	FMovieSceneEntitySystemRunner* ActiveRunner = Linker->GetActiveRunner();
+	if (!ensure(ActiveRunner))
+	{
+		return;
+	}
+
+	ESystemPhase CurrentPhase = ActiveRunner->GetCurrentPhase();
+	if (CurrentPhase == ESystemPhase::Instantiation)
+	{
+		OnInstantiation();
+	}
+	else if (CurrentPhase == ESystemPhase::Evaluation)
+	{
+		OnEvaluation(InPrerequisites, Subsequents);
+	}
+}
+
+
+void UMovieSceneCustomPrimitiveDataSystem::OnInstantiation()
+{
+	using namespace UE::MovieScene;
+
+	FBuiltInComponentTypes* BuiltInComponents = FBuiltInComponentTypes::Get();
+	FMovieSceneTracksComponentTypes* TracksComponents = FMovieSceneTracksComponentTypes::Get();
+
+	auto HandleUnlinkedAllocation = [this](const FEntityAllocation* Allocation)
+	{
+		this->ScalarParameterTracker.VisitUnlinkedAllocation(Allocation);
+	};
+
+	auto HandleUpdatedAllocation = [this](const FEntityAllocation* Allocation, TRead<UObject*> InObjects, TRead<FName> ParameterNames)
+	{
+		this->ScalarParameterTracker.VisitActiveAllocation(Allocation, InObjects, ParameterNames);
+	};
+
+	// First step handle any new bound objects
+	FEntityTaskBuilder()
+		.Read(BuiltInComponents->BoundObject)
+		.Read(TracksComponents->ScalarParameterName)
+		.FilterAll({ BuiltInComponents->Tags.NeedsLink, TracksComponents->Tags.CustomPrimitiveData })
+		.FilterNone({ BuiltInComponents->Tags.NeedsUnlink })
+		.Iterate_PerAllocation(&Linker->EntityManager, HandleUpdatedAllocation);
+
+	// Next handle any entities that are going away
+	FEntityTaskBuilder()
+		.FilterAll({ BuiltInComponents->Tags.NeedsUnlink, TracksComponents->Tags.CustomPrimitiveData })
+		.Iterate_PerAllocation(&Linker->EntityManager, HandleUnlinkedAllocation);
+
+	// Process all blended parameters
+	TOverlappingCustomPrimitiveDataHandler<FCustomPrimitiveDataMixin> ScalarHandler(this);
+	ScalarParameterTracker.ProcessInvalidatedOutputs(Linker, ScalarHandler);
+
+
+	// Gather inputs that contribute to the parameter by excluding outputs (which will not have an instance handle)
+	TPreAnimatedStateTaskParams<FObjectKey, FName> Params;
+	Params.AdditionalFilter.Reset();
+	Params.AdditionalFilter.None({ BuiltInComponents->BlendChannelOutput });
+	Params.AdditionalFilter.All({ BuiltInComponents->Tags.NeedsLink, TracksComponents->Tags.CustomPrimitiveData });
+
+	ScalarParameterStorage->BeginTrackingAndCachePreAnimatedValuesTask(Linker, Params, BuiltInComponents->BoundObject, TracksComponents->ScalarParameterName);
+
+}
+
+void UMovieSceneCustomPrimitiveDataSystem::OnSchedulePersistentTasks(UE::MovieScene::IEntitySystemScheduler* TaskScheduler)
+{
+	using namespace UE::MovieScene;
+
+	FBuiltInComponentTypes* BuiltInComponents = FBuiltInComponentTypes::Get();
+	FMovieSceneTracksComponentTypes* TracksComponents = FMovieSceneTracksComponentTypes::Get();
+
+	FEntityTaskBuilder()
+		.Read(BuiltInComponents->BoundObject)
+		.Read(TracksComponents->ScalarParameterName)
+		.Read(BuiltInComponents->DoubleResult[0])
+		.FilterAll({ TracksComponents->Tags.CustomPrimitiveData })
+		.FilterNone({ BuiltInComponents->BlendChannelInput })
+		.SetDesiredThread(Linker->EntityManager.GetDispatchThread())
+		.Fork_PerAllocation<FApplyCustomPrimitiveDataParameters>(&Linker->EntityManager, TaskScheduler);
+}
+
+void UMovieSceneCustomPrimitiveDataSystem::OnEvaluation(FSystemTaskPrerequisites& InPrerequisites, FSystemSubsequentTasks& Subsequents)
+{
+	using namespace UE::MovieScene;
+
+	FBuiltInComponentTypes* BuiltInComponents = FBuiltInComponentTypes::Get();
+	FMovieSceneTracksComponentTypes* TracksComponents = FMovieSceneTracksComponentTypes::Get();
+
+	if (Linker->EntityManager.ContainsComponent(TracksComponents->ScalarParameterName))
+	{
+		FEntityTaskBuilder()
+			.Read(BuiltInComponents->BoundObject)
+			.Read(TracksComponents->ScalarParameterName)
+			.Read(BuiltInComponents->DoubleResult[0])
+			.FilterAll({ TracksComponents->Tags.CustomPrimitiveData })
+			.FilterNone({ BuiltInComponents->BlendChannelInput })
+			.SetDesiredThread(Linker->EntityManager.GetDispatchThread())
+			.Dispatch_PerAllocation<FApplyCustomPrimitiveDataParameters>(&Linker->EntityManager, InPrerequisites, &Subsequents);
+	}
+}
+
diff --git a/Engine/Source/Runtime/MovieSceneTracks/Private/Tracks/MovieSceneCustomPrimitiveDataTrack.cpp b/Engine/Source/Runtime/MovieSceneTracks/Private/Tracks/MovieSceneCustomPrimitiveDataTrack.cpp
new file mode 100644
--- /dev/null
+++ b/Engine/Source/Runtime/MovieSceneTracks/Private/Tracks/MovieSceneCustomPrimitiveDataTrack.cpp
@@ -0,0 +1,384 @@
+// Copyright Epic Games, Inc. All Rights Reserved.
+
+#include "Tracks/MovieSceneCustomPrimitiveDataTrack.h"
+#include "MovieScene.h"
+#include "MovieSceneCommonHelpers.h"
+#include "Evaluation/MovieSceneEvaluationField.h"
+#include "EntitySystem/MovieSceneEntityBuilder.h"
+#include "EntitySystem/BuiltInComponentTypes.h"
+#include "MovieSceneTracksComponentTypes.h"
+#include "Sections/MovieSceneCustomPrimitiveDataSection.h"
+#include "MaterialTypes.h"
+#include "Materials/MaterialInterface.h"
+#include "Components/PrimitiveComponent.h"
+#include "Components/MeshComponent.h"
+
+#if WITH_EDITORONLY_DATA
+#include "IMovieScenePlayer.h"
+#endif
+
+#include UE_INLINE_GENERATED_CPP_BY_NAME(MovieSceneCustomPrimitiveDataTrack)
+#define LOCTEXT_NAMESPACE "CustomPrimitiveDataTrack"
+
+UMovieSceneCustomPrimitiveDataTrack::UMovieSceneCustomPrimitiveDataTrack(const FObjectInitializer& ObjectInitializer)
+	: Super(ObjectInitializer)
+{
+#if WITH_EDITORONLY_DATA
+	TrackTint = FColor(64, 192, 64, 65);
+#endif
+
+	BuiltInTreePopulationMode = ETreePopulationMode::Blended;
+
+	SupportedBlendTypes.Add(EMovieSceneBlendType::Absolute);
+	SupportedBlendTypes.Add(EMovieSceneBlendType::Additive);
+	SupportedBlendTypes.Add(EMovieSceneBlendType::AdditiveFromBase);
+}
+
+
+bool UMovieSceneCustomPrimitiveDataTrack::SupportsType(TSubclassOf<UMovieSceneSection> SectionClass) const
+{
+	return SectionClass == UMovieSceneCustomPrimitiveDataSection::StaticClass();
+}
+
+
+UMovieSceneSection* UMovieSceneCustomPrimitiveDataTrack::CreateNewSection()
+{
+	UMovieSceneSection* NewSection = NewObject<UMovieSceneCustomPrimitiveDataSection>(this, NAME_None, RF_Transactional);
+	NewSection->SetBlendType(EMovieSceneBlendType::Absolute);
+	return NewSection;
+}
+
+
+void UMovieSceneCustomPrimitiveDataTrack::RemoveAllAnimationData()
+{
+	Sections.Empty();
+	SectionToKey = nullptr;
+}
+
+
+bool UMovieSceneCustomPrimitiveDataTrack::HasSection(const UMovieSceneSection& Section) const
+{
+	return Sections.Contains(&Section);
+}
+
+
+void UMovieSceneCustomPrimitiveDataTrack::AddSection(UMovieSceneSection& Section)
+{
+	Sections.Add(&Section);
+
+	if (Sections.Num() > 1)
+	{
+		SetSectionToKey(&Section);
+	}
+}
+
+
+void UMovieSceneCustomPrimitiveDataTrack::RemoveSection(UMovieSceneSection& Section)
+{
+	Sections.Remove(&Section);
+	if (SectionToKey == &Section)
+	{
+		if (Sections.Num() > 0)
+		{
+			SectionToKey = Sections[0];
+		}
+		else
+		{
+			SectionToKey = nullptr;
+		}
+	}
+}
+
+void UMovieSceneCustomPrimitiveDataTrack::RemoveSectionAt(int32 SectionIndex)
+{
+	bool bResetSectionToKey = (SectionToKey == Sections[SectionIndex]);
+
+	Sections.RemoveAt(SectionIndex);
+
+	if (bResetSectionToKey)
+	{
+		SectionToKey = Sections.Num() > 0 ? Sections[0] : nullptr;
+	}
+}
+
+
+bool UMovieSceneCustomPrimitiveDataTrack::IsEmpty() const
+{
+	return Sections.Num() == 0;
+}
+
+bool UMovieSceneCustomPrimitiveDataTrack::SupportsMultipleRows() const
+{
+	return true;
+}
+
+void UMovieSceneCustomPrimitiveDataTrack::SetSectionToKey(UMovieSceneSection* InSection)
+{
+	SectionToKey = InSection;
+}
+
+UMovieSceneSection* UMovieSceneCustomPrimitiveDataTrack::GetSectionToKey() const
+{
+	return SectionToKey;
+}
+
+const TArray<UMovieSceneSection*>& UMovieSceneCustomPrimitiveDataTrack::GetAllSections() const
+{
+	return Sections;
+}
+
+
+void UMovieSceneCustomPrimitiveDataTrack::AddScalarParameterKey(uint8 CustomPrimitiveDataStartIndex, FFrameNumber Time, float Value)
+{
+	AddScalarParameterKey(CustomPrimitiveDataStartIndex, Time, INDEX_NONE, Value);
+}
+
+void UMovieSceneCustomPrimitiveDataTrack::AddScalarParameterKey(uint8 CustomPrimitiveDataStartIndex, FFrameNumber Time, int32 RowIndex, float Value)
+{
+	UMovieSceneSection* NearestSection = SectionToKey;
+	if (NearestSection == nullptr || (RowIndex != INDEX_NONE && NearestSection->GetRowIndex() != RowIndex))
+	{
+		NearestSection = MovieSceneHelpers::FindNearestSectionAtTime(Sections, Time, RowIndex);
+	}
+	if (NearestSection == nullptr)
+	{
+		NearestSection = Cast<UMovieSceneCustomPrimitiveDataSection>(CreateNewSection());
+
+		UMovieScene* MovieScene = GetTypedOuter<UMovieScene>();
+		check(MovieScene);
+
+		NearestSection->SetRange(MovieScene->GetPlaybackRange());
+		Sections.Add(NearestSection);
+	}
+	if (NearestSection != nullptr && NearestSection->TryModify())
+	{
+		if (UMovieSceneCustomPrimitiveDataSection* NearestCustomPrimitiveDataSection = Cast<UMovieSceneCustomPrimitiveDataSection>(NearestSection))
+		{
+			NearestCustomPrimitiveDataSection->AddScalarParameterKey(*FString::FromInt(CustomPrimitiveDataStartIndex), Time, Value);
+		}
+	}
+}
+
+void UMovieSceneCustomPrimitiveDataTrack::AddVector2DParameterKey(uint8 CustomPrimitiveDataStartIndex, FFrameNumber Time, FVector2D Value)
+{
+	AddVector2DParameterKey(CustomPrimitiveDataStartIndex, Time, INDEX_NONE, Value);
+}
+
+void UMovieSceneCustomPrimitiveDataTrack::AddVector2DParameterKey(uint8 CustomPrimitiveDataStartIndex, FFrameNumber Time, int32 RowIndex, FVector2D Value)
+{
+	UMovieSceneSection* NearestSection = SectionToKey;
+	if (NearestSection == nullptr || (RowIndex != INDEX_NONE && NearestSection->GetRowIndex() != RowIndex))
+	{
+		NearestSection = MovieSceneHelpers::FindNearestSectionAtTime(Sections, Time, RowIndex);
+	}
+	if (NearestSection == nullptr)
+	{
+		NearestSection = Cast<UMovieSceneCustomPrimitiveDataSection>(CreateNewSection());
+
+		UMovieScene* MovieScene = GetTypedOuter<UMovieScene>();
+		check(MovieScene);
+
+		NearestSection->SetRange(MovieScene->GetPlaybackRange());
+		Sections.Add(NearestSection);
+	}
+	if (NearestSection != nullptr && NearestSection->TryModify())
+	{
+		if (UMovieSceneCustomPrimitiveDataSection* NearestCustomPrimitiveDataSection = Cast<UMovieSceneCustomPrimitiveDataSection>(NearestSection))
+		{
+			NearestCustomPrimitiveDataSection->AddVector2DParameterKey(*FString::FromInt(CustomPrimitiveDataStartIndex), Time, Value);
+		}
+	}
+}
+
+void UMovieSceneCustomPrimitiveDataTrack::AddVectorParameterKey(uint8 CustomPrimitiveDataStartIndex, FFrameNumber Time, FVector Value)
+{
+	AddVectorParameterKey(CustomPrimitiveDataStartIndex, Time, INDEX_NONE, Value);
+}
+
+void UMovieSceneCustomPrimitiveDataTrack::AddVectorParameterKey(uint8 CustomPrimitiveDataStartIndex, FFrameNumber Time, int32 RowIndex, FVector Value)
+{
+	UMovieSceneSection* NearestSection = SectionToKey;
+	if (NearestSection == nullptr || (RowIndex != INDEX_NONE && NearestSection->GetRowIndex() != RowIndex))
+	{
+		NearestSection = MovieSceneHelpers::FindNearestSectionAtTime(Sections, Time, RowIndex);
+	}
+	if (NearestSection == nullptr)
+	{
+		NearestSection = Cast<UMovieSceneCustomPrimitiveDataSection>(CreateNewSection());
+
+		UMovieScene* MovieScene = GetTypedOuter<UMovieScene>();
+		check(MovieScene);
+
+		NearestSection->SetRange(MovieScene->GetPlaybackRange());
+		Sections.Add(NearestSection);
+	}
+	if (NearestSection != nullptr && NearestSection->TryModify())
+	{
+		if (UMovieSceneCustomPrimitiveDataSection* NearestCustomPrimitiveDataSection = Cast<UMovieSceneCustomPrimitiveDataSection>(NearestSection))
+		{
+			NearestCustomPrimitiveDataSection->AddVectorParameterKey(*FString::FromInt(CustomPrimitiveDataStartIndex), Time, Value);
+		}
+	}
+}
+
+void UMovieSceneCustomPrimitiveDataTrack::AddColorParameterKey(uint8 CustomPrimitiveDataStartIndex, FFrameNumber Time, FLinearColor Value)
+{
+	AddColorParameterKey(CustomPrimitiveDataStartIndex, Time, INDEX_NONE, Value);
+}
+
+void UMovieSceneCustomPrimitiveDataTrack::AddColorParameterKey(uint8 CustomPrimitiveDataStartIndex, FFrameNumber Time, int32 RowIndex, FLinearColor Value)
+{
+	UMovieSceneSection* NearestSection = SectionToKey;
+	if (NearestSection == nullptr || (RowIndex != INDEX_NONE && NearestSection->GetRowIndex() != RowIndex))
+	{
+		NearestSection = MovieSceneHelpers::FindNearestSectionAtTime(Sections, Time, RowIndex);
+	}
+	if (NearestSection == nullptr)
+	{
+		NearestSection = Cast<UMovieSceneCustomPrimitiveDataSection>(CreateNewSection());
+
+		UMovieScene* MovieScene = GetTypedOuter<UMovieScene>();
+		check(MovieScene);
+
+		NearestSection->SetRange(MovieScene->GetPlaybackRange());
+		Sections.Add(NearestSection);
+	}
+	if (NearestSection != nullptr && NearestSection->TryModify())
+	{
+		if (UMovieSceneCustomPrimitiveDataSection* NearestCustomPrimitiveDataSection = Cast<UMovieSceneCustomPrimitiveDataSection>(NearestSection))
+		{
+			NearestCustomPrimitiveDataSection->AddColorParameterKey(*FString::FromInt(CustomPrimitiveDataStartIndex), Time, Value);
+		}
+	}
+}
+
+
+void UMovieSceneCustomPrimitiveDataTrack::ImportEntityImpl(UMovieSceneEntitySystemLinker* EntityLinker, const FEntityImportParams& Params, FImportedEntity* OutImportedEntity)
+{
+	// These tracks don't define any entities for themselves
+	checkf(false, TEXT("This track should never have created entities for itself - this assertion indicates an error in the entity-component field"));
+}
+
+void UMovieSceneCustomPrimitiveDataTrack::ExtendEntityImpl(UMovieSceneParameterSection* Section, UMovieSceneEntitySystemLinker* EntityLinker, const UE::MovieScene::FEntityImportParams& Params, UE::MovieScene::FImportedEntity* OutImportedEntity)
+{
+	using namespace UE::MovieScene;
+
+	FBuiltInComponentTypes* BuiltInComponents = FBuiltInComponentTypes::Get();
+	FMovieSceneTracksComponentTypes* TracksComponents = FMovieSceneTracksComponentTypes::Get();
+
+	// Custom primitive data parameters are always absolute blends for the time being
+	OutImportedEntity->AddBuilder(
+		FEntityBuilder()
+		.AddTag(TracksComponents->Tags.CustomPrimitiveData)
+		// If the section has no valid blend type (legacy data), make it use absolute blending.
+		// Otherwise, the base section class will add the appropriate blend type tag in BuildDefaultComponents.
+		.AddTagConditional(BuiltInComponents->Tags.AbsoluteBlend, !Section->GetBlendType().IsValid())
+	);
+}
+
+bool UMovieSceneCustomPrimitiveDataTrack::PopulateEvaluationFieldImpl(const TRange<FFrameNumber>& EffectiveRange, const FMovieSceneEvaluationFieldEntityMetaData& InMetaData, FMovieSceneEntityComponentFieldBuilder* OutFieldBuilder)
+{
+	const FMovieSceneTrackEvaluationField& LocalEvaluationField = GetEvaluationField();
+
+	for (const FMovieSceneTrackEvaluationFieldEntry& Entry : LocalEvaluationField.Entries)
+	{
+		UMovieSceneCustomPrimitiveDataSection* CPDSection = Cast<UMovieSceneCustomPrimitiveDataSection>(Entry.Section);
+		if (CPDSection)
+		{
+			if (IsRowEvalDisabled(Entry.Section->GetRowIndex()))
+			{
+				continue;
+			}
+
+			TRange<FFrameNumber> SectionEffectiveRange = TRange<FFrameNumber>::Intersection(EffectiveRange, Entry.Range);
+			if (!SectionEffectiveRange.IsEmpty())
+			{
+				FMovieSceneEvaluationFieldEntityMetaData SectionMetaData = InMetaData;
+				SectionMetaData.Flags = Entry.Flags;
+				CPDSection->ExternalPopulateEvaluationField(SectionEffectiveRange, SectionMetaData, OutFieldBuilder);
+			}
+		}
+	}
+
+	return true;
+}
+
+
+#if WITH_EDITORONLY_DATA
+
+void UMovieSceneCustomPrimitiveDataTrack::GetCPDMaterialData(IMovieScenePlayer& Player, FGuid BoundObjectId, FMovieSceneSequenceID SequenceID, TSortedMap<uint8, TArray<FCustomPrimitiveDataMaterialParametersData>>& OutCPDMaterialData)
+{
+	// Empty the map
+	OutCPDMaterialData.Empty();
+
+	auto FindMaterialData = [this, &OutCPDMaterialData](EMaterialParameterType ParameterType, UMaterialInterface* MaterialInterface, const FComponentMaterialInfo& MaterialInfo)
+	{
+		TMap<FMaterialParameterInfo, FMaterialParameterMetadata> Parameters;
+
+		MaterialInterface->GetAllParametersOfType(ParameterType, Parameters);
+
+		for (const TPair<FMaterialParameterInfo, FMaterialParameterMetadata>& Parameter : Parameters)
+		{
+			const FMaterialParameterInfo& ParameterInfo = Parameter.Key;
+			const FMaterialParameterMetadata& ParameterMetadata = Parameter.Value;
+
+			if (ParameterMetadata.PrimitiveDataIndex > INDEX_NONE)
+			{
+				OutCPDMaterialData.FindOrAdd(ParameterMetadata.PrimitiveDataIndex).Add(
+					{
+						ParameterType,
+						MaterialInfo,
+						ParameterInfo,
+						MaterialInterface
+					});
+			}
+		}
+	};
+
+	// Find the bound object
+	for (TWeakObjectPtr<> WeakObject : Player.FindBoundObjects(BoundObjectId, SequenceID))
+	{
+		if (USceneComponent* SceneComponent = Cast<USceneComponent>(WeakObject.Get()))
+		{
+			// Find each type of component with materials we can, and search for parameters of scalar and vector types that specify PrimitiveDataIndex
+			if (UPrimitiveComponent* PrimitiveComponent = Cast<UPrimitiveComponent>(SceneComponent))
+			{
+				int32 NumMaterials = PrimitiveComponent->GetNumMaterials();
+				TArray<FName> MaterialSlotNames = PrimitiveComponent->GetMaterialSlotNames();
+				UMeshComponent* MeshComponent = Cast<UMeshComponent>(SceneComponent);
+				if (NumMaterials > 0 || MeshComponent)
+				{
+					for (int32 MaterialIndex = 0; MaterialIndex < NumMaterials; MaterialIndex++)
+					{
+						if (UMaterialInterface* IndexedMaterial = PrimitiveComponent->GetMaterial(MaterialIndex))
+						{
+							FName MaterialSlotName = MaterialSlotNames.IsValidIndex(MaterialIndex) ? MaterialSlotNames[MaterialIndex] : FName();
+							FComponentMaterialInfo MaterialInfo{ MaterialSlotName, MaterialIndex, EComponentMaterialType::IndexedMaterial };
+							FindMaterialData(EMaterialParameterType::Scalar, IndexedMaterial, MaterialInfo);
+							FindMaterialData(EMaterialParameterType::Vector, IndexedMaterial, MaterialInfo);
+						}
+					}
+					if (MeshComponent)
+					{
+						if (UMaterialInterface* OverlayMaterial = MeshComponent->GetOverlayMaterial())
+						{
+							FComponentMaterialInfo MaterialInfo{ FName(), 0, EComponentMaterialType::OverlayMaterial };
+							FindMaterialData(EMaterialParameterType::Scalar, OverlayMaterial, MaterialInfo);
+							FindMaterialData(EMaterialParameterType::Vector, OverlayMaterial, MaterialInfo);
+						}
+					}
+				}
+			}
+		}
+	}
+}
+
+
+FText UMovieSceneCustomPrimitiveDataTrack::GetDefaultDisplayName() const
+{
+	return LOCTEXT("CustomPrimitiveDataTrackName", "Custom Primitive Data");
+}
+
+#endif
+
+#undef LOCTEXT_NAMESPACE
\ No newline at end of file
diff --git a/Engine/Source/Runtime/MovieSceneTracks/Public/MovieSceneTracksComponentTypes.h b/Engine/Source/Runtime/MovieSceneTracks/Public/MovieSceneTracksComponentTypes.h
--- a/Engine/Source/Runtime/MovieSceneTracks/Public/MovieSceneTracksComponentTypes.h
+++ b/Engine/Source/Runtime/MovieSceneTracks/Public/MovieSceneTracksComponentTypes.h
@@ -597,6 +597,7 @@ struct FMovieSceneTracksComponentTypes
 	struct
 	{
 		FComponentTypeID BoundMaterialChanged;
+		FComponentTypeID CustomPrimitiveData;
 	} Tags;
 
 	TComponentTypeID<FLevelVisibilityComponentData> LevelVisibility;
diff --git a/Engine/Source/Runtime/MovieSceneTracks/Public/Sections/MovieSceneCustomPrimitiveDataSection.h b/Engine/Source/Runtime/MovieSceneTracks/Public/Sections/MovieSceneCustomPrimitiveDataSection.h
new file mode 100644
--- /dev/null
+++ b/Engine/Source/Runtime/MovieSceneTracks/Public/Sections/MovieSceneCustomPrimitiveDataSection.h
@@ -0,0 +1,42 @@
+// Copyright Epic Games, Inc. All Rights Reserved.
+
+#pragma once
+
+#include "CoreMinimal.h"
+#include "Sections/MovieSceneParameterSection.h"
+
+#include "MovieSceneCustomPrimitiveDataSection.generated.h"
+
+
+/**
+ * Movie scene section that animates custom primitive data scalars, vector2s, vector3s, and colors
+ */
+UCLASS(MinimalAPI)
+class UMovieSceneCustomPrimitiveDataSection : public UMovieSceneParameterSection
+{
+	GENERATED_BODY()
+
+public:
+
+	MOVIESCENETRACKS_API virtual void ImportEntityImpl(UMovieSceneEntitySystemLinker* EntityLinker, const FEntityImportParams& Params, FImportedEntity* OutImportedEntity) override;
+
+	MOVIESCENETRACKS_API void ExternalPopulateEvaluationField(const TRange<FFrameNumber>& EffectiveRange, const FMovieSceneEvaluationFieldEntityMetaData& InMetaData, FMovieSceneEntityComponentFieldBuilder* OutFieldBuilder);
+
+
+	MOVIESCENETRACKS_API void ReconstructChannelProxy() override;
+
+private:
+
+#if WITH_EDITORONLY_DATA
+
+public:
+	uint64 GetChannelsUsed() const { return ChannelsUsedBitmap; }
+
+protected:
+
+	UPROPERTY()
+	uint64 ChannelsUsedBitmap = 0;
+
+#endif
+
+};
\ No newline at end of file
diff --git a/Engine/Source/Runtime/MovieSceneTracks/Public/Systems/MovieSceneCustomPrimitiveDataSystem.h b/Engine/Source/Runtime/MovieSceneTracks/Public/Systems/MovieSceneCustomPrimitiveDataSystem.h
new file mode 100644
--- /dev/null
+++ b/Engine/Source/Runtime/MovieSceneTracks/Public/Systems/MovieSceneCustomPrimitiveDataSystem.h
@@ -0,0 +1,79 @@
+// Copyright Epic Games, Inc. All Rights Reserved.
+
+#pragma once
+
+#include "EntitySystem/MovieSceneEntitySystem.h"
+#include "EntitySystem/MovieSceneOverlappingEntityTracker.h"
+#include "Evaluation/PreAnimatedState/MovieScenePreAnimatedStateStorage.h"
+#include "EntitySystem/BuiltInComponentTypes.h"
+
+#include "MovieSceneCustomPrimitiveDataSystem.generated.h"
+
+class UMovieSceneBlenderSystem;
+class UMovieScenePiecewiseDoubleBlenderSystem;
+
+namespace UE::MovieScene
+{
+
+struct FPreAnimatedCustomPrimitiveDataEntryStorage;
+
+struct FAnimatedCustomPrimitiveDataInfo
+{
+	/** Weak linker ptr - only assigned if the output entity is ever allocated */
+	TWeakObjectPtr<UMovieSceneEntitySystemLinker> WeakLinker;
+	/** Weak blender system ptr - only assigned if the blend channel is ever allocated */
+	TWeakObjectPtr<UMovieSceneBlenderSystem> WeakBlenderSystem;
+	int32 NumContributors = 0;
+	FMovieSceneEntityID OutputEntityID;
+	FMovieSceneBlendChannelID BlendChannelID;
+
+	~FAnimatedCustomPrimitiveDataInfo();
+};
+
+MOVIESCENETRACKS_API void CollectGarbageForOutput(FAnimatedCustomPrimitiveDataInfo* Output);
+
+} // namespace UE::MovieScene
+
+
+/**
+ * System responsible for tracking and animating custom primitive data entities.
+ * Operates on the following component types from FMovieSceneTracksComponentTypes:
+ *
+ * Instantiation: Tracks any ScalarParameterName with a CustomPrimitiveData tag on a BoundObject where that BoundObject is a Primitive Component.
+ *                Manages adding BlendChannelInputs and Outputs where multiple entities animate the same custom primitive data index.
+ *
+ * Evaluation:    Visits any BoundObject with the supported parameter name and either a BlendChannelOutput component
+ *                or no BlendChannelInput, and applies the resulting parameter to the Custom Primitive Data on the Primitive Component.
+ */
+UCLASS(MinimalAPI)
+class UMovieSceneCustomPrimitiveDataSystem
+	: public UMovieSceneEntitySystem
+{
+public:
+
+	GENERATED_BODY()
+
+	UMovieSceneCustomPrimitiveDataSystem(const FObjectInitializer& ObjInit);
+
+private:
+
+	virtual void OnLink() override;
+	virtual void OnUnlink() override;
+	virtual void OnSchedulePersistentTasks(UE::MovieScene::IEntitySystemScheduler* TaskScheduler) override;
+	virtual void OnRun(FSystemTaskPrerequisites& InPrerequisites, FSystemSubsequentTasks& Subsequents) override;
+
+	void OnInstantiation();
+	void OnEvaluation(FSystemTaskPrerequisites& InPrerequisites, FSystemSubsequentTasks& Subsequents);
+
+private:
+
+	/** Overlapping trackers that track multiple entities animating the same bound object and name */
+	UE::MovieScene::TOverlappingEntityTracker<UE::MovieScene::FAnimatedCustomPrimitiveDataInfo, FObjectKey, FName> ScalarParameterTracker;
+
+	/** Holds pre-animated values for custom primitive data entries */
+	TSharedPtr<UE::MovieScene::FPreAnimatedCustomPrimitiveDataEntryStorage> ScalarParameterStorage;
+public:
+
+	UPROPERTY()
+	TObjectPtr<UMovieScenePiecewiseDoubleBlenderSystem> DoubleBlenderSystem;
+};
diff --git a/Engine/Source/Runtime/MovieSceneTracks/Public/Tracks/MovieSceneCustomPrimitiveDataTrack.h b/Engine/Source/Runtime/MovieSceneTracks/Public/Tracks/MovieSceneCustomPrimitiveDataTrack.h
new file mode 100644
--- /dev/null
+++ b/Engine/Source/Runtime/MovieSceneTracks/Public/Tracks/MovieSceneCustomPrimitiveDataTrack.h
@@ -0,0 +1,165 @@
+// Copyright Epic Games, Inc. All Rights Reserved.
+
+#pragma once
+
+#include "CoreMinimal.h"
+#include "UObject/ObjectMacros.h"
+#include "Sections/MovieSceneParameterSection.h"
+#include "MovieSceneNameableTrack.h"
+#include "EntitySystem/IMovieSceneEntityProvider.h"
+#include "MaterialTypes.h"
+#include "MovieSceneMaterialTrack.h"
+#include "MovieSceneCustomPrimitiveDataTrack.generated.h"
+
+
+#if WITH_EDITORONLY_DATA
+/**
+ * Used to cache mapping between custom primitive data indices and found material parameters using those indices.
+ * This is used for nicer display names on channels animating those parameters.
+*/
+struct FCustomPrimitiveDataMaterialParametersData
+{
+	EMaterialParameterType MaterialParameterType;
+	FComponentMaterialInfo MaterialInfo;
+	FMaterialParameterInfo ParameterInfo;
+	TWeakObjectPtr<UMaterialInterface> MaterialAsset;
+};
+#endif
+
+
+/**
+ * Handles manipulation of custom primitive data in a movie scene.
+ */
+UCLASS(MinimalAPI)
+class UMovieSceneCustomPrimitiveDataTrack
+	: public UMovieSceneNameableTrack
+	, public IMovieSceneEntityProvider
+	, public IMovieSceneParameterSectionExtender
+{
+	GENERATED_BODY()
+
+public:
+
+	MOVIESCENETRACKS_API UMovieSceneCustomPrimitiveDataTrack(const FObjectInitializer& ObjectInitializer);
+
+	// UMovieSceneTrack interface
+
+	MOVIESCENETRACKS_API virtual bool SupportsType(TSubclassOf<UMovieSceneSection> SectionClass) const override;
+	MOVIESCENETRACKS_API virtual UMovieSceneSection* CreateNewSection() override;
+	MOVIESCENETRACKS_API virtual void RemoveAllAnimationData() override;
+	MOVIESCENETRACKS_API virtual bool HasSection(const UMovieSceneSection& Section) const override;
+	MOVIESCENETRACKS_API virtual void AddSection(UMovieSceneSection& Section) override;
+	MOVIESCENETRACKS_API virtual void RemoveSection(UMovieSceneSection& Section) override;
+	MOVIESCENETRACKS_API virtual void RemoveSectionAt(int32 SectionIndex) override;
+	MOVIESCENETRACKS_API virtual bool IsEmpty() const override;
+	MOVIESCENETRACKS_API virtual const TArray<UMovieSceneSection*>& GetAllSections() const override;	
+	MOVIESCENETRACKS_API virtual bool SupportsMultipleRows() const override;
+	MOVIESCENETRACKS_API virtual void SetSectionToKey(UMovieSceneSection* Section) override;
+	MOVIESCENETRACKS_API virtual UMovieSceneSection* GetSectionToKey() const override;
+
+public:
+
+	/**
+	* Adds a scalar parameter key to the track. 
+	* @param CustomPrimitiveDataStartIndex The start index for the custom primitive data
+	* @param Time The time to add the new key.
+	* @param The value for the new key.
+	*/
+	MOVIESCENETRACKS_API void AddScalarParameterKey(uint8 CustomPrimitiveDataStartIndex, FFrameNumber Position, float Value);
+
+	/**
+	 * Adds a scalar parameter key to the track. 
+	* @param CustomPrimitiveDataStartIndex The start index for the custom primitive data
+	 * @param Time The time to add the new key.
+	 * @param RowIndex The preferred row index on which to look for sections.
+	 * @param The value for the new key.
+	 */
+	MOVIESCENETRACKS_API void AddScalarParameterKey(uint8 CustomPrimitiveDataStartIndex, FFrameNumber Position, int32 RowIndex, float Value);
+	
+	/**
+	 * Adds a Vector2D parameter key to the track.
+	 * @param CustomPrimitiveDataStartIndex The start index for the custom primitive data
+	 * @param Time The time to add the new key.
+	 * @param The value for the new key.
+	 */
+	MOVIESCENETRACKS_API void AddVector2DParameterKey(uint8 CustomPrimitiveDataStartIndex, FFrameNumber Position, FVector2D Value);
+
+	/**
+	 * Adds a Vector2D parameter key to the track.
+	 * @param CustomPrimitiveDataStartIndex The start index for the custom primitive data
+	 * @param Time The time to add the new key.
+	 * @param RowIndex The preferred row index on which to look for sections.
+	 * @param The value for the new key.
+	 */
+	MOVIESCENETRACKS_API void AddVector2DParameterKey(uint8 CustomPrimitiveDataStartIndex, FFrameNumber Position, int32 RowIndex, FVector2D Value);
+
+	/**
+	 * Adds a Vector parameter key to the track.
+	 * @param CustomPrimitiveDataStartIndex The start index for the custom primitive data
+	 * @param Time The time to add the new key.
+	 * @param The value for the new key.
+	 */
+	MOVIESCENETRACKS_API void AddVectorParameterKey(uint8 CustomPrimitiveDataStartIndex, FFrameNumber Position, FVector Value);
+
+	/**
+	 * Adds a Vector parameter key to the track.
+	 * @param CustomPrimitiveDataStartIndex The start index for the custom primitive data
+	 * @param Time The time to add the new key.
+	 * @param RowIndex The preferred row index on which to look for sections.
+	 * @param The value for the new key.
+	 */
+	MOVIESCENETRACKS_API void AddVectorParameterKey(uint8 CustomPrimitiveDataStartIndex, FFrameNumber Position, int32 RowIndex, FVector Value);
+
+	/**
+	* Adds a Vector2D parameter key to the track.
+	* @param CustomPrimitiveDataStartIndex The start index for the custom primitive data
+	* @param Time The time to add the new key.
+	* @param The value for the new key.
+	*/
+	MOVIESCENETRACKS_API void AddColorParameterKey(uint8 CustomPrimitiveDataStartIndex, FFrameNumber Position, FLinearColor Value);
+
+	/**
+	* Adds a Vector2D parameter key to the track.
+	* @param CustomPrimitiveDataStartIndex The start index for the custom primitive data
+	* @param Time The time to add the new key.
+	* @param RowIndex The preferred row index on which to look for sections.
+	* @param The value for the new key.
+	*/
+	MOVIESCENETRACKS_API void AddColorParameterKey(uint8 CustomPrimitiveDataStartIndex, FFrameNumber Position, int32 RowIndex, FLinearColor Value);
+
+	/*~ IMovieSceneEntityProvider */
+	virtual void ImportEntityImpl(UMovieSceneEntitySystemLinker* EntityLinker, const FEntityImportParams& Params, FImportedEntity* OutImportedEntity) override;
+	virtual bool PopulateEvaluationFieldImpl(const TRange<FFrameNumber>& EffectiveRange, const FMovieSceneEvaluationFieldEntityMetaData& InMetaData, FMovieSceneEntityComponentFieldBuilder* OutFieldBuilder) override;
+
+	/*~ IMovieSceneParameterSectionExtender */
+	virtual void ExtendEntityImpl(UMovieSceneParameterSection* Section, UMovieSceneEntitySystemLinker* EntityLinker, const UE::MovieScene::FEntityImportParams& Params, UE::MovieScene::FImportedEntity* OutImportedEntity) override;
+
+#if WITH_EDITORONLY_DATA
+
+	/*
+	* Returns the mappings from custom primitive data indices to any material parameters on the bound object.
+	* @param Player Movie scene player
+	* @param BoundObjectID Bound object ID for the component this track is on
+	* @param SequenceID The sequence ID this track is on
+	*/
+	MOVIESCENETRACKS_API void GetCPDMaterialData(IMovieScenePlayer& Player, FGuid BoundObjectId, FMovieSceneSequenceID SequenceID, TSortedMap<uint8, TArray<FCustomPrimitiveDataMaterialParametersData>>& OutCPDMaterialData);
+
+
+public:
+
+	FText GetDefaultDisplayName() const override;
+
+private:
+
+#endif
+
+private:
+
+	/** The sections owned by this track .*/
+	UPROPERTY()
+	TArray<TObjectPtr<UMovieSceneSection>> Sections;
+
+	/** Section we should Key */
+	UPROPERTY()
+	TObjectPtr<UMovieSceneSection> SectionToKey;
+};
\ No newline at end of file
