Sequencer: Parse takes by asset name rather than shot display name. This allows you to rename the shot label and the takes which are name based, will stay intact.

This also fixed a bug where when you change takes, the shot label doesn't come along with it.

From here on out, the shot display name is just for display. All the take operations are based on the asset name

diff --git a/Engine/Source/Editor/MovieSceneTools/Private/MovieSceneToolHelpers.cpp b/Engine/Source/Editor/MovieSceneTools/Private/MovieSceneToolHelpers.cpp
--- a/Engine/Source/Editor/MovieSceneTools/Private/MovieSceneToolHelpers.cpp
+++ b/Engine/Source/Editor/MovieSceneTools/Private/MovieSceneToolHelpers.cpp
@@ -499,16 +499,7 @@ FString MovieSceneToolHelpers::GenerateNewSubsequenceName(const TArray<UMovieSce
 	}
 
 	FString NextSectionName = NextSection && NextSection->GetSequence() ? NextSection->GetSequence()->GetName() : FString();
-	if (const UMovieSceneCinematicShotSection* ShotSection = Cast<UMovieSceneCinematicShotSection>(NextSection))
-	{
-		NextSectionName = ShotSection->GetShotDisplayName();
-	}
-
 	FString CurrentSectionName = CurrentSection && CurrentSection->GetSequence() ? CurrentSection->GetSequence()->GetName() : FString();
-	if (const UMovieSceneCinematicShotSection* ShotSection = Cast<UMovieSceneCinematicShotSection>(CurrentSection))
-	{
-		CurrentSectionName = ShotSection->GetShotDisplayName();
-	}
 
 	// This is the first or last shot
 	if ((CurrentSection == nullptr && NextSection) || (CurrentSection != nullptr && CurrentSection == NextSection))
@@ -679,13 +670,8 @@ void MovieSceneToolHelpers::GatherTakes(const UMovieSceneSection* Section, TArra
 	uint32 ShotNumberDigits = 0;
 	uint32 TakeNumberDigits = 0;
 
-	FString SubSectionName = SubSection->GetSequence()->GetName();
-	if (const UMovieSceneCinematicShotSection* ShotSection = Cast<UMovieSceneCinematicShotSection>(SubSection))
-	{
-		SubSectionName = ShotSection->GetShotDisplayName();
-	}
-
-	ParseShotName(SubSectionName, ShotPrefix, ShotNumber, OutCurrentTakeNumber, ShotNumberDigits, TakeNumberDigits);
+	FString SequenceName = SubSection->GetSequence()->GetName();
+	ParseShotName(SequenceName, ShotPrefix, ShotNumber, OutCurrentTakeNumber, ShotNumberDigits, TakeNumberDigits);
 
 	// Gather up all level sequence assets
 	FAssetRegistryModule& AssetRegistryModule = FModuleManager::LoadModuleChecked<FAssetRegistryModule>(TEXT("AssetRegistry"));
@@ -734,13 +720,8 @@ bool MovieSceneToolHelpers::GetTakeNumber(const UMovieSceneSection* Section, FAs
 	uint32 ShotNumberDigits = 0;
 	uint32 TakeNumberDigits = 0;
 
-	FString SubSectionName = SubSection->GetSequence()->GetName();
-	if (const UMovieSceneCinematicShotSection* ShotSection = Cast<UMovieSceneCinematicShotSection>(SubSection))
-	{
-		SubSectionName = ShotSection->GetShotDisplayName();
-	}
-
-	ParseShotName(SubSectionName, ShotPrefix, ShotNumber, TakeNumberDummy, ShotNumberDigits, TakeNumberDigits);
+	FString SequenceName = SubSection->GetSequence()->GetName();
+	ParseShotName(SequenceName, ShotPrefix, ShotNumber, TakeNumberDummy, ShotNumberDigits, TakeNumberDigits);
 	
 	// Gather up all level sequence assets
 	FAssetRegistryModule& AssetRegistryModule = FModuleManager::LoadModuleChecked<FAssetRegistryModule>(TEXT("AssetRegistry"));
diff --git a/Engine/Source/Editor/MovieSceneTools/Private/TrackEditors/SubTrackEditor.cpp b/Engine/Source/Editor/MovieSceneTools/Private/TrackEditors/SubTrackEditor.cpp
--- a/Engine/Source/Editor/MovieSceneTools/Private/TrackEditors/SubTrackEditor.cpp
+++ b/Engine/Source/Editor/MovieSceneTools/Private/TrackEditors/SubTrackEditor.cpp
@@ -525,7 +525,9 @@ void FSubTrackEditor::CreateNewTake(UMovieSceneSubSection* Section)
 	uint32 ShotNumberDigits = 0;
 	uint32 TakeNumberDigits = 0;
 	
-	if (MovieSceneToolHelpers::ParseShotName(GetSubSectionDisplayName(Section), ShotPrefix, ShotNumber, TakeNumber, ShotNumberDigits, TakeNumberDigits))
+	FString SequenceName = Section->GetSequence() ? Section->GetSequence()->GetName() : FString();
+
+	if (MovieSceneToolHelpers::ParseShotName(SequenceName, ShotPrefix, ShotNumber, TakeNumber, ShotNumberDigits, TakeNumberDigits))
 	{
 		TArray<FAssetData> AssetData;
 		uint32 CurrentTakeNumber = INDEX_NONE;
@@ -544,39 +546,47 @@ void FSubTrackEditor::CreateNewTake(UMovieSceneSubSection* Section)
 			}
 		}
 
-		FString NewShotName = MovieSceneToolHelpers::ComposeShotName(ShotPrefix, ShotNumber, NewTakeNumber, ShotNumberDigits, TakeNumberDigits);
+		FString NewSectionName = MovieSceneToolHelpers::ComposeShotName(ShotPrefix, ShotNumber, NewTakeNumber, ShotNumberDigits, TakeNumberDigits);
 
-		TRange<FFrameNumber> NewShotRange         = Section->GetRange();
-		FFrameNumber         NewShotStartOffset   = Section->Parameters.StartFrameOffset;
-		float                NewShotTimeScale     = Section->Parameters.TimeScale;
-		int32                NewShotPrerollFrames = Section->GetPreRollFrames();
+		TRange<FFrameNumber> NewSectionRange         = Section->GetRange();
+		FFrameNumber         NewSectionStartOffset   = Section->Parameters.StartFrameOffset;
+		float                NewSectionTimeScale     = Section->Parameters.TimeScale;
+		int32                NewSectionPrerollFrames = Section->GetPreRollFrames();
 		int32                NewRowIndex          = Section->GetRowIndex();
-		FFrameNumber         NewShotStartTime     = NewShotRange.GetLowerBound().IsClosed() ? UE::MovieScene::DiscreteInclusiveLower(NewShotRange) : 0;
-		FColor               NewShotColorTint     = Section->GetColorTint();
+		FFrameNumber         NewSectionStartTime     = NewSectionRange.GetLowerBound().IsClosed() ? UE::MovieScene::DiscreteInclusiveLower(NewSectionRange) : 0;
+		FColor               NewSectionColorTint     = Section->GetColorTint();
 		UMovieSceneSubTrack* SubTrack = CastChecked<UMovieSceneSubTrack>(Section->GetOuter());
 		FString NewSequencePath = FPaths::GetPath(Section->GetSequence()->GetPathName());
 
-		if (UMovieSceneSequence* NewSequence = MovieSceneToolHelpers::CreateSequence(NewShotName, NewSequencePath, Section))
+		if (UMovieSceneSequence* NewSequence = MovieSceneToolHelpers::CreateSequence(NewSectionName, NewSequencePath, Section))
 		{
 			const FScopedTransaction Transaction(LOCTEXT("NewTake_Transaction", "New Take"));
 
 			int32 Duration = UE::MovieScene::DiscreteSize(Section->GetRange());
 
-			UMovieSceneSubSection* NewShot = SubTrack->AddSequence(NewSequence, NewShotStartTime, Duration);
+			UMovieSceneSubSection* NewSection = SubTrack->AddSequence(NewSequence, NewSectionStartTime, Duration);
 			SubTrack->RemoveSection(*Section);
 
-			NewShot->SetRange(NewShotRange);
-			NewShot->Parameters.StartFrameOffset = NewShotStartOffset;
-			NewShot->Parameters.TimeScale = NewShotTimeScale;
-			NewShot->SetPreRollFrames(NewShotPrerollFrames);
-			NewShot->SetRowIndex(NewRowIndex);
-			NewShot->SetColorTint(NewShotColorTint);
+			NewSection->SetRange(NewSectionRange);
+			NewSection->Parameters.StartFrameOffset = NewSectionStartOffset;
+			NewSection->Parameters.TimeScale = NewSectionTimeScale;
+			NewSection->SetPreRollFrames(NewSectionPrerollFrames);
+			NewSection->SetRowIndex(NewRowIndex);
+			NewSection->SetColorTint(NewSectionColorTint);
+
+			UMovieSceneCinematicShotSection* ShotSection = Cast<UMovieSceneCinematicShotSection>(Section);
+			UMovieSceneCinematicShotSection* NewShotSection = Cast<UMovieSceneCinematicShotSection>(NewSection);
 
-			MovieSceneToolHelpers::SetTakeNumber(NewShot, NewTakeNumber);
+			if (ShotSection && NewShotSection)
+			{
+				NewShotSection->SetShotDisplayName(ShotSection->GetShotDisplayName());
+			}
+
+			MovieSceneToolHelpers::SetTakeNumber(NewSection, NewTakeNumber);
 
 			GetSequencer()->NotifyMovieSceneDataChanged( EMovieSceneDataChangeType::MovieSceneStructureItemsChanged );
 			GetSequencer()->EmptySelection();
-			GetSequencer()->SelectSection(NewShot);
+			GetSequencer()->SelectSection(NewSection);
 			GetSequencer()->ThrobSectionSelection();
 		}
 	}
@@ -629,6 +639,14 @@ void FSubTrackEditor::ChangeTake(UMovieSceneSequence* Sequence)
 			NewSection->SetRowIndex(NewSectionRowIndex);
 			NewSection->SetColorTint(NewSectionColorTint);
 
+			UMovieSceneCinematicShotSection* ShotSection = Cast<UMovieSceneCinematicShotSection>(Section);
+			UMovieSceneCinematicShotSection* NewShotSection = Cast<UMovieSceneCinematicShotSection>(NewSection);
+
+			if (ShotSection && NewShotSection)
+			{
+				NewShotSection->SetShotDisplayName(ShotSection->GetShotDisplayName());
+			}
+
 			bChangedTake = true;
 		}
 	}
