Sequencer: Fix duplicate LOC strings

#rnx

diff --git a/Engine/Source/Editor/MovieSceneTools/Private/TrackEditors/CustomPrimitiveDataTrackEditor.cpp b/Engine/Source/Editor/MovieSceneTools/Private/TrackEditors/CustomPrimitiveDataTrackEditor.cpp
--- a/Engine/Source/Editor/MovieSceneTools/Private/TrackEditors/CustomPrimitiveDataTrackEditor.cpp
+++ b/Engine/Source/Editor/MovieSceneTools/Private/TrackEditors/CustomPrimitiveDataTrackEditor.cpp
@@ -241,11 +241,11 @@ void FCustomPrimitiveDataTrackEditor::OnBuildAddParameterMenu(FMenuBuilder& Menu
 				{
 					if (CanAddParameter(CPDTrack, GetStartIndex, 1))
 					{
-						return FText::Format(LOCTEXT("CustomPrimitiveDataMaterialTooltipScalar", "Add a track for a new scalar parameter starting at index {0}"), FText::AsNumber(StartIndex));
+						return FText::Format(LOCTEXT("CustomPrimitiveDataNewMaterialTooltipScalar", "Add a track for a new scalar parameter starting at index {0}"), FText::AsNumber(StartIndex));
 					}
 					else
 					{
-						return FText::Format(LOCTEXT("CustomPrimitiveDataMaterialTooltipScalarInvalid", "Can't add a scalar parameter starting at index {0} as it overlaps with other parameters using that index."), FText::AsNumber(StartIndex));
+						return FText::Format(LOCTEXT("CustomPrimitiveDataNewMaterialTooltipScalarInvalid", "Can't add a scalar parameter starting at index {0} as it overlaps with other parameters using that index."), FText::AsNumber(StartIndex));
 					}
 				}), 
 			FSlateIcon(),
@@ -260,11 +260,11 @@ void FCustomPrimitiveDataTrackEditor::OnBuildAddParameterMenu(FMenuBuilder& Menu
 				{
 					if (CanAddParameter(CPDTrack, GetStartIndex, 2))
 					{
-						return FText::Format(LOCTEXT("CustomPrimitiveDataMaterialTooltipVector2D", "Add a track for a new Vector2D parameter starting at index {0}"), FText::AsNumber(StartIndex));
+						return FText::Format(LOCTEXT("CustomPrimitiveDataNewMaterialTooltipVector2D", "Add a track for a new Vector2D parameter starting at index {0}"), FText::AsNumber(StartIndex));
 					}
 					else
 					{
-						return FText::Format(LOCTEXT("CustomPrimitiveDataMaterialTooltipVector2DInvalid", "Can't add a Vector2D parameter starting at index {0} as it overlaps with other parameters using that index."), FText::AsNumber(StartIndex));
+						return FText::Format(LOCTEXT("CustomPrimitiveDataNewMaterialTooltipVector2DInvalid", "Can't add a Vector2D parameter starting at index {0} as it overlaps with other parameters using that index."), FText::AsNumber(StartIndex));
 					}
 				}), 
 			FSlateIcon(),
@@ -279,11 +279,11 @@ void FCustomPrimitiveDataTrackEditor::OnBuildAddParameterMenu(FMenuBuilder& Menu
 				{
 					if (CanAddParameter(CPDTrack, GetStartIndex, 3))
 					{
-						return FText::Format(LOCTEXT("CustomPrimitiveDataMaterialTooltipVector", "Add a track for a new Vector parameter starting at index {0}"), FText::AsNumber(StartIndex));
+						return FText::Format(LOCTEXT("CustomPrimitiveDataNewMaterialTooltipVector", "Add a track for a new Vector parameter starting at index {0}"), FText::AsNumber(StartIndex));
 					}
 					else
 					{
-						return FText::Format(LOCTEXT("CustomPrimitiveDataMaterialTooltipVectorInvalid", "Can't add a Vector parameter starting at index {0} as it overlaps with other parameters using that index."), FText::AsNumber(StartIndex));
+						return FText::Format(LOCTEXT("CustomPrimitiveDataNewMaterialTooltipVectorInvalid", "Can't add a Vector parameter starting at index {0} as it overlaps with other parameters using that index."), FText::AsNumber(StartIndex));
 					}
 				}),
 			FSlateIcon(),
@@ -298,11 +298,11 @@ void FCustomPrimitiveDataTrackEditor::OnBuildAddParameterMenu(FMenuBuilder& Menu
 				{
 					if (CanAddParameter(CPDTrack, GetStartIndex, 4))
 					{
-						return FText::Format(LOCTEXT("CustomPrimitiveDataMaterialTooltipColor", "Add a track for a new color parameter starting at index {0}"), FText::AsNumber(StartIndex));
+						return FText::Format(LOCTEXT("CustomPrimitiveDataNewMaterialTooltipColor", "Add a track for a new color parameter starting at index {0}"), FText::AsNumber(StartIndex));
 					}
 					else
 					{
-						return FText::Format(LOCTEXT("CustomPrimitiveDataMaterialTooltipVectorColorInvalid", "Can't add a color parameter starting at index {0} as it overlaps with other parameters using that index."), FText::AsNumber(StartIndex));
+						return FText::Format(LOCTEXT("CustomPrimitiveDataNewMaterialTooltipVectorColorInvalid", "Can't add a color parameter starting at index {0} as it overlaps with other parameters using that index."), FText::AsNumber(StartIndex));
 					}
 				}), 
 			FSlateIcon(),
