Sequencer: Add pre/post cycles to integer and bool channels.  Integer supports all but linear, bool supports constant(default), cycle and oscillate

diff --git a/Engine/Source/Editor/MovieSceneTools/Private/Channels/BoolChannelCurveModel.cpp b/Engine/Source/Editor/MovieSceneTools/Private/Channels/BoolChannelCurveModel.cpp
--- a/Engine/Source/Editor/MovieSceneTools/Private/Channels/BoolChannelCurveModel.cpp
+++ b/Engine/Source/Editor/MovieSceneTools/Private/Channels/BoolChannelCurveModel.cpp
@@ -122,6 +122,38 @@ TUniquePtr<IBufferedCurveModel> FBoolChannelCurveModel::CreateBufferedCurveCopy(
 	return nullptr;
 }
 
+void FBoolChannelCurveModel::GetCurveAttributes(FCurveAttributes& OutCurveAttributes) const
+{
+	FMovieSceneBoolChannel* Channel = GetChannelHandle().Get();
+	if (Channel)
+	{
+		OutCurveAttributes.SetPreExtrapolation(Channel->PreInfinityExtrap);
+		OutCurveAttributes.SetPostExtrapolation(Channel->PostInfinityExtrap);
+	}
+}
+
+void FBoolChannelCurveModel::SetCurveAttributes(const FCurveAttributes& InCurveAttributes)
+{
+	FMovieSceneBoolChannel* Channel = GetChannelHandle().Get();
+	UMovieSceneSection* Section = Cast<UMovieSceneSection>(GetOwningObject());
+	if (Channel && Section && !IsReadOnly())
+	{
+		Section->MarkAsChanged();
+
+		if (InCurveAttributes.HasPreExtrapolation())
+		{
+			Channel->PreInfinityExtrap = InCurveAttributes.GetPreExtrapolation();
+		}
+
+		if (InCurveAttributes.HasPostExtrapolation())
+		{
+			Channel->PostInfinityExtrap = InCurveAttributes.GetPostExtrapolation();
+		}
+
+		CurveModifiedDelegate.Broadcast();
+	}
+}
+
 double FBoolChannelCurveModel::GetKeyValue(TArrayView<const bool> Values, int32 Index) const
 {
 	return (double)Values[Index];
diff --git a/Engine/Source/Editor/MovieSceneTools/Private/Channels/BuiltInChannelEditors.cpp b/Engine/Source/Editor/MovieSceneTools/Private/Channels/BuiltInChannelEditors.cpp
--- a/Engine/Source/Editor/MovieSceneTools/Private/Channels/BuiltInChannelEditors.cpp
+++ b/Engine/Source/Editor/MovieSceneTools/Private/Channels/BuiltInChannelEditors.cpp
@@ -28,6 +28,8 @@
 #include "Channels/MovieSceneChannelEditorData.h"
 #include "Channels/MovieSceneFloatChannel.h"
 #include "Channels/MovieSceneDoubleChannel.h"
+#include "Channels/MovieSceneBoolChannel.h"
+#include "Channels/MovieSceneIntegerChannel.h"
 #include "Channels/DoubleChannelCurveModel.h"
 #include "Channels/FloatChannelCurveModel.h"
 #include "Channels/IntegerChannelCurveModel.h"
@@ -1470,7 +1472,8 @@ struct FCurveChannelSectionMenuExtension : TSharedFromThis<FCurveChannelSectionM
 		);
 	}
 
-	void GetChannels(TArray<FMovieSceneFloatChannel*>& FloatChannels, TArray<FMovieSceneDoubleChannel*>& DoubleChannels) const
+	void GetChannels(TArray<FMovieSceneFloatChannel*>& FloatChannels, TArray<FMovieSceneDoubleChannel*>& DoubleChannels,
+		TArray<FMovieSceneIntegerChannel*>& IntegerChannels, TArray<FMovieSceneBoolChannel*>& BoolChannels) const
 	{
 		ISequencer* Sequencer = WeakSequencer.Pin().Get();
 		if (!Sequencer)
@@ -1494,10 +1497,20 @@ struct FCurveChannelSectionMenuExtension : TSharedFromThis<FCurveChannelSectionM
 				FMovieSceneDoubleChannel* Channel = static_cast<FMovieSceneDoubleChannel*>(Handle.Get());
 				DoubleChannels.Add(Channel);
 			}
+			else if (Handle.GetChannelTypeName() == FMovieSceneIntegerChannel::StaticStruct()->GetFName())
+			{
+				FMovieSceneIntegerChannel* Channel = static_cast<FMovieSceneIntegerChannel*>(Handle.Get());
+				IntegerChannels.Add(Channel);
+			}
+			else if (Handle.GetChannelTypeName() == FMovieSceneBoolChannel::StaticStruct()->GetFName())
+			{
+				FMovieSceneBoolChannel* Channel = static_cast<FMovieSceneBoolChannel*>(Handle.Get());
+				BoolChannels.Add(Channel);
+			}
 		}
 
 		// Otherwise, the channels of all the sections
-		if (FloatChannels.Num() + DoubleChannels.Num() == 0)
+		if (FloatChannels.Num() + DoubleChannels.Num() + IntegerChannels.Num() + BoolChannels.Num() == 0)
 		{
 			for (TWeakObjectPtr<UMovieSceneSection> WeakSection : Sections)
 			{
@@ -1512,6 +1525,14 @@ struct FCurveChannelSectionMenuExtension : TSharedFromThis<FCurveChannelSectionM
 					{
 						DoubleChannels.Add(Channel);
 					}
+					for (FMovieSceneIntegerChannel* Channel : ChannelProxy.GetChannels<FMovieSceneIntegerChannel>())
+					{
+						IntegerChannels.Add(Channel);
+					}
+					for (FMovieSceneBoolChannel* Channel : ChannelProxy.GetChannels<FMovieSceneBoolChannel>())
+					{
+						BoolChannels.Add(Channel);
+					}
 				}
 			}
 		}
@@ -1521,9 +1542,12 @@ struct FCurveChannelSectionMenuExtension : TSharedFromThis<FCurveChannelSectionM
 	{
 		TArray<FMovieSceneFloatChannel*> FloatChannels;
 		TArray<FMovieSceneDoubleChannel*> DoubleChannels;
-		GetChannels(FloatChannels, DoubleChannels);
+		TArray<FMovieSceneIntegerChannel*> IntegerChannels;
+		TArray<FMovieSceneBoolChannel*> BoolChannels;
+
+		GetChannels(FloatChannels, DoubleChannels,IntegerChannels,BoolChannels);
 
-		if (FloatChannels.Num() + DoubleChannels.Num() == 0)
+		if (FloatChannels.Num() + DoubleChannels.Num() + IntegerChannels.Num() + BoolChannels.Num() == 0)
 		{
 			return;
 		}
@@ -1554,6 +1578,18 @@ struct FCurveChannelSectionMenuExtension : TSharedFromThis<FCurveChannelSectionM
 			DestExtrap = ExtrapMode;
 			bAnythingChanged = true;
 		}
+		for (FMovieSceneIntegerChannel* Channel : IntegerChannels)
+		{
+			TEnumAsByte<ERichCurveExtrapolation>& DestExtrap = bPreInfinity ? Channel->PreInfinityExtrap : Channel->PostInfinityExtrap;
+			DestExtrap = ExtrapMode;
+			bAnythingChanged = true;
+		}
+		for (FMovieSceneBoolChannel* Channel : BoolChannels)
+		{
+			TEnumAsByte<ERichCurveExtrapolation>& DestExtrap = bPreInfinity ? Channel->PreInfinityExtrap : Channel->PostInfinityExtrap;
+			DestExtrap = ExtrapMode;
+			bAnythingChanged = true;
+		}
 
 		if (bAnythingChanged)
 		{
@@ -1572,7 +1608,10 @@ struct FCurveChannelSectionMenuExtension : TSharedFromThis<FCurveChannelSectionM
 	{
 		TArray<FMovieSceneFloatChannel*> FloatChannels;
 		TArray<FMovieSceneDoubleChannel*> DoubleChannels;
-		GetChannels(FloatChannels, DoubleChannels);
+		TArray<FMovieSceneIntegerChannel*> IntegerChannels;
+		TArray<FMovieSceneBoolChannel*> BoolChannels;
+
+		GetChannels(FloatChannels, DoubleChannels, IntegerChannels, BoolChannels);
 
 		for (FMovieSceneFloatChannel* Channel : FloatChannels)
 		{
@@ -1590,6 +1629,22 @@ struct FCurveChannelSectionMenuExtension : TSharedFromThis<FCurveChannelSectionM
 				return false;
 			}
 		}
+		for (FMovieSceneIntegerChannel* Channel : IntegerChannels)
+		{
+			ERichCurveExtrapolation SourceExtrap = bPreInfinity ? Channel->PreInfinityExtrap : Channel->PostInfinityExtrap;
+			if (SourceExtrap != ExtrapMode)
+			{
+				return false;
+			}
+		}
+		for (FMovieSceneBoolChannel* Channel : BoolChannels)
+		{
+			ERichCurveExtrapolation SourceExtrap = bPreInfinity ? Channel->PreInfinityExtrap : Channel->PostInfinityExtrap;
+			if (SourceExtrap != ExtrapMode)
+			{
+				return false;
+			}
+		}
 
 		return true;
 	}
@@ -1738,6 +1793,22 @@ void ExtendSectionMenu(FMenuBuilder& OuterMenuBuilder, TSharedPtr<FExtender> Men
 	MenuExtender->AddMenuExtension("SequencerChannels", EExtensionHook::First, nullptr, FMenuExtensionDelegate::CreateLambda([Extension](FMenuBuilder& MenuBuilder) { Extension->ExtendMenu(MenuBuilder); }));
 }
 
+void ExtendSectionMenu(FMenuBuilder& OuterMenuBuilder, TSharedPtr<FExtender> MenuExtender, TArray<TMovieSceneChannelHandle<FMovieSceneIntegerChannel>>&& Channels, TArrayView<UMovieSceneSection* const> Sections, TWeakPtr<ISequencer> InSequencer)
+{
+	TSharedRef<FCurveChannelSectionMenuExtension> Extension = FCurveChannelSectionMenuExtension::GetOrCreate(InSequencer);
+	Extension->AddSections(Sections);
+
+	MenuExtender->AddMenuExtension("SequencerChannels", EExtensionHook::First, nullptr, FMenuExtensionDelegate::CreateLambda([Extension](FMenuBuilder& MenuBuilder) { Extension->ExtendMenu(MenuBuilder); }));
+}
+
+void ExtendSectionMenu(FMenuBuilder& OuterMenuBuilder, TSharedPtr<FExtender> MenuExtender, TArray<TMovieSceneChannelHandle<FMovieSceneBoolChannel>>&& Channels, TArrayView<UMovieSceneSection* const> Sections, TWeakPtr<ISequencer> InSequencer)
+{
+	TSharedRef<FCurveChannelSectionMenuExtension> Extension = FCurveChannelSectionMenuExtension::GetOrCreate(InSequencer);
+	Extension->AddSections(Sections);
+
+	MenuExtender->AddMenuExtension("SequencerChannels", EExtensionHook::First, nullptr, FMenuExtensionDelegate::CreateLambda([Extension](FMenuBuilder& MenuBuilder) { Extension->ExtendMenu(MenuBuilder); }));
+}
+
 void ExtendKeyMenu(FMenuBuilder& OuterMenuBuilder, TSharedPtr<FExtender> MenuExtender, TArray<TExtendKeyMenuParams<FMovieSceneFloatChannel>>&& Channels, TWeakPtr<ISequencer> InSequencer)
 {
 	TSharedRef<FFloatChannelKeyMenuExtension> Extension = MakeShared<FFloatChannelKeyMenuExtension>(InSequencer, MoveTemp(Channels));
diff --git a/Engine/Source/Editor/MovieSceneTools/Private/Channels/BuiltInChannelEditors.h b/Engine/Source/Editor/MovieSceneTools/Private/Channels/BuiltInChannelEditors.h
--- a/Engine/Source/Editor/MovieSceneTools/Private/Channels/BuiltInChannelEditors.h
+++ b/Engine/Source/Editor/MovieSceneTools/Private/Channels/BuiltInChannelEditors.h
@@ -71,6 +71,8 @@ void DrawKeys(FMovieSceneEventChannel*    Channel, TArrayView<const FKeyHandle>
 /** Context menu overrides */
 void ExtendSectionMenu(FMenuBuilder& OuterMenuBuilder, TSharedPtr<FExtender> MenuExtender, TArray<TMovieSceneChannelHandle<FMovieSceneDoubleChannel>>&& Channels, TArrayView<UMovieSceneSection* const> Sections, TWeakPtr<ISequencer> InSequencer);
 void ExtendSectionMenu(FMenuBuilder& OuterMenuBuilder, TSharedPtr<FExtender> MenuExtender, TArray<TMovieSceneChannelHandle<FMovieSceneFloatChannel>>&& Channels, TArrayView<UMovieSceneSection* const> Sections, TWeakPtr<ISequencer> InSequencer);
+void ExtendSectionMenu(FMenuBuilder& OuterMenuBuilder, TSharedPtr<FExtender> MenuExtender, TArray<TMovieSceneChannelHandle<FMovieSceneIntegerChannel>>&& Channels, TArrayView<UMovieSceneSection* const> Sections, TWeakPtr<ISequencer> InSequencer);
+void ExtendSectionMenu(FMenuBuilder& OuterMenuBuilder, TSharedPtr<FExtender> MenuExtender, TArray<TMovieSceneChannelHandle<FMovieSceneBoolChannel>>&& Channels, TArrayView<UMovieSceneSection* const> Sections, TWeakPtr<ISequencer> InSequencer);
 void ExtendKeyMenu(FMenuBuilder& OuterMenuBuilder, TSharedPtr<FExtender> MenuExtender, TArray<TExtendKeyMenuParams<FMovieSceneDoubleChannel>>&& Channels, TWeakPtr<ISequencer> InSequencer);
 void ExtendKeyMenu(FMenuBuilder& OuterMenuBuilder, TSharedPtr<FExtender> MenuExtender, TArray<TExtendKeyMenuParams<FMovieSceneFloatChannel>>&& Channels, TWeakPtr<ISequencer> InSequencer);
 
diff --git a/Engine/Source/Editor/MovieSceneTools/Private/Channels/ChannelCurveModel.cpp b/Engine/Source/Editor/MovieSceneTools/Private/Channels/ChannelCurveModel.cpp
--- a/Engine/Source/Editor/MovieSceneTools/Private/Channels/ChannelCurveModel.cpp
+++ b/Engine/Source/Editor/MovieSceneTools/Private/Channels/ChannelCurveModel.cpp
@@ -100,28 +100,54 @@ void FChannelCurveModel<ChannelType, ChannelValue, KeyType>::DrawCurve(const FCu
 
 		// Add the lower bound of the visible space
 		const bool bValidRange = StartingIndex < EndingIndex;
-		if (bValidRange)
+		//if we aren't just doing the default constant then we need to sample
+		const bool PreNotConstant = (Channel->PreInfinityExtrap != RCCE_None && Channel->PreInfinityExtrap != RCCE_Constant);
+		const bool PostNotConstant = (Channel->PostInfinityExtrap != RCCE_None && Channel->PostInfinityExtrap != RCCE_Constant);
+		if (bValidRange && (PreNotConstant || PostNotConstant))
 		{
-			OutInterpolatingPoints.Add(MakeTuple(StartFrame / TickResolution, GetKeyValue(Values, StartingIndex)));
-		}
+			const FFrameRate DisplayResolution = Section->GetTypedOuter<UMovieScene>()->GetDisplayRate();
+			const FFrameNumber StartTimeInDisplay = FFrameRate::TransformTime(FFrameTime(StartFrame), TickResolution, DisplayResolution).FloorToFrame();
+			const FFrameNumber EndTimeInDisplay = FFrameRate::TransformTime(FFrameTime(EndFrame), TickResolution, DisplayResolution).CeilToFrame();
 
-		TOptional<double> PreviousValue;
-		for (int32 KeyIndex = StartingIndex; KeyIndex < EndingIndex; ++KeyIndex)
+			double Value = 0.0;
+			TOptional<double> PreviousValue;
+			for (FFrameNumber DisplayFrameNumber = StartTimeInDisplay; DisplayFrameNumber <= EndTimeInDisplay; ++DisplayFrameNumber)
+			{
+				FFrameNumber TickFrameNumber = FFrameRate::TransformTime(FFrameTime(DisplayFrameNumber), DisplayResolution, TickResolution).FrameNumber;
+				Evaluate(TickFrameNumber / TickResolution, Value);
+				if (PreviousValue.IsSet() && PreviousValue.GetValue() != Value)
+				{
+					OutInterpolatingPoints.Add(MakeTuple(TickFrameNumber / TickResolution, PreviousValue.GetValue()));
+				}
+				OutInterpolatingPoints.Add(MakeTuple(TickFrameNumber / TickResolution, Value));
+				PreviousValue = Value;
+			}
+		}
+		else
 		{
-			double Value = GetKeyValue(Values, KeyIndex);
-			if (PreviousValue.IsSet() && PreviousValue.GetValue() != Value)
+			if (bValidRange)
 			{
-				OutInterpolatingPoints.Add(MakeTuple(Times[KeyIndex] / TickResolution, PreviousValue.GetValue()));
+				OutInterpolatingPoints.Add(MakeTuple(StartFrame / TickResolution, GetKeyValue(Values, StartingIndex)));
 			}
 
-			OutInterpolatingPoints.Add(MakeTuple(Times[KeyIndex] / TickResolution, Value));
-			PreviousValue = Value;
-		}
+			TOptional<double> PreviousValue;
+			for (int32 KeyIndex = StartingIndex; KeyIndex < EndingIndex; ++KeyIndex)
+			{
+				double Value = GetKeyValue(Values, KeyIndex);
+				if (PreviousValue.IsSet() && PreviousValue.GetValue() != Value)
+				{
+					OutInterpolatingPoints.Add(MakeTuple(Times[KeyIndex] / TickResolution, PreviousValue.GetValue()));
+				}
 
-		// Add the upper bound of the visible space
-		if (bValidRange)
-		{
-			OutInterpolatingPoints.Add(MakeTuple(EndFrame / TickResolution, GetKeyValue(Values, EndingIndex - 1)));
+				OutInterpolatingPoints.Add(MakeTuple(Times[KeyIndex] / TickResolution, Value));
+				PreviousValue = Value;
+			}
+
+			// Add the upper bound of the visible space
+			if (bValidRange)
+			{
+				OutInterpolatingPoints.Add(MakeTuple(EndFrame / TickResolution, GetKeyValue(Values, EndingIndex - 1)));
+			}
 		}
 	}
 }
diff --git a/Engine/Source/Editor/MovieSceneTools/Private/Channels/IntegerChannelCurveModel.cpp b/Engine/Source/Editor/MovieSceneTools/Private/Channels/IntegerChannelCurveModel.cpp
--- a/Engine/Source/Editor/MovieSceneTools/Private/Channels/IntegerChannelCurveModel.cpp
+++ b/Engine/Source/Editor/MovieSceneTools/Private/Channels/IntegerChannelCurveModel.cpp
@@ -122,6 +122,38 @@ TUniquePtr<IBufferedCurveModel> FIntegerChannelCurveModel::CreateBufferedCurveCo
 	return nullptr;
 }
 
+void FIntegerChannelCurveModel::GetCurveAttributes(FCurveAttributes& OutCurveAttributes) const
+{
+	FMovieSceneIntegerChannel* Channel = GetChannelHandle().Get();
+	if (Channel)
+	{
+		OutCurveAttributes.SetPreExtrapolation(Channel->PreInfinityExtrap);
+		OutCurveAttributes.SetPostExtrapolation(Channel->PostInfinityExtrap);
+	}
+}
+
+void FIntegerChannelCurveModel::SetCurveAttributes(const FCurveAttributes& InCurveAttributes)
+{
+	FMovieSceneIntegerChannel* Channel = GetChannelHandle().Get();
+	UMovieSceneSection* Section = Cast<UMovieSceneSection>(GetOwningObject());
+	if (Channel && Section && !IsReadOnly())
+	{
+		Section->MarkAsChanged();
+
+		if (InCurveAttributes.HasPreExtrapolation())
+		{
+			Channel->PreInfinityExtrap = InCurveAttributes.GetPreExtrapolation();
+		}
+
+		if (InCurveAttributes.HasPostExtrapolation())
+		{
+			Channel->PostInfinityExtrap = InCurveAttributes.GetPostExtrapolation();
+		}
+
+		CurveModifiedDelegate.Broadcast();
+	}
+}
+
 double FIntegerChannelCurveModel::GetKeyValue(TArrayView<const int32> Values, int32 Index) const
 {
 	return (double)Values[Index];
diff --git a/Engine/Source/Editor/MovieSceneTools/Public/Channels/BoolChannelCurveModel.h b/Engine/Source/Editor/MovieSceneTools/Public/Channels/BoolChannelCurveModel.h
--- a/Engine/Source/Editor/MovieSceneTools/Public/Channels/BoolChannelCurveModel.h
+++ b/Engine/Source/Editor/MovieSceneTools/Public/Channels/BoolChannelCurveModel.h
@@ -30,6 +30,8 @@ public:
 	// FCurveModel
 	virtual void CreateKeyProxies(TArrayView<const FKeyHandle> InKeyHandles, TArrayView<UObject*> OutObjects) override;
 	virtual TUniquePtr<IBufferedCurveModel> CreateBufferedCurveCopy() const override;
+	virtual void GetCurveAttributes(FCurveAttributes& OutAttributes) const override;
+	virtual void SetCurveAttributes(const FCurveAttributes& InAttributes) override;
 
 protected:
 
diff --git a/Engine/Source/Editor/MovieSceneTools/Public/Channels/IntegerChannelCurveModel.h b/Engine/Source/Editor/MovieSceneTools/Public/Channels/IntegerChannelCurveModel.h
--- a/Engine/Source/Editor/MovieSceneTools/Public/Channels/IntegerChannelCurveModel.h
+++ b/Engine/Source/Editor/MovieSceneTools/Public/Channels/IntegerChannelCurveModel.h
@@ -30,10 +30,12 @@ public:
 	// FCurveModel
 	virtual void CreateKeyProxies(TArrayView<const FKeyHandle> InKeyHandles, TArrayView<UObject*> OutObjects) override;
 	virtual TUniquePtr<IBufferedCurveModel> CreateBufferedCurveCopy() const override;
-
+	virtual void GetCurveAttributes(FCurveAttributes& OutAttributes) const override;
+	virtual void SetCurveAttributes(const FCurveAttributes& InAttributes) override;
 protected:
 
 	// FChannelCurveModel
 	virtual double GetKeyValue(TArrayView<const int32> Values, int32 Index) const override;
 	virtual void SetKeyValue(int32 Index, double KeyValue) const override;
+
 };
\ No newline at end of file
diff --git a/Engine/Source/Runtime/MovieScene/Private/Channels/MovieSceneBoolChannel.cpp b/Engine/Source/Runtime/MovieScene/Private/Channels/MovieSceneBoolChannel.cpp
--- a/Engine/Source/Runtime/MovieScene/Private/Channels/MovieSceneBoolChannel.cpp
+++ b/Engine/Source/Runtime/MovieScene/Private/Channels/MovieSceneBoolChannel.cpp
@@ -46,7 +46,56 @@ bool FMovieSceneBoolChannel::Evaluate(FFrameTime InTime, bool& OutValue) const
 {
 	if (Times.Num())
 	{
-		const int32 Index = FMath::Max(0, Algo::UpperBound(Times, InTime.FrameNumber)-1);
+		const FFrameNumber MinFrame = Times[0];
+		const FFrameNumber MaxFrame = Times.Last();
+		//we do None, Constant, and Linear first there is no cycling and so we just exit
+		if (InTime < FFrameTime(MinFrame))
+		{
+			if (PreInfinityExtrap == RCCE_None)
+			{
+				return false;
+			}
+
+			if (PreInfinityExtrap == RCCE_Constant || PreInfinityExtrap == RCCE_Linear)
+			{
+				OutValue = Values[0];
+				return true;
+			}
+		}
+		else if (InTime > FFrameTime(MaxFrame))
+		{
+			if (PostInfinityExtrap == RCCE_None)
+			{
+				return false;
+			}
+
+			if (PostInfinityExtrap == RCCE_Constant || PreInfinityExtrap == RCCE_Linear)
+			{
+				OutValue = Values.Last();
+				return true;
+			}
+		}
+
+		// Compute the cycled time based on extrapolation
+		UE::MovieScene::FCycleParams Params = UE::MovieScene::CycleTime(MinFrame, MaxFrame, InTime);
+
+		// Deal with cycles and oscillation
+		if (InTime < FFrameTime(MinFrame))
+		{
+			if (PreInfinityExtrap == RCCE_Oscillate)
+			{
+				Params.Oscillate(MinFrame.Value, MaxFrame.Value);
+			}
+		}
+		else if (InTime > FFrameTime(MaxFrame))
+		{
+			if (PostInfinityExtrap == RCCE_Oscillate)
+			{
+				Params.Oscillate(MinFrame.Value, MaxFrame.Value);
+			}
+		}
+		
+		const int32 Index = FMath::Max(0, Algo::UpperBound(Times, Params.Time)-1);
 		OutValue = Values[Index];
 		return true;
 	}
diff --git a/Engine/Source/Runtime/MovieScene/Private/Channels/MovieSceneCurveChannelImpl.cpp b/Engine/Source/Runtime/MovieScene/Private/Channels/MovieSceneCurveChannelImpl.cpp
--- a/Engine/Source/Runtime/MovieScene/Private/Channels/MovieSceneCurveChannelImpl.cpp
+++ b/Engine/Source/Runtime/MovieScene/Private/Channels/MovieSceneCurveChannelImpl.cpp
@@ -51,7 +51,7 @@ static FAutoConsoleVariableRef CVarEnableCachedChannelEvaluation(
 
 template<typename ChannelType>
 static typename ChannelType::CurveValueType
-EvalForTwoKeys(
+	EvalForTwoKeys(
 		const typename ChannelType::ChannelValueType& Key1, FFrameNumber Key1Time,
 		const typename ChannelType::ChannelValueType& Key2, FFrameNumber Key2Time,
 		FFrameNumber InTime,
@@ -100,45 +100,6 @@ EvalForTwoKeys(
 	}
 }
 
-struct FCycleParams
-{
-	double ValueOffset;
-	FFrameTime Time;
-	int32 CycleCount;
-	int32 Duration;
-	bool bMirrorCurve;
-
-	FCycleParams(FFrameTime InTime, int32 InDuration)
-		: ValueOffset(0.0)
-		, Time(InTime)
-		, CycleCount(0)
-		, Duration(InDuration)
-		, bMirrorCurve(false)
-	{}
-
-	FORCEINLINE void ComputePreValueOffset(double FirstValue, double LastValue)
-	{
-		// CycleCount is negative for pre-extrap
-		ValueOffset = (LastValue-FirstValue) * CycleCount;
-	}
-	FORCEINLINE void ComputePostValueOffset(double FirstValue, double LastValue)
-	{
-		ValueOffset = (LastValue-FirstValue) * CycleCount;
-	}
-	FORCEINLINE bool ShouldMirrorCurve() const
-	{
-		return bMirrorCurve;
-	}
-	FORCEINLINE void Oscillate(int32 MinFrame, int32 MaxFrame)
-	{
-		if (FMath::Abs(CycleCount) % 2 == 1)
-		{
-			bMirrorCurve = true;
-			Time = MinFrame + (FFrameTime(MaxFrame) - Time);
-		}
-	}
-};
-
 FCycleParams CycleTime(FFrameNumber MinFrame, FFrameNumber MaxFrame, FFrameTime InTime)
 {
 	FCycleParams Params(InTime, MaxFrame.Value - MinFrame.Value);
@@ -152,21 +113,21 @@ FCycleParams CycleTime(FFrameNumber MinFrame, FFrameNumber MaxFrame, FFrameTime
 		const int32 CycleCount = ((MaxFrame - InTime) / Params.Duration).FloorToFrame().Value;
 
 		Params.CycleCount = -CycleCount;
-		Params.Time = InTime + FFrameTime(Params.Duration*CycleCount);
+		Params.Time = InTime + FFrameTime(Params.Duration * CycleCount);
 	}
 	else if (InTime > MaxFrame)
 	{
 		const int32 CycleCount = ((InTime - MinFrame) / Params.Duration).FloorToFrame().Value;
 
 		Params.CycleCount = CycleCount;
-		Params.Time = InTime - FFrameTime(Params.Duration*CycleCount);
+		Params.Time = InTime - FFrameTime(Params.Duration * CycleCount);
 	}
 
 	return Params;
 }
 
-} // namespace MovieScene
-} // namespace UE
+}
+}
 
 
 template<typename ChannelType>
diff --git a/Engine/Source/Runtime/MovieScene/Private/Channels/MovieSceneIntegerChannel.cpp b/Engine/Source/Runtime/MovieScene/Private/Channels/MovieSceneIntegerChannel.cpp
--- a/Engine/Source/Runtime/MovieScene/Private/Channels/MovieSceneIntegerChannel.cpp
+++ b/Engine/Source/Runtime/MovieScene/Private/Channels/MovieSceneIntegerChannel.cpp
@@ -47,8 +47,62 @@ bool FMovieSceneIntegerChannel::Evaluate(FFrameTime InTime, int32& OutValue) con
 {
 	if (Times.Num())
 	{
-		const int32 Index = FMath::Max(0, Algo::UpperBound(Times, InTime.FrameNumber)-1);
-		OutValue = Values[Index];
+		const FFrameNumber MinFrame = Times[0];
+		const FFrameNumber MaxFrame = Times.Last();
+		//we do None,Constant, and Linear first ,there is no cycling and so we just exit
+		if (InTime < FFrameTime(MinFrame))
+		{
+			if (PreInfinityExtrap == RCCE_None)
+			{
+				return false;
+			}
+
+			if (PreInfinityExtrap == RCCE_Constant || PreInfinityExtrap == RCCE_Linear)
+			{
+				OutValue = Values[0];
+				return true;
+			}
+		}
+		else if (InTime > FFrameTime(MaxFrame))
+		{
+			if (PostInfinityExtrap == RCCE_None)
+			{
+				return false;
+			}
+
+			if (PostInfinityExtrap == RCCE_Constant || PreInfinityExtrap == RCCE_Linear)
+			{
+				OutValue = Values.Last();
+				return true;
+			}
+		}
+
+		// Compute the cycled time based on extrapolation
+		UE::MovieScene::FCycleParams Params = UE::MovieScene::CycleTime(MinFrame, MaxFrame, InTime);
+
+		// Deal with offset cycles and oscillation
+		// Move int over to double then we will convert back to int if doing an offset
+		const double FirstValue = double(Values[0]);
+		const double LastValue = double(Values.Last());
+		if (InTime < FFrameTime(MinFrame))
+		{
+			switch (PreInfinityExtrap)
+			{
+			case RCCE_CycleWithOffset: Params.ComputePreValueOffset(FirstValue, LastValue); break;
+			case RCCE_Oscillate:       Params.Oscillate(MinFrame.Value, MaxFrame.Value);                       break;
+			}
+		}
+		else if (InTime > FFrameTime(MaxFrame))
+		{
+			switch (PostInfinityExtrap)
+			{
+			case RCCE_CycleWithOffset: Params.ComputePostValueOffset(FirstValue, LastValue); break;
+			case RCCE_Oscillate:       Params.Oscillate(MinFrame.Value, MaxFrame.Value);    break;
+			}
+		}
+
+		const int32 Index = FMath::Max(0, Algo::UpperBound(Times, Params.Time)-1);
+		OutValue = Values[Index] + (int32)(Params.ValueOffset + 0.5);
 		return true;
 	}
 	else if (bHasDefaultValue)
diff --git a/Engine/Source/Runtime/MovieScene/Public/Channels/MovieSceneBoolChannel.h b/Engine/Source/Runtime/MovieScene/Public/Channels/MovieSceneBoolChannel.h
--- a/Engine/Source/Runtime/MovieScene/Public/Channels/MovieSceneBoolChannel.h
+++ b/Engine/Source/Runtime/MovieScene/Public/Channels/MovieSceneBoolChannel.h
@@ -32,7 +32,7 @@ struct FMovieSceneBoolChannel : public FMovieSceneChannel
 	GENERATED_BODY()
 
 	FMovieSceneBoolChannel()
-		: DefaultValue(false), bHasDefaultValue(false)
+		: PreInfinityExtrap(RCCE_Constant), PostInfinityExtrap(RCCE_Constant), DefaultValue(false), bHasDefaultValue(false)
 	{}
 
 	/**
@@ -152,6 +152,14 @@ public:
 	{
 		bHasDefaultValue = false;
 	}
+public:
+	/** Pre-infinity extrapolation state, bool channel only supports constant, cycle and oscillate */
+	UPROPERTY()
+	TEnumAsByte<ERichCurveExtrapolation> PreInfinityExtrap;
+
+	/** Post-infinity extrapolation state, bool channel only supports  constant, cycle and oscillate */
+	UPROPERTY()
+	TEnumAsByte<ERichCurveExtrapolation> PostInfinityExtrap;
 
 protected:
 
diff --git a/Engine/Source/Runtime/MovieScene/Public/Channels/MovieSceneChannel.h b/Engine/Source/Runtime/MovieScene/Public/Channels/MovieSceneChannel.h
--- a/Engine/Source/Runtime/MovieScene/Public/Channels/MovieSceneChannel.h
+++ b/Engine/Source/Runtime/MovieScene/Public/Channels/MovieSceneChannel.h
@@ -44,7 +44,55 @@ struct FKeyMoveEventItem
 	int32 NewIndex;
 	FFrameNumber NewFrame;
 };
+/*
+*  Class to help with pre/post cycling of values
+*/
+namespace UE
+{
+namespace MovieScene
+{
+struct FCycleParams
+{
+	double ValueOffset;
+	FFrameTime Time;
+	int32 CycleCount;
+	int32 Duration;
+	bool bMirrorCurve;
+
+	FCycleParams(FFrameTime InTime, int32 InDuration)
+		: ValueOffset(0.0)
+		, Time(InTime)
+		, CycleCount(0)
+		, Duration(InDuration)
+		, bMirrorCurve(false)
+	{}
+
+	FORCEINLINE void ComputePreValueOffset(double FirstValue, double LastValue)
+	{
+		// CycleCount is negative for pre-extrap
+		ValueOffset = (LastValue - FirstValue) * CycleCount;
+	}
+	FORCEINLINE void ComputePostValueOffset(double FirstValue, double LastValue)
+	{
+		ValueOffset = (LastValue - FirstValue) * CycleCount;
+	}
+	FORCEINLINE bool ShouldMirrorCurve() const
+	{
+		return bMirrorCurve;
+	}
+	FORCEINLINE void Oscillate(int32 MinFrame, int32 MaxFrame)
+	{
+		if (FMath::Abs(CycleCount) % 2 == 1)
+		{
+			bMirrorCurve = true;
+			Time = MinFrame + (FFrameTime(MaxFrame) - Time);
+		}
+	}
+};
+MOVIESCENE_API  FCycleParams CycleTime(FFrameNumber MinFrame, FFrameNumber MaxFrame, FFrameTime InTime);
 
+} // namespace MovieScene
+} // namespace UE
 /*
 * Note if any Channel uses these delegate's they need a custom serializer to make sure they stick around on undo/redo. Dynamic delegates are too heavy.
 */
diff --git a/Engine/Source/Runtime/MovieScene/Public/Channels/MovieSceneIntegerChannel.h b/Engine/Source/Runtime/MovieScene/Public/Channels/MovieSceneIntegerChannel.h
--- a/Engine/Source/Runtime/MovieScene/Public/Channels/MovieSceneIntegerChannel.h
+++ b/Engine/Source/Runtime/MovieScene/Public/Channels/MovieSceneIntegerChannel.h
@@ -33,7 +33,7 @@ struct FMovieSceneIntegerChannel : public FMovieSceneChannel
 	GENERATED_BODY()
 
 	FMovieSceneIntegerChannel()
-		: DefaultValue(0), bHasDefaultValue(false)
+		: PreInfinityExtrap(RCCE_Constant), PostInfinityExtrap(RCCE_Constant), DefaultValue(0), bHasDefaultValue(false)
 	{}
 
 	/**
@@ -175,7 +175,14 @@ public:
 	{
 		bHasDefaultValue = false;
 	}
+public:
+	/** Pre-infinity extrapolation state, integer channel supports them all but linear since that requires a tangent*/
+	UPROPERTY()
+	TEnumAsByte<ERichCurveExtrapolation> PreInfinityExtrap;
 
+	/** Post-infinity extrapolation state, integer channel supports them all but linear since that requires a tangent*/
+	UPROPERTY()
+	TEnumAsByte<ERichCurveExtrapolation> PostInfinityExtrap;
 private:
 
 	UPROPERTY(meta=(KeyTimes))
