Sequencer: Remove dependency on SceneOutliner private

#rnx

diff --git a/Engine/Source/Editor/MovieSceneTools/MovieSceneTools.Build.cs b/Engine/Source/Editor/MovieSceneTools/MovieSceneTools.Build.cs
--- a/Engine/Source/Editor/MovieSceneTools/MovieSceneTools.Build.cs
+++ b/Engine/Source/Editor/MovieSceneTools/MovieSceneTools.Build.cs
@@ -9,7 +9,6 @@ public class MovieSceneTools : ModuleRules
 		PrivateIncludePaths.AddRange(
             new string[] {
 				System.IO.Path.Combine(GetModuleDirectory("UnrealEd"), "Private"), //compatibility for FBX importer
-				System.IO.Path.Combine(GetModuleDirectory("SceneOutliner"), "Private"),
 			}
         );
 
