Sequencer: Import from Animation Root transform keys as relative to the current time

diff --git a/Engine/Source/Editor/MovieSceneTools/Private/TrackEditors/TransformTrackEditor.cpp b/Engine/Source/Editor/MovieSceneTools/Private/TrackEditors/TransformTrackEditor.cpp
--- a/Engine/Source/Editor/MovieSceneTools/Private/TrackEditors/TransformTrackEditor.cpp
+++ b/Engine/Source/Editor/MovieSceneTools/Private/TrackEditors/TransformTrackEditor.cpp
@@ -1170,6 +1170,8 @@ void F3DTransformTrackEditor::ImportAnimSequenceTransforms(const FAssetData& Ass
 {
 	FSlateApplication::Get().DismissAllMenus();
 
+	FQualifiedFrameTime CurrentTime = Sequencer->GetLocalTime();
+
 	UAnimSequence* AnimSequence = Cast<UAnimSequence>(Asset.GetAsset());
 
 	// find object binding to recover any component transforms we need to incorporate (for characters)
@@ -1246,6 +1248,8 @@ void F3DTransformTrackEditor::ImportAnimSequenceTransforms(const FAssetData& Ass
 				float Time;
 			};
 
+			float MinTime = FLT_MAX;
+
 			TArray<FTempTransformKey> TempKeys;
 
 			TArray<FName> BoneTrackNames;
@@ -1266,6 +1270,8 @@ void F3DTransformTrackEditor::ImportAnimSequenceTransforms(const FAssetData& Ass
 				TempKey.WoundRotation = TempKey.Transform.GetRotation().Rotator();
 
 				TempKeys.Add(TempKey);
+
+				MinTime = FMath::Min(MinTime, TempKey.Time);
 			}
 
 			int32 TransformCount = TempKeys.Num();
@@ -1279,11 +1285,15 @@ void F3DTransformTrackEditor::ImportAnimSequenceTransforms(const FAssetData& Ass
 				FMath::WindRelativeAnglesDegrees(Rotator.Roll, NextRotator.Roll);
 			}
 
+			FFrameNumber MinKeyTime = (MinTime * TickResolution).RoundToFrame();
+
 			TRange<FFrameNumber> Range = Section->GetRange();
 			for(const FTempTransformKey& TempKey : TempKeys)
 			{
 				FFrameNumber KeyTime = (TempKey.Time * TickResolution).RoundToFrame();
 
+				KeyTime = CurrentTime.Time.FrameNumber + (KeyTime - MinKeyTime);
+
 				Range = TRange<FFrameNumber>::Hull(Range, TRange<FFrameNumber>(KeyTime));
 
 				const FVector3f Translation = (FVector3f)TempKey.Transform.GetTranslation();
