Sequencer: re-add FBX import/export menu entries to template sequences

diff --git a/Engine/Plugins/MovieScene/LevelSequenceEditor/Source/LevelSequenceEditor/Private/LevelSequenceEditorModule.cpp b/Engine/Plugins/MovieScene/LevelSequenceEditor/Source/LevelSequenceEditor/Private/LevelSequenceEditorModule.cpp
--- a/Engine/Plugins/MovieScene/LevelSequenceEditor/Source/LevelSequenceEditor/Private/LevelSequenceEditorModule.cpp
+++ b/Engine/Plugins/MovieScene/LevelSequenceEditor/Source/LevelSequenceEditor/Private/LevelSequenceEditorModule.cpp
@@ -26,7 +26,6 @@
 #include "Misc/LevelSequenceEditorActorSpawner.h"
 #include "Misc/LevelSequenceEditorHelpers.h"
 #include "Misc/LevelSequenceEditorSettings.h"
-#include "Misc/LevelSequenceFBXInterop.h"
 #include "Misc/MovieSceneSequenceEditor_LevelSequence.h"
 #include "SequencerCommands.h"
 #include "SequencerSettings.h"
diff --git a/Engine/Plugins/MovieScene/LevelSequenceEditor/Source/LevelSequenceEditor/Private/Misc/LevelSequenceCustomization.cpp b/Engine/Plugins/MovieScene/LevelSequenceEditor/Source/LevelSequenceEditor/Private/Misc/LevelSequenceCustomization.cpp
--- a/Engine/Plugins/MovieScene/LevelSequenceEditor/Source/LevelSequenceEditor/Private/Misc/LevelSequenceCustomization.cpp
+++ b/Engine/Plugins/MovieScene/LevelSequenceEditor/Source/LevelSequenceEditor/Private/Misc/LevelSequenceCustomization.cpp
@@ -10,7 +10,7 @@
 #include "ISequencer.h"
 #include "ISequencerModule.h"
 #include "LevelSequenceEditorCommands.h"
-#include "Misc/LevelSequenceFBXInterop.h"
+#include "LevelSequenceFBXInterop.h"
 #include "Modules/ModuleManager.h"
 #include "MovieScene.h"
 #include "MVVM/Extensions/IObjectBindingExtension.h"
diff --git a/Engine/Plugins/MovieScene/LevelSequenceEditor/Source/LevelSequenceEditor/Private/Misc/LevelSequenceFBXInterop.cpp b/Engine/Plugins/MovieScene/LevelSequenceEditor/Source/LevelSequenceEditor/Private/Misc/LevelSequenceFBXInterop.cpp
deleted file mode 100644
diff --git a/Engine/Plugins/MovieScene/LevelSequenceEditor/Source/LevelSequenceEditor/Private/Misc/LevelSequenceFBXInterop.h b/Engine/Plugins/MovieScene/LevelSequenceEditor/Source/LevelSequenceEditor/Private/Misc/LevelSequenceFBXInterop.h
deleted file mode 100644
diff --git a/Engine/Plugins/MovieScene/LevelSequenceEditor/Source/LevelSequenceEditor/Public/SequencerExportTask.h b/Engine/Plugins/MovieScene/LevelSequenceEditor/Source/LevelSequenceEditor/Public/SequencerExportTask.h
deleted file mode 100644
diff --git a/Engine/Plugins/MovieScene/TemplateSequence/Source/TemplateSequenceEditor/Private/Customizations/TemplateSequenceCustomizationBase.cpp b/Engine/Plugins/MovieScene/TemplateSequence/Source/TemplateSequenceEditor/Private/Customizations/TemplateSequenceCustomizationBase.cpp
--- a/Engine/Plugins/MovieScene/TemplateSequence/Source/TemplateSequenceEditor/Private/Customizations/TemplateSequenceCustomizationBase.cpp
+++ b/Engine/Plugins/MovieScene/TemplateSequence/Source/TemplateSequenceEditor/Private/Customizations/TemplateSequenceCustomizationBase.cpp
@@ -2,8 +2,11 @@
 
 #include "TemplateSequenceCustomizationBase.h"
 #include "Framework/Application/SlateApplication.h"
+#include "Framework/MultiBox/MultiBoxBuilder.h"
+#include "Framework/MultiBox/MultiBoxExtender.h"
 #include "GameFramework/Actor.h"
 #include "ISequencer.h"
+#include "LevelSequenceFBXInterop.h"
 #include "Misc/TemplateSequenceEditorUtil.h"
 #include "ScopedTransaction.h"
 #include "TemplateSequence.h"
@@ -20,6 +23,7 @@ void FTemplateSequenceCustomizationBase::RegisterSequencerCustomization(FSequenc
 
 	FSequencerCustomizationInfo BaseInfo;
 	BaseInfo.OnPaste.BindRaw(this, &FTemplateSequenceCustomizationBase::OnPaste);
+	BaseInfo.OnBuildObjectBindingContextMenu = FOnGetSequencerMenuExtender::CreateRaw(this, &FTemplateSequenceCustomizationBase::CreateObjectBindingContextMenuExtender);
 	Builder.AddCustomization(BaseInfo);
 }
 
@@ -72,6 +76,51 @@ ESequencerPasteSupport FTemplateSequenceCustomizationBase::OnPaste()
 	return ESequencerPasteSupport::Tracks | ESequencerPasteSupport::Sections;
 }
 
+TSharedPtr<FExtender> FTemplateSequenceCustomizationBase::CreateObjectBindingContextMenuExtender(UE::Sequencer::FViewModelPtr InViewModel)
+{
+	using namespace UE::Sequencer;
+
+	TSharedRef<FExtender> Extender = MakeShared<FExtender>();
+	TSharedPtr<FObjectBindingModel> ObjectBindingModel = InViewModel->CastThisShared<FObjectBindingModel>();
+	Extender->AddMenuExtension(
+			"ObjectBindingActions", EExtensionHook::Before, nullptr,
+			FMenuExtensionDelegate::CreateRaw(this, &FTemplateSequenceCustomizationBase::ExtendObjectBindingContextMenu, ObjectBindingModel));
+	return Extender.ToSharedPtr();
+}
+
+void FTemplateSequenceCustomizationBase::ExtendObjectBindingContextMenu(FMenuBuilder& MenuBuilder, TSharedPtr<UE::Sequencer::FObjectBindingModel> ObjectBindingModel)
+{
+	using namespace UE::Sequencer;
+
+	MenuBuilder.BeginSection("Import/Export", LOCTEXT("ImportExportMenuSectionName", "Import/Export"));
+
+	TSharedPtr<ISequencer> SequencerPtr = Sequencer->AsShared();
+
+	MenuBuilder.AddMenuEntry(
+		LOCTEXT("ImportFBX", "Import..."),
+		LOCTEXT("ImportFBXTooltip", "Import FBX animation to this object"),
+		FSlateIcon(),
+		FUIAction(
+			FExecuteAction::CreateLambda([=] {
+					FLevelSequenceFBXInterop Interop(SequencerPtr);
+					Interop.ImportFBXOntoSelectedNodes();
+				})
+		));
+
+	MenuBuilder.AddMenuEntry(
+		LOCTEXT("ExportFBX", "Export..."),
+		LOCTEXT("ExportFBXTooltip", "Export FBX animation from this object"),
+		FSlateIcon(),
+		FUIAction(
+			FExecuteAction::CreateLambda([=] {
+					FLevelSequenceFBXInterop Interop(SequencerPtr);
+					Interop.ExportFBX();
+				})
+		));
+
+	MenuBuilder.EndSection();
+}
+
 void FTemplateSequenceCustomizationBase::OnMovieSceneDataChanged(EMovieSceneDataChangeType ChangeType)
 {
 	if (ChangeType == EMovieSceneDataChangeType::TrackValueChanged || ChangeType == EMovieSceneDataChangeType::TrackValueChangedRefreshImmediately)
diff --git a/Engine/Plugins/MovieScene/TemplateSequence/Source/TemplateSequenceEditor/Private/Customizations/TemplateSequenceCustomizationBase.h b/Engine/Plugins/MovieScene/TemplateSequence/Source/TemplateSequenceEditor/Private/Customizations/TemplateSequenceCustomizationBase.h
--- a/Engine/Plugins/MovieScene/TemplateSequence/Source/TemplateSequenceEditor/Private/Customizations/TemplateSequenceCustomizationBase.h
+++ b/Engine/Plugins/MovieScene/TemplateSequence/Source/TemplateSequenceEditor/Private/Customizations/TemplateSequenceCustomizationBase.h
@@ -2,6 +2,8 @@
 
 #pragma once
 
+#include "MVVM/ViewModelPtr.h"
+#include "MVVM/ViewModels/ObjectBindingModel.h"
 #include "SequencerCustomizationManager.h"
 
 class UActorFactory;
@@ -26,6 +28,9 @@ protected:
 	ISequencer* GetSequencer() const { return Sequencer; }
 	UTemplateSequence* GetTemplateSequence() const { return TemplateSequence; }
 
+	TSharedPtr<FExtender> CreateObjectBindingContextMenuExtender(UE::Sequencer::FViewModelPtr InViewModel);
+	void ExtendObjectBindingContextMenu(FMenuBuilder& MenuBuilder, TSharedPtr<UE::Sequencer::FObjectBindingModel> ObjectBindingModel);
+
 private:
 	ESequencerPasteSupport OnPaste();
 	void OnMovieSceneDataChanged(EMovieSceneDataChangeType ChangeType);
diff --git a/Engine/Plugins/MovieScene/TemplateSequence/Source/TemplateSequenceEditor/TemplateSequenceEditor.Build.cs b/Engine/Plugins/MovieScene/TemplateSequence/Source/TemplateSequenceEditor/TemplateSequenceEditor.Build.cs
--- a/Engine/Plugins/MovieScene/TemplateSequence/Source/TemplateSequenceEditor/TemplateSequenceEditor.Build.cs
+++ b/Engine/Plugins/MovieScene/TemplateSequence/Source/TemplateSequenceEditor/TemplateSequenceEditor.Build.cs
@@ -21,7 +21,8 @@ public class TemplateSequenceEditor : ModuleRules
 				"MovieScene",
 				"MovieSceneTools",
 				"MovieSceneTracks",
-				"Sequencer",				
+				"Sequencer",
+				"SequencerCore",
 				"Slate",
 				"SlateCore",
 				"UnrealEd",
diff --git a/Engine/Source/Editor/MovieSceneTools/Private/LevelSequenceFBXInterop.cpp b/Engine/Source/Editor/MovieSceneTools/Private/LevelSequenceFBXInterop.cpp
new file mode 100644
--- /dev/null
+++ b/Engine/Source/Editor/MovieSceneTools/Private/LevelSequenceFBXInterop.cpp
@@ -0,0 +1,247 @@
+// Copyright Epic Games, Inc. All Rights Reserved.
+
+#include "LevelSequenceFBXInterop.h"
+
+#include "DesktopPlatformModule.h"
+#include "EditorDirectories.h"
+#include "Exporters/Exporter.h"
+#include "FbxExporter.h"
+#include "Framework/Application/SlateApplication.h"
+#include "Framework/Notifications/NotificationManager.h"
+#include "ISequencer.h"
+#include "MovieSceneToolHelpers.h"
+#include "MVVM/Extensions/IObjectBindingExtension.h"
+#include "MVVM/Extensions/IOutlinerExtension.h"
+#include "MVVM/Extensions/ITrackExtension.h"
+#include "MVVM/Selection/Selection.h"
+#include "MVVM/ViewModelPtr.h"
+#include "MVVM/ViewModels/ObjectBindingModel.h"
+#include "MVVM/ViewModels/SequencerEditorViewModel.h"
+#include "MVVM/ViewModels/OutlinerViewModel.h"
+#include "MVVM/ViewModels/ViewModelIterators.h"
+#include "SequencerExportTask.h"
+#include "SequencerUtilities.h"
+#include "UObject/UObjectIterator.h"
+#include "UnrealExporter.h"
+#include "Widgets/Notifications/SNotificationList.h"
+
+#define LOCTEXT_NAMESPACE "LevelSequenceFBXInterop"
+
+FLevelSequenceFBXInterop::FLevelSequenceFBXInterop(TSharedPtr<ISequencer> InSequencer)
+	: Sequencer(InSequencer)
+{
+}
+
+void FLevelSequenceFBXInterop::ImportFBX()
+{
+	using namespace UE::Sequencer;
+
+	TMap<FGuid, FString> ObjectBindingNameMap;
+
+	TSharedPtr<FSequencerEditorViewModel> EditorViewModel = Sequencer->GetViewModel();
+	TSharedPtr<FOutlinerViewModel> OutlinerViewModel = EditorViewModel->GetOutliner();
+	TParentFirstChildIterator<IObjectBindingExtension> ObjectBindingIt = OutlinerViewModel->GetRootItem()->GetDescendantsOfType<IObjectBindingExtension>();
+
+	// Only visit the first object binding in a given sub-hierarchy so we get top-level object bindings.
+	// This is done by skipping the branch on each iteration
+	for ( ; ObjectBindingIt; ++ObjectBindingIt)
+	{
+		FGuid ObjectBinding = ObjectBindingIt->GetObjectGuid();
+		ObjectBindingNameMap.Add(ObjectBinding, (*ObjectBindingIt).AsModel()->CastThisChecked<IOutlinerExtension>()->GetLabel().ToString());
+
+		ObjectBindingIt.IgnoreCurrentChildren();
+	}
+
+	MovieSceneToolHelpers::ImportFBXWithDialog(Sequencer->GetFocusedMovieSceneSequence(), *Sequencer, ObjectBindingNameMap, TOptional<bool>());
+}
+
+void FLevelSequenceFBXInterop::ImportFBXOntoSelectedNodes()
+{
+	using namespace UE::Sequencer;
+
+	// The object binding and names to match when importing from fbx
+	TMap<FGuid, FString> ObjectBindingNameMap;
+
+	TSharedPtr<FSequencerEditorViewModel> EditorViewModel = Sequencer->GetViewModel();
+	for (TViewModelPtr<FObjectBindingModel> ObjectBindingNode : EditorViewModel->GetSelection()->Outliner.Filter<FObjectBindingModel>())
+	{
+		FGuid ObjectBinding = ObjectBindingNode->GetObjectGuid();
+		ObjectBindingNameMap.Add(ObjectBinding, ObjectBindingNode->GetLabel().ToString());
+	}
+
+	MovieSceneToolHelpers::ImportFBXWithDialog(Sequencer->GetFocusedMovieSceneSequence(), *Sequencer, ObjectBindingNameMap, TOptional<bool>(false));
+}
+
+void FLevelSequenceFBXInterop::ExportFBX()
+{
+	using namespace UE::Sequencer;
+
+	TArray<UExporter*> Exporters;
+	TArray<FString> SaveFilenames;
+	IDesktopPlatform* DesktopPlatform = FDesktopPlatformModule::Get();
+	bool bExportFileNamePicked = false;
+	if ( DesktopPlatform != NULL )
+	{
+		FString FileTypes = "FBX document|*.fbx";
+		UMovieSceneSequence* Sequence = Sequencer->GetFocusedMovieSceneSequence();
+		for (TObjectIterator<UClass> It; It; ++It)
+		{
+			if (!It->IsChildOf(UExporter::StaticClass()) || It->HasAnyClassFlags(CLASS_Abstract | CLASS_Deprecated | CLASS_NewerVersionExists))
+			{
+				continue;
+			}
+
+			UExporter* Default = It->GetDefaultObject<UExporter>();
+			if (!Default->SupportsObject(Sequence))
+			{
+				continue;
+			}
+
+			for (int32 i = 0; i < Default->FormatExtension.Num(); ++i)
+			{
+				const FString& FormatExtension = Default->FormatExtension[i];
+				const FString& FormatDescription = Default->FormatDescription[i];
+
+				if (FileTypes.Len() > 0)
+				{
+					FileTypes += TEXT("|");
+				}
+				FileTypes += FormatDescription;
+				FileTypes += TEXT("|*.");
+				FileTypes += FormatExtension;
+			}
+
+			Exporters.Add(Default);
+		}
+
+		bExportFileNamePicked = DesktopPlatform->SaveFileDialog(
+			FSlateApplication::Get().FindBestParentWindowHandleForDialogs(nullptr),
+			LOCTEXT( "ExportLevelSequence", "Export Level Sequence" ).ToString(),
+			*( FEditorDirectories::Get().GetLastDirectory( ELastDirectory::FBX ) ),
+			TEXT( "" ),
+			*FileTypes,
+			EFileDialogFlags::None,
+			SaveFilenames );
+	}
+
+	if ( bExportFileNamePicked )
+	{
+		FString ExportFilename = SaveFilenames[0];
+		FEditorDirectories::Get().SetLastDirectory( ELastDirectory::FBX, FPaths::GetPath( ExportFilename ) ); // Save path as default for next time.
+
+		// Make sure external selection is up to date since export could happen on tracks that have been right clicked but not have their underlying bound objects selected yet since that happens on mouse up.
+		FSequencerUtilities::SynchronizeExternalSelectionWithSequencerSelection(Sequencer.ToSharedRef());
+		
+		// Select selected nodes if there are selected nodes
+		TArray<FGuid> Bindings;
+		TArray<UMovieSceneTrack*> Tracks;
+		UMovieScene* MovieScene = Sequencer->GetFocusedMovieSceneSequence()->GetMovieScene();
+		TSharedPtr<FSequencerEditorViewModel> EditorViewModel = Sequencer->GetViewModel();
+		for (FViewModelPtr Node : EditorViewModel->GetSelection()->Outliner)
+		{
+			if (TViewModelPtr<IObjectBindingExtension> ObjectBindingNode = Node.ImplicitCast())
+			{
+				Bindings.Add(ObjectBindingNode->GetObjectGuid());
+
+				for (TViewModelPtr<IOutlinerExtension> DescendantNode : Node->GetDescendantsOfType<IOutlinerExtension>())
+				{
+					if (!EditorViewModel->GetSelection()->Outliner.IsSelected(DescendantNode) && DescendantNode.AsModel()->IsA<IObjectBindingExtension>())
+					{
+						IObjectBindingExtension* DescendantObjectBindingNode = DescendantNode.AsModel()->CastThisChecked<IObjectBindingExtension>();
+						Bindings.Add(DescendantObjectBindingNode->GetObjectGuid());
+					}
+				}
+			}
+			else if (TViewModelPtr<ITrackExtension> TrackNode = Node.ImplicitCast())
+			{
+				UMovieSceneTrack* Track = TrackNode->GetTrack();
+				if (Track && MovieScene->ContainsTrack(*Track))
+				{
+					Tracks.Add(Track);
+				}
+			}
+		}
+
+		FString FileExtension = FPaths::GetExtension(ExportFilename);
+		if (FileExtension == TEXT("fbx"))
+		{
+			ExportFBXInternal(ExportFilename, Bindings, (Bindings.Num() + Tracks.Num()) > 0 ? Tracks : MovieScene->GetTracks());
+		}
+		else
+		{
+			for (UExporter* Exporter : Exporters)
+			{
+				if (Exporter->FormatExtension.Contains(FileExtension))
+				{
+					USequencerExportTask* ExportTask = NewObject<USequencerExportTask>();
+					TStrongObjectPtr<USequencerExportTask> ExportTaskGuard(ExportTask);
+					ExportTask->Object = Sequencer->GetFocusedMovieSceneSequence();
+					ExportTask->Exporter = nullptr;
+					ExportTask->Filename = ExportFilename;
+					ExportTask->bSelected = false;
+					ExportTask->bReplaceIdentical = true;
+					ExportTask->bPrompt = false;
+					ExportTask->bUseFileArchive = false;
+					ExportTask->bWriteEmptyFiles = false;
+					ExportTask->bAutomated = false;
+					ExportTask->Exporter = NewObject<UExporter>(GetTransientPackage(), Exporter->GetClass());
+
+					ExportTask->SequencerContext = Sequencer->GetPlaybackContext();
+
+					UExporter::RunAssetExportTask(ExportTask);
+
+					ExportTask->Object = nullptr;
+					ExportTask->Exporter = nullptr;
+					ExportTask->SequencerContext = nullptr;
+
+					break;
+				}
+			}
+		}
+	}
+}
+
+void FLevelSequenceFBXInterop::ExportFBXInternal(const FString& ExportFilename, const TArray<FGuid>& Bindings, const TArray<UMovieSceneTrack*>& Tracks)
+{
+	UnFbx::FFbxExporter* Exporter = UnFbx::FFbxExporter::GetInstance();
+	//Show the fbx export dialog options
+	bool ExportCancel = false;
+	bool ExportAll = false;
+	Exporter->FillExportOptions(false, true, ExportFilename, ExportCancel, ExportAll);
+	if (!ExportCancel)
+	{
+		UMovieScene* MovieScene = Sequencer->GetFocusedMovieSceneSequence()->GetMovieScene();
+		UWorld* World = Cast<UWorld>(Sequencer->GetPlaybackContext());
+		FMovieSceneSequenceIDRef Template = Sequencer->GetFocusedTemplateID();
+		UnFbx::FFbxExporter::FLevelSequenceNodeNameAdapter NodeNameAdapter(MovieScene, Sequencer.Get(), Template);
+
+		{
+			FSpawnableRestoreState SpawnableRestoreState(MovieScene);
+			if (SpawnableRestoreState.bWasChanged)
+			{
+				// Evaluate at the beginning of the subscene time to ensure that spawnables are created before export
+				Sequencer->SetLocalTimeDirectly(UE::MovieScene::DiscreteInclusiveLower(FSequencerUtilities::GetTimeBounds(Sequencer.ToSharedRef())));
+			}
+
+			FMovieSceneSequenceTransform RootToLocalTransform = Sequencer->GetFocusedMovieSceneSequenceTransform();
+			if (MovieSceneToolHelpers::ExportFBX(World, MovieScene, Sequencer.Get(), Bindings, Tracks, NodeNameAdapter, Template, ExportFilename, RootToLocalTransform))
+			{
+				FNotificationInfo Info(NSLOCTEXT("Sequencer", "ExportFBXSucceeded", "FBX Export Succeeded."));
+				Info.Hyperlink = FSimpleDelegate::CreateStatic([](FString InFilename) { FPlatformProcess::ExploreFolder(*InFilename); }, ExportFilename);
+				Info.HyperlinkText = FText::FromString(ExportFilename);
+				Info.ExpireDuration = 5.0f;
+				FSlateNotificationManager::Get().AddNotification(Info)->SetCompletionState(SNotificationItem::CS_Success);
+			}
+			else
+			{
+				FNotificationInfo Info(NSLOCTEXT("Sequencer", "ExportFBXFailed", "FBX Export Failed."));
+				Info.ExpireDuration = 5.0f;
+				FSlateNotificationManager::Get().AddNotification(Info)->SetCompletionState(SNotificationItem::CS_Fail);
+			}
+		}
+
+		Sequencer->ForceEvaluate();
+	}
+}
+
+#undef LOCTEXT_NAMESPACE
diff --git a/Engine/Source/Editor/MovieSceneTools/Public/LevelSequenceFBXInterop.h b/Engine/Source/Editor/MovieSceneTools/Public/LevelSequenceFBXInterop.h
new file mode 100644
--- /dev/null
+++ b/Engine/Source/Editor/MovieSceneTools/Public/LevelSequenceFBXInterop.h
@@ -0,0 +1,29 @@
+// Copyright Epic Games, Inc. All Rights Reserved.
+
+#pragma once
+
+#include "ISequencer.h"
+
+class MOVIESCENETOOLS_API FLevelSequenceFBXInterop
+{
+public:
+
+	FLevelSequenceFBXInterop(TSharedPtr<ISequencer> InSequencer);
+
+	/** Imports the animation from an fbx file. */
+	void ImportFBX();
+	void ImportFBXOntoSelectedNodes();
+
+	/** Exports the animation to an fbx file. */
+	void ExportFBX();
+
+private:
+
+	/** Exports sequence to a FBX file */
+	void ExportFBXInternal(const FString& ExportFilename, const TArray<FGuid>& Bindings, const TArray<UMovieSceneTrack*>& Tracks);
+
+private:
+
+	TSharedPtr<ISequencer> Sequencer;
+};
+
diff --git a/Engine/Source/Editor/MovieSceneTools/Public/SequencerExportTask.h b/Engine/Source/Editor/MovieSceneTools/Public/SequencerExportTask.h
new file mode 100644
--- /dev/null
+++ b/Engine/Source/Editor/MovieSceneTools/Public/SequencerExportTask.h
@@ -0,0 +1,24 @@
+// Copyright Epic Games, Inc. All Rights Reserved.
+
+#pragma once
+
+#include "CoreMinimal.h"
+#include "UObject/ObjectMacros.h"
+#include "UObject/Object.h"
+#include "AssetExportTask.h"
+#include "SequencerExportTask.generated.h"
+
+/**
+ * Contains data for a group of assets to import
+ */ 
+UCLASS(Transient, BlueprintType)
+class USequencerExportTask : public UAssetExportTask
+{
+	GENERATED_BODY()
+
+public:
+
+	/* A UWorld for LevelSequences, UUserWidget for WidgetAnimations, or AActor for Actor Sequences, etc... */
+	UPROPERTY(BlueprintReadWrite, Category = Miscellaneous)
+	TObjectPtr<UObject> SequencerContext;
+};
