Curve Editor: Show arrive tangent if previous key is cubic. This allows the arrive tangent of a constant key to be visible and manipulatable.

diff --git a/Engine/Source/Editor/CurveEditor/Private/RichCurveEditorModel.cpp b/Engine/Source/Editor/CurveEditor/Private/RichCurveEditorModel.cpp
--- a/Engine/Source/Editor/CurveEditor/Private/RichCurveEditorModel.cpp
+++ b/Engine/Source/Editor/CurveEditor/Private/RichCurveEditorModel.cpp
@@ -385,6 +385,17 @@ void FRichCurveEditorModel::GetKeyAttributes(TArrayView<const FKeyHandle> InKeys
 
 					Attributes.SetInterpMode(ThisKey.InterpMode);
 
+					// If the previous key is cubic, show the arrive tangent handle even if this key is constant
+					FKeyHandle PreviousKeyHandle = RichCurve.GetPreviousKey(InKeys[Index]);
+					const bool bGetArriveTangent = RichCurve.IsKeyHandleValid(PreviousKeyHandle) && RichCurve.GetKeyRef(PreviousKeyHandle).InterpMode == RCIM_Cubic;
+					if (bGetArriveTangent)
+					{
+						if (&ThisKey != FirstKey)
+						{
+							Attributes.SetArriveTangent(ThisKey.ArriveTangent);
+						}
+					}
+
 					if (ThisKey.InterpMode != RCIM_Constant && ThisKey.InterpMode != RCIM_Linear)
 					{
 						Attributes.SetTangentMode(ThisKey.TangentMode);
diff --git a/Engine/Source/Editor/MovieSceneTools/Private/Channels/BezierChannelCurveModel.cpp b/Engine/Source/Editor/MovieSceneTools/Private/Channels/BezierChannelCurveModel.cpp
--- a/Engine/Source/Editor/MovieSceneTools/Private/Channels/BezierChannelCurveModel.cpp
+++ b/Engine/Source/Editor/MovieSceneTools/Private/Channels/BezierChannelCurveModel.cpp
@@ -217,7 +217,18 @@ void FBezierChannelCurveModel<ChannelType, ChannelValue, KeyType>::GetKeyAttribu
 
 				Attributes.SetInterpMode(KeyValue.InterpMode);
 
-				if (KeyValue.InterpMode != RCIM_Constant && KeyValue.InterpMode != RCIM_Linear)
+				// If the previous key is cubic, show the arrive tangent handle even if this key is constant
+				const int32 PreviousKeyIndex = KeyIndex - 1;
+				const bool bGetArriveTangent = Values.IsValidIndex(PreviousKeyIndex) && Values[PreviousKeyIndex].InterpMode == RCIM_Cubic;
+				if (bGetArriveTangent)
+				{
+					if (KeyIndex != 0)
+					{
+						Attributes.SetArriveTangent(KeyValue.Tangent.ArriveTangent / TimeInterval);
+					}
+				}
+
+				if ((KeyValue.InterpMode != RCIM_Constant && KeyValue.InterpMode != RCIM_Linear))
 				{
 					Attributes.SetTangentMode(KeyValue.TangentMode);
 					if (KeyIndex != 0)
