Sequencer- Add fixup UX logic for Component Material Parameter and Material Switcher tracks. We now detect when a track bound to a material by slot name no longer matches both slot name and material index, and allow the user to either automatically or manually rebind the track to the material slot and index.
As part of this, change display name tooltips on tracks back to using a function rather than a saved value so they can be dynamically specified. This needed Sequencer and ObjectBindingID passed to the function for context. This was only very recently changed to use saved values, so this shouldn't cause issues.

diff --git a/Engine/Source/Editor/MovieSceneTools/Private/TrackEditors/MaterialTrackEditor.cpp b/Engine/Source/Editor/MovieSceneTools/Private/TrackEditors/MaterialTrackEditor.cpp
--- a/Engine/Source/Editor/MovieSceneTools/Private/TrackEditors/MaterialTrackEditor.cpp
+++ b/Engine/Source/Editor/MovieSceneTools/Private/TrackEditors/MaterialTrackEditor.cpp
@@ -287,6 +287,131 @@ TSharedRef<ISequencerTrackEditor> FComponentMaterialTrackEditor::CreateTrackEdit
 }
 
 
+void FComponentMaterialTrackEditor::BuildTrackContextMenu(FMenuBuilder& MenuBuilder, UMovieSceneTrack* Track)
+{
+	if (UMovieSceneComponentMaterialTrack* MaterialTrack = Cast<UMovieSceneComponentMaterialTrack>(Track))
+	{
+		FGuid ObjectBinding = MaterialTrack->FindObjectBindingGuid();
+		UObject* BoundObject = GetSequencer()->FindSpawnedObjectOrTemplate(ObjectBinding);
+		const FComponentMaterialInfo& MaterialInfo = MaterialTrack->GetMaterialInfo();
+		if (MaterialInfo.MaterialType == EComponentMaterialType::Empty || MaterialInfo.MaterialType == EComponentMaterialType::IndexedMaterial)
+		{
+			if (UPrimitiveComponent* Component = Cast<UPrimitiveComponent>(BoundObject))
+			{
+				int32 FoundMaterialIndex = INDEX_NONE;
+				FName FoundSlotName;
+				UMaterialInterface* Material = nullptr;
+				FText AutoRebindTooltip;
+				bool bNeedsRebindFix = false;
+
+				if (!MaterialInfo.MaterialSlotName.IsNone())
+				{
+					Material = Component->GetMaterialByName(MaterialInfo.MaterialSlotName);
+					if (Material)
+					{
+						FoundSlotName = MaterialInfo.MaterialSlotName;
+					}
+					FoundMaterialIndex = Component->GetMaterialIndex(MaterialInfo.MaterialSlotName);
+					if (FoundMaterialIndex != INDEX_NONE && FoundMaterialIndex != MaterialInfo.MaterialSlotIndex)
+					{
+						bNeedsRebindFix = true;
+						// Found material by slot name, but the indices don't match what is cached. Auto-rebind would change index to the one in the found slot name.
+						AutoRebindTooltip = FText::Format(LOCTEXT("AutoRebindToNewSlotIndex", "Rebind track to index {0}, keeping same slot {1}"), FText::AsNumber(FoundMaterialIndex), FText::FromName(MaterialInfo.MaterialSlotName));
+					}
+				}
+				if (!Material)
+				{
+					bNeedsRebindFix = true;
+					// Couldn't find binding by slot name or one wasn't specified, try by index
+					Material = Component->GetMaterial(MaterialInfo.MaterialSlotIndex);
+					if (Material)
+					{
+						// Found material by index, but not by slot name. Auto-rebind would change slot name to the current one
+						FoundMaterialIndex = MaterialInfo.MaterialSlotIndex;
+						TArray<FName> SlotNames = Component->GetMaterialSlotNames();
+						check(SlotNames.IsValidIndex(MaterialInfo.MaterialSlotIndex));
+						FoundSlotName = SlotNames[MaterialInfo.MaterialSlotIndex];
+						AutoRebindTooltip = FText::Format(LOCTEXT("AutoRebindToNewSlotName", "Rebind track to slot {0}, keeping same index {1}"), FText::FromName(FoundSlotName), FText::AsNumber(MaterialInfo.MaterialSlotIndex));
+					}
+					// If we didn't find a material, we don't create a tooltip, because we won't be able to 'auto' rebind, just manually bind.
+				}
+
+				// If we didn't find a material, a slot name wasn't specified or wasn't found, or an index was mismatched, allow the user to rebind
+				if (bNeedsRebindFix)
+				{
+					MenuBuilder.BeginSection("Fix Material Binding", LOCTEXT("FixMaterialBindingSectionName", "Fix Material Binding"));
+
+					// If we set this tooltip, then we are able to auto-rebind
+					if (!AutoRebindTooltip.IsEmpty())
+					{
+						MenuBuilder.AddMenuEntry(
+							LOCTEXT("AutoFixMaterialBinding", "Auto-Fix Material Binding"),
+							AutoRebindTooltip,
+							FSlateIcon(),
+							FUIAction(
+								FExecuteAction::CreateLambda([MaterialTrack, FoundMaterialIndex, FoundSlotName]() 
+									{ 
+										FScopedTransaction Transaction(LOCTEXT("FixupMaterialBinding", "Fixup material binding"));
+										MaterialTrack->Modify();
+										MaterialTrack->SetMaterialInfo(FComponentMaterialInfo{ FoundSlotName, FoundMaterialIndex, EComponentMaterialType::IndexedMaterial });
+										MaterialTrack->SetDisplayName(FText::Format(LOCTEXT("SlotMaterialTrackName", "Material Slot: {0}"), FText::FromName(FoundSlotName)));
+									})));
+					}
+
+					MenuBuilder.AddSubMenu(
+						LOCTEXT("RebindMaterial", "Re-Bind Material Track..."),
+						LOCTEXT("RebindMaterialTooltip", "Re-Bind this material track to a different material slot"),
+						FNewMenuDelegate::CreateSP(this, &FComponentMaterialTrackEditor::FillRebindMaterialTrackMenu, MaterialTrack, Component, ObjectBinding)
+					);
+
+					MenuBuilder.EndSection();
+					MenuBuilder.AddMenuSeparator();
+				}
+			}
+		}
+	}
+}
+
+void FComponentMaterialTrackEditor::FillRebindMaterialTrackMenu(FMenuBuilder& MenuBuilder, UMovieSceneComponentMaterialTrack* MaterialTrack, UPrimitiveComponent* PrimitiveComponent, FGuid ObjectBinding)
+{
+	auto GetMaterialInfoForTrack = [](UMovieSceneTrack* InTrack)
+	{
+		UMovieSceneComponentMaterialTrack* MaterialTrack = Cast<UMovieSceneComponentMaterialTrack>(InTrack);
+		return MaterialTrack ? MaterialTrack->GetMaterialInfo() : FComponentMaterialInfo();
+	};
+	int32 NumMaterials = PrimitiveComponent->GetNumMaterials();
+	TArray<FName> MaterialSlotNames = PrimitiveComponent->GetMaterialSlotNames();
+	const UMovieScene* MovieScene = GetFocusedMovieScene();
+	const FMovieSceneBinding* Binding = Algo::FindBy(MovieScene->GetBindings(), ObjectBinding, &FMovieSceneBinding::GetObjectGuid);
+	if (Binding && NumMaterials > 0)
+	{
+		MenuBuilder.BeginSection("RebindMaterialSlots", LOCTEXT("RebindMaterialSlots", "Material Slots"));
+		{
+			for (int32 MaterialIndex = 0; MaterialIndex < NumMaterials; MaterialIndex++)
+			{
+				FName MaterialSlotName = MaterialSlotNames.IsValidIndex(MaterialIndex) ? MaterialSlotNames[MaterialIndex] : FName();
+				FComponentMaterialInfo MaterialInfo{ MaterialSlotName, MaterialIndex, EComponentMaterialType::IndexedMaterial };
+				const bool bAlreadyExists = Algo::FindBy(Binding->GetTracks(), MaterialInfo, GetMaterialInfoForTrack) != nullptr;
+				if (bAlreadyExists)
+				{
+					continue;
+				}
+				FUIAction RebindToMaterialAction(FExecuteAction::CreateLambda([MaterialTrack, MaterialInfo]() 
+					{
+						FScopedTransaction Transaction(LOCTEXT("FixupMaterialBinding", "Fixup material binding"));
+						MaterialTrack->Modify();
+						MaterialTrack->SetMaterialInfo(MaterialInfo); 
+						MaterialTrack->SetDisplayName(FText::Format(LOCTEXT("SlotMaterialTrackName", "Material Slot: {0}"), FText::FromName(MaterialInfo.MaterialSlotName)));
+					}));
+				FText RebindToMaterialLabel = FText::Format(LOCTEXT("RebindToMaterialSlot", "Slot: {0}, Index: {1}"), FText::FromName(MaterialSlotName), FText::AsNumber(MaterialIndex));
+				FText RebindToMaterialToolTip = FText::Format(LOCTEXT("RebindToMaterialSlotToolTip", "Rebind this track to material slot {0}, index {1}"), FText::FromName(MaterialSlotName), FText::AsNumber(MaterialIndex));
+				MenuBuilder.AddMenuEntry(RebindToMaterialLabel, RebindToMaterialToolTip, FSlateIcon(), RebindToMaterialAction);
+			}
+		}
+		MenuBuilder.EndSection();
+	}
+}
+
 bool FComponentMaterialTrackEditor::SupportsType( TSubclassOf<UMovieSceneTrack> Type ) const
 {
 	return Type == UMovieSceneComponentMaterialTrack::StaticClass();
@@ -509,7 +634,6 @@ void FComponentMaterialTrackEditor::HandleAddComponentMaterialActionExecute(USce
 			// Construct display name from MaterialInfo
 			UMaterialInterface* MaterialInterface = GetMaterialInterfaceForTrack(ObjectHandle, MaterialTrack);
 			FText TrackDisplayName;
-			FText TrackTooltipText;
 			switch (MaterialInfo.MaterialType)
 			{
 			case EComponentMaterialType::Empty:
@@ -517,16 +641,12 @@ void FComponentMaterialTrackEditor::HandleAddComponentMaterialActionExecute(USce
 			case EComponentMaterialType::IndexedMaterial:
 				TrackDisplayName = !MaterialInfo.MaterialSlotName.IsNone() ? FText::Format(LOCTEXT("SlotMaterialTrackName", "Material Slot: {0}"), FText::FromName(MaterialInfo.MaterialSlotName))
 					: FText::Format(LOCTEXT("IndexedMaterialTrackName", "Material Element {0}"), FText::AsNumber(MaterialInfo.MaterialSlotIndex));
-				TrackTooltipText = !MaterialInfo.MaterialSlotName.IsNone() ? FText::Format(LOCTEXT("SlotMaterialTrackTooltip", "Material parameter track for slot {0} at index {1}, asset {2}"), FText::FromName(MaterialInfo.MaterialSlotName), FText::AsNumber(MaterialInfo.MaterialSlotIndex), MaterialInterface ? FText::FromString(MaterialInterface->GetName()) : FText())
-					: FText::Format(LOCTEXT("IndexedMaterialTrackTooltip", "Material parameter track for element at index {0}, asset {1}"), FText::AsNumber(MaterialInfo.MaterialSlotIndex), MaterialInterface ? FText::FromString(MaterialInterface->GetName()) : FText());
 				break;
 			case EComponentMaterialType::OverlayMaterial:
 				TrackDisplayName = FText::Format(LOCTEXT("OverlayMaterialTrackName", "Overlay Material {0}"), MaterialInterface ? FText::FromString(MaterialInterface->GetName()) : FText());
-				TrackTooltipText = FText::Format(LOCTEXT("OverlayMaterialTrackTooltip", "Material parameter track for overlay material {0}"), MaterialInterface ? FText::FromString(MaterialInterface->GetName()) : FText());
 				break;
 			case EComponentMaterialType::DecalMaterial:
 				TrackDisplayName = FText::Format(LOCTEXT("DecalMaterialTrackName", "Decal Material {0}"), MaterialInterface ? FText::FromString(MaterialInterface->GetName()) : FText());
-				TrackTooltipText = FText::Format(LOCTEXT("DecalMaterialTrackTooltip", "Material parameter track for decal material {0}"), MaterialInterface ? FText::FromString(MaterialInterface->GetName()) : FText());
 				break;
 			default:
 				break;
@@ -536,16 +656,10 @@ void FComponentMaterialTrackEditor::HandleAddComponentMaterialActionExecute(USce
 			{
 				MaterialTrack->SetDisplayName(TrackDisplayName);
 			}
-			if (!TrackTooltipText.IsEmpty())
-			{
-				MaterialTrack->SetDisplayNameTooltipText(TrackTooltipText);
-			}
 		}
 	}
 
 	SequencerPtr->NotifyMovieSceneDataChanged(EMovieSceneDataChangeType::MovieSceneStructureItemAdded);
 }
 
-
-
 #undef LOCTEXT_NAMESPACE
diff --git a/Engine/Source/Editor/MovieSceneTools/Private/TrackEditors/PrimitiveMaterialTrackEditor.cpp b/Engine/Source/Editor/MovieSceneTools/Private/TrackEditors/PrimitiveMaterialTrackEditor.cpp
--- a/Engine/Source/Editor/MovieSceneTools/Private/TrackEditors/PrimitiveMaterialTrackEditor.cpp
+++ b/Engine/Source/Editor/MovieSceneTools/Private/TrackEditors/PrimitiveMaterialTrackEditor.cpp
@@ -129,7 +129,6 @@ void FPrimitiveMaterialTrackEditor::CreateTrackForElement(TArray<FGuid> ObjectBi
 		NewTrack->SetMaterialInfo(MaterialInfo);
 		// Construct display names from MaterialInfo
 		FText TrackDisplayName;
-		FText TrackTooltipText;
 		switch (MaterialInfo.MaterialType)
 		{
 		case EComponentMaterialType::Empty:
@@ -138,16 +137,12 @@ void FPrimitiveMaterialTrackEditor::CreateTrackForElement(TArray<FGuid> ObjectBi
 			TrackDisplayName = !MaterialInfo.MaterialSlotName.IsNone() ? FText::Format(LOCTEXT("SlotMaterialSwitcherTrackName", "Material Slot: {0}"), FText::FromName(MaterialInfo.MaterialSlotName))
 				: FText::Format(LOCTEXT("IndexedMaterialSwitcherTrackName", "Material Element {0}"), FText::AsNumber(MaterialInfo.MaterialSlotIndex));
 
-			TrackTooltipText = !MaterialInfo.MaterialSlotName.IsNone() ? FText::Format(LOCTEXT("SlotMaterialSwitcherTrackTooltip", "Material switcher for {0} at index {1}"), FText::FromName(MaterialInfo.MaterialSlotName), FText::AsNumber(MaterialInfo.MaterialSlotIndex))
-				: FText::Format(LOCTEXT("IndexedMaterialSwitcherTrackTooltip", "Material switcher for element at index {0}"), FText::AsNumber(MaterialInfo.MaterialSlotIndex));
 			break;
 		case EComponentMaterialType::OverlayMaterial:
 			TrackDisplayName = LOCTEXT("OverlayMaterialSwitcherTrackName", "Overlay Material");
-			TrackTooltipText = LOCTEXT("OverlayMaterialSwitcherTrackTooltip", "Material switcher for overlay material");
 			break;
 		case EComponentMaterialType::DecalMaterial:
 			TrackDisplayName = LOCTEXT("DecalMaterialSwitcherTrackName", "Decal Material");
-			TrackTooltipText = LOCTEXT("DecalMaterialSwitcherTrackTooltip", "Material switcher for decal material");
 			break;
 		default:
 			break;
@@ -157,10 +152,6 @@ void FPrimitiveMaterialTrackEditor::CreateTrackForElement(TArray<FGuid> ObjectBi
 		{
 			NewTrack->SetDisplayName(TrackDisplayName);
 		}
-		if (!TrackTooltipText.IsEmpty())
-		{
-			NewTrack->SetDisplayNameTooltipText(TrackTooltipText);
-		}
 
 		NewTrack->AddSection(*NewTrack->CreateNewSection());
 	}
@@ -168,4 +159,129 @@ void FPrimitiveMaterialTrackEditor::CreateTrackForElement(TArray<FGuid> ObjectBi
 	GetSequencer()->NotifyMovieSceneDataChanged(EMovieSceneDataChangeType::MovieSceneStructureItemAdded);
 }
 
+void FPrimitiveMaterialTrackEditor::BuildTrackContextMenu(FMenuBuilder& MenuBuilder, UMovieSceneTrack* Track)
+{
+	if (UMovieScenePrimitiveMaterialTrack* MaterialTrack = Cast<UMovieScenePrimitiveMaterialTrack>(Track))
+	{
+		FGuid ObjectBinding = MaterialTrack->FindObjectBindingGuid();
+		UObject* BoundObject = GetSequencer()->FindSpawnedObjectOrTemplate(ObjectBinding);
+		const FComponentMaterialInfo& MaterialInfo = MaterialTrack->GetMaterialInfo();
+		if (MaterialInfo.MaterialType == EComponentMaterialType::Empty || MaterialInfo.MaterialType == EComponentMaterialType::IndexedMaterial)
+		{
+			if (UPrimitiveComponent* Component = Cast<UPrimitiveComponent>(BoundObject))
+			{
+				int32 FoundMaterialIndex = INDEX_NONE;
+				FName FoundSlotName;
+				UMaterialInterface* Material = nullptr;
+				FText AutoRebindTooltip;
+				bool bNeedsRebindFix = false;
+
+				if (!MaterialInfo.MaterialSlotName.IsNone())
+				{
+					Material = Component->GetMaterialByName(MaterialInfo.MaterialSlotName);
+					if (Material)
+					{
+						FoundSlotName = MaterialInfo.MaterialSlotName;
+					}
+					FoundMaterialIndex = Component->GetMaterialIndex(MaterialInfo.MaterialSlotName);
+					if (FoundMaterialIndex != INDEX_NONE && FoundMaterialIndex != MaterialInfo.MaterialSlotIndex)
+					{
+						bNeedsRebindFix = true;
+						// Found material by slot name, but the indices don't match what is cached. Auto-rebind would change index to the one in the found slot name.
+						AutoRebindTooltip = FText::Format(LOCTEXT("AutoRebindToNewSlotSwitcherIndex", "Rebind track to index {0}, keeping same slot {1}"), FText::AsNumber(FoundMaterialIndex), FText::FromName(MaterialInfo.MaterialSlotName));
+					}
+				}
+				if (!Material)
+				{
+					bNeedsRebindFix = true;
+					// Couldn't find binding by slot name or one wasn't specified, try by index
+					Material = Component->GetMaterial(MaterialInfo.MaterialSlotIndex);
+					if (Material)
+					{
+						// Found material by index, but not by slot name. Auto-rebind would change slot name to the current one
+						FoundMaterialIndex = MaterialInfo.MaterialSlotIndex;
+						TArray<FName> SlotNames = Component->GetMaterialSlotNames();
+						check(SlotNames.IsValidIndex(MaterialInfo.MaterialSlotIndex));
+						FoundSlotName = SlotNames[MaterialInfo.MaterialSlotIndex];
+						AutoRebindTooltip = FText::Format(LOCTEXT("AutoRebindToNewSlotSwitcherName", "Rebind track to slot {0}, keeping same index {1}"), FText::FromName(FoundSlotName), FText::AsNumber(MaterialInfo.MaterialSlotIndex));
+					}
+					// If we didn't find a material, we don't create a tooltip, because we won't be able to 'auto' rebind, just manually bind.
+				}
+
+				// If we didn't find a material, a slot name wasn't specified or wasn't found, or an index was mismatched, allow the user to rebind
+				if (bNeedsRebindFix)
+				{
+					MenuBuilder.BeginSection("Fix Material Switcher Binding", LOCTEXT("FixMaterialBindingSwitcherSectionName", "Fix Material Switcher Binding"));
+
+					// If we set this tooltip, then we are able to auto-rebind
+					if (!AutoRebindTooltip.IsEmpty())
+					{
+						MenuBuilder.AddMenuEntry(
+							LOCTEXT("AutoFixMaterialSwitcherBinding", "Auto-Fix Material Switcher Binding"),
+							AutoRebindTooltip,
+							FSlateIcon(),
+							FUIAction(
+								FExecuteAction::CreateLambda([MaterialTrack, FoundMaterialIndex, FoundSlotName]()
+									{
+										FScopedTransaction Transaction(LOCTEXT("FixupMaterialSwitcherBinding", "Fixup material switcher binding"));
+										MaterialTrack->Modify();
+										MaterialTrack->SetMaterialInfo(FComponentMaterialInfo{ FoundSlotName, FoundMaterialIndex, EComponentMaterialType::IndexedMaterial });
+										MaterialTrack->SetDisplayName(FText::Format(LOCTEXT("SlotMaterialSwitcherTrackName", "Material Slot: {0}"), FText::FromName(FoundSlotName)));
+									})));
+					}
+
+					MenuBuilder.AddSubMenu(
+						LOCTEXT("RebindMaterialSwitcher", "Re-Bind Material Switcher Track..."),
+						LOCTEXT("RebindMaterialSwitcherTooltip", "Re-Bind this material switcher track to a different material slot"),
+						FNewMenuDelegate::CreateSP(this, &FPrimitiveMaterialTrackEditor::FillRebindMaterialTrackMenu, MaterialTrack, Component, ObjectBinding)
+					);
+
+					MenuBuilder.EndSection();
+					MenuBuilder.AddMenuSeparator();
+				}
+			}
+		}
+	}
+}
+
+void FPrimitiveMaterialTrackEditor::FillRebindMaterialTrackMenu(FMenuBuilder& MenuBuilder, UMovieScenePrimitiveMaterialTrack* MaterialTrack, UPrimitiveComponent* PrimitiveComponent, FGuid ObjectBinding)
+{
+	auto GetMaterialInfoForTrack = [](UMovieSceneTrack* InTrack)
+	{
+		UMovieScenePrimitiveMaterialTrack* MaterialTrack = Cast<UMovieScenePrimitiveMaterialTrack>(InTrack);
+		return MaterialTrack ? MaterialTrack->GetMaterialInfo() : FComponentMaterialInfo();
+	};
+	int32 NumMaterials = PrimitiveComponent->GetNumMaterials();
+	TArray<FName> MaterialSlotNames = PrimitiveComponent->GetMaterialSlotNames();
+	const UMovieScene* MovieScene = GetFocusedMovieScene();
+	const FMovieSceneBinding* Binding = Algo::FindBy(MovieScene->GetBindings(), ObjectBinding, &FMovieSceneBinding::GetObjectGuid);
+	if (Binding && NumMaterials > 0)
+	{
+		MenuBuilder.BeginSection("RebindMaterialSlots", LOCTEXT("RebindMaterialSlots", "Material Slots"));
+		{
+			for (int32 MaterialIndex = 0; MaterialIndex < NumMaterials; MaterialIndex++)
+			{
+				FName MaterialSlotName = MaterialSlotNames.IsValidIndex(MaterialIndex) ? MaterialSlotNames[MaterialIndex] : FName();
+				FComponentMaterialInfo MaterialInfo{ MaterialSlotName, MaterialIndex, EComponentMaterialType::IndexedMaterial };
+				const bool bAlreadyExists = Algo::FindBy(Binding->GetTracks(), MaterialInfo, GetMaterialInfoForTrack) != nullptr;
+				if (bAlreadyExists)
+				{
+					continue;
+				}
+				FUIAction RebindToMaterialAction(FExecuteAction::CreateLambda([MaterialTrack, MaterialInfo]()
+					{
+						FScopedTransaction Transaction(LOCTEXT("FixupMaterialSwitcherBinding", "Fixup material switcher binding"));
+						MaterialTrack->Modify();
+						MaterialTrack->SetMaterialInfo(MaterialInfo);
+						MaterialTrack->SetDisplayName(FText::Format(LOCTEXT("SlotMaterialSwitcherTrackName", "Material Slot: {0}"), FText::FromName(MaterialInfo.MaterialSlotName)));
+					}));
+				FText RebindToMaterialLabel = FText::Format(LOCTEXT("RebindToMaterialSwitcherSlot", "Slot: {0}, Index: {1}"), FText::FromName(MaterialSlotName), FText::AsNumber(MaterialIndex));
+				FText RebindToMaterialToolTip = FText::Format(LOCTEXT("RebindToMaterialSwitcherSlotToolTip", "Rebind this track to material slot {0}, index {1}"), FText::FromName(MaterialSlotName), FText::AsNumber(MaterialIndex));
+				MenuBuilder.AddMenuEntry(RebindToMaterialLabel, RebindToMaterialToolTip, FSlateIcon(), RebindToMaterialAction);
+			}
+		}
+		MenuBuilder.EndSection();
+	}
+}
+
 #undef LOCTEXT_NAMESPACE
diff --git a/Engine/Source/Editor/MovieSceneTools/Private/TrackEditors/PrimitiveMaterialTrackEditor.h b/Engine/Source/Editor/MovieSceneTools/Private/TrackEditors/PrimitiveMaterialTrackEditor.h
--- a/Engine/Source/Editor/MovieSceneTools/Private/TrackEditors/PrimitiveMaterialTrackEditor.h
+++ b/Engine/Source/Editor/MovieSceneTools/Private/TrackEditors/PrimitiveMaterialTrackEditor.h
@@ -21,7 +21,10 @@ public:
 
 	//~ ISequencerTrackEditor interface
 	virtual void ExtendObjectBindingTrackMenu(TSharedRef<FExtender> Extender, const TArray<FGuid>& ObjectBindings, const UClass* ObjectClass) override;
+	virtual void BuildTrackContextMenu(FMenuBuilder& MenuBuilder, UMovieSceneTrack* Track) override;
 
 	void ConstructObjectBindingTrackMenu(FMenuBuilder& MenuBuilder, TArray<FGuid> ObjectBindings);
 	void CreateTrackForElement(TArray<FGuid> ObjectBindingIDs, FComponentMaterialInfo MaterialInfo);
+	/** Callback for rebinding a component material track to a different material slot */
+	void FillRebindMaterialTrackMenu(FMenuBuilder& MenuBuilder, class UMovieScenePrimitiveMaterialTrack* MaterialTrack, class UPrimitiveComponent* Component, FGuid ObjectBinding);
 };
diff --git a/Engine/Source/Editor/MovieSceneTools/Public/TrackEditors/MaterialTrackEditor.h b/Engine/Source/Editor/MovieSceneTools/Public/TrackEditors/MaterialTrackEditor.h
--- a/Engine/Source/Editor/MovieSceneTools/Public/TrackEditors/MaterialTrackEditor.h
+++ b/Engine/Source/Editor/MovieSceneTools/Public/TrackEditors/MaterialTrackEditor.h
@@ -93,6 +93,7 @@ public:
 	virtual void ExtendObjectBindingTrackMenu(TSharedRef<FExtender> Extender, const TArray<FGuid>& ObjectBindings, const UClass* ObjectClass) override;
 	virtual bool SupportsType( TSubclassOf<UMovieSceneTrack> Type ) const override;
 	virtual bool GetDefaultExpansionState(UMovieSceneTrack* InTrack) const override;
+	virtual void BuildTrackContextMenu(FMenuBuilder& MenuBuilder, UMovieSceneTrack* Track) override;
 
 protected:
 
@@ -105,5 +106,7 @@ private:
 	void ConstructObjectBindingTrackMenu(FMenuBuilder& MenuBuilder, TArray<FGuid> ObjectBindings);
 	/** Callback for executing the add component material track. */
 	void HandleAddComponentMaterialActionExecute(USceneComponent* Component, FComponentMaterialInfo MaterialInfo);
+	/** Callback for rebinding a component material track to a different material slot */
+	void FillRebindMaterialTrackMenu(FMenuBuilder& MenuBuilder, class UMovieSceneComponentMaterialTrack* MaterialTrack, class UPrimitiveComponent* Component, FGuid ObjectBinding);
 
 };
diff --git a/Engine/Source/Editor/Sequencer/Private/MVVM/ViewModels/TrackModel.cpp b/Engine/Source/Editor/Sequencer/Private/MVVM/ViewModels/TrackModel.cpp
--- a/Engine/Source/Editor/Sequencer/Private/MVVM/ViewModels/TrackModel.cpp
+++ b/Engine/Source/Editor/Sequencer/Private/MVVM/ViewModels/TrackModel.cpp
@@ -530,40 +530,25 @@ FSlateColor FTrackModel::GetLabelColor() const
 		return FSlateColor::UseForeground();
 	}
 
-	// Display track node is red if the property track is not bound to valid property
-	if (UMovieScenePropertyTrack* PropertyTrack = Cast<UMovieScenePropertyTrack>(Track))
+	FMovieSceneLabelParams LabelParams;
+	LabelParams.bIsDimmed = IsDimmed();
+	if (TViewModelPtr<FSequenceModel> SequenceModel = FindAncestorOfType<FSequenceModel>())
 	{
-		const bool bIsDimmed = IsDimmed();
-
-		// 3D transform tracks don't map to property bindings as below
-		if (Track->IsA<UMovieScene3DTransformTrack>() || Track->IsA<UMovieScenePrimitiveMaterialTrack>())
-		{
-			return FOutlinerItemModel::GetLabelColor();
-		}
-
-		// If there is no object binding extension, don't tint it
-		TSharedPtr<IObjectBindingExtension> ParentBinding = FindAncestorOfType<IObjectBindingExtension>();
-		if (!ParentBinding)
-		{
-			return FOutlinerItemModel::GetLabelColor();
-		}
-
-		// Return a normal colour if we have at least one bound object for which the property binding resolves
-		// correctly. Otherwise, return a red colour indicating a binding issue.
-		TArray<UObject*> BoundObjects;
-		FindBoundObjects(BoundObjects);
-		for (UObject* BoundObject : BoundObjects)
+		if (TSharedPtr<FSequencerEditorViewModel> SequencerModel = SequenceModel->GetEditor())
 		{
-			FTrackInstancePropertyBindings PropertyBinding(PropertyTrack->GetPropertyName(), PropertyTrack->GetPropertyPath().ToString());
-			if (PropertyBinding.GetProperty(*BoundObject))
+			LabelParams.SequenceID = SequenceModel->GetSequenceID();
+			LabelParams.Player = SequencerModel->GetSequencer().Get();
+			if (LabelParams.Player)
 			{
-				return bIsDimmed ? FSlateColor::UseSubduedForeground() : FSlateColor::UseForeground();
+				if (TViewModelPtr<FObjectBindingModel> ObjectBindingModel = FindAncestorOfType<FObjectBindingModel>())
+				{
+					LabelParams.BindingID = ObjectBindingModel->GetObjectGuid();
+				}
 			}
 		}
-		return bIsDimmed ? FSlateColor(FLinearColor::Red.Desaturate(0.6f)) : FLinearColor::Red;
 	}
 
-	return FOutlinerItemModel::GetLabelColor();
+	return Track->GetLabelColor(LabelParams);
 }
 
 FText FTrackModel::GetLabelToolTipText() const
@@ -574,33 +559,25 @@ FText FTrackModel::GetLabelToolTipText() const
 		return FText();
 	}
 
-	if (UMovieScenePropertyTrack* PropertyTrack = Cast<UMovieScenePropertyTrack>(Track))
+	FMovieSceneLabelParams LabelParams;
+	LabelParams.bIsDimmed = IsDimmed();
+	if (TViewModelPtr<FSequenceModel> SequenceModel = FindAncestorOfType<FSequenceModel>())
 	{
-		TArray<UObject*> BoundObjects;
-		FindBoundObjects(BoundObjects);
-		for (UObject* BoundObject : BoundObjects)
+		if (TSharedPtr<FSequencerEditorViewModel> SequencerModel = SequenceModel->GetEditor())
 		{
-			FTrackInstancePropertyBindings PropertyBinding(PropertyTrack->GetPropertyName(), PropertyTrack->GetPropertyPath().ToString());
-			if (FProperty* BoundProperty = PropertyBinding.GetProperty(*BoundObject))
+			LabelParams.SequenceID = SequenceModel->GetSequenceID();
+			LabelParams.Player = SequencerModel->GetSequencer().Get();
+			if (LabelParams.Player)
 			{
-				FString PropertyName = BoundProperty->GetMetaData(TEXT("DisplayName"));
-				if (PropertyName.IsEmpty())
+				if (TViewModelPtr<FObjectBindingModel> ObjectBindingModel = FindAncestorOfType<FObjectBindingModel>())
 				{
-					PropertyName = BoundProperty->GetName();
+					LabelParams.BindingID = ObjectBindingModel->GetObjectGuid();
 				}
-
-				FString CategoryName = BoundProperty->GetMetaData(TEXT("Category")).Replace(TEXT("|"), TEXT(" \u00BB "));
-				if (!CategoryName.IsEmpty())
-				{
-					CategoryName.Append(TEXT(" \u00BB "));
-				}
-
-				return FText::FromString(FString::Printf(TEXT("%s%s\n(Path: %s)"), *CategoryName, *PropertyName, *PropertyBinding.GetPropertyPath()));
+				return Track->GetDisplayNameToolTipText(LabelParams);
 			}
 		}
 	}
-	
-	return Track->GetDisplayNameToolTipText();
+	return FText();
 }
 
 TSharedRef<SWidget> FTrackModel::CreateOutlinerView(const FCreateOutlinerViewParams& InParams)
diff --git a/Engine/Source/Editor/Sequencer/Private/MVVM/ViewModels/TrackRowModel.cpp b/Engine/Source/Editor/Sequencer/Private/MVVM/ViewModels/TrackRowModel.cpp
--- a/Engine/Source/Editor/Sequencer/Private/MVVM/ViewModels/TrackRowModel.cpp
+++ b/Engine/Source/Editor/Sequencer/Private/MVVM/ViewModels/TrackRowModel.cpp
@@ -195,37 +195,15 @@ FSlateColor FTrackRowModel::GetLabelColor() const
 	{
 		return FSlateColor::UseForeground();
 	}
-
-	// Display track node is red if the property track is not bound to valid property
-	if (UMovieScenePropertyTrack* PropertyTrack = Cast<UMovieScenePropertyTrack>(Track))
+	FMovieSceneLabelParams LabelParams;
+	LabelParams.bIsDimmed = IsDimmed();
+	LabelParams.Player = Sequencer.Get();
+	LabelParams.SequenceID = SequenceModel->GetSequenceID();
+	if (TViewModelPtr<FObjectBindingModel> ObjectBindingModel = FindAncestorOfType<FObjectBindingModel>())
 	{
-		const bool bIsDimmed = IsDimmed();
-
-		// 3D transform tracks don't map to property bindings as below
-		if (Track->IsA<UMovieScene3DTransformTrack>() || Track->IsA<UMovieScenePrimitiveMaterialTrack>())
-		{
-			return bIsDimmed ? FSlateColor::UseSubduedForeground() : FSlateColor::UseForeground();
-		}
-
-		if (TSharedPtr<IObjectBindingExtension> ParentBinding = FindAncestorOfType<IObjectBindingExtension>())
-		{
-			for (TWeakObjectPtr<> WeakObject : Sequencer->FindBoundObjects(ParentBinding->GetObjectGuid(), Sequencer->GetFocusedTemplateID()))
-			{
-				if (UObject* Object = WeakObject.Get())
-				{
-					FTrackInstancePropertyBindings PropertyBinding(PropertyTrack->GetPropertyName(), PropertyTrack->GetPropertyPath().ToString());
-					if (PropertyBinding.GetProperty(*Object))
-					{
-						return bIsDimmed ? FSlateColor::UseSubduedForeground() : FSlateColor::UseForeground();
-					}
-				}
-			}
-
-			return bIsDimmed ? FSlateColor(FLinearColor::Red.Desaturate(0.6f)) : FLinearColor::Red;
-		}
+		LabelParams.BindingID = ObjectBindingModel->GetObjectGuid();
 	}
-
-	return FOutlinerItemModel::GetLabelColor();
+	return Track->GetLabelColor(LabelParams);
 }
 
 TSharedRef<SWidget> FTrackRowModel::CreateOutlinerView(const FCreateOutlinerViewParams& InParams)
diff --git a/Engine/Source/Editor/SequencerCore/Private/MVVM/Views/SOutlinerItemViewBase.cpp b/Engine/Source/Editor/SequencerCore/Private/MVVM/Views/SOutlinerItemViewBase.cpp
--- a/Engine/Source/Editor/SequencerCore/Private/MVVM/Views/SOutlinerItemViewBase.cpp
+++ b/Engine/Source/Editor/SequencerCore/Private/MVVM/Views/SOutlinerItemViewBase.cpp
@@ -252,7 +252,8 @@ FText SOutlinerItemViewBase::GetLabel() const
 
 FSlateColor SOutlinerItemViewBase::GetLabelColor() const
 {
-	return IsDimmed() ? FSlateColor::UseSubduedForeground() : FSlateColor::UseForeground();
+	TViewModelPtr<IOutlinerExtension> Outliner = WeakOutlinerExtension.Pin();
+	return Outliner ? Outliner->GetLabelColor() : (IsDimmed() ? FSlateColor::UseSubduedForeground() : FSlateColor::UseForeground());
 }
 
 FSlateFontInfo SOutlinerItemViewBase::GetLabelFont() const
diff --git a/Engine/Source/Runtime/MovieScene/Public/MovieSceneTrack.h b/Engine/Source/Runtime/MovieScene/Public/MovieSceneTrack.h
--- a/Engine/Source/Runtime/MovieScene/Public/MovieSceneTrack.h
+++ b/Engine/Source/Runtime/MovieScene/Public/MovieSceneTrack.h
@@ -25,6 +25,10 @@
 #include "UObject/UnrealNames.h"
 #include "UObject/UnrealType.h"
 
+#if WITH_EDITOR
+#include "Styling/SlateColor.h"
+#endif
+
 #include "MovieSceneTrack.generated.h"
 
 class UMovieSceneSection;
@@ -152,6 +156,15 @@ enum class EMovieSceneSectionMovedResult
 	SectionsChanged = 1
 };
 ENUM_CLASS_FLAGS(EMovieSceneSectionMovedResult);
+
+/** Parameters for helping to determine dynamic label color/tooltip*/
+struct FMovieSceneLabelParams
+{
+	class IMovieScenePlayer* Player = nullptr;
+	FGuid BindingID;
+	FMovieSceneSequenceID SequenceID;
+	bool bIsDimmed = false;
+};
 #endif
 
 /**
@@ -447,7 +460,14 @@ public:
 	 *
 	 * @return Display tooltip text.
 	 */
-	virtual FText GetDisplayNameToolTipText() const { return FText::GetEmpty(); }
+	virtual FText GetDisplayNameToolTipText(const FMovieSceneLabelParams& LabelParams) const { return FText::GetEmpty(); }
+
+	/**
+	 * Gets the track label's color.
+	 * 
+	 * @return Track label color.
+	 */
+	virtual FSlateColor GetLabelColor(const FMovieSceneLabelParams& LabelParams) const { return FSlateColor::UseForeground(); }
 
 	/**
 	 * Get this track's color tint.
diff --git a/Engine/Source/Runtime/MovieSceneTracks/Private/Tracks/MovieScene3DTransformTrack.cpp b/Engine/Source/Runtime/MovieSceneTracks/Private/Tracks/MovieScene3DTransformTrack.cpp
--- a/Engine/Source/Runtime/MovieSceneTracks/Private/Tracks/MovieScene3DTransformTrack.cpp
+++ b/Engine/Source/Runtime/MovieSceneTracks/Private/Tracks/MovieScene3DTransformTrack.cpp
@@ -327,5 +327,10 @@ TArray<FTrajectoryKey> UMovieScene3DTransformTrack::GetTrajectoryData(FFrameNumb
 	return Result;
 }
 
+FSlateColor UMovieScene3DTransformTrack::GetLabelColor(const FMovieSceneLabelParams& LabelParams) const
+{
+	return LabelParams.bIsDimmed ? FSlateColor::UseSubduedForeground() : FSlateColor::UseForeground();
+}
+
 #endif	// WITH_EDITOR
 
diff --git a/Engine/Source/Runtime/MovieSceneTracks/Private/Tracks/MovieSceneMaterialTrack.cpp b/Engine/Source/Runtime/MovieSceneTracks/Private/Tracks/MovieSceneMaterialTrack.cpp
--- a/Engine/Source/Runtime/MovieSceneTracks/Private/Tracks/MovieSceneMaterialTrack.cpp
+++ b/Engine/Source/Runtime/MovieSceneTracks/Private/Tracks/MovieSceneMaterialTrack.cpp
@@ -8,9 +8,16 @@
 #include "EntitySystem/BuiltInComponentTypes.h"
 #include "MovieSceneTracksComponentTypes.h"
 #include "Sections/MovieSceneComponentMaterialParameterSection.h"
+#include "IMovieScenePlayer.h"
+#include "Components/DecalComponent.h"
+#include "Components/PrimitiveComponent.h"
+#include "Components/MeshComponent.h"
+#include "Materials/MaterialInterface.h"
 
 #include UE_INLINE_GENERATED_CPP_BY_NAME(MovieSceneMaterialTrack)
 
+#define LOCTEXT_NAMESPACE "MaterialTrack"
+
 UMovieSceneMaterialTrack::UMovieSceneMaterialTrack(const FObjectInitializer& ObjectInitializer)
 	: Super(ObjectInitializer)
 {
@@ -288,6 +295,182 @@ FText UMovieSceneComponentMaterialTrack::GetDefaultDisplayName() const
 	return FText::FromString(FString::Printf(TEXT("Material Element %i"), MaterialIndex_DEPRECATED));
 }
 
+FText UMovieSceneComponentMaterialTrack::GetDisplayNameToolTipText(const FMovieSceneLabelParams& LabelParams) const
+{
+	if (LabelParams.Player)
+	{
+		UMaterialInterface* Material = nullptr;
+		for (TWeakObjectPtr<> WeakObject : LabelParams.Player->FindBoundObjects(LabelParams.BindingID, LabelParams.SequenceID))
+		{
+			switch (MaterialInfo.MaterialType)
+			{
+			case EComponentMaterialType::Empty:
+				break;
+			case EComponentMaterialType::IndexedMaterial:
+				if (UPrimitiveComponent* Component = Cast<UPrimitiveComponent>(WeakObject.Get()))
+				{
+					if (!MaterialInfo.MaterialSlotName.IsNone())
+					{
+						Material = Component->GetMaterialByName(MaterialInfo.MaterialSlotName);
+						int32 FoundMaterialIndex = Component->GetMaterialIndex(MaterialInfo.MaterialSlotName);
+						if (Material)
+						{
+							// Found a binding by slot name
+							return (FoundMaterialIndex == MaterialInfo.MaterialSlotIndex) ? FText::Format(LOCTEXT("SlotMaterialTrackTooltip", "Material parameter track for slot {0} at index {1}, asset {2}"), FText::FromName(MaterialInfo.MaterialSlotName), FText::AsNumber(MaterialInfo.MaterialSlotIndex), FText::FromString(Material->GetName()))
+								: FText::Format(LOCTEXT("SlotMaterialTrackTooltip_WrongIndexCached", "Material parameter track for slot {0} at index {1}, asset {2}. \n Asset has cached material index {3} that does not match material found with slot name. This can be corrected in right-click menu."), FText::FromName(MaterialInfo.MaterialSlotName), FText::AsNumber(FoundMaterialIndex), FText::FromString(Material->GetName()), FText::AsNumber(MaterialInfo.MaterialSlotIndex));
+						}
+					}
+					if (!Material)
+					{
+						// Couldn't find binding by slot name or one wasn't specified, try by index
+						Material = Component->GetMaterial(MaterialInfo.MaterialSlotIndex);
+						if (Material)
+						{
+							// Found by index, but may want to rebind to new slot name
+							return MaterialInfo.MaterialSlotName.IsNone() ? FText::Format(LOCTEXT("IndexedMaterialTrackTooltip_NoSlotName", "Material parameter track for element at index {0}, asset {1}. \nBinding is missing Material Slot Name. This can be corrected in right-click menu."), FText::AsNumber(MaterialInfo.MaterialSlotIndex), FText::FromString(Material->GetName()))
+								: FText::Format(LOCTEXT("IndexedMaterialTrackTooltip_IncorrectSlotName", "Material parameter track for element at index {0}, asset {1}. \n Binding references missing Material Slot Name {2}. This can be rebound in right-click menu."), FText::AsNumber(MaterialInfo.MaterialSlotIndex), FText::FromString(Material->GetName()), FText::FromName(MaterialInfo.MaterialSlotName));
+						}
+						else
+						{
+							// Couldn't find binding by slot or index
+							return FText::Format(LOCTEXT("SlotMaterialTrackTooltip_MissingMaterial", "Material could not be found for slot {0} and cached index {1}"), FText::FromName(MaterialInfo.MaterialSlotName), FText::AsNumber(MaterialInfo.MaterialSlotIndex));
+						}
+					}
+				}
+				break;
+			case EComponentMaterialType::OverlayMaterial:
+				if (UMeshComponent* Component = Cast<UMeshComponent>(WeakObject.Get()))
+				{
+					Material = Component->GetOverlayMaterial();
+					if (Material)
+					{
+						return FText::Format(LOCTEXT("OverlayMaterialTrackTooltip", "Material parameter track for overlay material {0}"), FText::FromString(Material->GetName()));
+					}
+					else
+					{
+						return LOCTEXT("OverlayMaterialTrackTooltip_MissingMaterial", "No overlay material could be found");
+					}
+				}
+				break;
+			case EComponentMaterialType::DecalMaterial:
+				if (UDecalComponent* DecalComponent = Cast<UDecalComponent>(WeakObject.Get()))
+				{
+					Material = DecalComponent->GetDecalMaterial();
+					if (Material)
+					{
+						return FText::Format(LOCTEXT("DecalMaterialTrackTooltip", "Material parameter track for decal material {0}"), FText::FromString(Material->GetName()));
+					}
+					else
+					{
+						return LOCTEXT("DecalMaterialTrackTooltip_MissingMaterial", "No decal material could be found");
+					}
+				}
+				break;
+			default:
+				break;
+			}
+		}
+	}
+	return FText();
+}
+
+FSlateColor UMovieSceneComponentMaterialTrack::GetLabelColor(const FMovieSceneLabelParams& LabelParams) const
+{
+	auto GetDimmedColor = [&LabelParams](FLinearColor LinearColor)
+	{
+		return FSlateColor(LabelParams.bIsDimmed ? LinearColor.Desaturate(0.6f) : LinearColor);
+	};
+	if (LabelParams.Player)
+	{
+		for (TWeakObjectPtr<> WeakObject : LabelParams.Player->FindBoundObjects(LabelParams.BindingID, LabelParams.SequenceID))
+		{
+			switch (MaterialInfo.MaterialType)
+			{
+			case EComponentMaterialType::Empty:
+				break;
+			case EComponentMaterialType::IndexedMaterial:
+				if (UPrimitiveComponent* Component = Cast<UPrimitiveComponent>(WeakObject.Get()))
+				{
+					if (MaterialInfo.MaterialSlotName.IsNone())
+					{
+						// No material slot name. Can we find by index?
+						if (Component->GetMaterial(MaterialInfo.MaterialSlotIndex))
+						{
+							// Found by index, but show yellow because missing slot name
+							return GetDimmedColor(FLinearColor::Yellow);
+						}
+						else
+						{
+							// No material found at all
+							return GetDimmedColor(FLinearColor::Red);
+						}
+					}
+					else
+					{
+						UMaterialInterface* Material = Component->GetMaterialByName(MaterialInfo.MaterialSlotName);
+						int32 FoundMaterialIndex = Component->GetMaterialIndex(MaterialInfo.MaterialSlotName);
+						if (Material)
+						{
+							if (FoundMaterialIndex != MaterialInfo.MaterialSlotIndex)
+							{
+								// Indices don't match- yellow
+								return GetDimmedColor(FLinearColor::Yellow);
+							}
+						}
+						else
+						{
+							// Can't find material by name, try by index
+							Material = Component->GetMaterial(MaterialInfo.MaterialSlotIndex);
+							if (Material)
+							{
+								// Found, but didn't match name. Yellow
+								return GetDimmedColor(FLinearColor::Yellow);
+							}
+							else
+							{
+								// Nothing found. Red.
+								return GetDimmedColor(FLinearColor::Red);
+							}
+						}
+					}
+				}
+				break;
+			case EComponentMaterialType::OverlayMaterial:
+				if (UMeshComponent* Component = Cast<UMeshComponent>(WeakObject.Get()))
+				{
+					UMaterialInterface* Material = Component->GetOverlayMaterial();
+					if (!Material)
+					{
+						return GetDimmedColor(FLinearColor::Red);
+					}
+				}
+				else
+				{
+					return GetDimmedColor(FLinearColor::Red);
+				}
+				break;
+			case EComponentMaterialType::DecalMaterial:
+				if (UDecalComponent* DecalComponent = Cast<UDecalComponent>(WeakObject.Get()))
+				{
+					UMaterialInterface* Material = DecalComponent->GetDecalMaterial();
+					if (!Material)
+					{
+						return GetDimmedColor(FLinearColor::Red);
+					}
+				}
+				else
+				{
+					return GetDimmedColor(FLinearColor::Red);
+				}
+				break;
+			default:
+				break;
+			}
+		}
+	}
+	return LabelParams.bIsDimmed ? FSlateColor::UseSubduedForeground() : FSlateColor::UseForeground();
+}
+
 void UMovieSceneComponentMaterialTrack::PostLoad()
 {
 	Super::PostLoad();
@@ -300,3 +483,4 @@ void UMovieSceneComponentMaterialTrack::PostLoad()
 }
 #endif
 
+#undef LOCTEXT_NAMESPACE
\ No newline at end of file
diff --git a/Engine/Source/Runtime/MovieSceneTracks/Private/Tracks/MovieScenePrimitiveMaterialTrack.cpp b/Engine/Source/Runtime/MovieSceneTracks/Private/Tracks/MovieScenePrimitiveMaterialTrack.cpp
--- a/Engine/Source/Runtime/MovieSceneTracks/Private/Tracks/MovieScenePrimitiveMaterialTrack.cpp
+++ b/Engine/Source/Runtime/MovieSceneTracks/Private/Tracks/MovieScenePrimitiveMaterialTrack.cpp
@@ -4,9 +4,14 @@
 #include "Tracks/MovieSceneMaterialTrack.h"
 #include "Sections/MovieScenePrimitiveMaterialSection.h"
 #include "Evaluation/MovieSceneEvaluationTrack.h"
+#include "Components/DecalComponent.h"
+#include "Components/PrimitiveComponent.h"
+#include "Components/MeshComponent.h"
 
 #include UE_INLINE_GENERATED_CPP_BY_NAME(MovieScenePrimitiveMaterialTrack)
 
+#define LOCTEXT_NAMESPACE "PrimitiveMaterialTrack"
+
 
 UMovieScenePrimitiveMaterialTrack::UMovieScenePrimitiveMaterialTrack(const FObjectInitializer& ObjInit)
 	: Super(ObjInit)
@@ -30,6 +35,156 @@ bool UMovieScenePrimitiveMaterialTrack::SupportsType(TSubclassOf<UMovieSceneSect
 }
 
 #if WITH_EDITORONLY_DATA
+
+FText UMovieScenePrimitiveMaterialTrack::GetDisplayNameToolTipText(const FMovieSceneLabelParams& LabelParams) const
+{
+	if (LabelParams.Player)
+	{
+		UMaterialInterface* Material = nullptr;
+		for (TWeakObjectPtr<> WeakObject : LabelParams.Player->FindBoundObjects(LabelParams.BindingID, LabelParams.SequenceID))
+		{
+			switch (MaterialInfo.MaterialType)
+			{
+			case EComponentMaterialType::Empty:
+				break;
+			case EComponentMaterialType::IndexedMaterial:
+				if (UPrimitiveComponent* Component = Cast<UPrimitiveComponent>(WeakObject.Get()))
+				{
+					if (!MaterialInfo.MaterialSlotName.IsNone())
+					{
+						Material = Component->GetMaterialByName(MaterialInfo.MaterialSlotName);
+						int32 FoundMaterialIndex = Component->GetMaterialIndex(MaterialInfo.MaterialSlotName);
+						if (Material)
+						{
+							// Found a binding by slot name
+							return FText::Format(LOCTEXT("SlotMaterialSwitcherTrackTooltip", "Material switcher for {0} at index {1}"), FText::FromName(MaterialInfo.MaterialSlotName), FText::AsNumber(MaterialInfo.MaterialSlotIndex));
+						}
+					}
+					if (!Material)
+					{
+						// Couldn't find binding by slot name or one wasn't specified, try by index
+						Material = Component->GetMaterial(MaterialInfo.MaterialSlotIndex);
+						if (Material)
+						{
+							// Found by index, but may want to rebind to new slot name
+							return MaterialInfo.MaterialSlotName.IsNone() ? FText::Format(LOCTEXT("IndexedMaterialSwitcherTrackTooltip_MissingSlotName", "Material switcher for element at index {0}\nBinding is missing Material Slot Name. This can be corrected in right-click menu."), FText::AsNumber(MaterialInfo.MaterialSlotIndex))
+								: FText::Format(LOCTEXT("IndexedMaterialSwitcherTrackTooltip_IncorrectSlotName", "Material switcher for element at index {0} \n Binding references missing Material Slot Name {1}. This can be rebound in right-click menu."), FText::AsNumber(MaterialInfo.MaterialSlotIndex), FText::FromName(MaterialInfo.MaterialSlotName));
+						}
+						else
+						{
+							// Couldn't find binding by slot or index
+							return FText::Format(LOCTEXT("SlotMaterialSwitcherTrackTooltip_MissingMaterial", "Material could not be found for slot {0} and cached index {1}"), FText::FromName(MaterialInfo.MaterialSlotName), FText::AsNumber(MaterialInfo.MaterialSlotIndex));
+						}
+					}
+				}
+				break;
+			case EComponentMaterialType::OverlayMaterial:
+				if (UMeshComponent* Component = Cast<UMeshComponent>(WeakObject.Get()))
+				{
+					return LOCTEXT("OverlayMaterialSwitcherTrackTooltip", "Material switcher for overlay material");
+				}
+				break;
+			case EComponentMaterialType::DecalMaterial:
+				if (UDecalComponent* DecalComponent = Cast<UDecalComponent>(WeakObject.Get()))
+				{
+					return LOCTEXT("DecalMaterialSwitcherTrackTooltip", "Material switcher for decal material");
+				}
+				break;
+			default:
+				break;
+			}
+		}
+	}
+	return FText();
+}
+
+FSlateColor UMovieScenePrimitiveMaterialTrack::GetLabelColor(const FMovieSceneLabelParams& LabelParams) const
+{
+	auto GetDimmedColor = [&LabelParams](FLinearColor LinearColor)
+	{
+		return FSlateColor(LabelParams.bIsDimmed ? LinearColor.Desaturate(0.6f) : LinearColor);
+	};
+	if (LabelParams.Player)
+	{
+		for (TWeakObjectPtr<> WeakObject : LabelParams.Player->FindBoundObjects(LabelParams.BindingID, LabelParams.SequenceID))
+		{
+			switch (MaterialInfo.MaterialType)
+			{
+			case EComponentMaterialType::Empty:
+				break;
+			case EComponentMaterialType::IndexedMaterial:
+				if (UPrimitiveComponent* Component = Cast<UPrimitiveComponent>(WeakObject.Get()))
+				{
+					if (MaterialInfo.MaterialSlotName.IsNone())
+					{
+						// No material slot name. Can we find by index?
+						if (Component->GetMaterial(MaterialInfo.MaterialSlotIndex))
+						{
+							// Found by index, but show yellow because missing slot name
+							return GetDimmedColor(FLinearColor::Yellow);
+						}
+						else
+						{
+							// No material found at all
+							return GetDimmedColor(FLinearColor::Red);
+						}
+					}
+					else
+					{
+						UMaterialInterface* Material = Component->GetMaterialByName(MaterialInfo.MaterialSlotName);
+						int32 FoundMaterialIndex = Component->GetMaterialIndex(MaterialInfo.MaterialSlotName);
+						if (Material)
+						{
+							if (FoundMaterialIndex != MaterialInfo.MaterialSlotIndex)
+							{
+								// Indices don't match- yellow
+								return GetDimmedColor(FLinearColor::Yellow);
+							}
+						}
+						else
+						{
+							// Can't find material by name, try by index
+							Material = Component->GetMaterial(MaterialInfo.MaterialSlotIndex);
+							if (Material)
+							{
+								// Found, but didn't match name. Yellow
+								return GetDimmedColor(FLinearColor::Yellow);
+							}
+							else
+							{
+								// Nothing found. Red.
+								return GetDimmedColor(FLinearColor::Red);
+							}
+						}
+					}
+				}
+				break;
+			case EComponentMaterialType::OverlayMaterial:
+			{
+				UMeshComponent* MeshComponent = Cast<UMeshComponent>(WeakObject.Get());
+				if (!MeshComponent)
+				{
+					return GetDimmedColor(FLinearColor::Red);
+				}
+			}
+				break;
+			case EComponentMaterialType::DecalMaterial:
+			{
+				UDecalComponent* DecalComponent = Cast<UDecalComponent>(WeakObject.Get());
+				if (!DecalComponent)
+				{
+					return GetDimmedColor(FLinearColor::Red);
+				}
+			}
+				break;
+			default:
+				break;
+			}
+		}
+	}
+	return LabelParams.bIsDimmed ? FSlateColor::UseSubduedForeground() : FSlateColor::UseForeground();
+}
+
 void UMovieScenePrimitiveMaterialTrack::PostLoad()
 {
 	Super::PostLoad();
@@ -62,4 +217,6 @@ void UMovieScenePrimitiveMaterialTrack::SetMaterialIndex(int32 InMaterialIndex)
 	MaterialInfo.MaterialSlotIndex = InMaterialIndex;
 	// Assumption is if this is being called by old code, we should be indexed
 	MaterialInfo.MaterialType = EComponentMaterialType::IndexedMaterial;
-}
\ No newline at end of file
+}
+
+#undef LOCTEXT_NAMESPACE
\ No newline at end of file
diff --git a/Engine/Source/Runtime/MovieSceneTracks/Private/Tracks/MovieScenePropertyTrack.cpp b/Engine/Source/Runtime/MovieSceneTracks/Private/Tracks/MovieScenePropertyTrack.cpp
--- a/Engine/Source/Runtime/MovieSceneTracks/Private/Tracks/MovieScenePropertyTrack.cpp
+++ b/Engine/Source/Runtime/MovieSceneTracks/Private/Tracks/MovieScenePropertyTrack.cpp
@@ -1,6 +1,8 @@
 // Copyright Epic Games, Inc. All Rights Reserved.
 
 #include "Tracks/MovieScenePropertyTrack.h"
+
+#include "IMovieScenePlayer.h"
 #include "Channels/MovieSceneSectionChannelOverrideRegistry.h"
 #include "Algo/Sort.h"
 #include "MovieScene.h"
@@ -79,11 +81,59 @@ FText UMovieScenePropertyTrack::GetDefaultDisplayName() const
 	return FText::FromName(PropertyBinding.PropertyName);
 }
 
-FText UMovieScenePropertyTrack::GetDisplayNameToolTipText() const
+FText UMovieScenePropertyTrack::GetDisplayNameToolTipText(const FMovieSceneLabelParams& LabelParams) const
 {
+	if (!LabelParams.BindingID.IsValid() || !LabelParams.Player)
+	{
+		return FText();
+	}
+	const TArrayView<TWeakObjectPtr<>> FoundBoundObjects = LabelParams.Player->FindBoundObjects(LabelParams.BindingID, LabelParams.SequenceID);
+	for (const TWeakObjectPtr<> BoundObject : FoundBoundObjects)
+	{
+		FTrackInstancePropertyBindings InstancePropertyBinding(GetPropertyName(), GetPropertyPath().ToString());
+		if (FProperty* BoundProperty = InstancePropertyBinding.GetProperty(*BoundObject))
+		{
+			FString PropertyName = BoundProperty->GetMetaData(TEXT("DisplayName"));
+			if (PropertyName.IsEmpty())
+			{
+				PropertyName = BoundProperty->GetName();
+			}
+
+			FString CategoryName = BoundProperty->GetMetaData(TEXT("Category")).Replace(TEXT("|"), TEXT(" \u00BB "));
+			if (!CategoryName.IsEmpty())
+			{
+				CategoryName.Append(TEXT(" \u00BB "));
+			}
+
+			return FText::FromString(FString::Printf(TEXT("%s%s\n(Path: %s)"), *CategoryName, *PropertyName, *InstancePropertyBinding.GetPropertyPath()));
+		}
+	}
+	
 	return FText::FromName(PropertyBinding.PropertyPath);
 }
 
+FSlateColor UMovieScenePropertyTrack::GetLabelColor(const FMovieSceneLabelParams& LabelParams) const
+{
+	// If there is no object binding extension, don't tint it
+	if (!LabelParams.BindingID.IsValid() || !LabelParams.Player)
+	{
+		return LabelParams.bIsDimmed ? FSlateColor::UseSubduedForeground() : FSlateColor::UseForeground();
+	}
+
+	// Return a normal colour if we have at least one bound object for which the property binding resolves
+	// correctly. Otherwise, return a red colour indicating a binding issue.
+	const TArrayView<TWeakObjectPtr<>> FoundBoundObjects = LabelParams.Player->FindBoundObjects(LabelParams.BindingID, LabelParams.SequenceID);
+	for (const TWeakObjectPtr<> BoundObject : FoundBoundObjects)
+	{
+		FTrackInstancePropertyBindings InstancePropertyBinding(GetPropertyName(), GetPropertyPath().ToString());
+		if (InstancePropertyBinding.GetProperty(*BoundObject))
+		{
+			return LabelParams.bIsDimmed ? FSlateColor::UseSubduedForeground() : FSlateColor::UseForeground();
+		}
+	}
+	return LabelParams.bIsDimmed ? FSlateColor(FLinearColor::Red.Desaturate(0.6f)) : FLinearColor::Red;
+}
+
 FName UMovieScenePropertyTrack::GetTrackName() const
 {
 	return UniqueTrackName;
diff --git a/Engine/Source/Runtime/MovieSceneTracks/Public/Tracks/MovieScene3DTransformTrack.h b/Engine/Source/Runtime/MovieSceneTracks/Public/Tracks/MovieScene3DTransformTrack.h
--- a/Engine/Source/Runtime/MovieSceneTracks/Public/Tracks/MovieScene3DTransformTrack.h
+++ b/Engine/Source/Runtime/MovieSceneTracks/Public/Tracks/MovieScene3DTransformTrack.h
@@ -79,7 +79,7 @@ public:
 #if WITH_EDITOR
 
 	MOVIESCENETRACKS_API TArray<FTrajectoryKey> GetTrajectoryData(FFrameNumber Time, int32 MaxNumDataPoints, TRange<FFrameNumber>) const;
-
+	MOVIESCENETRACKS_API FSlateColor GetLabelColor(const FMovieSceneLabelParams& LabelParams) const override;
 #endif
 
 private:
diff --git a/Engine/Source/Runtime/MovieSceneTracks/Public/Tracks/MovieSceneMaterialTrack.h b/Engine/Source/Runtime/MovieSceneTracks/Public/Tracks/MovieSceneMaterialTrack.h
--- a/Engine/Source/Runtime/MovieSceneTracks/Public/Tracks/MovieSceneMaterialTrack.h
+++ b/Engine/Source/Runtime/MovieSceneTracks/Public/Tracks/MovieSceneMaterialTrack.h
@@ -8,6 +8,9 @@
 #include "MovieSceneNameableTrack.h"
 #include "EntitySystem/IMovieSceneEntityProvider.h"
 #include "MaterialTypes.h"
+#if WITH_EDITOR
+#include "Styling/SlateColor.h"
+#endif
 #include "MovieSceneMaterialTrack.generated.h"
 
 UENUM()
@@ -228,8 +231,10 @@ public:
 	}
 
 #if WITH_EDITOR
-	virtual FText GetDisplayNameToolTipText() const override { return TooltipText; }
-	void SetDisplayNameTooltipText(FText InTooltipText) { TooltipText = InTooltipText; }
+	virtual FText GetDisplayNameToolTipText(const FMovieSceneLabelParams& LabelParams) const override;
+
+	// We override label color if material binding is broken/partially broken.
+	virtual FSlateColor GetLabelColor(const FMovieSceneLabelParams& LabelParams) const override;
 #endif
 
 
@@ -243,9 +248,6 @@ private:
 	UPROPERTY(meta = (DeprecatedProperty, DeprecationMessage = "Use MaterialInfo instead."))
 	int32 MaterialIndex_DEPRECATED;
 
-	UPROPERTY()
-	FText TooltipText;
-
 #endif
 private:
 	/** The info on the material this track is animating.*/
diff --git a/Engine/Source/Runtime/MovieSceneTracks/Public/Tracks/MovieScenePrimitiveMaterialTrack.h b/Engine/Source/Runtime/MovieSceneTracks/Public/Tracks/MovieScenePrimitiveMaterialTrack.h
--- a/Engine/Source/Runtime/MovieSceneTracks/Public/Tracks/MovieScenePrimitiveMaterialTrack.h
+++ b/Engine/Source/Runtime/MovieSceneTracks/Public/Tracks/MovieScenePrimitiveMaterialTrack.h
@@ -35,8 +35,10 @@ public:
 	virtual bool SupportsType(TSubclassOf<UMovieSceneSection> SectionClass) const override;
 
 #if WITH_EDITOR
-	virtual FText GetDisplayNameToolTipText() const override { return TooltipText; }
-	void SetDisplayNameTooltipText(FText InTooltipText) { TooltipText = InTooltipText; }
+	virtual FText GetDisplayNameToolTipText(const FMovieSceneLabelParams& LabelParams) const override;
+
+	// We override label color if material binding is broken/partially broken.
+	virtual FSlateColor GetLabelColor(const FMovieSceneLabelParams& LabelParams) const override;
 #endif
 
 protected:
@@ -46,9 +48,6 @@ protected:
 private:
 	UPROPERTY()
 	int32 MaterialIndex_DEPRECATED;
-		
-	UPROPERTY()
-	FText TooltipText;
 
 #endif
 
diff --git a/Engine/Source/Runtime/MovieSceneTracks/Public/Tracks/MovieScenePropertyTrack.h b/Engine/Source/Runtime/MovieSceneTracks/Public/Tracks/MovieScenePropertyTrack.h
--- a/Engine/Source/Runtime/MovieSceneTracks/Public/Tracks/MovieScenePropertyTrack.h
+++ b/Engine/Source/Runtime/MovieSceneTracks/Public/Tracks/MovieScenePropertyTrack.h
@@ -39,7 +39,8 @@ public:
 
 #if WITH_EDITORONLY_DATA
 	MOVIESCENETRACKS_API virtual FText GetDefaultDisplayName() const override;
-	MOVIESCENETRACKS_API virtual FText GetDisplayNameToolTipText() const override;
+	MOVIESCENETRACKS_API virtual FText GetDisplayNameToolTipText(const FMovieSceneLabelParams& LabelParams) const override;
+	MOVIESCENETRACKS_API virtual FSlateColor GetLabelColor(const FMovieSceneLabelParams& LabelParams) const override;
 	virtual bool CanRename() const override { return false; }
 	MOVIESCENETRACKS_API virtual FName GetTrackName() const override;
 #endif
