Sequencer - Allow MovieSceneByteProperty tracks to find UEnums nested in structs in order to properly show the UI for animating these enums.

diff --git a/Engine/Source/Editor/MovieSceneTools/Private/TrackEditors/PropertyTrackEditors/BytePropertyTrackEditor.cpp b/Engine/Source/Editor/MovieSceneTools/Private/TrackEditors/PropertyTrackEditors/BytePropertyTrackEditor.cpp
--- a/Engine/Source/Editor/MovieSceneTools/Private/TrackEditors/PropertyTrackEditors/BytePropertyTrackEditor.cpp
+++ b/Engine/Source/Editor/MovieSceneTools/Private/TrackEditors/PropertyTrackEditors/BytePropertyTrackEditor.cpp
@@ -28,8 +28,7 @@ TSharedRef<ISequencerTrackEditor> FBytePropertyTrackEditor::CreateTrackEditor( T
 	return MakeShareable(new FBytePropertyTrackEditor(OwningSequencer));
 }
 
-
-UEnum* GetEnumForByteTrack(TSharedPtr<ISequencer> Sequencer, const FGuid& OwnerObjectHandle, FName PropertyName, UMovieSceneByteTrack* ByteTrack)
+UEnum* GetEnumForByteTrack(TSharedPtr<ISequencer> Sequencer, const FGuid& OwnerObjectHandle, const FPropertyChangedParams& PropertyChangedParams, UMovieSceneByteTrack* ByteTrack)
 {
 	TSet<UEnum*> PropertyEnums;
 
@@ -41,7 +40,7 @@ UEnum* GetEnumForByteTrack(TSharedPtr<ISequencer> Sequencer, const FGuid& OwnerO
 			continue;
 		}
 
-		FProperty* Property = RuntimeObject->GetClass()->FindPropertyByName(PropertyName);
+		FProperty* Property = PropertyChangedParams.PropertyPath.GetLeafMostProperty().Property.Get();
 		if (Property != nullptr)
 		{
 			UEnum* Enum = nullptr;
@@ -77,23 +76,23 @@ UEnum* GetEnumForByteTrack(TSharedPtr<ISequencer> Sequencer, const FGuid& OwnerO
 }
 
 
-UMovieSceneTrack* FBytePropertyTrackEditor::AddTrack(UMovieScene* FocusedMovieScene, const FGuid& ObjectHandle, TSubclassOf<class UMovieSceneTrack> TrackClass, FName UniqueTypeName)
+void FBytePropertyTrackEditor::GenerateKeysFromPropertyChanged( const FPropertyChangedParams& PropertyChangedParams, UMovieSceneSection* SectionToKey, FGeneratedTrackKeys& OutGeneratedKeys )
 {
-	UMovieSceneTrack* NewTrack = FPropertyTrackEditor::AddTrack(FocusedMovieScene, ObjectHandle, TrackClass, UniqueTypeName);
-	UMovieSceneByteTrack* ByteTrack = Cast<UMovieSceneByteTrack>(NewTrack);
-	UEnum* TrackEnum = GetEnumForByteTrack(GetSequencer(), ObjectHandle, UniqueTypeName, ByteTrack);
+	uint8 KeyedValue = PropertyChangedParams.GetPropertyValue<uint8>();
+	OutGeneratedKeys.Add(FMovieSceneChannelValueSetter::Create<FMovieSceneByteChannel>(0, KeyedValue, true));
+}
 
-	if (TrackEnum != nullptr)
+void FBytePropertyTrackEditor::InitializeNewTrack(UMovieSceneByteTrack* NewTrack, FPropertyChangedParams PropertyChangedParams)
+{
+	FPropertyTrackEditor<UMovieSceneByteTrack>::InitializeNewTrack(NewTrack, PropertyChangedParams);
+	if (NewTrack)
 	{
-		ByteTrack->SetEnum(TrackEnum);
-	}
-
-	return NewTrack;
-}
+		UEnum* TrackEnum = GetEnumForByteTrack(GetSequencer(), NewTrack->FindObjectBindingGuid(), PropertyChangedParams, NewTrack);
 
+		if (TrackEnum != nullptr)
+		{
+			NewTrack->SetEnum(TrackEnum);
+		}
+	}
 
-void FBytePropertyTrackEditor::GenerateKeysFromPropertyChanged( const FPropertyChangedParams& PropertyChangedParams, UMovieSceneSection* SectionToKey, FGeneratedTrackKeys& OutGeneratedKeys )
-{
-	uint8 KeyedValue = PropertyChangedParams.GetPropertyValue<uint8>();
-	OutGeneratedKeys.Add(FMovieSceneChannelValueSetter::Create<FMovieSceneByteChannel>(0, KeyedValue, true));
 }
diff --git a/Engine/Source/Editor/MovieSceneTools/Private/TrackEditors/PropertyTrackEditors/BytePropertyTrackEditor.h b/Engine/Source/Editor/MovieSceneTools/Private/TrackEditors/PropertyTrackEditors/BytePropertyTrackEditor.h
--- a/Engine/Source/Editor/MovieSceneTools/Private/TrackEditors/PropertyTrackEditors/BytePropertyTrackEditor.h
+++ b/Engine/Source/Editor/MovieSceneTools/Private/TrackEditors/PropertyTrackEditors/BytePropertyTrackEditor.h
@@ -61,15 +61,11 @@ public:
 	 */
 	static TSharedRef<ISequencerTrackEditor> CreateTrackEditor(TSharedRef<ISequencer> OwningSequencer);
 
-public:
-
-	//~ ISequencerTrackEditor interface
-
-	virtual UMovieSceneTrack* AddTrack(UMovieScene* FocusedMovieScene, const FGuid& ObjectHandle, TSubclassOf<class UMovieSceneTrack> TrackClass, FName UniqueTypeName) override;
-
 protected:
 
 	//~ FPropertyTrackEditor interface
 
 	virtual void GenerateKeysFromPropertyChanged(const FPropertyChangedParams& PropertyChangedParams, UMovieSceneSection* SectionToKey, FGeneratedTrackKeys& OutGeneratedKeys) override;
+	virtual void InitializeNewTrack(UMovieSceneByteTrack* NewTrack, FPropertyChangedParams PropertyChangedParams) override;
+
 };
