Sequencer: Fixed transform origins being removed from non-absolute auto-key operations

This code was correct for absolute sections, but not for any other blend types.

diff --git a/Engine/Source/Editor/MovieSceneTools/Private/TrackEditors/TransformTrackEditor.cpp b/Engine/Source/Editor/MovieSceneTools/Private/TrackEditors/TransformTrackEditor.cpp
--- a/Engine/Source/Editor/MovieSceneTools/Private/TrackEditors/TransformTrackEditor.cpp
+++ b/Engine/Source/Editor/MovieSceneTools/Private/TrackEditors/TransformTrackEditor.cpp
@@ -1028,7 +1028,7 @@ FTransformData F3DTransformTrackEditor::RecomposeTransform(const FTransformData&
 	}
 
 	// Account for the transform origin only if this is not parented because the transform origin is already being applied to the parent.
-	if (!SceneComponent->GetAttachParent())
+	if (!SceneComponent->GetAttachParent() && Section->GetBlendType() == EMovieSceneBlendType::Absolute)
 	{
 		CurrentTransform *= GetTransformOrigin().Inverse();
 	}
