Sequencer: Add toggle lock camera command to curve editor shared bindings

diff --git a/Engine/Source/Editor/MovieSceneTools/Private/TrackEditors/CameraCutTrackEditor.cpp b/Engine/Source/Editor/MovieSceneTools/Private/TrackEditors/CameraCutTrackEditor.cpp
--- a/Engine/Source/Editor/MovieSceneTools/Private/TrackEditors/CameraCutTrackEditor.cpp
+++ b/Engine/Source/Editor/MovieSceneTools/Private/TrackEditors/CameraCutTrackEditor.cpp
@@ -101,6 +101,12 @@ void FCameraCutTrackEditor::BindCommands(TSharedRef<FUICommandList> SequencerCom
 		Commands.ToggleLockCamera,
 		FExecuteAction::CreateSP( this, &FCameraCutTrackEditor::ToggleLockCamera) );
 
+	TSharedPtr<FUICommandList> CurveEditorSharedBindings = GetSequencer()->GetCommandBindings(ESequencerCommandBindings::CurveEditor);
+	if (CurveEditorSharedBindings)
+	{	
+		CurveEditorSharedBindings->MapAction(Commands.ToggleLockCamera, *SequencerCommandBindings->GetActionForCommand(Commands.ToggleLockCamera));
+	}
+
 	Commands.BindingCount++;
 }
 
