Sequencer: Fix issue where Movie Scene Event UObject payload parameters were not properly creating references to the UObjects and therefore renaming those UObjects would not create redirectors and the references would be lost.

diff --git a/Engine/Source/Editor/MovieSceneTools/Private/MovieSceneDirectorBlueprintEndpointCustomization.cpp b/Engine/Source/Editor/MovieSceneTools/Private/MovieSceneDirectorBlueprintEndpointCustomization.cpp
--- a/Engine/Source/Editor/MovieSceneTools/Private/MovieSceneDirectorBlueprintEndpointCustomization.cpp
+++ b/Engine/Source/Editor/MovieSceneTools/Private/MovieSceneDirectorBlueprintEndpointCustomization.cpp
@@ -304,13 +304,21 @@ void FMovieSceneDirectorBlueprintEndpointCustomization::CustomizeChildren(TShare
 					continue;
 				}
 
-				const FString* PayloadVariable = PayloadVariables.Find(Field->GetFName());
+				const FMovieSceneDirectorBlueprintVariableValue* PayloadVariable = PayloadVariables.Find(Field->GetFName());
 				if (PayloadVariable)
 				{
 					AllValidNames.Add(Field->GetFName());
 
 					// We have an override for this variable
-					const bool bImportSuccess = FBlueprintEditorUtils::PropertyValueFromString(Field, *PayloadVariable, StructData->GetStructMemory());
+					bool bImportSuccess = false;
+					if (PayloadVariable->ObjectValue)
+					{
+						bImportSuccess = FBlueprintEditorUtils::PropertyValueFromString(Field, PayloadVariable->ObjectValue->GetPathName(), StructData->GetStructMemory());
+					}
+					if (!bImportSuccess)
+					{
+						bImportSuccess = FBlueprintEditorUtils::PropertyValueFromString(Field, *PayloadVariable->Value, StructData->GetStructMemory());
+					}
 					if (!bImportSuccess)
 					{
 						// @todo: error
@@ -318,13 +326,18 @@ void FMovieSceneDirectorBlueprintEndpointCustomization::CustomizeChildren(TShare
 				}
 				else if (UEdGraphPin* Pin = CommonEndpoint->FindPin(Field->GetFName(), EGPD_Output))
 				{
-					if (!Pin->DefaultValue.IsEmpty())
+					bool bImportSuccess = false;
+					if (Pin->DefaultObject)
 					{
-						const bool bImportSuccess = FBlueprintEditorUtils::PropertyValueFromString(Field, Pin->DefaultValue, StructData->GetStructMemory());
-						if (!bImportSuccess)
-						{
-							// @todo: error
-						}
+						bImportSuccess = FBlueprintEditorUtils::PropertyValueFromString(Field, Pin->DefaultObject->GetPathName(), StructData->GetStructMemory());
+					}
+					if (!bImportSuccess && !Pin->DefaultValue.IsEmpty())
+					{
+						bImportSuccess = FBlueprintEditorUtils::PropertyValueFromString(Field, Pin->DefaultValue, StructData->GetStructMemory());
+					}
+					if (!bImportSuccess)
+					{
+						// @todo: error
 					}
 				}
 
@@ -391,6 +404,22 @@ void FMovieSceneDirectorBlueprintEndpointCustomization::OnPayloadVariableChanged
 		return;
 	}
 
+	TObjectPtr<UObject> NewValueObject = FindObject<UObject>(nullptr, *NewValueString);
+
+#if WITH_EDITORONLY_DATA
+	// Fixup redirectors
+	while (Cast<UObjectRedirector>(NewValueObject) != nullptr)
+	{
+		NewValueObject = Cast<UObjectRedirector>(NewValueObject)->DestinationObject;
+		if (NewValueObject)
+		{
+			NewValueString = NewValueObject->GetPathName();
+		}
+	}
+#endif
+
+	FMovieSceneDirectorBlueprintVariableValue NewPayloadVariable{ NewValueObject, NewValueString };
+
 	TArray<UObject*> EditObjects;
 	GetEditObjects(EditObjects);
 
@@ -404,7 +433,7 @@ void FMovieSceneDirectorBlueprintEndpointCustomization::OnPayloadVariableChanged
 
 	for (int32 Index = 0; Index < RawData.Num(); ++Index)
 	{
-		bChangedAnything |= SetPayloadVariable(EditObjects[Index], RawData[Index], Property->GetFName(), NewValueString);
+		bChangedAnything |= SetPayloadVariable(EditObjects[Index], RawData[Index], Property->GetFName(), NewPayloadVariable);
 	}
 
 	if (bChangedAnything)
@@ -1306,7 +1335,7 @@ void FMovieSceneDirectorBlueprintEndpointCustomization::SetEndpoint(const FMovie
 	UEdGraphPin* CallTargetPin = FMovieSceneDirectorBlueprintUtils::FindCallTargetPin(NewEndpoint, EndpointDefinition.PossibleCallTargetClass);
 
 	// Map of the Payload Variable Names to their Default Values as Strings
-	TMap<FName, FString> PayloadVariables;
+	TMap<FName, FMovieSceneDirectorBlueprintVariableValue> PayloadVariables;
 	if (PayloadTemplate && EnumHasAnyFlags(AutoCreatePayload, EAutoCreatePayload::Variables | EAutoCreatePayload::Pins))
 	{
 		UFunction* PayloadTemplateFunction = nullptr;
@@ -1347,7 +1376,7 @@ void FMovieSceneDirectorBlueprintEndpointCustomization::SetEndpoint(const FMovie
 				// Make a payload variable for this pin
 				if (EnumHasAnyFlags(AutoCreatePayload, EAutoCreatePayload::Variables))
 				{
-					PayloadVariables.Add(PayloadPin->PinName, PayloadPin->DefaultValue);
+					PayloadVariables.Add(PayloadPin->PinName, FMovieSceneDirectorBlueprintVariableValue{ PayloadPin->DefaultObject, PayloadPin->DefaultValue });
 				}
 
 				// Make a matching user pin on the endpoint node
@@ -1368,7 +1397,7 @@ void FMovieSceneDirectorBlueprintEndpointCustomization::SetEndpoint(const FMovie
 
 	// Create payload variables for new parameters, remove payload variables for parameters
 	// that don't exist anymore.
-	const FString EmptyValue;
+	const FMovieSceneDirectorBlueprintVariableValue EmptyValue;
 	for (const TPair<UMovieSceneSequence*, FSequenceData>& Pair : AllSequenceData)
 	{
 		const TArray<void*>& RawData = Pair.Value.RawData;
@@ -1379,7 +1408,7 @@ void FMovieSceneDirectorBlueprintEndpointCustomization::SetEndpoint(const FMovie
 			FPayloadVariableMap OldPayloadVariables;
 			GetPayloadVariables(EditObjects[Index], RawData[Index], OldPayloadVariables);
 
-			for (const TPair<FName, FString>& PayloadVar : PayloadVariables)
+			for (const TPair<FName, FMovieSceneDirectorBlueprintVariableValue>& PayloadVar : PayloadVariables)
 			{
 				if (!OldPayloadVariables.Contains(PayloadVar.Key))
 				{
@@ -1387,7 +1416,7 @@ void FMovieSceneDirectorBlueprintEndpointCustomization::SetEndpoint(const FMovie
 				}
 			}
 
-			for (const TPair<FName, FString>& PayloadVar : OldPayloadVariables)
+			for (const TPair<FName, FMovieSceneDirectorBlueprintVariableValue>& PayloadVar : OldPayloadVariables)
 			{
 				if (!PayloadVariables.Contains(PayloadVar.Key))
 				{
diff --git a/Engine/Source/Editor/MovieSceneTools/Private/MovieSceneDirectorBlueprintUtils.cpp b/Engine/Source/Editor/MovieSceneTools/Private/MovieSceneDirectorBlueprintUtils.cpp
--- a/Engine/Source/Editor/MovieSceneTools/Private/MovieSceneDirectorBlueprintUtils.cpp
+++ b/Engine/Source/Editor/MovieSceneTools/Private/MovieSceneDirectorBlueprintUtils.cpp
@@ -295,7 +295,7 @@ FMovieSceneDirectorBlueprintEntrypointResult FMovieSceneDirectorBlueprintUtils::
 				continue;
 			}
 
-			const FString* PayloadVariable = EndpointCall.PayloadVariables.Find(Pin->PinName);
+			const FMovieSceneDirectorBlueprintVariableValue* PayloadVariable = EndpointCall.PayloadVariables.Find(Pin->PinName);
 			if (!PayloadVariable)
 			{
 				continue;
@@ -309,7 +309,8 @@ FMovieSceneDirectorBlueprintEntrypointResult FMovieSceneDirectorBlueprintUtils::
 			else
 			{
 				bool bMarkAsModified = false;
-				Schema->TrySetDefaultValue(*Pin, *PayloadVariable, bMarkAsModified);
+				Schema->TrySetDefaultValue(*Pin, PayloadVariable->Value, bMarkAsModified);
+				Schema->TrySetDefaultObject(*Pin, PayloadVariable->ObjectValue, bMarkAsModified);
 				bSuccess = true;
 			}
 
@@ -353,7 +354,7 @@ FMovieSceneDirectorBlueprintEntrypointResult FMovieSceneDirectorBlueprintUtils::
 
 bool FMovieSceneDirectorBlueprintUtils::GenerateEntryPointRawActorParameter(
 			FKismetCompilerContext* Compiler, UEdGraph* Graph, UEdGraphNode* Endpoint, 
-			UK2Node* OriginNode, UEdGraphPin* DestPin, const FString& PayloadValue)
+			UK2Node* OriginNode, UEdGraphPin* DestPin, const FMovieSceneDirectorBlueprintVariableValue& PayloadValue)
 {
 	// Raw actor properties are represented as soft object ptrs under the hood so that we can still serialize them
 	// To make this work, we have to make the following graph:
@@ -394,9 +395,10 @@ bool FMovieSceneDirectorBlueprintUtils::GenerateEntryPointRawActorParameter(
 
 	const UEdGraphSchema*  Schema = Graph->GetSchema();
 
-	// Set the default value for the path string
+	// Set the default value for the path string/object
 	const bool bMarkAsModified = false;
-	Schema->TrySetDefaultValue(*PathInput, PayloadValue, bMarkAsModified);
+	Schema->TrySetDefaultValue(*PathInput, PayloadValue.Value, bMarkAsModified);
+	Schema->TrySetDefaultObject(*PathInput, PayloadValue.ObjectValue, bMarkAsModified);
 
 	bool bSuccess = true;
 	bSuccess &= Schema->TryCreateConnection(PathOutput, SoftRefInput);
diff --git a/Engine/Source/Editor/MovieSceneTools/Private/MovieSceneDynamicBindingBlueprintExtension.cpp b/Engine/Source/Editor/MovieSceneTools/Private/MovieSceneDynamicBindingBlueprintExtension.cpp
--- a/Engine/Source/Editor/MovieSceneTools/Private/MovieSceneDynamicBindingBlueprintExtension.cpp
+++ b/Engine/Source/Editor/MovieSceneTools/Private/MovieSceneDynamicBindingBlueprintExtension.cpp
@@ -102,7 +102,7 @@ void UMovieSceneDynamicBindingBlueprintExtension::HandleGenerateFunctionGraphs(F
 			}
 			for (auto& Pair : DynamicBinding.PayloadVariables)
 			{
-				EndpointCall.PayloadVariables.Add(Pair.Key, Pair.Value.Value);
+				EndpointCall.PayloadVariables.Add(Pair.Key, FMovieSceneDirectorBlueprintVariableValue{ Pair.Value.ObjectValue, Pair.Value.Value });
 			}
 
 			// Create the endpoint call, and clean-up stale payload variables.
diff --git a/Engine/Source/Editor/MovieSceneTools/Private/MovieSceneDynamicBindingCustomization.cpp b/Engine/Source/Editor/MovieSceneTools/Private/MovieSceneDynamicBindingCustomization.cpp
--- a/Engine/Source/Editor/MovieSceneTools/Private/MovieSceneDynamicBindingCustomization.cpp
+++ b/Engine/Source/Editor/MovieSceneTools/Private/MovieSceneDynamicBindingCustomization.cpp
@@ -46,11 +46,11 @@ void FMovieSceneDynamicBindingCustomization::GetPayloadVariables(UObject* EditOb
 	const FMovieSceneDynamicBinding* DynamicBinding = static_cast<FMovieSceneDynamicBinding*>(RawData);
 	for (const TPair<FName, FMovieSceneDynamicBindingPayloadVariable>& Pair : DynamicBinding->PayloadVariables)
 	{
-		OutPayloadVariables.Add(Pair.Key, Pair.Value.Value);
+		OutPayloadVariables.Add(Pair.Key, FMovieSceneDirectorBlueprintVariableValue{ Pair.Value.ObjectValue, Pair.Value.Value });
 	}
 }
 
-bool FMovieSceneDynamicBindingCustomization::SetPayloadVariable(UObject* EditObject, void* RawData, FName FieldName, const FString& NewVariableValue)
+bool FMovieSceneDynamicBindingCustomization::SetPayloadVariable(UObject* EditObject, void* RawData, FName FieldName, const FMovieSceneDirectorBlueprintVariableValue& NewVariableValue)
 {
 	UMovieScene* MovieScene = Cast<UMovieScene>(EditObject);
 	FMovieSceneDynamicBinding* DynamicBinding = static_cast<FMovieSceneDynamicBinding*>(RawData);
@@ -61,7 +61,7 @@ bool FMovieSceneDynamicBindingCustomization::SetPayloadVariable(UObject* EditObj
 
 	MovieScene->Modify();
 
-	if (!NewVariableValue.IsEmpty())
+	if (!NewVariableValue.Value.IsEmpty())
 	{
 		FMovieSceneDynamicBindingPayloadVariable* PayloadVariable = DynamicBinding->PayloadVariables.Find(FieldName);
 		if (!PayloadVariable)
@@ -69,7 +69,8 @@ bool FMovieSceneDynamicBindingCustomization::SetPayloadVariable(UObject* EditObj
 			PayloadVariable = &DynamicBinding->PayloadVariables.Add(FieldName);
 		}
 
-		PayloadVariable->Value = NewVariableValue;
+		PayloadVariable->Value = NewVariableValue.Value;
+		PayloadVariable->ObjectValue = NewVariableValue.ObjectValue;
 	}
 	else
 	{
diff --git a/Engine/Source/Editor/MovieSceneTools/Private/MovieSceneEventCustomization.cpp b/Engine/Source/Editor/MovieSceneTools/Private/MovieSceneEventCustomization.cpp
--- a/Engine/Source/Editor/MovieSceneTools/Private/MovieSceneEventCustomization.cpp
+++ b/Engine/Source/Editor/MovieSceneTools/Private/MovieSceneEventCustomization.cpp
@@ -51,11 +51,11 @@ void FMovieSceneEventCustomization::GetPayloadVariables(UObject* EditObject, voi
 	const FMovieSceneEvent* EntryPoint = static_cast<FMovieSceneEvent*>(RawData);
 	for (const TPair<FName, FMovieSceneEventPayloadVariable>& Pair : EntryPoint->PayloadVariables)
 	{
-		OutPayloadVariables.Add(Pair.Key, Pair.Value.Value);
+		OutPayloadVariables.Add(Pair.Key, FMovieSceneDirectorBlueprintVariableValue{ Pair.Value.ObjectValue, Pair.Value.Value });
 	}
 }
 
-bool FMovieSceneEventCustomization::SetPayloadVariable(UObject* EditObject, void* RawData, FName FieldName, const FString& NewVariableValue)
+bool FMovieSceneEventCustomization::SetPayloadVariable(UObject* EditObject, void* RawData, FName FieldName, const FMovieSceneDirectorBlueprintVariableValue& NewVariableValue)
 {
 	UMovieSceneEventSectionBase* EventSection = Cast<UMovieSceneEventSectionBase>(EditObject);
 	FMovieSceneEvent* EntryPoint = static_cast<FMovieSceneEvent*>(RawData);
@@ -72,7 +72,8 @@ bool FMovieSceneEventCustomization::SetPayloadVariable(UObject* EditObject, void
 		PayloadVariable = &EntryPoint->PayloadVariables.Add(FieldName);
 	}
 
-	PayloadVariable->Value = NewVariableValue;
+	PayloadVariable->Value = NewVariableValue.Value;
+	PayloadVariable->ObjectValue = NewVariableValue.ObjectValue;
 	return true;
 }
 
diff --git a/Engine/Source/Editor/MovieSceneTools/Private/MovieSceneEventUtils.cpp b/Engine/Source/Editor/MovieSceneTools/Private/MovieSceneEventUtils.cpp
--- a/Engine/Source/Editor/MovieSceneTools/Private/MovieSceneEventUtils.cpp
+++ b/Engine/Source/Editor/MovieSceneTools/Private/MovieSceneEventUtils.cpp
@@ -255,7 +255,7 @@ UK2Node_FunctionEntry* FMovieSceneEventUtils::GenerateEntryPoint(FMovieSceneEven
 	FMovieSceneDirectorBlueprintEndpointCall EndpointCall;
 	for (const TPair<FName, FMovieSceneEventPayloadVariable>& Pair : EntrypointDefinition->PayloadVariables)
 	{
-		EndpointCall.PayloadVariables.Add(Pair.Key, Pair.Value.Value);
+		EndpointCall.PayloadVariables.Add(Pair.Key, FMovieSceneDirectorBlueprintVariableValue{ Pair.Value.ObjectValue, Pair.Value.Value });
 	}
 	if (!EntrypointDefinition->BoundObjectPinName.IsNone())
 	{
diff --git a/Engine/Source/Editor/MovieSceneTools/Public/MovieSceneDirectorBlueprintEndpointCustomization.h b/Engine/Source/Editor/MovieSceneTools/Public/MovieSceneDirectorBlueprintEndpointCustomization.h
--- a/Engine/Source/Editor/MovieSceneTools/Public/MovieSceneDirectorBlueprintEndpointCustomization.h
+++ b/Engine/Source/Editor/MovieSceneTools/Public/MovieSceneDirectorBlueprintEndpointCustomization.h
@@ -50,7 +50,7 @@ public:
 
 protected:
 
-	using FPayloadVariableMap = TMap<FName, FString, TInlineSetAllocator<8>>;
+	using FPayloadVariableMap = TMap<FName, FMovieSceneDirectorBlueprintVariableValue, TInlineSetAllocator<8>>;
 
 	/**
 	 * User-interface information for "well-known parameters", i.e. parameters that can be bound
@@ -92,9 +92,10 @@ protected:
 	 * @param RawData			The raw data for which to set the endpoint call payload variable
 	 * @param FieldName			The name of the payload variable to set
 	 * @param NewVariableValue	The string value of the payload variable to set
+	 * @param NewVariableObject	In the case the payload is a UObject, a reference to the UObject
 	 * @return					Whether the payload variable was set
 	 */
-	virtual bool SetPayloadVariable(UObject* EditObject, void* RawData, FName FieldName, const FString& NewVariableValue) = 0;
+	virtual bool SetPayloadVariable(UObject* EditObject, void* RawData, FName FieldName, const FMovieSceneDirectorBlueprintVariableValue& NewVariableValue) = 0;
 
 	/**
 	 * Get the pin names setup with "well-known parameters" for the currently edited objects.
diff --git a/Engine/Source/Editor/MovieSceneTools/Public/MovieSceneDirectorBlueprintUtils.h b/Engine/Source/Editor/MovieSceneTools/Public/MovieSceneDirectorBlueprintUtils.h
--- a/Engine/Source/Editor/MovieSceneTools/Public/MovieSceneDirectorBlueprintUtils.h
+++ b/Engine/Source/Editor/MovieSceneTools/Public/MovieSceneDirectorBlueprintUtils.h
@@ -51,6 +51,18 @@ struct MOVIESCENETOOLS_API FMovieSceneDirectorBlueprintEndpointParameter
 	static FString SanitizePinName(const FString& InPinName);
 };
 
+/*
+* Structure for holding value of a blueprint variable 
+*/
+struct MOVIESCENETOOLS_API FMovieSceneDirectorBlueprintVariableValue
+{
+	/* In the case the variable is a UObject, reference to the UObject */
+	TObjectPtr<UObject> ObjectValue;
+
+	/** String value of the variable */
+	FString Value;
+};
+
 /**
  * Parameters for creating a blueprint endpoint
  */
@@ -122,7 +134,7 @@ struct MOVIESCENETOOLS_API FMovieSceneDirectorBlueprintEndpointCall
 	TArray<FName> ExposedPinNames;
 
 	/** The values to pass to the endpoint node's parameters */
-	TMap<FName, FString> PayloadVariables;
+	TMap<FName, FMovieSceneDirectorBlueprintVariableValue> PayloadVariables;
 
 	/** Callback for customizing the endpoint call */
 	FMovieSceneCustomizeDirectorBlueprintEndpointCallDelegate CustomizeEndpointCall;
@@ -210,6 +222,6 @@ private:
 
 	static bool GenerateEntryPointRawActorParameter(
 			FKismetCompilerContext* Compiler, UEdGraph* Graph, UEdGraphNode* Endpoint, 
-			UK2Node* OriginNode, UEdGraphPin* DestPin, const FString& PayloadValue);
+			UK2Node* OriginNode, UEdGraphPin* DestPin, const FMovieSceneDirectorBlueprintVariableValue& PayloadValue);
 };
 
diff --git a/Engine/Source/Editor/MovieSceneTools/Public/MovieSceneDynamicBindingCustomization.h b/Engine/Source/Editor/MovieSceneTools/Public/MovieSceneDynamicBindingCustomization.h
--- a/Engine/Source/Editor/MovieSceneTools/Public/MovieSceneDynamicBindingCustomization.h
+++ b/Engine/Source/Editor/MovieSceneTools/Public/MovieSceneDynamicBindingCustomization.h
@@ -31,7 +31,7 @@ public:
 protected:
 
 	virtual void GetPayloadVariables(UObject* EditObject, void* RawData, FPayloadVariableMap& OutPayloadVariables) const override;
-	virtual bool SetPayloadVariable(UObject* EditObject, void* RawData, FName FieldName, const FString& NewVariableValue) override;
+	virtual bool SetPayloadVariable(UObject* EditObject, void* RawData, FName FieldName, const FMovieSceneDirectorBlueprintVariableValue& NewVariableValue) override;
 	virtual UK2Node* FindEndpoint(UMovieSceneSequence* Sequence, UBlueprint* Blueprint, UObject* EditObject, void* RawData) const override;
 	virtual void GetWellKnownParameterPinNames(UObject* EditObject, void* RawData, TArray<FName>& OutWellKnownParameters) const override;
 	virtual void GetWellKnownParameterCandidates(UK2Node* Endpoint, TArray<FWellKnownParameterCandidates>& OutCandidates) const override;
diff --git a/Engine/Source/Editor/MovieSceneTools/Public/MovieSceneEventCustomization.h b/Engine/Source/Editor/MovieSceneTools/Public/MovieSceneEventCustomization.h
--- a/Engine/Source/Editor/MovieSceneTools/Public/MovieSceneEventCustomization.h
+++ b/Engine/Source/Editor/MovieSceneTools/Public/MovieSceneEventCustomization.h
@@ -22,7 +22,7 @@ public:
 protected:
 
 	virtual void GetPayloadVariables(UObject* EditObject, void* RawData, FPayloadVariableMap& OutPayloadVariables) const override;
-	virtual bool SetPayloadVariable(UObject* EditObject, void* RawData, FName FieldName, const FString& NewVariableValue) override;
+	virtual bool SetPayloadVariable(UObject* EditObject, void* RawData, FName FieldName, const FMovieSceneDirectorBlueprintVariableValue& NewValue) override;
 	virtual UK2Node* FindEndpoint(UMovieSceneSequence* Sequence, UBlueprint* Blueprint, UObject* EditObject, void* RawData) const override;
 	virtual void GetWellKnownParameterPinNames(UObject* EditObject, void* RawData, TArray<FName>& OutWellKnownParameters) const override;
 	virtual void GetWellKnownParameterCandidates(UK2Node* Endpoint, TArray<FWellKnownParameterCandidates>& OutCandidates) const override;
diff --git a/Engine/Source/Runtime/MovieScene/Public/MovieSceneDynamicBinding.h b/Engine/Source/Runtime/MovieScene/Public/MovieSceneDynamicBinding.h
--- a/Engine/Source/Runtime/MovieScene/Public/MovieSceneDynamicBinding.h
+++ b/Engine/Source/Runtime/MovieScene/Public/MovieSceneDynamicBinding.h
@@ -20,6 +20,9 @@ USTRUCT(BlueprintType)
 struct FMovieSceneDynamicBindingPayloadVariable
 {
 	GENERATED_BODY()
+		
+	UPROPERTY()
+	TObjectPtr<UObject> ObjectValue;
 
 	UPROPERTY(EditAnywhere, Category = "Sequencer|Dynamic Binding")
 	FString Value;
diff --git a/Engine/Source/Runtime/MovieSceneTracks/Public/Channels/MovieSceneEvent.h b/Engine/Source/Runtime/MovieSceneTracks/Public/Channels/MovieSceneEvent.h
--- a/Engine/Source/Runtime/MovieSceneTracks/Public/Channels/MovieSceneEvent.h
+++ b/Engine/Source/Runtime/MovieSceneTracks/Public/Channels/MovieSceneEvent.h
@@ -33,6 +33,10 @@ struct FMovieSceneEventPayloadVariable
 {
 	GENERATED_BODY()
 
+	UPROPERTY()
+	/** If the value for this pin should be an object, we store a pointer to it */
+	TObjectPtr<class UObject> ObjectValue;
+
 	UPROPERTY(EditAnywhere, Category="Sequencer|Event")
 	FString Value;
 };
