Sequencer: TrySetDefaultValue or TrySetDefaultObject, not both

diff --git a/Engine/Source/Editor/MovieSceneTools/Private/MovieSceneDirectorBlueprintUtils.cpp b/Engine/Source/Editor/MovieSceneTools/Private/MovieSceneDirectorBlueprintUtils.cpp
--- a/Engine/Source/Editor/MovieSceneTools/Private/MovieSceneDirectorBlueprintUtils.cpp
+++ b/Engine/Source/Editor/MovieSceneTools/Private/MovieSceneDirectorBlueprintUtils.cpp
@@ -309,8 +309,14 @@ FMovieSceneDirectorBlueprintEntrypointResult FMovieSceneDirectorBlueprintUtils::
 			else
 			{
 				bool bMarkAsModified = false;
-				Schema->TrySetDefaultValue(*Pin, PayloadVariable->Value, bMarkAsModified);
-				Schema->TrySetDefaultObject(*Pin, PayloadVariable->ObjectValue, bMarkAsModified);
+				if (!PayloadVariable->Value.IsEmpty())
+				{
+					Schema->TrySetDefaultValue(*Pin, PayloadVariable->Value, bMarkAsModified);
+				}
+				else if (!PayloadVariable->ObjectValue.IsNull())
+				{
+					Schema->TrySetDefaultObject(*Pin, PayloadVariable->ObjectValue, bMarkAsModified);
+				}
 				bSuccess = true;
 			}
 
@@ -397,8 +403,14 @@ bool FMovieSceneDirectorBlueprintUtils::GenerateEntryPointRawActorParameter(
 
 	// Set the default value for the path string/object
 	const bool bMarkAsModified = false;
-	Schema->TrySetDefaultValue(*PathInput, PayloadValue.Value, bMarkAsModified);
-	Schema->TrySetDefaultObject(*PathInput, PayloadValue.ObjectValue, bMarkAsModified);
+	if (!PayloadValue.Value.IsEmpty())
+	{
+		Schema->TrySetDefaultValue(*PathInput, PayloadValue.Value, bMarkAsModified);
+	}
+	else if (!PayloadValue.ObjectValue.IsNull())
+	{
+		Schema->TrySetDefaultObject(*PathInput, PayloadValue.ObjectValue, bMarkAsModified);
+	}
 
 	bool bSuccess = true;
 	bSuccess &= Schema->TryCreateConnection(PathOutput, SoftRefInput);
