Sequencer: IsNull deprecation

diff --git a/Engine/Source/Editor/MovieSceneTools/Private/MovieSceneDirectorBlueprintUtils.cpp b/Engine/Source/Editor/MovieSceneTools/Private/MovieSceneDirectorBlueprintUtils.cpp
--- a/Engine/Source/Editor/MovieSceneTools/Private/MovieSceneDirectorBlueprintUtils.cpp
+++ b/Engine/Source/Editor/MovieSceneTools/Private/MovieSceneDirectorBlueprintUtils.cpp
@@ -313,7 +313,7 @@ FMovieSceneDirectorBlueprintEntrypointResult FMovieSceneDirectorBlueprintUtils::
 				{
 					Schema->TrySetDefaultValue(*Pin, PayloadVariable->Value, bMarkAsModified);
 				}
-				else if (!PayloadVariable->ObjectValue.IsNull())
+				else if (!PayloadVariable->ObjectValue)
 				{
 					Schema->TrySetDefaultObject(*Pin, PayloadVariable->ObjectValue, bMarkAsModified);
 				}
@@ -407,7 +407,7 @@ bool FMovieSceneDirectorBlueprintUtils::GenerateEntryPointRawActorParameter(
 	{
 		Schema->TrySetDefaultValue(*PathInput, PayloadValue.Value, bMarkAsModified);
 	}
-	else if (!PayloadValue.ObjectValue.IsNull())
+	else if (!PayloadValue.ObjectValue)
 	{
 		Schema->TrySetDefaultObject(*PathInput, PayloadValue.ObjectValue, bMarkAsModified);
 	}
