Sequencer: Prefer TrySetDefaultObject over TrySetDefaultValue

diff --git a/Engine/Source/Editor/MovieSceneTools/Private/MovieSceneDirectorBlueprintUtils.cpp b/Engine/Source/Editor/MovieSceneTools/Private/MovieSceneDirectorBlueprintUtils.cpp
--- a/Engine/Source/Editor/MovieSceneTools/Private/MovieSceneDirectorBlueprintUtils.cpp
+++ b/Engine/Source/Editor/MovieSceneTools/Private/MovieSceneDirectorBlueprintUtils.cpp
@@ -309,13 +309,13 @@ FMovieSceneDirectorBlueprintEntrypointResult FMovieSceneDirectorBlueprintUtils::
 			else
 			{
 				bool bMarkAsModified = false;
-				if (!PayloadVariable->Value.IsEmpty())
+				if (PayloadVariable->ObjectValue)
 				{
-					Schema->TrySetDefaultValue(*Pin, PayloadVariable->Value, bMarkAsModified);
+					Schema->TrySetDefaultObject(*Pin, PayloadVariable->ObjectValue, bMarkAsModified);
 				}
-				else if (PayloadVariable->ObjectValue)
+				else if (!PayloadVariable->Value.IsEmpty())
 				{
-					Schema->TrySetDefaultObject(*Pin, PayloadVariable->ObjectValue, bMarkAsModified);
+					Schema->TrySetDefaultValue(*Pin, PayloadVariable->Value, bMarkAsModified);
 				}
 				bSuccess = true;
 			}
@@ -403,13 +403,13 @@ bool FMovieSceneDirectorBlueprintUtils::GenerateEntryPointRawActorParameter(
 
 	// Set the default value for the path string/object
 	const bool bMarkAsModified = false;
-	if (!PayloadValue.Value.IsEmpty())
+	if (PayloadValue.ObjectValue)
 	{
-		Schema->TrySetDefaultValue(*PathInput, PayloadValue.Value, bMarkAsModified);
+		Schema->TrySetDefaultObject(*PathInput, PayloadValue.ObjectValue, bMarkAsModified);
 	}
-	else if (PayloadValue.ObjectValue)
+	else if (!PayloadValue.Value.IsEmpty())
 	{
-		Schema->TrySetDefaultObject(*PathInput, PayloadValue.ObjectValue, bMarkAsModified);
+		Schema->TrySetDefaultValue(*PathInput, PayloadValue.Value, bMarkAsModified);
 	}
 
 	bool bSuccess = true;
