Sequencer: Fix crash when binding overrides array change

The repro in this case is adding binding overrides during runtime, then removing them once playback finished.

diff --git a/Engine/Source/Editor/MovieSceneTools/Private/MovieSceneObjectBindingIDCustomization.cpp b/Engine/Source/Editor/MovieSceneTools/Private/MovieSceneObjectBindingIDCustomization.cpp
--- a/Engine/Source/Editor/MovieSceneTools/Private/MovieSceneObjectBindingIDCustomization.cpp
+++ b/Engine/Source/Editor/MovieSceneTools/Private/MovieSceneObjectBindingIDCustomization.cpp
@@ -164,12 +164,17 @@ FMovieSceneObjectBindingID FMovieSceneObjectBindingIDCustomization::GetCurrentVa
 	TArray<void*> Ptrs;
 	StructProperty->AccessRawData(Ptrs);
 
-	FMovieSceneObjectBindingID Value = Ptrs.Num() > 0 ? *static_cast<FMovieSceneObjectBindingID*>(Ptrs[0]) : FMovieSceneObjectBindingID();
+	if (Ptrs.Num() == 0 || Ptrs[0] == nullptr)
+	{
+		return FMovieSceneObjectBindingID();
+	}
+
+	FMovieSceneObjectBindingID Value = *static_cast<FMovieSceneObjectBindingID*>(Ptrs[0]);
 
 	// If more than one value and not all equal, return empty
 	for (int32 Index = 1; Index < Ptrs.Num(); ++Index)
 	{
-		if (*static_cast<FMovieSceneObjectBindingID*>(Ptrs[Index]) != Value)
+		if (Ptrs[Index] != nullptr && *static_cast<FMovieSceneObjectBindingID*>(Ptrs[Index]) != Value)
 		{
 			return FMovieSceneObjectBindingID();
 		}
