Sequencer:
- Added Support for keying Transform properties from Details Panel

diff --git a/Engine/Source/Editor/MovieSceneTools/Private/TrackEditors/TransformTrackEditor.cpp b/Engine/Source/Editor/MovieSceneTools/Private/TrackEditors/TransformTrackEditor.cpp
--- a/Engine/Source/Editor/MovieSceneTools/Private/TrackEditors/TransformTrackEditor.cpp
+++ b/Engine/Source/Editor/MovieSceneTools/Private/TrackEditors/TransformTrackEditor.cpp
@@ -43,6 +43,7 @@
 #include "Constraints/TransformConstraintChannelInterface.h"
 #include "Misc/TransactionObjectEvent.h"
 #include "Engine/Selection.h"
+#include "ISequencerObjectChangeListener.h"
 
 #define LOCTEXT_NAMESPACE "MovieScene_TransformTrack"
 
@@ -77,6 +78,28 @@ F3DTransformTrackEditor::F3DTransformTrackEditor( TSharedRef<ISequencer> InSeque
 	{
 		GEditor->RegisterForUndo(this);
 	}
+
+	if (TSharedPtr<ISequencer> SequencerPtr = FMovieSceneTrackEditor::GetSequencer())
+	{
+	    const FProperty* LocationProperty = FindFProperty<FProperty>(USceneComponent::StaticClass(), USceneComponent::GetRelativeLocationPropertyName());
+		const FProperty* RotationProperty = FindFProperty<FProperty>(USceneComponent::StaticClass(), USceneComponent::GetRelativeRotationPropertyName());
+		const FProperty* Scale3DProperty = FindFProperty<FProperty>(USceneComponent::StaticClass(), USceneComponent::GetRelativeScale3DPropertyName());
+
+		ISequencerObjectChangeListener& ObjectChangeListener = SequencerPtr->GetObjectChangeListener();
+		auto AddTransformProperty = [this, &ObjectChangeListener](const FProperty* Property, EMovieSceneTransformChannel TransformChannel)
+		{
+			if (Property)
+			{
+				TransformProperties.Add({ Property, TransformChannel });
+				ObjectChangeListener.GetOnAnimatablePropertyChanged(Property)
+					.AddRaw(this, &F3DTransformTrackEditor::OnTransformPropertyChanged, TransformChannel);
+			}
+		};
+
+		AddTransformProperty(LocationProperty, EMovieSceneTransformChannel::Translation);
+		AddTransformProperty(RotationProperty, EMovieSceneTransformChannel::Rotation);
+		AddTransformProperty(Scale3DProperty, EMovieSceneTransformChannel::Scale);
+	}
 }
 
 F3DTransformTrackEditor::~F3DTransformTrackEditor()
@@ -85,6 +108,15 @@ F3DTransformTrackEditor::~F3DTransformTrackEditor()
 	{
 		GEditor->UnregisterForUndo(this);
 	}
+
+	if (TSharedPtr<ISequencer> SequencerPtr = FMovieSceneTrackEditor::GetSequencer())
+	{
+		ISequencerObjectChangeListener& ObjectChangeListener = SequencerPtr->GetObjectChangeListener();
+		for (const FTransformPropertyInfo& TransformProperty : TransformProperties)
+		{
+			ObjectChangeListener.GetOnAnimatablePropertyChanged(TransformProperty.Property).RemoveAll(this);
+		}
+	}
 }
 //for 5.2 we will move this over to the header
 static TSet<TWeakObjectPtr<UMovieScene3DTransformSection>> SectionsToClear;
@@ -1084,6 +1116,72 @@ int32 GetPreviousKey(FMovieSceneDoubleChannel& Channel, FFrameNumber Time)
 	return Index;
 }
 
+void F3DTransformTrackEditor::OnTransformPropertyChanged(const FPropertyChangedParams& PropertyChangedParams, EMovieSceneTransformChannel TransformChannel)
+{
+	// Key Property sends in one object at a time
+	USceneComponent* SceneComponent = MovieSceneHelpers::SceneComponentFromRuntimeObject(PropertyChangedParams.ObjectsThatChanged[0]);
+	if (!SceneComponent)
+	{
+		return;
+	}
+
+	UObject* ObjectToKey = SceneComponent;
+
+	// Set Owning Actor to key instead of Scene Component if it's the root
+	AActor* OwningActor = SceneComponent->GetOwner();
+	if (OwningActor && SceneComponent == OwningActor->GetRootComponent())
+	{
+		ObjectToKey = OwningActor;
+	}
+
+	auto InitializeNewTrack = [](UMovieScene3DTransformTrack* NewTrack)
+	{
+		NewTrack->SetPropertyNameAndPath(TransformPropertyName, TransformPropertyName.ToString());
+	};
+
+	auto GenerateKeys = [this, SceneComponent, TransformChannel](UMovieSceneSection* Section, FGeneratedTrackKeys& OutGeneratedKeys)
+	{
+		UMovieScene3DTransformSection* TransformSection = CastChecked<UMovieScene3DTransformSection>(Section);
+
+		// Ensure that the Transform Channel is masked in
+		TransformSection->SetMask(TransformSection->GetMask().GetChannels() | TransformChannel);
+
+		GetTransformKeys(TOptional<FTransformData>(), FTransformData(SceneComponent), TransformChannel, SceneComponent, Section, OutGeneratedKeys);
+	};
+
+	auto KeyProperty = [this, ObjectToKey, PropertyChangedParams, InitializeNewTrack, GenerateKeys, TransformChannel](FFrameNumber KeyTime)
+	{
+		FKeyPropertyResult KeyPropertyResult = this->AddKeysToObjects({ ObjectToKey },
+			KeyTime,
+			PropertyChangedParams.KeyMode,
+			UMovieScene3DTransformTrack::StaticClass(),
+			TransformPropertyName,
+			InitializeNewTrack,
+			GenerateKeys);
+
+		for (TWeakObjectPtr<UMovieSceneSection>& WeakSection : KeyPropertyResult.SectionsKeyed)
+		{
+			if (UMovieScene3DTransformSection* Section = Cast<UMovieScene3DTransformSection>(WeakSection.Get()))
+			{
+				FMovieSceneConstraintChannelHelper::CompensateIfNeeded(GetSequencer(), Section, KeyTime);
+			}
+		}
+
+		// For new sections created, mask so only the relevant transform channel appears
+		for (TWeakObjectPtr<UMovieSceneSection>& WeakSection : KeyPropertyResult.SectionsCreated)
+		{
+			if (UMovieScene3DTransformSection* Section = Cast<UMovieScene3DTransformSection>(WeakSection.Get()))
+			{
+				Section->SetMask(TransformChannel);
+			}
+		}
+
+		return KeyPropertyResult;
+	};
+
+	AnimatablePropertyChanged(FOnKeyProperty::CreateLambda(KeyProperty));
+}
+
 void F3DTransformTrackEditor::ProcessKeyOperation(UObject* ObjectToKey, TArrayView<const UE::Sequencer::FKeySectionOperation> SectionsToKey, ISequencer& InSequencer, FFrameNumber KeyTime)
 {
 	using namespace UE::MovieScene;
diff --git a/Engine/Source/Editor/MovieSceneTools/Private/TrackEditors/TransformTrackEditor.h b/Engine/Source/Editor/MovieSceneTools/Private/TrackEditors/TransformTrackEditor.h
--- a/Engine/Source/Editor/MovieSceneTools/Private/TrackEditors/TransformTrackEditor.h
+++ b/Engine/Source/Editor/MovieSceneTools/Private/TrackEditors/TransformTrackEditor.h
@@ -73,6 +73,8 @@ public:
 
 private:
 
+	void OnTransformPropertyChanged(const FPropertyChangedParams& PropertyChangedParams, EMovieSceneTransformChannel TransformChannel);
+
 	void ProcessKeyOperation(UObject* ObjectToKey, TArrayView<const UE::Sequencer::FKeySectionOperation> SectionsToKey, ISequencer& InSequencer, FFrameNumber KeyTime);
 
 	/** Returns whether or not a transform track can be added for an actor with a specific handle. */
@@ -198,6 +200,14 @@ private:
 	/** Mapping of objects to their existing transform data (for comparing against new transform data) */
 	TMap< TWeakObjectPtr<UObject>, FTransformData > ObjectToExistingTransform;
 
+	struct FTransformPropertyInfo
+	{
+		const FProperty* Property;
+		EMovieSceneTransformChannel TransformChannel;
+	};
+	/** Array of transform property info for the scene component transform properties for explicit support */
+	TArray<FTransformPropertyInfo, TFixedAllocator<3>> TransformProperties;
+
 	/** Command Bindings added by the Transform Track Editor to Sequencer and curve editor. */
 	TSharedPtr<FUICommandList> CommandBindings;
 
diff --git a/Engine/Source/Editor/MovieSceneTools/Public/KeyframeTrackEditor.h b/Engine/Source/Editor/MovieSceneTools/Public/KeyframeTrackEditor.h
--- a/Engine/Source/Editor/MovieSceneTools/Public/KeyframeTrackEditor.h
+++ b/Engine/Source/Editor/MovieSceneTools/Public/KeyframeTrackEditor.h
@@ -634,6 +634,7 @@ private:
 				if (bSectionCreated && GetSequencer()->GetInfiniteKeyAreas())
 				{
 					SectionToKey->SetRange(TRange<FFrameNumber>::All());
+					KeyPropertyResult.SectionsCreated.Add(SectionToKey);
 				}
 			}
 
diff --git a/Engine/Source/Editor/Sequencer/Private/SequencerObjectChangeListener.cpp b/Engine/Source/Editor/Sequencer/Private/SequencerObjectChangeListener.cpp
--- a/Engine/Source/Editor/Sequencer/Private/SequencerObjectChangeListener.cpp
+++ b/Engine/Source/Editor/Sequencer/Private/SequencerObjectChangeListener.cpp
@@ -140,6 +140,11 @@ FOnAnimatablePropertyChanged& FSequencerObjectChangeListener::GetOnAnimatablePro
 	return PropertyChangedEventMap.FindOrAdd( PropertyKey );
 }
 
+FOnAnimatablePropertyChanged& FSequencerObjectChangeListener::GetOnAnimatablePropertyChanged(const FProperty* Property)
+{
+	return PropertyPathChangedEventMap.FindOrAdd(Property);
+}
+
 FOnPropagateObjectChanges& FSequencerObjectChangeListener::GetOnPropagateObjectChanges()
 {
 	return OnPropagateObjectChanges;
@@ -199,8 +204,13 @@ const FOnAnimatablePropertyChanged* FSequencerObjectChangeListener::FindProperty
 	// we check for flags on the outer property but must not use setter functions that would not know the index to set
 	const bool bCanApplyFunction = PropertyOrContainer == &Property;
 
-	const FOnAnimatablePropertyChanged* DelegatePtr = PropertyChangedEventMap.Find(PropertyKey);
-	if (DelegatePtr != nullptr)
+	// Early return if explicitly supported
+	if (const FOnAnimatablePropertyChanged* DelegatePtr = PropertyPathChangedEventMap.Find(&Property))
+	{
+		return DelegatePtr;
+	}
+
+	if (const FOnAnimatablePropertyChanged* DelegatePtr = PropertyChangedEventMap.Find(PropertyKey))
 	{
 		FString PropertyVarName = PropertyOrContainer->GetName();
 
@@ -310,6 +320,12 @@ bool FSequencerObjectChangeListener::CanKeyProperty_Internal(FCanKeyPropertyPara
 			// If there is a custom accessor for this specific property path, it is animatable (as long as there is a supported track editor registered for the property type)
 			if (UE::MovieScene::GlobalCustomAccessorExists(CanKeyPropertyParams.ObjectClass, InOutPropertyPath.ToString(TEXT("."))))
 			{
+				if (const FOnAnimatablePropertyChanged* DelegatePtr = PropertyPathChangedEventMap.Find(Property))
+				{
+					InOutProperty = Property;
+					InOutDelegate = *DelegatePtr;
+					return true;
+				}
 				if (const FOnAnimatablePropertyChanged* DelegatePtr = PropertyChangedEventMap.Find(PropertyKey))
 				{
 					InOutProperty = Property;
diff --git a/Engine/Source/Editor/Sequencer/Private/SequencerObjectChangeListener.h b/Engine/Source/Editor/Sequencer/Private/SequencerObjectChangeListener.h
--- a/Engine/Source/Editor/Sequencer/Private/SequencerObjectChangeListener.h
+++ b/Engine/Source/Editor/Sequencer/Private/SequencerObjectChangeListener.h
@@ -26,6 +26,7 @@ public:
 
 	/** ISequencerObjectChangeListener interface */
 	virtual FOnAnimatablePropertyChanged& GetOnAnimatablePropertyChanged( FAnimatedPropertyKey PropertyKey ) override;
+	virtual FOnAnimatablePropertyChanged& GetOnAnimatablePropertyChanged(const FProperty* Property) override;
 	virtual FOnPropagateObjectChanges& GetOnPropagateObjectChanges() override;
 	virtual FOnObjectPropertyChanged& GetOnAnyPropertyChanged(UObject& Object) override;
 	virtual void ReportObjectDestroyed(UObject& Object) override;
@@ -92,6 +93,9 @@ private:
 	/** A mapping of property classes to multi-cast delegate that is broadcast when properties of that type change */
 	TMap< FAnimatedPropertyKey, FOnAnimatablePropertyChanged > PropertyChangedEventMap;
 
+	/** A mapping of an explicitly supported property path to multi-cast delegate that is broadcast when properties of that type change */
+	TMap<const FProperty*, FOnAnimatablePropertyChanged> PropertyPathChangedEventMap;
+
 	/** A mapping of object instance to property change event */
 	TMap< FObjectKey, FOnObjectPropertyChanged > ObjectToPropertyChangedEvent;
 
diff --git a/Engine/Source/Editor/Sequencer/Public/ISequencerObjectChangeListener.h b/Engine/Source/Editor/Sequencer/Public/ISequencerObjectChangeListener.h
--- a/Engine/Source/Editor/Sequencer/Public/ISequencerObjectChangeListener.h
+++ b/Engine/Source/Editor/Sequencer/Public/ISequencerObjectChangeListener.h
@@ -23,6 +23,11 @@ public:
 	 */
 	virtual FOnAnimatablePropertyChanged& GetOnAnimatablePropertyChanged( FAnimatedPropertyKey PropertyKey ) = 0;
 
+	/**
+	 * A delegate for when a property of a specific path is changed. 
+	 */
+	virtual FOnAnimatablePropertyChanged& GetOnAnimatablePropertyChanged(const FProperty* Property) = 0;
+
 	/**
 	 * A delegate for when object changes should be propagated to/from puppet actors
 	 */
diff --git a/Engine/Source/Editor/Sequencer/Public/MovieSceneTrackEditor.h b/Engine/Source/Editor/Sequencer/Public/MovieSceneTrackEditor.h
--- a/Engine/Source/Editor/Sequencer/Public/MovieSceneTrackEditor.h
+++ b/Engine/Source/Editor/Sequencer/Public/MovieSceneTrackEditor.h
@@ -35,6 +35,8 @@ struct FKeyPropertyResult
 		bHandleCreated |= A.bHandleCreated;
 		bTrackCreated |= A.bTrackCreated;
 		bKeyCreated |= A.bKeyCreated;
+		SectionsCreated.Append(A.SectionsCreated);
+		SectionsKeyed.Append(A.SectionsKeyed);
 	}
 
 	/* Was the track modified in any way? */
