Sequencer: Fix shot name being erroneously carried over when switching takes and creating new takes. It should only be carried over it's not the same as the sequence's name (ie. it was a custom shot name)

diff --git a/Engine/Source/Editor/MovieSceneTools/Private/TrackEditors/SubTrackEditor.cpp b/Engine/Source/Editor/MovieSceneTools/Private/TrackEditors/SubTrackEditor.cpp
--- a/Engine/Source/Editor/MovieSceneTools/Private/TrackEditors/SubTrackEditor.cpp
+++ b/Engine/Source/Editor/MovieSceneTools/Private/TrackEditors/SubTrackEditor.cpp
@@ -577,7 +577,8 @@ void FSubTrackEditor::CreateNewTake(UMovieSceneSubSection* Section)
 			UMovieSceneCinematicShotSection* ShotSection = Cast<UMovieSceneCinematicShotSection>(Section);
 			UMovieSceneCinematicShotSection* NewShotSection = Cast<UMovieSceneCinematicShotSection>(NewSection);
 
-			if (ShotSection && NewShotSection)
+			// If the old shot's name is not the same as the sequence's name, assume the user had customized the shot name, so carry it over
+			if (ShotSection && NewShotSection && ShotSection->GetSequence() && ShotSection->GetShotDisplayName() != ShotSection->GetSequence()->GetName())
 			{
 				NewShotSection->SetShotDisplayName(ShotSection->GetShotDisplayName());
 			}
@@ -642,7 +643,8 @@ void FSubTrackEditor::ChangeTake(UMovieSceneSequence* Sequence)
 			UMovieSceneCinematicShotSection* ShotSection = Cast<UMovieSceneCinematicShotSection>(Section);
 			UMovieSceneCinematicShotSection* NewShotSection = Cast<UMovieSceneCinematicShotSection>(NewSection);
 
-			if (ShotSection && NewShotSection)
+			// If the old shot's name is not the same as the sequence's name, assume the user had customized the shot name, so carry it over
+			if (ShotSection && NewShotSection && ShotSection->GetSequence() && ShotSection->GetShotDisplayName() != ShotSection->GetSequence()->GetName())
 			{
 				NewShotSection->SetShotDisplayName(ShotSection->GetShotDisplayName());
 			}
