Sequencer: Remove check for material slot index and just continue if the index doesn't exist to allow for manual rebinding

This can happen if a material slot is deleted and the sequencer track is still referring to an out of bounds index.

diff --git a/Engine/Source/Editor/MovieSceneTools/Private/TrackEditors/PrimitiveMaterialTrackEditor.cpp b/Engine/Source/Editor/MovieSceneTools/Private/TrackEditors/PrimitiveMaterialTrackEditor.cpp
--- a/Engine/Source/Editor/MovieSceneTools/Private/TrackEditors/PrimitiveMaterialTrackEditor.cpp
+++ b/Engine/Source/Editor/MovieSceneTools/Private/TrackEditors/PrimitiveMaterialTrackEditor.cpp
@@ -201,9 +201,11 @@ void FPrimitiveMaterialTrackEditor::BuildTrackContextMenu(FMenuBuilder& MenuBuil
 						// Found material by index, but not by slot name. Auto-rebind would change slot name to the current one
 						FoundMaterialIndex = MaterialInfo.MaterialSlotIndex;
 						TArray<FName> SlotNames = Component->GetMaterialSlotNames();
-						check(SlotNames.IsValidIndex(MaterialInfo.MaterialSlotIndex));
-						FoundSlotName = SlotNames[MaterialInfo.MaterialSlotIndex];
-						AutoRebindTooltip = FText::Format(LOCTEXT("AutoRebindToNewSlotSwitcherName", "Rebind track to slot {0}, keeping same index {1}"), FText::FromName(FoundSlotName), FText::AsNumber(MaterialInfo.MaterialSlotIndex));
+						if (SlotNames.IsValidIndex(MaterialInfo.MaterialSlotIndex))
+						{
+							FoundSlotName = SlotNames[MaterialInfo.MaterialSlotIndex];
+							AutoRebindTooltip = FText::Format(LOCTEXT("AutoRebindToNewSlotSwitcherName", "Rebind track to slot {0}, keeping same index {1}"), FText::FromName(FoundSlotName), FText::AsNumber(MaterialInfo.MaterialSlotIndex));
+						}
 					}
 					// If we didn't find a material, we don't create a tooltip, because we won't be able to 'auto' rebind, just manually bind.
 				}
