Sequencer: Add ensure if channel doesn't exist.

No repro at the moment. This will prevent the crash, but there's probably going to be something wrong with keyframing, so emit a warning.

#rnx

diff --git a/Engine/Source/Editor/MovieSceneTools/Private/TrackEditors/TransformTrackEditor.cpp b/Engine/Source/Editor/MovieSceneTools/Private/TrackEditors/TransformTrackEditor.cpp
--- a/Engine/Source/Editor/MovieSceneTools/Private/TrackEditors/TransformTrackEditor.cpp
+++ b/Engine/Source/Editor/MovieSceneTools/Private/TrackEditors/TransformTrackEditor.cpp
@@ -1416,31 +1416,34 @@ void F3DTransformTrackEditor::ProcessKeyOperation(UObject* ObjectToKey, TArrayVi
 			{
 				FMovieSceneDoubleChannel* Channel = static_cast<FMovieSceneDoubleChannel*>(Handle.Get());
 
-				double Value =
-					Handle.GetChannelIndex() == 0 ? RecomposedTransform.Translation[0] :
-					Handle.GetChannelIndex() == 1 ? RecomposedTransform.Translation[1] :
-					Handle.GetChannelIndex() == 2 ? RecomposedTransform.Translation[2] :
-					Handle.GetChannelIndex() == 3 ? RecomposedTransform.Rotation.Roll :
-					Handle.GetChannelIndex() == 4 ? RecomposedTransform.Rotation.Pitch :
-					Handle.GetChannelIndex() == 5 ? RecomposedTransform.Rotation.Yaw :
-					Handle.GetChannelIndex() == 6 ? RecomposedTransform.Scale[0] :
-					Handle.GetChannelIndex() == 7 ? RecomposedTransform.Scale[1] :
-					Handle.GetChannelIndex() == 8 ? RecomposedTransform.Scale[2] : 0.f;
-
-				if (KeyArea->GetName() == "Rotation.X" ||
-					KeyArea->GetName() == "Rotation.Y" ||
-					KeyArea->GetName() == "Rotation.Z")
+				if (ensureAlwaysMsgf(Channel, TEXT("Channel: %s for Key Area %s does not exist. Keying may not function properly"), *Handle.GetChannelTypeName().ToString(), *KeyArea->GetName().ToString()))
 				{
-					int32 PreviousKey = GetPreviousKey(*Channel, KeyTime);
-					if (PreviousKey != INDEX_NONE && PreviousKey < Channel->GetData().GetValues().Num())
+					double Value =
+						Handle.GetChannelIndex() == 0 ? RecomposedTransform.Translation[0] :
+						Handle.GetChannelIndex() == 1 ? RecomposedTransform.Translation[1] :
+						Handle.GetChannelIndex() == 2 ? RecomposedTransform.Translation[2] :
+						Handle.GetChannelIndex() == 3 ? RecomposedTransform.Rotation.Roll :
+						Handle.GetChannelIndex() == 4 ? RecomposedTransform.Rotation.Pitch :
+						Handle.GetChannelIndex() == 5 ? RecomposedTransform.Rotation.Yaw :
+						Handle.GetChannelIndex() == 6 ? RecomposedTransform.Scale[0] :
+						Handle.GetChannelIndex() == 7 ? RecomposedTransform.Scale[1] :
+						Handle.GetChannelIndex() == 8 ? RecomposedTransform.Scale[2] : 0.f;
+
+					if (KeyArea->GetName() == "Rotation.X" ||
+						KeyArea->GetName() == "Rotation.Y" ||
+						KeyArea->GetName() == "Rotation.Z")
 					{
-						double OldValue = Channel->GetData().GetValues()[PreviousKey].Value;
-						Value = UnwindChannel(OldValue, Value);
+						int32 PreviousKey = GetPreviousKey(*Channel, KeyTime);
+						if (PreviousKey != INDEX_NONE && PreviousKey < Channel->GetData().GetValues().Num())
+						{
+							double OldValue = Channel->GetData().GetValues()[PreviousKey].Value;
+							Value = UnwindChannel(OldValue, Value);
+						}
 					}
-				}
 
-				EMovieSceneKeyInterpolation Interpolation = GetInterpolationMode(Channel, KeyTime, InSequencer.GetKeyInterpolation());
-				AddKeyToChannel(Channel, KeyTime, Value, Interpolation);
+					EMovieSceneKeyInterpolation Interpolation = GetInterpolationMode(Channel, KeyTime, InSequencer.GetKeyInterpolation());
+					AddKeyToChannel(Channel, KeyTime, Value, Interpolation);
+				}
 			}
 			else
 			{
