Sequencer: Constraints getting removed if we compensate out of a sequence boundary so need ot re-evaluate to make sure it's set up.

diff --git a/Engine/Source/Editor/MovieSceneTools/Private/Constraints/MovieSceneConstraintChannelHelper.cpp b/Engine/Source/Editor/MovieSceneTools/Private/Constraints/MovieSceneConstraintChannelHelper.cpp
--- a/Engine/Source/Editor/MovieSceneTools/Private/Constraints/MovieSceneConstraintChannelHelper.cpp
+++ b/Engine/Source/Editor/MovieSceneTools/Private/Constraints/MovieSceneConstraintChannelHelper.cpp
@@ -193,6 +193,8 @@ void FCompensationEvaluator::ComputeLocalTransforms(
 			BakeHelper->StopBaking(MovieScene);
 		}
 	}
+	//get back to where we are at, should also make sure things are active
+	InSequencer->ForceEvaluate();
 }
 
 void FCompensationEvaluator::ComputeLocalTransformsForBaking(UWorld* InWorld, const TSharedPtr<ISequencer>& InSequencer, const TArray<FFrameNumber>& InFrames)
