Sequencer: fix locking/unlocking of cameras

When the level is about to be saved, properly remember which camera actor was locked to which viewport, and restore exactly that after the save operation has completed.

Also, do not mess up cinematic viewport lock when restoring camera locks.

diff --git a/Engine/Source/Editor/MovieSceneTools/Private/TrackEditors/TransformTrackEditor.cpp b/Engine/Source/Editor/MovieSceneTools/Private/TrackEditors/TransformTrackEditor.cpp
--- a/Engine/Source/Editor/MovieSceneTools/Private/TrackEditors/TransformTrackEditor.cpp
+++ b/Engine/Source/Editor/MovieSceneTools/Private/TrackEditors/TransformTrackEditor.cpp
@@ -378,22 +378,55 @@ void F3DTransformTrackEditor::OnPreSaveWorld(UWorld* World)
 {
 	LockedCameraBindings.Reset();
 
+	TSharedPtr<ISequencer> SequencerPtr = GetSequencer();
+
+	// Get the camera binding GUIDs.
 	TArray<FGuid> CameraBindingIDs;
-	GetSequencer()->GetCameraObjectBindings(CameraBindingIDs);
+	SequencerPtr->GetCameraObjectBindings(CameraBindingIDs);
+
+	// Match the camera binding GUIDs with what actor they are bound to.
+	TArray<AActor*> CameraBindingActors;
 	for (const FGuid& CameraBindingID : CameraBindingIDs)
 	{
-		if (IsCameraBindingLocked(CameraBindingID))
+		AActor* BoundActor = nullptr;
+		for (auto Object : SequencerPtr->FindObjectsInCurrentSequence(CameraBindingID))
 		{
-			LockedCameraBindings.Add(CameraBindingID);
+			AActor* Actor = Cast<AActor>(Object.Get());
+			if (Actor != nullptr)
+			{
+				BoundActor = Actor;
+				break;
+			}
+		}
+		CameraBindingActors.Add(BoundActor);
+	}
+
+	// Look at all the editor viewports and see if they're locked to one of our
+	// bound camera actors. If they are, associate them with the binding GUID.
+	for (FLevelEditorViewportClient* LevelVC : GEditor->GetLevelViewportClients())
+	{
+		if (LevelVC && LevelVC->GetViewMode() != VMI_Unknown)
+		{
+			AActor* ActorLock = LevelVC->GetActiveActorLock().Get();
+			if (!ActorLock)
+			{
+				continue;
+			}
+
+			int32 Index = CameraBindingActors.Find(ActorLock);
+			if (Index != INDEX_NONE)
+			{
+				LockedCameraBindings.Add(LevelVC, CameraBindingIDs[Index]);
+			}
 		}
 	}
 }
 
 void F3DTransformTrackEditor::OnPostSaveWorld(UWorld* World)
 {
-	for (const FGuid& CameraBindingID : LockedCameraBindings)
+	for (const TPair<FLevelEditorViewportClient*, FGuid> Pair : LockedCameraBindings)
 	{
-		LockCameraBinding(true, CameraBindingID);
+		LockCameraBinding(true, Pair.Value, Pair.Key, false);
 	}
 
 	LockedCameraBindings.Reset();
@@ -586,11 +619,13 @@ void F3DTransformTrackEditor::OnLockCameraClicked(ECheckBoxState CheckBoxState,
 	LockCameraBinding((CheckBoxState == ECheckBoxState::Checked), ObjectGuid);
 }
 
-void F3DTransformTrackEditor::LockCameraBinding(bool bLock, FGuid ObjectGuid)
+void F3DTransformTrackEditor::LockCameraBinding(bool bLock, FGuid ObjectGuid, FLevelEditorViewportClient* ViewportClient, bool bRemoveCinematicLock)
 {
-	TWeakObjectPtr<AActor> CameraActor;
+	TSharedPtr<ISequencer> SequencerPtr = GetSequencer();
 
-	for (auto Object : GetSequencer()->FindObjectsInCurrentSequence(ObjectGuid))
+	// Find the actor bound to the given object binding ID.
+	AActor* CameraActor = nullptr;
+	for (auto Object : SequencerPtr->FindObjectsInCurrentSequence(ObjectGuid))
 	{
 		AActor* Actor = Cast<AActor>(Object.Get());
 
@@ -601,70 +636,75 @@ void F3DTransformTrackEditor::LockCameraBinding(bool bLock, FGuid ObjectGuid)
 		}
 	}
 
-	// Lock the active viewport to the camera
-	if (bLock)
+	// Bail out if we didn't find the actor to lock the viewport to.
+	if (!CameraActor)
 	{
-		// Set the active viewport or any viewport if there is no active viewport
-		FViewport* ActiveViewport = GEditor->GetActiveViewport();
+		return;
+	}
 
-		FLevelEditorViewportClient* LevelVC = nullptr;
+	// Find the active viewport if no viewport was provided. If no viewport is active, or if the active viewport doesn't
+	// match our requirements, use the first one we find that fits.
+	if (ViewportClient == nullptr)
+	{
+		FViewport* ActiveViewport = GEditor->GetActiveViewport();
 
-		for(FLevelEditorViewportClient* Viewport : GEditor->GetLevelViewportClients())
+		for (FLevelEditorViewportClient* LevelVC : GEditor->GetLevelViewportClients())
 		{		
-			if (Viewport && Viewport->GetViewMode() != VMI_Unknown && Viewport->AllowsCinematicControl())
+			if (LevelVC && 
+					LevelVC->GetViewMode() != VMI_Unknown && 
+					(LevelVC->Viewport == ActiveViewport || ActiveViewport == nullptr))
 			{
-				LevelVC = Viewport;
-
-				if (LevelVC->Viewport == ActiveViewport)
-				{
-					break;
-				}
+				ViewportClient = LevelVC;
+				break;
 			}
 		}
+	}
 
-		if (LevelVC != nullptr && CameraActor.IsValid())
-		{
-			UCameraComponent* CameraComponent = MovieSceneHelpers::CameraComponentFromActor(CameraActor.Get());
+	// Bail out if we didn't find any acceptable viewport.
+	if (!ViewportClient)
+	{
+		return;
+	}
 
-			if (CameraComponent && CameraComponent->ProjectionMode == ECameraProjectionMode::Type::Perspective)
+	if (bLock)
+	{
+		// Lock the given/active/found viewport to the camera.
+		UCameraComponent* CameraComponent = MovieSceneHelpers::CameraComponentFromActor(CameraActor);
+		if (CameraComponent && CameraComponent->ProjectionMode == ECameraProjectionMode::Type::Perspective)
+		{
+			if (ViewportClient->GetViewportType() != LVT_Perspective)
 			{
-				if (LevelVC->GetViewportType() != LVT_Perspective)
-				{
-					LevelVC->SetViewportType(LVT_Perspective);
-				}
+				ViewportClient->SetViewportType(LVT_Perspective);
 			}
+		}
 
+		if (bRemoveCinematicLock)
+		{
 			GetSequencer()->SetPerspectiveViewportCameraCutEnabled(false);
-			LevelVC->SetCinematicActorLock(nullptr);
-			LevelVC->SetActorLock(CameraActor.Get());
-			LevelVC->bLockedCameraView = true;
-			LevelVC->UpdateViewForLockedActor();
-			LevelVC->Invalidate();
+			ViewportClient->SetCinematicActorLock(nullptr);
 		}
+
+		ViewportClient->SetActorLock(CameraActor);
+		ViewportClient->bLockedCameraView = true;
+		ViewportClient->UpdateViewForLockedActor();
+		ViewportClient->Invalidate();
 	}
-	// Otherwise, clear all locks on the camera
 	else
 	{
-		ClearLockedCameras(CameraActor.Get());
-	}
-}
-
-void F3DTransformTrackEditor::ClearLockedCameras(AActor* LockedActor)
-{
-	for(FLevelEditorViewportClient* LevelVC : GEditor->GetLevelViewportClients())
-	{
-		if (LevelVC && LevelVC->GetViewMode() != VMI_Unknown && LevelVC->AllowsCinematicControl())
+		// Clear the lock to this camera on the given/active/found viewport.
+		if (ViewportClient->IsActorLocked(CameraActor))
 		{
-			if (LevelVC->IsActorLocked(LockedActor))
+			if (bRemoveCinematicLock)
 			{
-				LevelVC->SetCinematicActorLock(nullptr);
-				LevelVC->SetActorLock(nullptr);
-				LevelVC->bLockedCameraView = false;
-				LevelVC->ViewFOV = LevelVC->FOVAngle;
-				LevelVC->RemoveCameraRoll();
-				LevelVC->UpdateViewForLockedActor();
-				LevelVC->Invalidate();
+				ViewportClient->SetCinematicActorLock(nullptr);
 			}
+
+			ViewportClient->SetActorLock(nullptr);
+			ViewportClient->bLockedCameraView = false;
+			ViewportClient->ViewFOV = ViewportClient->FOVAngle;
+			ViewportClient->RemoveCameraRoll();
+			ViewportClient->UpdateViewForLockedActor();
+			ViewportClient->Invalidate();
 		}
 	}
 }
diff --git a/Engine/Source/Editor/MovieSceneTools/Private/TrackEditors/TransformTrackEditor.h b/Engine/Source/Editor/MovieSceneTools/Private/TrackEditors/TransformTrackEditor.h
--- a/Engine/Source/Editor/MovieSceneTools/Private/TrackEditors/TransformTrackEditor.h
+++ b/Engine/Source/Editor/MovieSceneTools/Private/TrackEditors/TransformTrackEditor.h
@@ -19,6 +19,7 @@
 
 class AActor;
 struct FAssetData;
+class FLevelEditorViewportClient;
 class SHorizontalBox;
 class UTickableTransformConstraint;
 
@@ -128,17 +129,14 @@ private:
 	/** Delegate for locked camera button */
 	void OnLockCameraClicked(ECheckBoxState CheckBoxState, FGuid ObjectGuid);
 
-	/** Clear locked cameras */
-	void ClearLockedCameras(AActor* LockedActor);
-
 	/** Delegate for camera button lock tooltip */
 	FText GetLockCameraToolTip(FGuid ObjectGuid) const; 
 
 	/** Implementation of checking if a camera is locked */
 	bool IsCameraBindingLocked(FGuid ObjectGuid) const; 
 
-	/** Toggle whether a camera is locked */
-	void LockCameraBinding(bool bLock, FGuid ObjectGuid);
+	/** Toggle whether a camera is locked in the given viewport (or the active viewport if not provided) */
+	void LockCameraBinding(bool bLock, FGuid ObjectGuid, FLevelEditorViewportClient* ViewportClient = nullptr, bool bRemoveCinematicLock = true);
 
 	/** Generates transform keys based on the last transform, the current transform, and other options. 
 		One transform key is generated for each individual key to be added to the section. */
@@ -215,7 +213,7 @@ private:
 	TSharedPtr<FUICommandList> CommandBindings;
 
 	/** List of locked cameras to restore after save */
-	TArray<FGuid> LockedCameraBindings;
+	TMap<FLevelEditorViewportClient*, FGuid> LockedCameraBindings;
 
 	/** Array of sections that are getting undone, we need to recreate any constraint channel add, move key delegates to them*/
 	mutable TArray<TWeakObjectPtr<UMovieScene3DTransformSection>> SectionsGettingUndone;
