Sequencer: Fix new potential issue in Event Tracks caused by my CL27752174. That change introduced TObjectPtr<UObject> references instead of FString references to object parameters in the director blueprint endpoints. The reason for that change was to fix UE-165499 by allowing proper redirectors to get created and redirector fixup.

However, with that change, users can introduce illegal hard references to actors in levels by creating Actor parameters for event track endpoints and then creating keys on event tracks with that parameter pointing at actors in levels in Sequencer. These references then prevent the sequence from being saved and can create crashes.

This change changes the TObjectPtr<UObject> to FSoftObjectPath instead. As the previous change has not yet been rolled out in a UE version or hotfix, this change should be relatively safe. I have tested UE-165499 again with this new change and redirectors are still being resolved properly.

For completeness, UE-165499 should also be reverified.

diff --git a/Engine/Source/Editor/MovieSceneTools/Private/MovieSceneDirectorBlueprintEndpointCustomization.cpp b/Engine/Source/Editor/MovieSceneTools/Private/MovieSceneDirectorBlueprintEndpointCustomization.cpp
--- a/Engine/Source/Editor/MovieSceneTools/Private/MovieSceneDirectorBlueprintEndpointCustomization.cpp
+++ b/Engine/Source/Editor/MovieSceneTools/Private/MovieSceneDirectorBlueprintEndpointCustomization.cpp
@@ -311,9 +311,9 @@ void FMovieSceneDirectorBlueprintEndpointCustomization::CustomizeChildren(TShare
 
 					// We have an override for this variable
 					bool bImportSuccess = false;
-					if (PayloadVariable->ObjectValue)
+					if (PayloadVariable->ObjectValue.IsValid())
 					{
-						bImportSuccess = FBlueprintEditorUtils::PropertyValueFromString(Field, PayloadVariable->ObjectValue->GetPathName(), StructData->GetStructMemory());
+						bImportSuccess = FBlueprintEditorUtils::PropertyValueFromString(Field, PayloadVariable->ObjectValue.ToString(), StructData->GetStructMemory());
 					}
 					if (!bImportSuccess)
 					{
diff --git a/Engine/Source/Editor/MovieSceneTools/Private/MovieSceneDirectorBlueprintUtils.cpp b/Engine/Source/Editor/MovieSceneTools/Private/MovieSceneDirectorBlueprintUtils.cpp
--- a/Engine/Source/Editor/MovieSceneTools/Private/MovieSceneDirectorBlueprintUtils.cpp
+++ b/Engine/Source/Editor/MovieSceneTools/Private/MovieSceneDirectorBlueprintUtils.cpp
@@ -309,9 +309,9 @@ FMovieSceneDirectorBlueprintEntrypointResult FMovieSceneDirectorBlueprintUtils::
 			else
 			{
 				bool bMarkAsModified = false;
-				if (PayloadVariable->ObjectValue)
+				if (PayloadVariable->ObjectValue.IsValid())
 				{
-					Schema->TrySetDefaultObject(*Pin, PayloadVariable->ObjectValue, bMarkAsModified);
+					Schema->TrySetDefaultObject(*Pin, PayloadVariable->ObjectValue.ResolveObject(), bMarkAsModified);
 				}
 				else if (!PayloadVariable->Value.IsEmpty())
 				{
@@ -403,9 +403,9 @@ bool FMovieSceneDirectorBlueprintUtils::GenerateEntryPointRawActorParameter(
 
 	// Set the default value for the path string/object
 	const bool bMarkAsModified = false;
-	if (PayloadValue.ObjectValue)
+	if (PayloadValue.ObjectValue.IsValid())
 	{
-		Schema->TrySetDefaultObject(*PathInput, PayloadValue.ObjectValue, bMarkAsModified);
+		Schema->TrySetDefaultObject(*PathInput, PayloadValue.ObjectValue.ResolveObject(), bMarkAsModified);
 	}
 	else if (!PayloadValue.Value.IsEmpty())
 	{
diff --git a/Engine/Source/Editor/MovieSceneTools/Public/MovieSceneDirectorBlueprintUtils.h b/Engine/Source/Editor/MovieSceneTools/Public/MovieSceneDirectorBlueprintUtils.h
--- a/Engine/Source/Editor/MovieSceneTools/Public/MovieSceneDirectorBlueprintUtils.h
+++ b/Engine/Source/Editor/MovieSceneTools/Public/MovieSceneDirectorBlueprintUtils.h
@@ -57,7 +57,7 @@ struct MOVIESCENETOOLS_API FMovieSceneDirectorBlueprintEndpointParameter
 struct MOVIESCENETOOLS_API FMovieSceneDirectorBlueprintVariableValue
 {
 	/* In the case the variable is a UObject, reference to the UObject */
-	TObjectPtr<UObject> ObjectValue;
+	FSoftObjectPath ObjectValue;
 
 	/** String value of the variable */
 	FString Value;
diff --git a/Engine/Source/Runtime/MovieScene/Public/MovieSceneDynamicBinding.h b/Engine/Source/Runtime/MovieScene/Public/MovieSceneDynamicBinding.h
--- a/Engine/Source/Runtime/MovieScene/Public/MovieSceneDynamicBinding.h
+++ b/Engine/Source/Runtime/MovieScene/Public/MovieSceneDynamicBinding.h
@@ -22,7 +22,7 @@ struct FMovieSceneDynamicBindingPayloadVariable
 	GENERATED_BODY()
 		
 	UPROPERTY()
-	TObjectPtr<UObject> ObjectValue;
+	FSoftObjectPath ObjectValue;
 
 	UPROPERTY(EditAnywhere, Category = "Sequencer|Dynamic Binding")
 	FString Value;
diff --git a/Engine/Source/Runtime/MovieSceneTracks/Public/Channels/MovieSceneEvent.h b/Engine/Source/Runtime/MovieSceneTracks/Public/Channels/MovieSceneEvent.h
--- a/Engine/Source/Runtime/MovieSceneTracks/Public/Channels/MovieSceneEvent.h
+++ b/Engine/Source/Runtime/MovieSceneTracks/Public/Channels/MovieSceneEvent.h
@@ -13,6 +13,7 @@
 #include "UObject/ObjectPtr.h"
 #include "UObject/WeakObjectPtr.h"
 #include "UObject/WeakObjectPtrTemplates.h"
+#include "UObject/SoftObjectPtr.h"
 
 #include "MovieSceneEvent.generated.h"
 
@@ -35,7 +36,7 @@ struct FMovieSceneEventPayloadVariable
 
 	UPROPERTY()
 	/** If the value for this pin should be an object, we store a pointer to it */
-	TObjectPtr<class UObject> ObjectValue;
+	FSoftObjectPath ObjectValue;
 
 	UPROPERTY(EditAnywhere, Category="Sequencer|Event")
 	FString Value;
