Sequencer: Rework DrawFrameTimeHint so that any section can supply a frame time to draw.

Rework layer id for DrawFrameTimeHint so that it's drawn with a layer id after the section title drawing.

#rnx

diff --git a/Engine/Plugins/Animation/ControlRig/Source/ControlRigEditor/Private/Sequencer/ControlRigParameterTrackEditor.cpp b/Engine/Plugins/Animation/ControlRig/Source/ControlRigEditor/Private/Sequencer/ControlRigParameterTrackEditor.cpp
--- a/Engine/Plugins/Animation/ControlRig/Source/ControlRigEditor/Private/Sequencer/ControlRigParameterTrackEditor.cpp
+++ b/Engine/Plugins/Animation/ControlRig/Source/ControlRigEditor/Private/Sequencer/ControlRigParameterTrackEditor.cpp
@@ -25,7 +25,6 @@
 #include "UnrealEdGlobals.h"
 #include "ClassViewerModule.h"
 #include "ClassViewerFilter.h"
-#include "CommonMovieSceneTools.h"
 #include "AssetRegistry/AssetRegistryModule.h"
 #include "IContentBrowserSingleton.h"
 #include "ContentBrowserModule.h"
diff --git a/Engine/Plugins/Animation/ControlRig/Source/ControlRigEditor/Private/Sequencer/ControlRigSpaceChannelEditors.cpp b/Engine/Plugins/Animation/ControlRig/Source/ControlRigEditor/Private/Sequencer/ControlRigSpaceChannelEditors.cpp
--- a/Engine/Plugins/Animation/ControlRig/Source/ControlRigEditor/Private/Sequencer/ControlRigSpaceChannelEditors.cpp
+++ b/Engine/Plugins/Animation/ControlRig/Source/ControlRigEditor/Private/Sequencer/ControlRigSpaceChannelEditors.cpp
@@ -54,7 +54,6 @@
 #include "ScopedTransaction.h"
 #include "IControlRigEditorModule.h"
 #include "TimerManager.h"
-#include "CommonMovieSceneTools.h"
 #include "SequencerSectionPainter.h"
 #include "Rendering/DrawElements.h"
 #include "Fonts/FontMeasure.h"
diff --git a/Engine/Plugins/Experimental/GeometryCollectionPlugin/Source/GeometryCollectionSequencer/Private/GeometryCollectionTrackEditor.cpp b/Engine/Plugins/Experimental/GeometryCollectionPlugin/Source/GeometryCollectionSequencer/Private/GeometryCollectionTrackEditor.cpp
--- a/Engine/Plugins/Experimental/GeometryCollectionPlugin/Source/GeometryCollectionSequencer/Private/GeometryCollectionTrackEditor.cpp
+++ b/Engine/Plugins/Experimental/GeometryCollectionPlugin/Source/GeometryCollectionSequencer/Private/GeometryCollectionTrackEditor.cpp
@@ -15,7 +15,6 @@
 #include "UnrealEdGlobals.h"
 #include "MovieSceneGeometryCollectionTrack.h"
 #include "MovieSceneGeometryCollectionSection.h"
-#include "CommonMovieSceneTools.h"
 #include "ContentBrowserModule.h"
 #include "SequencerUtilities.h"
 #include "ISectionLayoutBuilder.h"
diff --git a/Engine/Plugins/Media/MediaCompositing/Source/MediaCompositingEditor/Private/Sequencer/MediaThumbnailSection.cpp b/Engine/Plugins/Media/MediaCompositing/Source/MediaCompositingEditor/Private/Sequencer/MediaThumbnailSection.cpp
--- a/Engine/Plugins/Media/MediaCompositing/Source/MediaCompositingEditor/Private/Sequencer/MediaThumbnailSection.cpp
+++ b/Engine/Plugins/Media/MediaCompositing/Source/MediaCompositingEditor/Private/Sequencer/MediaThumbnailSection.cpp
@@ -23,7 +23,6 @@
 #include "SequencerSectionPainter.h"
 #include "TrackEditorThumbnail/TrackEditorThumbnailPool.h"
 #include "Compilation/MovieSceneCompiledDataManager.h"
-#include "CommonMovieSceneTools.h"
 #include "Widgets/Layout/SBox.h"
 #include "Widgets/Text/STextBlock.h"
 #include "TimeToPixel.h"
diff --git a/Engine/Plugins/Runtime/HairStrands/Source/HairStrandsEditor/Private/GroomCacheTrackEditor.cpp b/Engine/Plugins/Runtime/HairStrands/Source/HairStrandsEditor/Private/GroomCacheTrackEditor.cpp
--- a/Engine/Plugins/Runtime/HairStrands/Source/HairStrandsEditor/Private/GroomCacheTrackEditor.cpp
+++ b/Engine/Plugins/Runtime/HairStrands/Source/HairStrandsEditor/Private/GroomCacheTrackEditor.cpp
@@ -1,7 +1,6 @@
 // Copyright Epic Games, Inc. All Rights Reserved.
 
 #include "GroomCacheTrackEditor.h"
-#include "CommonMovieSceneTools.h"
 #include "Fonts/FontMeasure.h"
 #include "Framework/Application/SlateApplication.h"
 #include "Framework/MultiBox/MultiBoxBuilder.h"
diff --git a/Engine/Source/Editor/MovieSceneTools/Private/CommonMovieSceneTools.cpp b/Engine/Source/Editor/MovieSceneTools/Private/CommonMovieSceneTools.cpp
deleted file mode 100644
diff --git a/Engine/Source/Editor/MovieSceneTools/Private/Sections/CinematicShotSection.cpp b/Engine/Source/Editor/MovieSceneTools/Private/Sections/CinematicShotSection.cpp
--- a/Engine/Source/Editor/MovieSceneTools/Private/Sections/CinematicShotSection.cpp
+++ b/Engine/Source/Editor/MovieSceneTools/Private/Sections/CinematicShotSection.cpp
@@ -18,7 +18,6 @@
 #include "Tracks/MovieSceneCameraCutTrack.h"
 #include "Sections/MovieSceneCameraCutSection.h"
 #include "Evaluation/MovieSceneEvaluationTemplateInstance.h"
-#include "CommonMovieSceneTools.h"
 #include "Subsystems/AssetEditorSubsystem.h"
 #include "Editor.h"
 
diff --git a/Engine/Source/Editor/MovieSceneTools/Private/Sections/EventSection.cpp b/Engine/Source/Editor/MovieSceneTools/Private/Sections/EventSection.cpp
--- a/Engine/Source/Editor/MovieSceneTools/Private/Sections/EventSection.cpp
+++ b/Engine/Source/Editor/MovieSceneTools/Private/Sections/EventSection.cpp
@@ -7,7 +7,6 @@
 #include "SequencerSectionPainter.h"
 #include "ISectionLayoutBuilder.h"
 #include "Styling/AppStyle.h"
-#include "CommonMovieSceneTools.h"
 #include "Tracks/MovieSceneEventTrack.h"
 #include "Sections/MovieSceneEventSection.h"
 #include "Sections/MovieSceneEventTriggerSection.h"
@@ -22,6 +21,7 @@
 #include "MovieSceneSequence.h"
 #include "EditorFontGlyphs.h"
 #include "ScopedTransaction.h"
+#include "TimeToPixel.h"
 
 #define LOCTEXT_NAMESPACE "EventSection"
 
diff --git a/Engine/Source/Editor/MovieSceneTools/Private/Sections/MovieSceneSectionHelpers.cpp b/Engine/Source/Editor/MovieSceneTools/Private/Sections/MovieSceneSectionHelpers.cpp
--- a/Engine/Source/Editor/MovieSceneTools/Private/Sections/MovieSceneSectionHelpers.cpp
+++ b/Engine/Source/Editor/MovieSceneTools/Private/Sections/MovieSceneSectionHelpers.cpp
@@ -5,13 +5,13 @@
 #include "Sections/MovieSceneColorSection.h"
 #include "Channels/MovieSceneChannelProxy.h"
 #include "Channels/MovieSceneFloatChannel.h"
-#include "CommonMovieSceneTools.h"
 
 #include "Misc/FrameNumber.h"
 #include "Widgets/Colors/SColorPicker.h"
 #include "ScopedTransaction.h"
 #include "Framework/Application/SlateApplication.h"
 #include "Engine/Engine.h"
+#include "TimeToPixel.h"
 
 void MovieSceneSectionHelpers::ConsolidateColorCurves( TArray< TTuple<float, FLinearColor> >& OutColorKeys, const FLinearColor& DefaultColor, TArrayView<const FMovieSceneFloatChannel* const> ColorChannels, const FTimeToPixel& TimeConverter )
 {
diff --git a/Engine/Source/Editor/MovieSceneTools/Private/TrackEditors/AudioTrackEditor.cpp b/Engine/Source/Editor/MovieSceneTools/Private/TrackEditors/AudioTrackEditor.cpp
--- a/Engine/Source/Editor/MovieSceneTools/Private/TrackEditors/AudioTrackEditor.cpp
+++ b/Engine/Source/Editor/MovieSceneTools/Private/TrackEditors/AudioTrackEditor.cpp
@@ -23,7 +23,6 @@
 #include "UnrealEdGlobals.h"
 #include "Tracks/MovieSceneAudioTrack.h"
 #include "Sections/MovieSceneAudioSection.h"
-#include "CommonMovieSceneTools.h"
 #include "AudioDevice.h"
 #include "Sound/SoundNodeWavePlayer.h"
 #include "Sound/SoundNodeDialoguePlayer.h"
@@ -39,7 +38,7 @@
 
 #include "DragAndDrop/AssetDragDropOp.h"
 #include "Misc/QualifiedFrameTime.h"
-
+#include "TimeToPixel.h"
 
 #define LOCTEXT_NAMESPACE "FAudioTrackEditor"
 
diff --git a/Engine/Source/Editor/MovieSceneTools/Private/TrackEditors/ParticleTrackEditor.cpp b/Engine/Source/Editor/MovieSceneTools/Private/TrackEditors/ParticleTrackEditor.cpp
--- a/Engine/Source/Editor/MovieSceneTools/Private/TrackEditors/ParticleTrackEditor.cpp
+++ b/Engine/Source/Editor/MovieSceneTools/Private/TrackEditors/ParticleTrackEditor.cpp
@@ -14,10 +14,9 @@
 #include "Tracks/MovieSceneParticleTrack.h"
 #include "ISectionLayoutBuilder.h"
 #include "Sections/MovieSceneParticleSection.h"
-#include "CommonMovieSceneTools.h"
 #include "Particles/ParticleLODLevel.h"
 #include "Particles/ParticleModuleRequired.h"
-
+#include "TimeToPixel.h"
 
 namespace AnimatableParticleEditorConstants
 {
diff --git a/Engine/Source/Editor/MovieSceneTools/Private/TrackEditors/SkeletalAnimationTrackEditor.cpp b/Engine/Source/Editor/MovieSceneTools/Private/TrackEditors/SkeletalAnimationTrackEditor.cpp
--- a/Engine/Source/Editor/MovieSceneTools/Private/TrackEditors/SkeletalAnimationTrackEditor.cpp
+++ b/Engine/Source/Editor/MovieSceneTools/Private/TrackEditors/SkeletalAnimationTrackEditor.cpp
@@ -19,7 +19,6 @@
 #include "UnrealEdGlobals.h"
 #include "Tracks/MovieSceneSkeletalAnimationTrack.h"
 #include "Sections/MovieSceneSkeletalAnimationSection.h"
-#include "CommonMovieSceneTools.h"
 #include "AssetRegistry/AssetRegistryModule.h"
 #include "IContentBrowserSingleton.h"
 #include "ContentBrowserModule.h"
@@ -47,7 +46,6 @@
 #include "Engine/SimpleConstructionScript.h"
 #include "Engine/Blueprint.h"
 
-#include "CommonMovieSceneTools.h"
 #include "Subsystems/AssetEditorSubsystem.h"
 #include "Editor.h"
 
@@ -81,6 +79,7 @@
 #include "AnimSequenceLevelSequenceLink.h"
 #include "UObject/SavePackage.h"
 #include "AnimSequencerInstanceProxy.h"
+#include "TimeToPixel.h"
 
 int32 FSkeletalAnimationTrackEditor::NumberActive = 0;
 
@@ -579,6 +578,67 @@ FText FSkeletalAnimationSection::GetSectionToolTip() const
 	return FText::GetEmpty();
 }
 
+TOptional<FFrameTime> FSkeletalAnimationSection::GetSectionTime(FSequencerSectionPainter& InPainter) const
+{
+	if (!InPainter.bIsSelected || !Sequencer.Pin() || Section.Params.Animation == nullptr)
+	{
+		return TOptional<FFrameTime>();
+	}
+
+	FFrameTime CurrentTime = Sequencer.Pin()->GetLocalTime().Time;
+	if (!Section.GetRange().Contains(CurrentTime.FrameNumber))
+	{
+		return TOptional<FFrameTime>();
+	}
+
+	const FTimeToPixel& TimeToPixelConverter = InPainter.GetTimeConverter();
+	FFrameRate TickResolution = TimeToPixelConverter.GetTickResolution();
+
+	// Draw the current time next to the scrub handle
+	const double AnimTime = Section.MapTimeToAnimation(CurrentTime, TickResolution);
+	const FFrameRate SamplingFrameRate = Section.Params.Animation->GetSamplingFrameRate();
+
+	FQualifiedFrameTime HintFrameTime;
+	if (!UAnimationBlueprintLibrary::EvaluateRootBoneTimecodeAttributesAtTime(Section.Params.Animation, static_cast<float>(AnimTime), HintFrameTime))
+	{
+		const FFrameTime FrameTime = SamplingFrameRate.AsFrameTime(AnimTime);
+		HintFrameTime = FQualifiedFrameTime(FrameTime, SamplingFrameRate);
+	}
+
+	// Get the desired frame display format and zero padding from
+	// the sequencer settings, if possible.
+	TAttribute<EFrameNumberDisplayFormats> DisplayFormatAttr(EFrameNumberDisplayFormats::Frames);
+	TAttribute<uint8> ZeroPadFrameNumbersAttr(0u);
+	if (const USequencerSettings* SequencerSettings = Sequencer.Pin()->GetSequencerSettings())
+	{
+		DisplayFormatAttr.Set(SequencerSettings->GetTimeDisplayFormat());
+		ZeroPadFrameNumbersAttr.Set(SequencerSettings->GetZeroPadFrames());
+	}
+
+	// No frame rate conversion necessary since we're displaying
+	// the source frame time/rate.
+	const TAttribute<FFrameRate> TickResolutionAttr(HintFrameTime.Rate);
+	const TAttribute<FFrameRate> DisplayRateAttr(HintFrameTime.Rate);
+
+	FFrameNumberInterface FrameNumberInterface(DisplayFormatAttr, ZeroPadFrameNumbersAttr, TickResolutionAttr, DisplayRateAttr);
+
+	float Subframe = 0.0f;
+	if (UAnimationBlueprintLibrary::EvaluateRootBoneTimecodeSubframeAttributeAtTime(Section.Params.Animation, static_cast<float>(AnimTime), Subframe))
+	{
+		if (FMath::IsNearlyEqual(Subframe, FMath::RoundToFloat(Subframe)))
+		{
+			FrameNumberInterface.SetSubframeIndicator(FString::Printf(TEXT(" (%d)"), FMath::RoundToInt(Subframe)));
+		}
+		else
+		{
+			FrameNumberInterface.SetSubframeIndicator(FString::Printf(TEXT(" (%s)"), *LexToSanitizedString(Subframe)));
+		}
+	}
+
+	return HintFrameTime.Time;
+}
+
+
 float FSkeletalAnimationSection::GetSectionHeight() const
 {
 	return (float)SkeletalAnimationEditorConstants::AnimationTrackHeight;
@@ -638,57 +698,6 @@ int32 FSkeletalAnimationSection::OnPaintSection( FSequencerSectionPainter& Paint
 		}
 	}
 
-	TSharedPtr<ISequencer> SequencerPtr = Sequencer.Pin();
-	if (Painter.bIsSelected && SequencerPtr.IsValid())
-	{
-		FFrameTime CurrentTime = SequencerPtr->GetLocalTime().Time;
-		if (Section.GetRange().Contains(CurrentTime.FrameNumber) && Section.Params.Animation != nullptr)
-		{
-			// Draw the current time next to the scrub handle
-			const double AnimTime = Section.MapTimeToAnimation(CurrentTime, TickResolution);
-			const FFrameRate SamplingFrameRate = Section.Params.Animation->GetSamplingFrameRate();
-
-			FQualifiedFrameTime HintFrameTime;
-			if (!UAnimationBlueprintLibrary::EvaluateRootBoneTimecodeAttributesAtTime(Section.Params.Animation, static_cast<float>(AnimTime), HintFrameTime))
-			{
-				const FFrameTime FrameTime = SamplingFrameRate.AsFrameTime(AnimTime);
-				HintFrameTime = FQualifiedFrameTime(FrameTime, SamplingFrameRate);
-			}
-
-			// Get the desired frame display format and zero padding from
-			// the sequencer settings, if possible.
-			TAttribute<EFrameNumberDisplayFormats> DisplayFormatAttr(EFrameNumberDisplayFormats::Frames);
-			TAttribute<uint8> ZeroPadFrameNumbersAttr(0u);
-			if (const USequencerSettings* SequencerSettings = SequencerPtr->GetSequencerSettings())
-			{
-				DisplayFormatAttr.Set(SequencerSettings->GetTimeDisplayFormat());
-				ZeroPadFrameNumbersAttr.Set(SequencerSettings->GetZeroPadFrames());
-			}
-
-			// No frame rate conversion necessary since we're displaying
-			// the source frame time/rate.
-			const TAttribute<FFrameRate> TickResolutionAttr(HintFrameTime.Rate);
-			const TAttribute<FFrameRate> DisplayRateAttr(HintFrameTime.Rate);
-
-			FFrameNumberInterface FrameNumberInterface(DisplayFormatAttr, ZeroPadFrameNumbersAttr, TickResolutionAttr, DisplayRateAttr);
-
-			float Subframe = 0.0f;
-			if (UAnimationBlueprintLibrary::EvaluateRootBoneTimecodeSubframeAttributeAtTime(Section.Params.Animation, static_cast<float>(AnimTime), Subframe))
-			{
-				if (FMath::IsNearlyEqual(Subframe, FMath::RoundToFloat(Subframe)))
-				{
-					FrameNumberInterface.SetSubframeIndicator(FString::Printf(TEXT(" (%d)"), FMath::RoundToInt(Subframe)));
-				}
-				else
-				{
-					FrameNumberInterface.SetSubframeIndicator(FString::Printf(TEXT(" (%s)"), *LexToSanitizedString(Subframe)));
-				}
-			}
-
-			DrawFrameTimeHint(Painter, CurrentTime, HintFrameTime.Time, &FrameNumberInterface);
-		}
-	}
-	
 	return LayerId;
 }
 
diff --git a/Engine/Source/Editor/MovieSceneTools/Private/TrackEditors/SkeletalAnimationTrackEditor.h b/Engine/Source/Editor/MovieSceneTools/Private/TrackEditors/SkeletalAnimationTrackEditor.h
--- a/Engine/Source/Editor/MovieSceneTools/Private/TrackEditors/SkeletalAnimationTrackEditor.h
+++ b/Engine/Source/Editor/MovieSceneTools/Private/TrackEditors/SkeletalAnimationTrackEditor.h
@@ -151,6 +151,7 @@ public:
 	virtual UMovieSceneSection* GetSectionObject() override;
 	virtual FText GetSectionTitle() const override;
 	virtual FText GetSectionToolTip() const override;
+	virtual TOptional<FFrameTime> GetSectionTime(FSequencerSectionPainter& InPainter) const override;
 	virtual float GetSectionHeight() const override;
 	virtual FMargin GetContentPadding() const override;
 	virtual int32 OnPaintSection( FSequencerSectionPainter& Painter ) const override;
diff --git a/Engine/Source/Editor/MovieSceneTools/Private/TrackEditors/SubTrackEditor.cpp b/Engine/Source/Editor/MovieSceneTools/Private/TrackEditors/SubTrackEditor.cpp
--- a/Engine/Source/Editor/MovieSceneTools/Private/TrackEditors/SubTrackEditor.cpp
+++ b/Engine/Source/Editor/MovieSceneTools/Private/TrackEditors/SubTrackEditor.cpp
@@ -37,8 +37,6 @@
 #include "IStructureDetailsView.h"
 #include "PropertyEditorModule.h"
 
-#include "CommonMovieSceneTools.h"
-
 namespace SubTrackEditorConstants
 {
 	const float TrackHeight = 50.0f;
diff --git a/Engine/Source/Editor/MovieSceneTools/Private/TrackEditors/SubTrackEditorBase.cpp b/Engine/Source/Editor/MovieSceneTools/Private/TrackEditors/SubTrackEditorBase.cpp
--- a/Engine/Source/Editor/MovieSceneTools/Private/TrackEditors/SubTrackEditorBase.cpp
+++ b/Engine/Source/Editor/MovieSceneTools/Private/TrackEditors/SubTrackEditorBase.cpp
@@ -111,35 +111,6 @@ FSubSectionPainterResult FSubSectionPainterUtil::PaintSection(TSharedPtr<const I
         );
     }
 
-    if (Params.bDrawFrameNumberHintWhenSelected && InPainter.bIsSelected && Sequencer.IsValid())
-    {
-        FFrameTime CurrentTime = Sequencer->GetLocalTime().Time;
-        if (SectionRange.Contains(CurrentTime.FrameNumber))
-        {
-            const UMovieScene* SubSequenceMovieScene = SectionObject.GetSequence()->GetMovieScene();
-            const FFrameRate DisplayRate = SubSequenceMovieScene->GetDisplayRate();
-            const FFrameRate TickResolution = SubSequenceMovieScene->GetTickResolution();
-			const FFrameTime HintFrameTime = CurrentTime * SectionObject.OuterToInnerTransform();
-
-			// Get the desired frame display format and zero padding from
-			// the sequencer settings, if possible.
-			TAttribute<EFrameNumberDisplayFormats> DisplayFormatAttr(EFrameNumberDisplayFormats::Frames);
-			TAttribute<uint8> ZeroPadFrameNumbersAttr(0u);
-			if (const USequencerSettings* SequencerSettings = Sequencer->GetSequencerSettings())
-			{
-				DisplayFormatAttr.Set(SequencerSettings->GetTimeDisplayFormat());
-				ZeroPadFrameNumbersAttr.Set(SequencerSettings->GetZeroPadFrames());
-			}
-
-			const TAttribute<FFrameRate> TickResolutionAttr(TickResolution);
-			const TAttribute<FFrameRate> DisplayRateAttr(DisplayRate);
-
-			const FFrameNumberInterface FrameNumberInterface(DisplayFormatAttr, ZeroPadFrameNumbersAttr, TickResolutionAttr, DisplayRateAttr);
-
-			DrawFrameTimeHint(InPainter, CurrentTime, HintFrameTime, &FrameNumberInterface);
-        }
-    }
-
     InPainter.LayerId = LayerId;
 
     return FSSPR_Success;
diff --git a/Engine/Source/Editor/MovieSceneTools/Public/CommonMovieSceneTools.h b/Engine/Source/Editor/MovieSceneTools/Public/CommonMovieSceneTools.h
--- a/Engine/Source/Editor/MovieSceneTools/Public/CommonMovieSceneTools.h
+++ b/Engine/Source/Editor/MovieSceneTools/Public/CommonMovieSceneTools.h
@@ -1,25 +1,13 @@
 // Copyright Epic Games, Inc. All Rights Reserved.
 
-#pragma once
-
-#include "CoreTypes.h"
-#include "FrameNumberNumericInterface.h"
-#include "Layout/Geometry.h"
-#include "Misc/FrameRate.h"
-#include "Misc/FrameTime.h"
-#include "SequencerSectionPainter.h"
-#include "TimeToPixel.h"
-
-class FSequencerSectionPainter;
-struct FFrameNumberInterface;
-
 /**
- * Draw a frame time next to the scrub handle
- *
- * @param InPainter Structure that affords common painting operations
- * @param CurrentTime Current time of the scrub handle
- * @param FrameTime Frame time to draw
- * @param FrameNumberInterface (optional) Interface to control the display format and/or frame rate conversion of the drawn frame time.
- *            If not provided, the frame time will be drawn as a frame number without any subframe.
+ * DEPRECATED CommonMovieSceneTools.h. This is no longer used.
  */
-void DrawFrameTimeHint(FSequencerSectionPainter& InPainter, const FFrameTime& CurrentTime, const FFrameTime& FrameTime, const FFrameNumberInterface* FrameNumberInterface = nullptr);
+ 
+ #pragma once
+  
+#ifdef _MSC_VER
+    #pragma message(__FILE__"(9): warning: CommonMovieSceneTools.h is no longer used")
+#else
+    #pragma message("CommonMovieSceneTools.h is no longer used")
+#endif
\ No newline at end of file
diff --git a/Engine/Source/Editor/MovieSceneTools/Public/TrackEditors/SubTrackEditorBase.h b/Engine/Source/Editor/MovieSceneTools/Public/TrackEditors/SubTrackEditorBase.h
--- a/Engine/Source/Editor/MovieSceneTools/Public/TrackEditors/SubTrackEditorBase.h
+++ b/Engine/Source/Editor/MovieSceneTools/Public/TrackEditors/SubTrackEditorBase.h
@@ -6,7 +6,6 @@
 #include "Editor.h"
 #include "ISequencerSection.h"
 #include "Sections/MovieSceneSubSection.h"
-#include "CommonMovieSceneTools.h"
 #include "Styling/AppStyle.h"
 #include "Framework/Application/SlateApplication.h"
 #include "ISequencer.h"
@@ -27,15 +26,14 @@ class ISequencer;
  */
 struct FSubSectionPainterParams
 {
-    FSubSectionPainterParams() : bShowTrackNum(true), bDrawFrameNumberHintWhenSelected(true)
+    FSubSectionPainterParams() : bShowTrackNum(true)
     {}
     FSubSectionPainterParams(FMargin InContentPadding)
-	: ContentPadding(InContentPadding), bShowTrackNum(true), bDrawFrameNumberHintWhenSelected(true)
+	: ContentPadding(InContentPadding), bShowTrackNum(true)
     {}
 
     FMargin ContentPadding;
     bool bShowTrackNum;
-    bool bDrawFrameNumberHintWhenSelected;
 };
 
 /**
@@ -137,6 +135,7 @@ public:
     virtual UMovieSceneSection* GetSectionObject() override;
     virtual FText GetSectionTitle() const override;
 	virtual FText GetSectionToolTip() const override;
+	virtual TOptional<FFrameTime> GetSectionTime(FSequencerSectionPainter& InPainter) const override;
 	virtual float GetSectionHeight() const override;
     virtual bool IsReadOnly() const override;
     virtual int32 OnPaintSection( FSequencerSectionPainter& InPainter ) const override;
@@ -253,6 +252,33 @@ FText TSubSectionMixin<ParentSectionClass>::GetSectionToolTip() const
 		);
 	}
 }
+
+template<typename ParentSectionClass>
+TOptional<FFrameTime> TSubSectionMixin<ParentSectionClass>::GetSectionTime(FSequencerSectionPainter& InPainter) const
+{
+	if (!InPainter.bIsSelected)
+	{
+		return TOptional<FFrameTime>();
+	}
+
+	TSharedPtr<ISequencer> Sequencer = GetSequencer();
+	if (!Sequencer)
+	{
+		return TOptional<FFrameTime>();
+	}
+
+	FFrameTime CurrentTime = Sequencer->GetLocalTime().Time;
+	if (!SubSectionObject.GetRange().Contains(CurrentTime.FrameNumber))
+	{
+		return TOptional<FFrameTime>();
+	}
+
+	const UMovieScene* SubSequenceMovieScene = SubSectionObject.GetSequence()->GetMovieScene();
+	const FFrameTime HintFrameTime = CurrentTime * SubSectionObject.OuterToInnerTransform();
+
+	return HintFrameTime;
+}
+
 #undef LOCTEXT_NAMESPACE
 
 template<typename ParentSectionClass>
diff --git a/Engine/Source/Editor/Sequencer/Private/SSequencerDebugVisualizer.cpp b/Engine/Source/Editor/Sequencer/Private/SSequencerDebugVisualizer.cpp
--- a/Engine/Source/Editor/Sequencer/Private/SSequencerDebugVisualizer.cpp
+++ b/Engine/Source/Editor/Sequencer/Private/SSequencerDebugVisualizer.cpp
@@ -3,7 +3,6 @@
 #include "SSequencerDebugVisualizer.h"
 #include "MVVM/Selection/Selection.h"
 #include "MVVM/ViewModels/SequencerEditorViewModel.h"
-#include "CommonMovieSceneTools.h"
 #include "Styling/AppStyle.h"
 #include "Widgets/SToolTip.h"
 #include "Widgets/Layout/SBorder.h"
diff --git a/Engine/Source/Editor/Sequencer/Private/SSequencerSection.cpp b/Engine/Source/Editor/Sequencer/Private/SSequencerSection.cpp
--- a/Engine/Source/Editor/Sequencer/Private/SSequencerSection.cpp
+++ b/Engine/Source/Editor/Sequencer/Private/SSequencerSection.cpp
@@ -24,13 +24,14 @@
 #include "Sequencer.h"
 #include "SequencerSectionPainter.h"
 #include "MovieSceneSequence.h"
-#include "CommonMovieSceneTools.h"
 #include "ISequencerEditTool.h"
 #include "ISequencerSection.h"
 #include "SequencerHotspots.h"
 #include "Widgets/SOverlay.h"
 #include "MovieScene.h"
 #include "Fonts/FontCache.h"
+#include "Fonts/FontMeasure.h"
+#include "FrameNumberNumericInterface.h"
 #include "Framework/Application/SlateApplication.h"
 #include "MovieSceneTimeHelpers.h"
 #include "Tracks/MovieScenePropertyTrack.h"
@@ -1312,6 +1313,58 @@ TSharedPtr<FTrackAreaViewModel> SSequencerSection::GetTrackAreaViewModel() const
 	return nullptr;
 }
 
+void DrawFrameTimeHint(FSequencerSectionPainter& InPainter, const FFrameTime& CurrentTime, const FFrameTime& FrameTime, const FFrameNumberInterface* FrameNumberInterface, int32 LayerId)
+{
+	FString FrameTimeString;
+	if (FrameNumberInterface)
+	{
+		FrameTimeString = FrameNumberInterface->ToString(FrameTime.AsDecimal());
+	}
+	else
+	{
+		FrameTimeString = FString::FromInt(FrameTime.GetFrame().Value);
+	}
+
+	const FSlateFontInfo SmallLayoutFont = FCoreStyle::GetDefaultFontStyle("Bold", 10);
+	const TSharedRef< FSlateFontMeasure > FontMeasureService = FSlateApplication::Get().GetRenderer()->GetFontMeasureService();
+	FVector2D TextSize = FontMeasureService->Measure(FrameTimeString, SmallLayoutFont);
+
+	const float PixelX = InPainter.GetTimeConverter().FrameToPixel(CurrentTime);
+
+	// Flip the text position if getting near the end of the view range
+	static const float TextOffsetPx = 10.f;
+	bool  bDrawLeft = (InPainter.SectionGeometry.Size.X - PixelX) < (TextSize.X + 22.f) - TextOffsetPx;
+	float TextPosition = bDrawLeft ? PixelX - TextSize.X - TextOffsetPx : PixelX + TextOffsetPx;
+	//handle mirrored labels
+	const float MajorTickHeight = 9.0f;
+	FVector2D TextOffset(TextPosition, InPainter.SectionGeometry.Size.Y - (MajorTickHeight + TextSize.Y));
+
+	const FLinearColor DrawColor = FAppStyle::GetSlateColor("SelectionColor").GetColor(FWidgetStyle()).CopyWithNewOpacity(InPainter.GhostAlpha);
+	const FVector2D BoxPadding = FVector2D(4.0f, 2.0f);
+	// draw time string
+
+	FSlateDrawElement::MakeBox(
+		InPainter.DrawElements,
+		LayerId,
+		InPainter.SectionGeometry.ToPaintGeometry(TextSize + 2.0f * BoxPadding, FSlateLayoutTransform(TextOffset - BoxPadding)),
+		FAppStyle::GetBrush("WhiteBrush"),
+		ESlateDrawEffect::None,
+		FLinearColor::Black.CopyWithNewOpacity(0.5f * InPainter.GhostAlpha)
+	);
+
+	ESlateDrawEffect DrawEffects = InPainter.bParentEnabled ? ESlateDrawEffect::None : ESlateDrawEffect::DisabledEffect;
+
+	FSlateDrawElement::MakeText(
+		InPainter.DrawElements,
+		LayerId,
+		InPainter.SectionGeometry.ToPaintGeometry(TextSize, FSlateLayoutTransform(TextOffset)),
+		FrameTimeString,
+		SmallLayoutFont,
+		DrawEffects,
+		DrawColor
+	);
+}
+
 int32 SSequencerSection::OnPaint( const FPaintArgs& Args, const FGeometry& AllottedGeometry, const FSlateRect& MyCullingRect, FSlateWindowElementList& OutDrawElements, int32 LayerId, const FWidgetStyle& InWidgetStyle, bool bParentEnabled ) const
 {
 	TSharedPtr<FSectionModel> SectionModel = WeakSectionModel.Pin();
@@ -1447,6 +1500,32 @@ int32 SSequencerSection::OnPaint( const FPaintArgs& Args, const FGeometry& Allot
 
 	++LayerId;
 
+	TOptional<FFrameTime> SectionTime = SectionInterface->GetSectionTime(Painter);
+	if (SectionTime.IsSet())
+	{
+		const FFrameRate DisplayRate = GetSequencer().GetFocusedDisplayRate();
+		const FFrameRate TickResolution = GetSequencer().GetFocusedTickResolution();
+
+		// Get the desired frame display format and zero padding from
+		// the sequencer settings, if possible.
+		TAttribute<EFrameNumberDisplayFormats> DisplayFormatAttr(EFrameNumberDisplayFormats::Frames);
+		TAttribute<uint8> ZeroPadFrameNumbersAttr(0u);
+		if (const USequencerSettings* SequencerSettings = GetSequencer().GetSequencerSettings())
+		{
+			DisplayFormatAttr.Set(SequencerSettings->GetTimeDisplayFormat());
+			ZeroPadFrameNumbersAttr.Set(SequencerSettings->GetZeroPadFrames());
+		}
+
+		const TAttribute<FFrameRate> TickResolutionAttr(TickResolution);
+		const TAttribute<FFrameRate> DisplayRateAttr(DisplayRate);
+
+		const FFrameNumberInterface FrameNumberInterface(DisplayFormatAttr, ZeroPadFrameNumbersAttr, TickResolutionAttr, DisplayRateAttr);
+
+		DrawFrameTimeHint(Painter, GetSequencer().GetLocalTime().Time, SectionTime.GetValue(), &FrameNumberInterface, LayerId);
+	}
+
+	++LayerId;
+
 	if (Painter.bClipRectEnabled)
 	{
 		OutDrawElements.PopClip();
diff --git a/Engine/Source/Editor/Sequencer/Private/SSequencerShotFilterOverlay.cpp b/Engine/Source/Editor/Sequencer/Private/SSequencerShotFilterOverlay.cpp
--- a/Engine/Source/Editor/Sequencer/Private/SSequencerShotFilterOverlay.cpp
+++ b/Engine/Source/Editor/Sequencer/Private/SSequencerShotFilterOverlay.cpp
@@ -3,7 +3,6 @@
 #include "SSequencerShotFilterOverlay.h"
 #include "Rendering/DrawElements.h"
 #include "Styling/AppStyle.h"
-#include "CommonMovieSceneTools.h"
 #include "TimeToPixel.h"
 
 /* SSequencerShotFilterOverlay interface
diff --git a/Engine/Source/Editor/Sequencer/Private/Widgets/SSequencerLayerBar.cpp b/Engine/Source/Editor/Sequencer/Private/Widgets/SSequencerLayerBar.cpp
--- a/Engine/Source/Editor/Sequencer/Private/Widgets/SSequencerLayerBar.cpp
+++ b/Engine/Source/Editor/Sequencer/Private/Widgets/SSequencerLayerBar.cpp
@@ -16,7 +16,6 @@
 #include "MVVM/ViewModels/ViewModel.h"
 #include "MVVM/Views/STrackAreaView.h"
 
-#include "CommonMovieSceneTools.h"
 #include "SequencerCommonHelpers.h"
 #include "MovieSceneTimeHelpers.h"
 #include "Sequencer.h"
diff --git a/Engine/Source/Editor/Sequencer/Public/ISequencerSection.h b/Engine/Source/Editor/Sequencer/Public/ISequencerSection.h
--- a/Engine/Source/Editor/Sequencer/Public/ISequencerSection.h
+++ b/Engine/Source/Editor/Sequencer/Public/ISequencerSection.h
@@ -161,6 +161,11 @@ public:
 	 */
 	virtual FText GetSectionToolTip() const { return GetSectionTitle(); }
 
+	/**
+	 * @return The local section time
+	 */
+	virtual TOptional<FFrameTime> GetSectionTime(FSequencerSectionPainter& InPainter) const { return TOptional<FFrameTime>(); }
+
 	/**
 	 * @return The amount of padding to apply to non-interactive portions of the section interface (such as section text)
 	 */
