Sequencer: Modify when adding a section for AddTrack

This fixes the UI not updating when adding a Spawn Track from the menu

#rnx

diff --git a/Engine/Source/Editor/MovieSceneTools/Private/TrackEditors/SpawnTrackEditor.cpp b/Engine/Source/Editor/MovieSceneTools/Private/TrackEditors/SpawnTrackEditor.cpp
--- a/Engine/Source/Editor/MovieSceneTools/Private/TrackEditors/SpawnTrackEditor.cpp
+++ b/Engine/Source/Editor/MovieSceneTools/Private/TrackEditors/SpawnTrackEditor.cpp
@@ -43,6 +43,7 @@ UMovieSceneTrack* FSpawnTrackEditor::AddTrack(UMovieScene* FocusedMovieScene, co
 
 	if (auto* SpawnTrack = Cast<UMovieSceneSpawnTrack>(NewTrack))
 	{
+		SpawnTrack->Modify();
 		SpawnTrack->SetObjectId(ObjectHandle);
 		SpawnTrack->AddSection(*SpawnTrack->CreateNewSection());
 	}
