Sequencer: Mark SequencePlayer as deprecated for 5.4. It can be accessed through GetSequencePlayer()

#noqatesting

diff --git a/Engine/Plugins/Developer/Concert/ConcertSync/ConcertSyncClient/Source/ConcertSyncClient/Private/ConcertClientSequencerManager.cpp b/Engine/Plugins/Developer/Concert/ConcertSync/ConcertSyncClient/Source/ConcertSyncClient/Private/ConcertClientSequencerManager.cpp
--- a/Engine/Plugins/Developer/Concert/ConcertSync/ConcertSyncClient/Source/ConcertSyncClient/Private/ConcertClientSequencerManager.cpp
+++ b/Engine/Plugins/Developer/Concert/ConcertSync/ConcertSyncClient/Source/ConcertSyncClient/Private/ConcertClientSequencerManager.cpp
@@ -665,9 +665,9 @@ void FConcertClientSequencerManager::ApplyCloseEventToPlayers(const FConcertSequ
 
 	ALevelSequenceActor* LevelSequenceActor = *Player;
 
-	if (CloseEvent.bControllerClose && LevelSequenceActor && LevelSequenceActor->SequencePlayer)
+	if (CloseEvent.bControllerClose && LevelSequenceActor && LevelSequenceActor->GetSequencePlayer())
 	{
-		LevelSequenceActor->SequencePlayer->Stop();
+		LevelSequenceActor->GetSequencePlayer()->Stop();
 	}
 
 	if (!CanClose(CloseEvent))
@@ -833,13 +833,12 @@ bool FConcertClientSequencerManager::CanClose(const FConcertSequencerCloseEvent&
 
 void FConcertClientSequencerManager::DestroyPlayer(ALevelSequenceActor* LevelSequenceActor)
 {
-	if (LevelSequenceActor && LevelSequenceActor->SequencePlayer)
+	if (LevelSequenceActor && LevelSequenceActor->GetSequencePlayer())
 	{
 		UE_LOG(LogConcertSequencerSync, VeryVerbose, TEXT("Destroying LevelSequenceActor: %s"), *LevelSequenceActor->GetPathName());
-		LevelSequenceActor->SequencePlayer->Stop();
+		LevelSequenceActor->GetSequencePlayer()->Stop();
 
 		LevelSequenceActor->SetSequence(nullptr);
-		LevelSequenceActor->SequencePlayer = nullptr;
 		LevelSequenceActor->Destroy(false, false);
 	}
 }
@@ -939,7 +938,7 @@ void FConcertClientSequencerManager::ApplyTimeAdjustmentToPlayers(const FConcert
 	}
 
 	ALevelSequenceActor* LevelSequenceActor = *SeqPlayer;
-	if (LevelSequenceActor && LevelSequenceActor->SequencePlayer)
+	if (LevelSequenceActor && LevelSequenceActor->GetSequencePlayer())
 	{
 		UMovieScene* MovieScene = LevelSequenceActor->LevelSequenceAsset->GetMovieScene();
 		if (MovieScene)
@@ -1191,9 +1190,9 @@ void FConcertClientSequencerManager::ApplyEventToPlayers(const FConcertSequencer
 	}
 
 	ALevelSequenceActor* LevelSequenceActor = *SeqPlayer;
-	if (LevelSequenceActor && LevelSequenceActor->SequencePlayer)
+	if (LevelSequenceActor && LevelSequenceActor->GetSequencePlayer())
 	{
-		ULevelSequencePlayer* Player = LevelSequenceActor->SequencePlayer;
+		ULevelSequencePlayer* Player = LevelSequenceActor->GetSequencePlayer();
 		float LatencyCompensationMs = GetLatencyCompensationMs();
 
 		FFrameRate SequenceRate = Player->GetFrameRate();
diff --git a/Engine/Plugins/VirtualProduction/Takes/Source/TakeRecorderSources/Private/TakeRecorderLevelSequenceSource.cpp b/Engine/Plugins/VirtualProduction/Takes/Source/TakeRecorderSources/Private/TakeRecorderLevelSequenceSource.cpp
--- a/Engine/Plugins/VirtualProduction/Takes/Source/TakeRecorderSources/Private/TakeRecorderLevelSequenceSource.cpp
+++ b/Engine/Plugins/VirtualProduction/Takes/Source/TakeRecorderSources/Private/TakeRecorderLevelSequenceSource.cpp
@@ -59,11 +59,12 @@ void UTakeRecorderLevelSequenceSource::StartRecording(const FTimecode& InSection
 	// Play the level sequence actors	
 	for (auto ActorToTrigger : ActorsToTrigger)
 	{
-		if (ActorToTrigger->SequencePlayer)
+		ULevelSequencePlayer* SequencePlayer = ActorToTrigger->GetSequencePlayer();
+		if (SequencePlayer)
 		{
-			ActorToTrigger->SequencePlayer->SetDisableCameraCuts(true);
-			ActorToTrigger->SequencePlayer->SetPlaybackPosition(FMovieSceneSequencePlaybackParams(ActorToTrigger->SequencePlayer->GetStartTime().Time, EUpdatePositionMethod::Jump));
-			ActorToTrigger->SequencePlayer->Play();
+			SequencePlayer->SetDisableCameraCuts(true);
+			SequencePlayer->SetPlaybackPosition(FMovieSceneSequencePlaybackParams(SequencePlayer->GetStartTime().Time, EUpdatePositionMethod::Jump));
+			SequencePlayer->Play();
 		}
 		else
 		{
@@ -83,7 +84,7 @@ void UTakeRecorderLevelSequenceSource::StopRecording(class ULevelSequence* InSeq
 		}
 
 		ALevelSequenceActor* ActorToTrigger = Actor.Get();
-		ULevelSequencePlayer* SequencePlayer = ActorToTrigger->SequencePlayer;
+		ULevelSequencePlayer* SequencePlayer = ActorToTrigger->GetSequencePlayer();
 		if (SequencePlayer)
 		{
 			SequencePlayer->SetDisableCameraCuts(false);
@@ -145,7 +146,7 @@ FText UTakeRecorderLevelSequenceSource::GetDescriptionTextImpl() const
 		}
 
 		ALevelSequenceActor* ActorToTrigger = Actor.Get();
-		ULevelSequencePlayer* SequencePlayer = ActorToTrigger->SequencePlayer;
+		ULevelSequencePlayer* SequencePlayer = ActorToTrigger->GetSequencePlayer();
 		if (SequencePlayer && SequencePlayer->IsPlaying())
 		{
 			++NumPlaying;
diff --git a/Engine/Source/Editor/MovieSceneTools/Private/AutomatedLevelSequenceCapture.cpp b/Engine/Source/Editor/MovieSceneTools/Private/AutomatedLevelSequenceCapture.cpp
--- a/Engine/Source/Editor/MovieSceneTools/Private/AutomatedLevelSequenceCapture.cpp
+++ b/Engine/Source/Editor/MovieSceneTools/Private/AutomatedLevelSequenceCapture.cpp
@@ -106,7 +106,7 @@ UMovieSceneCinematicShotTrack* GetCinematicShotTrack(TWeakObjectPtr<ALevelSequen
 
 UMovieSceneMotionVectorSimulationSystem* FindMotionVectorSimulation(ALevelSequenceActor* LevelSequenceActor)
 {
-	return static_cast<IMovieScenePlayer*>(LevelSequenceActor->SequencePlayer)->GetEvaluationTemplate().GetEntitySystemLinker()->FindSystem<UMovieSceneMotionVectorSimulationSystem>();
+	return static_cast<IMovieScenePlayer*>(LevelSequenceActor->GetSequencePlayer())->GetEvaluationTemplate().GetEntitySystemLinker()->FindSystem<UMovieSceneMotionVectorSimulationSystem>();
 }
 
 UAutomatedLevelSequenceCapture::UAutomatedLevelSequenceCapture(const FObjectInitializer& Init)
@@ -334,7 +334,7 @@ void UAutomatedLevelSequenceCapture::Initialize(TSharedPtr<FSceneViewport> InVie
 	{
 		// Ensure it doesn't loop (-1 is indefinite)
 		Actor->PlaybackSettings.LoopCount.Value = 0;
-		Actor->SequencePlayer->SetTimeController(TimeController);
+		Actor->GetSequencePlayer()->SetTimeController(TimeController);
 		Actor->PlaybackSettings.bPauseAtEnd = true;
 		Actor->PlaybackSettings.bAutoPlay = false;
 
@@ -350,11 +350,11 @@ void UAutomatedLevelSequenceCapture::Initialize(TSharedPtr<FSceneViewport> InVie
 		}
 
 		// Make sure we're not playing yet, and have a fully up to date player based on the above settings (in case AutoPlay was called from BeginPlay)
-		if( Actor->SequencePlayer != nullptr )
+		if( Actor->GetSequencePlayer() != nullptr )
 		{
-			if (Actor->SequencePlayer->IsPlaying())
+			if (Actor->GetSequencePlayer()->IsPlaying())
 			{
-				Actor->SequencePlayer->Stop();
+				Actor->GetSequencePlayer()->Stop();
 			}
 			Actor->InitializePlayer();
 		}
@@ -604,11 +604,11 @@ void UAutomatedLevelSequenceCapture::SetupFrameRange()
 				}
 
 				// Override the movie scene's playback range
-				Actor->SequencePlayer->SetFrameRate(Settings.GetFrameRate());
-				Actor->SequencePlayer->SetFrameRange(PlaybackStartFrame.Value, (PlaybackEndFrame - PlaybackStartFrame).Value);
-				Actor->SequencePlayer->SetPlaybackPosition(FMovieSceneSequencePlaybackParams(FFrameTime(PlaybackStartFrame), EUpdatePositionMethod::Jump));
+				Actor->GetSequencePlayer()->SetFrameRate(Settings.GetFrameRate());
+				Actor->GetSequencePlayer()->SetFrameRange(PlaybackStartFrame.Value, (PlaybackEndFrame - PlaybackStartFrame).Value);
+				Actor->GetSequencePlayer()->SetPlaybackPosition(FMovieSceneSequencePlaybackParams(FFrameTime(PlaybackStartFrame), EUpdatePositionMethod::Jump));
 
-				Actor->SequencePlayer->SetSnapshotOffsetFrames(WarmUpFrameCount);
+				Actor->GetSequencePlayer()->SetSnapshotOffsetFrames(WarmUpFrameCount);
 			}
 		}
 	}
@@ -645,7 +645,7 @@ void UAutomatedLevelSequenceCapture::OnTick(float DeltaSeconds)
 {
 	ALevelSequenceActor* Actor = LevelSequenceActor.Get();
 
-	if (!Actor || !Actor->SequencePlayer)
+	if (!Actor || !Actor->GetSequencePlayer())
 	{
 		return;
 	}
@@ -689,7 +689,7 @@ void UAutomatedLevelSequenceCapture::OnTick(float DeltaSeconds)
 		// Bind to the event so we know when to capture a frame
 		if (!bIsAudioCapturePass)
 		{
-			OnPlayerUpdatedBinding = Actor->SequencePlayer->OnSequenceUpdated().AddUObject( this, &UAutomatedLevelSequenceCapture::SequenceUpdated );
+			OnPlayerUpdatedBinding = Actor->GetSequencePlayer()->OnSequenceUpdated().AddUObject( this, &UAutomatedLevelSequenceCapture::SequenceUpdated );
 		}
 
 		StartWarmup();
@@ -714,7 +714,7 @@ void UAutomatedLevelSequenceCapture::OnTick(float DeltaSeconds)
 	}
 	else if( CaptureState == ELevelSequenceCaptureState::ReadyToWarmUp )
 	{
-		Actor->SequencePlayer->Play();
+		Actor->GetSequencePlayer()->Play();
 		// Start warming up
 		CaptureState = ELevelSequenceCaptureState::WarmingUp;
 	}
@@ -738,7 +738,7 @@ void UAutomatedLevelSequenceCapture::OnTick(float DeltaSeconds)
 		}
 	}
 
-	if( bCapturing && !Actor->SequencePlayer->IsPlaying() && CaptureState != ELevelSequenceCaptureState::Paused )
+	if( bCapturing && !Actor->GetSequencePlayer()->IsPlaying() && CaptureState != ELevelSequenceCaptureState::Paused )
 	{
 		++ShotIndex;
 
@@ -750,14 +750,14 @@ void UAutomatedLevelSequenceCapture::OnTick(float DeltaSeconds)
 			FFrameNumber StartTimePlayRateSpace = ConvertFrameTime(StartTime, MovieScene->GetTickResolution(), Settings.GetFrameRate()).CeilToFrame();
 			FFrameNumber EndTimePlayRateSpace   = ConvertFrameTime(EndTime,   MovieScene->GetTickResolution(), Settings.GetFrameRate()).CeilToFrame();
 
-			Actor->SequencePlayer->SetFrameRange(StartTimePlayRateSpace.Value, (EndTimePlayRateSpace - StartTimePlayRateSpace).Value);
-			Actor->SequencePlayer->SetPlaybackPosition(FMovieSceneSequencePlaybackParams(FFrameTime(StartTimePlayRateSpace), EUpdatePositionMethod::Jump));
-			Actor->SequencePlayer->Play();
+			Actor->GetSequencePlayer()->SetFrameRange(StartTimePlayRateSpace.Value, (EndTimePlayRateSpace - StartTimePlayRateSpace).Value);
+			Actor->GetSequencePlayer()->SetPlaybackPosition(FMovieSceneSequencePlaybackParams(FFrameTime(StartTimePlayRateSpace), EUpdatePositionMethod::Jump));
+			Actor->GetSequencePlayer()->Play();
 
 			// We need to re-register to the binding when we start each shot. When a shot reaches the last frame it unregisters the binding so that
 			// any subsequent seeking doesn't accidentally render extra frames. SetupShot doesn't get called until after the first time we finish
 			// rendering a shot so this doesn't register the delegate twice on the first go.
-			OnPlayerUpdatedBinding = Actor->SequencePlayer->OnSequenceUpdated().AddUObject(this, &UAutomatedLevelSequenceCapture::SequenceUpdated);
+			OnPlayerUpdatedBinding = Actor->GetSequencePlayer()->OnSequenceUpdated().AddUObject(this, &UAutomatedLevelSequenceCapture::SequenceUpdated);
 
 			CaptureState = ELevelSequenceCaptureState::FinishedWarmUp;
 			UpdateFrameState();
@@ -770,7 +770,7 @@ void UAutomatedLevelSequenceCapture::OnTick(float DeltaSeconds)
 			if (IsAudioPassIfNeeded() && CaptureState != ELevelSequenceCaptureState::Setup)
 			{
 				// If they don't want to render audio, or they have rendered an audio pass, we finish and finalize the data.
-				Actor->SequencePlayer->OnSequenceUpdated().Remove( OnPlayerUpdatedBinding );
+				Actor->GetSequencePlayer()->OnSequenceUpdated().Remove( OnPlayerUpdatedBinding );
 				FinalizeWhenReady();
 
 				// Restore our cached variables since these are the actual one represented in the in-engine UI
@@ -782,7 +782,7 @@ void UAutomatedLevelSequenceCapture::OnTick(float DeltaSeconds)
 			{
 				// Reset us to use the platform clock for controlling the playback rate of the sequence. The audio system
 				// uses the platform clock for timings as well.
-				Actor->SequencePlayer->SetTimeController(MakeShared<FMovieSceneTimeController_PlatformClock>());
+				Actor->GetSequencePlayer()->SetTimeController(MakeShared<FMovieSceneTimeController_PlatformClock>());
 				CaptureState = ELevelSequenceCaptureState::Setup;
 				
 				// We'll now repeat the whole process including warmups and delays. The audio capture will pause recording while we are delayed.
@@ -805,7 +805,7 @@ void UAutomatedLevelSequenceCapture::DelayBeforeWarmupFinished()
 void UAutomatedLevelSequenceCapture::PauseFinished()
 {
 	ALevelSequenceActor* Actor = LevelSequenceActor.Get();
-	Actor->SequencePlayer->Play();
+	Actor->GetSequencePlayer()->Play();
 
 	CaptureState = ELevelSequenceCaptureState::FinishedWarmUp;
 
@@ -828,7 +828,7 @@ void UAutomatedLevelSequenceCapture::SequenceUpdated(const UMovieSceneSequencePl
 		UpdateFrameState();
 
 		ALevelSequenceActor* Actor = LevelSequenceActor.Get();
-		if (Actor && Actor->SequencePlayer)
+		if (Actor && Actor->GetSequencePlayer())
 		{
 			// If this is a new shot, set the state to shot warm up and pause on this frame until warmed up			
 			const bool bHasMultipleShots = PreviousState.CurrentShotName != PreviousState.RootName;
@@ -838,7 +838,7 @@ void UAutomatedLevelSequenceCapture::SequenceUpdated(const UMovieSceneSequencePl
 			const bool bDelayingBeforeShotWarmUp = (bNewShot && DelayBeforeShotWarmUp > 0);
 			const bool bDelayingEveryFrame = (bNewFrame && DelayEveryFrame > 0);
 
-			if (Actor->SequencePlayer->IsPlaying() && ( bDelayingBeforeShotWarmUp || bDelayingEveryFrame ))
+			if (Actor->GetSequencePlayer()->IsPlaying() && ( bDelayingBeforeShotWarmUp || bDelayingEveryFrame ))
 			{
 				if (bIsAudioCapturePass)
 				{
@@ -865,7 +865,7 @@ void UAutomatedLevelSequenceCapture::SequenceUpdated(const UMovieSceneSequencePl
 				}
 
 				Actor->GetWorld()->GetTimerManager().SetTimer(DelayTimer, FTimerDelegate::CreateUObject(this, &UAutomatedLevelSequenceCapture::PauseFinished), DelayBeforeShotWarmUp + DelayEveryFrame, false);
-				Actor->SequencePlayer->Pause();
+				Actor->GetSequencePlayer()->Pause();
 			}
 			else if (CaptureState == ELevelSequenceCaptureState::FinishedWarmUp)
 			{
@@ -893,7 +893,7 @@ void UAutomatedLevelSequenceCapture::SequenceUpdated(const UMovieSceneSequencePl
 					}
 				}
 
-				bool bOnLastFrame = (CurrentTime.FrameNumber >= Actor->SequencePlayer->GetStartTime().Time.FrameNumber + Actor->SequencePlayer->GetFrameDuration() - 1);
+				bool bOnLastFrame = (CurrentTime.FrameNumber >= Actor->GetSequencePlayer()->GetStartTime().Time.FrameNumber + Actor->GetSequencePlayer()->GetFrameDuration() - 1);
 				bool bLastShot = NumShots == 0 ? true : ShotIndex == NumShots - 1;
 
 				CaptureThisFrame((CurrentTime - PreviousTime) / Settings.GetFrameRate());
@@ -912,7 +912,7 @@ void UAutomatedLevelSequenceCapture::SequenceUpdated(const UMovieSceneSequencePl
 					{
 						FinalizeWhenReady();
 					}
-					Actor->SequencePlayer->OnSequenceUpdated().Remove(OnPlayerUpdatedBinding);
+					Actor->GetSequencePlayer()->OnSequenceUpdated().Remove(OnPlayerUpdatedBinding);
 				}
 
 
@@ -927,9 +927,9 @@ void UAutomatedLevelSequenceCapture::UpdateFrameState()
 {
 	ALevelSequenceActor* Actor = LevelSequenceActor.Get();
 
-	if (Actor && Actor->SequencePlayer)
+	if (Actor && Actor->GetSequencePlayer())
 	{
-		Actor->SequencePlayer->TakeFrameSnapshot(CachedState);
+		Actor->GetSequencePlayer()->TakeFrameSnapshot(CachedState);
 	}
 }
 
diff --git a/Engine/Source/Editor/SequenceRecorder/Private/SequenceRecorder.cpp b/Engine/Source/Editor/SequenceRecorder/Private/SequenceRecorder.cpp
--- a/Engine/Source/Editor/SequenceRecorder/Private/SequenceRecorder.cpp
+++ b/Engine/Source/Editor/SequenceRecorder/Private/SequenceRecorder.cpp
@@ -709,11 +709,12 @@ bool FSequenceRecorder::StartRecordingInternal(UWorld* World)
 				// Always initialize the player so that the playback settings/range can be initialized from editor.
 				DupActorToTrigger->InitializePlayer();
 				
-				if (DupActorToTrigger->SequencePlayer)
+				ULevelSequencePlayer* SequencePlayer = DupActorToTrigger->GetSequencePlayer();
+				if (SequencePlayer)
 				{
-					DupActorToTrigger->SequencePlayer->SetDisableCameraCuts(true);
-					DupActorToTrigger->SequencePlayer->SetPlaybackPosition(FMovieSceneSequencePlaybackParams(0, EUpdatePositionMethod::Jump));
-					DupActorToTrigger->SequencePlayer->Play();
+					SequencePlayer->SetDisableCameraCuts(true);
+					SequencePlayer->SetPlaybackPosition(FMovieSceneSequencePlaybackParams(0, EUpdatePositionMethod::Jump));
+					SequencePlayer->Play();
 				}
 				else
 				{
@@ -1039,7 +1040,7 @@ bool FSequenceRecorder::StopRecording(bool bAllowLooping)
 		if (DupActorsToTrigger[DupActorToTriggerIndex].IsValid())
 		{
 			ALevelSequenceActor* DupActorToTrigger = DupActorsToTrigger[DupActorToTriggerIndex].Get();
-			ULevelSequencePlayer* SequencePlayer = DupActorToTrigger->SequencePlayer;
+			ULevelSequencePlayer* SequencePlayer = DupActorToTrigger->GetSequencePlayer();
 			if (SequencePlayer)
 			{
 				SequencePlayer->SetDisableCameraCuts(false);
diff --git a/Engine/Source/Runtime/LevelSequence/Private/LevelSequenceActor.cpp b/Engine/Source/Runtime/LevelSequence/Private/LevelSequenceActor.cpp
--- a/Engine/Source/Runtime/LevelSequence/Private/LevelSequenceActor.cpp
+++ b/Engine/Source/Runtime/LevelSequence/Private/LevelSequenceActor.cpp
@@ -68,10 +68,12 @@ ALevelSequenceActor::ALevelSequenceActor(const FObjectInitializer& Init)
 	DefaultInstanceData = Init.CreateDefaultSubobject<UDefaultLevelSequenceInstanceData>(this, "InstanceData");
 
 	// SequencePlayer must be a default sub object for it to be replicated correctly
+PRAGMA_DISABLE_DEPRECATION_WARNINGS // make SequencePlayer protected and remove this for 5.6
 	SequencePlayer = Init.CreateDefaultSubobject<ULevelSequencePlayer>(this, "AnimationPlayer");
-	SequencePlayer->OnPlay.AddDynamic(this, &ALevelSequenceActor::ShowBurnin);
-	SequencePlayer->OnPlayReverse.AddDynamic(this, &ALevelSequenceActor::ShowBurnin);
-	SequencePlayer->OnStop.AddDynamic(this, &ALevelSequenceActor::HideBurnin);
+PRAGMA_ENABLE_DEPRECATION_WARNINGS // SequencePlayer
+	GetSequencePlayer()->OnPlay.AddDynamic(this, &ALevelSequenceActor::ShowBurnin);
+	GetSequencePlayer()->OnPlayReverse.AddDynamic(this, &ALevelSequenceActor::ShowBurnin);
+	GetSequencePlayer()->OnStop.AddDynamic(this, &ALevelSequenceActor::HideBurnin);
 	bOverrideInstanceData = false;
 
 	// The level sequence actor defaults to never ticking by the tick manager because it is ticked separately in LevelTick
@@ -90,15 +92,15 @@ void ALevelSequenceActor::PostInitProperties()
 
 	// Have to initialize this here as any properties set on default subobjects inside the constructor
 	// Get stomped by the CDO's properties when the constructor exits.
-	SequencePlayer->SetPlaybackClient(this);
-	SequencePlayer->SetPlaybackSettings(PlaybackSettings);
+	GetSequencePlayer()->SetPlaybackClient(this);
+	GetSequencePlayer()->SetPlaybackSettings(PlaybackSettings);
 }
 
 void ALevelSequenceActor::RewindForReplay()
 {
-	if (SequencePlayer)
+	if (GetSequencePlayer())
 	{
-		SequencePlayer->RewindForReplay();
+		GetSequencePlayer()->RewindForReplay();
 	}
 }
 
@@ -129,7 +131,9 @@ bool ALevelSequenceActor::GetIsReplicatedPlayback() const
 
 ULevelSequencePlayer* ALevelSequenceActor::GetSequencePlayer() const
 {
-	return SequencePlayer && SequencePlayer->GetSequence() ? SequencePlayer : nullptr;
+PRAGMA_DISABLE_DEPRECATION_WARNINGS // make SequencePlayer protected and remove this for 5.6
+	return SequencePlayer;
+PRAGMA_ENABLE_DEPRECATION_WARNINGS
 }
 
 void ALevelSequenceActor::SetReplicatePlayback(bool bInReplicatePlayback)
@@ -142,7 +146,9 @@ void ALevelSequenceActor::GetLifetimeReplicatedProps(TArray<FLifetimeProperty>&
 {
 	Super::GetLifetimeReplicatedProps(OutLifetimeProps);
 
+PRAGMA_DISABLE_DEPRECATION_WARNINGS // make SequencePlayer protected and remove this for 5.6
 	DOREPLIFETIME(ALevelSequenceActor, SequencePlayer);
+PRAGMA_ENABLE_DEPRECATION_WARNINGS
 	DOREPLIFETIME(ALevelSequenceActor, LevelSequenceAsset);
 }
 
@@ -157,7 +163,7 @@ void ALevelSequenceActor::PostInitializeComponents()
 	
 	// Initialize this player for tick as soon as possible to ensure that a persistent
 	// reference to the tick manager is maintained
-	SequencePlayer->InitializeForTick(this);
+	GetSequencePlayer()->InitializeForTick(this);
 
 	InitializePlayer();
 }
@@ -166,33 +172,33 @@ void ALevelSequenceActor::BeginPlay()
 {
 	Super::BeginPlay();
 
-	if (SequencePlayer)
+	if (GetSequencePlayer())
 	{
-		AddReplicatedSubObject(SequencePlayer);
+		AddReplicatedSubObject(GetSequencePlayer());
 	}
 
 	if (PlaybackSettings.bAutoPlay)
 	{
-		SequencePlayer->Play();
+		GetSequencePlayer()->Play();
 	}
 }
 
 void ALevelSequenceActor::EndPlay(const EEndPlayReason::Type EndPlayReason)
 {
-	if (SequencePlayer)
+	if (GetSequencePlayer())
 	{
-		RemoveReplicatedSubObject(SequencePlayer);
+		RemoveReplicatedSubObject(GetSequencePlayer());
 
 		// Stop may modify a lot of actor state so it needs to be called
 		// during EndPlay (when Actors + World are still valid) instead
 		// of waiting for the UObject to be destroyed by GC.
-		SequencePlayer->Stop();
+		GetSequencePlayer()->Stop();
 
-		SequencePlayer->OnPlay.RemoveAll(this);
-		SequencePlayer->OnPlayReverse.RemoveAll(this);
-		SequencePlayer->OnStop.RemoveAll(this);
+		GetSequencePlayer()->OnPlay.RemoveAll(this);
+		GetSequencePlayer()->OnPlayReverse.RemoveAll(this);
+		GetSequencePlayer()->OnStop.RemoveAll(this);
 
-		SequencePlayer->TearDown();
+		GetSequencePlayer()->TearDown();
 	}
 
 	Super::EndPlay(EndPlayReason);
@@ -218,7 +224,7 @@ void ALevelSequenceActor::PostLoad()
 	}
 #endif
 
-	SequencePlayer->SetPlaybackSettings(PlaybackSettings);
+	GetSequencePlayer()->SetPlaybackSettings(PlaybackSettings);
 
 #if WITH_EDITORONLY_DATA
 	if (LevelSequence_DEPRECATED.IsValid())
@@ -267,14 +273,14 @@ ULevelSequence* ALevelSequenceActor::GetSequence() const
 
 void ALevelSequenceActor::SetSequence(ULevelSequence* InSequence)
 {
-	if (!SequencePlayer->IsPlaying())
+	if (!GetSequencePlayer()->IsPlaying())
 	{
 		LevelSequenceAsset = InSequence;
 
 		// cbb: should ideally null out the template and player when no sequence is assigned, but that's currently not possible
 		if (InSequence)
 		{
-			SequencePlayer->Initialize(InSequence, GetLevel(), CameraSettings);
+			GetSequencePlayer()->Initialize(InSequence, GetLevel(), CameraSettings);
 		}
 	}
 }
@@ -284,9 +290,9 @@ void ALevelSequenceActor::InitializePlayer()
 	if (LevelSequenceAsset && GetWorld()->IsGameWorld())
 	{
 		// Level sequence is already loaded. Initialize the player if it's not already initialized with this sequence
-		if (LevelSequenceAsset != SequencePlayer->GetSequence() || SequencePlayer->GetEvaluationTemplate().GetRunner() == nullptr)
+		if (LevelSequenceAsset != GetSequencePlayer()->GetSequence() || GetSequencePlayer()->GetEvaluationTemplate().GetRunner() == nullptr)
 		{
-			SequencePlayer->Initialize(LevelSequenceAsset, GetLevel(), CameraSettings);
+			GetSequencePlayer()->Initialize(LevelSequenceAsset, GetLevel(), CameraSettings);
 		}
 	}
 }
@@ -356,10 +362,10 @@ void ALevelSequenceActor::SetBinding(FMovieSceneObjectBindingID Binding, const T
 	else
 	{
 		BindingOverrides->SetBinding(Binding, TArray<UObject*>(Actors), bAllowBindingsFromAsset);
-		if (SequencePlayer)
+		if (GetSequencePlayer())
 		{
-			FMovieSceneSequenceID SequenceID = Binding.ResolveSequenceID(MovieSceneSequenceID::Root, *SequencePlayer);
-			SequencePlayer->State.Invalidate(Binding.GetGuid(), SequenceID);
+			FMovieSceneSequenceID SequenceID = Binding.ResolveSequenceID(MovieSceneSequenceID::Root, *GetSequencePlayer());
+			GetSequencePlayer()->State.Invalidate(Binding.GetGuid(), SequenceID);
 		}
 	}
 }
@@ -394,10 +400,10 @@ void ALevelSequenceActor::AddBinding(FMovieSceneObjectBindingID Binding, AActor*
 	else
 	{
 		BindingOverrides->AddBinding(Binding, Actor, bAllowBindingsFromAsset);
-		if (SequencePlayer)
+		if (GetSequencePlayer())
 		{
-			FMovieSceneSequenceID SequenceID = Binding.ResolveSequenceID(MovieSceneSequenceID::Root, *SequencePlayer);
-			SequencePlayer->State.Invalidate(Binding.GetGuid(), SequenceID);
+			FMovieSceneSequenceID SequenceID = Binding.ResolveSequenceID(MovieSceneSequenceID::Root, *GetSequencePlayer());
+			GetSequencePlayer()->State.Invalidate(Binding.GetGuid(), SequenceID);
 		}
 	}
 }
@@ -432,10 +438,10 @@ void ALevelSequenceActor::RemoveBinding(FMovieSceneObjectBindingID Binding, AAct
 	else
 	{
 		BindingOverrides->RemoveBinding(Binding, Actor);
-		if (SequencePlayer)
+		if (GetSequencePlayer())
 		{
-			FMovieSceneSequenceID SequenceID = Binding.ResolveSequenceID(MovieSceneSequenceID::Root, *SequencePlayer);
-			SequencePlayer->State.Invalidate(Binding.GetGuid(), SequenceID);
+			FMovieSceneSequenceID SequenceID = Binding.ResolveSequenceID(MovieSceneSequenceID::Root, *GetSequencePlayer());
+			GetSequencePlayer()->State.Invalidate(Binding.GetGuid(), SequenceID);
 		}
 	}
 }
@@ -470,10 +476,10 @@ void ALevelSequenceActor::ResetBinding(FMovieSceneObjectBindingID Binding)
 	else
 	{
 		BindingOverrides->ResetBinding(Binding);
-		if (SequencePlayer)
+		if (GetSequencePlayer())
 		{
-			FMovieSceneSequenceID SequenceID = Binding.ResolveSequenceID(MovieSceneSequenceID::Root, *SequencePlayer);
-			SequencePlayer->State.Invalidate(Binding.GetGuid(), SequenceID);
+			FMovieSceneSequenceID SequenceID = Binding.ResolveSequenceID(MovieSceneSequenceID::Root, *GetSequencePlayer());
+			GetSequencePlayer()->State.Invalidate(Binding.GetGuid(), SequenceID);
 		}
 	}
 }
@@ -481,26 +487,26 @@ void ALevelSequenceActor::ResetBinding(FMovieSceneObjectBindingID Binding)
 void ALevelSequenceActor::ResetBindings()
 {
 	BindingOverrides->ResetBindings();
-	if (SequencePlayer)
+	if (GetSequencePlayer())
 	{
-		SequencePlayer->State.ClearObjectCaches(*SequencePlayer);
+		GetSequencePlayer()->State.ClearObjectCaches(*GetSequencePlayer());
 	}
 }
 
 FMovieSceneObjectBindingID ALevelSequenceActor::FindNamedBinding(FName InBindingName) const
 {
-	if (ensureAlways(SequencePlayer))
+	if (ensureAlways(GetSequencePlayer()))
 	{
-		return SequencePlayer->GetSequence()->FindBindingByTag(InBindingName);
+		return GetSequencePlayer()->GetSequence()->FindBindingByTag(InBindingName);
 	}
 	return FMovieSceneObjectBindingID();
 }
 
 const TArray<FMovieSceneObjectBindingID>& ALevelSequenceActor::FindNamedBindings(FName InBindingName) const
 {
-	if (ensureAlways(SequencePlayer))
+	if (ensureAlways(GetSequencePlayer()))
 	{
-		return SequencePlayer->GetSequence()->FindBindingsByTag(InBindingName);
+		return GetSequencePlayer()->GetSequence()->FindBindingsByTag(InBindingName);
 	}
 
 	static TArray<FMovieSceneObjectBindingID> EmptyBindings;
@@ -511,7 +517,7 @@ void ALevelSequenceActor::PostNetReceive()
 {
 	Super::PostNetReceive();
 
-	if (LevelSequenceAsset && SequencePlayer && SequencePlayer->GetSequence() != LevelSequenceAsset)
+	if (LevelSequenceAsset && GetSequencePlayer() && GetSequencePlayer()->GetSequence() != LevelSequenceAsset)
 	{
 		InitializePlayer();
 	}
diff --git a/Engine/Source/Runtime/LevelSequence/Private/LevelSequenceBurnIn.cpp b/Engine/Source/Runtime/LevelSequence/Private/LevelSequenceBurnIn.cpp
--- a/Engine/Source/Runtime/LevelSequence/Private/LevelSequenceBurnIn.cpp
+++ b/Engine/Source/Runtime/LevelSequence/Private/LevelSequenceBurnIn.cpp
@@ -12,10 +12,11 @@ ULevelSequenceBurnIn::ULevelSequenceBurnIn( const FObjectInitializer& ObjectInit
 void ULevelSequenceBurnIn::TakeSnapshotsFrom(ALevelSequenceActor& InActor)
 {
 	LevelSequenceActor = &InActor;
-	if (ensure(InActor.SequencePlayer))
+	ULevelSequencePlayer* SequencePlayer = InActor.GetSequencePlayer();
+	if (ensure(SequencePlayer))
 	{
-		InActor.SequencePlayer->OnSequenceUpdated().AddUObject(this, &ULevelSequenceBurnIn::OnSequenceUpdated);
-		InActor.SequencePlayer->TakeFrameSnapshot(FrameInformation);
+		SequencePlayer->OnSequenceUpdated().AddUObject(this, &ULevelSequenceBurnIn::OnSequenceUpdated);
+		SequencePlayer->TakeFrameSnapshot(FrameInformation);
 	}
 }
 
diff --git a/Engine/Source/Runtime/LevelSequence/Private/LevelSequencePlayer.cpp b/Engine/Source/Runtime/LevelSequence/Private/LevelSequencePlayer.cpp
--- a/Engine/Source/Runtime/LevelSequence/Private/LevelSequencePlayer.cpp
+++ b/Engine/Source/Runtime/LevelSequence/Private/LevelSequencePlayer.cpp
@@ -69,7 +69,7 @@ ULevelSequencePlayer* ULevelSequencePlayer::CreateLevelSequencePlayer(UObject* W
 	ALevelSequenceActor* Actor = World->SpawnActor<ALevelSequenceActor>(SpawnParams);
 
 	Actor->PlaybackSettings = Settings;
-	Actor->SequencePlayer->SetPlaybackSettings(Settings);
+	Actor->GetSequencePlayer()->SetPlaybackSettings(Settings);
 
 	Actor->SetSequence(InLevelSequence);
 
@@ -79,7 +79,7 @@ ULevelSequencePlayer* ULevelSequencePlayer::CreateLevelSequencePlayer(UObject* W
 	FTransform DefaultTransform;
 	Actor->FinishSpawning(DefaultTransform);
 
-	return Actor->SequencePlayer;
+	return Actor->GetSequencePlayer();
 }
 
 /* ULevelSequencePlayer implementation
diff --git a/Engine/Source/Runtime/LevelSequence/Public/LevelSequenceActor.h b/Engine/Source/Runtime/LevelSequence/Public/LevelSequenceActor.h
--- a/Engine/Source/Runtime/LevelSequence/Public/LevelSequenceActor.h
+++ b/Engine/Source/Runtime/LevelSequence/Public/LevelSequenceActor.h
@@ -83,7 +83,8 @@ public:
 	UPROPERTY(EditAnywhere, BlueprintReadOnly, Category="Playback", meta=(ShowOnlyInnerProperties, ExposeOnSpawn))
 	FMovieSceneSequencePlaybackSettings PlaybackSettings;
 
-	UPROPERTY(Instanced, transient, replicated, BlueprintReadOnly, BlueprintGetter=GetSequencePlayer, Category="Playback", meta=(ExposeFunctionCategories="Sequencer|Player"))
+	UE_DEPRECATED(5.4, "Direct access to SequencePlayer will be removed. Please use GetSequencePlayer()")
+	UPROPERTY(Instanced, transient, replicated, BlueprintReadOnly, BlueprintGetter = GetSequencePlayer, Category = "Playback", meta = (ExposeFunctionCategories = "Sequencer|Player"))
 	TObjectPtr<ULevelSequencePlayer> SequencePlayer;
 
 	UPROPERTY(EditAnywhere, replicated, BlueprintReadOnly, Category="General", meta=(AllowedClasses="/Script/LevelSequence.LevelSequence", ExposeOnSpawn))
