Sequencer: Nullptr check for subsequences when drawing the local frame time hint

diff --git a/Engine/Source/Editor/MovieSceneTools/Public/TrackEditors/SubTrackEditorBase.h b/Engine/Source/Editor/MovieSceneTools/Public/TrackEditors/SubTrackEditorBase.h
--- a/Engine/Source/Editor/MovieSceneTools/Public/TrackEditors/SubTrackEditorBase.h
+++ b/Engine/Source/Editor/MovieSceneTools/Public/TrackEditors/SubTrackEditorBase.h
@@ -273,6 +273,11 @@ TOptional<FFrameTime> TSubSectionMixin<ParentSectionClass>::GetSectionTime(FSequ
 		return TOptional<FFrameTime>();
 	}
 
+	if (!SubSectionObject.GetSequence() || !SubSectionObject.GetSequence()->GetMovieScene())
+	{
+		return TOptional<FFrameTime>();
+	}
+
 	const UMovieScene* SubSequenceMovieScene = SubSectionObject.GetSequence()->GetMovieScene();
 	const FFrameTime HintFrameTime = CurrentTime * SubSectionObject.OuterToInnerTransform();
 
