Sequencer: Fixed a crash caused by a dangling delegate pointer

diff --git a/Engine/Source/Editor/MovieSceneTools/Private/TrackEditors/AudioTrackEditor.cpp b/Engine/Source/Editor/MovieSceneTools/Private/TrackEditors/AudioTrackEditor.cpp
--- a/Engine/Source/Editor/MovieSceneTools/Private/TrackEditors/AudioTrackEditor.cpp
+++ b/Engine/Source/Editor/MovieSceneTools/Private/TrackEditors/AudioTrackEditor.cpp
@@ -937,6 +937,11 @@ FAudioTrackEditor::FAudioTrackEditor( TSharedRef<ISequencer> InSequencer )
 
 FAudioTrackEditor::~FAudioTrackEditor()
 {
+	TSharedPtr<ISequencer> SequencerPtr = GetSequencer();
+	if (MovieSceneChangedDelegate.IsValid())
+	{
+		SequencerPtr->OnMovieSceneDataChanged().Remove(MovieSceneChangedDelegate);
+	}
 }
 
 
