Fix for crash with invalid Sequencer shared pointer.

diff --git a/Engine/Source/Editor/MovieSceneTools/Private/TrackEditors/AudioTrackEditor.cpp b/Engine/Source/Editor/MovieSceneTools/Private/TrackEditors/AudioTrackEditor.cpp
--- a/Engine/Source/Editor/MovieSceneTools/Private/TrackEditors/AudioTrackEditor.cpp
+++ b/Engine/Source/Editor/MovieSceneTools/Private/TrackEditors/AudioTrackEditor.cpp
@@ -938,7 +938,7 @@ FAudioTrackEditor::FAudioTrackEditor( TSharedRef<ISequencer> InSequencer )
 FAudioTrackEditor::~FAudioTrackEditor()
 {
 	TSharedPtr<ISequencer> SequencerPtr = GetSequencer();
-	if (MovieSceneChangedDelegate.IsValid())
+	if (SequencerPtr.IsValid() && MovieSceneChangedDelegate.IsValid())
 	{
 		SequencerPtr->OnMovieSceneDataChanged().Remove(MovieSceneChangedDelegate);
 	}
