Curve: Fix GetValueRange where DBL_MIN means minimum positive. Change this to use NumericLimits with Max and Lowest.

diff --git a/Engine/Source/Editor/MovieSceneTools/Private/Channels/BezierChannelCurveModel.cpp b/Engine/Source/Editor/MovieSceneTools/Private/Channels/BezierChannelCurveModel.cpp
--- a/Engine/Source/Editor/MovieSceneTools/Private/Channels/BezierChannelCurveModel.cpp
+++ b/Engine/Source/Editor/MovieSceneTools/Private/Channels/BezierChannelCurveModel.cpp
@@ -487,8 +487,8 @@ void FBezierChannelCurveModel<ChannelType, ChannelValue, KeyType>::GetValueRange
 			FFrameRate TickResolution = Section->GetTypedOuter<UMovieScene>()->GetTickResolution();
 			double ToTime = TickResolution.AsInterval();
 			int32 LastKeyIndex = Values.Num() - 1;
-			MinValue = DBL_MAX;
-			MaxValue = DBL_MIN;
+			MinValue = TNumericLimits<double>::Max();
+			MaxValue = TNumericLimits<double>::Lowest();
 
 			for (int32 i = 0; i < Values.Num(); i++)
 			{
@@ -517,11 +517,11 @@ void FBezierChannelCurveModel<ChannelType, ChannelValue, KeyType>::GetValueRange
 		}
 	}
 	//if nothing found just set to zero
-	if (MinValue == DBL_MAX)
+	if (MinValue == TNumericLimits<double>::Max())
 	{
 		MinValue = 0.0;
 	}
-	if (MaxValue == DBL_MIN)
+	if (MaxValue == TNumericLimits<double>::Lowest())
 	{
 		MaxValue = 0.0;
 	}
@@ -531,9 +531,9 @@ void FBezierChannelCurveModel<ChannelType, ChannelValue, KeyType>::GetValueRange
 template<typename ChannelType, typename ChannelValue, typename KeyType> 
 void FBezierChannelCurveModel<ChannelType, ChannelValue, KeyType>::GetValueRange(double& MinValue, double& MaxValue) const
 {
-	const double InMinTime = DBL_MIN;
-	const double InMaxtime = DBL_MAX;
-	GetValueRange(InMinTime, InMaxtime, MinValue, MaxValue);
+	const double InMinTime = TNumericLimits<double>::Lowest();
+	const double InMaxTime = TNumericLimits<double>::Max();
+	GetValueRange(InMinTime, InMaxTime, MinValue, MaxValue);
 }
 
 template<typename ChannelType, typename ChannelValue, typename KeyType> 
