Sequencer- Binding Lifetime Track.
The Binding Lifetime track is an optional track that can be added to any object binding. It allows the user to set time range(s) for when the object in question will be bound by Sequencer. Outside of those ranges, the object will be released, and any tracks underneath it will restore state/keep state as set.

Currently this track works for both spawnables and possessables. This track is a precursor to other binding types coming to Sequencer, and in the future this track will allow the user to trigger events as part of the binding activation/deactivation so that changes like AI control for example can be set when binding/unbinding NPC characters or player characters.

As part of this change, it was necessary to clean up some of the sequencer UI color settings on labels, as well as a few changes to the layer bar UI to facilitate a hatch overlay for the sections of the binding tracks that are in a deactivated area.

diff --git a/Engine/Source/Editor/EditorStyle/Private/StarshipStyle.cpp b/Engine/Source/Editor/EditorStyle/Private/StarshipStyle.cpp
--- a/Engine/Source/Editor/EditorStyle/Private/StarshipStyle.cpp
+++ b/Engine/Source/Editor/EditorStyle/Private/StarshipStyle.cpp
@@ -2155,6 +2155,7 @@ void FStarshipEditorStyle::FStyle::SetupSequencerStyles()
 		Set("Sequencer.Section.Background", new BORDER_BRUSH(TEXT("Sequencer/SectionBackground"), FMargin(4.f / 16.f)));
 		Set("Sequencer.Section.BackgroundTint", new BOX_BRUSH(TEXT("Sequencer/SectionBackgroundTint"), FMargin(4 / 16.f)));
 		Set("Sequencer.Section.CollapsedSelectedSectionOverlay", new IMAGE_BRUSH(TEXT("Sequencer/Section/CollapsedSelectedSectionOverlay"), Icon16x16, FLinearColor::White, ESlateBrushTileType::Both));
+		Set("Sequencer.Section.SequencerDeactivatedOverlay", new IMAGE_BRUSH(TEXT("Sequencer/Section/SequencerDeactivatedOverlay"), Icon16x16, FLinearColor::White, ESlateBrushTileType::Both));
 		Set("Sequencer.Section.ErroredSectionOverlay", new BORDER_BRUSH(TEXT("Sequencer/Section/CollapsedSelectedSectionOverlay"), FMargin(4.f / 16.f)));
 		Set("Sequencer.Section.SectionHeaderSelectedSectionOverlay", new IMAGE_BRUSH(TEXT("Sequencer/Section/SectionHeaderSelectedSectionOverlay"), Icon16x16, FLinearColor::White, ESlateBrushTileType::Both));
 		Set("Sequencer.Section.SelectedTrackTint", new BOX_BRUSH(TEXT("Sequencer/SelectedTrackTint"), FMargin(0.f, 0.5f)));
diff --git a/Engine/Source/Editor/MovieSceneTools/Private/MVVM/MovieSceneToolsViewModelTypeIDs.cpp b/Engine/Source/Editor/MovieSceneTools/Private/MVVM/MovieSceneToolsViewModelTypeIDs.cpp
--- a/Engine/Source/Editor/MovieSceneTools/Private/MVVM/MovieSceneToolsViewModelTypeIDs.cpp
+++ b/Engine/Source/Editor/MovieSceneTools/Private/MVVM/MovieSceneToolsViewModelTypeIDs.cpp
@@ -4,6 +4,7 @@
 #include "MVVM/ViewModelTypeID.h"
 #include "MVVM/ViewModels/CameraCutTrackModel.h"
 #include "MVVM/ViewModels/CinematicShotTrackModel.h"
+#include "MVVM/ViewModels/BindingLifetimeTrackModel.h"
 
 
 namespace UE
@@ -14,6 +15,7 @@ namespace Sequencer
 // Model types
 UE_SEQUENCER_DEFINE_CASTABLE(FCameraCutTrackModel);
 UE_SEQUENCER_DEFINE_CASTABLE(FCinematicShotTrackModel);
+UE_SEQUENCER_DEFINE_CASTABLE(FBindingLifetimeTrackModel);
 
 } // namespace Sequencer
 } // namespace UE
diff --git a/Engine/Source/Editor/MovieSceneTools/Private/MVVM/ViewModels/BindingLifetimeTrackModel.cpp b/Engine/Source/Editor/MovieSceneTools/Private/MVVM/ViewModels/BindingLifetimeTrackModel.cpp
new file mode 100644
--- /dev/null
+++ b/Engine/Source/Editor/MovieSceneTools/Private/MVVM/ViewModels/BindingLifetimeTrackModel.cpp
@@ -0,0 +1,97 @@
+// Copyright Epic Games, Inc. All Rights Reserved.
+
+#include "MVVM/ViewModels/BindingLifetimeTrackModel.h"
+
+#include "MovieSceneTrack.h"
+#include "Templates/Casts.h"
+#include "Tracks/MovieSceneBindingLifetimeTrack.h"
+
+namespace UE
+{
+	namespace Sequencer
+	{
+
+		TSharedPtr<FTrackModel> FBindingLifetimeTrackModel::CreateTrackModel(UMovieSceneTrack* Track)
+		{
+			if (UMovieSceneBindingLifetimeTrack* BindingLifetimeTrack = Cast<UMovieSceneBindingLifetimeTrack>(Track))
+			{
+				return MakeShared<FBindingLifetimeTrackModel>(BindingLifetimeTrack);
+			}
+			return nullptr;
+		}
+
+		FBindingLifetimeTrackModel::FBindingLifetimeTrackModel(UMovieSceneBindingLifetimeTrack* Track)
+			: FTrackModel(Track)
+		{
+		}
+
+		void FBindingLifetimeTrackModel::OnConstruct()
+		{
+			FTrackModel::OnConstruct();
+			RecalculateInverseLifetimeRange();
+		}
+
+		void FBindingLifetimeTrackModel::RecalculateInverseLifetimeRange()
+		{
+			InverseLifetimeRange.Reset();
+			InverseLifetimeRange.Add(TRange<FFrameNumber>::All());
+			// Iterate through our sections, removing them from the range
+
+			auto RemoveRangeFromSet = [&](TRange<FFrameNumber> SectionRange) {
+				for (int32 Index = 0; Index < InverseLifetimeRange.Num(); ++Index)
+				{
+					const TRange<FFrameNumber>& Current = InverseLifetimeRange[Index];
+					if (Current.Overlaps(SectionRange))
+					{
+						// Special case that difference doesn't handle well
+						if (!SectionRange.HasLowerBound() && !SectionRange.HasUpperBound())
+						{
+							InverseLifetimeRange.Reset();
+							return;
+						}
+
+						TArray<TRange<FFrameNumber>> SplitRanges = TRange<FFrameNumber>::Difference(Current, SectionRange);
+						for (const TRange<FFrameNumber>& NewRange : SplitRanges)
+						{
+							// Splitting infinite ranges keeps an infinite range, which we don't want to keep
+							if (!NewRange.HasLowerBound() && !NewRange.HasUpperBound())
+							{
+								continue;
+							}
+							
+							InverseLifetimeRange.Add(NewRange);
+						}
+						InverseLifetimeRange.RemoveAtSwap(Index--);
+					}
+				}
+			};
+
+			for (const TViewModelPtr<FSectionModel>& Item : GetSectionModels().IterateSubList<FSectionModel>())
+			{
+				RemoveRangeFromSet(Item->GetRange());
+			}
+		}
+
+		FSortingKey FBindingLifetimeTrackModel::GetSortingKey() const
+		{
+			FSortingKey SortingKey;
+			if (UMovieSceneTrack* Track = GetTrack())
+			{
+				SortingKey.CustomOrder = Track->GetSortingOrder();
+			}
+			return SortingKey.PrioritizeBy(4);
+		}
+
+		void FBindingLifetimeTrackModel::OnDeferredModifyFlush()
+		{
+			// Do the child update
+			FTrackModel::OnDeferredModifyFlush();
+
+			// Recalculate our valid range
+			RecalculateInverseLifetimeRange();
+		}
+
+
+	} // namespace Sequencer
+} // namespace UE
+
diff --git a/Engine/Source/Editor/MovieSceneTools/Private/MovieSceneToolsModule.cpp b/Engine/Source/Editor/MovieSceneTools/Private/MovieSceneToolsModule.cpp
--- a/Engine/Source/Editor/MovieSceneTools/Private/MovieSceneToolsModule.cpp
+++ b/Engine/Source/Editor/MovieSceneTools/Private/MovieSceneToolsModule.cpp
@@ -54,11 +54,13 @@
 #include "TrackEditors/CameraShakeSourceShakeTrackEditor.h"
 #include "TrackEditors/CVarTrackEditor.h"
 #include "TrackEditors/CustomPrimitiveDataTrackEditor.h"
+#include "TrackEditors/BindingLifetimeTrackEditor.h"
 
 #include "Channels/PerlinNoiseChannelInterface.h"
 
 #include "MVVM/ViewModels/CameraCutTrackModel.h"
 #include "MVVM/ViewModels/CinematicShotTrackModel.h"
+#include "MVVM/ViewModels/BindingLifetimeTrackModel.h"
 
 #include "MovieSceneBuiltInEasingFunctionCustomization.h"
 #include "MovieSceneAlphaBlendOptionCustomization.h"
@@ -180,11 +182,13 @@ void FMovieSceneToolsModule::StartupModule()
 		PrimitiveMaterialCreateEditorHandle = SequencerModule.RegisterTrackEditor(FOnCreateTrackEditor::CreateStatic(&FPrimitiveMaterialTrackEditor::CreateTrackEditor));
 		CameraShakeSourceShakeCreateEditorHandle = SequencerModule.RegisterTrackEditor(FOnCreateTrackEditor::CreateStatic(&FCameraShakeSourceShakeTrackEditor::CreateTrackEditor));
 		CVarTrackCreateEditorHandle = SequencerModule.RegisterTrackEditor(FOnCreateTrackEditor::CreateStatic(&FCVarTrackEditor::CreateTrackEditor));
-		CustomPrimitiveDataTrackCreateEditorHandle = SequencerModule.RegisterTrackEditor(FOnCreateTrackEditor::CreateStatic(&FCustomPrimitiveDataTrackEditor::CreateTrackEditor));
+		BindingLifetimeTrackCreateEditorHandle = SequencerModule.RegisterTrackEditor(FOnCreateTrackEditor::CreateStatic(&FBindingLifetimeTrackEditor::CreateTrackEditor));
+
 
 		// register track models
 		CameraCutTrackModelHandle = SequencerModule.RegisterTrackModel(FOnCreateTrackModel::CreateStatic(&FCameraCutTrackModel::CreateTrackModel));
 		CinematicShotTrackModelHandle = SequencerModule.RegisterTrackModel(FOnCreateTrackModel::CreateStatic(&FCinematicShotTrackModel::CreateTrackModel));
+		BindingLifetimeTrackModelHandle = SequencerModule.RegisterTrackModel(FOnCreateTrackModel::CreateStatic(&FBindingLifetimeTrackModel::CreateTrackModel));
 
 		RegisterClipboardConversions();
 
@@ -343,6 +347,7 @@ void FMovieSceneToolsModule::ShutdownModule()
 	SequencerModule.UnRegisterTrackEditor( ObjectTrackCreateEditorHandle );
 	SequencerModule.UnRegisterTrackEditor( PrimitiveMaterialCreateEditorHandle );
 	SequencerModule.UnRegisterTrackEditor( CVarTrackCreateEditorHandle );
+	SequencerModule.UnRegisterTrackEditor(BindingLifetimeTrackCreateEditorHandle);
 
 	// unregister track models
 	SequencerModule.UnregisterTrackModel( CameraCutTrackModelHandle );
diff --git a/Engine/Source/Editor/MovieSceneTools/Private/Sections/BindingLifetimeSection.cpp b/Engine/Source/Editor/MovieSceneTools/Private/Sections/BindingLifetimeSection.cpp
new file mode 100644
--- /dev/null
+++ b/Engine/Source/Editor/MovieSceneTools/Private/Sections/BindingLifetimeSection.cpp
@@ -0,0 +1,26 @@
+// Copyright Epic Games, Inc. All Rights Reserved.
+
+#include "Sections/BindingLifetimeSection.h"
+#include "SequencerSectionPainter.h"
+
+#define LOCTEXT_NAMESPACE "BindingLifetimeSection"
+
+bool FBindingLifetimeSection::IsSectionSelected() const
+{
+	TSharedPtr<ISequencer> SequencerPtr = Sequencer.Pin();
+
+	TArray<UMovieSceneTrack*> SelectedTracks;
+	SequencerPtr->GetSelectedTracks(SelectedTracks);
+
+	UMovieSceneSection* Section = WeakSection.Get();
+	UMovieSceneTrack* Track = Section ? CastChecked<UMovieSceneTrack>(Section->GetOuter()) : nullptr;
+	return Track && SelectedTracks.Contains(Track);
+}
+
+int32 FBindingLifetimeSection::OnPaintSection(FSequencerSectionPainter& Painter) const
+{
+	int32 LayerId = Painter.PaintSectionBackground();
+	return LayerId;
+}
+
+#undef LOCTEXT_NAMESPACE // "BindingLifetimeSection"
\ No newline at end of file
diff --git a/Engine/Source/Editor/MovieSceneTools/Private/Sections/BindingLifetimeSection.h b/Engine/Source/Editor/MovieSceneTools/Private/Sections/BindingLifetimeSection.h
new file mode 100644
--- /dev/null
+++ b/Engine/Source/Editor/MovieSceneTools/Private/Sections/BindingLifetimeSection.h
@@ -0,0 +1,28 @@
+// Copyright Epic Games, Inc. All Rights Reserved.
+
+#pragma once
+
+#include "ISequencerSection.h"
+#include "ISequencer.h"
+
+class FSequencerSectionPainter;
+
+class FBindingLifetimeSection
+	: public FSequencerSection
+{
+public:
+
+	FBindingLifetimeSection(UMovieSceneSection& InSectionObject, TWeakPtr<ISequencer> InSequencer)
+		: FSequencerSection(InSectionObject)
+		, Sequencer(InSequencer)
+	{}
+
+protected:
+
+	bool IsSectionSelected() const; 
+	virtual int32 OnPaintSection(FSequencerSectionPainter& Painter) const override;
+
+protected:
+
+	TWeakPtr<ISequencer> Sequencer;
+};
diff --git a/Engine/Source/Editor/MovieSceneTools/Private/TrackEditors/BindingLifetimeTrackEditor.cpp b/Engine/Source/Editor/MovieSceneTools/Private/TrackEditors/BindingLifetimeTrackEditor.cpp
new file mode 100644
--- /dev/null
+++ b/Engine/Source/Editor/MovieSceneTools/Private/TrackEditors/BindingLifetimeTrackEditor.cpp
@@ -0,0 +1,194 @@
+// Copyright Epic Games, Inc. All Rights Reserved.
+
+#include "TrackEditors/BindingLifetimeTrackEditor.h"
+
+#include "Containers/UnrealString.h"
+#include "Delegates/Delegate.h"
+#include "Framework/Commands/UIAction.h"
+#include "Framework/MultiBox/MultiBoxBuilder.h"
+#include "HAL/Platform.h"
+#include "ISequencer.h"
+#include "Internationalization/Internationalization.h"
+#include "Misc/Guid.h"
+#include "MovieScene.h"
+#include "MovieSceneSequence.h"
+#include "MovieSceneTrack.h"
+#include "ScopedTransaction.h"
+#include "Templates/Casts.h"
+#include "Textures/SlateIcon.h"
+#include "Tracks/MovieSceneBindingLifetimeTrack.h"
+#include "Sections/MovieSceneBindingLifetimeSection.h"
+#include "UObject/Class.h"
+#include "UObject/UnrealNames.h"
+#include "MovieSceneTrackEditor.h"
+#include "Sections/BindingLifetimeSection.h"
+#include "Widgets/SBoxPanel.h"
+#include "SequencerUtilities.h"
+
+class ISequencerTrackEditor;
+
+
+#define LOCTEXT_NAMESPACE "FBindingLifetimeTrackEditor"
+
+
+TSharedRef<ISequencerTrackEditor> FBindingLifetimeTrackEditor::CreateTrackEditor(TSharedRef<ISequencer> InSequencer)
+{
+	return MakeShareable(new FBindingLifetimeTrackEditor(InSequencer));
+}
+
+
+FBindingLifetimeTrackEditor::FBindingLifetimeTrackEditor(TSharedRef<ISequencer> InSequencer)
+	: FMovieSceneTrackEditor(InSequencer)
+{ }
+
+TSharedRef<ISequencerSection> FBindingLifetimeTrackEditor::MakeSectionInterface(UMovieSceneSection & SectionObject, UMovieSceneTrack & Track, FGuid ObjectBinding)
+{
+	return MakeShared<FBindingLifetimeSection>(SectionObject, GetSequencer());
+}
+
+void FBindingLifetimeTrackEditor::CreateNewSection(UMovieSceneTrack* Track, bool bSelect)
+{
+	TSharedPtr<ISequencer> SequencerPtr = GetSequencer();
+	if (SequencerPtr.IsValid())
+	{
+		UMovieScene* FocusedMovieScene = GetFocusedMovieScene();
+		FQualifiedFrameTime CurrentTime = SequencerPtr->GetLocalTime();
+
+		FScopedTransaction Transaction(LOCTEXT("CreateNewSectionTransactionText", "Add Section"));
+
+
+		UMovieSceneSection* NewSection = NewObject<UMovieSceneSection>(Track, UMovieSceneBindingLifetimeSection::StaticClass(), NAME_None, RF_Transactional);
+		check(NewSection);
+
+		Track->Modify();
+
+		// first section by default should be infinite
+		if (Track->GetAllSections().Num() == 0)
+		{
+			NewSection->SetRange(TRange<FFrameNumber>::All());
+		}
+		else
+		{
+			int32 Duration = 0;
+
+			if (CurrentTime.Time.FrameNumber < FocusedMovieScene->GetPlaybackRange().GetUpperBoundValue())
+			{
+				Duration = FocusedMovieScene->GetPlaybackRange().GetUpperBoundValue().Value - CurrentTime.Time.FrameNumber.Value;
+			}
+			else
+			{
+				const float DefaultLengthInSeconds = 5.f;
+				Duration = (DefaultLengthInSeconds * SequencerPtr->GetFocusedTickResolution()).FloorToFrame().Value;
+			}
+
+			NewSection->InitialPlacement(Track->GetAllSections(), CurrentTime.Time.FrameNumber.Value, Duration, false);
+		}
+
+		Track->AddSection(*NewSection);
+		Track->UpdateEasing();
+
+		if (bSelect)
+		{
+			SequencerPtr->EmptySelection();
+			SequencerPtr->SelectSection(NewSection);
+			SequencerPtr->ThrobSectionSelection();
+		}
+
+		SequencerPtr->NotifyMovieSceneDataChanged(EMovieSceneDataChangeType::MovieSceneStructureItemAdded);
+	}
+}
+
+
+UMovieSceneTrack* FBindingLifetimeTrackEditor::AddTrack(UMovieScene* FocusedMovieScene, const FGuid& ObjectHandle, TSubclassOf<UMovieSceneTrack> TrackClass, FName UniqueTypeName)
+{
+	UMovieSceneTrack* NewTrack = FMovieSceneTrackEditor::AddTrack(FocusedMovieScene, ObjectHandle, TrackClass, UniqueTypeName);
+
+	if (auto* BindingLifetimeTrack = Cast<UMovieSceneBindingLifetimeTrack>(NewTrack))
+	{
+		BindingLifetimeTrack->Modify();
+		CreateNewSection(BindingLifetimeTrack, false);
+	}
+
+	return NewTrack;
+}
+
+
+void FBindingLifetimeTrackEditor::BuildObjectBindingTrackMenu(FMenuBuilder& MenuBuilder, const TArray<FGuid>& ObjectBindings, const UClass* ObjectClass)
+{
+	UMovieSceneSequence* MovieSequence = GetSequencer()->GetFocusedMovieSceneSequence();
+
+	// TODO: Do we want to restrict this to level sequences or no?
+	if (!MovieSequence || MovieSequence->GetClass()->GetName() != TEXT("LevelSequence"))
+	{
+		return;
+	}
+
+	MenuBuilder.AddMenuEntry(
+		LOCTEXT("AddBindingLifetimeTrack", "Binding Lifetime"),
+		LOCTEXT("AddBindingLifetimeTrackTooltip", "Adds a new track that controls the lifetime of the track's object binding."),
+		FSlateIcon(),
+		FUIAction(
+			FExecuteAction::CreateRaw(this, &FBindingLifetimeTrackEditor::HandleAddBindingLifetimeTrackMenuEntryExecute, ObjectBindings),
+			FCanExecuteAction::CreateSP(this, &FBindingLifetimeTrackEditor::CanAddBindingLifetimeTrack, ObjectBindings[0])
+		)
+	);
+}
+
+
+bool FBindingLifetimeTrackEditor::SupportsType(TSubclassOf<UMovieSceneTrack> Type) const
+{
+	return (Type == UMovieSceneBindingLifetimeTrack::StaticClass());
+}
+
+
+bool FBindingLifetimeTrackEditor::SupportsSequence(UMovieSceneSequence* InSequence) const
+{
+	ETrackSupport TrackSupported = InSequence ? InSequence->IsTrackSupported(UMovieSceneBindingLifetimeTrack::StaticClass()) : ETrackSupport::NotSupported;
+	return TrackSupported == ETrackSupport::Supported;
+}
+
+
+void FBindingLifetimeTrackEditor::HandleAddBindingLifetimeTrackMenuEntryExecute(TArray<FGuid> ObjectBindings)
+{
+	FScopedTransaction AddSpawnTrackTransaction(LOCTEXT("AddBindingLifetimeTrack_Transaction", "Add Binding Lifetime Track"));
+
+	for (FGuid ObjectBinding : ObjectBindings)
+	{
+		AddTrack(GetSequencer()->GetFocusedMovieSceneSequence()->GetMovieScene(), ObjectBinding, UMovieSceneBindingLifetimeTrack::StaticClass(), NAME_None);
+	}
+	GetSequencer()->NotifyMovieSceneDataChanged(EMovieSceneDataChangeType::MovieSceneStructureItemAdded);
+}
+
+
+bool FBindingLifetimeTrackEditor::CanAddBindingLifetimeTrack(FGuid ObjectBinding) const
+{
+	return !GetSequencer()->GetFocusedMovieSceneSequence()->GetMovieScene()->FindTrack<UMovieSceneBindingLifetimeTrack>(ObjectBinding);
+}
+
+TSharedPtr<SWidget> FBindingLifetimeTrackEditor::BuildOutlinerEditWidget(const FGuid& ObjectBinding, UMovieSceneTrack* Track, const FBuildEditWidgetParams& Params)
+{
+	check(Track);
+
+	TSharedPtr<ISequencer> SequencerPtr = GetSequencer();
+	if (!SequencerPtr.IsValid())
+	{
+		return SNullWidget::NullWidget;
+	}
+
+	auto OnClickedCallback = [this, Track]() -> FReply
+	{
+		CreateNewSection(Track, true);
+		return FReply::Handled();
+	};
+
+	return SNew(SHorizontalBox)
+		+ SHorizontalBox::Slot()
+		.AutoWidth()
+		.VAlign(VAlign_Center)
+		[
+			FSequencerUtilities::MakeAddButton(LOCTEXT("AddSection", "Section"), FOnClicked::CreateLambda(OnClickedCallback), Params.NodeIsHovered, GetSequencer())
+		];
+}
+
+
+#undef LOCTEXT_NAMESPACE
diff --git a/Engine/Source/Editor/MovieSceneTools/Public/MVVM/ViewModels/BindingLifetimeTrackModel.h b/Engine/Source/Editor/MovieSceneTools/Public/MVVM/ViewModels/BindingLifetimeTrackModel.h
new file mode 100644
--- /dev/null
+++ b/Engine/Source/Editor/MovieSceneTools/Public/MVVM/ViewModels/BindingLifetimeTrackModel.h
@@ -0,0 +1,53 @@
+// Copyright Epic Games, Inc. All Rights Reserved.
+
+#pragma once
+
+#include "MVVM/Extensions/ISortableExtension.h"
+#include "MVVM/ICastable.h"
+#include "MVVM/ViewModels/TrackModel.h"
+#include "Templates/SharedPointer.h"
+#include "MVVM/Extensions/IBindingLifetimeExtension.h"
+
+class UMovieSceneBindingLifetimeTrack;
+class UMovieSceneTrack;
+namespace UE::Sequencer { template <typename T> struct TAutoRegisterViewModelTypeID; }
+
+namespace UE
+{
+	namespace Sequencer
+	{
+
+		class MOVIESCENETOOLS_API FBindingLifetimeTrackModel
+			: public FTrackModel
+			, public IBindingLifetimeExtension
+		{
+		public:
+
+			UE_SEQUENCER_DECLARE_CASTABLE(FBindingLifetimeTrackModel, FTrackModel, IBindingLifetimeExtension);
+
+			static TSharedPtr<FTrackModel> CreateTrackModel(UMovieSceneTrack* Track);
+
+			explicit FBindingLifetimeTrackModel(UMovieSceneBindingLifetimeTrack* Track);
+
+			FSortingKey GetSortingKey() const override;
+
+			void OnDeferredModifyFlush() override;
+
+			// IBindingLifetimeExtension
+			const TArray<TRange<FFrameNumber>>& GetInverseLifetimeRange() const override { return InverseLifetimeRange; }
+		
+		protected:
+			virtual void OnConstruct() override;
+
+		private:
+			void RecalculateInverseLifetimeRange();
+
+			// The inverse of the range created by all of our binding lifetime sections
+			// In other words, the ranges where the object binding should be deactivated.
+			TArray<TRange<FFrameNumber>> InverseLifetimeRange;
+
+		};
+
+	} // namespace Sequencer
+} // namespace UE
+
diff --git a/Engine/Source/Editor/MovieSceneTools/Public/MovieSceneToolsModule.h b/Engine/Source/Editor/MovieSceneTools/Public/MovieSceneToolsModule.h
--- a/Engine/Source/Editor/MovieSceneTools/Public/MovieSceneToolsModule.h
+++ b/Engine/Source/Editor/MovieSceneTools/Public/MovieSceneToolsModule.h
@@ -138,9 +138,11 @@ private:
 	FDelegateHandle CameraShakeSourceShakeCreateEditorHandle;
 	FDelegateHandle CVarTrackCreateEditorHandle;
 	FDelegateHandle CustomPrimitiveDataTrackCreateEditorHandle;
+	FDelegateHandle BindingLifetimeTrackCreateEditorHandle;
 
 	FDelegateHandle CameraCutTrackModelHandle;
 	FDelegateHandle CinematicShotTrackModelHandle;
+	FDelegateHandle BindingLifetimeTrackModelHandle;
 
 	FDelegateHandle GenerateEventEntryPointsHandle;
 	FDelegateHandle FixupDynamicBindingPayloadParameterNameHandle;
diff --git a/Engine/Source/Editor/MovieSceneTools/Public/TrackEditors/BindingLifetimeTrackEditor.h b/Engine/Source/Editor/MovieSceneTools/Public/TrackEditors/BindingLifetimeTrackEditor.h
new file mode 100644
--- /dev/null
+++ b/Engine/Source/Editor/MovieSceneTools/Public/TrackEditors/BindingLifetimeTrackEditor.h
@@ -0,0 +1,74 @@
+// Copyright Epic Games, Inc. All Rights Reserved.
+
+#pragma once
+
+#include "Containers/Array.h"
+#include "CoreMinimal.h"
+#include "ISequencer.h"
+#include "ISequencerSection.h"
+#include "Misc/Guid.h"
+#include "MovieSceneTrack.h"
+#include "Templates/SharedPointer.h"
+#include "Templates/SubclassOf.h"
+#include "MovieSceneTrackEditor.h"
+#include "UObject/NameTypes.h"
+
+class FMenuBuilder;
+class ISequencer;
+class SWidget;
+class UClass;
+class UMovieScene;
+class UMovieSceneSequence;
+class UMovieSceneTrack;
+class UObject;
+struct FBuildEditWidgetParams;
+struct FGuid;
+
+/**
+ * A track editor for controlling the lifetime of an object binding
+ */
+class MOVIESCENETOOLS_API FBindingLifetimeTrackEditor
+	: public FMovieSceneTrackEditor
+{
+public:
+
+	/**
+	 * Factory function to create an instance of this class (called by a sequencer).
+	 *
+	 * @param InSequencer The sequencer instance to be used by this tool.
+	 * @return The new instance of this class.
+	 */
+	static TSharedRef<ISequencerTrackEditor> CreateTrackEditor(TSharedRef<ISequencer> InSequencer);
+
+public:
+
+	/**
+	 * Creates and initializes a new instance.
+	 *
+	 * @param InSequencer The sequencer instance to be used by this tool.
+	 */
+	FBindingLifetimeTrackEditor(TSharedRef<ISequencer> InSequencer);
+
+	void CreateNewSection(UMovieSceneTrack* Track, bool bSelect);
+
+public:
+
+	// ISequencerTrackEditor interface
+
+	virtual UMovieSceneTrack* AddTrack(UMovieScene* FocusedMovieScene, const FGuid& ObjectHandle, TSubclassOf<class UMovieSceneTrack> TrackClass, FName UniqueTypeName) override;
+	virtual void BuildObjectBindingTrackMenu(FMenuBuilder& MenuBuilder, const TArray<FGuid>& ObjectBindings, const UClass* ObjectClass) override;
+	
+	virtual TSharedPtr<SWidget> BuildOutlinerEditWidget(const FGuid& ObjectBinding, UMovieSceneTrack* Track, const FBuildEditWidgetParams& Params) override;
+	
+	virtual bool HandleAssetAdded(UObject* Asset, const FGuid& TargetObjectGuid) override { return false; }
+	virtual bool SupportsType(TSubclassOf<UMovieSceneTrack> Type) const override;
+	virtual bool SupportsSequence(UMovieSceneSequence* InSequence) const override;
+
+	virtual TSharedRef<ISequencerSection> MakeSectionInterface(UMovieSceneSection& SectionObject, UMovieSceneTrack& Track, FGuid ObjectBinding) override;
+
+private:
+
+	/** Callback for executing the "Add Binding Lifetime Track" menu entry. */
+	void HandleAddBindingLifetimeTrackMenuEntryExecute(TArray<FGuid> ObjectBindings);
+	bool CanAddBindingLifetimeTrack(FGuid ObjectBinding) const;
+};
diff --git a/Engine/Source/Editor/Sequencer/Private/MVVM/SequencerViewModelTypeIDs.cpp b/Engine/Source/Editor/Sequencer/Private/MVVM/SequencerViewModelTypeIDs.cpp
--- a/Engine/Source/Editor/Sequencer/Private/MVVM/SequencerViewModelTypeIDs.cpp
+++ b/Engine/Source/Editor/Sequencer/Private/MVVM/SequencerViewModelTypeIDs.cpp
@@ -17,6 +17,7 @@
 #include "MVVM/Extensions/ITrackExtension.h"
 #include "MVVM/Extensions/IObjectBindingExtension.h"
 #include "MVVM/Extensions/ISnappableExtension.h"
+#include "MVVM/Extensions/IBindingLifetimeExtension.h"
 
 #include "MVVM/ViewModels/SequencerOutlinerViewModel.h"
 #include "MVVM/ViewModels/SequencerEditorViewModel.h"
@@ -38,6 +39,7 @@ namespace Sequencer
 {
 
 // Model types
+UE_SEQUENCER_DEFINE_CASTABLE(FBindingLifetimeOverlayModel);
 UE_SEQUENCER_DEFINE_CASTABLE(FCategoryGroupModel);
 UE_SEQUENCER_DEFINE_CASTABLE(FCategoryModel);
 UE_SEQUENCER_DEFINE_CASTABLE(FChannelGroupModel);
@@ -75,6 +77,7 @@ UE_SEQUENCER_DEFINE_VIEW_MODEL_TYPE_ID(FTrackRowModelStorageExtension);
 UE_SEQUENCER_DEFINE_VIEW_MODEL_TYPE_ID(ICurveEditorTreeItemExtension);
 UE_SEQUENCER_DEFINE_VIEW_MODEL_TYPE_ID(IObjectBindingExtension);
 UE_SEQUENCER_DEFINE_VIEW_MODEL_TYPE_ID(ISnappableExtension);
+UE_SEQUENCER_DEFINE_VIEW_MODEL_TYPE_ID(IBindingLifetimeExtension);
 UE_SEQUENCER_DEFINE_VIEW_MODEL_TYPE_ID(FOutlinerCacheExtension);
 
 } // namespace Sequencer
diff --git a/Engine/Source/Editor/Sequencer/Private/MVVM/ViewModels/BindingLifetimeOverlayModel.cpp b/Engine/Source/Editor/Sequencer/Private/MVVM/ViewModels/BindingLifetimeOverlayModel.cpp
new file mode 100644
--- /dev/null
+++ b/Engine/Source/Editor/Sequencer/Private/MVVM/ViewModels/BindingLifetimeOverlayModel.cpp
@@ -0,0 +1,69 @@
+// Copyright Epic Games, Inc. All Rights Reserved.
+
+#include "MVVM/ViewModels/BindingLifetimeOverlayModel.h"
+
+#include "Containers/ArrayView.h"
+#include "HAL/Platform.h"
+#include "MVVM/ViewModelPtr.h"
+#include "MVVM/ViewModels/EditorViewModel.h"
+#include "MVVM/ViewModels/SequencerEditorViewModel.h"
+#include "MVVM/ViewModels/ViewModelIterators.h"
+#include "MVVM/Views/STrackLane.h"
+#include "Math/RangeBound.h"
+#include "Misc/AssertionMacros.h"
+#include "Misc/FrameNumber.h"
+#include "MovieSceneTimeHelpers.h"
+#include "Templates/TypeHash.h"
+#include "Widgets/DeclarativeSyntaxSupport.h"
+#include "Widgets/SSequencerBindingLifetimeOverlay.h"
+#include "Sequencer.h"
+#include "AnimatedRange.h"
+
+namespace UE
+{
+	namespace Sequencer
+	{
+
+		FBindingLifetimeOverlayModel::FBindingLifetimeOverlayModel(TWeakPtr<FViewModel> LayerRoot
+			, TWeakPtr<FSequencerEditorViewModel> InEditorViewModel
+			, TViewModelPtr<IBindingLifetimeExtension> InBindingLifetimeTrack)
+			: WeakEditorViewModel(InEditorViewModel)
+			, BindingLifetimeTrack(InBindingLifetimeTrack)
+		{
+		}
+
+		FBindingLifetimeOverlayModel::~FBindingLifetimeOverlayModel()
+		{
+		}
+
+		TSharedPtr<ITrackLaneWidget> FBindingLifetimeOverlayModel::CreateTrackLaneView(const FCreateTrackLaneViewParams& InParams)
+		{
+			return SNew(SSequencerBindingLifetimeOverlay, InParams.OwningTrackLane->GetTrackAreaView(), InParams.Editor->CastThisSharedChecked<FSequencerEditorViewModel>(), SharedThis(this))
+					.Visibility(EVisibility::HitTestInvisible);
+		}
+
+		FTrackLaneVirtualAlignment FBindingLifetimeOverlayModel::ArrangeVirtualTrackLaneView() const
+		{
+			TSharedPtr<FSequencerEditorViewModel> EditorViewModel = WeakEditorViewModel.Pin();
+			TSharedPtr<FSequencer> Sequencer = EditorViewModel->GetSequencerImpl();
+
+			FFrameRate TickResolution = Sequencer->GetFocusedTickResolution();
+
+			TRange<FFrameNumber> ViewRange(TickResolution.AsFrameNumber(Sequencer->GetViewRange().GetLowerBoundValue()), TickResolution.AsFrameNumber(Sequencer->GetViewRange().GetUpperBoundValue()));
+
+
+			return FTrackLaneVirtualAlignment::Proportional(ViewRange, 1.f);
+		}
+
+		const TArray<TRange<FFrameNumber>>& FBindingLifetimeOverlayModel::GetInverseLifetimeRange() const
+		{
+			const static TArray<TRange<FFrameNumber>> EmptyLifetimeRange;
+			if (BindingLifetimeTrack)
+			{
+				return BindingLifetimeTrack->GetInverseLifetimeRange();
+			}
+			return EmptyLifetimeRange;
+		}
+	} // namespace Sequencer
+} // namespace UE
+
diff --git a/Engine/Source/Editor/Sequencer/Private/MVVM/ViewModels/LayerBarModel.cpp b/Engine/Source/Editor/Sequencer/Private/MVVM/ViewModels/LayerBarModel.cpp
--- a/Engine/Source/Editor/Sequencer/Private/MVVM/ViewModels/LayerBarModel.cpp
+++ b/Engine/Source/Editor/Sequencer/Private/MVVM/ViewModels/LayerBarModel.cpp
@@ -84,7 +84,7 @@ TSharedPtr<ITrackLaneWidget> FLayerBarModel::CreateTrackLaneView(const FCreateTr
 FTrackLaneVirtualAlignment FLayerBarModel::ArrangeVirtualTrackLaneView() const
 {
 	TRange<FFrameNumber> Range = ComputeRange();
-	return FTrackLaneVirtualAlignment::Proportional(Range, .8f);
+	return FTrackLaneVirtualAlignment::Proportional(Range, .8f, EVerticalAlignment::VAlign_Top);
 }
 
 ESelectionIntent FLayerBarModel::IsSelectable() const
diff --git a/Engine/Source/Editor/Sequencer/Private/MVVM/ViewModels/ObjectBindingModel.cpp b/Engine/Source/Editor/Sequencer/Private/MVVM/ViewModels/ObjectBindingModel.cpp
--- a/Engine/Source/Editor/Sequencer/Private/MVVM/ViewModels/ObjectBindingModel.cpp
+++ b/Engine/Source/Editor/Sequencer/Private/MVVM/ViewModels/ObjectBindingModel.cpp
@@ -5,15 +5,19 @@
 #include "MVVM/ViewModels/SequenceModel.h"
 #include "MVVM/ViewModels/TrackModel.h"
 #include "MVVM/ViewModels/LayerBarModel.h"
+#include "MVVM/ViewModels/BindingLifetimeTrackModel.h"
 #include "MVVM/ViewModels/ViewModelIterators.h"
 #include "MVVM/TrackModelStorageExtension.h"
 #include "MVVM/ObjectBindingModelStorageExtension.h"
 #include "MVVM/ViewModels/SequencerEditorViewModel.h"
 #include "MVVM/ViewModels/OutlinerViewModelDragDropOp.h"
 #include "MVVM/Views/SOutlinerObjectBindingView.h"
+#include "MVVM/Views/STrackLane.h"
 #include "MVVM/Selection/Selection.h"
 #include "MVVM/Extensions/IRecyclableExtension.h"
+#include "MVVM/Extensions/IBindingLifetimeExtension.h"
 #include "Algo/Sort.h"
+#include "AnimatedRange.h"
 #include "ClassViewerModule.h"
 #include "Containers/ArrayBuilder.h"
 #include "Engine/LevelStreaming.h"
@@ -42,6 +46,7 @@
 #include "SequencerUtilities.h"
 #include "Styling/AppStyle.h"
 #include "Styling/SlateIconFinder.h"
+#include "Widgets/SSequencerBindingLifetimeOverlay.h"
 
 #define LOCTEXT_NAMESPACE "ObjectBindingModel"
 
@@ -122,9 +127,11 @@ void GetKeyablePropertyPaths(UClass* Class, void* ValuePtr, UStruct* PropertySou
 FObjectBindingModel::FObjectBindingModel(FSequenceModel* InOwnerModel, const FMovieSceneBinding& InBinding)
 	: ObjectBindingID(InBinding.GetObjectGuid())
 	, TrackAreaList(EViewModelListType::TrackArea)
+	, TopLevelChildTrackAreaList(GetTopLevelChildTrackAreaGroupType())
 	, OwnerModel(InOwnerModel)
 {
 	RegisterChildList(&TrackAreaList);
+	RegisterChildList(&TopLevelChildTrackAreaList);
 
 	SetIdentifier(*ObjectBindingID.ToString());
 }
@@ -133,6 +140,12 @@ FObjectBindingModel::~FObjectBindingModel()
 {
 }
 
+EViewModelListType FObjectBindingModel::GetTopLevelChildTrackAreaGroupType()
+{
+	static EViewModelListType TopLevelChildTrackAreaGroup = RegisterCustomModelListType();
+	return TopLevelChildTrackAreaGroup;
+}
+
 void FObjectBindingModel::OnConstruct()
 {
 	if (!LayerBar)
@@ -145,7 +158,7 @@ void FObjectBindingModel::OnConstruct()
 			LayerBar = MakeShared<FLayerBarModel>(AsShared());
 			LayerBar->SetLinkedOutlinerItem(SharedThis(this));
 
-			GetChildrenForList(&TrackAreaList).AddChild(LayerBar);
+			GetChildrenForList(&TopLevelChildTrackAreaList).AddChild(LayerBar);
 		}
 	}
 
@@ -192,7 +205,7 @@ bool FObjectBindingModel::SupportsRebinding() const
 FTrackAreaParameters FObjectBindingModel::GetTrackAreaParameters() const
 {
 	FTrackAreaParameters Params;
-	Params.LaneType = ETrackAreaLaneType::Inline;
+	Params.LaneType = ETrackAreaLaneType::Nested;
 	return Params;
 }
 
@@ -201,13 +214,28 @@ FViewModelVariantIterator FObjectBindingModel::GetTrackAreaModelList() const
 	return &TrackAreaList;
 }
 
+FViewModelVariantIterator FObjectBindingModel::GetTopLevelChildTrackAreaModels() const
+{
+	return &TopLevelChildTrackAreaList;
+}
+
 void FObjectBindingModel::AddTrack(UMovieSceneTrack* Track)
 {
 	FTrackModelStorageExtension* TrackStorage = OwnerModel->CastDynamic<FTrackModelStorageExtension>();
 
-	TSharedPtr<FTrackModel> TrackModel = TrackStorage->CreateModelForTrack(Track, AsShared());
+	TViewModelPtr<FTrackModel> TrackModel = TrackStorage->CreateModelForTrack(Track, AsShared());
 
 	GetChildrenForList(&OutlinerChildList).AddChild(TrackModel);
+
+	if (TrackModel->IsA<IBindingLifetimeExtension>())
+	{
+		if (!BindingLifetimeOverlayModel)
+		{
+			BindingLifetimeOverlayModel = MakeShared<FBindingLifetimeOverlayModel>(AsShared(), GetEditor(), TrackModel.ImplicitCast());
+			BindingLifetimeOverlayModel->SetLinkedOutlinerItem(SharedThis(this));
+			GetChildrenForList(&TrackAreaList).AddChild(BindingLifetimeOverlayModel);
+		}
+	}
 }
 
 void FObjectBindingModel::RemoveTrack(UMovieSceneTrack* Track)
@@ -218,6 +246,14 @@ void FObjectBindingModel::RemoveTrack(UMovieSceneTrack* Track)
 	if (TrackModel)
 	{
 		TrackModel->RemoveFromParent();
+		if (TrackModel->IsA<IBindingLifetimeExtension>())
+		{
+			if (BindingLifetimeOverlayModel)
+			{
+				BindingLifetimeOverlayModel->RemoveFromParent();
+				BindingLifetimeOverlayModel.Reset();
+			}
+		}
 	}
 }
 
@@ -353,6 +389,11 @@ FSlateColor FObjectBindingModel::GetLabelColor() const
 	// e.g. Spawnables don't have valid object bindings when their track hasn't spawned them yet,
 	// so we override the default behavior of red with a gray so that users don't think there is something wrong.
 	TFunction<FSlateColor(const FObjectBindingModel&)> GetObjectBindingAncestorInvalidLabelColor = [&](const FObjectBindingModel& InObjectBindingModel) -> FSlateColor {
+		if (!Sequencer->State.GetBindingActivation(InObjectBindingModel.GetObjectGuid(), OwnerModel->GetSequenceID()))
+		{
+			return FSlateColor::UseSubduedForeground();
+		}
+		
 		if (TSharedPtr<FObjectBindingModel> ParentBindingModel = InObjectBindingModel.FindAncestorOfType<FObjectBindingModel>())
 		{
 			return GetObjectBindingAncestorInvalidLabelColor(*ParentBindingModel.Get());
@@ -1199,6 +1240,8 @@ void FObjectBindingModel::Delete()
 	{
 		ParentFolder->GetFolder()->RemoveChildObjectBinding(ObjectBindingID);
 	}
+
+	BindingLifetimeOverlayModel.Reset();
 }
 
 } // namespace Sequencer
diff --git a/Engine/Source/Editor/Sequencer/Private/MVVM/ViewModels/OutlinerItemModel.cpp b/Engine/Source/Editor/Sequencer/Private/MVVM/ViewModels/OutlinerItemModel.cpp
--- a/Engine/Source/Editor/Sequencer/Private/MVVM/ViewModels/OutlinerItemModel.cpp
+++ b/Engine/Source/Editor/Sequencer/Private/MVVM/ViewModels/OutlinerItemModel.cpp
@@ -404,6 +404,31 @@ TSharedPtr<SWidget> FOutlinerItemModelMixin::CreateContextMenuWidget(const FCrea
 	return nullptr;
 }
 
+FSlateColor FOutlinerItemModelMixin::GetLabelColor() const
+{
+	if (TViewModelPtr<FSequenceModel> SequenceModel = AsViewModel()->FindAncestorOfType<FSequenceModel>())
+	{
+		if (TSharedPtr<FSequencerEditorViewModel> SequencerModel = SequenceModel->GetEditor())
+		{
+			if (IMovieScenePlayer* Player = SequencerModel->GetSequencer().Get())
+			{
+				if (TViewModelPtr<FObjectBindingModel> ObjectBindingModel = AsViewModel()->FindAncestorOfType<FObjectBindingModel>())
+				{
+					// If the object binding model has an invalid binding, we want to use its label color, as it may be red or gray depending on situation
+					// and we want the children of that to have the same color.
+					// Otherwise, we can use the track's label color below
+					TArrayView<TWeakObjectPtr<> > BoundObjects = Player->FindBoundObjects(ObjectBindingModel->GetObjectGuid(), SequenceModel->GetSequenceID());
+					if (BoundObjects.Num() == 0)
+					{
+						return ObjectBindingModel->GetLabelColor();
+					}
+				}
+			}
+		}
+	}
+	return IOutlinerExtension::GetLabelColor();
+}
+
 void FOutlinerItemModelMixin::BuildContextMenu(FMenuBuilder& MenuBuilder)
 {
 	TSharedPtr<FSequencer> Sequencer = StaticCastSharedPtr<FSequencer>(GetEditor()->GetSequencer());
diff --git a/Engine/Source/Editor/Sequencer/Private/Widgets/SSequencerBindingLifetimeOverlay.cpp b/Engine/Source/Editor/Sequencer/Private/Widgets/SSequencerBindingLifetimeOverlay.cpp
new file mode 100644
--- /dev/null
+++ b/Engine/Source/Editor/Sequencer/Private/Widgets/SSequencerBindingLifetimeOverlay.cpp
@@ -0,0 +1,127 @@
+// Copyright Epic Games, Inc. All Rights Reserved.
+
+#include "Widgets/SSequencerBindingLifetimeOverlay.h"
+
+#include "MVVM/ObjectBindingModelStorageExtension.h"
+#include "MVVM/ViewModels/BindingLifetimeOverlayModel.h"
+#include "MVVM/ViewModels/SequenceModel.h"
+#include "MVVM/ViewModels/SequencerEditorViewModel.h"
+#include "MVVM/ViewModels/TrackAreaViewModel.h"
+#include "MVVM/ViewModels/ViewModel.h"
+#include "MVVM/Views/STrackAreaView.h"
+
+#include "SequencerCommonHelpers.h"
+#include "MovieSceneTimeHelpers.h"
+#include "Sequencer.h"
+#include "Rendering/DrawElementTypes.h"
+
+#define LOCTEXT_NAMESPACE "SSequencerBindingLifetimeOverlay"
+
+namespace UE
+{
+	namespace Sequencer
+	{
+		void SSequencerBindingLifetimeOverlay::Construct(const FArguments& InArgs, TWeakPtr<STrackAreaView> InWeakTrackArea, TWeakPtr<FSequencerEditorViewModel> InWeakEditor, TWeakPtr<FBindingLifetimeOverlayModel> InWeakBindingLifetimeOverlayModel)
+		{
+			WeakBindingLifetimeOverlayModel = InWeakBindingLifetimeOverlayModel;
+			WeakEditor = InWeakEditor;
+			WeakTrackArea = InWeakTrackArea; 
+			TimeToPixel = InWeakTrackArea.Pin()->GetTimeToPixel();
+		}
+
+		TSharedPtr<FBindingLifetimeOverlayModel> SSequencerBindingLifetimeOverlay::GetBindingLifetimeOverlayModel() const
+		{
+			return WeakBindingLifetimeOverlayModel.Pin();
+		}
+
+		TSharedPtr<FSequencer> SSequencerBindingLifetimeOverlay::GetSequencer() const
+		{
+			if (TSharedPtr<FSequencerEditorViewModel> Editor = WeakEditor.Pin())
+			{
+				return Editor->GetSequencerImpl();
+			}
+			return nullptr;
+		}
+
+		int32 SSequencerBindingLifetimeOverlay::GetOverlapPriority() const
+		{
+			return 1000;
+		}
+
+		TSharedRef<const SWidget> SSequencerBindingLifetimeOverlay::AsWidget() const
+		{
+			return AsShared();
+		}
+
+		FTrackLaneScreenAlignment SSequencerBindingLifetimeOverlay::GetAlignment(const FTimeToPixel& InTimeToPixel, const FGeometry& InParentGeometry) const
+		{
+			if (TSharedPtr<FBindingLifetimeOverlayModel> BindingLifetimeOverlayModel = GetBindingLifetimeOverlayModel())
+			{
+				return BindingLifetimeOverlayModel->ArrangeVirtualTrackLaneView().ToScreen(InTimeToPixel, InParentGeometry);
+			}
+			return FTrackLaneScreenAlignment{};
+		}
+
+
+		int32 SSequencerBindingLifetimeOverlay::OnPaint(const FPaintArgs& Args, const FGeometry& AllottedGeometry, const FSlateRect& MyCullingRect, FSlateWindowElementList& OutDrawElements, int32 LayerId, const FWidgetStyle& InWidgetStyle, bool bParentEnabled) const
+		{
+			const ESlateDrawEffect DrawEffects = bParentEnabled
+				? ESlateDrawEffect::None
+				: ESlateDrawEffect::DisabledEffect;
+			static const FSlateBrush* InactiveLifetimeOverlay = FAppStyle::GetBrush("Sequencer.Section.SequencerDeactivatedOverlay");
+
+			if (TSharedPtr<FBindingLifetimeOverlayModel> BindingLifetimeOverlayModel = GetBindingLifetimeOverlayModel())
+			{
+				// If applicable, paint an overlay over the range that the object binding is deactivated
+				const TArray<TRange<FFrameNumber>>& InverseLifetimeRange = BindingLifetimeOverlayModel->GetInverseLifetimeRange();
+				if (!InverseLifetimeRange.IsEmpty())
+				{
+					for (TRange<FFrameNumber> ExcludedRange : InverseLifetimeRange)
+					{
+						const FFrameNumber LowerFrame = ExcludedRange.HasLowerBound() ? ExcludedRange.GetLowerBoundValue() : TimeToPixel->PixelToFrame(AllottedGeometry.Position.X).FloorToFrame();
+						const FFrameNumber UpperFrame = ExcludedRange.HasUpperBound() ? ExcludedRange.GetUpperBoundValue() : TimeToPixel->PixelToFrame(AllottedGeometry.Position.X + AllottedGeometry.Size.X).CeilToFrame();
+
+						const float ExcludedRangeL = TimeToPixel->FrameToPixel(LowerFrame);
+						const float ExcludedRangeR = TimeToPixel->FrameToPixel(UpperFrame);
+
+						// Black line on left
+						FSlateDrawElement::MakeBox(
+							OutDrawElements,
+							LayerId + 50,
+							AllottedGeometry.ToPaintGeometry(FVector2f(1.0f, AllottedGeometry.Size.Y), FSlateLayoutTransform(FVector2f(ExcludedRangeL, 0.f))),
+							FAppStyle::GetBrush("WhiteBrush"),
+							ESlateDrawEffect::None,
+							FLinearColor::Black
+						);
+
+						// Black tint for excluded regions
+						// Artificially increase layer id to ensure we paint over any possible section elements
+						FSlateDrawElement::MakeBox(
+							OutDrawElements,
+							LayerId + 50,
+							AllottedGeometry.ToPaintGeometry(FVector2f(ExcludedRangeR - ExcludedRangeL, AllottedGeometry.Size.Y), FSlateLayoutTransform(FVector2f(ExcludedRangeL, 0.0f))),
+							InactiveLifetimeOverlay,
+							ESlateDrawEffect::None
+						);
+
+
+						// Black line on right
+						FSlateDrawElement::MakeBox(
+							OutDrawElements,
+							LayerId + 50,
+							AllottedGeometry.ToPaintGeometry(FVector2f(1.0f, AllottedGeometry.Size.Y), FSlateLayoutTransform(FVector2f(ExcludedRangeR, 0.f))),
+							FAppStyle::GetBrush("WhiteBrush"),
+							ESlateDrawEffect::None,
+							FLinearColor::Black
+						);
+
+					}
+				}
+			}
+			return SCompoundWidget::OnPaint(Args, AllottedGeometry, MyCullingRect, OutDrawElements, LayerId + 1, InWidgetStyle, bParentEnabled);
+		}
+
+	} // namespace Sequencer
+} // namespace UE
+
+#undef LOCTEXT_NAMESPACE
diff --git a/Engine/Source/Editor/Sequencer/Private/Widgets/SSequencerBindingLifetimeOverlay.h b/Engine/Source/Editor/Sequencer/Private/Widgets/SSequencerBindingLifetimeOverlay.h
new file mode 100644
--- /dev/null
+++ b/Engine/Source/Editor/Sequencer/Private/Widgets/SSequencerBindingLifetimeOverlay.h
@@ -0,0 +1,65 @@
+// Copyright Epic Games, Inc. All Rights Reserved.
+
+#pragma once
+
+#include "Math/Range.h"
+#include "Misc/FrameNumber.h"
+
+#include "Widgets/SCompoundWidget.h"
+#include "Widgets/DeclarativeSyntaxSupport.h"
+
+#include "MVVM/Extensions/ITrackLaneExtension.h"
+#include "MVVM/Extensions/ViewModelExtensionCollection.h"
+#include "MVVM/ViewModels/BindingLifetimeOverlayModel.h"
+
+class ISequencer;
+class FSequencer;
+class FScopedTransaction;
+class FSequencerDisplayNode;
+struct FTimeToPixel;
+
+namespace UE
+{
+	namespace Sequencer
+	{
+		class FObjectBindingModel;
+		class FSequencerEditorViewModel;
+		class FViewModel;
+		class STrackAreaView;
+
+		class SSequencerBindingLifetimeOverlay
+			: public SCompoundWidget
+			, public ITrackLaneWidget
+		{
+		public:
+			SLATE_BEGIN_ARGS(SSequencerBindingLifetimeOverlay) {}
+			SLATE_END_ARGS()
+
+			void Construct(const FArguments& InArgs, TWeakPtr<STrackAreaView> InWeakTrackArea, TWeakPtr<FSequencerEditorViewModel> InWeakEditor, TWeakPtr<FBindingLifetimeOverlayModel> InWeakBindingLifetimeOverlayModel);
+
+			TSharedPtr<FBindingLifetimeOverlayModel> GetBindingLifetimeOverlayModel() const;
+			TSharedPtr<FSequencer> GetSequencer() const;
+
+
+
+		private:
+
+			/*~ ITrackLaneWidget interface */
+			TSharedRef<const SWidget> AsWidget() const override;
+			FTrackLaneScreenAlignment GetAlignment(const FTimeToPixel& TimeToPixel, const FGeometry& InParentGeometry) const override;
+			int32 GetOverlapPriority() const override;
+
+			/*~ SCompoundWidget interface */
+			int32 OnPaint(const FPaintArgs& Args, const FGeometry& AllottedGeometry, const FSlateRect& MyCullingRect, FSlateWindowElementList& OutDrawElements, int32 LayerId, const FWidgetStyle& InWidgetStyle, bool bParentEnabled) const override;
+
+		private:
+
+			TWeakPtr<FBindingLifetimeOverlayModel> WeakBindingLifetimeOverlayModel;
+			TWeakPtr<FSequencerEditorViewModel> WeakEditor;
+			TWeakPtr<STrackAreaView> WeakTrackArea;
+			TSharedPtr<FTimeToPixel> TimeToPixel;
+		};
+
+	} // namespace Sequencer
+} // namespace UE
+
diff --git a/Engine/Source/Editor/Sequencer/Private/Widgets/SSequencerLayerBar.cpp b/Engine/Source/Editor/Sequencer/Private/Widgets/SSequencerLayerBar.cpp
--- a/Engine/Source/Editor/Sequencer/Private/Widgets/SSequencerLayerBar.cpp
+++ b/Engine/Source/Editor/Sequencer/Private/Widgets/SSequencerLayerBar.cpp
@@ -254,6 +254,26 @@ void SSequencerLayerBar::Construct(const FArguments& InArgs, TWeakPtr<STrackArea
 	];
 }
 
+FVector2D SSequencerLayerBar::ComputeDesiredSize(float LayoutScale) const
+{
+	TViewModelPtr<FLinkedOutlinerExtension> LayerBar = WeakLayerBar.Pin();
+	TViewModelPtr<IOutlinerExtension>  OutlinerItem = LayerBar ? LayerBar->GetLinkedOutlinerItem() : nullptr;
+	TViewModelPtr<ITrackAreaExtension> TrackArea = OutlinerItem.ImplicitCast();
+
+	float Height = 0.f;
+
+	if (TrackArea && OutlinerItem)
+	{
+		FTrackAreaParameters Parameters = TrackArea->GetTrackAreaParameters();
+
+		Height = OutlinerItem->GetOutlinerSizing().GetTotalHeight();
+
+		Height -= Parameters.TrackLanePadding.Top + Parameters.TrackLanePadding.Bottom;
+	}
+
+	return FVector2D(100.f, Height);
+}
+
 TSharedPtr<FLayerBarModel> SSequencerLayerBar::GetLayerBarModel() const
 {
 	return WeakLayerBar.Pin();
@@ -296,11 +316,17 @@ int32 SSequencerLayerBar::OnPaint(const FPaintArgs& Args, const FGeometry& Allot
 	static const FSlateBrush* LayerBarBackgroundBrush = FAppStyle::Get().GetBrush("Sequencer.LayerBar.Background");
 	static const FSlateBrush* SelectedSectionOverlay  = FAppStyle::Get().GetBrush("Sequencer.Section.CollapsedSelectedSectionOverlay");
 
+	TViewModelPtr<FLinkedOutlinerExtension> LayerBar = WeakLayerBar.Pin();
+	TViewModelPtr<IOutlinerExtension>  OutlinerItem = LayerBar ? LayerBar->GetLinkedOutlinerItem() : nullptr;
+	float TotalNodeHeight = OutlinerItem ? OutlinerItem->GetOutlinerSizing().GetTotalHeight() : 0.0f;
 	// Draw the background
 	FSlateDrawElement::MakeBox(
 		OutDrawElements,
 		LayerId,
-		AllottedGeometry.ToPaintGeometry(),
+		AllottedGeometry.ToPaintGeometry(
+			FVector2f(AllottedGeometry.GetLocalSize().X, TotalNodeHeight),
+			FSlateLayoutTransform()
+		),
 		LayerBarBackgroundBrush,
 		DrawEffects
 	);
@@ -333,7 +359,7 @@ int32 SSequencerLayerBar::OnPaint(const FPaintArgs& Args, const FGeometry& Allot
 		FSlateDrawElement::MakeBox(
 			OutDrawElements,
 			LayerId,
-			AllottedGeometry.ToPaintGeometry(AllottedGeometry.GetLocalSize() - FVector2f(2.f,3.f), FSlateLayoutTransform(FVector2f(1.f, 1.f))),
+			AllottedGeometry.ToPaintGeometry(FVector2f(AllottedGeometry.GetLocalSize().X - 2.f, TotalNodeHeight-3.f), FSlateLayoutTransform(FVector2f(1.f, 1.f))),
 			SelectedSectionOverlay,
 			DrawEffects,
 			SelectionColor.CopyWithNewOpacity(0.8f)
@@ -341,7 +367,7 @@ int32 SSequencerLayerBar::OnPaint(const FPaintArgs& Args, const FGeometry& Allot
 		++LayerId;
 	}
 
-	return SCompoundWidget::OnPaint(Args, AllottedGeometry, MyCullingRect, OutDrawElements, LayerId, InWidgetStyle, bParentEnabled);
+	return SCompoundWidget::OnPaint(Args, AllottedGeometry.MakeChild(FVector2f(AllottedGeometry.GetLocalSize().X, TotalNodeHeight), FSlateLayoutTransform()), MyCullingRect, OutDrawElements, LayerId, InWidgetStyle, bParentEnabled);
 }
 
 FReply SSequencerLayerBar::OnMouseButtonDown(const FGeometry& MyGeometry, const FPointerEvent& MouseEvent)
diff --git a/Engine/Source/Editor/Sequencer/Private/Widgets/SSequencerLayerBar.h b/Engine/Source/Editor/Sequencer/Private/Widgets/SSequencerLayerBar.h
--- a/Engine/Source/Editor/Sequencer/Private/Widgets/SSequencerLayerBar.h
+++ b/Engine/Source/Editor/Sequencer/Private/Widgets/SSequencerLayerBar.h
@@ -49,6 +49,7 @@ private:
 
 	/*~ SCompoundWidget interface */
 	int32 OnPaint( const FPaintArgs& Args, const FGeometry& AllottedGeometry, const FSlateRect& MyCullingRect, FSlateWindowElementList& OutDrawElements, int32 LayerId, const FWidgetStyle& InWidgetStyle, bool bParentEnabled ) const override;
+	FVector2D ComputeDesiredSize(float LayoutScale) const override;
 
 	FReply OnMouseButtonDown(const FGeometry& MyGeometry, const FPointerEvent& MouseEvent) override;
 	FReply OnMouseButtonUp(const FGeometry& MyGeometry, const FPointerEvent& MouseEvent) override;
diff --git a/Engine/Source/Editor/Sequencer/Public/MVVM/Extensions/IBindingLifetimeExtension.h b/Engine/Source/Editor/Sequencer/Public/MVVM/Extensions/IBindingLifetimeExtension.h
new file mode 100644
--- /dev/null
+++ b/Engine/Source/Editor/Sequencer/Public/MVVM/Extensions/IBindingLifetimeExtension.h
@@ -0,0 +1,27 @@
+// Copyright Epic Games, Inc. All Rights Reserved.
+
+#pragma once
+
+#include "MVVM/ViewModelTypeID.h"
+
+struct FGuid;
+
+namespace UE
+{
+	namespace Sequencer
+	{
+
+		class SEQUENCER_API IBindingLifetimeExtension
+		{
+		public:
+
+			UE_SEQUENCER_DECLARE_VIEW_MODEL_TYPE_ID(IBindingLifetimeExtension)
+
+			virtual ~IBindingLifetimeExtension() {}
+
+			virtual const TArray<TRange<FFrameNumber>>& GetInverseLifetimeRange() const = 0;
+		};
+
+	} // namespace Sequencer
+} // namespace UE
+
diff --git a/Engine/Source/Editor/Sequencer/Public/MVVM/ViewModels/BindingLifetimeOverlayModel.h b/Engine/Source/Editor/Sequencer/Public/MVVM/ViewModels/BindingLifetimeOverlayModel.h
new file mode 100644
--- /dev/null
+++ b/Engine/Source/Editor/Sequencer/Public/MVVM/ViewModels/BindingLifetimeOverlayModel.h
@@ -0,0 +1,51 @@
+// Copyright Epic Games, Inc. All Rights Reserved.
+
+#pragma once
+
+#include "MVVM/Extensions/ITrackLaneExtension.h"
+#include "MVVM/Extensions/IBindingLifetimeExtension.h"
+#include "MVVM/Extensions/LinkedOutlinerExtension.h"
+#include "MVVM/Extensions/ViewModelExtensionCollection.h"
+#include "MVVM/ICastable.h"
+#include "MVVM/ViewModels/ViewModel.h"
+#include "Math/Range.h"
+#include "Templates/SharedPointer.h"
+#include "SequencerEditorViewModel.h"
+
+namespace UE::Sequencer { template <typename T> struct TAutoRegisterViewModelTypeID; }
+struct FFrameNumber;
+
+namespace UE
+{
+	namespace Sequencer
+	{
+		class SEQUENCER_API FBindingLifetimeOverlayModel
+			: public FViewModel
+			, public FLinkedOutlinerExtension
+			, public ITrackLaneExtension
+		{
+		public:
+
+			UE_SEQUENCER_DECLARE_CASTABLE(FBindingLifetimeOverlayModel, FViewModel
+				, FLinkedOutlinerExtension
+				, ITrackLaneExtension
+			);
+
+			FBindingLifetimeOverlayModel(TWeakPtr<FViewModel> LayerRoot, TWeakPtr<FSequencerEditorViewModel> InEditorViewModel, TViewModelPtr<IBindingLifetimeExtension> InBindingLifetimeTrack);
+			~FBindingLifetimeOverlayModel();
+
+			/*~ ITrackLaneExtension Interface */
+			TSharedPtr<ITrackLaneWidget> CreateTrackLaneView(const FCreateTrackLaneViewParams& InParams) override;
+			FTrackLaneVirtualAlignment ArrangeVirtualTrackLaneView() const override;
+
+			const TArray<TRange<FFrameNumber>>& GetInverseLifetimeRange() const;
+
+		private:
+
+			TWeakPtr<FSequencerEditorViewModel> WeakEditorViewModel;
+			TViewModelPtr<IBindingLifetimeExtension> BindingLifetimeTrack;
+		};
+
+	} // namespace Sequencer
+} // namespace UE
+
diff --git a/Engine/Source/Editor/Sequencer/Public/MVVM/ViewModels/ObjectBindingModel.h b/Engine/Source/Editor/Sequencer/Public/MVVM/ViewModels/ObjectBindingModel.h
--- a/Engine/Source/Editor/Sequencer/Public/MVVM/ViewModels/ObjectBindingModel.h
+++ b/Engine/Source/Editor/Sequencer/Public/MVVM/ViewModels/ObjectBindingModel.h
@@ -6,6 +6,7 @@
 #include "MVVM/ViewModels/ViewModel.h"
 #include "MVVM/ViewModels/ViewModelHierarchy.h"
 #include "MVVM/ViewModels/OutlinerItemModel.h"
+#include  "MVVM/ViewModels/BindingLifetimeOverlayModel.h"
 #include "MVVM/Extensions/IRenameableExtension.h"
 #include "MVVM/Extensions/ITrackAreaExtension.h"
 #include "MVVM/Extensions/IGroupableExtension.h"
@@ -14,6 +15,7 @@
 #include "MVVM/Extensions/IDraggableOutlinerExtension.h"
 #include "MVVM/Extensions/IDeletableExtension.h"
 
+
 struct FMovieSceneBinding;
 
 class UMovieScene;
@@ -64,6 +66,8 @@ public:
 	FObjectBindingModel(FSequenceModel* OwnerModel, const FMovieSceneBinding& InBinding);
 	~FObjectBindingModel();
 
+	static EViewModelListType GetTopLevelChildTrackAreaGroupType();
+
 	void AddTrack(UMovieSceneTrack* Track);
 	void RemoveTrack(UMovieSceneTrack* Track);
 
@@ -85,6 +89,7 @@ public:
 	/*~ ITrackAreaExtension */
 	FTrackAreaParameters GetTrackAreaParameters() const override;
 	FViewModelVariantIterator GetTrackAreaModelList() const override;
+	FViewModelVariantIterator GetTopLevelChildTrackAreaModels() const override;
 
 	/*~ IGroupableExtension */
 	void GetIdentifierForGrouping(TStringBuilder<128>& OutString) const override;
@@ -148,7 +153,9 @@ protected:
 	FGuid ObjectBindingID;
 	FGuid ParentObjectBindingID;
 	FViewModelListHead TrackAreaList;
+	FViewModelListHead TopLevelChildTrackAreaList;
 	TSharedPtr<FLayerBarModel> LayerBar;
+	TSharedPtr<FBindingLifetimeOverlayModel> BindingLifetimeOverlayModel;
 	FSequenceModel* OwnerModel;
 };
 
diff --git a/Engine/Source/Editor/Sequencer/Public/MVVM/ViewModels/OutlinerItemModel.h b/Engine/Source/Editor/Sequencer/Public/MVVM/ViewModels/OutlinerItemModel.h
--- a/Engine/Source/Editor/Sequencer/Public/MVVM/ViewModels/OutlinerItemModel.h
+++ b/Engine/Source/Editor/Sequencer/Public/MVVM/ViewModels/OutlinerItemModel.h
@@ -49,6 +49,7 @@ public:
 	void SetExpansion(bool bInIsExpanded) override;
 	bool IsFilteredOut() const override;
 	TSharedPtr<SWidget> CreateContextMenuWidget(const FCreateOutlinerContextMenuWidgetParams& InParams) override;
+	FSlateColor GetLabelColor() const override;
 
 	/*~ ICurveEditorTreeItemExtension */
 	virtual bool HasCurves() const override;
diff --git a/Engine/Source/Editor/Sequencer/Public/MVVM/ViewModels/TrackModel.h b/Engine/Source/Editor/Sequencer/Public/MVVM/ViewModels/TrackModel.h
--- a/Engine/Source/Editor/Sequencer/Public/MVVM/ViewModels/TrackModel.h
+++ b/Engine/Source/Editor/Sequencer/Public/MVVM/ViewModels/TrackModel.h
@@ -85,7 +85,7 @@ public:
 	void OnModifiedIndirectly(UMovieSceneSignedObject*) override;
 
 	/* IDeferredSignedObjectFlushSignal */
-	void OnDeferredModifyFlush() override;
+	virtual void OnDeferredModifyFlush() override;
 
 	/*~ IOutlinerExtension */
 	FOutlinerSizing GetOutlinerSizing() const override;
@@ -140,10 +140,10 @@ public:
 	bool CanDelete(FText* OutErrorMessage) const override;
 	void Delete() override;
 
-private:
-
 	/*~ FViewModel interface */
-	void OnConstruct() override;
+	virtual void OnConstruct() override;
+
+private:
 
 	void ForceUpdate();
 
diff --git a/Engine/Source/Editor/SequencerCore/Private/MVVM/Extensions/IGeometryExtension.cpp b/Engine/Source/Editor/SequencerCore/Private/MVVM/Extensions/IGeometryExtension.cpp
--- a/Engine/Source/Editor/SequencerCore/Private/MVVM/Extensions/IGeometryExtension.cpp
+++ b/Engine/Source/Editor/SequencerCore/Private/MVVM/Extensions/IGeometryExtension.cpp
@@ -43,6 +43,12 @@ float IGeometryExtension::UpdateVirtualGeometry(float InitialVirtualPosition, TS
 			{
 				It->ReportVirtualGeometry(Geometry);
 			}
+
+			for (TTypedIterator<IGeometryExtension, FViewModelVariantIterator>
+				It(TrackAreaExtension->GetTopLevelChildTrackAreaModels()); It; ++It)
+			{
+				It->ReportVirtualGeometry(Geometry);
+			}
 		}
 	}
 	else
diff --git a/Engine/Source/Editor/SequencerCore/Private/MVVM/Views/STrackLane.cpp b/Engine/Source/Editor/SequencerCore/Private/MVVM/Views/STrackLane.cpp
--- a/Engine/Source/Editor/SequencerCore/Private/MVVM/Views/STrackLane.cpp
+++ b/Engine/Source/Editor/SequencerCore/Private/MVVM/Views/STrackLane.cpp
@@ -168,11 +168,8 @@ void STrackLane::RecreateWidgets()
 	ViewParams.OwningTrackLane = SharedThis(this);
 	ViewParams.TimeToPixel = TimeToPixel;
 
-	// Construct views for this track lane
-	for (TTypedIterator<ITrackLaneExtension, FViewModelVariantIterator> It(TrackAreaExtension->GetTrackAreaModelList()); It; ++It)
+	auto ConstructTrackLaneView = [&](TViewModelPtr<ITrackLaneExtension> Model)
 	{
-		TViewModelPtr<ITrackLaneExtension> Model = *It;
-
 		TSharedPtr<ITrackLaneWidget> ParentView;
 
 		if (ParentLane)
@@ -201,8 +198,18 @@ void STrackLane::RecreateWidgets()
 			SlotArguments.AttachWidget(NewView->AsWidget());
 			Children.AddSlot(MoveTemp(SlotArguments));
 		}
+	};
+
+	// Construct views for this track lane
+	for (TTypedIterator<ITrackLaneExtension, FViewModelVariantIterator> It(TrackAreaExtension->GetTrackAreaModelList()); It; ++It)
+	{
+		ConstructTrackLaneView(*It);
 	}
 
+	for (TTypedIterator<ITrackLaneExtension, FViewModelVariantIterator> It(TrackAreaExtension->GetTopLevelChildTrackAreaModels()); It; ++It)
+	{
+		ConstructTrackLaneView(*It);
+	}
 }
 
 void STrackLane::OnArrangeChildren( const FGeometry& AllottedGeometry, FArrangedChildren& ArrangedChildren ) const
diff --git a/Engine/Source/Runtime/LevelSequence/Private/LevelSequence.cpp b/Engine/Source/Runtime/LevelSequence/Private/LevelSequence.cpp
--- a/Engine/Source/Runtime/LevelSequence/Private/LevelSequence.cpp
+++ b/Engine/Source/Runtime/LevelSequence/Private/LevelSequence.cpp
@@ -32,6 +32,7 @@
 #include "Tracks/MovieSceneSpawnTrack.h"
 #include "Tracks/MovieSceneSubTrack.h"
 #include "Tracks/MovieSceneCVarTrack.h"
+#include "Tracks/MovieSceneBindingLifetimeTrack.h"
 #include "Modules/ModuleManager.h"
 #include "LevelSequencePlayer.h"
 #include "Compilation/MovieSceneCompiledDataManager.h"
@@ -141,7 +142,8 @@ ETrackSupport ULevelSequence::IsTrackSupported(TSubclassOf<class UMovieSceneTrac
 		InTrackClass == UMovieSceneSlomoTrack::StaticClass() ||
 		InTrackClass == UMovieSceneSpawnTrack::StaticClass() ||
 		InTrackClass == UMovieSceneSubTrack::StaticClass() ||
-		InTrackClass == UMovieSceneCVarTrack::StaticClass())
+		InTrackClass == UMovieSceneCVarTrack::StaticClass() ||
+		InTrackClass == UMovieSceneBindingLifetimeTrack::StaticClass())
 	{
 		return ETrackSupport::Supported;
 	}
diff --git a/Engine/Source/Runtime/MovieScene/Private/EntitySystem/BuiltInComponentTypes.cpp b/Engine/Source/Runtime/MovieScene/Private/EntitySystem/BuiltInComponentTypes.cpp
--- a/Engine/Source/Runtime/MovieScene/Private/EntitySystem/BuiltInComponentTypes.cpp
+++ b/Engine/Source/Runtime/MovieScene/Private/EntitySystem/BuiltInComponentTypes.cpp
@@ -386,6 +386,8 @@ FBuiltInComponentTypes::FBuiltInComponentTypes()
 	ComponentRegistry->NewComponentType(&Interrogation.Instance,  TEXT("Interrogation Instance"));
 	ComponentRegistry->NewComponentType(&Interrogation.OutputKey, TEXT("Interrogation Output"));
 
+	ComponentRegistry->NewComponentType(&BindingLifetime, TEXT("Binding Lifetime"));
+
 	Tags.RestoreState            = ComponentRegistry->NewTag(TEXT("Is Restore State Entity"));
 	Tags.IgnoreHierarchicalBias  = ComponentRegistry->NewTag(TEXT("Ignore Hierarchical Bias"));
 	Tags.BlendHierarchicalBias   = ComponentRegistry->NewTag(TEXT("Blend Hierarchical Bias"));
diff --git a/Engine/Source/Runtime/MovieScene/Private/EntitySystem/MovieSceneBindingLifetimeSystem.cpp b/Engine/Source/Runtime/MovieScene/Private/EntitySystem/MovieSceneBindingLifetimeSystem.cpp
new file mode 100644
--- /dev/null
+++ b/Engine/Source/Runtime/MovieScene/Private/EntitySystem/MovieSceneBindingLifetimeSystem.cpp
@@ -0,0 +1,105 @@
+// Copyright Epic Games, Inc. All Rights Reserved.
+
+#include "EntitySystem/MovieSceneBindingLifetimeSystem.h"
+#include "EntitySystem/MovieSceneEntitySystemLinker.h"
+#include "EntitySystem/MovieSceneEntitySystemRunner.h"
+#include "EntitySystem/MovieSceneEntitySystemTask.h"
+#include "EntitySystem/BuiltInComponentTypes.h"
+
+#include "Evaluation/MovieSceneEvaluationTemplateInstance.h"
+#include "IMovieScenePlayer.h"
+#include "MovieSceneSequence.h"
+
+#include "EntitySystem/MovieSceneEntitySystemTypes.h"
+#include "EntitySystem/MovieSceneEntityRange.h"
+#include "EntitySystem/MovieSceneEntitySystemTask.h"
+#include "EntitySystem/MovieSceneEntityInstantiatorSystem.h"
+#include "Evaluation/MovieSceneEvaluationOperand.h"
+
+#include "MovieScene.h"
+#include "MovieSceneExecutionToken.h"
+#include "IMovieScenePlaybackClient.h"
+#include "EntitySystem/MovieSceneSpawnablesSystem.h"
+
+#include UE_INLINE_GENERATED_CPP_BY_NAME(MovieSceneBindingLifetimeSystem)
+#define LOCTEXT_NAMESPACE "MovieSceneBindingLifetimeSystem"
+
+namespace UE
+{
+	namespace MovieScene
+	{
+		const FMovieSceneAnimTypeID BindingLifetimeAnimTypeID = FMovieSceneAnimTypeID::Unique();
+	}
+}
+
+
+UMovieSceneBindingLifetimeSystem::UMovieSceneBindingLifetimeSystem(const FObjectInitializer& ObjInit)
+	: Super(ObjInit)
+{
+	using namespace UE::MovieScene;
+
+	Phase = UE::MovieScene::ESystemPhase::Spawn;
+	FBuiltInComponentTypes* BuiltInComponentTypes = FBuiltInComponentTypes::Get();
+	RelevantComponent = BuiltInComponentTypes->BindingLifetime;
+
+	if (HasAnyFlags(RF_ClassDefaultObject))
+	{
+		DefineImplicitPrerequisite(GetClass(), UMovieSceneSpawnablesSystem::StaticClass());
+	}
+}
+
+FMovieSceneAnimTypeID UMovieSceneBindingLifetimeSystem::GetAnimTypeID()
+{
+	return UE::MovieScene::BindingLifetimeAnimTypeID;
+}
+
+void UMovieSceneBindingLifetimeSystem::OnRun(FSystemTaskPrerequisites& InPrerequisites, FSystemSubsequentTasks& Subsequents)
+{
+	using namespace UE::MovieScene;
+	FBuiltInComponentTypes* BuiltInComponents = FBuiltInComponentTypes::Get();
+	if (!Linker->EntityManager.Contains(FEntityComponentFilter().Any({ BuiltInComponents->Tags.NeedsLink, BuiltInComponents->Tags.NeedsUnlink })))
+	{
+		return;
+	}
+
+	FInstanceRegistry* InstanceRegistry = Linker->GetInstanceRegistry();
+
+	bool bLink = false;
+	auto SetBindingActivation = [InstanceRegistry, &bLink](FMovieSceneEntityID EntityID, FInstanceHandle InstanceHandle, const FMovieSceneBindingLifetimeComponentData& BindingLifetime)
+	{
+		const FSequenceInstance& SequenceInstance = InstanceRegistry->GetInstance(InstanceHandle);
+
+		FMovieSceneSequenceID SequenceID = SequenceInstance.GetSequenceID();
+		IMovieScenePlayer* Player = SequenceInstance.GetPlayer();
+		if (Player)
+		{
+			// For now we use the linking/unlinking of the inactive ranges to set the binding activations
+			if (BindingLifetime.BindingLifetimeState == EMovieSceneBindingLifetimeState::InActive)
+			{
+				Player->State.SetBindingActivation(BindingLifetime.BindingGuid, SequenceID, !bLink);
+			}
+		}
+	};
+
+	FBuiltInComponentTypes* BuiltInComponentTypes = FBuiltInComponentTypes::Get();
+
+	// Unlink stale bindinglifetime entities
+	FEntityTaskBuilder()
+		.ReadEntityIDs()
+		.Read(BuiltInComponents->InstanceHandle)
+		.Read(BuiltInComponentTypes->BindingLifetime)
+		.FilterAll({ BuiltInComponents->Tags.NeedsUnlink })
+		.Iterate_PerEntity(&Linker->EntityManager, SetBindingActivation);
+
+	// Link new bindinglifetime entities
+	bLink = true;
+	FEntityTaskBuilder()
+		.ReadEntityIDs()
+		.Read(BuiltInComponents->InstanceHandle)
+		.Read(BuiltInComponentTypes->BindingLifetime)
+		.FilterAll({ BuiltInComponents->Tags.NeedsLink })
+		.Iterate_PerEntity(&Linker->EntityManager, SetBindingActivation);
+}
+
+#undef LOCTEXT_NAMESPACE
+
diff --git a/Engine/Source/Runtime/MovieScene/Private/EntitySystem/MovieSceneSpawnablesSystem.cpp b/Engine/Source/Runtime/MovieScene/Private/EntitySystem/MovieSceneSpawnablesSystem.cpp
--- a/Engine/Source/Runtime/MovieScene/Private/EntitySystem/MovieSceneSpawnablesSystem.cpp
+++ b/Engine/Source/Runtime/MovieScene/Private/EntitySystem/MovieSceneSpawnablesSystem.cpp
@@ -166,6 +166,18 @@ void UMovieSceneSpawnablesSystem::OnRun(FSystemTaskPrerequisites& InPrerequisite
 
 		UObject* ExistingSpawnedObject = Player->GetSpawnRegister().FindSpawnedObject(SpawnableBindingID, SequenceID).Get();
 
+		if (!Player->State.GetBindingActivation(SpawnableBindingID, SequenceID))
+		{
+			// If the binding is currently inactive, don't spawn the object.
+
+			// If we have an existing spawned object, then we need to destroy the spawned object here.
+			if (ExistingSpawnedObject)
+			{
+				DestroyOldSpawnables(InstanceHandle, SpawnableBindingID);
+			}
+			return;
+		}
+
 		// Check whether the binding is overridden - if it is we cannot spawn a new object
 		if (const IMovieScenePlaybackClient* PlaybackClient = Player->GetPlaybackClient())
 		{
diff --git a/Engine/Source/Runtime/MovieScene/Private/Evaluation/MovieSceneEvaluationState.cpp b/Engine/Source/Runtime/MovieScene/Private/Evaluation/MovieSceneEvaluationState.cpp
--- a/Engine/Source/Runtime/MovieScene/Private/Evaluation/MovieSceneEvaluationState.cpp
+++ b/Engine/Source/Runtime/MovieScene/Private/Evaluation/MovieSceneEvaluationState.cpp
@@ -207,6 +207,24 @@ void FMovieSceneObjectCache::InvalidateIfValid(const FGuid& InGuid)
 	}
 }
 
+bool FMovieSceneObjectCache::GetBindingActivation(const FGuid& InGuid) const
+{
+	return !InactiveBindingIds.Contains(InGuid);
+}
+
+void FMovieSceneObjectCache::SetBindingActivation(const FGuid& InGuid, bool bActive)
+{
+	if (bActive)
+	{
+		InactiveBindingIds.Remove(InGuid);
+	}
+	else
+	{
+		InactiveBindingIds.Add(InGuid);
+	}
+	InvalidateInternal(InGuid);
+}
+
 bool FMovieSceneObjectCache::InvalidateIfValidInternal(const FGuid& InGuid)
 {
 	// Don't manipulate the actual map structure, since this can be called from inside an iterator
@@ -332,6 +350,12 @@ void FMovieSceneObjectCache::UpdateBindings(const FGuid& InGuid, IMovieScenePlay
 		return;
 	}
 
+	// Binding is inactive, do not resolve.
+	if (InactiveBindingIds.Contains(InGuid))
+	{
+		return;
+	}
+
 	// If we have overrides for this binding, ask the player to find it for us (most probably in a different cache
 	// for a different sequence).
 	// TODO-lchabant: we could technically end up in a circular override that creates an infinite loop...
@@ -531,6 +555,21 @@ void FMovieSceneEvaluationState::Invalidate(const FGuid& InGuid, FMovieSceneSequ
 	}
 }
 
+
+bool FMovieSceneEvaluationState::GetBindingActivation(const FGuid& InGuid, FMovieSceneSequenceIDRef InSequenceID) const
+{
+	if (const FMovieSceneObjectCache* Cache = FindObjectCache(InSequenceID))
+	{
+		return Cache->GetBindingActivation(InGuid);
+	}
+	return true;
+}
+
+void FMovieSceneEvaluationState::SetBindingActivation(const FGuid& InGuid, FMovieSceneSequenceIDRef InSequenceID, bool bActive)
+{
+	GetObjectCache(InSequenceID).SetBindingActivation(InGuid, bActive);
+}
+
 void FMovieSceneEvaluationState::ClearObjectCaches(IMovieScenePlayer& Player)
 {
 	for (auto& Pair : ObjectCaches)
diff --git a/Engine/Source/Runtime/MovieScene/Private/Sections/MovieSceneBindingLifetimeSection.cpp b/Engine/Source/Runtime/MovieScene/Private/Sections/MovieSceneBindingLifetimeSection.cpp
new file mode 100644
--- /dev/null
+++ b/Engine/Source/Runtime/MovieScene/Private/Sections/MovieSceneBindingLifetimeSection.cpp
@@ -0,0 +1,48 @@
+// Copyright Epic Games, Inc. All Rights Reserved.
+
+#include "Sections/MovieSceneBindingLifetimeSection.h"
+#include "Tracks/MovieSceneBindingLifetimeTrack.h"
+
+#include "EntitySystem/MovieSceneEntityBuilder.h"
+#include "EntitySystem/MovieSceneEntityManager.h"
+#include "EntitySystem/BuiltInComponentTypes.h"
+#include "EntitySystem/MovieSceneEntitySystemLinker.h"
+#include "Evaluation/MovieSceneEvaluationField.h"
+
+#include UE_INLINE_GENERATED_CPP_BY_NAME(MovieSceneBindingLifetimeSection)
+
+UMovieSceneBindingLifetimeSection::UMovieSceneBindingLifetimeSection(const FObjectInitializer& ObjInit)
+	: Super(ObjInit)
+{
+	bSupportsInfiniteRange = true;
+}
+
+void UMovieSceneBindingLifetimeSection::ImportEntityImpl(UMovieSceneEntitySystemLinker* EntityLinker, const FEntityImportParams& Params, FImportedEntity* OutImportedEntity)
+{
+	using namespace UE::MovieScene;
+
+	FBuiltInComponentTypes* BuiltInComponentTypes = FBuiltInComponentTypes::Get();
+
+	OutImportedEntity->AddBuilder(
+		FEntityBuilder()
+		.AddConditional(BuiltInComponentTypes->BindingLifetime, FMovieSceneBindingLifetimeComponentData{ Params.GetObjectBindingID(), EMovieSceneBindingLifetimeState::Active }, Params.GetObjectBindingID().IsValid())
+	);
+}
+
+bool UMovieSceneBindingLifetimeSection::PopulateEvaluationFieldImpl(const TRange<FFrameNumber>& EffectiveRange, const FMovieSceneEvaluationFieldEntityMetaData& InMetaData, FMovieSceneEntityComponentFieldBuilder* OutFieldBuilder)
+{
+	// By default, Binding Lifetime sections do not populate any evaluation field entries
+	// that is the job of its outer UMovieSceneTrack through a call to ExternalPopulateEvaluationField
+	return true;
+}
+
+void UMovieSceneBindingLifetimeSection::ExternalPopulateEvaluationField(const TRange<FFrameNumber>& EffectiveRange, const FMovieSceneEvaluationFieldEntityMetaData& InMetaData, FMovieSceneEntityComponentFieldBuilder* OutFieldBuilder)
+{
+	using namespace UE::MovieScene;
+
+	const int32 MetaDataIndex = OutFieldBuilder->AddMetaData(InMetaData);
+
+	const int32 EntityIndex = OutFieldBuilder->FindOrAddEntity(this, 1);
+	OutFieldBuilder->AddPersistentEntity(EffectiveRange, EntityIndex, MetaDataIndex);
+}
+
diff --git a/Engine/Source/Runtime/MovieScene/Private/Tracks/MovieSceneBindingLifetimeTrack.cpp b/Engine/Source/Runtime/MovieScene/Private/Tracks/MovieSceneBindingLifetimeTrack.cpp
new file mode 100644
--- /dev/null
+++ b/Engine/Source/Runtime/MovieScene/Private/Tracks/MovieSceneBindingLifetimeTrack.cpp
@@ -0,0 +1,176 @@
+// Copyright Epic Games, Inc. All Rights Reserved.
+
+#include "Tracks/MovieSceneBindingLifetimeTrack.h"
+#include "MovieSceneCommonHelpers.h"
+#include "Sections/MovieSceneBindingLifetimeSection.h"
+#include "Serialization/ObjectReader.h"
+#include "Serialization/ObjectWriter.h"
+#include "MovieScene.h"
+#include "MovieSceneTimeHelpers.h"
+#include "EntitySystem/BuiltInComponentTypes.h"
+
+#include UE_INLINE_GENERATED_CPP_BY_NAME(MovieSceneBindingLifetimeTrack)
+
+#define LOCTEXT_NAMESPACE "MovieSceneBindingLifetimeTrack"
+
+
+/* UMovieSceneBindingLifetimeTrack structors
+ *****************************************************************************/
+
+UMovieSceneBindingLifetimeTrack::UMovieSceneBindingLifetimeTrack(const FObjectInitializer& Obj)
+	: Super(Obj)
+{
+#if WITH_EDITORONLY_DATA
+	TrackTint = FColor(26, 117, 49, 150);
+#endif
+}
+
+
+bool UMovieSceneBindingLifetimeTrack::SupportsType(TSubclassOf<UMovieSceneSection> SectionClass) const
+{
+	return SectionClass == UMovieSceneBindingLifetimeSection::StaticClass();
+}
+
+UMovieSceneSection* UMovieSceneBindingLifetimeTrack::CreateNewSection()
+{
+	return NewObject<UMovieSceneBindingLifetimeSection>(this, NAME_None, RF_Transactional);
+}
+
+
+bool UMovieSceneBindingLifetimeTrack::HasSection(const UMovieSceneSection& Section) const
+{
+	return Sections.ContainsByPredicate([&](const UMovieSceneSection* In) { return In == &Section; });
+}
+
+
+void UMovieSceneBindingLifetimeTrack::AddSection(UMovieSceneSection& Section)
+{
+	Sections.Add(&Section);
+}
+
+
+void UMovieSceneBindingLifetimeTrack::RemoveSection(UMovieSceneSection& Section)
+{
+	Sections.RemoveAll([&](const UMovieSceneSection* In) { return In == &Section; });
+}
+
+void UMovieSceneBindingLifetimeTrack::RemoveSectionAt(int32 SectionIndex)
+{
+	Sections.RemoveAt(SectionIndex);
+}
+
+void UMovieSceneBindingLifetimeTrack::RemoveAllAnimationData()
+{
+	Sections.Empty();
+}
+
+
+bool UMovieSceneBindingLifetimeTrack::IsEmpty() const
+{
+	return Sections.Num() == 0;
+}
+
+
+const TArray<UMovieSceneSection*>& UMovieSceneBindingLifetimeTrack::GetAllSections() const
+{
+	return Sections;
+}
+
+#if WITH_EDITORONLY_DATA
+
+FText UMovieSceneBindingLifetimeTrack::GetDisplayName() const
+{
+	return LOCTEXT("TrackName", "Binding Lifetime");
+}
+
+#endif
+
+void UMovieSceneBindingLifetimeTrack::ImportEntityImpl(UMovieSceneEntitySystemLinker* EntityLinker, const FEntityImportParams& Params, FImportedEntity* OutImportedEntity)
+{
+	using namespace UE::MovieScene;
+
+	FBuiltInComponentTypes* BuiltInComponentTypes = FBuiltInComponentTypes::Get();
+
+		OutImportedEntity->AddBuilder(
+			FEntityBuilder()
+			.AddConditional(BuiltInComponentTypes->BindingLifetime, FMovieSceneBindingLifetimeComponentData{ Params.GetObjectBindingID(), EMovieSceneBindingLifetimeState::InActive }, Params.GetObjectBindingID().IsValid())
+		);
+}
+
+
+TArray<FFrameNumberRange> CalculateInverseLifetimeRange(const FMovieSceneTrackEvaluationField& LocalEvaluationField)
+{
+	TArray<FFrameNumberRange> InverseLifetimeRange;
+	InverseLifetimeRange.Add(FFrameNumberRange::All());
+	// Iterate through our sections, removing them from the range
+
+	auto RemoveRangeFromSet = [&](FFrameNumberRange SectionRange) {
+		for (int32 Index = 0; Index < InverseLifetimeRange.Num(); ++Index)
+		{
+			const FFrameNumberRange& Current = InverseLifetimeRange[Index];
+			if (Current.Overlaps(SectionRange))
+			{
+				// Special case that difference doesn't handle well
+				if (!SectionRange.HasLowerBound() && !SectionRange.HasUpperBound())
+				{
+					InverseLifetimeRange.Reset();
+					return;
+				}
+
+				TArray<FFrameNumberRange> SplitRanges = FFrameNumberRange::Difference(Current, SectionRange);
+				for (const FFrameNumberRange& NewRange : SplitRanges)
+				{
+					// Splitting infinite ranges keeps an infinite range, which we don't want to keep
+					if (!NewRange.HasLowerBound() && !NewRange.HasUpperBound())
+					{
+						continue;
+					}
+
+					InverseLifetimeRange.Add(NewRange);
+				}
+				InverseLifetimeRange.RemoveAtSwap(Index--);
+			}
+		}
+	};
+
+	for (const FMovieSceneTrackEvaluationFieldEntry& Entry : LocalEvaluationField.Entries)
+	{
+		RemoveRangeFromSet(Entry.Range);
+	}
+	return InverseLifetimeRange;
+}
+
+bool UMovieSceneBindingLifetimeTrack::PopulateEvaluationFieldImpl(const TRange<FFrameNumber>& EffectiveRange, const FMovieSceneEvaluationFieldEntityMetaData& InMetaData, FMovieSceneEntityComponentFieldBuilder* OutFieldBuilder)
+{
+	const FMovieSceneTrackEvaluationField& LocalEvaluationField = GetEvaluationField();
+
+	TArray<FFrameNumberRange> InverseLifetimeRange = CalculateInverseLifetimeRange(LocalEvaluationField);
+	// Add an entity for each section of the inverse range to define areas where our lifetime is inactive.
+	for (const FFrameNumberRange& InverseRange : InverseLifetimeRange)
+	{
+		OutFieldBuilder->AddPersistentEntity(InverseRange, this, 0, OutFieldBuilder->AddMetaData(InMetaData));
+	}
+
+	for (const FMovieSceneTrackEvaluationFieldEntry& Entry : LocalEvaluationField.Entries)
+	{
+		UMovieSceneBindingLifetimeSection* BindingLifetimeSection = Cast<UMovieSceneBindingLifetimeSection>(Entry.Section);
+		if (BindingLifetimeSection)
+		{
+			FFrameNumberRange SectionEffectiveRange = FFrameNumberRange::Intersection(EffectiveRange, Entry.Range);
+			if (!SectionEffectiveRange.IsEmpty())
+			{
+				FMovieSceneEvaluationFieldEntityMetaData SectionMetaData = InMetaData;
+				SectionMetaData.Flags = Entry.Flags;
+				BindingLifetimeSection->ExternalPopulateEvaluationField(SectionEffectiveRange, SectionMetaData, OutFieldBuilder);
+			}
+		}
+	}
+
+	return true;
+}
+
+
+#undef LOCTEXT_NAMESPACE
+
+
+
diff --git a/Engine/Source/Runtime/MovieScene/Public/EntitySystem/BuiltInComponentTypes.h b/Engine/Source/Runtime/MovieScene/Public/EntitySystem/BuiltInComponentTypes.h
--- a/Engine/Source/Runtime/MovieScene/Public/EntitySystem/BuiltInComponentTypes.h
+++ b/Engine/Source/Runtime/MovieScene/Public/EntitySystem/BuiltInComponentTypes.h
@@ -327,6 +327,21 @@ private:
 };
 
 
+enum class EMovieSceneBindingLifetimeState : uint8
+{
+	/** Object Binding State is Active. Tracks on this binding will be active. */
+	Active,
+	/** Object Binding is Inactive. Tracks on this binding will be inactive and states will be restored/kept depending on settings. */
+	InActive
+};
+
+struct FMovieSceneBindingLifetimeComponentData
+{
+	FGuid BindingGuid;
+
+	EMovieSceneBindingLifetimeState BindingLifetimeState = EMovieSceneBindingLifetimeState::Active;
+};
+
 /**
  * Specifies a unique, sorted path of hbiases that contribute to a blended output
  * Supports up to 8 unique HBiases in its path
@@ -601,6 +616,9 @@ public:
 	// An FGuid relating to a spawnable binding in a sequence
 	TComponentTypeID<FGuid> SpawnableBinding;
 
+	// Data relating to the lifetime of bindings
+	TComponentTypeID<FMovieSceneBindingLifetimeComponentData> BindingLifetime;
+
 public:
 
 	// A boolean repesenting the output of a bool property track or channel
diff --git a/Engine/Source/Runtime/MovieScene/Public/EntitySystem/MovieSceneBindingLifetimeSystem.h b/Engine/Source/Runtime/MovieScene/Public/EntitySystem/MovieSceneBindingLifetimeSystem.h
new file mode 100644
--- /dev/null
+++ b/Engine/Source/Runtime/MovieScene/Public/EntitySystem/MovieSceneBindingLifetimeSystem.h
@@ -0,0 +1,29 @@
+// Copyright Epic Games, Inc. All Rights Reserved.
+
+#pragma once
+
+#include "EntitySystem/MovieSceneEntitySystem.h"
+#include "EntitySystem/MovieSceneCachedEntityFilterResult.h"
+#include "MovieSceneBindingLifetimeSystem.generated.h"
+
+class UObject;
+struct FMovieSceneAnimTypeID;
+
+/**
+ * Systems that (optionally) controls the lifetime of bindings by communicating with the MovieSceneObjectCache.
+ */
+
+UCLASS(MinimalAPI)
+class UMovieSceneBindingLifetimeSystem : public UMovieSceneEntitySystem
+{
+public:
+	GENERATED_BODY()
+
+	MOVIESCENE_API UMovieSceneBindingLifetimeSystem(const FObjectInitializer& ObjInit);
+
+	MOVIESCENE_API static FMovieSceneAnimTypeID GetAnimTypeID();
+
+private:
+
+	MOVIESCENE_API virtual void OnRun(FSystemTaskPrerequisites& InPrerequisites, FSystemSubsequentTasks& Subsequents) override;
+};
diff --git a/Engine/Source/Runtime/MovieScene/Public/Evaluation/MovieSceneEvaluationState.h b/Engine/Source/Runtime/MovieScene/Public/Evaluation/MovieSceneEvaluationState.h
--- a/Engine/Source/Runtime/MovieScene/Public/Evaluation/MovieSceneEvaluationState.h
+++ b/Engine/Source/Runtime/MovieScene/Public/Evaluation/MovieSceneEvaluationState.h
@@ -118,6 +118,18 @@ struct FMovieSceneObjectCache
 	 */
 	MOVIESCENE_API void InvalidateIfValid(const FGuid& InGuid);
 
+
+	/* Gets the current activation on the provided binding. If no binding lifetime track is present, true will be returned.
+	 * @param InGuid			The object binding ID
+	 */
+	MOVIESCENE_API bool GetBindingActivation(const FGuid& InGuid) const;
+
+	/* Sets the binding to either active or inactive. Inactive bindings will invalidate and not resolve while inactive.
+	 * @param InGuid			The object binding ID
+	 * @param bActive		    Where to activate or deactivate the binding.
+	 */
+	MOVIESCENE_API void SetBindingActivation(const FGuid& InGuid, bool bActive);
+
 	/**
 	 * Completely erase all knowledge of, anc caches for all object bindings
 	 */
@@ -222,6 +234,9 @@ private:
 	  */
 	TMap<FMovieSceneObjectBindingID, FGuidArray, FDefaultSetAllocator> ReverseMappedBindings;
 
+	/* A set of inactive binding ids based on Binding Lifetime track. While inactive, these will be prevented from resolving.*/
+	TSet<FGuid> InactiveBindingIds;
+
 	/** Serial number for this cache */
 	uint32 SerialNumber = 0;
 
@@ -298,6 +313,20 @@ struct FMovieSceneEvaluationState
 		return nullptr;
 	}
 
+	/**
+	 * Find an object cache pertaining to the specified sequence
+	 *
+	 * @param InSequenceID		The sequence ID to lookup
+	 */
+	FORCEINLINE const FMovieSceneObjectCache* FindObjectCache(FMovieSceneSequenceIDRef SequenceID) const
+	{
+		if (const FVersionedObjectCache* Cache = ObjectCaches.Find(SequenceID))
+		{
+			return &Cache->ObjectCache;
+		}
+		return nullptr;
+	}
+
 	/**
 	 * Get an object cache pertaining to the specified sequence
 	 *
@@ -326,6 +355,20 @@ struct FMovieSceneEvaluationState
 	 */
 	MOVIESCENE_API void Invalidate(const FGuid& InGuid, FMovieSceneSequenceIDRef InSequenceID);
 
+
+	/* Gets the current activation on the provided binding. If no binding lifetime track is present, true will be returned.
+	 * @param InGuid			The object binding ID
+	 * @param InSequenceID		The sequence ID to which the object binding belongs
+	 */
+	MOVIESCENE_API bool GetBindingActivation(const FGuid& InGuid, FMovieSceneSequenceIDRef InSequenceID) const;
+
+	/* Sets the binding to either active or inactive. Inactive bindings will invalidate and not resolve while inactive. 
+	 * @param InGuid			The object binding ID
+	 * @param InSequenceID		The sequence ID to which the object binding belongs
+	 * @param bActive		    Where to activate or deactivate the binding.
+	 */
+	MOVIESCENE_API void SetBindingActivation(const FGuid& InGuid, FMovieSceneSequenceIDRef InSequenceID, bool bActive);
+
 	/**
 	 * Forcably clear all object caches
 	 */
diff --git a/Engine/Source/Runtime/MovieScene/Public/Sections/MovieSceneBindingLifetimeSection.h b/Engine/Source/Runtime/MovieScene/Public/Sections/MovieSceneBindingLifetimeSection.h
new file mode 100644
--- /dev/null
+++ b/Engine/Source/Runtime/MovieScene/Public/Sections/MovieSceneBindingLifetimeSection.h
@@ -0,0 +1,28 @@
+// Copyright Epic Games, Inc. All Rights Reserved.
+
+#pragma once
+
+#include "MovieSceneSection.h"
+#include "EntitySystem/IMovieSceneEntityProvider.h"
+#include "MovieSceneBindingLifetimeSection.generated.h"
+
+
+/**
+ * Binding lifetime section that will connect an object binding while active, and disconnect it afterwards.
+ */
+UCLASS(MinimalAPI)
+class UMovieSceneBindingLifetimeSection
+	: public UMovieSceneSection
+	, public IMovieSceneEntityProvider
+{
+public:
+	GENERATED_BODY()
+
+	UMovieSceneBindingLifetimeSection(const FObjectInitializer& ObjInit);
+
+	virtual bool PopulateEvaluationFieldImpl(const TRange<FFrameNumber>& EffectiveRange, const FMovieSceneEvaluationFieldEntityMetaData& InMetaData, FMovieSceneEntityComponentFieldBuilder* OutFieldBuilder) override;
+	virtual void ImportEntityImpl(UMovieSceneEntitySystemLinker* EntityLinker, const FEntityImportParams& Params, FImportedEntity* OutImportedEntity) override;
+	void ExternalPopulateEvaluationField(const TRange<FFrameNumber>& EffectiveRange, const FMovieSceneEvaluationFieldEntityMetaData& InMetaData, FMovieSceneEntityComponentFieldBuilder* OutFieldBuilder);
+};
+
+
diff --git a/Engine/Source/Runtime/MovieScene/Public/Tracks/MovieSceneBindingLifetimeTrack.h b/Engine/Source/Runtime/MovieScene/Public/Tracks/MovieSceneBindingLifetimeTrack.h
new file mode 100644
--- /dev/null
+++ b/Engine/Source/Runtime/MovieScene/Public/Tracks/MovieSceneBindingLifetimeTrack.h
@@ -0,0 +1,66 @@
+// Copyright Epic Games, Inc. All Rights Reserved.
+
+#pragma once
+
+#include "Containers/Array.h"
+#include "Containers/ContainerAllocationPolicies.h"
+#include "CoreMinimal.h"
+#include "CoreTypes.h"
+#include "Curves/IntegralCurve.h"
+#include "Internationalization/Text.h"
+#include "Misc/Guid.h"
+#include "MovieSceneSection.h"
+#include "MovieSceneTrack.h"
+#include "Templates/SubclassOf.h"
+#include "UObject/ObjectMacros.h"
+#include "UObject/ObjectPtr.h"
+#include "UObject/UObjectGlobals.h"
+#include "EntitySystem/IMovieSceneEntityProvider.h"
+
+#include "MovieSceneBindingLifetimeTrack.generated.h"
+
+/**
+ * Handles when an object binding should be activated/deactivated
+ */
+UCLASS(MinimalAPI)
+class UMovieSceneBindingLifetimeTrack
+	: public UMovieSceneTrack
+	, public IMovieSceneEntityProvider
+{
+public:
+
+	MOVIESCENE_API UMovieSceneBindingLifetimeTrack(const FObjectInitializer& Obj);
+	GENERATED_BODY()
+
+public:
+
+public:
+
+	// UMovieSceneTrack interface
+	MOVIESCENE_API virtual bool SupportsType(TSubclassOf<UMovieSceneSection> SectionClass) const override;
+	MOVIESCENE_API virtual UMovieSceneSection* CreateNewSection() override;
+	MOVIESCENE_API virtual bool HasSection(const UMovieSceneSection& Section) const override;
+	MOVIESCENE_API virtual void AddSection(UMovieSceneSection& Section) override;
+	MOVIESCENE_API virtual void RemoveSection(UMovieSceneSection& Section) override;
+	MOVIESCENE_API virtual void RemoveSectionAt(int32 SectionIndex) override;
+	MOVIESCENE_API virtual void RemoveAllAnimationData() override;
+	MOVIESCENE_API virtual bool IsEmpty() const override;
+	MOVIESCENE_API virtual const TArray<UMovieSceneSection*>& GetAllSections() const override;	
+	
+	/*~ IMovieSceneEntityProvider */
+	virtual void ImportEntityImpl(UMovieSceneEntitySystemLinker* EntityLinker, const FEntityImportParams& Params, FImportedEntity* OutImportedEntity) override;
+	virtual bool PopulateEvaluationFieldImpl(const TRange<FFrameNumber>& EffectiveRange, const FMovieSceneEvaluationFieldEntityMetaData& InMetaData, FMovieSceneEntityComponentFieldBuilder* OutFieldBuilder) override;
+
+
+#if WITH_EDITORONLY_DATA
+	MOVIESCENE_API virtual FText GetDisplayName() const override;
+#endif
+
+protected:
+
+	/** All the sections in this track */
+	UPROPERTY()
+	TArray<TObjectPtr<UMovieSceneSection>> Sections;
+};
+
+
diff --git a/Engine/Source/Runtime/MovieSceneTracks/Private/Systems/MovieSceneDataLayerSystem.cpp b/Engine/Source/Runtime/MovieSceneTracks/Private/Systems/MovieSceneDataLayerSystem.cpp
--- a/Engine/Source/Runtime/MovieSceneTracks/Private/Systems/MovieSceneDataLayerSystem.cpp
+++ b/Engine/Source/Runtime/MovieSceneTracks/Private/Systems/MovieSceneDataLayerSystem.cpp
@@ -3,6 +3,7 @@
 #include "Systems/MovieSceneDataLayerSystem.h"
 #include "EntitySystem/MovieSceneEntitySystemTask.h"
 #include "EntitySystem/MovieSceneSpawnablesSystem.h"
+#include "EntitySystem/MovieSceneBindingLifetimeSystem.h"
 #include "EntitySystem/MovieSceneEntitySystemLinker.h"
 #include "EntitySystem/BuiltInComponentTypes.h"
 #include "Evaluation/PreAnimatedState/MovieScenePreAnimatedEntityCaptureSource.h"
@@ -499,7 +500,7 @@ UMovieSceneDataLayerSystem::UMovieSceneDataLayerSystem(const FObjectInitializer&
 
 	if (HasAnyFlags(RF_ClassDefaultObject))
 	{
-		DefineImplicitPrerequisite(GetClass(), UMovieSceneSpawnablesSystem::StaticClass());
+		DefineImplicitPrerequisite(GetClass(), UMovieSceneBindingLifetimeSystem::StaticClass());
 	}
 	else
 	{
