Constraints: fixed misbehavior when undoing changes with the curve editor opened

#rnx

diff --git a/Engine/Source/Editor/MovieSceneTools/Private/Channels/ConstraintChannelCurveModel.cpp b/Engine/Source/Editor/MovieSceneTools/Private/Channels/ConstraintChannelCurveModel.cpp
--- a/Engine/Source/Editor/MovieSceneTools/Private/Channels/ConstraintChannelCurveModel.cpp
+++ b/Engine/Source/Editor/MovieSceneTools/Private/Channels/ConstraintChannelCurveModel.cpp
@@ -46,7 +46,7 @@ FConstraintChannelCurveModel::FConstraintChannelCurveModel(TMovieSceneChannelHan
 	ChannelHandle = NewChannelProxy->MakeHandle<FMovieSceneConstraintChannel>(InChannel.GetChannelIndex());
 	if (FMovieSceneChannelProxy* ChannelProxy = InChannel.GetChannelProxy())
 	{
-		OnDestroyHandle = NewChannelProxy->OnDestroy.AddRaw(this, &FConstraintChannelCurveModel::NeedNewProxy);
+		OnDestroyHandle = NewChannelProxy->OnDestroy.AddRaw(this, &FConstraintChannelCurveModel::FixupCurve);
 	}
 
 	WeakSection = OwningSection;
@@ -64,7 +64,7 @@ UObject* FConstraintChannelCurveModel::GetOwningObject() const
 
 bool FConstraintChannelCurveModel::HasChangedAndResetTest()
 {
-	if (UMovieSceneSection* Section = WeakSection.Get())
+	if (const UMovieSceneSection* Section = WeakSection.Get())
 	{
 		if (Section->GetSignature() != LastSignature)
 		{
@@ -83,33 +83,21 @@ FConstraintChannelCurveModel::~FConstraintChannelCurveModel()
 		ChannelProxy->OnDestroy.Remove(OnDestroyHandle);
 	}
 }
-void FConstraintChannelCurveModel::NeedNewProxy() const
-{	
-	bNeedNewProxy = true; 
-}
 
-void FConstraintChannelCurveModel::FixupCurve() const
+void FConstraintChannelCurveModel::FixupCurve()
 {
-	if (UMovieSceneSection* Section = WeakSection.Get())
+	if (const UMovieSceneSection* Section = WeakSection.Get())
 	{
-		if (bNeedNewProxy)
-		{
-			if (IMovieSceneConstrainedSection* ConstrainedSection = Cast<IMovieSceneConstrainedSection>(Section))
-			{
-				ConstrainedSection->OnConstraintsChanged();
-			}
-			FMovieSceneChannelProxy* NewChannelProxy = &Section->GetChannelProxy();
-			ChannelHandle = NewChannelProxy->MakeHandle<FMovieSceneConstraintChannel>(ChannelHandle.GetChannelIndex());
-			OnDestroyHandle = NewChannelProxy->OnDestroy.AddRaw(this, &FConstraintChannelCurveModel::NeedNewProxy);
-			bNeedNewProxy = false;
-		}
+		FMovieSceneChannelProxy* NewChannelProxy = &Section->GetChannelProxy();
+		ChannelHandle = NewChannelProxy->MakeHandle<FMovieSceneConstraintChannel>(ChannelHandle.GetChannelIndex());
+		OnDestroyHandle = NewChannelProxy->OnDestroy.AddRaw(this, &FConstraintChannelCurveModel::FixupCurve);
 	}
 }
 
 TArray<FKeyBarCurveModel::FBarRange> FConstraintChannelCurveModel::FindRanges()
 {
 	FMovieSceneConstraintChannel* Channel = GetChannel();
-	UMovieSceneSection* Section = WeakSection.Get();
+	const UMovieSceneSection* Section = WeakSection.Get();
 	TArray<FKeyBarCurveModel::FBarRange> Ranges = FConstraintChannelEditor::GetBarRanges(Channel, Section);
 	return Ranges;
 }
@@ -121,12 +109,9 @@ const void* FConstraintChannelCurveModel::GetCurve() const
 
 FMovieSceneConstraintChannel* FConstraintChannelCurveModel::GetChannel() const
 {
-	FixupCurve();
-	FMovieSceneConstraintChannel* Channel = ChannelHandle.Get();
-	return Channel;
+	return ChannelHandle.Get();
 }
 
-
 void FConstraintChannelCurveModel::Modify()
 {
 	if (UMovieSceneSection* Section = WeakSection.Get())
diff --git a/Engine/Source/Editor/MovieSceneTools/Public/Channels/ConstraintChannelCurveModel.h b/Engine/Source/Editor/MovieSceneTools/Public/Channels/ConstraintChannelCurveModel.h
--- a/Engine/Source/Editor/MovieSceneTools/Public/Channels/ConstraintChannelCurveModel.h
+++ b/Engine/Source/Editor/MovieSceneTools/Public/Channels/ConstraintChannelCurveModel.h
@@ -75,10 +75,7 @@ public:
 
 
 private:
-	//Fixup the curve and get the right proxy later when we access it due to the fact the constraints channels may not be setup correctly
-	//These functions all need to be const since they can be called by const functions so we need to specify the things they change as mutable
-	void FixupCurve() const;
-	void NeedNewProxy() const;
+	void FixupCurve();
 	FMovieSceneConstraintChannel* GetChannel() const;
 
 private:
@@ -88,6 +85,5 @@ private:
 	TWeakPtr<ISequencer> WeakSequencer;
 	TWeakObjectPtr<UTickableConstraint> Constraint;
 
-	mutable FDelegateHandle OnDestroyHandle;
-	mutable bool bNeedNewProxy = false; 
+	FDelegateHandle OnDestroyHandle;
 };
\ No newline at end of file
