Sequencer: Snapping: Constriants:  Perform Euler filter when snapping/compensation static meshes

diff --git a/Engine/Source/Editor/MovieSceneTools/Private/MovieSceneToolHelpers.cpp b/Engine/Source/Editor/MovieSceneTools/Private/MovieSceneToolHelpers.cpp
--- a/Engine/Source/Editor/MovieSceneTools/Private/MovieSceneToolHelpers.cpp
+++ b/Engine/Source/Editor/MovieSceneTools/Private/MovieSceneToolHelpers.cpp
@@ -5349,7 +5349,50 @@ bool MovieSceneToolHelpers::AddTransformKeys(
 			}
 		}
 	}
-
+	//need to handle euler flips.. 3,4,5 is rotation
+	if (bKeyRotation)
+	{
+		//figure out start/endframe, frames may not be in order
+		FFrameNumber StartTime = Frames[0], EndTime = Frames[0];
+		for (const FFrameNumber& Frame : Frames)
+		{
+			if (Frame < StartTime)
+			{
+				StartTime = Frame;
+			}
+			else if (Frame > EndTime)
+			{
+				EndTime = Frame;
+			}
+		}
+		TSortedMap<FFrameNumber, FFrameNumber> FrameSet;
+		TRange<FFrameNumber> WithinRange(0, 0);
+		WithinRange.SetLowerBoundValue(StartTime);
+		WithinRange.SetUpperBoundValue(EndTime);
+		TArray<FFrameNumber> KeyTimes;
+		TArray<FKeyHandle> KeyHandles;
+		for (int32 ChannelIndex  =  3; ChannelIndex <=5;  ++ChannelIndex)
+		{
+			KeyTimes.SetNum(0);
+			KeyHandles.SetNum(0);
+			if (DoubleChannels[ChannelIndex].Get())
+			{
+				TMovieSceneChannelData<FMovieSceneDoubleValue> ChannelData = DoubleChannels[ChannelIndex].Get()->GetData();
+				ChannelData.GetKeys(WithinRange, &KeyTimes, &KeyHandles);
+				for (int32 KeyIndex = 0; KeyIndex < KeyTimes.Num() - 1; ++KeyIndex)
+				{
+					const int32 ValueIndex = ChannelData.GetIndex(KeyHandles[KeyIndex]);
+					const int32 NextValueIndex = ChannelData.GetIndex(KeyHandles[KeyIndex + 1]);
+
+					TArrayView<const FMovieSceneDoubleValue> Values = DoubleChannels[ChannelIndex].Get()->GetValues(); //re-get the array since the winding may change values
+					double Value = Values[ValueIndex].Value;
+					double NextValue = Values[NextValueIndex].Value;
+					FMath::WindRelativeAnglesDegrees(Value,NextValue);
+					AssignValue(DoubleChannels[ChannelIndex].Get(), KeyHandles[KeyIndex + 1], NextValue);
+				}
+			}
+		}
+	}
 	//now we need to set auto tangents
 	for (const int32 ChannelIndex : ChannelsIndexToKey)
 	{
