Sequencer: Constraints:  Constrained actor needs to be in Sequencer in order for ConstraintBindingID to be valid

diff --git a/Engine/Source/Editor/MovieSceneTools/Private/Constraints/MovieSceneConstraintChannelHelper.cpp b/Engine/Source/Editor/MovieSceneTools/Private/Constraints/MovieSceneConstraintChannelHelper.cpp
--- a/Engine/Source/Editor/MovieSceneTools/Private/Constraints/MovieSceneConstraintChannelHelper.cpp
+++ b/Engine/Source/Editor/MovieSceneTools/Private/Constraints/MovieSceneConstraintChannelHelper.cpp
@@ -1195,7 +1195,7 @@ void FMovieSceneConstraintChannelHelper::CreateBindingIDForHandle(const TSharedP
 			}
 			else
 			{
-				const FGuid Guid = InSequencer->GetHandleToObject(Actor, false); //don't create it???
+				const FGuid Guid = InSequencer->GetHandleToObject(Actor, true); //make sure object is in sequencer or binding id will be empty we won't resolve the binding
 				InHandle->ConstraintBindingID = UE::MovieScene::FRelativeObjectBindingID(Guid);
 			}
 		}
