Sequencer- Fix a few issues/bugs with Binding Lifetime tracks and overlapping sections on non-blending tracks in general.
* Broke out some copy/pasted code into a helper function for calculating the inverse of the lifetime section ranges, and protected against hanging in case somehow some of the sections overlap.
* Prevent re-making a section infinite if by doing so would cause it to overlap with another section in its range and the track doesn't allow blending on its sections.
* Fix a crash bug in the initial placement code for the binding lifetime sections caused by an edge case in the logic to prevent overlaps.

diff --git a/Engine/Source/Editor/MovieSceneTools/Private/MVVM/ViewModels/BindingLifetimeTrackModel.cpp b/Engine/Source/Editor/MovieSceneTools/Private/MVVM/ViewModels/BindingLifetimeTrackModel.cpp
--- a/Engine/Source/Editor/MovieSceneTools/Private/MVVM/ViewModels/BindingLifetimeTrackModel.cpp
+++ b/Engine/Source/Editor/MovieSceneTools/Private/MVVM/ViewModels/BindingLifetimeTrackModel.cpp
@@ -33,43 +33,10 @@ namespace UE
 
 		void FBindingLifetimeTrackModel::RecalculateInverseLifetimeRange()
 		{
-			InverseLifetimeRange.Reset();
-			InverseLifetimeRange.Add(TRange<FFrameNumber>::All());
-			// Iterate through our sections, removing them from the range
-
-			auto RemoveRangeFromSet = [&](TRange<FFrameNumber> SectionRange) {
-				for (int32 Index = 0; Index < InverseLifetimeRange.Num(); ++Index)
-				{
-					const TRange<FFrameNumber>& Current = InverseLifetimeRange[Index];
-					if (Current.Overlaps(SectionRange))
-					{
-						// Special case that difference doesn't handle well
-						if (!SectionRange.HasLowerBound() && !SectionRange.HasUpperBound())
-						{
-							InverseLifetimeRange.Reset();
-							return;
-						}
-
-						TArray<TRange<FFrameNumber>> SplitRanges = TRange<FFrameNumber>::Difference(Current, SectionRange);
-						for (const TRange<FFrameNumber>& NewRange : SplitRanges)
-						{
-							// Splitting infinite ranges keeps an infinite range, which we don't want to keep
-							if (!NewRange.HasLowerBound() && !NewRange.HasUpperBound())
-							{
-								continue;
-							}
-							
-							InverseLifetimeRange.Add(NewRange);
-						}
-						InverseLifetimeRange.RemoveAtSwap(Index--);
-					}
-				}
-			};
-
-			for (const TViewModelPtr<FSectionModel>& Item : GetSectionModels().IterateSubList<FSectionModel>())
-			{
-				RemoveRangeFromSet(Item->GetRange());
-			}
+			TArray<FFrameNumberRange> SectionRanges;
+			Algo::Transform(GetSectionModels().IterateSubList<FSectionModel>(), SectionRanges, [](const TViewModelPtr<FSectionModel>& Item) { return Item->GetRange(); });
+			
+			InverseLifetimeRange = UMovieSceneBindingLifetimeTrack::CalculateInverseLifetimeRange(SectionRanges);
 		}
 
 		FSortingKey FBindingLifetimeTrackModel::GetSortingKey() const
diff --git a/Engine/Source/Editor/MovieSceneTools/Private/MovieSceneSectionDetailsCustomization.cpp b/Engine/Source/Editor/MovieSceneTools/Private/MovieSceneSectionDetailsCustomization.cpp
--- a/Engine/Source/Editor/MovieSceneTools/Private/MovieSceneSectionDetailsCustomization.cpp
+++ b/Engine/Source/Editor/MovieSceneTools/Private/MovieSceneSectionDetailsCustomization.cpp
@@ -29,7 +29,9 @@ void FMovieSceneSectionDetailsCustomization::CustomizeDetails(IDetailLayoutBuild
 	TArray<TWeakObjectPtr<UObject>> Objects;
 	DetailBuilder.GetObjectsBeingCustomized(Objects);
 
-	bool bSectionsAreInfinite = true;
+	bool bSectionCanHaveOpenLowerBound = true;
+	bool bSectionCanHaveOpenUpperBound = true;
+
 	for (TWeakObjectPtr<UObject> Object : Objects)
 	{
 		if (Object.IsValid() && Object->IsA(UMovieSceneSection::StaticClass()))
@@ -37,9 +39,21 @@ void FMovieSceneSectionDetailsCustomization::CustomizeDetails(IDetailLayoutBuild
 			UMovieSceneSection* MovieSceneSection = (UMovieSceneSection*)Object.Get();
 			if (!MovieSceneSection->GetSupportsInfiniteRange())
 			{
-				bSectionsAreInfinite = false;
+				bSectionCanHaveOpenLowerBound = false;
+				bSectionCanHaveOpenUpperBound = false;
 				break;
 			}
+			else 
+			{
+				if (!MovieSceneSection->CanHaveOpenLowerBound())
+				{
+					bSectionCanHaveOpenLowerBound = false;
+				}
+				if (!MovieSceneSection->CanHaveOpenUpperBound())
+				{
+					bSectionCanHaveOpenUpperBound = false;
+				}
+			}
 		}
 	}
 
@@ -77,8 +91,8 @@ void FMovieSceneSectionDetailsCustomization::CustomizeDetails(IDetailLayoutBuild
 		.Padding(1)
 		[
 			SNew(SButton)
-			.Visibility_Lambda([bSectionsAreInfinite]() -> EVisibility {
-				return bSectionsAreInfinite ? EVisibility::Visible : EVisibility::Collapsed;
+			.Visibility_Lambda([bSectionCanHaveOpenLowerBound]() -> EVisibility {
+				return bSectionCanHaveOpenLowerBound ? EVisibility::Visible : EVisibility::Collapsed;
 			})
 			.OnClicked(this, &FMovieSceneSectionDetailsCustomization::ToggleRangeStartBounded)
 			.ContentPadding(0)
@@ -123,8 +137,8 @@ void FMovieSceneSectionDetailsCustomization::CustomizeDetails(IDetailLayoutBuild
 		.Padding(1)
 		[
 			SNew(SButton)
-			.Visibility_Lambda([bSectionsAreInfinite]() -> EVisibility {
-				return bSectionsAreInfinite ? EVisibility::Visible : EVisibility::Collapsed;
+			.Visibility_Lambda([bSectionCanHaveOpenUpperBound]() -> EVisibility {
+				return bSectionCanHaveOpenUpperBound ? EVisibility::Visible : EVisibility::Collapsed;
 			})
 			.OnClicked(this, &FMovieSceneSectionDetailsCustomization::ToggleRangeEndBounded)
 			.ContentPadding(0)
diff --git a/Engine/Source/Editor/MovieSceneTools/Public/MVVM/ViewModels/BindingLifetimeTrackModel.h b/Engine/Source/Editor/MovieSceneTools/Public/MVVM/ViewModels/BindingLifetimeTrackModel.h
--- a/Engine/Source/Editor/MovieSceneTools/Public/MVVM/ViewModels/BindingLifetimeTrackModel.h
+++ b/Engine/Source/Editor/MovieSceneTools/Public/MVVM/ViewModels/BindingLifetimeTrackModel.h
@@ -34,7 +34,7 @@ namespace UE
 			void OnDeferredModifyFlush() override;
 
 			// IBindingLifetimeExtension
-			const TArray<TRange<FFrameNumber>>& GetInverseLifetimeRange() const override { return InverseLifetimeRange; }
+			const TArray<FFrameNumberRange>& GetInverseLifetimeRange() const override { return InverseLifetimeRange; }
 		
 		protected:
 			virtual void OnConstruct() override;
@@ -44,7 +44,7 @@ namespace UE
 
 			// The inverse of the range created by all of our binding lifetime sections
 			// In other words, the ranges where the object binding should be deactivated.
-			TArray<TRange<FFrameNumber>> InverseLifetimeRange;
+			TArray<FFrameNumberRange> InverseLifetimeRange;
 
 		};
 
diff --git a/Engine/Source/Editor/Sequencer/Private/MVVM/ViewModels/BindingLifetimeOverlayModel.cpp b/Engine/Source/Editor/Sequencer/Private/MVVM/ViewModels/BindingLifetimeOverlayModel.cpp
--- a/Engine/Source/Editor/Sequencer/Private/MVVM/ViewModels/BindingLifetimeOverlayModel.cpp
+++ b/Engine/Source/Editor/Sequencer/Private/MVVM/ViewModels/BindingLifetimeOverlayModel.cpp
@@ -55,9 +55,9 @@ namespace UE
 			return FTrackLaneVirtualAlignment::Proportional(ViewRange, 1.f);
 		}
 
-		const TArray<TRange<FFrameNumber>>& FBindingLifetimeOverlayModel::GetInverseLifetimeRange() const
+		const TArray<FFrameNumberRange>& FBindingLifetimeOverlayModel::GetInverseLifetimeRange() const
 		{
-			const static TArray<TRange<FFrameNumber>> EmptyLifetimeRange;
+			const static TArray<FFrameNumberRange> EmptyLifetimeRange;
 			if (BindingLifetimeTrack)
 			{
 				return BindingLifetimeTrack->GetInverseLifetimeRange();
diff --git a/Engine/Source/Editor/Sequencer/Private/Widgets/SSequencerBindingLifetimeOverlay.cpp b/Engine/Source/Editor/Sequencer/Private/Widgets/SSequencerBindingLifetimeOverlay.cpp
--- a/Engine/Source/Editor/Sequencer/Private/Widgets/SSequencerBindingLifetimeOverlay.cpp
+++ b/Engine/Source/Editor/Sequencer/Private/Widgets/SSequencerBindingLifetimeOverlay.cpp
@@ -73,10 +73,10 @@ namespace UE
 			if (TSharedPtr<FBindingLifetimeOverlayModel> BindingLifetimeOverlayModel = GetBindingLifetimeOverlayModel())
 			{
 				// If applicable, paint an overlay over the range that the object binding is deactivated
-				const TArray<TRange<FFrameNumber>>& InverseLifetimeRange = BindingLifetimeOverlayModel->GetInverseLifetimeRange();
+				const TArray<FFrameNumberRange>& InverseLifetimeRange = BindingLifetimeOverlayModel->GetInverseLifetimeRange();
 				if (!InverseLifetimeRange.IsEmpty())
 				{
-					for (TRange<FFrameNumber> ExcludedRange : InverseLifetimeRange)
+					for (const FFrameNumberRange& ExcludedRange : InverseLifetimeRange)
 					{
 						const FFrameNumber LowerFrame = ExcludedRange.HasLowerBound() ? ExcludedRange.GetLowerBoundValue() : TimeToPixel->PixelToFrame(AllottedGeometry.Position.X).FloorToFrame();
 						const FFrameNumber UpperFrame = ExcludedRange.HasUpperBound() ? ExcludedRange.GetUpperBoundValue() : TimeToPixel->PixelToFrame(AllottedGeometry.Position.X + AllottedGeometry.Size.X).CeilToFrame();
diff --git a/Engine/Source/Editor/Sequencer/Public/MVVM/Extensions/IBindingLifetimeExtension.h b/Engine/Source/Editor/Sequencer/Public/MVVM/Extensions/IBindingLifetimeExtension.h
--- a/Engine/Source/Editor/Sequencer/Public/MVVM/Extensions/IBindingLifetimeExtension.h
+++ b/Engine/Source/Editor/Sequencer/Public/MVVM/Extensions/IBindingLifetimeExtension.h
@@ -19,7 +19,7 @@ namespace UE
 
 			virtual ~IBindingLifetimeExtension() {}
 
-			virtual const TArray<TRange<FFrameNumber>>& GetInverseLifetimeRange() const = 0;
+			virtual const TArray<FFrameNumberRange>& GetInverseLifetimeRange() const = 0;
 		};
 
 	} // namespace Sequencer
diff --git a/Engine/Source/Editor/Sequencer/Public/MVVM/ViewModels/BindingLifetimeOverlayModel.h b/Engine/Source/Editor/Sequencer/Public/MVVM/ViewModels/BindingLifetimeOverlayModel.h
--- a/Engine/Source/Editor/Sequencer/Public/MVVM/ViewModels/BindingLifetimeOverlayModel.h
+++ b/Engine/Source/Editor/Sequencer/Public/MVVM/ViewModels/BindingLifetimeOverlayModel.h
@@ -38,7 +38,7 @@ namespace UE
 			TSharedPtr<ITrackLaneWidget> CreateTrackLaneView(const FCreateTrackLaneViewParams& InParams) override;
 			FTrackLaneVirtualAlignment ArrangeVirtualTrackLaneView() const override;
 
-			const TArray<TRange<FFrameNumber>>& GetInverseLifetimeRange() const;
+			const TArray<FFrameNumberRange>& GetInverseLifetimeRange() const;
 
 		private:
 
diff --git a/Engine/Source/Runtime/MovieScene/Private/MovieSceneSection.cpp b/Engine/Source/Runtime/MovieScene/Private/MovieSceneSection.cpp
--- a/Engine/Source/Runtime/MovieScene/Private/MovieSceneSection.cpp
+++ b/Engine/Source/Runtime/MovieScene/Private/MovieSceneSection.cpp
@@ -456,6 +456,84 @@ void UMovieSceneSection::GetOverlappingSections(TArray<UMovieSceneSection*>& Out
 	}
 }
 
+/* Returns whether this section can have an open lower bound. This will generally be false if sections of this type cannot be blended and there is another section on the same row before this one.*/
+bool UMovieSceneSection::CanHaveOpenLowerBound() const
+{
+	if (!GetBlendType().IsValid())
+	{
+		UMovieSceneTrack* Track = GetTypedOuter<UMovieSceneTrack>();
+		if (!Track)
+		{
+			return true;
+		}
+
+		TRange<FFrameNumber> ThisRange = GetRange();
+
+		if (!ThisRange.HasLowerBound())
+		{
+			return true;
+		}
+
+		for (UMovieSceneSection* Section : Track->GetAllSections())
+		{
+			if (!Section || (Section == this))
+			{
+				continue;
+			}
+
+			if (Section->GetRowIndex() != GetRowIndex())
+			{
+				continue;
+			}
+
+			if (Section->GetRange().Overlaps(ThisRange) || (Section->GetRange().HasUpperBound() && Section->GetRange().GetUpperBoundValue() <= ThisRange.GetLowerBoundValue()))
+			{
+				return false;
+			}
+		}
+	}
+	return true;
+}
+
+/* Returns whether this section can have an open upper bound. This will generally be false if sections of this type cannot be blended and there is another section on the same row after this one.*/
+bool UMovieSceneSection::CanHaveOpenUpperBound() const
+{
+	if (!GetBlendType().IsValid())
+	{
+		UMovieSceneTrack* Track = GetTypedOuter<UMovieSceneTrack>();
+		if (!Track)
+		{
+			return true;
+		}
+
+		TRange<FFrameNumber> ThisRange = GetRange();
+
+		if (!ThisRange.HasUpperBound())
+		{
+			return true;
+		}
+
+		for (UMovieSceneSection* Section : Track->GetAllSections())
+		{
+			if (!Section || (Section == this))
+			{
+				continue;
+			}
+
+			if (Section->GetRowIndex() != GetRowIndex())
+			{
+				continue;
+			}
+
+			if (Section->GetRange().Overlaps(ThisRange) || (Section->GetRange().HasLowerBound() && Section->GetRange().GetLowerBoundValue() >= ThisRange.GetUpperBoundValue()))
+			{
+				return false;
+			}
+		}
+	}
+	return true;
+}
+
 
 const UMovieSceneSection* UMovieSceneSection::OverlapsWithSections(const TArray<UMovieSceneSection*>& Sections, int32 TrackDelta, int32 TimeDelta) const
 {
diff --git a/Engine/Source/Runtime/MovieScene/Private/Sections/MovieSceneBindingLifetimeSection.cpp b/Engine/Source/Runtime/MovieScene/Private/Sections/MovieSceneBindingLifetimeSection.cpp
--- a/Engine/Source/Runtime/MovieScene/Private/Sections/MovieSceneBindingLifetimeSection.cpp
+++ b/Engine/Source/Runtime/MovieScene/Private/Sections/MovieSceneBindingLifetimeSection.cpp
@@ -72,12 +72,22 @@ void UMovieSceneBindingLifetimeSection::InitialPlacement(const TArray<UMovieScen
 		}
 		else
 		{
-			// We're likely overlapping the first infinite section here. Split it.
-			UMovieSceneSection* NewSection = OverlappedSection->SplitSection(FQualifiedFrameTime(InStartTime, GetTypedOuter<UMovieScene>()->GetTickResolution()), false);
-			// The new section will be the one on the right and will still have an open upper bound. Resize it to fit the new section.
-			TRange<FFrameNumber> NewSectionRange = NewSection->GetRange();
-			NewSectionRange.SetLowerBoundValue(InStartTime + Duration);
-			NewSection->SetRange(NewSectionRange);
+			// Edge case- the other range has an infinite end, but can't be split because the point it would be split at is its start.
+			if (OtherRange.GetLowerBound().IsClosed() && OtherRange.GetLowerBoundValue() == SectionRange.GetLowerBound().GetValue())
+			{
+				// Adjust the lowerbound of OtherRange
+				OtherRange.SetLowerBoundValue(SectionRange.GetUpperBound().GetValue());
+				OverlappedSection->SetRange(OtherRange);
+			}
+			else
+			{
+				// We're overlapping an at-least partially infinite section here. Split it.
+				UMovieSceneSection* NewSection = OverlappedSection->SplitSection(FQualifiedFrameTime(InStartTime, GetTypedOuter<UMovieScene>()->GetTickResolution()), false);
+				// The new section will be the one on the right and will still have an open upper bound. Resize it to fit the new section.
+				TRange<FFrameNumber> NewSectionRange = NewSection->GetRange();
+				NewSectionRange.SetLowerBoundValue(InStartTime + Duration);
+				NewSection->SetRange(NewSectionRange);
+			}
 		}
 	}
 }
diff --git a/Engine/Source/Runtime/MovieScene/Private/Tracks/MovieSceneBindingLifetimeTrack.cpp b/Engine/Source/Runtime/MovieScene/Private/Tracks/MovieSceneBindingLifetimeTrack.cpp
--- a/Engine/Source/Runtime/MovieScene/Private/Tracks/MovieSceneBindingLifetimeTrack.cpp
+++ b/Engine/Source/Runtime/MovieScene/Private/Tracks/MovieSceneBindingLifetimeTrack.cpp
@@ -97,8 +97,7 @@ void UMovieSceneBindingLifetimeTrack::ImportEntityImpl(UMovieSceneEntitySystemLi
 		);
 }
 
-
-TArray<FFrameNumberRange> CalculateInverseLifetimeRange(const FMovieSceneTrackEvaluationField& LocalEvaluationField)
+TArray<FFrameNumberRange> UMovieSceneBindingLifetimeTrack::CalculateInverseLifetimeRange(const TArray<FFrameNumberRange>& Ranges)
 {
 	TArray<FFrameNumberRange> InverseLifetimeRange;
 	InverseLifetimeRange.Add(FFrameNumberRange::All());
@@ -133,9 +132,18 @@ TArray<FFrameNumberRange> CalculateInverseLifetimeRange(const FMovieSceneTrackEv
 		}
 	};
 
-	for (const FMovieSceneTrackEvaluationFieldEntry& Entry : LocalEvaluationField.Entries)
+	TArray<TRange<FFrameNumber>> SectionRanges;
+	for (const FFrameNumberRange& Range : Ranges)
 	{
-		RemoveRangeFromSet(Entry.Range);
+		if (Algo::AnyOf(SectionRanges, [&](const TRange<FFrameNumber> OtherRange) { return OtherRange.Overlaps(Range); }))
+		{
+			// Lifetime Range sections have managed to overlap, which should not be allowed.
+			ensure(false);
+			InverseLifetimeRange.Reset();
+			return InverseLifetimeRange;
+		}
+		SectionRanges.Add(Range);
+		RemoveRangeFromSet(Range);
 	}
 	return InverseLifetimeRange;
 }
@@ -143,8 +151,10 @@ TArray<FFrameNumberRange> CalculateInverseLifetimeRange(const FMovieSceneTrackEv
 bool UMovieSceneBindingLifetimeTrack::PopulateEvaluationFieldImpl(const TRange<FFrameNumber>& EffectiveRange, const FMovieSceneEvaluationFieldEntityMetaData& InMetaData, FMovieSceneEntityComponentFieldBuilder* OutFieldBuilder)
 {
 	const FMovieSceneTrackEvaluationField& LocalEvaluationField = GetEvaluationField();
+	TArray<FFrameNumberRange> Ranges;
+	Algo::Transform(LocalEvaluationField.Entries, Ranges, [](const FMovieSceneTrackEvaluationFieldEntry& Entry) { return Entry.Range; });
 
-	TArray<FFrameNumberRange> InverseLifetimeRange = CalculateInverseLifetimeRange(LocalEvaluationField);
+	TArray<FFrameNumberRange> InverseLifetimeRange = CalculateInverseLifetimeRange(Ranges);
 	// Add an entity for each section of the inverse range to define areas where our lifetime is inactive.
 	for (const FFrameNumberRange& InverseRange : InverseLifetimeRange)
 	{
diff --git a/Engine/Source/Runtime/MovieScene/Public/MovieSceneSection.h b/Engine/Source/Runtime/MovieScene/Public/MovieSceneSection.h
--- a/Engine/Source/Runtime/MovieScene/Public/MovieSceneSection.h
+++ b/Engine/Source/Runtime/MovieScene/Public/MovieSceneSection.h
@@ -606,6 +606,12 @@ public:
 	 */
 	MOVIESCENE_API void GetOverlappingSections(TArray<UMovieSceneSection*>& OutSections, bool bSameRow, bool bIncludeThis);
 
+	/* Returns whether this section can have an open lower bound. This will generally be false if sections of this type cannot be blended and there is another section on the same row before this one.*/
+	MOVIESCENE_API bool CanHaveOpenLowerBound() const;
+
+	/* Returns whether this section can have an open upper bound. This will generally be false if sections of this type cannot be blended and there is another section on the same row after this one.*/
+	MOVIESCENE_API bool CanHaveOpenUpperBound() const;
+
 	/**
 	 * Evaluate this sections's easing functions based on the specified time
 	 */
diff --git a/Engine/Source/Runtime/MovieScene/Public/Tracks/MovieSceneBindingLifetimeTrack.h b/Engine/Source/Runtime/MovieScene/Public/Tracks/MovieSceneBindingLifetimeTrack.h
--- a/Engine/Source/Runtime/MovieScene/Public/Tracks/MovieSceneBindingLifetimeTrack.h
+++ b/Engine/Source/Runtime/MovieScene/Public/Tracks/MovieSceneBindingLifetimeTrack.h
@@ -51,6 +51,8 @@ public:
 	virtual void ImportEntityImpl(UMovieSceneEntitySystemLinker* EntityLinker, const FEntityImportParams& Params, FImportedEntity* OutImportedEntity) override;
 	virtual bool PopulateEvaluationFieldImpl(const TRange<FFrameNumber>& EffectiveRange, const FMovieSceneEvaluationFieldEntityMetaData& InMetaData, FMovieSceneEntityComponentFieldBuilder* OutFieldBuilder) override;
 
+	MOVIESCENE_API static TArray<FFrameNumberRange> CalculateInverseLifetimeRange(const TArray<FFrameNumberRange>& Ranges);
+
 
 #if WITH_EDITORONLY_DATA
 	MOVIESCENE_API virtual FText GetDisplayName() const override;
