Sequencer: Null check when locking camera to viewport and ther is no active viewport.

diff --git a/Engine/Source/Editor/MovieSceneTools/Private/MovieSceneToolHelpers.cpp b/Engine/Source/Editor/MovieSceneTools/Private/MovieSceneToolHelpers.cpp
--- a/Engine/Source/Editor/MovieSceneTools/Private/MovieSceneToolHelpers.cpp
+++ b/Engine/Source/Editor/MovieSceneTools/Private/MovieSceneToolHelpers.cpp
@@ -1279,7 +1279,7 @@ void MovieSceneToolHelpers::LockCameraActorToViewport(const TSharedPtr<ISequence
 	Sequencer->SetPerspectiveViewportCameraCutEnabled(false);
 
 	// Lock the viewport to this camera
-	if (CameraActor && CameraActor->GetLevel())
+	if (GCurrentLevelEditingViewportClient && CameraActor && CameraActor->GetLevel())
 	{
 		GCurrentLevelEditingViewportClient->SetCinematicActorLock(nullptr);
 		GCurrentLevelEditingViewportClient->SetActorLock(CameraActor);
