Sequencer: Fix FBX Export context with GetWorld()

#rnx

diff --git a/Engine/Source/Editor/MovieSceneTools/Private/LevelSequenceFBXInterop.cpp b/Engine/Source/Editor/MovieSceneTools/Private/LevelSequenceFBXInterop.cpp
--- a/Engine/Source/Editor/MovieSceneTools/Private/LevelSequenceFBXInterop.cpp
+++ b/Engine/Source/Editor/MovieSceneTools/Private/LevelSequenceFBXInterop.cpp
@@ -211,7 +211,7 @@ void FLevelSequenceFBXInterop::ExportFBXInternal(const FString& ExportFilename,
 	if (!ExportCancel)
 	{
 		UMovieScene* MovieScene = Sequencer->GetFocusedMovieSceneSequence()->GetMovieScene();
-		UWorld* World = Cast<UWorld>(Sequencer->GetPlaybackContext());
+		UWorld* World = Sequencer->GetPlaybackContext()->GetWorld();
 		FMovieSceneSequenceIDRef Template = Sequencer->GetFocusedTemplateID();
 		UnFbx::FFbxExporter::FLevelSequenceNodeNameAdapter NodeNameAdapter(MovieScene, Sequencer.Get(), Template);
 
