Sequencer: Fix merge issues with custom primitive data track editor registration getting lost

#rnx

diff --git a/Engine/Source/Editor/MovieSceneTools/Private/MovieSceneToolsModule.cpp b/Engine/Source/Editor/MovieSceneTools/Private/MovieSceneToolsModule.cpp
--- a/Engine/Source/Editor/MovieSceneTools/Private/MovieSceneToolsModule.cpp
+++ b/Engine/Source/Editor/MovieSceneTools/Private/MovieSceneToolsModule.cpp
@@ -182,9 +182,9 @@ void FMovieSceneToolsModule::StartupModule()
 		PrimitiveMaterialCreateEditorHandle = SequencerModule.RegisterTrackEditor(FOnCreateTrackEditor::CreateStatic(&FPrimitiveMaterialTrackEditor::CreateTrackEditor));
 		CameraShakeSourceShakeCreateEditorHandle = SequencerModule.RegisterTrackEditor(FOnCreateTrackEditor::CreateStatic(&FCameraShakeSourceShakeTrackEditor::CreateTrackEditor));
 		CVarTrackCreateEditorHandle = SequencerModule.RegisterTrackEditor(FOnCreateTrackEditor::CreateStatic(&FCVarTrackEditor::CreateTrackEditor));
+		CustomPrimitiveDataTrackCreateEditorHandle = SequencerModule.RegisterTrackEditor(FOnCreateTrackEditor::CreateStatic(&FCustomPrimitiveDataTrackEditor::CreateTrackEditor));
 		BindingLifetimeTrackCreateEditorHandle = SequencerModule.RegisterTrackEditor(FOnCreateTrackEditor::CreateStatic(&FBindingLifetimeTrackEditor::CreateTrackEditor));
 
-
 		// register track models
 		CameraCutTrackModelHandle = SequencerModule.RegisterTrackModel(FOnCreateTrackModel::CreateStatic(&FCameraCutTrackModel::CreateTrackModel));
 		CinematicShotTrackModelHandle = SequencerModule.RegisterTrackModel(FOnCreateTrackModel::CreateStatic(&FCinematicShotTrackModel::CreateTrackModel));
@@ -347,7 +347,8 @@ void FMovieSceneToolsModule::ShutdownModule()
 	SequencerModule.UnRegisterTrackEditor( ObjectTrackCreateEditorHandle );
 	SequencerModule.UnRegisterTrackEditor( PrimitiveMaterialCreateEditorHandle );
 	SequencerModule.UnRegisterTrackEditor( CVarTrackCreateEditorHandle );
-	SequencerModule.UnRegisterTrackEditor(BindingLifetimeTrackCreateEditorHandle);
+	SequencerModule.UnRegisterTrackEditor( CustomPrimitiveDataTrackCreateEditorHandle );
+	SequencerModule.UnRegisterTrackEditor( BindingLifetimeTrackCreateEditorHandle );
 
 	// unregister track models
 	SequencerModule.UnregisterTrackModel( CameraCutTrackModelHandle );
