Fix audio clock source dialog not suppressing when user selects 'no' when dismissing dialog.

diff --git a/Engine/Source/Editor/MaterialEditor/Private/MaterialEditor.cpp b/Engine/Source/Editor/MaterialEditor/Private/MaterialEditor.cpp
--- a/Engine/Source/Editor/MaterialEditor/Private/MaterialEditor.cpp
+++ b/Engine/Source/Editor/MaterialEditor/Private/MaterialEditor.cpp
@@ -2686,7 +2686,7 @@ bool FMaterialEditor::UpdateOriginalMaterial()
 				FText::Format(NSLOCTEXT("UnrealEd", "Error_CompileErrorsInDefaultMaterial", "The current material has the following compilation errors:\n{0}\nThis material is a Default Material which must be available as a code fallback at all times, compilation errors are not allowed."), JoinedErrors),
 				NSLOCTEXT("UnrealEd", "Warning_CompileErrorsInDefaultMaterial_Title", "Error: Compilation errors in Default Material"), "Error_CompileErrorsInDefaultMaterial");
 			Info.ConfirmText = NSLOCTEXT("ModalDialogs", "CompileErrorsInDefaultMaterialOk", "Ok");
-			Info.bDontPersistSuppressionAcrossSessions = true;
+			Info.DialogMode = FSuppressableWarningDialog::EMode::DontPersistSuppressionAcrossSessions;
 			Info.WrapMessageAt = 0.0f;
 
 			FSuppressableWarningDialog CompileErrors(Info);
@@ -2700,7 +2700,7 @@ bool FMaterialEditor::UpdateOriginalMaterial()
 				NSLOCTEXT("UnrealEd", "Warning_CompileErrorsInMaterial_Title", "Warning: Compilation errors in this Material" ), "Warning_CompileErrorsInMaterial");
 			Info.ConfirmText = NSLOCTEXT("ModalDialogs", "CompileErrorsInMaterialConfirm", "Continue");
 			Info.CancelText = NSLOCTEXT("ModalDialogs", "CompileErrorsInMaterialCancel", "Abort");
-			Info.bDontPersistSuppressionAcrossSessions = true;
+			Info.DialogMode = FSuppressableWarningDialog::EMode::DontPersistSuppressionAcrossSessions;
 			Info.WrapMessageAt = 0.0f;
 
 			FSuppressableWarningDialog CompileErrorsWarning( Info );
diff --git a/Engine/Source/Editor/MovieSceneTools/Private/TrackEditors/AudioTrackEditor.cpp b/Engine/Source/Editor/MovieSceneTools/Private/TrackEditors/AudioTrackEditor.cpp
--- a/Engine/Source/Editor/MovieSceneTools/Private/TrackEditors/AudioTrackEditor.cpp
+++ b/Engine/Source/Editor/MovieSceneTools/Private/TrackEditors/AudioTrackEditor.cpp
@@ -1409,11 +1409,9 @@ void FAudioTrackEditor::OnAttachedAudioEnterPressed(const TArray<FAssetData>& As
 
 void FAudioTrackEditor::RegisterMovieSceneChangedDelegate(TSharedRef<ISequencer> InSequencer)
 {
-	// Check the sequence prior to installing the delegate
-	if (CheckSequenceClockSource())
+	if (SequenceContainsAudioTrack(InSequencer->GetRootMovieSceneSequence()))
 	{
-		// If we are here then either the clock source for this sequence is already set to audio clock or
-		// there is an audio track and the user has been notified. Either way, no need to install the delegate. 
+		// This sequence already has an audio track. Don't install the delegate.
 		return;
 	}
 	
@@ -1436,6 +1434,55 @@ void FAudioTrackEditor::RegisterMovieSceneChangedDelegate(TSharedRef<ISequencer>
 	});
 }
 
+bool FAudioTrackEditor::SequenceContainsAudioTrack(const UMovieSceneSequence* InSequence)
+{
+	if (!InSequence)
+	{
+		return false;
+	}
+
+	if (UMovieScene* MovieScene = InSequence->GetMovieScene())
+	{
+		for (const UMovieSceneTrack* Track : MovieScene->GetTracks())
+		{
+			if (Cast<UMovieSceneAudioTrack>(Track))
+			{
+				return true;
+			}
+
+			const UMovieSceneSubTrack* SubTrack = Cast<UMovieSceneSubTrack>(Track);
+			if (!SubTrack)
+			{
+				continue;
+			}
+
+			for (const UMovieSceneSection* Section : SubTrack->GetAllSections())
+			{
+				const UMovieSceneSubSection* SubSection = Cast<UMovieSceneSubSection>(Section);
+				if (!SubSection)
+				{
+					continue;
+				}
+
+				UMovieSceneSequence* SubSequence = SubSection->GetSequence();
+				if (!SubSequence)
+				{
+					continue;
+				}
+				else
+				{
+					if (SequenceContainsAudioTrack(SubSequence))
+					{
+						return true;
+					}
+				}
+			}
+		}
+	}
+
+	return false;
+}
+
 bool FAudioTrackEditor::CheckSequenceClockSource()
 {
 	TSharedPtr<ISequencer> SequencerPtr = GetSequencer();
@@ -1445,7 +1492,7 @@ bool FAudioTrackEditor::CheckSequenceClockSource()
 	{
 		if (UMovieScene* MovieScene = RootSequence->GetMovieScene())
 		{
-			const bool bHasAudioTrack = (MovieScene->FindTrack(UMovieSceneAudioTrack::StaticClass()) != nullptr);
+			const bool bHasAudioTrack = SequenceContainsAudioTrack(RootSequence);
 			const bool bIsUsingAudioClock = (MovieScene->GetClockSource() == EUpdateClockSource::Audio);
 
 			if (bIsUsingAudioClock)
@@ -1453,9 +1500,12 @@ bool FAudioTrackEditor::CheckSequenceClockSource()
 				// If sequence is already using audio clock, we're done
 				return true;
 			} 
-			else if (bHasAudioTrack && !MovieScene->IsReadOnly())
+			else if (bHasAudioTrack)
 			{
-				PromptUserForClockSource();
+				if (!MovieScene->IsReadOnly())
+				{
+					PromptUserForClockSource();
+				}
 
 				// Only prompt once per sequencer instance to avoid dialog thrashing
 				return true;
@@ -1471,13 +1521,14 @@ void FAudioTrackEditor::PromptUserForClockSource()
 	FSuppressableWarningDialog::FSetupInfo SetupInfo(
 		LOCTEXT("AutoSelectAudioClockSource_Message", "It is recommended to use the audio clock as the clock source when working with audio tracks in sequencer for improved synchronization between animation and audio. Would you like to switch the clock source now?"),
 		LOCTEXT("AutoSelectAudioClockSource_Title", "Use Audio Clock Source?"),
-		TEXT("AutoSelectAudioClockSource_Dialog"));
+		TEXT("AutoSelectAudioClockSource_SuppressDialog"));
 
 	SetupInfo.ConfirmText = LOCTEXT("AutoSelectAudioClockSource_ConfirmText", "Yes");
 	SetupInfo.CancelText = LOCTEXT("AutoSelectAudioClockSource_CancelText", "No");
 	SetupInfo.CheckBoxText = LOCTEXT("AutoSelectAudioClockSource_CheckBoxText", "Don't show this again");
 	SetupInfo.bDefaultToSuppressInTheFuture = false;
-
+	SetupInfo.DialogMode = FSuppressableWarningDialog::EMode::PersistUserResponse;
+	
 	FSuppressableWarningDialog SwitchToAudioClockSourceDialog(SetupInfo);
 	FSuppressableWarningDialog::EResult Result = SwitchToAudioClockSourceDialog.ShowModal();
 
@@ -1486,10 +1537,6 @@ void FAudioTrackEditor::PromptUserForClockSource()
 		// Configure this sequence's clock source to use the audio clock
 		SetClockSoureToAudioClock();
 	}
-	else if (Result == FSuppressableWarningDialog::Suppressed)
-	{
-		UE_LOG(LogMovieScene, Display, TEXT("It is recommended to use the audio clock as the clock source when working with audio tracks in sequencer for improved synchronization between animation and audio. Consider switching to the audio clock source."));
-	}
 }
 
 void FAudioTrackEditor::SetClockSoureToAudioClock()
diff --git a/Engine/Source/Editor/MovieSceneTools/Public/TrackEditors/AudioTrackEditor.h b/Engine/Source/Editor/MovieSceneTools/Public/TrackEditors/AudioTrackEditor.h
--- a/Engine/Source/Editor/MovieSceneTools/Public/TrackEditors/AudioTrackEditor.h
+++ b/Engine/Source/Editor/MovieSceneTools/Public/TrackEditors/AudioTrackEditor.h
@@ -99,6 +99,9 @@ private:
 	/** Registers a delegate with the given sequencer for monitoring edits */
 	void RegisterMovieSceneChangedDelegate(TSharedRef<ISequencer> InSequencer);
 
+	/** Returns true if the given Sequence or any subsequence contains an audio track */
+	bool SequenceContainsAudioTrack(const UMovieSceneSequence* InSequence);
+
 	/** Will return true if a sequence contains an audio track and the user was notified about the potential clock source issue */
 	bool CheckSequenceClockSource();
 
diff --git a/Engine/Source/Editor/UnrealEd/Private/Dialogs/Dialogs.cpp b/Engine/Source/Editor/UnrealEd/Private/Dialogs/Dialogs.cpp
--- a/Engine/Source/Editor/UnrealEd/Private/Dialogs/Dialogs.cpp
+++ b/Engine/Source/Editor/UnrealEd/Private/Dialogs/Dialogs.cpp
@@ -828,9 +828,21 @@ FSuppressableWarningDialog::FSuppressableWarningDialog(const FSetupInfo& Info)
 	IniSettingName = Info.IniSettingName;
 	IniSettingFileName = Info.IniSettingFileName;
 	Prompt = Info.Message;
-	bDontPersistSuppressionAcrossSessions = Info.bDontPersistSuppressionAcrossSessions;
+	ResponseIniSettingName = Info.IniSettingName + TEXT("_ConfirmResponse");
 
-	if (bDontPersistSuppressionAcrossSessions)
+	// bDontPersistSuppressionAcrossSessions takes precedence until removed
+	PRAGMA_DISABLE_DEPRECATION_WARNINGS
+	if (Info.bDontPersistSuppressionAcrossSessions)
+	PRAGMA_ENABLE_DEPRECATION_WARNINGS
+	{
+		DialogMode = EMode::DontPersistSuppressionAcrossSessions;
+	}
+	else
+	{
+		DialogMode = Info.DialogMode;
+	}
+
+	if (DialogMode == EMode::DontPersistSuppressionAcrossSessions)
 	{
 		bShouldSuppressDialog = SuppressedInTheSession.Contains(SuppressableWarningDialogGetSessionKey(IniSettingName, IniSettingFileName));
 	}
@@ -871,7 +883,7 @@ FSuppressableWarningDialog::EResult FSuppressableWarningDialog::ShowModal() cons
 	bool bShouldSuppressDialog = false;
 
 	// Get the setting from the config file.
-	if (bDontPersistSuppressionAcrossSessions)
+	if (DialogMode == EMode::DontPersistSuppressionAcrossSessions)
 	{
 		bShouldSuppressDialog = SuppressedInTheSession.Contains(SuppressableWarningDialogGetSessionKey(IniSettingName, IniSettingFileName));
 	}
@@ -885,12 +897,13 @@ FSuppressableWarningDialog::EResult FSuppressableWarningDialog::ShowModal() cons
 	{
 		GEditor->EditorAddModalWindow(ModalWindow.ToSharedRef());
 		RetCode = (MessageBox->GetResponse()) ? Confirm : Cancel;
+		
+		// Set the ini variable to the state of the disable check box
+		bShouldSuppressDialog = MessageBox->GetCheckBoxState();
 
 		if( RetCode == Confirm )
 		{
-			// Set the ini variable to the state of the disable check box
-			bShouldSuppressDialog = MessageBox->GetCheckBoxState();
-			if (bDontPersistSuppressionAcrossSessions)
+			if (DialogMode == EMode::DontPersistSuppressionAcrossSessions)
 			{
 				if (bShouldSuppressDialog)
 				{
@@ -899,7 +912,20 @@ FSuppressableWarningDialog::EResult FSuppressableWarningDialog::ShowModal() cons
 			}
 			else
 			{
-				GConfig->SetBool( *ConfigSection, *IniSettingName, bShouldSuppressDialog, IniSettingFileName );
+				GConfig->SetBool(*ConfigSection, *IniSettingName, bShouldSuppressDialog, IniSettingFileName);
+
+				if (DialogMode == EMode::PersistUserResponse)
+				{
+					GConfig->SetBool(*ConfigSection, *ResponseIniSettingName, true, IniSettingFileName);
+				}
+			}
+		}
+		else
+		{
+			if (DialogMode == EMode::PersistUserResponse)
+			{
+				GConfig->SetBool(*ConfigSection, *IniSettingName, bShouldSuppressDialog, IniSettingFileName);
+				GConfig->SetBool(*ConfigSection, *ResponseIniSettingName, false, IniSettingFileName);
 			}
 		}
 	}
@@ -907,6 +933,15 @@ FSuppressableWarningDialog::EResult FSuppressableWarningDialog::ShowModal() cons
 	{
 		// If the dialog is suppressed, log the warning
 		UE_LOG(LogDialogs, Warning, TEXT("Suppressed: %s"), *Prompt.ToString());
+
+		if (DialogMode == EMode::PersistUserResponse)
+		{
+			// Get the saved response
+			bool bWasConfirmed = false;
+			GConfig->GetBool(*ConfigSection, *ResponseIniSettingName, bWasConfirmed, IniSettingFileName);
+
+			RetCode = bWasConfirmed ? Confirm : Cancel;
+		}
 	}
 
 	return RetCode;
diff --git a/Engine/Source/Editor/UnrealEd/Public/Dialogs/Dialogs.h b/Engine/Source/Editor/UnrealEd/Public/Dialogs/Dialogs.h
--- a/Engine/Source/Editor/UnrealEd/Public/Dialogs/Dialogs.h
+++ b/Engine/Source/Editor/UnrealEd/Public/Dialogs/Dialogs.h
@@ -77,7 +77,22 @@ class FSuppressableWarningDialog
 {
 public:
 
-	/** 
+	/**
+	 * Optional mode of operation for FSuppressableWarningDialog
+	 */
+	enum class EMode
+	{
+		/** Default behavior for dialog */
+		Default,
+
+		/** Dialog suppression will not persist after editor closes */
+		DontPersistSuppressionAcrossSessions,
+
+		/** Persist user result */
+		PersistUserResponse
+	};
+
+	/**
 	 * Struct used to initialize FSuppressableWarningDialog
 	 * 
 	 * User must provide confirm text, and cancel text (if using cancel button)
@@ -100,8 +115,12 @@ public:
 		bool bDefaultToSuppressInTheFuture;
 
 		/** If true suppression will not persist for future editor sessions */
+		UE_DEPRECATED(5.4, "bDontPersistSuppressionAcrossSessions is deprecated, please use FSetupInfo::DialogMode instead.")
 		bool bDontPersistSuppressionAcrossSessions;
 
+		/** Optional mode of operation for FSuppressableWarningDialog */
+		EMode DialogMode;
+
 		/** Wrap message at specified length, zero or negative number will disable the wrapping */
 		float WrapMessageAt;
 
@@ -132,6 +151,7 @@ public:
 			, IniSettingFileName(InIniSettingFileName)
 			, bDefaultToSuppressInTheFuture(false)
 			, bDontPersistSuppressionAcrossSessions(false)
+			, DialogMode(EMode::Default)
 			, WrapMessageAt(512.0f)
 			, ConfirmText()
 			, CancelText()
@@ -164,6 +184,9 @@ private:
 	/** Name of the flag which controls whether to launch the warning */
 	FString IniSettingName;
 
+	/** The name of the setting which stores the user response when dismissing the dialog. */
+	FString ResponseIniSettingName;
+
 	/** Name of the file which stores the IniSettingName flag result */
 	FString IniSettingFileName;
 
@@ -176,8 +199,8 @@ private:
 	/** Cached pointer to the message box held within the window */
 	TSharedPtr<class SModalDialogWithCheckbox> MessageBox;
 
-	/** If true suppression will not persist for future editor sessions */
-	bool bDontPersistSuppressionAcrossSessions;
+	/** Optional mode of operation */
+	EMode DialogMode;
 
 	/** Set of session only suppressions */
 	static TSet<FString> SuppressedInTheSession;
