Fix clock source dialog crash when this pointer is invalid

diff --git a/Engine/Source/Editor/MovieSceneTools/Private/TrackEditors/AudioTrackEditor.cpp b/Engine/Source/Editor/MovieSceneTools/Private/TrackEditors/AudioTrackEditor.cpp
--- a/Engine/Source/Editor/MovieSceneTools/Private/TrackEditors/AudioTrackEditor.cpp
+++ b/Engine/Source/Editor/MovieSceneTools/Private/TrackEditors/AudioTrackEditor.cpp
@@ -932,19 +932,27 @@ void FAudioSection::RegenerateWaveforms(TRange<float> DrawRange, int32 XOffset,
 FAudioTrackEditor::FAudioTrackEditor( TSharedRef<ISequencer> InSequencer )
 	: FMovieSceneTrackEditor( InSequencer ) 
 {
-	RegisterMovieSceneChangedDelegate(InSequencer);
 }
 
 FAudioTrackEditor::~FAudioTrackEditor()
+{
+}
+
+void FAudioTrackEditor::OnInitialize()
+{
+	RegisterMovieSceneChangedDelegate();
+}
+
+void FAudioTrackEditor::OnRelease()
 {
 	TSharedPtr<ISequencer> SequencerPtr = GetSequencer();
 	if (SequencerPtr.IsValid() && MovieSceneChangedDelegate.IsValid())
 	{
 		SequencerPtr->OnMovieSceneDataChanged().Remove(MovieSceneChangedDelegate);
+		MovieSceneChangedDelegate.Reset();
 	}
 }
 
-
 TSharedRef<ISequencerTrackEditor> FAudioTrackEditor::CreateTrackEditor( TSharedRef<ISequencer> InSequencer )
 {
 	return MakeShareable( new FAudioTrackEditor( InSequencer ) );
@@ -1407,31 +1415,37 @@ void FAudioTrackEditor::OnAttachedAudioEnterPressed(const TArray<FAssetData>& As
 	}
 }
 
-void FAudioTrackEditor::RegisterMovieSceneChangedDelegate(TSharedRef<ISequencer> InSequencer)
+void FAudioTrackEditor::RegisterMovieSceneChangedDelegate()
 {
-	if (SequenceContainsAudioTrack(InSequencer->GetRootMovieSceneSequence()))
+	TSharedPtr<ISequencer> SequencerPtr = GetSequencer();
+	if (SequencerPtr.IsValid())
 	{
-		// This sequence already has an audio track. Don't install the delegate.
-		return;
-	}
+		if (SequenceContainsAudioTrack(SequencerPtr->GetRootMovieSceneSequence()))
+		{
+			// This sequence already has an audio track. Don't install the delegate.
+			return;
+		}
 	
-	// Add delegate for scene data change events
-	MovieSceneChangedDelegate = InSequencer->OnMovieSceneDataChanged().AddLambda([this](EMovieSceneDataChangeType InChangeType)
+		// Add delegate for scene data change events
+		MovieSceneChangedDelegate = SequencerPtr->OnMovieSceneDataChanged().AddSP(this, &FAudioTrackEditor::OnMovieSceneDataChanged);
+	}
+}
+
+void FAudioTrackEditor::OnMovieSceneDataChanged(EMovieSceneDataChangeType InChangeType)
+{
+	if (InChangeType == EMovieSceneDataChangeType::MovieSceneStructureItemAdded)
 	{
-		if (InChangeType == EMovieSceneDataChangeType::MovieSceneStructureItemAdded)
+		if (CheckSequenceClockSource())
 		{
-			if (CheckSequenceClockSource())
+			// The user has been notified, remove the delegate
+			TSharedPtr<ISequencer> SequencerPtr = GetSequencer();
+			if (SequencerPtr.IsValid() && MovieSceneChangedDelegate.IsValid())
 			{
-				// The user has been notified, remove the delegate
-				TSharedPtr<ISequencer> SequencerPtr = GetSequencer();
-				if (SequencerPtr.IsValid() && MovieSceneChangedDelegate.IsValid())
-				{
-					SequencerPtr->OnMovieSceneDataChanged().Remove(MovieSceneChangedDelegate);
-					MovieSceneChangedDelegate.Reset();
-				}
+				SequencerPtr->OnMovieSceneDataChanged().Remove(MovieSceneChangedDelegate);
+				MovieSceneChangedDelegate.Reset();
 			}
 		}
-	});
+	}
 }
 
 bool FAudioTrackEditor::SequenceContainsAudioTrack(const UMovieSceneSequence* InSequence)
diff --git a/Engine/Source/Editor/MovieSceneTools/Public/TrackEditors/AudioTrackEditor.h b/Engine/Source/Editor/MovieSceneTools/Public/TrackEditors/AudioTrackEditor.h
--- a/Engine/Source/Editor/MovieSceneTools/Public/TrackEditors/AudioTrackEditor.h
+++ b/Engine/Source/Editor/MovieSceneTools/Public/TrackEditors/AudioTrackEditor.h
@@ -52,6 +52,8 @@ public:
 
 	// ISequencerTrackEditor interface
 
+	virtual void OnInitialize() override;
+	virtual void OnRelease() override;
 	virtual void BuildAddTrackMenu(FMenuBuilder& MenuBuilder) override;
 	virtual void BuildObjectBindingTrackMenu(FMenuBuilder& MenuBuilder, const TArray<FGuid>& ObjectBindings, const UClass* ObjectClass) override;
 	virtual TSharedPtr<SWidget> BuildOutlinerEditWidget(const FGuid& ObjectBinding, UMovieSceneTrack* Track, const FBuildEditWidgetParams& Params) override;
@@ -96,8 +98,11 @@ private:
 	/** Attached audio asset enter pressed */
 	void OnAttachedAudioEnterPressed(const TArray<FAssetData>& AssetData, TArray<FGuid> ObjectBindings);
 
-	/** Registers a delegate with the given sequencer for monitoring edits */
-	void RegisterMovieSceneChangedDelegate(TSharedRef<ISequencer> InSequencer);
+	/** Registers a delegate with the sequencer for monitoring edits */
+	void RegisterMovieSceneChangedDelegate();
+
+	/** Movie Scene Data Changed Delegate */
+	void OnMovieSceneDataChanged(EMovieSceneDataChangeType InChangeType);
 
 	/** Returns true if the given Sequence or any subsequence contains an audio track */
 	bool SequenceContainsAudioTrack(const UMovieSceneSequence* InSequence);
