Sequencer: Set font and padding for string key editor

diff --git a/Engine/Source/Editor/MovieSceneTools/Private/CurveKeyEditors/SStringCurveKeyEditor.cpp b/Engine/Source/Editor/MovieSceneTools/Private/CurveKeyEditors/SStringCurveKeyEditor.cpp
--- a/Engine/Source/Editor/MovieSceneTools/Private/CurveKeyEditors/SStringCurveKeyEditor.cpp
+++ b/Engine/Source/Editor/MovieSceneTools/Private/CurveKeyEditors/SStringCurveKeyEditor.cpp
@@ -20,6 +20,8 @@ void SStringCurveKeyEditor::Construct(const FArguments& InArgs, const TSequencer
 	ChildSlot
 	[
 		SNew(SEditableTextBox)
+		.Padding(FMargin(0.0, 2.0, 0.0, 2.0))
+		.Font(FAppStyle::GetFontStyle("Sequencer.AnimationOutliner.RegularFont"))
 		.MinDesiredWidth(10.f)
 		.SelectAllTextWhenFocused(true)
 		.Text(this, &SStringCurveKeyEditor::GetText)
