Sequencer- Add support for animating material parameters on volumetric cloud component materials.

diff --git a/Engine/Source/Editor/MovieSceneTools/Private/TrackEditors/MaterialTrackEditor.cpp b/Engine/Source/Editor/MovieSceneTools/Private/TrackEditors/MaterialTrackEditor.cpp
--- a/Engine/Source/Editor/MovieSceneTools/Private/TrackEditors/MaterialTrackEditor.cpp
+++ b/Engine/Source/Editor/MovieSceneTools/Private/TrackEditors/MaterialTrackEditor.cpp
@@ -16,6 +16,7 @@
 #include "Engine/Selection.h"
 #include "ISequencerModule.h"
 #include "Components/MeshComponent.h"
+#include "Components/VolumetricCloudComponent.h"
 
 
 #define LOCTEXT_NAMESPACE "MaterialTrackEditor"
@@ -475,6 +476,12 @@ UMaterialInterface* FComponentMaterialTrackEditor::GetMaterialInterfaceForTrack(
 			return DecalComponent->GetDecalMaterial();
 		}
 		break;
+	case EComponentMaterialType::VolumetricCloudMaterial:
+		if (UVolumetricCloudComponent* CloudComponent = Cast<UVolumetricCloudComponent>(Object))
+		{
+			return CloudComponent->GetMaterial();
+		}
+		break;
 	default:
 		break;
 	}
@@ -483,7 +490,7 @@ UMaterialInterface* FComponentMaterialTrackEditor::GetMaterialInterfaceForTrack(
 
 void FComponentMaterialTrackEditor::ExtendObjectBindingTrackMenu(TSharedRef<FExtender> Extender, const TArray<FGuid>& ObjectBindings, const UClass* ObjectClass)
 {
-	if (ObjectClass->IsChildOf(UPrimitiveComponent::StaticClass()) || ObjectClass->IsChildOf(UDecalComponent::StaticClass()))
+	if (ObjectClass->IsChildOf(UPrimitiveComponent::StaticClass()) || ObjectClass->IsChildOf(UDecalComponent::StaticClass()) || ObjectClass->IsChildOf(UVolumetricCloudComponent::StaticClass()))
 	{
 		Extender->AddMenuExtension(SequencerMenuExtensionPoints::AddTrackMenu_PropertiesSection, EExtensionHook::Before, nullptr, FMenuExtensionDelegate::CreateSP(this, &FComponentMaterialTrackEditor::ConstructObjectBindingTrackMenu, ObjectBindings));
 	}
@@ -576,6 +583,25 @@ void FComponentMaterialTrackEditor::ConstructObjectBindingTrackMenu(FMenuBuilder
 			MenuBuilder.EndSection();
 		}
 	}
+	else if (UVolumetricCloudComponent* CloudComponent = Cast<UVolumetricCloudComponent>(SceneComponent))
+	{
+		if (UMaterialInterface* CloudMaterial = CloudComponent->GetMaterial())
+		{
+			MenuBuilder.BeginSection("Materials", LOCTEXT("MaterialSection", "Material Parameters"));
+			{
+				FComponentMaterialInfo MaterialInfo{ FName(), 0, EComponentMaterialType::VolumetricCloudMaterial };
+				const bool bAlreadyExists = Algo::FindBy(Binding->GetTracks(), MaterialInfo, GetMaterialInfoForTrack) != nullptr;
+				if (!bAlreadyExists)
+				{
+					FUIAction AddComponentMaterialAction(FExecuteAction::CreateRaw(this, &FComponentMaterialTrackEditor::HandleAddComponentMaterialActionExecute, SceneComponent, MaterialInfo));
+					FText AddCloudMaterialLabel = FText::Format(LOCTEXT("AddCloudMaterialLabelFormat", "Volumetric Cloud: {0}"), FText::FromString(CloudMaterial->GetName()));
+					FText AddCloudMaterialToolTip = FText::Format(LOCTEXT("AddCloudMaterialToolTipFormat", "Add volumetric cloud material {0}"), FText::FromString(CloudMaterial->GetName()));
+					MenuBuilder.AddMenuEntry(AddCloudMaterialLabel, AddCloudMaterialToolTip, FSlateIcon(), AddComponentMaterialAction);
+				}
+			}
+			MenuBuilder.EndSection();
+		}
+	}
 }
 
 void FComponentMaterialTrackEditor::HandleAddComponentMaterialActionExecute(USceneComponent* Component, FComponentMaterialInfo MaterialInfo)
@@ -647,7 +673,10 @@ void FComponentMaterialTrackEditor::HandleAddComponentMaterialActionExecute(USce
 				break;
 			case EComponentMaterialType::DecalMaterial:
 				TrackDisplayName = FText::Format(LOCTEXT("DecalMaterialTrackName", "Decal Material {0}"), MaterialInterface ? FText::FromString(MaterialInterface->GetName()) : FText());
-				break;
+				break;			
+			case EComponentMaterialType::VolumetricCloudMaterial:
+					TrackDisplayName = FText::Format(LOCTEXT("CloudMaterialTrackName", "Volumetric Cloud Material {0}"), MaterialInterface ? FText::FromString(MaterialInterface->GetName()) : FText());
+					break;
 			default:
 				break;
 
diff --git a/Engine/Source/Editor/MovieSceneTools/Private/TrackEditors/PrimitiveMaterialTrackEditor.cpp b/Engine/Source/Editor/MovieSceneTools/Private/TrackEditors/PrimitiveMaterialTrackEditor.cpp
--- a/Engine/Source/Editor/MovieSceneTools/Private/TrackEditors/PrimitiveMaterialTrackEditor.cpp
+++ b/Engine/Source/Editor/MovieSceneTools/Private/TrackEditors/PrimitiveMaterialTrackEditor.cpp
@@ -9,6 +9,7 @@
 #include "Materials/MaterialInterface.h"
 #include "Algo/Find.h"
 #include "Components/MeshComponent.h"
+#include "Components/VolumetricCloudComponent.h"
 
 
 #define LOCTEXT_NAMESPACE "PrimitiveMaterialTrackEditor"
@@ -114,6 +115,22 @@ void FPrimitiveMaterialTrackEditor::ConstructObjectBindingTrackMenu(FMenuBuilder
 		}
 		MenuBuilder.EndSection();
 	}
+	else if (UVolumetricCloudComponent* CloudComponent = Cast<UVolumetricCloudComponent>(SceneComponent))
+	{
+		MenuBuilder.BeginSection("Materials", LOCTEXT("MaterialSection", "Material Parameters"));
+		{
+			FComponentMaterialInfo MaterialInfo{ FName(), 0, EComponentMaterialType::VolumetricCloudMaterial };
+			const bool bAlreadyExists = Algo::FindBy(Binding->GetTracks(), MaterialInfo, GetMaterialInfoForTrack) != nullptr;
+			if (!bAlreadyExists)
+			{
+				FUIAction AddMaterialSwitcherAction(FExecuteAction::CreateSP(this, &FPrimitiveMaterialTrackEditor::CreateTrackForElement, ObjectBindings, MaterialInfo));
+				FText CloudMaterialSwitcherLabel = LOCTEXT("CloudMaterialSwitcher_Format", "Volumetric Cloud Material Switcher");
+				FText CloudMaterialSwitcherTooltip = LOCTEXT("CloudMaterialSwitcherTooltip_Format", "Add volumetric cloud material switcher");
+				MenuBuilder.AddMenuEntry(CloudMaterialSwitcherLabel, CloudMaterialSwitcherTooltip, FSlateIcon(), AddMaterialSwitcherAction);
+			}
+		}
+		MenuBuilder.EndSection();
+	}
 }
 
 void FPrimitiveMaterialTrackEditor::CreateTrackForElement(TArray<FGuid> ObjectBindingIDs, FComponentMaterialInfo MaterialInfo)
@@ -143,7 +160,10 @@ void FPrimitiveMaterialTrackEditor::CreateTrackForElement(TArray<FGuid> ObjectBi
 			break;
 		case EComponentMaterialType::DecalMaterial:
 			TrackDisplayName = LOCTEXT("DecalMaterialSwitcherTrackName", "Decal Material");
-			break;
+			break;		
+		case EComponentMaterialType::VolumetricCloudMaterial:
+				TrackDisplayName = LOCTEXT("CloudMaterialSwitcherTrackName", "Volumetric Cloud Material");
+				break;
 		default:
 			break;
 
diff --git a/Engine/Source/Runtime/Engine/Classes/Components/VolumetricCloudComponent.h b/Engine/Source/Runtime/Engine/Classes/Components/VolumetricCloudComponent.h
--- a/Engine/Source/Runtime/Engine/Classes/Components/VolumetricCloudComponent.h
+++ b/Engine/Source/Runtime/Engine/Classes/Components/VolumetricCloudComponent.h
@@ -177,6 +177,8 @@ class UVolumetricCloudComponent : public USceneComponent
 	UFUNCTION(BlueprintCallable, Category = "Rendering")
 	ENGINE_API void SetMaterial(UMaterialInterface* NewValue);
 
+	ENGINE_API UMaterialInterface* GetMaterial() const { return Material; }
+
 	// Deprecated functions but still valid because they forward data correctly.
 	UE_DEPRECATED(5.0, "This function has been replaced by SetReflectionViewSampleCountScale.")
 	UFUNCTION(BlueprintCallable, Category = "Rendering", meta = (DeprecatedFunction, DeprecationMessage = "This function has been replaced by SetReflectionViewSampleCountScale."))
diff --git a/Engine/Source/Runtime/MovieSceneTracks/Private/Systems/MovieSceneComponentMaterialSystem.cpp b/Engine/Source/Runtime/MovieSceneTracks/Private/Systems/MovieSceneComponentMaterialSystem.cpp
--- a/Engine/Source/Runtime/MovieSceneTracks/Private/Systems/MovieSceneComponentMaterialSystem.cpp
+++ b/Engine/Source/Runtime/MovieSceneTracks/Private/Systems/MovieSceneComponentMaterialSystem.cpp
@@ -17,6 +17,7 @@
 #include "Components/PrimitiveComponent.h"
 #include "Components/DecalComponent.h"
 #include "Components/MeshComponent.h"
+#include "Components/VolumetricCloudComponent.h"
 
 #include UE_INLINE_GENERATED_CPP_BY_NAME(MovieSceneComponentMaterialSystem)
 
@@ -76,6 +77,12 @@ UMaterialInterface* FComponentMaterialAccessor::GetMaterial() const
 			return DecalComponent->GetDecalMaterial();
 		}
 		break;
+	case EComponentMaterialType::VolumetricCloudMaterial:
+		if (UVolumetricCloudComponent* CloudComponent = Cast<UVolumetricCloudComponent>(Object))
+		{
+			return CloudComponent->GetMaterial();
+		}
+		break;
 	default:
 		break;
 	}
@@ -114,6 +121,12 @@ void FComponentMaterialAccessor::SetMaterial(UMaterialInterface* InMaterial) con
 			DecalComponent->SetDecalMaterial(InMaterial);
 		}
 		break;
+	case EComponentMaterialType::VolumetricCloudMaterial:
+		if (UVolumetricCloudComponent* CloudComponent = Cast<UVolumetricCloudComponent>(Object))
+		{
+			CloudComponent->SetMaterial(InMaterial);
+		}
+		break;
 	default:
 		break;
 	}
@@ -164,6 +177,14 @@ UMaterialInstanceDynamic* FComponentMaterialAccessor::CreateDynamicMaterial(UMat
 			return DecalComponent->CreateDynamicMaterialInstance();
 		}
 		break;
+	case EComponentMaterialType::VolumetricCloudMaterial:
+		if (UVolumetricCloudComponent* CloudComponent = Cast<UVolumetricCloudComponent>(Object))
+		{
+			UMaterialInstanceDynamic* Result = MakeDynamicMaterial();
+			CloudComponent->SetMaterial(Result);
+			return Result;
+		}
+		break;
 	default:
 		break;
 	}
diff --git a/Engine/Source/Runtime/MovieSceneTracks/Private/Tracks/MovieSceneMaterialTrack.cpp b/Engine/Source/Runtime/MovieSceneTracks/Private/Tracks/MovieSceneMaterialTrack.cpp
--- a/Engine/Source/Runtime/MovieSceneTracks/Private/Tracks/MovieSceneMaterialTrack.cpp
+++ b/Engine/Source/Runtime/MovieSceneTracks/Private/Tracks/MovieSceneMaterialTrack.cpp
@@ -13,6 +13,7 @@
 #include "Components/PrimitiveComponent.h"
 #include "Components/MeshComponent.h"
 #include "Materials/MaterialInterface.h"
+#include "Components/VolumetricCloudComponent.h"
 
 #include UE_INLINE_GENERATED_CPP_BY_NAME(MovieSceneMaterialTrack)
 
@@ -365,7 +366,21 @@ FText UMovieSceneComponentMaterialTrack::GetDisplayNameToolTipText(const FMovieS
 						return LOCTEXT("DecalMaterialTrackTooltip_MissingMaterial", "No decal material could be found");
 					}
 				}
-				break;
+				break;			
+			case EComponentMaterialType::VolumetricCloudMaterial:
+					if (UVolumetricCloudComponent* CloudComponent = Cast<UVolumetricCloudComponent>(WeakObject.Get()))
+					{
+						Material = CloudComponent->GetMaterial();
+						if (Material)
+						{
+							return FText::Format(LOCTEXT("CloudMaterialTrackTooltip", "Material parameter track for volumetric cloud material {0}"), FText::FromString(Material->GetName()));
+						}
+						else
+						{
+							return LOCTEXT("CloudMaterialTrackTooltip_MissingMaterial", "No volumetric cloud material could be found");
+						}
+					}
+					break;
 			default:
 				break;
 			}
@@ -462,6 +477,19 @@ FSlateColor UMovieSceneComponentMaterialTrack::GetLabelColor(const FMovieSceneLa
 				{
 					return GetDimmedColor(FLinearColor::Red);
 				}
+			case EComponentMaterialType::VolumetricCloudMaterial:
+				if (UVolumetricCloudComponent* CloudComponent = Cast<UVolumetricCloudComponent>(WeakObject.Get()))
+				{
+					UMaterialInterface* Material = CloudComponent->GetMaterial();
+					if (!Material)
+					{
+						return GetDimmedColor(FLinearColor::Red);
+					}
+				}
+				else
+				{
+					return GetDimmedColor(FLinearColor::Red);
+				}
 				break;
 			default:
 				break;
diff --git a/Engine/Source/Runtime/MovieSceneTracks/Public/Tracks/MovieSceneMaterialTrack.h b/Engine/Source/Runtime/MovieSceneTracks/Public/Tracks/MovieSceneMaterialTrack.h
--- a/Engine/Source/Runtime/MovieSceneTracks/Public/Tracks/MovieSceneMaterialTrack.h
+++ b/Engine/Source/Runtime/MovieSceneTracks/Public/Tracks/MovieSceneMaterialTrack.h
@@ -23,7 +23,9 @@ enum class EComponentMaterialType
 	/* An overlay material on a mesh component*/
 	OverlayMaterial,
 	/* A decal material*/
-	DecalMaterial
+	DecalMaterial,
+	/* Volumetric Cloud Material*/
+	VolumetricCloudMaterial
 };
 
 /**
