Sequencer: fix displaying of looping sub-sections

We were iterating until our offset from the start of the section reached the end time of the section... which meant we iterated for the whole length of the section plus its offset from the origin. Now we just iterate for the *length* of the section.

diff --git a/Engine/Source/Editor/MovieSceneTools/Private/TrackEditors/SubTrackEditorBase.cpp b/Engine/Source/Editor/MovieSceneTools/Private/TrackEditors/SubTrackEditorBase.cpp
--- a/Engine/Source/Editor/MovieSceneTools/Private/TrackEditors/SubTrackEditorBase.cpp
+++ b/Engine/Source/Editor/MovieSceneTools/Private/TrackEditors/SubTrackEditorBase.cpp
@@ -253,9 +253,6 @@ void FSubSectionPainterUtil::DoPaintNonLoopingSection(const UMovieSceneSubSectio
 
 void FSubSectionPainterUtil::DoPaintLoopingSection(const UMovieSceneSubSection& SectionObject, const UMovieSceneSequence& InnerSequence, FSequencerSectionPainter& InPainter, ESlateDrawEffect DrawEffects)
 {
-    const FFrameNumber SectionStartFrame = SectionObject.GetInclusiveStartFrame();
-    const FFrameNumber SectionEndFrame   = SectionObject.GetExclusiveEndFrame();
-
     const TRange<FFrameNumber> SectionRange = SectionObject.GetRange();
     const int32 SectionSize = UE::MovieScene::DiscreteSize(SectionRange);
     const float PixelsPerFrame = InPainter.SectionGeometry.Size.X / float(SectionSize);
@@ -280,7 +277,7 @@ void FSubSectionPainterUtil::DoPaintLoopingSection(const UMovieSceneSubSection&
     // Draw separators where the sub-sequence is looping. To be consistent with the non-looping case, we draw a red and green
     // separator back to back.
     uint32 MaxLoopBoundaries = 100;
-    while (CurOffsetFrame < SectionEndFrame)
+    while (CurOffsetFrame < SectionSize)
     {
         const int32 CurOffset = CurOffsetFrame.Value;
         FSlateDrawElement::MakeBox(
