Sequencer: Make animation sections adhere to view density heights

diff --git a/Engine/Source/Editor/MovieSceneTools/Private/TrackEditors/SkeletalAnimationTrackEditor.cpp b/Engine/Source/Editor/MovieSceneTools/Private/TrackEditors/SkeletalAnimationTrackEditor.cpp
--- a/Engine/Source/Editor/MovieSceneTools/Private/TrackEditors/SkeletalAnimationTrackEditor.cpp
+++ b/Engine/Source/Editor/MovieSceneTools/Private/TrackEditors/SkeletalAnimationTrackEditor.cpp
@@ -718,7 +718,7 @@ TOptional<FFrameTime> FSkeletalAnimationSection::GetSectionTime(FSequencerSectio
 
 float FSkeletalAnimationSection::GetSectionHeight(const UE::Sequencer::FViewDensityInfo& ViewDensity) const
 {
-	return SkeletalAnimationEditorConstants::AnimationTrackHeight;
+	return ViewDensity.UniformHeight.Get(SkeletalAnimationEditorConstants::AnimationTrackHeight);
 }
 
 
