Sequencer: Corrected section heights for shot sections that would clip thumbnails

diff --git a/Engine/Source/Editor/MovieSceneTools/Private/Sections/CinematicShotSection.cpp b/Engine/Source/Editor/MovieSceneTools/Private/Sections/CinematicShotSection.cpp
--- a/Engine/Source/Editor/MovieSceneTools/Private/Sections/CinematicShotSection.cpp
+++ b/Engine/Source/Editor/MovieSceneTools/Private/Sections/CinematicShotSection.cpp
@@ -68,7 +68,7 @@ FText FCinematicShotSection::GetSectionTitle() const
 
 float FCinematicShotSection::GetSectionHeight(const UE::Sequencer::FViewDensityInfo& ViewDensity) const
 {
-	return FViewportThumbnailSection::GetSectionHeight(ViewDensity) + 2*9.f;
+	return FViewportThumbnailSection::GetSectionHeight(ViewDensity) + 2*13.f;
 }
 
 FMargin FCinematicShotSection::GetContentPadding() const
