Sequencer: Introduced a new column that only shows an 'Add keyframe' button for more condensed views

This column is mutually exclusive with the Key Frame Navigation column

diff --git a/Engine/Source/Editor/MovieSceneTools/Private/TrackEditors/CameraCutTrackEditor.cpp b/Engine/Source/Editor/MovieSceneTools/Private/TrackEditors/CameraCutTrackEditor.cpp
--- a/Engine/Source/Editor/MovieSceneTools/Private/TrackEditors/CameraCutTrackEditor.cpp
+++ b/Engine/Source/Editor/MovieSceneTools/Private/TrackEditors/CameraCutTrackEditor.cpp
@@ -199,9 +199,24 @@ TSharedPtr<SWidget> FCameraCutTrackEditor::BuildOutlinerColumnWidget(const FBuil
 			Params.ViewModel);
 	}
 
-	// Show the camera lock button in the Edit area if the nav column is disabled
-	const bool bEditColumn = (ColumnName == FCommonOutlinerNames::Edit && !Params.TreeViewRow->IsColumnVisible(FCommonOutlinerNames::Nav));
-	if (ColumnName == FCommonOutlinerNames::Nav || bEditColumn)
+	bool bAddCameraLock = false;
+	if (ColumnName == FCommonOutlinerNames::Nav)
+	{
+		bAddCameraLock = true;
+	}
+	else if (ColumnName == FCommonOutlinerNames::KeyFrame)
+	{
+		// Add the camera lock button to the keyframe column if Nav is disabled
+		bAddCameraLock = Params.TreeViewRow->IsColumnVisible(FCommonOutlinerNames::Nav) == false;
+	}
+	else if (ColumnName == FCommonOutlinerNames::Edit)
+	{
+		// Add the camera lock button to the edit column if both Nav and KeyFrame are disabled
+		bAddCameraLock = Params.TreeViewRow->IsColumnVisible(FCommonOutlinerNames::Nav) == false &&
+			Params.TreeViewRow->IsColumnVisible(FCommonOutlinerNames::KeyFrame) == false;
+	}
+
+	if (bAddCameraLock)
 	{
 		TSharedRef<SWidget> Button = SNew(SCheckBox)
 		.Style(FAppStyle::Get(), "Sequencer.Outliner.ToggleButton")
@@ -215,8 +230,9 @@ TSharedPtr<SWidget> FCameraCutTrackEditor::BuildOutlinerColumnWidget(const FBuil
 			.Image(FAppStyle::GetBrush("Sequencer.Outliner.CameraLock"))
 		];
 
-		if (bEditColumn)
+		if (ColumnName == FCommonOutlinerNames::Edit)
 		{
+			// Needs to be left aligned in the edit column because this column slot is set to fill
 			return SNew(SBox)
 			.HAlign(HAlign_Left)
 			.Padding(4.f, 0.f)
diff --git a/Engine/Source/Editor/MovieSceneTools/Private/TrackEditors/CinematicShotTrackEditor.cpp b/Engine/Source/Editor/MovieSceneTools/Private/TrackEditors/CinematicShotTrackEditor.cpp
--- a/Engine/Source/Editor/MovieSceneTools/Private/TrackEditors/CinematicShotTrackEditor.cpp
+++ b/Engine/Source/Editor/MovieSceneTools/Private/TrackEditors/CinematicShotTrackEditor.cpp
@@ -12,6 +12,7 @@
 #include "Sections/MovieSceneCinematicShotSection.h"
 #include "SequencerSettings.h"
 #include "MVVM/Views/ViewUtilities.h"
+#include "MVVM/ViewModels/TrackRowModel.h"
 #include "MVVM/ViewModels/OutlinerColumns/OutlinerColumnTypes.h"
 #include "MVVM/Extensions/ITrackExtension.h"
 #include "TrackEditorThumbnail/TrackEditorThumbnailPool.h"
@@ -86,34 +87,53 @@ TSharedPtr<SWidget> FCinematicShotTrackEditor::BuildOutlinerColumnWidget(const F
 			Params.ViewModel);
 	}
 
-	// Show the camera lock button in the Edit area if the nav column is disabled
-	const bool bEditColumn = (ColumnName == FCommonOutlinerNames::Edit && !Params.TreeViewRow->IsColumnVisible(FCommonOutlinerNames::Nav));
-	if (ColumnName == FCommonOutlinerNames::Nav || bEditColumn)
+	if (!Params.ViewModel->IsA<FTrackRowModel>())
 	{
-		TSharedRef<SWidget> Button = SNew(SCheckBox)
-		.Style(FAppStyle::Get(), "Sequencer.Outliner.ToggleButton")
-		.Type(ESlateCheckBoxType::ToggleButton)
-		.IsFocusable(false)
-		.IsChecked(this, &FCinematicShotTrackEditor::AreShotsLocked)
-		.OnCheckStateChanged(this, &FCinematicShotTrackEditor::OnLockShotsClicked)
-		.ToolTipText(this, &FCinematicShotTrackEditor::GetLockShotsToolTip)
-		[
-			SNew(SImage)
-				.Image(FAppStyle::GetBrush("Sequencer.Outliner.CameraLock"))
-		];
+		bool bAddCameraLock = false;
+		if (ColumnName == FCommonOutlinerNames::Nav)
+		{
+			bAddCameraLock = true;
+		}
+		else if (ColumnName == FCommonOutlinerNames::KeyFrame)
+		{
+			// Add the camera lock button to the keyframe column if Nav is disabled
+			bAddCameraLock = Params.TreeViewRow->IsColumnVisible(FCommonOutlinerNames::Nav) == false;
+		}
+		else if (ColumnName == FCommonOutlinerNames::Edit)
+		{
+			// Add the camera lock button to the edit column if both Nav and KeyFrame are disabled
+			bAddCameraLock = Params.TreeViewRow->IsColumnVisible(FCommonOutlinerNames::Nav) == false &&
+				Params.TreeViewRow->IsColumnVisible(FCommonOutlinerNames::KeyFrame) == false;
+		}
 
-		if (bEditColumn)
+		if (bAddCameraLock)
 		{
-			return SNew(SBox)
-			.HAlign(HAlign_Left)
-			.Padding(4.f, 0.f)
+			TSharedRef<SWidget> Button = SNew(SCheckBox)
+			.Style(FAppStyle::Get(), "Sequencer.Outliner.ToggleButton")
+			.Type(ESlateCheckBoxType::ToggleButton)
+			.IsFocusable(false)
+			.IsChecked(this, &FCinematicShotTrackEditor::AreShotsLocked)
+			.OnCheckStateChanged(this, &FCinematicShotTrackEditor::OnLockShotsClicked)
+			.ToolTipText(this, &FCinematicShotTrackEditor::GetLockShotsToolTip)
 			[
-				Button
+				SNew(SImage)
+				.Image(FAppStyle::GetBrush("Sequencer.Outliner.CameraLock"))
 			];
-		}
-		else
-		{
-			return Button;
+
+			if (ColumnName == FCommonOutlinerNames::Edit)
+			{
+				// Needs to be left aligned in the edit column because this column slot is set to fill
+				return SNew(SBox)
+				.HAlign(HAlign_Left)
+				.Padding(4.f, 0.f)
+				[
+					Button
+				];
+			}
+			else
+			{
+				return Button;
+			}
 		}
 	}
 
diff --git a/Engine/Source/Editor/MovieSceneTools/Private/TrackEditors/TransformTrackEditor.cpp b/Engine/Source/Editor/MovieSceneTools/Private/TrackEditors/TransformTrackEditor.cpp
--- a/Engine/Source/Editor/MovieSceneTools/Private/TrackEditors/TransformTrackEditor.cpp
+++ b/Engine/Source/Editor/MovieSceneTools/Private/TrackEditors/TransformTrackEditor.cpp
@@ -502,12 +502,31 @@ void F3DTransformTrackEditor::BuildObjectBindingColumnWidgets(TFunctionRef<TShar
 {
 	using namespace UE::Sequencer;
 
-	if (InColumnName == FCommonOutlinerNames::Nav || (InColumnName == FCommonOutlinerNames::Edit && !InParams.TreeViewRow->IsColumnVisible(FCommonOutlinerNames::Nav)))
+	bool bAddCameraLock = false;
+	if (InColumnName == FCommonOutlinerNames::Nav)
 	{
+		bAddCameraLock = true;
+	}
+	else if (InColumnName == FCommonOutlinerNames::KeyFrame)
+	{
+		// Add the camera lock button to the keyframe column if Nav is disabled
+		bAddCameraLock = InParams.TreeViewRow->IsColumnVisible(FCommonOutlinerNames::Nav) == false;
+	}
+	else if (InColumnName == FCommonOutlinerNames::Edit)
+	{
+		// Add the camera lock button to the edit column if both Nav and KeyFrame are disabled
+		bAddCameraLock = InParams.TreeViewRow->IsColumnVisible(FCommonOutlinerNames::Nav) == false &&
+			InParams.TreeViewRow->IsColumnVisible(FCommonOutlinerNames::KeyFrame) == false;
+	}
+
+	if (bAddCameraLock)
+	{
+		const bool bEditColumn = InColumnName == FCommonOutlinerNames::Edit;
 		FGuid ObjectGuid = ObjectBinding->GetObjectGuid();
 		GetEditBox()->AddSlot()
 		.VAlign(VAlign_Center)
-		.HAlign(HAlign_Center)
+		.HAlign(bEditColumn ? HAlign_Left : HAlign_Center)
+		.Padding(bEditColumn ? FMargin(4.f, 0.f) : FMargin(0.f))
 		.AutoWidth()
 		[
 			SNew(SCheckBox)
diff --git a/Engine/Source/Editor/Sequencer/Private/MVVM/ViewModels/CategoryModel.cpp b/Engine/Source/Editor/Sequencer/Private/MVVM/ViewModels/CategoryModel.cpp
--- a/Engine/Source/Editor/Sequencer/Private/MVVM/ViewModels/CategoryModel.cpp
+++ b/Engine/Source/Editor/Sequencer/Private/MVVM/ViewModels/CategoryModel.cpp
@@ -224,6 +224,24 @@ TSharedPtr<SWidget> FCategoryGroupModel::CreateOutlinerViewForColumn(const FCrea
 
 	}
 
+	if (InColumnName == FCommonOutlinerNames::KeyFrame)
+	{
+		EKeyNavigationButtons Buttons = EKeyNavigationButtons::AddKey;
+
+		return SNew(SSequencerKeyNavigationButtons, SharedThis(this), Editor->GetSequencer())
+			.Buttons(Buttons);
+	}
+
+	if (InColumnName == FCommonOutlinerNames::Nav)
+	{
+		EKeyNavigationButtons Buttons = InParams.TreeViewRow->IsColumnVisible(FCommonOutlinerNames::KeyFrame)
+			? EKeyNavigationButtons::NavOnly
+			: EKeyNavigationButtons::All;
+
+		return SNew(SSequencerKeyNavigationButtons, SharedThis(this), Editor->GetSequencer())
+			.Buttons(Buttons);
+	}
+
 	if (InColumnName == FCommonOutlinerNames::Nav)
 	{
 		return SNew(SSequencerKeyNavigationButtons, SharedThis(this), Editor->GetSequencer());
diff --git a/Engine/Source/Editor/Sequencer/Private/MVVM/ViewModels/ChannelModel.cpp b/Engine/Source/Editor/Sequencer/Private/MVVM/ViewModels/ChannelModel.cpp
--- a/Engine/Source/Editor/Sequencer/Private/MVVM/ViewModels/ChannelModel.cpp
+++ b/Engine/Source/Editor/Sequencer/Private/MVVM/ViewModels/ChannelModel.cpp
@@ -845,9 +845,21 @@ TSharedPtr<SWidget> FChannelGroupOutlinerModel::CreateOutlinerViewForColumn(cons
 			.Visibility(this, &FChannelGroupOutlinerModel::GetKeyEditorVisibility);
 	}
 
+	if (InColumnName == FCommonOutlinerNames::KeyFrame)
+	{
+		EKeyNavigationButtons Buttons = EKeyNavigationButtons::AddKey;
+		return SNew(SSequencerKeyNavigationButtons, SharedThis(this), Editor->GetSequencer())
+			.Buttons(Buttons);
+	}
+
 	if (InColumnName == FCommonOutlinerNames::Nav)
 	{
-		return SNew(SSequencerKeyNavigationButtons, SharedThis(this), Editor->GetSequencer());
+		EKeyNavigationButtons Buttons = InParams.TreeViewRow->IsColumnVisible(FCommonOutlinerNames::KeyFrame)
+			? EKeyNavigationButtons::NavOnly
+			: EKeyNavigationButtons::All;
+
+		return SNew(SSequencerKeyNavigationButtons, SharedThis(this), Editor->GetSequencer())
+			.Buttons(Buttons);
 	}
 
 	if (InColumnName == FCommonOutlinerNames::ColorPicker)
diff --git a/Engine/Source/Editor/Sequencer/Private/MVVM/Views/SSequencerKeyNavigationButtons.h b/Engine/Source/Editor/Sequencer/Private/MVVM/Views/SSequencerKeyNavigationButtons.h
--- a/Engine/Source/Editor/Sequencer/Private/MVVM/Views/SSequencerKeyNavigationButtons.h
+++ b/Engine/Source/Editor/Sequencer/Private/MVVM/Views/SSequencerKeyNavigationButtons.h
@@ -29,7 +29,8 @@ class SSequencerKeyNavigationButtons
 {
 public:
 
-	SLATE_BEGIN_ARGS(SSequencerKeyNavigationButtons){}
+	SLATE_BEGIN_ARGS(SSequencerKeyNavigationButtons) : _Buttons(EKeyNavigationButtons::All) {}
+		SLATE_ARGUMENT(EKeyNavigationButtons, Buttons)
 	SLATE_END_ARGS()
 
 	void Construct(const FArguments& InArgs, const TSharedPtr<FViewModel>& InModel, TSharedPtr<ISequencer> InSequencer)
@@ -40,6 +41,7 @@ public:
 		ChildSlot
 		[
 			SNew(SKeyNavigationButtons, InModel)
+			.Buttons(InArgs._Buttons)
 			.AddKeyToolTip(FText::Format(LOCTEXT("AddKeyButton", "Add a new key at the current time ({0})"), FSequencerCommands::Get().SetKey->GetInputText()))
 			.PreviousKeyToolTip(FText::Format(LOCTEXT("PreviousKeyButton", "Set the time to the previous key ({0})"), FSequencerCommands::Get().StepToPreviousKey->GetInputText()))
 			.NextKeyToolTip(FText::Format(LOCTEXT("NextKeyButton", "Set the time to the next key ({0})"), FSequencerCommands::Get().StepToNextKey->GetInputText()))
diff --git a/Engine/Source/Editor/Sequencer/Private/MovieSceneTrackEditor.cpp b/Engine/Source/Editor/Sequencer/Private/MovieSceneTrackEditor.cpp
--- a/Engine/Source/Editor/Sequencer/Private/MovieSceneTrackEditor.cpp
+++ b/Engine/Source/Editor/Sequencer/Private/MovieSceneTrackEditor.cpp
@@ -22,6 +22,28 @@
 #include "MVVM/Views/SOutlinerItemViewBase.h"
 #include "SKeyAreaEditorSwitcher.h"
 
+namespace UE::Sequencer
+{
+	bool HasKeyableAreas(const FViewModel& ViewModel)
+	{
+		for (TSharedPtr<FChannelGroupModel> ChannelGroup : ViewModel.GetDescendantsOfType<FChannelGroupModel>())
+		{
+			for (const TWeakViewModelPtr<FChannelModel>& WeakChannel : ChannelGroup->GetChannels())
+			{
+				if (TViewModelPtr<FChannelModel> Channel = WeakChannel.Pin())
+				{
+					if (Channel->GetKeyArea()->CanCreateKeyEditor())
+					{
+						return true;
+					}
+				}
+			}
+		}
+
+		return false;
+	}
+}
+
 FMovieSceneTrackEditor::FMovieSceneTrackEditor(TSharedRef<ISequencer> InSequencer)
 	: Sequencer(InSequencer)
 { 
@@ -245,31 +267,29 @@ TSharedPtr<SWidget> FMovieSceneTrackEditor::BuildOutlinerColumnWidget(const FBui
 		return BuildOutlinerEditWidget(ObjectBinding ? ObjectBinding->GetObjectGuid() : FGuid(), Track, EditWidgetParams);
 	}
 
-	if (ColumnName == FCommonOutlinerNames::Nav)
+	if (ColumnName == FCommonOutlinerNames::KeyFrame)
 	{
-		bool bHasKeyableAreas = false;
-		for (TSharedPtr<FChannelGroupModel> ChannelGroup : Params.ViewModel->GetDescendantsOfType<FChannelGroupModel>())
+		if (HasKeyableAreas(*Params.ViewModel))
 		{
-			for (const TWeakViewModelPtr<FChannelModel>& WeakChannel : ChannelGroup->GetChannels())
-			{
-				if (TViewModelPtr<FChannelModel> Channel = WeakChannel.Pin())
-				{
-					if (Channel->GetKeyArea()->CanCreateKeyEditor())
-					{
-						bHasKeyableAreas = true;
-						break;
-					}
-				}
-			}
-			if (bHasKeyableAreas)
-			{
-				break;
-			}
+			EKeyNavigationButtons Buttons = EKeyNavigationButtons::AddKey;
+
+			return SNew(SSequencerKeyNavigationButtons, Params.ViewModel, Editor->GetSequencer())
+				.Buttons(Buttons);
 		}
 
-		if (bHasKeyableAreas)
+		return nullptr;
+	}
+
+	if (ColumnName == FCommonOutlinerNames::Nav)
+	{
+		if (HasKeyableAreas(*Params.ViewModel))
 		{
-			return SNew(SSequencerKeyNavigationButtons, Params.ViewModel, Editor->GetSequencer());
+			EKeyNavigationButtons Buttons = Params.TreeViewRow->IsColumnVisible(FCommonOutlinerNames::KeyFrame)
+				? EKeyNavigationButtons::NavOnly
+				: EKeyNavigationButtons::All;
+
+			return SNew(SSequencerKeyNavigationButtons, Params.ViewModel, Editor->GetSequencer())
+				.Buttons(Buttons);
 		}
 
 		return nullptr;
diff --git a/Engine/Source/Editor/Sequencer/Private/SSequencer.cpp b/Engine/Source/Editor/Sequencer/Private/SSequencer.cpp
--- a/Engine/Source/Editor/Sequencer/Private/SSequencer.cpp
+++ b/Engine/Source/Editor/Sequencer/Private/SSequencer.cpp
@@ -2411,29 +2411,50 @@ void SSequencer::FillViewDensityMenu(FMenuBuilder& InMenuBuilder)
 
 void SSequencer::FillColumnVisibilityMenu(FMenuBuilder& InMenuBuilder)
 {
+	using namespace UE::Sequencer;
+
 	const bool bShouldCloseWindowAfterMenuSelection = true;
 
 	for (FSequencerOutlinerColumnVisibility& ColumnVisibility : OutlinerColumnVisibilities)
 	{
+		auto ToggleVisibility = [this, &ColumnVisibility]
+		{
+			ColumnVisibility.bIsColumnVisible = !ColumnVisibility.bIsColumnVisible;
+			if(ColumnVisibility.bIsColumnVisible)
+			{
+				FName ColumnName = ColumnVisibility.Column->GetColumnName();
+				FSequencerOutlinerColumnVisibility* AutoDisable = nullptr;
+
+				// Auto disable mutually exclusive columns
+				if (ColumnName == FCommonOutlinerNames::Nav)
+				{
+					AutoDisable = Algo::FindBy(this->OutlinerColumnVisibilities, FCommonOutlinerNames::KeyFrame, [](const FSequencerOutlinerColumnVisibility& In) { return In.Column->GetColumnName(); });
+				}
+				else if (ColumnName == FCommonOutlinerNames::KeyFrame)
+				{
+					AutoDisable = Algo::FindBy(this->OutlinerColumnVisibilities, FCommonOutlinerNames::Nav, [](const FSequencerOutlinerColumnVisibility& In) { return In.Column->GetColumnName(); });
+				}
+
+				if (AutoDisable)
+				{
+					AutoDisable->bIsColumnVisible = false;
+				}
+			}
+			
+			this->UpdateOutlinerViewColumns();
+		};
+
+
 		InMenuBuilder.AddMenuEntry(
 			ColumnVisibility.Column->GetColumnLabel(),
 			LOCTEXT("SetColumnVisibilityTooltip", "Enable or disable this outliner column"),
 			FSlateIcon(),
 			FUIAction(
-				FExecuteAction::CreateLambda([this, &ColumnVisibility] {
-					ColumnVisibility.bIsColumnVisible = !ColumnVisibility.bIsColumnVisible;
-					UpdateOutlinerViewColumns();
-					}),
+				FExecuteAction::CreateLambda(ToggleVisibility),
 				FCanExecuteAction(),
-				FIsActionChecked::CreateLambda([this, ColumnVisibility] { 
-						for (const FColumnVisibilitySetting& ColumnSettings : GetSequencerSettings()->GetOutlinerColumnSettings())
-						{
-							if (ColumnSettings.ColumnName == ColumnVisibility.Column.Get()->GetColumnName())
-							{
-								return ColumnSettings.bIsVisible;
-							}
-						}
-						return false;
+				FIsActionChecked::CreateLambda([&ColumnVisibility] { 
+						// Capture by ref here since the array itself cannot be re-allocated while this menu is open
+						return ColumnVisibility.bIsColumnVisible;
 					})
 				),
 				NAME_None,
diff --git a/Engine/Source/Editor/Sequencer/Private/SequencerModule.cpp b/Engine/Source/Editor/Sequencer/Private/SequencerModule.cpp
--- a/Engine/Source/Editor/Sequencer/Private/SequencerModule.cpp
+++ b/Engine/Source/Editor/Sequencer/Private/SequencerModule.cpp
@@ -39,6 +39,7 @@
 #include "MVVM/ViewModels/OutlinerColumns/EditOutlinerColumn.h"
 #include "MVVM/ViewModels/OutlinerColumns/AddOutlinerColumn.h"
 #include "MVVM/ViewModels/OutlinerColumns/NavOutlinerColumn.h"
+#include "MVVM/ViewModels/OutlinerColumns/KeyFrameOutlinerColumn.h"
 #include "MVVM/ViewModels/OutlinerColumns/ColorPickerOutlinerColumn.h"
 
 #include "ToolMenus.h"
@@ -397,6 +398,7 @@ public:
 			AddOutlinerColumnHandle   = RegisterOutlinerColumn(FOnCreateOutlinerColumn::CreateStatic([]{ return TSharedRef<IOutlinerColumn>(MakeShared<FAddOutlinerColumn>()); }));
 
 			// Register right gutter columns
+			KeyFrameOutlinerColumnHandle     = RegisterOutlinerColumn(FOnCreateOutlinerColumn::CreateStatic([]{ return TSharedRef<IOutlinerColumn>(MakeShared<FKeyFrameOutlinerColumn>()); }));
 			NavOutlinerColumnHandle          = RegisterOutlinerColumn(FOnCreateOutlinerColumn::CreateStatic([]{ return TSharedRef<IOutlinerColumn>(MakeShared<FNavOutlinerColumn>()); }));
 			ColorPickerOutlinerColumnHandle  = RegisterOutlinerColumn(FOnCreateOutlinerColumn::CreateStatic([]{ return TSharedRef<IOutlinerColumn>(MakeShared<FColorPickerOutlinerColumn>()); }));
 
@@ -452,6 +454,7 @@ public:
 			UnregisterOutlinerColumn(LabelOutlinerColumnHandle);
 			UnregisterOutlinerColumn(EditOutlinerColumnHandle);
 			UnregisterOutlinerColumn(AddOutlinerColumnHandle);
+			UnregisterOutlinerColumn(KeyFrameOutlinerColumnHandle);
 			UnregisterOutlinerColumn(NavOutlinerColumnHandle);
 			UnregisterOutlinerColumn(ColorPickerOutlinerColumnHandle);
 		}
@@ -643,6 +646,7 @@ private:
 	FDelegateHandle LabelOutlinerColumnHandle;
 	FDelegateHandle EditOutlinerColumnHandle;
 	FDelegateHandle AddOutlinerColumnHandle;
+	FDelegateHandle KeyFrameOutlinerColumnHandle;
 	FDelegateHandle NavOutlinerColumnHandle;
 	FDelegateHandle ColorPickerOutlinerColumnHandle;
 };
diff --git a/Engine/Source/Editor/Sequencer/Public/MVVM/ViewModels/TrackRowModel.h b/Engine/Source/Editor/Sequencer/Public/MVVM/ViewModels/TrackRowModel.h
--- a/Engine/Source/Editor/Sequencer/Public/MVVM/ViewModels/TrackRowModel.h
+++ b/Engine/Source/Editor/Sequencer/Public/MVVM/ViewModels/TrackRowModel.h
@@ -25,7 +25,7 @@ namespace Sequencer
 
 class FSectionModel;
 
-class FTrackRowModel
+class SEQUENCER_API FTrackRowModel
 	: public FMuteSoloOutlinerItemModel
 	, public ITrackAreaExtension
 	, public ILockableExtension
diff --git a/Engine/Source/Editor/SequencerCore/Private/MVVM/ViewModels/OutlinerColumns/ColorPickerOutlinerColumn.cpp b/Engine/Source/Editor/SequencerCore/Private/MVVM/ViewModels/OutlinerColumns/ColorPickerOutlinerColumn.cpp
--- a/Engine/Source/Editor/SequencerCore/Private/MVVM/ViewModels/OutlinerColumns/ColorPickerOutlinerColumn.cpp
+++ b/Engine/Source/Editor/SequencerCore/Private/MVVM/ViewModels/OutlinerColumns/ColorPickerOutlinerColumn.cpp
@@ -9,7 +9,7 @@ FColorPickerOutlinerColumn::FColorPickerOutlinerColumn()
 {
 	Name     = FCommonOutlinerNames::ColorPicker;
 	Label    = NSLOCTEXT("FColorPickerOutlinerColumn", "ColorPickerColumnColorPicker", "ColorPicker");
-	Position = FOutlinerColumnPosition{ 20, EOutlinerColumnGroup::RightGutter };
+	Position = FOutlinerColumnPosition{ 30, EOutlinerColumnGroup::RightGutter };
 	Layout   = FOutlinerColumnLayout{ 6.f, FMargin(0.f), HAlign_Fill, VAlign_Fill, EOutlinerColumnSizeMode::Fixed, EOutlinerColumnFlags::None };
 }
 
diff --git a/Engine/Source/Editor/SequencerCore/Private/MVVM/ViewModels/OutlinerColumns/KeyFrameOutlinerColumn.cpp b/Engine/Source/Editor/SequencerCore/Private/MVVM/ViewModels/OutlinerColumns/KeyFrameOutlinerColumn.cpp
new file mode 100644
--- /dev/null
+++ b/Engine/Source/Editor/SequencerCore/Private/MVVM/ViewModels/OutlinerColumns/KeyFrameOutlinerColumn.cpp
@@ -0,0 +1,21 @@
+// Copyright Epic Games, Inc. All Rights Reserved.
+
+#include "MVVM/ViewModels/OutlinerColumns/KeyFrameOutlinerColumn.h"
+
+namespace UE::Sequencer
+{
+
+FKeyFrameOutlinerColumn::FKeyFrameOutlinerColumn()
+{
+	Name     = FCommonOutlinerNames::KeyFrame;
+	Label    = NSLOCTEXT("FKeyFrameOutlinerColumn", "ColumnName", "Key Frame");
+	Position = FOutlinerColumnPosition{ 20, EOutlinerColumnGroup::RightGutter };
+	Layout   = FOutlinerColumnLayout{ 34.f, FMargin(0.f), HAlign_Center, VAlign_Center, EOutlinerColumnSizeMode::Fixed, EOutlinerColumnFlags::OverflowSubsequentEmptyCells };
+}
+
+bool FKeyFrameOutlinerColumn::IsColumnVisibleByDefault() const
+{
+	return false;
+}
+
+} // namespace UE::Sequencer
diff --git a/Engine/Source/Editor/SequencerCore/Private/MVVM/ViewModels/OutlinerColumns/NavOutlinerColumn.cpp b/Engine/Source/Editor/SequencerCore/Private/MVVM/ViewModels/OutlinerColumns/NavOutlinerColumn.cpp
--- a/Engine/Source/Editor/SequencerCore/Private/MVVM/ViewModels/OutlinerColumns/NavOutlinerColumn.cpp
+++ b/Engine/Source/Editor/SequencerCore/Private/MVVM/ViewModels/OutlinerColumns/NavOutlinerColumn.cpp
@@ -8,9 +8,9 @@ namespace UE::Sequencer
 FNavOutlinerColumn::FNavOutlinerColumn()
 {
 	Name     = FCommonOutlinerNames::Nav;
-	Label    = NSLOCTEXT("FNavOutlinerColumn", "NavColumnNav", "Nav");
+	Label    = NSLOCTEXT("FNavOutlinerColumn", "ColumnName", "Key Navigation");
 	Position = FOutlinerColumnPosition{ 10, EOutlinerColumnGroup::RightGutter };
-	Layout   = FOutlinerColumnLayout{ 58.f, FMargin(8.f, 0.f), HAlign_Center, VAlign_Center, EOutlinerColumnSizeMode::Fixed, EOutlinerColumnFlags::OverflowSubsequentEmptyCells };
+	Layout   = FOutlinerColumnLayout{ 74.f, FMargin(0.f), HAlign_Center, VAlign_Center, EOutlinerColumnSizeMode::Fixed, EOutlinerColumnFlags::OverflowSubsequentEmptyCells };
 }
 
 } // namespace UE::Sequencer
diff --git a/Engine/Source/Editor/SequencerCore/Private/MVVM/ViewModels/OutlinerColumns/OutlinerColumnTypes.cpp b/Engine/Source/Editor/SequencerCore/Private/MVVM/ViewModels/OutlinerColumns/OutlinerColumnTypes.cpp
--- a/Engine/Source/Editor/SequencerCore/Private/MVVM/ViewModels/OutlinerColumns/OutlinerColumnTypes.cpp
+++ b/Engine/Source/Editor/SequencerCore/Private/MVVM/ViewModels/OutlinerColumns/OutlinerColumnTypes.cpp
@@ -13,6 +13,7 @@ FName FCommonOutlinerNames::Label("Label");
 FName FCommonOutlinerNames::Edit("Edit");
 FName FCommonOutlinerNames::Add("Add");
 FName FCommonOutlinerNames::Nav("Nav");
+FName FCommonOutlinerNames::KeyFrame("KeyFrame");
 FName FCommonOutlinerNames::ColorPicker("ColorPicker");
 
 } // namespace UE::Sequencer
\ No newline at end of file
diff --git a/Engine/Source/Editor/SequencerCore/Private/MVVM/Views/SKeyNavigationButtons.cpp b/Engine/Source/Editor/SequencerCore/Private/MVVM/Views/SKeyNavigationButtons.cpp
--- a/Engine/Source/Editor/SequencerCore/Private/MVVM/Views/SKeyNavigationButtons.cpp
+++ b/Engine/Source/Editor/SequencerCore/Private/MVVM/Views/SKeyNavigationButtons.cpp
@@ -8,7 +8,9 @@
 #include "HAL/PlatformCrt.h"
 #include "Layout/Children.h"
 #include "Layout/Margin.h"
+#include "MVVM/ViewModelPtr.h"
 #include "MVVM/ViewModels/ViewModel.h"
+#include "MVVM/Extensions/IHoveredExtension.h"
 #include "Math/NumericLimits.h"
 #include "Misc/FrameNumber.h"
 #include "Misc/FrameTime.h"
@@ -30,50 +32,82 @@ void SKeyNavigationButtons::Construct(const FArguments& InArgs, const TSharedPtr
 {
 	WeakModel = InModel;
 
+	TAttribute<bool> IsHovered;
+	if (TSharedPtr<IHoveredExtension> Hoverable = InModel->CastThisShared<IHoveredExtension>())
+	{
+		IsHovered = MakeAttributeSP(Hoverable.ToSharedRef(), &IHoveredExtension::IsHovered);
+	}
+
 	GetNavigatableTimesEvent = InArgs._GetNavigatableTimes;
 	AddKeyEvent = InArgs._OnAddKey;
 	SetTimeEvent = InArgs._OnSetTime;
 	TimeAttribute = InArgs._Time;
 
-	ChildSlot
-	[
-		SNew(SHorizontalBox)
-		
-		// Previous key slot
-		+ SHorizontalBox::Slot()
+	TSharedRef<SHorizontalBox> BoxPanel = SNew(SHorizontalBox);
+
+	const float CommonPadding = 4.f;
+
+	if (EnumHasAnyFlags(InArgs._Buttons, EKeyNavigationButtons::PreviousKey))
+	{
+		BoxPanel->AddSlot()
 		.VAlign(VAlign_Center)
+		.Padding(CommonPadding, 0.f)
 		.AutoWidth()
 		[
 			SNew(SOutlinerColumnButton)
 			.IsFocusable(false)
+			.IsRowHovered(IsHovered)
 			.ToolTipText(InArgs._PreviousKeyToolTip)
 			.Image(FAppStyle::GetBrush("Sequencer.Outliner.PreviousKey"))
 			.OnClicked(this, &SKeyNavigationButtons::OnPreviousKeyClicked)
-		]
+		];
+	}
 
-		// Add key slot
-		+ SHorizontalBox::Slot()
+	if (EnumHasAnyFlags(InArgs._Buttons, EKeyNavigationButtons::AddKey))
+	{
+		BoxPanel->AddSlot()
 		.VAlign(VAlign_Center)
-		.Padding(8.f, 0.f)
+		.Padding(CommonPadding, 0.f)
 		.AutoWidth()
 		[
 			SNew(SOutlinerColumnButton)
 			.IsFocusable(false)
+			.IsRowHovered(IsHovered)
 			.ToolTipText(InArgs._AddKeyToolTip)
 			.Image(FAppStyle::GetBrush("Sequencer.Outliner.AddKey"))
 			.OnClicked(this, &SKeyNavigationButtons::OnAddKeyClicked)
-		]
-		// Next key slot
-		+ SHorizontalBox::Slot()
+		];
+	}
+
+	if (EnumHasAnyFlags(InArgs._Buttons, EKeyNavigationButtons::NextKey))
+	{
+		BoxPanel->AddSlot()
 		.VAlign(VAlign_Center)
+		.Padding(CommonPadding, 0.f)
 		.AutoWidth()
 		[
 			SNew(SOutlinerColumnButton)
 			.IsFocusable(false)
+			.IsRowHovered(IsHovered)
 			.ToolTipText(InArgs._NextKeyToolTip)
 			.Image(FAppStyle::GetBrush("Sequencer.Outliner.NextKey"))
 			.OnClicked(this, &SKeyNavigationButtons::OnNextKeyClicked)
-		]
+		];
+	}
+
+	if (BoxPanel->NumSlots() > 0)
+	{
+		FMargin LeadingPadding = BoxPanel->GetSlot(0).GetPadding();
+		LeadingPadding.Left += CommonPadding;
+		BoxPanel->GetSlot(0).SetPadding(LeadingPadding);
+
+		FMargin TrailingPadding = BoxPanel->GetSlot(BoxPanel->NumSlots()-1).GetPadding();
+		TrailingPadding.Right += CommonPadding;
+		BoxPanel->GetSlot(BoxPanel->NumSlots()-1).SetPadding(TrailingPadding);
+	}
+	ChildSlot
+	[
+		BoxPanel
 	];
 }
 
diff --git a/Engine/Source/Editor/SequencerCore/Private/MVVM/Views/SOutlinerView.cpp b/Engine/Source/Editor/SequencerCore/Private/MVVM/Views/SOutlinerView.cpp
--- a/Engine/Source/Editor/SequencerCore/Private/MVVM/Views/SOutlinerView.cpp
+++ b/Engine/Source/Editor/SequencerCore/Private/MVVM/Views/SOutlinerView.cpp
@@ -305,6 +305,7 @@ int32 SOutlinerView::CreateOutlinerColumnsForGroup(int32 ColumnIndex, EOutlinerC
 		ColumnMetaData->Columns.Add(Layout);
 
 		// Keep track of how to create widgets for this column
+		// Column Generators must be defined before HeaderRow->AddColumn since AddColumn might re-generate widgets
 		ColumnGenerators.Add(ColumnName, [Column](const FCreateOutlinerColumnParams& Params, const TSharedRef<SOutlinerViewRow>& Row){
 			return Column->IsItemCompatibleWithColumn(Params)
 				? Column->CreateColumnWidget(Params, Row)
@@ -343,22 +344,26 @@ void SOutlinerView::InsertSeparatorColumn(int32 InsertIndex, int32 SeparatorID)
 		EOutlinerColumnSizeMode::Fixed,
 		EOutlinerColumnFlags::Hidden
 	};
+
 	FName SeparatorName(NAME_Separator, SeparatorID);
 
 	// Add a 1px separator column
 	ColumnMetaData->Columns.Insert(SeparatorLayout, InsertIndex);
+
+	// Column Generators must be defined before HeaderRow->InsertColumn since InsertColumn might re-generate widgets
+	ColumnGenerators.Add(SeparatorName,
+		[](const FCreateOutlinerColumnParams&, const TSharedRef<SOutlinerViewRow>&)->TSharedPtr<SWidget>
+	{
+		return SNew(SImage)
+			.Image(FAppStyle::Get().GetBrush("Sequencer.Outliner.Separator"));
+	}
+	);
+
 	HeaderRow->InsertColumn(
 		SHeaderRow::Column(SeparatorName)
 		.FixedWidth(1.f),
 		InsertIndex
 	);
-	ColumnGenerators.Add(SeparatorName,
-		[](const FCreateOutlinerColumnParams&, const TSharedRef<SOutlinerViewRow>&)->TSharedPtr<SWidget>
-		{
-			return SNew(SImage)
-			.Image(FAppStyle::Get().GetBrush("Sequencer.Outliner.Separator"));
-		}
-	);
 }
 
 void SOutlinerView::UpdateOutlinerColumns()
@@ -403,13 +408,13 @@ void SOutlinerView::UpdateOutlinerColumns()
 	int32 InsertIndex = NumLeftGutter;
 	if (InsertIndex < ColumnMetaData->Columns.Num())
 	{
-		InsertSeparatorColumn(InsertIndex++, NumSeparators++);
+		InsertSeparatorColumn(InsertIndex++, ++NumSeparators);
 		if (NumCenter > 0)
 		{
 			InsertIndex += NumCenter;
 			if (InsertIndex < ColumnMetaData->Columns.Num())
 			{
-				InsertSeparatorColumn(InsertIndex++, NumSeparators++);
+				InsertSeparatorColumn(InsertIndex++, ++NumSeparators);
 			}
 		}
 	}
diff --git a/Engine/Source/Editor/SequencerCore/Public/MVVM/ViewModels/OutlinerColumns/KeyFrameOutlinerColumn.h b/Engine/Source/Editor/SequencerCore/Public/MVVM/ViewModels/OutlinerColumns/KeyFrameOutlinerColumn.h
new file mode 100644
--- /dev/null
+++ b/Engine/Source/Editor/SequencerCore/Public/MVVM/ViewModels/OutlinerColumns/KeyFrameOutlinerColumn.h
@@ -0,0 +1,20 @@
+// Copyright Epic Games, Inc. All Rights Reserved.
+
+#pragma once
+
+#include "MVVM/ViewModels/OutlinerColumns/OutlinerColumnBase.h"
+
+namespace UE::Sequencer
+{
+
+class FKeyFrameOutlinerColumn
+	: public FOutlinerColumnBase
+{
+public:
+
+	SEQUENCERCORE_API FKeyFrameOutlinerColumn();
+
+	bool IsColumnVisibleByDefault() const override;
+};
+
+} // namespace UE::Sequencer
\ No newline at end of file
diff --git a/Engine/Source/Editor/SequencerCore/Public/MVVM/ViewModels/OutlinerColumns/OutlinerColumnTypes.h b/Engine/Source/Editor/SequencerCore/Public/MVVM/ViewModels/OutlinerColumns/OutlinerColumnTypes.h
--- a/Engine/Source/Editor/SequencerCore/Public/MVVM/ViewModels/OutlinerColumns/OutlinerColumnTypes.h
+++ b/Engine/Source/Editor/SequencerCore/Public/MVVM/ViewModels/OutlinerColumns/OutlinerColumnTypes.h
@@ -71,6 +71,7 @@ struct FCommonOutlinerNames
 	static SEQUENCERCORE_API FName Edit;
 	static SEQUENCERCORE_API FName Add;
 	static SEQUENCERCORE_API FName Nav;
+	static SEQUENCERCORE_API FName KeyFrame;
 	static SEQUENCERCORE_API FName ColorPicker;
 };
 
diff --git a/Engine/Source/Editor/SequencerCore/Public/MVVM/Views/SKeyNavigationButtons.h b/Engine/Source/Editor/SequencerCore/Public/MVVM/Views/SKeyNavigationButtons.h
--- a/Engine/Source/Editor/SequencerCore/Public/MVVM/Views/SKeyNavigationButtons.h
+++ b/Engine/Source/Editor/SequencerCore/Public/MVVM/Views/SKeyNavigationButtons.h
@@ -14,6 +14,7 @@
 #include "Templates/SharedPointer.h"
 #include "Widgets/DeclarativeSyntaxSupport.h"
 #include "Widgets/SCompoundWidget.h"
+#include "Misc/EnumClassFlags.h"
 
 struct FFrameNumber;
 
@@ -24,6 +25,17 @@ namespace Sequencer
 
 class FViewModel;
 
+enum class EKeyNavigationButtons : uint8
+{
+	PreviousKey = 1 << 0,
+	AddKey      = 1 << 1,
+	NextKey     = 1 << 2,
+
+	NavOnly = PreviousKey | NextKey,
+	All     = NavOnly | AddKey,
+};
+ENUM_CLASS_FLAGS(EKeyNavigationButtons)
+
 /**
  * A widget for navigating between keys on a sequencer track
  */
@@ -36,7 +48,9 @@ public:
 	DECLARE_DELEGATE_OneParam(FOnGetNavigatableTimes, TArray<FFrameNumber>&)
 	DECLARE_DELEGATE_TwoParams(FOnAddKey, FFrameTime, TSharedPtr<FViewModel>)
 
-	SLATE_BEGIN_ARGS(SKeyNavigationButtons) {}
+	SLATE_BEGIN_ARGS(SKeyNavigationButtons) : _Buttons(EKeyNavigationButtons::All) {}
+
+		SLATE_ARGUMENT(EKeyNavigationButtons, Buttons)
 
 		SLATE_ARGUMENT(FText, PreviousKeyToolTip)
 		SLATE_ARGUMENT(FText, NextKeyToolTip)
