Sequencer: Fixed channels disappearing when drag-manipulating component transforms from the details panel

  - Wrapped calls to On(Pre)TransformChanged with a FScopedSignedObjectModifyDefer that forces the result to be propagated to the Sequencer UI. Without this they wouldn't only be propagated at the end of the drag operation when the transaction was closed.
  - Fixed unnecessary calls to SetMask when the channel mask didn't need to change

diff --git a/Engine/Source/Editor/MovieSceneTools/Private/TrackEditors/TransformTrackEditor.cpp b/Engine/Source/Editor/MovieSceneTools/Private/TrackEditors/TransformTrackEditor.cpp
--- a/Engine/Source/Editor/MovieSceneTools/Private/TrackEditors/TransformTrackEditor.cpp
+++ b/Engine/Source/Editor/MovieSceneTools/Private/TrackEditors/TransformTrackEditor.cpp
@@ -359,6 +359,7 @@ void F3DTransformTrackEditor::OnPrePropertyChanged(UObject* InObject, const FEdi
 
 	if (InObject && bTransformationToChange)
 	{
+		UE::MovieScene::FScopedSignedObjectModifyDefer ForceFlush(true);
 		OnPreTransformChanged(*InObject);
 	}
 }
@@ -373,6 +374,7 @@ void F3DTransformTrackEditor::OnPostPropertyChanged(UObject* InObject, FProperty
 
 	if (InObject && bTransformationChanged)
 	{
+		UE::MovieScene::FScopedSignedObjectModifyDefer ForceFlush(true);
 		OnTransformChanged(*InObject);
 	}
 }
@@ -991,6 +993,16 @@ void F3DTransformTrackEditor::AddTransformKeys( UObject* ObjectToKey, const TOpt
 	};
 	auto GenerateKeys = [=, this](UMovieSceneSection* Section, FGeneratedTrackKeys& GeneratedKeys)
 	{
+		UMovieScene3DTransformSection* TransformSection = CastChecked<UMovieScene3DTransformSection>(Section);
+
+		// Ensure that the Transform Channel is masked in
+		EMovieSceneTransformChannel ExistingChannels = TransformSection->GetMask().GetChannels();
+		if (!EnumHasAnyFlags(ExistingChannels, ChannelsToKey))
+		{
+			TransformSection->Modify();
+			TransformSection->SetMask(TransformSection->GetMask().GetChannels() | ChannelsToKey);
+		}
+
 		this->GetTransformKeys(LastTransform, CurrentTransform, ChannelsToKey, ObjectToKey, Section, GeneratedKeys);
 	};
 	auto OnKeyProperty = [=, this](FFrameNumber Time) -> FKeyPropertyResult
@@ -1417,7 +1429,12 @@ void F3DTransformTrackEditor::OnTransformPropertyChanged(const FPropertyChangedP
 		UMovieScene3DTransformSection* TransformSection = CastChecked<UMovieScene3DTransformSection>(Section);
 
 		// Ensure that the Transform Channel is masked in
-		TransformSection->SetMask(TransformSection->GetMask().GetChannels() | TransformChannel);
+		EMovieSceneTransformChannel ExistingChannels = TransformSection->GetMask().GetChannels();
+		if (!EnumHasAnyFlags(ExistingChannels, TransformChannel))
+		{
+			TransformSection->Modify();
+			TransformSection->SetMask(TransformSection->GetMask().GetChannels() | TransformChannel);
+		}
 
 		GetTransformKeys(TOptional<FTransformData>(), FTransformData(SceneComponent), TransformChannel, SceneComponent, Section, OutGeneratedKeys);
 	};
