Sequencer: Fixed audio waveform drawing for expanded audio sections

diff --git a/Engine/Source/Editor/MovieSceneTools/Private/TrackEditors/AudioTrackEditor.cpp b/Engine/Source/Editor/MovieSceneTools/Private/TrackEditors/AudioTrackEditor.cpp
--- a/Engine/Source/Editor/MovieSceneTools/Private/TrackEditors/AudioTrackEditor.cpp
+++ b/Engine/Source/Editor/MovieSceneTools/Private/TrackEditors/AudioTrackEditor.cpp
@@ -725,7 +725,7 @@ int32 FAudioSection::OnPaintSection( FSequencerSectionPainter& Painter ) const
 		FSlateDrawElement::MakeViewport(
 			Painter.DrawElements,
 			++LayerId,
-			Painter.SectionGeometry.ToPaintGeometry(FVector2f(StoredXSize, Painter.SectionGeometry.GetLocalSize().Y), FSlateLayoutTransform(FVector2f(StoredXOffset, 0))),
+			Painter.SectionGeometry.ToPaintGeometry(FVector2f(StoredXSize, StoredSectionHeight), FSlateLayoutTransform(FVector2f(StoredXOffset, 0))),
 			WaveformThumbnail,
 			(Painter.bParentEnabled ? ESlateDrawEffect::None : ESlateDrawEffect::DisabledEffect) | ESlateDrawEffect::NoGamma,
 			FLinearColor::White
