Sequencer: Fix camera button disappearing when unspawned

diff --git a/Engine/Source/Editor/MovieSceneTools/Private/TrackEditors/TransformTrackEditor.cpp b/Engine/Source/Editor/MovieSceneTools/Private/TrackEditors/TransformTrackEditor.cpp
--- a/Engine/Source/Editor/MovieSceneTools/Private/TrackEditors/TransformTrackEditor.cpp
+++ b/Engine/Source/Editor/MovieSceneTools/Private/TrackEditors/TransformTrackEditor.cpp
@@ -587,7 +587,7 @@ EVisibility F3DTransformTrackEditor::IsCameraVisible(FGuid ObjectGuid) const
 		}
 	}
 
-	return EVisibility::Collapsed;
+	return EVisibility::Hidden;
 }
 
 ECheckBoxState F3DTransformTrackEditor::IsCameraLocked(FGuid ObjectGuid) const
