Sequencer: Fix setting string array property values by using PropertyChangedParams.GetPropertyValue which handles arrays

diff --git a/Engine/Source/Editor/MovieSceneTools/Private/TrackEditors/PropertyTrackEditors/StringPropertyTrackEditor.cpp b/Engine/Source/Editor/MovieSceneTools/Private/TrackEditors/PropertyTrackEditors/StringPropertyTrackEditor.cpp
--- a/Engine/Source/Editor/MovieSceneTools/Private/TrackEditors/PropertyTrackEditors/StringPropertyTrackEditor.cpp
+++ b/Engine/Source/Editor/MovieSceneTools/Private/TrackEditors/PropertyTrackEditors/StringPropertyTrackEditor.cpp
@@ -22,23 +22,8 @@ TSharedRef<ISequencerTrackEditor> FStringPropertyTrackEditor::CreateTrackEditor(
 	return MakeShareable(new FStringPropertyTrackEditor(OwningSequencer));
 }
 
-
 void FStringPropertyTrackEditor::GenerateKeysFromPropertyChanged( const FPropertyChangedParams& PropertyChangedParams, UMovieSceneSection* SectionToKey, FGeneratedTrackKeys& OutGeneratedKeys )
 {
-	void* CurrentObject = PropertyChangedParams.ObjectsThatChanged[0];
-	void* PropertyValue = nullptr;
-	for (int32 i = 0; i < PropertyChangedParams.PropertyPath.GetNumProperties(); i++)
-	{
-		if (FProperty* Property = PropertyChangedParams.PropertyPath.GetPropertyInfo(i).Property.Get())
-		{
-			CurrentObject = Property->ContainerPtrToValuePtr<FString>(CurrentObject, 0);
-		}
-	}
-
-	const FStrProperty* StrProperty = CastField<const FStrProperty>( PropertyChangedParams.PropertyPath.GetLeafMostProperty().Property.Get() );
-	if ( StrProperty )
-	{
-		FString StrPropertyValue = StrProperty->GetPropertyValue(CurrentObject);
-		OutGeneratedKeys.Add(FMovieSceneChannelValueSetter::Create<FMovieSceneStringChannel>(0, MoveTemp(StrPropertyValue), true));
-	}
+	FString KeyedValue = PropertyChangedParams.GetPropertyValue<FString>();
+	OutGeneratedKeys.Add(FMovieSceneChannelValueSetter::Create<FMovieSceneStringChannel>(0, MoveTemp(KeyedValue), true));
 }
