Sequencer: Constraints working with snapper. Fix getting actor transforms, can't rely on the old interragator that only evaluates the transform track, not any attachments or constraints, etc.. so need to go evaluate everything when getting transforms. Also when setting the transform make sure it's in the right space taking into account the constraint space.

diff --git a/Engine/Plugins/Animation/ControlRig/Source/ControlRigEditor/Private/EditMode/SControlRigOutliner.cpp b/Engine/Plugins/Animation/ControlRig/Source/ControlRigEditor/Private/EditMode/SControlRigOutliner.cpp
--- a/Engine/Plugins/Animation/ControlRig/Source/ControlRigEditor/Private/EditMode/SControlRigOutliner.cpp
+++ b/Engine/Plugins/Animation/ControlRig/Source/ControlRigEditor/Private/EditMode/SControlRigOutliner.cpp
@@ -1112,6 +1112,10 @@ void SControlRigOutliner::HandleSelectionChanged(TSharedPtr<FMultiRigTreeElement
 					CurrentSelection.Key->ClearControlSelection();
 				}
 			}
+			if (GEditor)
+			{
+				GEditor->Exec(GEditor->GetWorld(), TEXT("SELECT NONE"));
+			}
 		}
 	}
 
diff --git a/Engine/Plugins/Animation/ControlRig/Source/ControlRigEditor/Private/Tools/ControlRigSnapper.cpp b/Engine/Plugins/Animation/ControlRig/Source/ControlRigEditor/Private/Tools/ControlRigSnapper.cpp
--- a/Engine/Plugins/Animation/ControlRig/Source/ControlRigEditor/Private/Tools/ControlRigSnapper.cpp
+++ b/Engine/Plugins/Animation/ControlRig/Source/ControlRigEditor/Private/Tools/ControlRigSnapper.cpp
@@ -32,6 +32,7 @@
 #include "LevelSequenceActor.h"
 #include "Tools/BakingHelper.h"
 #include "Containers/SortedMap.h"
+#include "TransformConstraint.h"
 #include "Sequencer/ControlRigSequencerHelpers.h"
 #include "Constraints/MovieSceneConstraintChannelHelper.inl"
 
@@ -202,11 +203,11 @@ struct FGuidAndActor
 	AActor* Actor;
 
 	bool SetTransform(
-		ISequencer* Sequencer, const TArray<FFrameNumber>& Frames,
+		TSharedPtr<ISequencer>& Sequencer, const TArray<FFrameNumber>& Frames,
 		const TArray<FTransform>& WorldTransformsToSnapTo, const UControlRigSnapSettings* SnapSettings) const
 	{
 		// get section
-		UMovieScene3DTransformSection* TransformSection = MovieSceneToolHelpers::GetTransformSection(Sequencer, Guid);
+		UMovieScene3DTransformSection* TransformSection = MovieSceneToolHelpers::GetTransformSection(Sequencer.Get(), Guid);
 		if (!TransformSection)
 		{
 			return false;
@@ -218,8 +219,9 @@ struct FGuidAndActor
 		if (ParentActor)
 		{
 			FActorForWorldTransforms ActorSelection;
+			ActorSelection.Actor = ParentActor;
 			ActorSelection.SocketName = Actor->GetAttachParentSocketName();
-			MovieSceneToolHelpers::GetActorWorldTransforms(Sequencer, ActorSelection, Frames, ParentWorldTransforms);
+			MovieSceneToolHelpers::GetActorWorldTransforms(Sequencer.Get(), ActorSelection, Frames, ParentWorldTransforms);
 		}
 		else
 		{
@@ -258,9 +260,37 @@ struct FGuidAndActor
 			LocalTransforms.Add(WorldTransform.GetRelativeTransform(ParentTransform));
 		}
 
+		if (USceneComponent* SceneComponent = Actor->GetRootComponent())
+		{
+			UMovieScene* MovieScene = Sequencer->GetFocusedMovieSceneSequence()->GetMovieScene();
+			FFrameRate TickResolution = MovieScene->GetTickResolution();
+			FFrameRate DisplayRate = MovieScene->GetDisplayRate();
+			FMovieSceneSequenceTransform RootToLocalTransform = Sequencer->GetFocusedMovieSceneSequenceTransform();
+
+			//adjust keys for constraint
+			const FConstraintsManagerController& Controller = FConstraintsManagerController::Get(Actor->GetWorld());
+
+			for (int32 Index = 0; Index < LocalTransforms.Num(); ++Index) 
+			{
+				const FFrameNumber Frame = Frames[Index];
+				FTransform& CurrentTransform = LocalTransforms[Index];
+
+				FFrameTime GlobalTime(Frame);
+				GlobalTime = GlobalTime * RootToLocalTransform.InverseNoLooping();
+				FMovieSceneContext Context = FMovieSceneContext(FMovieSceneEvaluationRange(GlobalTime, TickResolution), Sequencer->GetPlaybackStatus()).SetHasJumped(true);
+				if (Index == 0) // similar with baking first time in we need to evaluate twice (think due to double buffering that happens with skel mesh components).
+				{
+					Sequencer->GetEvaluationTemplate().EvaluateSynchronousBlocking(Context);
+				}
+				Sequencer->GetEvaluationTemplate().EvaluateSynchronousBlocking(Context);
+				Controller.EvaluateAllConstraints();
+				FTransformConstraintUtils::UpdateTransformBasedOnConstraint(CurrentTransform, SceneComponent);
+			}
+		}
+
 		// add keys
 		MovieSceneToolHelpers::AddTransformKeys(TransformSection, Frames, LocalTransforms, Channels);
-
+		
 		if (SnapSettings && SnapSettings->BakingKeySettings != EBakingKeySettings::KeysOnly
 			&& SnapSettings->bReduceKeys == true && Frames.Num() > 2)
 		{
@@ -273,7 +303,7 @@ struct FGuidAndActor
 		}
 
 
-		// notify
+	
 		Sequencer->NotifyMovieSceneDataChanged(EMovieSceneDataChangeType::MovieSceneStructureItemsChanged);
 
 		return true;
@@ -410,7 +440,7 @@ static UMovieSceneControlRigParameterSection* GetControlRigSection(ISequencer* S
 }
 
 //Get Transform frames for specified snap items, will add them to the OutFrameMap
-static void GetTransformFrames(ISequencer * Sequencer, const FControlRigSnapperSelection& SnapItem, FFrameNumber StartFrame, FFrameNumber EndFrame, TSortedMap<FFrameNumber, FFrameNumber>& OutFrameMap)
+static void GetTransformFrames(TSharedPtr<ISequencer>&  Sequencer, const FControlRigSnapperSelection& SnapItem, FFrameNumber StartFrame, FFrameNumber EndFrame, TSortedMap<FFrameNumber, FFrameNumber>& OutFrameMap)
 {
 	TArray<FFrameNumber> FramesToUse;
 	TArray<FActorForWorldTransforms> ParentActors;
@@ -446,7 +476,7 @@ static void GetTransformFrames(ISequencer * Sequencer, const FControlRigSnapperS
 		}
 		if (Guid.IsValid())
 		{
-			if (UMovieScene3DTransformSection* TransformSection = MovieSceneToolHelpers::GetTransformSection(Sequencer, Guid))
+			if (UMovieScene3DTransformSection* TransformSection = MovieSceneToolHelpers::GetTransformSection(Sequencer.Get(), Guid))
 			{
 				TArrayView<FMovieSceneDoubleChannel*> Channels = TransformSection->GetChannelProxy().GetChannels<FMovieSceneDoubleChannel>();
 				TArray<FFrameNumber> TransformFrameTimes = FMovieSceneConstraintChannelHelper::GetTransformTimes(
@@ -460,7 +490,7 @@ static void GetTransformFrames(ISequencer * Sequencer, const FControlRigSnapperS
 	}
 	for (FControlRigForWorldTransforms& ControlRig : ParentControlRigs)
 	{
-		UMovieSceneControlRigParameterSection* Section = GetControlRigSection(Sequencer, ControlRig.ControlRig.Get());
+		UMovieSceneControlRigParameterSection* Section = GetControlRigSection(Sequencer.Get(), ControlRig.ControlRig.Get());
 		if (Section)
 		{
 			for (FName& ControlName : ControlRig.ControlNames)
@@ -483,7 +513,7 @@ bool FControlRigSnapper::SnapIt(FFrameNumber StartFrame, FFrameNumber EndFrame,c
 	TWeakPtr<ISequencer> InSequencer = GetSequencer();
 	if (InSequencer.IsValid() && InSequencer.Pin()->GetFocusedMovieSceneSequence() && ActorToSnap.IsValid())
 	{
-		ISequencer* Sequencer = InSequencer.Pin().Get();
+		TSharedPtr<ISequencer> Sequencer = InSequencer.Pin();
 		UMovieScene* MovieScene = Sequencer->GetFocusedMovieSceneSequence()->GetMovieScene();
 
 		TArray<FFrameNumber> Frames;
@@ -506,7 +536,7 @@ bool FControlRigSnapper::SnapIt(FFrameNumber StartFrame, FFrameNumber EndFrame,c
 		}
 		Sequencer->ForceEvaluate(); // force an evaluate so any control rig get's binding setup
 		TArray<FTransform> WorldTransformToSnap;
-		bool bSnapToFirstFrameNotParents = !CalculateWorldTransformsFromParents(Sequencer, ParentToSnap, Frames, WorldTransformToSnap);
+		bool bSnapToFirstFrameNotParents = !CalculateWorldTransformsFromParents(Sequencer.Get(), ParentToSnap, Frames, WorldTransformToSnap);
 		if (Frames.Num() != WorldTransformToSnap.Num())
 		{
 			UE_LOG(LogControlRig, Error, TEXT("Number of Frames %d to Snap different than Parent Frames %d"), Frames.Num(),WorldTransformToSnap.Num());
@@ -558,7 +588,7 @@ bool FControlRigSnapper::SnapIt(FFrameNumber StartFrame, FFrameNumber EndFrame,c
 				TArray<FTransform> OneTransform;
 				FActorForWorldTransforms ActorSelection;
 				ActorSelection.Actor = GuidActor.Actor;
-				MovieSceneToolHelpers::GetActorWorldTransforms(Sequencer, ActorSelection, OneFrame, OneTransform);
+				MovieSceneToolHelpers::GetActorWorldTransforms(Sequencer.Get(), ActorSelection, OneFrame, OneTransform);
 				if (bSnapToFirstFrameNotParents)
 				{
 					for (FTransform& Transform : WorldTransformToSnap)
@@ -607,7 +637,7 @@ bool FControlRigSnapper::SnapIt(FFrameNumber StartFrame, FFrameNumber EndFrame,c
 					TArray<FTransform> ControlRigParentWorldTransforms;
 					FActorForWorldTransforms ControlRigActorSelection;
 					ControlRigActorSelection.Actor = Actor;
-					MovieSceneToolHelpers::GetActorWorldTransforms(Sequencer, ControlRigActorSelection, Frames, ControlRigParentWorldTransforms);
+					MovieSceneToolHelpers::GetActorWorldTransforms(Sequencer.Get(), ControlRigActorSelection, Frames, ControlRigParentWorldTransforms);
 					FControlRigSnapper Snapper;
 					for (const FName& Name : ControlRigAndSelection.ControlNames)
 					{
@@ -620,7 +650,7 @@ bool FControlRigSnapper::SnapIt(FFrameNumber StartFrame, FFrameNumber EndFrame,c
 						{
 							OneFrame[0] = Frames[0];
 							CurrentParentWorldTransform[0] = ControlRigParentWorldTransforms[0];
-							Snapper.GetControlRigControlTransforms(Sequencer, ControlRig, Name, OneFrame, CurrentParentWorldTransform, CurrentControlRigTransform);
+							Snapper.GetControlRigControlTransforms(Sequencer.Get(), ControlRig, Name, OneFrame, CurrentParentWorldTransform, CurrentControlRigTransform);
 							if (bSnapToFirstFrameNotParents)
 							{
 								for (FTransform& Transform : WorldTransformToSnap)
@@ -662,7 +692,7 @@ bool FControlRigSnapper::SnapIt(FFrameNumber StartFrame, FFrameNumber EndFrame,c
 							OneFrame[0] = Frames[Index];
 							CurrentParentWorldTransform[0] = ControlRigParentWorldTransforms[Index];
 							//this will evaluate at the current frame which we want
-							GetControlRigControlTransforms(Sequencer, ControlRig, Name, OneFrame, CurrentParentWorldTransform, CurrentControlRigTransform);
+							GetControlRigControlTransforms(Sequencer.Get(), ControlRig, Name, OneFrame, CurrentParentWorldTransform, CurrentControlRigTransform);
 							if (SnapSettings)
 							{
 								FTransform& Transform = WorldTransformToSnap[Index];
@@ -690,7 +720,7 @@ bool FControlRigSnapper::SnapIt(FFrameNumber StartFrame, FFrameNumber EndFrame,c
 						}
 					}
 
-					UMovieSceneControlRigParameterSection* ControlRigSection = GetControlRigSection(Sequencer, ControlRig);
+					UMovieSceneControlRigParameterSection* ControlRigSection = GetControlRigSection(Sequencer.Get(), ControlRig);
 					if (ControlRigSection && SnapSettings && SnapSettings->BakingKeySettings != EBakingKeySettings::KeysOnly
 						&& SnapSettings->bReduceKeys == true && Frames.Num() > 2)
 					{
diff --git a/Engine/Source/Editor/MovieSceneTools/Private/Constraints/MovieSceneConstraintChannelHelper.cpp b/Engine/Source/Editor/MovieSceneTools/Private/Constraints/MovieSceneConstraintChannelHelper.cpp
--- a/Engine/Source/Editor/MovieSceneTools/Private/Constraints/MovieSceneConstraintChannelHelper.cpp
+++ b/Engine/Source/Editor/MovieSceneTools/Private/Constraints/MovieSceneConstraintChannelHelper.cpp
@@ -1025,7 +1025,8 @@ void FMovieSceneConstraintChannelHelper::CompensateIfNeeded(
 			}
 
 			const UTickableTransformConstraint* Constraint = Cast<UTickableTransformConstraint>(InChannel.GetConstraint().Get());
-			return Constraint && (Constraint->GetTargetHash() == InChildHash) && Constraint->NeedsCompensation();
+			//if no InChildHash specified(== INDEX_NONE) then do all!
+			return Constraint && (InChildHash == INDEX_NONE || Constraint->GetTargetHash() == InChildHash) && Constraint->NeedsCompensation();
 		}
 	);
 
diff --git a/Engine/Source/Editor/MovieSceneTools/Private/MovieSceneToolHelpers.cpp b/Engine/Source/Editor/MovieSceneTools/Private/MovieSceneToolHelpers.cpp
--- a/Engine/Source/Editor/MovieSceneTools/Private/MovieSceneToolHelpers.cpp
+++ b/Engine/Source/Editor/MovieSceneTools/Private/MovieSceneToolHelpers.cpp
@@ -4208,126 +4208,119 @@ static void GetSequenceSceneComponentWorldTransforms(USceneComponent* SceneCompo
 	Interrogator.Reset();
 }
 
-static void GetSequencerActorWorldTransforms(IMovieScenePlayer* Player, UMovieSceneSequence* InSequence, FMovieSceneSequenceIDRef Template, const FActorForWorldTransforms& ActorSelection, const TArray<FFrameNumber>& Frames, TArray<FTransform>& OutTransforms)
+static void GetSequencerActorWorldTransforms(IMovieScenePlayer* Player, FMovieSceneSequenceTransform RootToLocalTransform, UMovieSceneSequence* InSequence, FMovieSceneSequenceIDRef Template, const FActorForWorldTransforms& ActorSelection, const TArray<FFrameNumber>& Frames, TArray<FTransform>& OutTransforms)
 {
-	if (AActor* Actor = ActorSelection.Actor.Get())
+	USceneComponent* SceneComponent = nullptr;
+	AActor* Actor = ActorSelection.Actor.Get();
+	if (Actor)
 	{
-
-		USkeletalMeshComponent* SkelMeshComp = ActorSelection.Component.IsValid() ? Cast<USkeletalMeshComponent>(ActorSelection.Component.Get()) : nullptr;
+		SceneComponent = Actor->GetRootComponent();
+	}
+	else
+	{
+		SceneComponent = ActorSelection.Component.IsValid() ? Cast<USceneComponent>(ActorSelection.Component.Get()) : nullptr;
+		if (SceneComponent)
+		{
+			Actor = SceneComponent->GetTypedOuter<AActor>();
+		}
+	}
+	
+	if(Actor && SceneComponent)
+	{
+		USkeletalMeshComponent* SkelMeshComp =  Cast<USkeletalMeshComponent>(SceneComponent);
 
 		if (!SkelMeshComp)
 		{
 			SkelMeshComp = MovieSceneToolHelpers::AcquireSkeletalMeshFromObject(Actor);
 		}
 
-		if (ActorSelection.SocketName != NAME_None && SkelMeshComp)
+		if (UMovieScene* MovieScene = InSequence->GetMovieScene())
 		{
+			OutTransforms.SetNum(Frames.Num());
 
-			if (UMovieScene* MovieScene = InSequence->GetMovieScene())
-			{
-				OutTransforms.SetNum(Frames.Num());
+			FFrameRate TickResolution = MovieScene->GetTickResolution();
+			FFrameRate DisplayRate = MovieScene->GetDisplayRate();
+			const TArray<IMovieSceneToolsAnimationBakeHelper*>& BakeHelpers = FMovieSceneToolsModule::Get().GetAnimationBakeHelpers();
 
-				FMovieSceneSequenceTransform RootToLocalTransform;
+			for (IMovieSceneToolsAnimationBakeHelper* BakeHelper : BakeHelpers)
+			{
+				if (BakeHelper)
+				{
+					BakeHelper->StartBaking(MovieScene);
+				}
+			}
 
-				FFrameRate TickResolution = MovieScene->GetTickResolution();
-				FFrameRate DisplayRate = MovieScene->GetDisplayRate();
-				const TArray<IMovieSceneToolsAnimationBakeHelper*>& BakeHelpers = FMovieSceneToolsModule::Get().GetAnimationBakeHelpers();
+			for (int32 Index = 0; Index < Frames.Num(); ++Index)
+			{
+				const FFrameNumber& FrameNumber = Frames[Index];
+				FFrameTime GlobalTime(FrameNumber);
+				GlobalTime = GlobalTime * RootToLocalTransform.InverseNoLooping(); //player evals in root time so need to go back to it.
 
 				for (IMovieSceneToolsAnimationBakeHelper* BakeHelper : BakeHelpers)
 				{
 					if (BakeHelper)
 					{
-						BakeHelper->StartBaking(MovieScene);
+						BakeHelper->PreEvaluation(MovieScene, FrameNumber);
 					}
 				}
 
-				for (int32 Index = 0; Index < Frames.Num(); ++Index)
+				FMovieSceneContext Context = FMovieSceneContext(FMovieSceneEvaluationRange(GlobalTime, TickResolution), Player->GetPlaybackStatus()).SetHasJumped(true);
+				if (Index == 0) // similar with baking first time in we need to evaluate twice (think due to double buffering that happens with skel mesh components).
 				{
-					const FFrameNumber& FrameNumber = Frames[Index];
-					FFrameTime GlobalTime(FrameNumber);
-
-					for (IMovieSceneToolsAnimationBakeHelper* BakeHelper : BakeHelpers)
-					{
-						if (BakeHelper)
-						{
-							BakeHelper->PreEvaluation(MovieScene, FrameNumber);
-						}
-					}
-
-					FMovieSceneContext Context = FMovieSceneContext(FMovieSceneEvaluationRange(GlobalTime, TickResolution), Player->GetPlaybackStatus()).SetHasJumped(true);
-					if (Index == 0) // similar with baking first time in we need to evaluate twice (think due to double buffering that happens with skel mesh components).
-					{
-						Player->GetEvaluationTemplate().EvaluateSynchronousBlocking(Context);
-					}
 					Player->GetEvaluationTemplate().EvaluateSynchronousBlocking(Context);
+				}
+				Player->GetEvaluationTemplate().EvaluateSynchronousBlocking(Context);
 
 
-					const FConstraintsManagerController& Controller = FConstraintsManagerController::Get(Actor->GetWorld());
-					Controller.EvaluateAllConstraints();
+				const FConstraintsManagerController& Controller = FConstraintsManagerController::Get(Actor->GetWorld());
+				Controller.EvaluateAllConstraints();
 					
-					for (IMovieSceneToolsAnimationBakeHelper* BakeHelper : BakeHelpers)
-					{
-						if (BakeHelper)
-						{
-							BakeHelper->PostEvaluation(MovieScene, FrameNumber);
-						}
-					}
-
-					AActor* Parent = ActorSelection.Actor.Get();
-					while (Parent)
+				for (IMovieSceneToolsAnimationBakeHelper* BakeHelper : BakeHelpers)
+				{
+					if (BakeHelper)
 					{
-						TArray<USkeletalMeshComponent*> MeshComps;
-						Parent->GetComponents(MeshComps, true);
-
-						for (USkeletalMeshComponent* MeshComp : MeshComps)
-						{
-							MeshComp->TickAnimation(0.03f, false);
-							MeshComp->RefreshBoneTransforms();
-							MeshComp->RefreshFollowerComponents();
-							MeshComp->UpdateComponentToWorld();
-							MeshComp->FinalizeBoneTransform();
-							MeshComp->MarkRenderTransformDirty();
-							MeshComp->MarkRenderDynamicDataDirty();
-						}
-
-						Parent = Parent->GetAttachParentActor();
+						BakeHelper->PostEvaluation(MovieScene, FrameNumber);
 					}
-
-					OutTransforms[Index] = SkelMeshComp->GetSocketTransform(ActorSelection.SocketName);// GetSocketTransofrm is world space in theory*OutTransforms[Index];
-
 				}
 
-				for (IMovieSceneToolsAnimationBakeHelper* BakeHelper : BakeHelpers)
+				AActor* Parent = ActorSelection.Actor.Get();
+				while (Parent)
 				{
-					if (BakeHelper)
+					TArray<USkeletalMeshComponent*> MeshComps;
+					Parent->GetComponents(MeshComps, true);
+
+					for (USkeletalMeshComponent* MeshComp : MeshComps)
 					{
-						BakeHelper->StopBaking(MovieScene);
+						MeshComp->TickAnimation(0.03f, false);
+						MeshComp->RefreshBoneTransforms();
+						MeshComp->RefreshFollowerComponents();
+						MeshComp->UpdateComponentToWorld();
+						MeshComp->FinalizeBoneTransform();
+						MeshComp->MarkRenderTransformDirty();
+						MeshComp->MarkRenderDynamicDataDirty();
 					}
+
+					Parent = Parent->GetAttachParentActor();
 				}
+
+				OutTransforms[Index] = (SkelMeshComp && ActorSelection.SocketName != NAME_None)
+					? SkelMeshComp->GetSocketTransform(ActorSelection.SocketName)
+					: SceneComponent->GetComponentToWorld();
+
 			}
-		}
-		else //no attached skelmesh socket so use Interrogator
-		{
-			USceneComponent* SceneComponent = Actor->GetRootComponent();
-			if (SceneComponent)
+
+			for (IMovieSceneToolsAnimationBakeHelper* BakeHelper : BakeHelpers)
 			{
-				GetSequenceSceneComponentWorldTransforms(SceneComponent, Player, InSequence, Template, Frames, OutTransforms);
+				if (BakeHelper)
+				{
+					BakeHelper->StopBaking(MovieScene);
+				}
 			}
-			
-		}
-	}
-	else //no actor so check to see if there's a scene component
-	{
-		USceneComponent* SceneComponent = ActorSelection.Component.IsValid() ? Cast<USceneComponent>(ActorSelection.Component.Get()) : nullptr;
-		if (SceneComponent)
-		{
-			GetSequenceSceneComponentWorldTransforms(SceneComponent, Player, InSequence, Template, Frames, OutTransforms);
-		}
-
+		}	
 	}
 }
 
-static void GetNonSequencerActorWorldTransforms(IMovieScenePlayer* Player, UMovieSceneSequence* InSequence, FMovieSceneSequenceIDRef Template, const FActorForWorldTransforms& ActorSelection, const TArray<FFrameNumber>& Frames, TArray<FTransform>& OutTransforms)
+static void GetNonSequencerActorWorldTransforms(IMovieScenePlayer* Player, FMovieSceneSequenceTransform RootToLocalTransform, UMovieSceneSequence* InSequence, FMovieSceneSequenceIDRef Template, const FActorForWorldTransforms& ActorSelection, const TArray<FFrameNumber>& Frames, TArray<FTransform>& OutTransforms)
 {
 	FName SocketName = ActorSelection.SocketName;
 	AActor* Actor = ActorSelection.Actor.Get();
@@ -4342,7 +4335,7 @@ static void GetNonSequencerActorWorldTransforms(IMovieScenePlayer* Player, UMovi
 			{
 				if (InSequence->GetMovieScene()->FindTrack<UMovieScene3DTransformTrack>(ActorHandle))
 				{
-					GetSequencerActorWorldTransforms(Player, InSequence, Template, NewActorSelection, Frames, OutTransforms);
+					GetSequencerActorWorldTransforms(Player, RootToLocalTransform, InSequence, Template, NewActorSelection, Frames, OutTransforms);
 					for (FTransform& OutTransform : OutTransforms)
 					{
 						OutTransform = WorldTransform * OutTransform;
@@ -4386,7 +4379,7 @@ void MovieSceneToolHelpers::GetActorWorldTransforms(ISequencer* Sequencer, const
 	{
 		if (Sequencer->GetFocusedMovieSceneSequence()->GetMovieScene()->FindTrack<UMovieScene3DTransformTrack>(ObjectHandle))
 		{
-			GetSequencerActorWorldTransforms(Sequencer, Sequencer->GetFocusedMovieSceneSequence(), Template, ActorSelection, Frames, OutWorldTransforms);
+			GetSequencerActorWorldTransforms(Sequencer, Sequencer->GetFocusedMovieSceneSequenceTransform(), Sequencer->GetFocusedMovieSceneSequence(), Template, ActorSelection, Frames, OutWorldTransforms);
 			return;
 		}
 	}
@@ -4396,18 +4389,18 @@ void MovieSceneToolHelpers::GetActorWorldTransforms(ISequencer* Sequencer, const
 	{
 		if (Sequencer->GetFocusedMovieSceneSequence()->GetMovieScene()->FindTrack<UMovieScene3DTransformTrack>(ObjectHandle))
 		{
-			GetSequencerActorWorldTransforms(Sequencer, Sequencer->GetFocusedMovieSceneSequence(), Template, ActorSelection, Frames, OutWorldTransforms);
+			GetSequencerActorWorldTransforms(Sequencer, Sequencer->GetFocusedMovieSceneSequenceTransform(), Sequencer->GetFocusedMovieSceneSequence(), Template, ActorSelection, Frames, OutWorldTransforms);
 			return;
 		}
 	}
 	if (bAvoidEvaluates)
 	{
 		TArray<FFrameNumber> NoFrame; //this will make sure we don't evaluate
-		GetNonSequencerActorWorldTransforms(Sequencer, Sequencer->GetFocusedMovieSceneSequence(), Template, ActorSelection, NoFrame, OutWorldTransforms);
+		GetNonSequencerActorWorldTransforms(Sequencer, Sequencer->GetFocusedMovieSceneSequenceTransform(), Sequencer->GetFocusedMovieSceneSequence(), Template, ActorSelection, NoFrame, OutWorldTransforms);
 	}
 	else
 	{
-		GetNonSequencerActorWorldTransforms(Sequencer, Sequencer->GetFocusedMovieSceneSequence(), Template, ActorSelection, Frames, OutWorldTransforms);
+		GetNonSequencerActorWorldTransforms(Sequencer, Sequencer->GetFocusedMovieSceneSequenceTransform(), Sequencer->GetFocusedMovieSceneSequence(), Template, ActorSelection, Frames, OutWorldTransforms);
 	}
 }
 
@@ -4624,6 +4617,7 @@ void MovieSceneToolHelpers::GetActorWorldTransforms(IMovieScenePlayer* Player, U
 {
 	FGuid ActorHandle = GetHandleToObject(ActorSelection.Actor.Get(), InSequence, Player, Template,false);
 	FGuid ComponentHandle = GetHandleToObject(ActorSelection.Component.Get(), InSequence, Player, Template,false);
+	FMovieSceneSequenceTransform RootToLocalTransform;
 	if (ActorHandle.IsValid() || ComponentHandle.IsValid())
 	{
 		//we can have handles but if they don't have a transform track the interrogator will return identity
@@ -4644,16 +4638,16 @@ void MovieSceneToolHelpers::GetActorWorldTransforms(IMovieScenePlayer* Player, U
 		}
 		if (bHaveTransformTrack)
 		{
-			GetSequencerActorWorldTransforms(Player, InSequence, Template, ActorSelection, Frames, OutWorldTransforms);
+			GetSequencerActorWorldTransforms(Player, RootToLocalTransform,InSequence, Template, ActorSelection, Frames, OutWorldTransforms);
 		}
 		else
 		{
-			GetNonSequencerActorWorldTransforms(Player, InSequence, Template, ActorSelection, Frames, OutWorldTransforms);
+			GetNonSequencerActorWorldTransforms(Player, RootToLocalTransform, InSequence, Template, ActorSelection, Frames, OutWorldTransforms);
 		}
 	}
 	else
 	{
-		GetNonSequencerActorWorldTransforms(Player, InSequence, Template, ActorSelection, Frames, OutWorldTransforms);
+		GetNonSequencerActorWorldTransforms(Player, RootToLocalTransform, InSequence, Template, ActorSelection, Frames, OutWorldTransforms);
 	}
 }
 
@@ -5393,6 +5387,8 @@ bool MovieSceneToolHelpers::AddTransformKeys(
 			}
 		}
 	}
+
+
 	//now we need to set auto tangents
 	for (const int32 ChannelIndex : ChannelsIndexToKey)
 	{
diff --git a/Engine/Source/Editor/MovieSceneTools/Private/TrackEditors/TransformTrackEditor.cpp b/Engine/Source/Editor/MovieSceneTools/Private/TrackEditors/TransformTrackEditor.cpp
--- a/Engine/Source/Editor/MovieSceneTools/Private/TrackEditors/TransformTrackEditor.cpp
+++ b/Engine/Source/Editor/MovieSceneTools/Private/TrackEditors/TransformTrackEditor.cpp
@@ -1024,30 +1024,6 @@ void F3DTransformTrackEditor::AddTransformKeys( UObject* ObjectToKey, const TOpt
 	AnimatablePropertyChanged( FOnKeyProperty::CreateLambda(OnKeyProperty) );
 }
 
-//todo move to external function so it can also be used by sdk/python
-static void UpdateTransformBasedOnConstraint(FTransform& CurrentTransform, USceneComponent* SceneComponent)
-{
-	TArray< TWeakObjectPtr<UTickableConstraint> > Constraints;
-	AActor* ShapeActor = SceneComponent->GetTypedOuter<AActor>();
-	
-	if (ShapeActor)
-	{
-		FTransformConstraintUtils::GetParentConstraints(SceneComponent->GetWorld(), ShapeActor, Constraints);
-
-		const int32 LastActiveIndex = FTransformConstraintUtils::GetLastActiveConstraintIndex(Constraints);
-		if (Constraints.IsValidIndex(LastActiveIndex))
-		{
-			// switch to constraint space
-			const FTransform WorldTransform = SceneComponent->GetSocketTransform(SceneComponent->GetAttachSocketName());
-			const TOptional<FTransform> RelativeTransform =
-				FTransformConstraintUtils::GetConstraintsRelativeTransform(Constraints, CurrentTransform, WorldTransform);
-			if (RelativeTransform)
-			{
-				CurrentTransform = *RelativeTransform; 
-			}
-		}
-	}
-}
 
 FTransformData F3DTransformTrackEditor::RecomposeTransform(const FTransformData& InTransformData, UObject* AnimatedObject, UMovieSceneSection* Section)
 {
@@ -1150,7 +1126,7 @@ FTransformData F3DTransformTrackEditor::RecomposeTransform(const FTransformData&
 		CurrentTransform *= GetTransformOrigin().Inverse();
 	}
 
-	UpdateTransformBasedOnConstraint(CurrentTransform, SceneComponent);
+	FTransformConstraintUtils::UpdateTransformBasedOnConstraint(CurrentTransform, SceneComponent);
 	
 	return FTransformData(CurrentTransform.GetLocation(), CurrentTransform.GetRotation().Rotator(), CurrentTransform.GetScale3D());
 }
diff --git a/Engine/Source/Runtime/Experimental/Animation/Constraints/Private/TransformConstraint.cpp b/Engine/Source/Runtime/Experimental/Animation/Constraints/Private/TransformConstraint.cpp
--- a/Engine/Source/Runtime/Experimental/Animation/Constraints/Private/TransformConstraint.cpp
+++ b/Engine/Source/Runtime/Experimental/Animation/Constraints/Private/TransformConstraint.cpp
@@ -1847,6 +1847,30 @@ bool FTransformConstraintUtils::AddConstraint(
 	return true;
 }
 
+void FTransformConstraintUtils::UpdateTransformBasedOnConstraint(FTransform& CurrentTransform, USceneComponent* SceneComponent)
+{
+	TArray< TWeakObjectPtr<UTickableConstraint> > Constraints;
+	AActor* ShapeActor = SceneComponent->GetTypedOuter<AActor>();
+
+	if (ShapeActor)
+	{
+		FTransformConstraintUtils::GetParentConstraints(SceneComponent->GetWorld(), ShapeActor, Constraints);
+
+		const int32 LastActiveIndex = FTransformConstraintUtils::GetLastActiveConstraintIndex(Constraints);
+		if (Constraints.IsValidIndex(LastActiveIndex))
+		{
+			// switch to constraint space
+			const FTransform WorldTransform = SceneComponent->GetSocketTransform(SceneComponent->GetAttachSocketName());
+			const TOptional<FTransform> RelativeTransform =
+				FTransformConstraintUtils::GetConstraintsRelativeTransform(Constraints, CurrentTransform, WorldTransform);
+			if (RelativeTransform)
+			{
+				CurrentTransform = *RelativeTransform;
+			}
+		}
+	}
+}
+
 FTransform FTransformConstraintUtils::ComputeRelativeTransform(
 	const FTransform& InChildLocal,
 	const FTransform& InChildWorld,
diff --git a/Engine/Source/Runtime/Experimental/Animation/Constraints/Public/TransformConstraint.h b/Engine/Source/Runtime/Experimental/Animation/Constraints/Public/TransformConstraint.h
--- a/Engine/Source/Runtime/Experimental/Animation/Constraints/Public/TransformConstraint.h
+++ b/Engine/Source/Runtime/Experimental/Animation/Constraints/Public/TransformConstraint.h
@@ -498,4 +498,8 @@ struct FTransformConstraintUtils
 		const UTransformableHandle* InHandle,
 		TArray< TWeakObjectPtr<UTickableConstraint> >& OutConstraints,
 		const bool bIncludeTarget = false);
+	
+	/** Adjust the transform on a scene component so it's effected by the constraint*/
+	static CONSTRAINTS_API void UpdateTransformBasedOnConstraint(FTransform& CurrentTransform, USceneComponent* SceneComponent);
+
 };
