Sequencer: Fixed soft object path BP nodes not being correctly populated from payload variables

diff --git a/Engine/Source/Editor/MovieSceneTools/Private/MovieSceneDirectorBlueprintUtils.cpp b/Engine/Source/Editor/MovieSceneTools/Private/MovieSceneDirectorBlueprintUtils.cpp
--- a/Engine/Source/Editor/MovieSceneTools/Private/MovieSceneDirectorBlueprintUtils.cpp
+++ b/Engine/Source/Editor/MovieSceneTools/Private/MovieSceneDirectorBlueprintUtils.cpp
@@ -405,7 +405,7 @@ bool FMovieSceneDirectorBlueprintUtils::GenerateEntryPointRawActorParameter(
 	const bool bMarkAsModified = false;
 	if (PayloadValue.ObjectValue.IsValid())
 	{
-		Schema->TrySetDefaultObject(*PathInput, PayloadValue.ObjectValue.ResolveObject(), bMarkAsModified);
+		Schema->TrySetDefaultValue(*PathInput, PayloadValue.ObjectValue.ToString(), bMarkAsModified);
 	}
 	else if (!PayloadValue.Value.IsEmpty())
 	{
