Sequencer: Convert from display rate to tick resolution for DrawFrameHint

#robmerge 5.4
#rnx

diff --git a/Engine/Source/Editor/MovieSceneTools/Private/TrackEditors/SkeletalAnimationTrackEditor.cpp b/Engine/Source/Editor/MovieSceneTools/Private/TrackEditors/SkeletalAnimationTrackEditor.cpp
--- a/Engine/Source/Editor/MovieSceneTools/Private/TrackEditors/SkeletalAnimationTrackEditor.cpp
+++ b/Engine/Source/Editor/MovieSceneTools/Private/TrackEditors/SkeletalAnimationTrackEditor.cpp
@@ -682,6 +682,9 @@ TOptional<FFrameTime> FSkeletalAnimationSection::GetSectionTime(FSequencerSectio
 		HintFrameTime = FQualifiedFrameTime(FrameTime, SamplingFrameRate);
 	}
 
+	// Convert to tick resolution
+	HintFrameTime = FQualifiedFrameTime(ConvertFrameTime(HintFrameTime.Time, SamplingFrameRate, TickResolution), TickResolution);
+
 	// Get the desired frame display format and zero padding from
 	// the sequencer settings, if possible.
 	TAttribute<EFrameNumberDisplayFormats> DisplayFormatAttr(EFrameNumberDisplayFormats::Frames);
