Sequencer: Fix import fbx replace transform not being set properly.

The settings are temporarily changed to the settings passed into the function. bReplaceTransformTrack was missing from the copy.

diff --git a/Engine/Source/Editor/MovieSceneTools/Private/MovieSceneToolHelpers.cpp b/Engine/Source/Editor/MovieSceneTools/Private/MovieSceneToolHelpers.cpp
--- a/Engine/Source/Editor/MovieSceneTools/Private/MovieSceneToolHelpers.cpp
+++ b/Engine/Source/Editor/MovieSceneTools/Private/MovieSceneToolHelpers.cpp
@@ -3470,6 +3470,7 @@ bool MovieSceneToolHelpers::ImportFBXIfReady(UWorld* World, UMovieSceneSequence*
 	CurrentImportFBXSettings->ReduceKeysTolerance = ImportFBXSettings->ReduceKeysTolerance;
 	CurrentImportFBXSettings->bConvertSceneUnit = ImportFBXSettings->bConvertSceneUnit;
 	CurrentImportFBXSettings->ImportUniformScale = ImportFBXSettings->ImportUniformScale;
+	CurrentImportFBXSettings->bReplaceTransformTrack = ImportFBXSettings->bReplaceTransformTrack;
 	UnFbx::FFbxImporter* FbxImporter = UnFbx::FFbxImporter::GetInstance();
 
 	UnFbx::FFbxCurvesAPI CurveAPI;
