Constraints: handles can be invalid if baking out of a sequence boundary or with spwanables so make sure that sequencer is evaluated once the transforms computed.

#rnx

diff --git a/Engine/Source/Editor/MovieSceneTools/Private/Constraints/MovieSceneConstraintChannelHelper.cpp b/Engine/Source/Editor/MovieSceneTools/Private/Constraints/MovieSceneConstraintChannelHelper.cpp
--- a/Engine/Source/Editor/MovieSceneTools/Private/Constraints/MovieSceneConstraintChannelHelper.cpp
+++ b/Engine/Source/Editor/MovieSceneTools/Private/Constraints/MovieSceneConstraintChannelHelper.cpp
@@ -204,6 +204,11 @@ void FCompensationEvaluator::ComputeLocalTransformsForBaking(UWorld* InWorld, co
 		return;
 	}
 
+	if (!IsValid(Handle) || !Handle->IsValid())
+	{
+		return;
+	}
+	
 	using ConstraintPtr = TWeakObjectPtr<UTickableConstraint>;
 	const TArray< ConstraintPtr > Constraints = GetHandleTransformConstraints(InWorld);
 	
@@ -315,6 +320,13 @@ void FCompensationEvaluator::ComputeLocalTransformsForBaking(UWorld* InWorld, co
 		}
 	}
 
+	const bool bIsValidAfterBaking = IsValid(Handle) && Handle->IsValid();
+	if (!bIsValidAfterBaking)
+	{
+		// the handle might not be valid after baking due to spawnables or baking out of the sequence boundaries
+		// so force sequencer evaluation to make sure we're back to normal
+		InSequencer->ForceEvaluate();
+	}
 }
 void FCompensationEvaluator::ComputeLocalTransformsBeforeDeletion(
 	UWorld* InWorld,
