Sequencer: Fix CanTrimSectionLeft, CanTrimSectionRight, CanSplitSection so that you can't end up with empty sections

diff --git a/Engine/Source/Editor/MovieSceneTools/Private/MovieSceneToolHelpers.cpp b/Engine/Source/Editor/MovieSceneTools/Private/MovieSceneToolHelpers.cpp
--- a/Engine/Source/Editor/MovieSceneTools/Private/MovieSceneToolHelpers.cpp
+++ b/Engine/Source/Editor/MovieSceneTools/Private/MovieSceneToolHelpers.cpp
@@ -138,6 +138,39 @@ void MovieSceneToolHelpers::TrimSection(const TSet<UMovieSceneSection*>& Section
 	}
 }
 
+bool MovieSceneToolHelpers::CanTrimSectionLeft(const TSet<UMovieSceneSection*>& Sections, FQualifiedFrameTime Time)
+{
+	for (UMovieSceneSection* Section : Sections)
+	{
+		if (Section && Section->HasStartFrame() && Time.Time.FrameNumber > Section->GetInclusiveStartFrame())
+		{
+			// Don't allow an empty section
+			if (Section->HasEndFrame() && Time.Time.FrameNumber >= Section->GetExclusiveEndFrame())
+			{
+				continue;
+			}
+			return true;
+		}
+	}
+	return false;
+}
+
+bool MovieSceneToolHelpers::CanTrimSectionRight(const TSet<UMovieSceneSection*>& Sections, FQualifiedFrameTime Time)
+{
+	for (UMovieSceneSection* Section : Sections)
+	{
+		if (Section && Section->HasEndFrame() && Time.Time.FrameNumber < Section->GetExclusiveEndFrame())
+		{
+			// Don't allow an empty section
+			if (Section->HasStartFrame() && Time.Time.FrameNumber <= Section->GetInclusiveStartFrame())
+			{
+				continue;
+			}
+			return true;
+		}
+	}
+	return false;
+}
 
 void MovieSceneToolHelpers::TrimOrExtendSection(UMovieSceneTrack* Track, TOptional<int32> SpecifiedRowIndex, FQualifiedFrameTime Time, bool bTrimOrExtendLeft, bool bDeleteKeys)
 {
@@ -232,6 +265,26 @@ void MovieSceneToolHelpers::SplitSection(const TSet<UMovieSceneSection*>& Sectio
 	}
 }
 
+bool MovieSceneToolHelpers::CanSplitSection(const TSet<UMovieSceneSection*>& Sections, FQualifiedFrameTime Time)
+{
+	for (UMovieSceneSection* Section : Sections)
+	{
+		if (Section && Section->IsTimeWithinSection(Time.Time.FrameNumber))
+		{
+			if (Section->HasStartFrame() && Time.Time.FrameNumber <= Section->GetInclusiveStartFrame())
+			{
+				continue;
+			}
+			if (Section->HasEndFrame() && Time.Time.FrameNumber >= Section->GetExclusiveEndFrame())
+			{
+				continue;
+			}
+			return true;
+		}
+	}
+	return false;
+}
+
 bool MovieSceneToolHelpers::ParseShotName(const FString& InShotName, FString& ShotPrefix, uint32& ShotNumber, uint32& TakeNumber, uint32& ShotNumberDigits, uint32& TakeNumberDigits)
 {
 	// Parse a shot name
diff --git a/Engine/Source/Editor/MovieSceneTools/Public/MovieSceneToolHelpers.h b/Engine/Source/Editor/MovieSceneTools/Public/MovieSceneToolHelpers.h
--- a/Engine/Source/Editor/MovieSceneTools/Public/MovieSceneToolHelpers.h
+++ b/Engine/Source/Editor/MovieSceneTools/Public/MovieSceneToolHelpers.h
@@ -109,6 +109,8 @@ public:
 	 * @param bDeleteKeys Delete keys outside the split ranges
 	 */
 	static void TrimSection(const TSet<UMovieSceneSection*>& Sections, FQualifiedFrameTime Time, bool bTrimLeft, bool bDeleteKeys);
+	static bool CanTrimSectionLeft(const TSet<UMovieSceneSection*>& Sections, FQualifiedFrameTime Time);
+	static bool CanTrimSectionRight(const TSet<UMovieSceneSection*>& Sections, FQualifiedFrameTime Time);
 
 	/**
 	 * Trim or extend section at the given time
@@ -129,6 +131,7 @@ public:
 	 * @param bDeleteKeys Delete keys outside the split ranges
 	 */
 	static void SplitSection(const TSet<UMovieSceneSection*>& Sections, FQualifiedFrameTime Time, bool bDeleteKeys);
+	static bool CanSplitSection(const TSet<UMovieSceneSection*>& Sections, FQualifiedFrameTime Time);
 
 	/**
 	 * Parse a shot name into its components.
diff --git a/Engine/Source/Editor/Sequencer/Private/Sequencer.cpp b/Engine/Source/Editor/Sequencer/Private/Sequencer.cpp
--- a/Engine/Source/Editor/Sequencer/Private/Sequencer.cpp
+++ b/Engine/Source/Editor/Sequencer/Private/Sequencer.cpp
@@ -10673,27 +10673,15 @@ void FSequencer::BindCommands()
 		Commands.TranslateRight,
 		FExecuteAction::CreateSP( this, &FSequencer::TranslateSelectedKeysAndSections, false) );
 
-	auto CanTrimSection = [this]{
-		for (UMovieSceneSection* Section : this->ViewModel->GetSelection()->GetSelectedSections())
-		{
-			if (Section && Section->IsTimeWithinSection(GetLocalTime().Time.FrameNumber))
-			{
-				return true;
-			}
-		}
-		return false;
-	};
-
 	SequencerCommandBindings->MapAction(
 		Commands.TrimSectionLeft,
 		FExecuteAction::CreateSP( this, &FSequencer::TrimSection, true ),
-		FCanExecuteAction::CreateLambda(CanTrimSection));
-
+		FCanExecuteAction::CreateLambda([this] { return MovieSceneToolHelpers::CanTrimSectionLeft(this->ViewModel->GetSelection()->GetSelectedSections(), GetLocalTime()); }));
 
 	SequencerCommandBindings->MapAction(
 		Commands.TrimSectionRight,
 		FExecuteAction::CreateSP( this, &FSequencer::TrimSection, false ),
-		FCanExecuteAction::CreateLambda(CanTrimSection));
+		FCanExecuteAction::CreateLambda([this] { return MovieSceneToolHelpers::CanTrimSectionRight(this->ViewModel->GetSelection()->GetSelectedSections(), GetLocalTime()); }));
 
 	SequencerCommandBindings->MapAction(
 		Commands.TrimOrExtendSectionLeft,
@@ -10706,7 +10694,7 @@ void FSequencer::BindCommands()
 	SequencerCommandBindings->MapAction(
 		Commands.SplitSection,
 		FExecuteAction::CreateSP( this, &FSequencer::SplitSection ),
-		FCanExecuteAction::CreateLambda(CanTrimSection));
+		FCanExecuteAction::CreateLambda([this] { return MovieSceneToolHelpers::CanSplitSection(this->ViewModel->GetSelection()->GetSelectedSections(), GetLocalTime()); }));
 
 	// We can convert to spawnables if anything selected is a root-level possessable
 	auto CanConvertToSpawnables = [this]{
