Sequencer: Control Rig: Make sure when collapsing by key the key is in section range

diff --git a/Engine/Source/Editor/MovieSceneTools/Private/MovieSceneToolHelpers.cpp b/Engine/Source/Editor/MovieSceneTools/Private/MovieSceneToolHelpers.cpp
--- a/Engine/Source/Editor/MovieSceneTools/Private/MovieSceneToolHelpers.cpp
+++ b/Engine/Source/Editor/MovieSceneTools/Private/MovieSceneToolHelpers.cpp
@@ -5014,10 +5014,13 @@ static bool MergeChannels(TArray<ChannelType*>& Channels,const  TArray<UMovieSce
 		Handles.Reset();
 		ChannelType* Channel = Channels[ChannelIndex];
 		Channel->GetKeys(Range, &KeyTimes, &Handles);
-		double DNumChannels = (double)(Channels.Num());
 		for (int32 FrameIndex = 0; FrameIndex < KeyTimes.Num(); ++FrameIndex)
 		{
 			const FFrameNumber& Frame = KeyTimes[FrameIndex];
+			if (Sections[0]->GetRange().Contains(Frame) == false)  //frame is outside base range so skip
+			{
+				continue;
+			}
 			const FFrameTime FrameTime(Frame);
 			int32 KeyIndex = Channel->GetData().GetIndex(Handles[FrameIndex]);
 			ChannelValueType Value = Channel->GetData().GetValues()[KeyIndex];
@@ -5025,26 +5028,38 @@ static bool MergeChannels(TArray<ChannelType*>& Channels,const  TArray<UMovieSce
 			Value.Value = 0.0; //zero out the value we calculate it 
 			if (MergeAlgorithm == FChannelMergeAlgorithm::Average)
 			{
+				double DNumChannels = 0.0;
 				for (int32 WeightIndex = 0; WeightIndex < Sections.Num(); ++WeightIndex)
 				{
-					float Weight = Sections[WeightIndex]->GetTotalWeightValue(FrameTime);
-					CurveValueType WeightedValue = 0.0;
-					ChannelType* EachChannel = Channels[WeightIndex];
-					EachChannel->Evaluate(FrameTime, WeightedValue);
-					WeightedValue *= ((double)Weight / DNumChannels);
-					Value.Value += WeightedValue;
+					if (Sections[WeightIndex]->GetRange().Contains(Frame))  
+					{
+						float Weight = Sections[WeightIndex]->GetTotalWeightValue(FrameTime);
+						CurveValueType WeightedValue = 0.0;
+						ChannelType* EachChannel = Channels[WeightIndex];
+						EachChannel->Evaluate(FrameTime, WeightedValue);
+						WeightedValue *= ((double)Weight);
+						Value.Value += WeightedValue;
+						DNumChannels += 1.0;
+					}
+				}
+				if (DNumChannels > 0.0) //should always happen since base(0) at least be here
+				{
+					Value.Value /= DNumChannels;
 				}
 			}
 			else if (MergeAlgorithm == FChannelMergeAlgorithm::Add)
 			{
 				for (int32 WeightIndex = 0; WeightIndex < Sections.Num(); ++WeightIndex)
 				{
-					float Weight = Sections[WeightIndex]->GetTotalWeightValue(FrameTime);
-					CurveValueType WeightedValue = 0.0;
-					ChannelType* EachChannel = Channels[WeightIndex];
-					EachChannel->Evaluate(FrameTime, WeightedValue);
-					WeightedValue *= Weight;
-					Value.Value += WeightedValue;
+					if (Sections[WeightIndex]->GetRange().Contains(Frame))  
+					{
+						float Weight = Sections[WeightIndex]->GetTotalWeightValue(FrameTime);
+						CurveValueType WeightedValue = 0.0;
+						ChannelType* EachChannel = Channels[WeightIndex];
+						EachChannel->Evaluate(FrameTime, WeightedValue);
+						WeightedValue *= Weight;
+						Value.Value += WeightedValue;
+					}
 				}
 			}
 			KeysToSet.Add(TPair<FFrameNumber, ChannelValueType>(Frame, Value));
