Sequencer- Fix issue where you can't add Spawn Track to new custom spawnables if they don't already have one.

diff --git a/Engine/Source/Editor/MovieSceneTools/Private/TrackEditors/SpawnTrackEditor.cpp b/Engine/Source/Editor/MovieSceneTools/Private/TrackEditors/SpawnTrackEditor.cpp
--- a/Engine/Source/Editor/MovieSceneTools/Private/TrackEditors/SpawnTrackEditor.cpp
+++ b/Engine/Source/Editor/MovieSceneTools/Private/TrackEditors/SpawnTrackEditor.cpp
@@ -19,6 +19,8 @@
 #include "Tracks/MovieSceneSpawnTrack.h"
 #include "UObject/Class.h"
 #include "UObject/UnrealNames.h"
+#include "MovieSceneBindingReferences.h"
+#include "Bindings/MovieSceneSpawnableBinding.h"
 
 class ISequencerTrackEditor;
 
@@ -56,10 +58,30 @@ void FSpawnTrackEditor::BuildObjectBindingTrackMenu(FMenuBuilder& MenuBuilder, c
 {
 	UMovieSceneSequence* MovieSequence = GetSequencer()->GetFocusedMovieSceneSequence();
 
-	if (!MovieSequence || MovieSequence->GetClass()->GetName() != TEXT("LevelSequence") || !MovieSequence->GetMovieScene()->FindSpawnable(ObjectBindings[0]))
+	if (!MovieSequence || MovieSequence->GetClass()->GetName() != TEXT("LevelSequence"))
 	{
 		return;
 	}
+	else
+	{
+		bool bSpawnable = false;
+		if (MovieSequence->GetMovieScene()->FindSpawnable(ObjectBindings[0]))
+		{
+			bSpawnable = true;
+		}
+		else if (const FMovieSceneBindingReferences* BindingReferences = MovieSequence->GetBindingReferences())
+		{
+			if (Algo::AnyOf(BindingReferences->GetReferences(ObjectBindings[0]), [](const FMovieSceneBindingReference& BindingReference) { return BindingReference.CustomBinding && BindingReference.CustomBinding->IsA<UMovieSceneSpawnableBindingBase>(); }))
+			{
+				bSpawnable = true;
+			}
+		}
+
+		if (!bSpawnable)
+		{
+			return;
+		}
+	}
 
 	MenuBuilder.AddMenuEntry(
 		LOCTEXT("AddSpawnTrack", "Spawn Track"),
