Sequencer: Animation Baking in subsequences would get the incorrect time range.  Need both root/current movie scene sequence to get the range stored in the subsequence data. Had to combine the two reviews since of overlapping changes and P4 suckage at handling this, plus issues with preflights.  Fixed issues brought up in the reviews in particually moved sequence parameters over to one struct and made sure all deprecrated functions call the new ones.

diff --git a/Engine/Plugins/Animation/ControlRig/Source/ControlRig/Private/Sequencer/MovieSceneControlRigParameterSection.cpp b/Engine/Plugins/Animation/ControlRig/Source/ControlRig/Private/Sequencer/MovieSceneControlRigParameterSection.cpp
--- a/Engine/Plugins/Animation/ControlRig/Source/ControlRig/Private/Sequencer/MovieSceneControlRigParameterSection.cpp
+++ b/Engine/Plugins/Animation/ControlRig/Source/ControlRig/Private/Sequencer/MovieSceneControlRigParameterSection.cpp
@@ -3071,6 +3071,12 @@ void UMovieSceneControlRigParameterSection::RecordControlRigKey(FFrameNumber Fra
 }
 
 bool UMovieSceneControlRigParameterSection::LoadAnimSequenceIntoThisSection(UAnimSequence* AnimSequence, UMovieScene* MovieScene, UObject* BoundObject, bool bKeyReduce, float Tolerance, bool bResetControls, FFrameNumber InStartFrame, EMovieSceneKeyInterpolation InInterpolation)
+{
+	FFrameNumber SequenceStart = UE::MovieScene::DiscreteInclusiveLower(MovieScene->GetPlaybackRange());
+	return LoadAnimSequenceIntoThisSection(AnimSequence, SequenceStart, MovieScene, BoundObject, bKeyReduce, Tolerance, bResetControls, InStartFrame, InInterpolation);
+}
+
+bool UMovieSceneControlRigParameterSection::LoadAnimSequenceIntoThisSection(UAnimSequence* AnimSequence, const FFrameNumber& SequenceStart, UMovieScene* MovieScene, UObject* BoundObject, bool bKeyReduce, float Tolerance, bool bResetControls, const FFrameNumber& InStartFrame, EMovieSceneKeyInterpolation InInterpolation)
 {
 	USkeletalMeshComponent* SkelMeshComp = Cast<USkeletalMeshComponent>(BoundObject);
 	
@@ -3102,7 +3108,7 @@ bool UMovieSceneControlRigParameterSection::LoadAnimSequenceIntoThisSection(UAni
 	float Length = AnimSequence->GetPlayLength();
 	const FFrameRate& FrameRate = AnimSequence->GetSamplingFrameRate();
 
-	FFrameNumber StartFrame = UE::MovieScene::DiscreteInclusiveLower(MovieScene->GetPlaybackRange()) + InStartFrame;
+	FFrameNumber StartFrame = SequenceStart + InStartFrame;
 	FFrameNumber EndFrame = TickResolution.AsFrameNumber(Length) + StartFrame;
 
 	Modify();
diff --git a/Engine/Plugins/Animation/ControlRig/Source/ControlRig/Public/Sequencer/MovieSceneControlRigParameterSection.h b/Engine/Plugins/Animation/ControlRig/Source/ControlRig/Public/Sequencer/MovieSceneControlRigParameterSection.h
--- a/Engine/Plugins/Animation/ControlRig/Source/ControlRig/Public/Sequencer/MovieSceneControlRigParameterSection.h
+++ b/Engine/Plugins/Animation/ControlRig/Source/ControlRig/Public/Sequencer/MovieSceneControlRigParameterSection.h
@@ -315,7 +315,10 @@ public:
 
 	//Function to load an Anim Sequence into this section. It will automatically resize to the section size.
 	//Will return false if fails or is canceled
-	virtual bool LoadAnimSequenceIntoThisSection(UAnimSequence* Sequence, UMovieScene* MovieScene, UObject* BoundObject, bool bKeyReduce, float Tolerance, bool bResetControls, FFrameNumber InStartFrame , EMovieSceneKeyInterpolation InInterolation);
+	virtual bool LoadAnimSequenceIntoThisSection(UAnimSequence* Sequence, const FFrameNumber& SequenceStart, UMovieScene* MovieScene, UObject* BoundObject, bool bKeyReduce, float Tolerance, bool bResetControls, const FFrameNumber& InStartFrame, EMovieSceneKeyInterpolation InInterpolation);
+
+	UE_DEPRECATED(5.5, "LoadAnimSequenceIntoThisSection without taking a sequence start frame is deprecated, use version that takes start frame instead")
+	virtual bool LoadAnimSequenceIntoThisSection(UAnimSequence* Sequence, UMovieScene* MovieScene, UObject* BoundObject, bool bKeyReduce, float Tolerance, bool bResetControls, FFrameNumber InStartFrame , EMovieSceneKeyInterpolation InInterpolation);
 #endif
 	const TArray<bool>& GetControlsMask() 
 	{
diff --git a/Engine/Plugins/Animation/ControlRig/Source/ControlRigEditor/Private/ControlRigSequencerEditorLibrary.cpp b/Engine/Plugins/Animation/ControlRig/Source/ControlRigEditor/Private/ControlRigSequencerEditorLibrary.cpp
--- a/Engine/Plugins/Animation/ControlRig/Source/ControlRigEditor/Private/ControlRigSequencerEditorLibrary.cpp
+++ b/Engine/Plugins/Animation/ControlRig/Source/ControlRigEditor/Private/ControlRigSequencerEditorLibrary.cpp
@@ -1178,7 +1178,12 @@ bool UControlRigSequencerEditorLibrary::BakeToControlRig(UWorld* World, ULevelSe
 			{
 				UAnimSequence* TempAnimSequence = NewObject<UAnimSequence>(GetTransientPackage(), NAME_None);
 				TempAnimSequence->SetSkeleton(Skeleton);
-				bResult = MovieSceneToolHelpers::ExportToAnimSequence(TempAnimSequence, ExportOptions, MovieScene, Player, SkeletalMeshComp, Template, RootToLocalTransform);
+				FAnimExportSequenceParameters AESP;
+				AESP.Player = Player;
+				AESP.RootToLocalTransform = RootToLocalTransform;
+				AESP.MovieSceneSequence = LevelSequence;
+				AESP.RootMovieSceneSequence = LevelSequence;
+				bResult = MovieSceneToolHelpers::ExportToAnimSequence(TempAnimSequence, ExportOptions, AESP, SkeletalMeshComp);
 				if (bResult == false)
 				{
 					TempAnimSequence->MarkAsGarbage();
@@ -1278,8 +1283,8 @@ bool UControlRigSequencerEditorLibrary::BakeToControlRig(UWorld* World, ULevelSe
 						WeakSequencer.Pin()->NotifyMovieSceneDataChanged(EMovieSceneDataChangeType::MovieSceneStructureItemAdded);
 						DefaultInterpolation = WeakSequencer.Pin()->GetKeyInterpolation();
 					}
-
-					ParamSection->LoadAnimSequenceIntoThisSection(TempAnimSequence, MovieScene, SkeletalMeshComp,
+					FFrameNumber StartFrame = MovieScene->GetPlaybackRange().GetLowerBoundValue();
+					ParamSection->LoadAnimSequenceIntoThisSection(TempAnimSequence, StartFrame, MovieScene, SkeletalMeshComp,
 						bReduceKeys, Tolerance, bResetControls, FFrameNumber(0), DefaultInterpolation);
 
 					//Turn Off Any Skeletal Animation Tracks
@@ -1345,7 +1350,8 @@ bool UControlRigSequencerEditorLibrary::LoadAnimSequenceIntoControlRigSection(UM
 		{
 			InStartFrame = FFrameRate::TransformTime(FFrameTime(InStartFrame, 0), MovieScene->GetDisplayRate(),MovieScene->GetTickResolution()).RoundToFrame();
 		}
-		return Section->LoadAnimSequenceIntoThisSection(AnimSequence, MovieScene, SkelMeshComp, bKeyReduce, Tolerance, bResetControls, InStartFrame, Interpolation);
+		FFrameNumber StartFrame  = MovieScene->GetPlaybackRange().GetLowerBoundValue();
+		return Section->LoadAnimSequenceIntoThisSection(AnimSequence, StartFrame,  MovieScene, SkelMeshComp, bKeyReduce, Tolerance, bResetControls, InStartFrame, Interpolation);
 	}
 	return false;
 }
diff --git a/Engine/Plugins/Animation/ControlRig/Source/ControlRigEditor/Private/EditMode/SControlRigSnapper.cpp b/Engine/Plugins/Animation/ControlRig/Source/ControlRigEditor/Private/EditMode/SControlRigSnapper.cpp
--- a/Engine/Plugins/Animation/ControlRig/Source/ControlRigEditor/Private/EditMode/SControlRigSnapper.cpp
+++ b/Engine/Plugins/Animation/ControlRig/Source/ControlRigEditor/Private/EditMode/SControlRigSnapper.cpp
@@ -129,13 +129,15 @@ public:
 };
 
 //function to keep framenumber from beign too large do to user error, it can cause a crash if we snap too many frames.
-static void KeepFrameInRange(FFrameNumber& KeepFrameInRange, UMovieScene* MovieScene)
+static void KeepFrameInRange(TSharedPtr<ISequencer>& SequencerPtr, FFrameNumber& KeepFrameInRange, UMovieScene* MovieScene)
 {
-	if (MovieScene)
+	if (MovieScene && SequencerPtr.IsValid())
 	{
 		//limit values to 10x times the range from the end points.
-		FFrameNumber StartFrame = MovieScene->GetPlaybackRange().GetLowerBoundValue();
-		FFrameNumber EndFrame = MovieScene->GetPlaybackRange().GetUpperBoundValue();
+		TOptional<TRange<FFrameNumber>> OptionalRange = SequencerPtr->GetSubSequenceRange();
+		FFrameNumber StartFrame = OptionalRange.IsSet() ? OptionalRange.GetValue().GetLowerBoundValue() : MovieScene->GetPlaybackRange().GetLowerBoundValue();
+		FFrameNumber EndFrame = OptionalRange.IsSet() ? OptionalRange.GetValue().GetUpperBoundValue() : MovieScene->GetPlaybackRange().GetUpperBoundValue();
+
 		if (EndFrame < StartFrame)
 		{
 			FFrameNumber Temp = StartFrame;
@@ -285,15 +287,14 @@ void SControlRigSnapper::Construct(const FArguments& InArgs)
 								.ToolTipText(LOCTEXT("GetStartFrameTooltip", "Set first frame to snap"))
 								.OnTextCommitted_Lambda([this](const FText& InText, ETextCommit::Type TextCommitType)
 									{
-										TWeakPtr<ISequencer> Sequencer = Snapper.GetSequencer();
-										if (Sequencer.IsValid() && Sequencer.Pin()->GetFocusedMovieSceneSequence())
+										TSharedPtr<ISequencer> SequencerPtr = Snapper.GetSequencer().Pin();
+										if (SequencerPtr.IsValid() && SequencerPtr->GetFocusedMovieSceneSequence())
 										{
-
-											TOptional<double> NewFrameTime = Sequencer.Pin()->GetNumericTypeInterface()->FromString(InText.ToString(), 0);
+											TOptional<double> NewFrameTime = SequencerPtr->GetNumericTypeInterface()->FromString(InText.ToString(), 0);
 											if (NewFrameTime.IsSet())
 											{
 												StartFrame = FFrameNumber((int32)NewFrameTime.GetValue());
-												KeepFrameInRange(StartFrame, Sequencer.Pin()->GetFocusedMovieSceneSequence()->GetMovieScene());
+												KeepFrameInRange(SequencerPtr, StartFrame, SequencerPtr->GetFocusedMovieSceneSequence()->GetMovieScene());
 											}
 										}
 									})
@@ -322,12 +323,12 @@ void SControlRigSnapper::Construct(const FArguments& InArgs)
 										TWeakPtr<ISequencer> Sequencer = Snapper.GetSequencer();
 										if (Sequencer.IsValid() && Sequencer.Pin()->GetFocusedMovieSceneSequence())
 										{
-
-											TOptional<double> NewFrameTime = Sequencer.Pin()->GetNumericTypeInterface()->FromString(InText.ToString(), 0);
+											TSharedPtr<ISequencer> SequencerPtr = Sequencer.Pin();
+											TOptional<double> NewFrameTime = SequencerPtr->GetNumericTypeInterface()->FromString(InText.ToString(), 0);
 											if (NewFrameTime.IsSet())
 											{
 												EndFrame = FFrameNumber((int32)NewFrameTime.GetValue());
-												KeepFrameInRange(EndFrame, Sequencer.Pin()->GetFocusedMovieSceneSequence()->GetMovieScene());
+												KeepFrameInRange(SequencerPtr, EndFrame, SequencerPtr->GetFocusedMovieSceneSequence()->GetMovieScene());
 											}
 										}
 									})
@@ -493,8 +494,10 @@ void SControlRigSnapper::SetStartEndFrames()
 	if (Sequencer.IsValid()  && Sequencer.Pin()->GetFocusedMovieSceneSequence())
 	{
 		UMovieScene* MovieScene = Sequencer.Pin()->GetFocusedMovieSceneSequence()->GetMovieScene();
-		StartFrame = MovieScene->GetPlaybackRange().GetLowerBoundValue();
-		EndFrame = MovieScene->GetPlaybackRange().GetUpperBoundValue();
+		TOptional<TRange<FFrameNumber>> OptionalRange = Sequencer.Pin()->GetSubSequenceRange();
+		StartFrame = OptionalRange.IsSet() ? OptionalRange.GetValue().GetLowerBoundValue() : MovieScene->GetPlaybackRange().GetLowerBoundValue();
+		EndFrame = OptionalRange.IsSet() ? OptionalRange.GetValue().GetUpperBoundValue() : MovieScene->GetPlaybackRange().GetUpperBoundValue();
+
 	}
 }
 
diff --git a/Engine/Plugins/Animation/ControlRig/Source/ControlRigEditor/Private/Sequencer/ControlRigParameterTrackEditor.cpp b/Engine/Plugins/Animation/ControlRig/Source/ControlRigEditor/Private/Sequencer/ControlRigParameterTrackEditor.cpp
--- a/Engine/Plugins/Animation/ControlRig/Source/ControlRigEditor/Private/Sequencer/ControlRigParameterTrackEditor.cpp
+++ b/Engine/Plugins/Animation/ControlRig/Source/ControlRigEditor/Private/Sequencer/ControlRigParameterTrackEditor.cpp
@@ -989,7 +989,7 @@ bool FControlRigParameterTrackEditor::LoadAnimationIntoSection(TSharedPtr<ISeque
 	EMovieSceneKeyInterpolation DefaultInterpolation = SequencerPtr->GetKeyInterpolation();
 	UMovieSceneSequence* OwnerSequence = SequencerPtr->GetFocusedMovieSceneSequence();
 	UMovieScene* OwnerMovieScene = OwnerSequence->GetMovieScene();
-	if (ParamSection->LoadAnimSequenceIntoThisSection(AnimSequence, OwnerMovieScene, SkelMeshComp,
+	if (ParamSection->LoadAnimSequenceIntoThisSection(AnimSequence, FFrameNumber(0), OwnerMovieScene, SkelMeshComp,
 		false, 0.0, bResetControls, StartFrame, DefaultInterpolation))
 	{
 		if (bReduceKeys)
@@ -1010,6 +1010,7 @@ void FControlRigParameterTrackEditor::BakeToControlRig(UClass* InClass, FGuid Ob
 	{
 
 		UMovieSceneSequence* OwnerSequence = GetSequencer()->GetFocusedMovieSceneSequence();
+		UMovieSceneSequence* RootSequence = GetSequencer()->GetRootMovieSceneSequence();
 		UMovieScene* OwnerMovieScene = OwnerSequence->GetMovieScene();
 		{
 			UAnimSequence* TempAnimSequence = NewObject<UAnimSequence>(GetTransientPackage(), NAME_None);
@@ -1047,8 +1048,12 @@ void FControlRigParameterTrackEditor::BakeToControlRig(UClass* InClass, FGuid Ob
 
 			if (OptionWindow.Get()->ShouldExport())
 			{
-
-				bool bResult = MovieSceneToolHelpers::ExportToAnimSequence(TempAnimSequence, AnimSeqExportOption, OwnerMovieScene, ParentSequencer.Get(), SkelMeshComp, Template, RootToLocalTransform);
+				FAnimExportSequenceParameters AESP;
+				AESP.Player = ParentSequencer.Get();
+				AESP.RootToLocalTransform = RootToLocalTransform;
+				AESP.MovieSceneSequence = OwnerSequence;
+				AESP.RootMovieSceneSequence = RootSequence;
+				bool bResult = MovieSceneToolHelpers::ExportToAnimSequence(TempAnimSequence, AnimSeqExportOption, AESP, SkelMeshComp);
 				if (bResult == false)
 				{
 					TempAnimSequence->MarkAsGarbage();
@@ -1129,12 +1134,14 @@ void FControlRigParameterTrackEditor::BakeToControlRig(UClass* InClass, FGuid Ob
 					Track->SetTrackName(FName(*ObjectName));
 					Track->SetDisplayName(FText::FromString(ObjectName));
 
-					GetSequencer()->EmptySelection();
-					GetSequencer()->SelectSection(NewSection);
-					GetSequencer()->ThrobSectionSelection();
-					GetSequencer()->NotifyMovieSceneDataChanged(EMovieSceneDataChangeType::MovieSceneStructureItemAdded);
 					TSharedPtr<ISequencer> SequencerPtr = GetSequencer();
-					LoadAnimationIntoSection(SequencerPtr, TempAnimSequence, SkelMeshComp, FFrameNumber(0),
+					SequencerPtr->EmptySelection();
+					SequencerPtr->SelectSection(NewSection);
+					SequencerPtr->ThrobSectionSelection();
+					SequencerPtr->NotifyMovieSceneDataChanged(EMovieSceneDataChangeType::MovieSceneStructureItemAdded);
+					TOptional<TRange<FFrameNumber>> OptionalRange = SequencerPtr->GetSubSequenceRange();
+					FFrameNumber StartFrame = OptionalRange.IsSet() ? OptionalRange.GetValue().GetLowerBoundValue() : OwnerMovieScene->GetPlaybackRange().GetLowerBoundValue();
+					LoadAnimationIntoSection(SequencerPtr, TempAnimSequence, SkelMeshComp, StartFrame,
 						BakeSettings->bReduceKeys, BakeSettings->SmartReduce, BakeSettings->bResetControls, ParamSection);
 					//Turn Off Any Skeletal Animation Tracks
 					TArray<UMovieSceneSkeletalAnimationTrack*> SkelAnimationTracks;
@@ -1239,6 +1246,7 @@ void FControlRigParameterTrackEditor::BakeInvertedPose(UControlRig* InControlRig
 
 	UMovieSceneControlRigParameterSection* Section = Cast<UMovieSceneControlRigParameterSection>(Track->GetSectionToKey());
 	USkeletalMeshComponent* SkelMeshComp = Cast<USkeletalMeshComponent>(InControlRig->GetObjectBinding()->GetBoundObject());
+	UMovieSceneSequence* RootMovieSceneSequence = GetSequencer()->GetRootMovieSceneSequence();
 	UMovieSceneSequence* MovieSceneSequence = GetSequencer()->GetFocusedMovieSceneSequence();
 	UMovieScene* MovieScene = MovieSceneSequence->GetMovieScene();
 	UAnimSeqExportOption* ExportOptions = NewObject<UAnimSeqExportOption>(GetTransientPackage(), NAME_None);
@@ -1256,7 +1264,7 @@ void FControlRigParameterTrackEditor::BakeInvertedPose(UControlRig* InControlRig
 	
 	const FScopedTransaction Transaction(LOCTEXT("BakeInvertedPose_Transaction", "Bake Inverted Pose"));
 
-	UnFbx::FLevelSequenceAnimTrackAdapter AnimTrackAdapter(ParentSequencer.Get(), MovieScene, RootToLocalTransform);
+	UnFbx::FLevelSequenceAnimTrackAdapter AnimTrackAdapter(ParentSequencer.Get(), MovieSceneSequence, RootMovieSceneSequence, RootToLocalTransform);
 	int32 AnimationLength = AnimTrackAdapter.GetLength();
 	FScopedSlowTask Progress(AnimationLength, LOCTEXT("BakingToControlRig_SlowTask", "Baking To Control Rig..."));
 	Progress.MakeDialog(true);
@@ -1287,8 +1295,13 @@ void FControlRigParameterTrackEditor::BakeInvertedPose(UControlRig* InControlRig
 	});
 	FEndAnimationCB EndCallback = FEndAnimationCB::CreateLambda([]{});
 
-	MovieSceneToolHelpers::BakeToSkelMeshToCallbacks(MovieScene,ParentSequencer.Get(),
-		SkelMeshComp, Template, RootToLocalTransform, ExportOptions,
+	FAnimExportSequenceParameters AESP;
+	AESP.Player = ParentSequencer.Get();
+	AESP.RootToLocalTransform = RootToLocalTransform;
+	AESP.MovieSceneSequence = MovieSceneSequence;
+	AESP.RootMovieSceneSequence = RootMovieSceneSequence;
+
+	MovieSceneToolHelpers::BakeToSkelMeshToCallbacks(AESP,SkelMeshComp, ExportOptions,
 		InitCallback, StartCallback, TickCallback, EndCallback);
 
 	InControlRig->OnPreAdditiveValuesApplication_AnyThread().Remove(DelegateHandle);
diff --git a/Engine/Plugins/Animation/ControlRig/Source/ControlRigEditor/Private/Tools/ConstraintBaker.cpp b/Engine/Plugins/Animation/ControlRig/Source/ControlRigEditor/Private/Tools/ConstraintBaker.cpp
--- a/Engine/Plugins/Animation/ControlRig/Source/ControlRigEditor/Private/Tools/ConstraintBaker.cpp
+++ b/Engine/Plugins/Animation/ControlRig/Source/ControlRigEditor/Private/Tools/ConstraintBaker.cpp
@@ -58,8 +58,12 @@ void FConstraintBaker::GetMinimalFramesToBake(
 
 	// default init bounds to the scene bounds
 	const bool bHasSettings = InSettings.IsSet();
-	FFrameNumber FirstBakingFrame = bHasSettings ? InSettings->StartFrame : MovieScene->GetPlaybackRange().GetLowerBoundValue();
-	FFrameNumber LastBakingFrame = bHasSettings ? InSettings->EndFrame : MovieScene->GetPlaybackRange().GetUpperBoundValue();
+	TOptional<TRange<FFrameNumber>> OptionalRange = InSequencer->GetSubSequenceRange();
+	FFrameNumber MovieSceneStartFrame = OptionalRange.IsSet() ? OptionalRange.GetValue().GetLowerBoundValue() : MovieScene->GetPlaybackRange().GetLowerBoundValue();
+	FFrameNumber MovieSceneEndFrame = OptionalRange.IsSet() ? OptionalRange.GetValue().GetUpperBoundValue() : MovieScene->GetPlaybackRange().GetUpperBoundValue();
+
+	FFrameNumber FirstBakingFrame = bHasSettings ? InSettings->StartFrame : MovieSceneStartFrame;
+	FFrameNumber LastBakingFrame = bHasSettings ? InSettings->EndFrame : MovieSceneEndFrame;
 	
 	// set start to first active frame if any, if we have no range
 	if (!bHasSettings)
diff --git a/Engine/Plugins/Editor/SequencerAnimTools/Source/SequencerAnimTools/Private/Sequencer/AnimationBoneTrail.cpp b/Engine/Plugins/Editor/SequencerAnimTools/Source/SequencerAnimTools/Private/Sequencer/AnimationBoneTrail.cpp
--- a/Engine/Plugins/Editor/SequencerAnimTools/Source/SequencerAnimTools/Private/Sequencer/AnimationBoneTrail.cpp
+++ b/Engine/Plugins/Editor/SequencerAnimTools/Source/SequencerAnimTools/Private/Sequencer/AnimationBoneTrail.cpp
@@ -38,8 +38,13 @@ void FAnimTrajectoryCache::Evaluate(FTrajectoryCache* ParentTransformCache)
 	// TODO: for some reason SkeletalMeshComponent becomes invalid sometimes when evaluating, no clue why yet
 	FMovieSceneSequenceTransform MovieSceneSequenceTransform;
 	UAnimSeqExportOption* AnimSeqExportOption = NewObject<UAnimSeqExportOption>(GetTransientPackage(), NAME_None);
-	MovieSceneToolHelpers::ExportToAnimSequence(CachedAnimSequence, AnimSeqExportOption,Sequencer->GetFocusedMovieSceneSequence()->GetMovieScene(), Sequencer.Get(), SkeletalMeshComponent.Get(),
-		Sequencer->GetFocusedTemplateID(), MovieSceneSequenceTransform);
+	FAnimExportSequenceParameters AESP;
+	AESP.Player = Sequencer.Get();
+	AESP.RootToLocalTransform = MovieSceneSequenceTransform;
+	AESP.MovieSceneSequence = Sequencer->GetFocusedMovieSceneSequence();
+	AESP.RootMovieSceneSequence = Sequencer->GetRootMovieSceneSequence();
+
+	MovieSceneToolHelpers::ExportToAnimSequence(CachedAnimSequence, AnimSeqExportOption, AESP, SkeletalMeshComponent.Get());
 	AnimSeqExportOption->MarkAsGarbage();
 	Sequencer->ForceEvaluate();
 
diff --git a/Engine/Plugins/Importers/USDImporter/Source/USDStage/Private/USDLevelSequenceHelper.cpp b/Engine/Plugins/Importers/USDImporter/Source/USDStage/Private/USDLevelSequenceHelper.cpp
--- a/Engine/Plugins/Importers/USDImporter/Source/USDStage/Private/USDLevelSequenceHelper.cpp
+++ b/Engine/Plugins/Importers/USDImporter/Source/USDStage/Private/USDLevelSequenceHelper.cpp
@@ -329,16 +329,17 @@ namespace UsdLevelSequenceHelperImpl
 
 			ExportOptions->bTransactRecording = false;
 
-			FMovieSceneSequenceIDRef Template = MovieSceneSequenceID::Root;
 			FMovieSceneSequenceTransform RootToLocalTransform;
+			FAnimExportSequenceParameters AESP;
+			AESP.Player = Player;
+			AESP.RootToLocalTransform = RootToLocalTransform;
+			AESP.MovieSceneSequence = LevelSequence;
+			AESP.RootMovieSceneSequence = LevelSequence;
 			bResult = MovieSceneToolHelpers::ExportToAnimSequence(
 				AnimSequence,
 				ExportOptions,
-				MovieScene,
-				Player,
-				SkeletalMeshComp,
-				Template,
-				RootToLocalTransform
+				AESP,
+				SkeletalMeshComp
 			);
 			if (!bResult)
 			{
@@ -443,6 +444,7 @@ namespace UsdLevelSequenceHelperImpl
 			const bool bResetControls = true;
 			ParamSection->LoadAnimSequenceIntoThisSection(
 				AnimSequence,
+				FFrameNumber(0),
 				MovieScene,
 				SkeletalMeshComp,
 				bReduceKeys,
diff --git a/Engine/Plugins/MovieScene/SequencerScripting/Source/SequencerScriptingEditor/Private/SequencerTools.cpp b/Engine/Plugins/MovieScene/SequencerScripting/Source/SequencerScriptingEditor/Private/SequencerTools.cpp
--- a/Engine/Plugins/MovieScene/SequencerScripting/Source/SequencerScriptingEditor/Private/SequencerTools.cpp
+++ b/Engine/Plugins/MovieScene/SequencerScripting/Source/SequencerScriptingEditor/Private/SequencerTools.cpp
@@ -269,7 +269,12 @@ bool ExportFBXInternal(const FSequencerExportFBXParams& InParams, UMovieSceneSeq
 			Player->SetPlaybackPosition(FMovieSceneSequencePlaybackParams(StartTime, EUpdatePositionMethod::Play));
 		}
 
-		bDidExport = MovieSceneToolHelpers::ExportFBX(World, MovieScene, Player, Bindings, Tracks, NodeNameAdapter, Template, FBXFileName, RootToLocalTransform);
+		FAnimExportSequenceParameters AESP;
+		AESP.Player = Player;
+		AESP.RootToLocalTransform = RootToLocalTransform;
+		AESP.MovieSceneSequence = Sequence;
+		AESP.RootMovieSceneSequence = RootSequence;
+		bDidExport = MovieSceneToolHelpers::ExportFBX(World, AESP, Bindings, Tracks, NodeNameAdapter, Template, FBXFileName);
 	}
 
 	Player->Stop();
@@ -365,7 +370,12 @@ bool USequencerToolsFunctionLibrary::ExportAnimSequence(UWorld* World, ULevelSeq
 		if (SkeletalMeshComp && SkeletalMeshComp->GetSkeletalMeshAsset() && SkeletalMeshComp->GetSkeletalMeshAsset()->GetSkeleton())
 		{
 			AnimSequence->SetSkeleton(SkeletalMeshComp->GetSkeletalMeshAsset()->GetSkeleton());
-			bResult = MovieSceneToolHelpers::ExportToAnimSequence(AnimSequence,ExportOptions, MovieScene, Player, SkeletalMeshComp, Template, RootToLocalTransform);
+			FAnimExportSequenceParameters AESQ;
+			AESQ.Player = Player;
+			AESQ.RootToLocalTransform = RootToLocalTransform;
+			AESQ.MovieSceneSequence = Sequence;
+			AESQ.RootMovieSceneSequence = Sequence;
+			bResult = MovieSceneToolHelpers::ExportToAnimSequence(AnimSequence, ExportOptions, AESQ, SkeletalMeshComp);
 		}
 	}
 	
diff --git a/Engine/Source/Editor/MovieSceneTools/Private/LevelSequenceFBXInterop.cpp b/Engine/Source/Editor/MovieSceneTools/Private/LevelSequenceFBXInterop.cpp
--- a/Engine/Source/Editor/MovieSceneTools/Private/LevelSequenceFBXInterop.cpp
+++ b/Engine/Source/Editor/MovieSceneTools/Private/LevelSequenceFBXInterop.cpp
@@ -210,7 +210,9 @@ void FLevelSequenceFBXInterop::ExportFBXInternal(const FString& ExportFilename,
 	Exporter->FillExportOptions(false, true, ExportFilename, ExportCancel, ExportAll);
 	if (!ExportCancel)
 	{
-		UMovieScene* MovieScene = Sequencer->GetFocusedMovieSceneSequence()->GetMovieScene();
+		UMovieSceneSequence* MovieSceneSequence = Sequencer->GetFocusedMovieSceneSequence();
+		UMovieSceneSequence* RootMovieSceneSequence = Sequencer->GetRootMovieSceneSequence();
+		UMovieScene* MovieScene = MovieSceneSequence->GetMovieScene();
 		UWorld* World = Sequencer->GetPlaybackContext()->GetWorld();
 		FMovieSceneSequenceIDRef Template = Sequencer->GetFocusedTemplateID();
 		UnFbx::FFbxExporter::FLevelSequenceNodeNameAdapter NodeNameAdapter(MovieScene, Sequencer.Get(), Template);
@@ -224,7 +226,12 @@ void FLevelSequenceFBXInterop::ExportFBXInternal(const FString& ExportFilename,
 			}
 
 			FMovieSceneSequenceTransform RootToLocalTransform = Sequencer->GetFocusedMovieSceneSequenceTransform();
-			if (MovieSceneToolHelpers::ExportFBX(World, MovieScene, Sequencer.Get(), Bindings, Tracks, NodeNameAdapter, Template, ExportFilename, RootToLocalTransform))
+			FAnimExportSequenceParameters AESP;
+			AESP.Player = Sequencer.Get();
+			AESP.RootToLocalTransform = RootToLocalTransform;
+			AESP.MovieSceneSequence = MovieSceneSequence;
+			AESP.RootMovieSceneSequence = RootMovieSceneSequence;
+			if (MovieSceneToolHelpers::ExportFBX(World, AESP, Bindings, Tracks, NodeNameAdapter, Template, ExportFilename))
 			{
 				FNotificationInfo Info(NSLOCTEXT("Sequencer", "ExportFBXSucceeded", "FBX Export Succeeded."));
 				Info.Hyperlink = FSimpleDelegate::CreateStatic([](FString InFilename) { FPlatformProcess::ExploreFolder(*InFilename); }, ExportFilename);
diff --git a/Engine/Source/Editor/MovieSceneTools/Private/MovieSceneToolHelpers.cpp b/Engine/Source/Editor/MovieSceneTools/Private/MovieSceneToolHelpers.cpp
--- a/Engine/Source/Editor/MovieSceneTools/Private/MovieSceneToolHelpers.cpp
+++ b/Engine/Source/Editor/MovieSceneTools/Private/MovieSceneToolHelpers.cpp
@@ -3726,7 +3726,25 @@ void ExportLevelMesh(UnFbx::FFbxExporter* Exporter, ULevel* Level, IMovieScenePl
 	Exporter->ExportLevelMesh(Level, !bSelectedOnly, ActorToExport, NodeNameAdapter, bSaveAnimSeq);
 }
 
+//5.5 deprecrated
 bool MovieSceneToolHelpers::ExportFBX(UWorld* World, UMovieScene* MovieScene, IMovieScenePlayer* Player, const TArray<FGuid>& Bindings, const TArray<UMovieSceneTrack*>& Tracks, INodeNameAdapter& NodeNameAdapter, FMovieSceneSequenceIDRef& Template, const FString& InFBXFileName, FMovieSceneSequenceTransform& RootToLocalTransform)
+{
+	if (MovieScene)
+	{
+		if (UMovieSceneSequence* MovieSceneSequence = MovieScene->GetTypedOuter< UMovieSceneSequence>())
+		{
+			FAnimExportSequenceParameters AESP;
+			AESP.MovieSceneSequence = MovieSceneSequence;
+			AESP.RootMovieSceneSequence = MovieSceneSequence;
+			AESP.Player = Player;
+			AESP.RootToLocalTransform = RootToLocalTransform;
+			return MovieSceneToolHelpers::ExportFBX(World, AESP, Bindings, Tracks, NodeNameAdapter, Template, InFBXFileName);
+		}
+	}
+	return false;
+}
+
+bool MovieSceneToolHelpers::ExportFBX(UWorld* World, const FAnimExportSequenceParameters& AESP, const TArray<FGuid>& Bindings, const TArray<UMovieSceneTrack*>& Tracks, INodeNameAdapter& NodeNameAdapter, FMovieSceneSequenceIDRef& Template, const FString& InFBXFileName)
 {
 	UnFbx::FFbxExporter* Exporter = UnFbx::FFbxExporter::GetInstance();
 
@@ -3734,7 +3752,7 @@ bool MovieSceneToolHelpers::ExportFBX(UWorld* World, UMovieScene* MovieScene, IM
 	Exporter->SetTransformBaking(false);
 	Exporter->SetKeepHierarchy(true);
 
-	ExportLevelMesh(Exporter, World->PersistentLevel, Player, Bindings, NodeNameAdapter, Template);
+	ExportLevelMesh(Exporter, World->PersistentLevel, AESP.Player, Bindings, NodeNameAdapter, Template);
 
 	// Export streaming levels and actors
 	for (ULevelStreaming* StreamingLevel : World->GetStreamingLevels())
@@ -3743,16 +3761,16 @@ bool MovieSceneToolHelpers::ExportFBX(UWorld* World, UMovieScene* MovieScene, IM
 		{
 			if (ULevel* Level = StreamingLevel->GetLoadedLevel())
 			{
-				ExportLevelMesh(Exporter, Level, Player, Bindings, NodeNameAdapter, Template);
+				ExportLevelMesh(Exporter, Level, AESP.Player, Bindings, NodeNameAdapter, Template);
 			}
 		}
 	}
 
-	Exporter->ExportLevelSequence(MovieScene, Bindings, Player, NodeNameAdapter, Template, RootToLocalTransform);
+	Exporter->ExportLevelSequence(AESP.MovieSceneSequence, AESP.RootMovieSceneSequence, Bindings, AESP.Player, NodeNameAdapter, Template, AESP.RootToLocalTransform);
 
 	//Export given tracks
-	Exporter->ExportLevelSequenceTracks(MovieScene, Player, Template, nullptr, nullptr, Tracks, RootToLocalTransform);
-	
+	Exporter->ExportLevelSequenceTracks(AESP.MovieSceneSequence, AESP.RootMovieSceneSequence, AESP.Player, Template, nullptr, nullptr, Tracks, AESP.RootToLocalTransform);
+
 	// Save to disk
 	Exporter->WriteToFile(*InFBXFileName);
 
@@ -3838,6 +3856,14 @@ bool MovieSceneToolHelpers::BakeToSkelMeshToCallbacks(UMovieScene* MovieScene, I
 	USkeletalMeshComponent* InSkelMeshComp, FMovieSceneSequenceIDRef& Template, FMovieSceneSequenceTransform& RootToLocalTransform, UAnimSeqExportOption* ExportOptions,
 	FInitAnimationCB InitCallback, FStartAnimationCB StartCallback, FTickAnimationCB TickCallback, FEndAnimationCB EndCallback)
 {
+	return false;
+}
+
+
+bool MovieSceneToolHelpers::BakeToSkelMeshToCallbacks(const FAnimExportSequenceParameters& AESP, USkeletalMeshComponent* InSkelMeshComp, UAnimSeqExportOption* ExportOptions,
+	FInitAnimationCB InitCallback, FStartAnimationCB StartCallback, FTickAnimationCB TickCallback, FEndAnimationCB EndCallback)
+{
+	UMovieScene* MovieScene = AESP.MovieSceneSequence->GetMovieScene();
 	TArray< USkeletalMeshComponent*> SkelMeshComps;
 	if (ExportOptions->bEvaluateAllSkeletalMeshComponents)
 	{
@@ -3862,7 +3888,7 @@ bool MovieSceneToolHelpers::BakeToSkelMeshToCallbacks(UMovieScene* MovieScene, I
 		}
 	}
 
-	UnFbx::FLevelSequenceAnimTrackAdapter AnimTrackAdapter(Player, MovieScene, RootToLocalTransform);
+	UnFbx::FLevelSequenceAnimTrackAdapter AnimTrackAdapter(AESP.Player, AESP.MovieSceneSequence, AESP.RootMovieSceneSequence, AESP.RootToLocalTransform);
 	int32 LocalStartFrame = AnimTrackAdapter.GetLocalStartFrame();
 	int32 StartFrame = AnimTrackAdapter.GetStartFrame();
 	int32 AnimationLength = AnimTrackAdapter.GetLength();
@@ -3961,73 +3987,90 @@ bool MovieSceneToolHelpers::BakeToSkelMeshToCallbacks(UMovieScene* MovieScene, I
 	return true;
 }
 
+//5.5 deprecrated
 bool MovieSceneToolHelpers::ExportToAnimSequence(UAnimSequence* AnimSequence, UAnimSeqExportOption* ExportOptions, UMovieScene* MovieScene, IMovieScenePlayer* Player,
-	USkeletalMeshComponent* SkelMeshComp, FMovieSceneSequenceIDRef& Template, FMovieSceneSequenceTransform& RootToLocalTransform)
+	USkeletalMeshComponent* SkelMesh, FMovieSceneSequenceIDRef& Template, FMovieSceneSequenceTransform& RootToLocalTransform)
+{
+	if (MovieScene)
+	{
+		FAnimExportSequenceParameters AESP;
+		AESP.Player = Player;
+		AESP.RootToLocalTransform = RootToLocalTransform;
+		UMovieSceneSequence* OwnerSceneSequence = MovieScene->GetTypedOuter<UMovieSceneSequence>();
+		AESP.MovieSceneSequence = OwnerSceneSequence;
+		AESP.RootMovieSceneSequence = OwnerSceneSequence;
+		return MovieSceneToolHelpers::ExportToAnimSequence(AnimSequence, ExportOptions, AESP, SkelMesh);
+	}
+	else
+	{
+		return false;
+	}
+}
+bool MovieSceneToolHelpers::ExportToAnimSequence(UAnimSequence* AnimSequence, UAnimSeqExportOption* ExportOptions, const FAnimExportSequenceParameters& AESP,
+	USkeletalMeshComponent* SkelMeshComp)
 {
-	if (AnimSequence == nullptr || ExportOptions == nullptr || MovieScene == nullptr || SkelMeshComp == nullptr)
+	if (AnimSequence == nullptr || ExportOptions == nullptr || AESP.MovieSceneSequence == nullptr || AESP.RootMovieSceneSequence == nullptr || SkelMeshComp == nullptr)
 	{
 		UE_LOG(LogMovieScene, Error, TEXT("MovieSceneToolHelpers::ExportToAnimSequence All parameters must be valid."));
 		return false;
 	}
 	FAnimRecorderInstance AnimationRecorder;
-	FFrameRate SampleRate = MovieScene->GetDisplayRate();
-	FInitAnimationCB InitCallback = FInitAnimationCB::CreateLambda([&AnimationRecorder,SampleRate,ExportOptions,SkelMeshComp,AnimSequence]
-	{
-		FAnimationRecordingSettings RecordingSettings;
-		RecordingSettings.SampleFrameRate = SampleRate;
-		RecordingSettings.Interpolation = ExportOptions->Interpolation;
-		RecordingSettings.InterpMode = ExportOptions->CurveInterpolation;
-		if (ExportOptions->CurveInterpolation == ERichCurveInterpMode::RCIM_Constant ||
-			ExportOptions->CurveInterpolation == ERichCurveInterpMode::RCIM_None)
-		{
-			RecordingSettings.TangentMode = ERichCurveTangentMode::RCTM_None;
-		}
-		else
-		{
-			RecordingSettings.TangentMode = ERichCurveTangentMode::RCTM_Auto;
-		}
-		RecordingSettings.Length = 0;
-		RecordingSettings.bRemoveRootAnimation = false;
-		RecordingSettings.bCheckDeltaTimeAtBeginning = false;
-		RecordingSettings.bRecordTransforms = ExportOptions->bExportTransforms;
-		RecordingSettings.bRecordMorphTargets = ExportOptions->bExportMorphTargets;
-		RecordingSettings.bRecordAttributeCurves = ExportOptions->bExportAttributeCurves;
-		RecordingSettings.bRecordMaterialCurves = ExportOptions->bExportMaterialCurves;
-		RecordingSettings.bRecordInWorldSpace = ExportOptions->bRecordInWorldSpace;
-		RecordingSettings.IncludeAnimationNames = ExportOptions->IncludeAnimationNames;
-		RecordingSettings.ExcludeAnimationNames = ExportOptions->ExcludeAnimationNames;
-		RecordingSettings.bTransactRecording = ExportOptions->bTransactRecording;
-		AnimationRecorder.Init(SkelMeshComp, AnimSequence, nullptr, RecordingSettings);	
+	FFrameRate SampleRate = AESP.MovieSceneSequence->GetMovieScene()->GetDisplayRate();
+	FInitAnimationCB InitCallback = FInitAnimationCB::CreateLambda([&AnimationRecorder, SampleRate, ExportOptions, SkelMeshComp, AnimSequence]
+		{
+			FAnimationRecordingSettings RecordingSettings;
+			RecordingSettings.SampleFrameRate = SampleRate;
+			RecordingSettings.Interpolation = ExportOptions->Interpolation;
+			RecordingSettings.InterpMode = ExportOptions->CurveInterpolation;
+			if (ExportOptions->CurveInterpolation == ERichCurveInterpMode::RCIM_Constant ||
+				ExportOptions->CurveInterpolation == ERichCurveInterpMode::RCIM_None)
+			{
+				RecordingSettings.TangentMode = ERichCurveTangentMode::RCTM_None;
+			}
+			else
+			{
+				RecordingSettings.TangentMode = ERichCurveTangentMode::RCTM_Auto;
+			}
+			RecordingSettings.Length = 0;
+			RecordingSettings.bRemoveRootAnimation = false;
+			RecordingSettings.bCheckDeltaTimeAtBeginning = false;
+			RecordingSettings.bRecordTransforms = ExportOptions->bExportTransforms;
+			RecordingSettings.bRecordMorphTargets = ExportOptions->bExportMorphTargets;
+			RecordingSettings.bRecordAttributeCurves = ExportOptions->bExportAttributeCurves;
+			RecordingSettings.bRecordMaterialCurves = ExportOptions->bExportMaterialCurves;
+			RecordingSettings.bRecordInWorldSpace = ExportOptions->bRecordInWorldSpace;
+			RecordingSettings.IncludeAnimationNames = ExportOptions->IncludeAnimationNames;
+			RecordingSettings.ExcludeAnimationNames = ExportOptions->ExcludeAnimationNames;
+			RecordingSettings.bTransactRecording = ExportOptions->bTransactRecording;
+			AnimationRecorder.Init(SkelMeshComp, AnimSequence, nullptr, RecordingSettings);
 		});
 
-	
+
 	FStartAnimationCB StartCallback = FStartAnimationCB::CreateLambda([SkelMeshComp, &AnimationRecorder]
-	{
-		AnimationRecorder.BeginRecording();
-		SkelMeshComp->UpdateLODStatus();
-	});
+		{
+			AnimationRecorder.BeginRecording();
+			SkelMeshComp->UpdateLODStatus();
+		});
 
 	FTickAnimationCB TickCallback = FTickAnimationCB::CreateLambda([&AnimationRecorder](float DeltaTime, FFrameNumber FrameNumber)
-	{
-		AnimationRecorder.Update(DeltaTime);
+		{
+			AnimationRecorder.Update(DeltaTime);
 
-	});
+		});
 
 	FEndAnimationCB EndCallback = FEndAnimationCB::CreateLambda([&AnimationRecorder]
-	{
+		{
 			const bool bShowAnimationAssetCreatedToast = false;
 			AnimationRecorder.FinishRecording(bShowAnimationAssetCreatedToast);
-	});
-	
+		});
+
 
-	MovieSceneToolHelpers::BakeToSkelMeshToCallbacks(MovieScene,Player,
-		SkelMeshComp, Template, RootToLocalTransform, ExportOptions,
+	MovieSceneToolHelpers::BakeToSkelMeshToCallbacks(AESP, SkelMeshComp, ExportOptions,
 		InitCallback, StartCallback, TickCallback, EndCallback);
 	return AnimSequence->GetDataModel()->HasBeenPopulated();
 }
 
 
-
 FSpawnableRestoreState::FSpawnableRestoreState(UMovieScene* MovieScene)
 	: FSpawnableRestoreState(MovieScene, nullptr)
 {
diff --git a/Engine/Source/Editor/MovieSceneTools/Private/TrackEditors/SkeletalAnimationTrackEditor.cpp b/Engine/Source/Editor/MovieSceneTools/Private/TrackEditors/SkeletalAnimationTrackEditor.cpp
--- a/Engine/Source/Editor/MovieSceneTools/Private/TrackEditors/SkeletalAnimationTrackEditor.cpp
+++ b/Engine/Source/Editor/MovieSceneTools/Private/TrackEditors/SkeletalAnimationTrackEditor.cpp
@@ -1281,7 +1281,8 @@ void FSkeletalAnimationTrackEditor::OnSequencerSaved(ISequencer& )
 			ULevelSequenceAnimSequenceLink* LevelAnimLink = AssetUserDataInterface->GetAssetUserData< ULevelSequenceAnimSequenceLink >();
 			if (LevelAnimLink)
 			{
-				UMovieScene* MovieScene = SequencerPtr->GetFocusedMovieSceneSequence()->GetMovieScene();
+				UMovieSceneSequence* MovieSceneSequence = SequencerPtr->GetFocusedMovieSceneSequence();
+				UMovieSceneSequence* RootMovieSceneSequence = SequencerPtr->GetRootMovieSceneSequence();
 				FMovieSceneSequenceIDRef Template = SequencerPtr->GetFocusedTemplateID();
 				FMovieSceneSequenceTransform RootToLocalTransform = SequencerPtr->GetFocusedMovieSceneSequenceTransform();
 				for (int32 Index = LevelAnimLink->AnimSequenceLinks.Num() -1; Index >=0 ; --Index)
@@ -1325,7 +1326,12 @@ void FSkeletalAnimationTrackEditor::OnSequencerSaved(ISequencer& )
 						AnimSeqExportOption->Interpolation = Item.Interpolation;
 						AnimSeqExportOption->CurveInterpolation = Item.CurveInterpolation;
 
-						bool bResult = MovieSceneToolHelpers::ExportToAnimSequence(AnimSequence, AnimSeqExportOption, MovieScene, SequencerPtr.Get(), SkelMeshComp, Template, RootToLocalTransform);
+						FAnimExportSequenceParameters AESP;
+						AESP.Player = SequencerPtr.Get();
+						AESP.RootToLocalTransform = RootToLocalTransform;
+						AESP.MovieSceneSequence = MovieSceneSequence;
+						AESP.RootMovieSceneSequence = RootMovieSceneSequence;
+						bool bResult = MovieSceneToolHelpers::ExportToAnimSequence(AnimSequence, AnimSeqExportOption, AESP, SkelMeshComp);
 
 						AnimSeqExportOption->bExportMorphTargets = bSavedExportMorphTargets;
 						AnimSeqExportOption->bExportAttributeCurves = bSavedExportAttributeCurves;
@@ -1474,11 +1480,16 @@ bool FSkeletalAnimationTrackEditor::CreateAnimationSequence(const TArray<UObject
 			if (OptionWindow->ShouldExport())
 			{
 				const TSharedPtr<ISequencer> ParentSequencer = GetSequencer();
-				UMovieScene* MovieScene = ParentSequencer->GetFocusedMovieSceneSequence()->GetMovieScene();
-				FMovieSceneSequenceIDRef Template = ParentSequencer->GetFocusedTemplateID();
+				UMovieSceneSequence* MovieSceneSequence = ParentSequencer->GetFocusedMovieSceneSequence();
+				UMovieSceneSequence* RootMovieSceneSequence = ParentSequencer->GetRootMovieSceneSequence();
 				FMovieSceneSequenceTransform RootToLocalTransform = ParentSequencer->GetFocusedMovieSceneSequenceTransform();
+				FAnimExportSequenceParameters AESP;
+				AESP.Player = ParentSequencer.Get();
+				AESP.RootToLocalTransform = RootToLocalTransform;
+				AESP.MovieSceneSequence = MovieSceneSequence;
+				AESP.RootMovieSceneSequence = RootMovieSceneSequence;
 
-				bResult  = MovieSceneToolHelpers::ExportToAnimSequence(AnimSequence, AnimSeqExportOption,MovieScene, ParentSequencer.Get(), SkelMeshComp, Template, RootToLocalTransform);
+				bResult  = MovieSceneToolHelpers::ExportToAnimSequence(AnimSequence, AnimSeqExportOption, AESP, SkelMeshComp);
 			}
 		}
 
diff --git a/Engine/Source/Editor/MovieSceneTools/Public/MovieSceneToolHelpers.h b/Engine/Source/Editor/MovieSceneTools/Public/MovieSceneToolHelpers.h
--- a/Engine/Source/Editor/MovieSceneTools/Public/MovieSceneToolHelpers.h
+++ b/Engine/Source/Editor/MovieSceneTools/Public/MovieSceneToolHelpers.h
@@ -65,6 +65,15 @@ struct FFBXInOutParameters
 	float ImportUniformScaleBackup;
 };
 
+struct FAnimExportSequenceParameters
+{
+	FAnimExportSequenceParameters() = default;
+	UMovieSceneSequence* MovieSceneSequence;
+	UMovieSceneSequence* RootMovieSceneSequence;
+	IMovieScenePlayer* Player;
+	FMovieSceneSequenceTransform RootToLocalTransform;
+};
+
 //callback's used by skel mesh recorders
 
 DECLARE_DELEGATE(FInitAnimationCB);
@@ -322,16 +331,17 @@ public:
 	* Export FBX
 	*
 	* @param World The world to export from
-	* @param InMovieScene The movie scene to export frome
-	* @param MoviePlayer to use
+	* @param AnimExportSequenceParameters The sequence parameters used for evaluation
 	* @param Bindings The sequencer binding map
 	* @param Tracks The tracks to export
 	* @param NodeNameAdaptor Adaptor to look up actor names.
-	* @param InFBXFileName the fbx file name.
 	* @param Template Movie scene sequence id.
-	* @param RootToLocalTransform The root to local transform time.
+	* @param InFBXFileName the fbx file name.
 	* @return Whether the export was successful
 	*/
+	static bool ExportFBX(UWorld* World, const FAnimExportSequenceParameters& AnimExportSequenceParameters, const TArray<FGuid>& Bindings, const TArray<UMovieSceneTrack*>& Tracks, INodeNameAdapter& NodeNameAdapter, FMovieSceneSequenceIDRef& Template, const FString& InFBXFileName);
+
+	UE_DEPRECATED(5.5, "ExportFBX taking movie scene has been deprecated in favor of a new function that takes current and root movie scene sequences")
 	static bool ExportFBX(UWorld* World, UMovieScene* MovieScene, IMovieScenePlayer* Player, const TArray<FGuid>& Bindings, const TArray<UMovieSceneTrack*>& Tracks, INodeNameAdapter& NodeNameAdapter, FMovieSceneSequenceIDRef& Template,  const FString& InFBXFileName, FMovieSceneSequenceTransform& RootToLocalTransform);
 
 	/**
@@ -429,21 +439,23 @@ public:
 	 *
 	 * @param AnimSequence The sequence to save to.
 	 * @param ExportOptions The options to use when saving.
-	 * @param MovieScene The movie scene to export the object binding from
-	 * @param Player The Player to evaluate
-	 * @param SkelMesh The Player to evaluate
-	 * @param Template ID of the sequence template.
-	 * @param RootToLocalTransform Transform Offset to apply to exported anim sequence.
+	 * @param AnimExportSequenceParameters The Sequence parameters like the movie scene sequences, the player, template id, used for evaluation
+	 * @param SkelMesh The Player to evaluatee.
 	 * @return Whether or not it succeeds
 
 	*/
+	static bool ExportToAnimSequence(UAnimSequence* AnimSequence, UAnimSeqExportOption* ExportOptions, const FAnimExportSequenceParameters& AnimExportSequenceParameters,
+			USkeletalMeshComponent* SkelMesh);
+
+	UE_DEPRECATED(5.5, "ExportToAnimSequence taking movie scene has been deprecated in favor of a new function that takes current and root movie scene sequences")
 	static bool ExportToAnimSequence(UAnimSequence* AnimSequence, UAnimSeqExportOption* ExportOptions, UMovieScene* MovieScene, IMovieScenePlayer* Player,
 		USkeletalMeshComponent* SkelMesh, FMovieSceneSequenceIDRef& Template, FMovieSceneSequenceTransform& RootToLocalTransform);
 
 	/*
 	 * Bake the SkelMesh to a generic object wich implements a set of callbacks
 	 *
-	 * @param MovieScene The movie scene to export the object binding from
+	 * @param MovieSceneSequence The movie scene sequence that's current
+	 * @param RootMovieSceneSequence The root movie scene sequence
 	 * @param Player The Player to evaluate
 	 * @param SkelMesh The Player to evaluate
 	 * @param Template ID of the sequence template.
@@ -457,6 +469,10 @@ public:
 	
 	*/
 
+	static bool BakeToSkelMeshToCallbacks(const FAnimExportSequenceParameters& AnimExportSequenceParameters, USkeletalMeshComponent* SkelMesh, UAnimSeqExportOption* ExportOptions,
+		FInitAnimationCB InitCallback, FStartAnimationCB StartCallback, FTickAnimationCB TickCallback, FEndAnimationCB EndCallback);
+
+	UE_DEPRECATED(5.5, "BakeToSkelMeshToCallbacks taking movie scene has been deprecated in favor of a new function that takes current and root movie scene sequences")
 	static bool BakeToSkelMeshToCallbacks(UMovieScene* MovieScene, IMovieScenePlayer* Player,
 		USkeletalMeshComponent* SkelMesh, FMovieSceneSequenceIDRef& Template, FMovieSceneSequenceTransform& RootToLocalTransform, UAnimSeqExportOption* ExportOptions,
 		FInitAnimationCB InitCallback, FStartAnimationCB StartCallback, FTickAnimationCB TickCallback, FEndAnimationCB EndCallback);
diff --git a/Engine/Source/Editor/Sequencer/Private/Sequencer.h b/Engine/Source/Editor/Sequencer/Private/Sequencer.h
--- a/Engine/Source/Editor/Sequencer/Private/Sequencer.h
+++ b/Engine/Source/Editor/Sequencer/Private/Sequencer.h
@@ -186,13 +186,6 @@ public:
 	/** Select all keys and sections backward from the current time */
 	void SelectBackward();
 
-	/**
-	 * Get the currently viewed sub sequence range
-	 *
-	 * @return The sub sequence range, or an empty optional if we're viewing the root.
-	 */
-	TOptional<TRange<FFrameNumber>> GetSubSequenceRange() const;
-
 	/**
 	 * Compute a major grid interval and number of minor divisions to display
 	 */
@@ -709,6 +702,7 @@ public:
 	virtual FMovieSceneRootEvaluationTemplateInstance& GetEvaluationTemplate() override { return RootTemplateInstance; }
 	virtual void ResetToNewRootSequence(UMovieSceneSequence& NewSequence) override;
 	virtual void FocusSequenceInstance(UMovieSceneSubSection& InSubSection) override;
+	virtual TOptional<TRange<FFrameNumber>> GetSubSequenceRange() const override;
 	virtual TSharedPtr<UE::Sequencer::FSequencerEditorViewModel> GetViewModel() const override;
 	virtual void SuppressAutoEvaluation(UMovieSceneSequence* Sequence, const FGuid& InSequenceSignature) override;
 	virtual EAutoChangeMode GetAutoChangeMode() const override;
diff --git a/Engine/Source/Editor/Sequencer/Public/ISequencer.h b/Engine/Source/Editor/Sequencer/Public/ISequencer.h
--- a/Engine/Source/Editor/Sequencer/Public/ISequencer.h
+++ b/Engine/Source/Editor/Sequencer/Public/ISequencer.h
@@ -284,6 +284,17 @@ public:
 	 */
 	virtual void PopToSequenceInstance(FMovieSceneSequenceIDRef SequenceID) = 0;
 
+	/**
+	 * Get the currently viewed sub sequence range
+	 *
+	 * @return The sub sequence range, or an empty optional if we're viewing the root.
+	 */
+	virtual TOptional<TRange<FFrameNumber>> GetSubSequenceRange() const
+	{
+		TOptional<TRange<FFrameNumber>> Optional;
+		return Optional;
+	}
+
 	/**
 	 * Retrieve the top level view model for this sequence
 	 */
diff --git a/Engine/Source/Editor/UnrealEd/Private/Fbx/FbxMainExport.cpp b/Engine/Source/Editor/UnrealEd/Private/Fbx/FbxMainExport.cpp
--- a/Engine/Source/Editor/UnrealEd/Private/Fbx/FbxMainExport.cpp
+++ b/Engine/Source/Editor/UnrealEd/Private/Fbx/FbxMainExport.cpp
@@ -92,6 +92,9 @@
 #include "Sections/MovieSceneFloatSection.h"
 #include "Evaluation/MovieScenePlayback.h"
 #include "Evaluation/MovieSceneEvaluationTemplateInstance.h"
+#include "MovieSceneSequenceID.h"
+#include "Evaluation/MovieSceneSequenceHierarchy.h"
+#include "Compilation/MovieSceneCompiledDataManager.h"
 #include "MovieSceneSequence.h"
 #include "MovieSceneTimeHelpers.h"
 #include "DynamicMeshBuilder.h"
@@ -1697,36 +1700,79 @@ FbxNode* FFbxExporter::FLevelSequenceNodeNameAdapter::GetFbxNode(UObject* InObje
 	return nullptr;
 }
 
-FLevelSequenceAnimTrackAdapter::FLevelSequenceAnimTrackAdapter( IMovieScenePlayer* InMovieScenePlayer, UMovieScene* InMovieScene, const FMovieSceneSequenceTransform& InRootToLocalTransform, UMovieSceneSkeletalAnimationTrack* InAnimTrack)
+FLevelSequenceAnimTrackAdapter::FLevelSequenceAnimTrackAdapter(IMovieScenePlayer* InMovieScenePlayer, UMovieSceneSequence* InMovieSceneSequence, UMovieSceneSequence* InRootMovieSceneSequence, const FMovieSceneSequenceTransform& InRootToLocalTransform, UMovieSceneSkeletalAnimationTrack* InAnimTrack)
 {
 	MovieScenePlayer = InMovieScenePlayer;
-	MovieScene = InMovieScene;
+	MovieSceneSequence = InMovieSceneSequence;
+	RootMovieSceneSequence = InRootMovieSceneSequence;
+	MovieScene = MovieSceneSequence->GetMovieScene();
+	RootToLocalTransform = InRootToLocalTransform;
+	AnimTrack = InAnimTrack;
+}
+
+//5.5 deprecrated
+FLevelSequenceAnimTrackAdapter::FLevelSequenceAnimTrackAdapter(IMovieScenePlayer* InMovieScenePlayer, UMovieScene* InMovieScene, const FMovieSceneSequenceTransform& InRootToLocalTransform, UMovieSceneSkeletalAnimationTrack* InAnimTrack)
+{
+	MovieScenePlayer = InMovieScenePlayer;
+	MovieScene = MovieSceneSequence->GetMovieScene();
 	RootToLocalTransform = InRootToLocalTransform;
 	AnimTrack = InAnimTrack;
+	if (MovieScene)
+	{
+		UMovieSceneSequence* OwnerSceneSequence = MovieScene->GetTypedOuter<UMovieSceneSequence>();
+		MovieSceneSequence = OwnerSceneSequence;
+		RootMovieSceneSequence = OwnerSceneSequence;
+	}
+	else
+	{
+		MovieSceneSequence = nullptr;
+		RootMovieSceneSequence = nullptr;
+	}
 }
 
+TRange<FFrameNumber> FLevelSequenceAnimTrackAdapter::GetSequenceRange() const
+{
+	MovieScenePlayer->State.AssignSequence(MovieSceneSequenceID::Root, *RootMovieSceneSequence, *MovieScenePlayer);
+	FMovieSceneSequenceIDRef Template = MovieScenePlayer->State.FindSequenceId(MovieSceneSequence);
+
+	FMovieSceneSequenceHierarchy Hierarchy = FMovieSceneSequenceHierarchy();
+	UMovieSceneCompiledDataManager::CompileHierarchy(RootMovieSceneSequence, &Hierarchy,EMovieSceneServerClientMask::All);
+
+	const FMovieSceneSubSequenceData* SubSequenceData = Hierarchy.FindSubData(Template);
+	TRange<FFrameNumber> Range = MovieScene->GetPlaybackRange();
+	FFrameTime StartTime = FFrameRate::TransformTime(UE::MovieScene::DiscreteInclusiveLower(MovieScene->GetPlaybackRange()).Value, MovieScene->GetTickResolution(), MovieScene->GetDisplayRate());
+
+	if (SubSequenceData)
+	{
+		Range = SubSequenceData->PlayRange.Value;
+	}
+	return Range;
+}
 int32 FLevelSequenceAnimTrackAdapter::GetLocalStartFrame() const
 {
+	TRange<FFrameNumber> Range = GetSequenceRange();
+
 	FFrameRate TickResolution = MovieScene->GetTickResolution();
 	FFrameRate DisplayRate = MovieScene->GetDisplayRate();
-	return FFrameRate::TransformTime(FFrameTime(UE::MovieScene::DiscreteInclusiveLower(MovieScene->GetPlaybackRange())), TickResolution, DisplayRate).RoundToFrame().Value;
+	return FFrameRate::TransformTime(FFrameTime(UE::MovieScene::DiscreteInclusiveLower(Range)), TickResolution, DisplayRate).RoundToFrame().Value;
 }
 
 int32 FLevelSequenceAnimTrackAdapter::GetStartFrame() const
 {
+	TRange<FFrameNumber> Range = GetSequenceRange();
 	FFrameRate TickResolution = MovieScene->GetTickResolution();
 	FFrameRate DisplayRate = MovieScene->GetDisplayRate();
-	return FFrameRate::TransformTime(FFrameTime(UE::MovieScene::DiscreteInclusiveLower(MovieScene->GetPlaybackRange()) * RootToLocalTransform.InverseNoLooping()), TickResolution, DisplayRate).RoundToFrame().Value;
+	return FFrameRate::TransformTime(FFrameTime(UE::MovieScene::DiscreteInclusiveLower(Range) * RootToLocalTransform.InverseNoLooping()), TickResolution, DisplayRate).RoundToFrame().Value;
 }
 
 int32 FLevelSequenceAnimTrackAdapter::GetLength() const
 {
+	TRange<FFrameNumber> Range = GetSequenceRange();
 	FFrameRate TickResolution = MovieScene->GetTickResolution();
 	FFrameRate DisplayRate = MovieScene->GetDisplayRate();
-	return FFrameRate::TransformTime(FFrameTime(UE::MovieScene::DiscreteSize(MovieScene->GetPlaybackRange())), TickResolution, DisplayRate).RoundToFrame().Value;
+	return FFrameRate::TransformTime(FFrameTime(UE::MovieScene::DiscreteSize(Range)), TickResolution, DisplayRate).RoundToFrame().Value;
 }
 
-
 void FLevelSequenceAnimTrackAdapter::UpdateAnimation( int32 LocalFrame )
 {	
 	FFrameRate TickResolution = MovieScene->GetTickResolution();
@@ -1793,8 +1839,19 @@ float FLevelSequenceAnimTrackAdapter::GetAnimTime(int32 LocalFrame) const
 
 	return 0.f;
 }
-
 bool FFbxExporter::ExportLevelSequenceTracks(UMovieScene* MovieScene, IMovieScenePlayer* MovieScenePlayer, FMovieSceneSequenceIDRef InSequenceID, FbxNode* FbxActor, UObject* BoundObject, const TArray<UMovieSceneTrack*>& Tracks, const FMovieSceneSequenceTransform& RootToLocalTransform)
+{
+	if (MovieScene)
+	{
+		if (UMovieSceneSequence* OwnerSceneSequence = MovieScene->GetTypedOuter<UMovieSceneSequence>())
+		{
+			return ExportLevelSequenceTracks(OwnerSceneSequence, OwnerSceneSequence, MovieScenePlayer, InSequenceID, FbxActor, BoundObject, Tracks,  RootToLocalTransform);
+		}
+	}
+	return false;
+}
+
+bool FFbxExporter::ExportLevelSequenceTracks(UMovieSceneSequence* MovieSceneSequence, UMovieSceneSequence* RootMovieSceneSequence, IMovieScenePlayer* MovieScenePlayer, FMovieSceneSequenceIDRef InSequenceID, FbxNode* FbxActor, UObject* BoundObject, const TArray<UMovieSceneTrack*>& Tracks, const FMovieSceneSequenceTransform& RootToLocalTransform)
 {
 	AActor* Actor = Cast<AActor>(BoundObject);
 	if (!Actor)
@@ -1812,6 +1869,7 @@ bool FFbxExporter::ExportLevelSequenceTracks(UMovieScene* MovieScene, IMovieScen
 		SkeletalMeshComp = Actor ? Cast<USkeletalMeshComponent>(Actor->GetComponentByClass(USkeletalMeshComponent::StaticClass())) : nullptr;	
 	}
 
+	UMovieScene* MovieScene = MovieSceneSequence->GetMovieScene();
 	FFrameRate DisplayRate = MovieScene->GetDisplayRate();
 
 	bool bSkip3DTransformTrack = SkeletalMeshComp && GetExportOptions()->MapSkeletalMotionToRoot;
@@ -1911,7 +1969,7 @@ bool FFbxExporter::ExportLevelSequenceTracks(UMovieScene* MovieScene, IMovieScen
 	{
 		if (!bSkip3DTransformTrack)
 		{
-			FLevelSequenceAnimTrackAdapter AnimTrackAdapter(MovieScenePlayer, MovieScene, RootToLocalTransform, nullptr);
+			FLevelSequenceAnimTrackAdapter AnimTrackAdapter(MovieScenePlayer, MovieSceneSequence, RootMovieSceneSequence, RootToLocalTransform, nullptr);
 			ExportLevelSequenceBaked3DTransformTrack(AnimTrackAdapter, FbxActor, MovieScenePlayer, InSequenceID, TransformTracks, BoundObject, MovieScene->GetPlaybackRange(), RootToLocalTransform);
 		}
 
@@ -1960,21 +2018,32 @@ bool FFbxExporter::ExportLevelSequenceTracks(UMovieScene* MovieScene, IMovieScen
 		SkeletalMeshComp->GetOwner()->GetComponents(SkeletalMeshComponents);
 		for (USkeletalMeshComponent* SkeletalMeshComponent : SkeletalMeshComponents)
 		{
-			FLevelSequenceAnimTrackAdapter AnimTrackAdapter(MovieScenePlayer, MovieScene, RootToLocalTransform, SkeletalAnimationTrack);
+			FLevelSequenceAnimTrackAdapter AnimTrackAdapter(MovieScenePlayer, MovieSceneSequence,RootMovieSceneSequence, RootToLocalTransform, SkeletalAnimationTrack);
 			ExportAnimTrack(AnimTrackAdapter, Actor, SkeletalMeshComponent, 1.0 / DisplayRate.AsDecimal());
 		}
 	}
 
 	return true;
 }
-
 bool FFbxExporter::ExportLevelSequence(UMovieScene* MovieScene, const TArray<FGuid>& Bindings, IMovieScenePlayer* MovieScenePlayer, INodeNameAdapter& NodeNameAdapter, FMovieSceneSequenceIDRef SequenceID, const FMovieSceneSequenceTransform& RootToLocalTransform)
 {
-	if (MovieScene == nullptr || MovieScenePlayer == nullptr)
+	if (MovieScene)
 	{
-		return false;
+		if (UMovieSceneSequence* OwnerSceneSequence = MovieScene->GetTypedOuter<UMovieSceneSequence>())
+		{
+			return ExportLevelSequence(OwnerSceneSequence, OwnerSceneSequence, Bindings, MovieScenePlayer, NodeNameAdapter, SequenceID, RootToLocalTransform);
+		}
 	}
+	return false;
+}
 
+bool FFbxExporter::ExportLevelSequence(UMovieSceneSequence* MovieSceneSequence, UMovieSceneSequence* RootMovieSceneSequence, const TArray<FGuid>& Bindings, IMovieScenePlayer* MovieScenePlayer, INodeNameAdapter& NodeNameAdapter, FMovieSceneSequenceIDRef SequenceID, const FMovieSceneSequenceTransform& RootToLocalTransform)
+{
+	if (MovieSceneSequence == nullptr || RootMovieSceneSequence == nullptr || MovieScenePlayer == nullptr)
+	{
+		return false;
+	}
+	UMovieScene* MovieScene = MovieSceneSequence->GetMovieScene();
 	double FrameRate = MovieScene->GetDisplayRate().AsDecimal();
 	FbxTime::EMode TimeMode = FbxTime::ConvertFrameRateToTimeMode(FrameRate);
 
@@ -2021,7 +2090,7 @@ bool FFbxExporter::ExportLevelSequence(UMovieScene* MovieScene, const TArray<FGu
 				// now it should export everybody
 				if (FbxActor)
 				{
-					ExportLevelSequenceTracks(MovieScene, MovieScenePlayer, SequenceID, FbxActor, RuntimeObject.Get(), MovieSceneBinding.GetTracks(), RootToLocalTransform);
+					ExportLevelSequenceTracks(MovieSceneSequence, RootMovieSceneSequence, MovieScenePlayer, SequenceID, FbxActor, RuntimeObject.Get(), MovieSceneBinding.GetTracks(), RootToLocalTransform);
 					bAnyBindingsExported = true;
 				}
 			}
@@ -2030,7 +2099,7 @@ bool FFbxExporter::ExportLevelSequence(UMovieScene* MovieScene, const TArray<FGu
 		// If no bindings exported, create a dummy actor to export tracks onto
 		if (!bAnyBindingsExported)
 		{
-			ExportLevelSequenceTracks(MovieScene, MovieScenePlayer, SequenceID, nullptr, nullptr, MovieSceneBinding.GetTracks(), RootToLocalTransform);
+			ExportLevelSequenceTracks(MovieSceneSequence, RootMovieSceneSequence, MovieScenePlayer, SequenceID, nullptr, nullptr, MovieSceneBinding.GetTracks(), RootToLocalTransform);
 		}
 	}
 
diff --git a/Engine/Source/Editor/UnrealEd/Private/FbxExporter.h b/Engine/Source/Editor/UnrealEd/Private/FbxExporter.h
--- a/Engine/Source/Editor/UnrealEd/Private/FbxExporter.h
+++ b/Engine/Source/Editor/UnrealEd/Private/FbxExporter.h
@@ -28,6 +28,7 @@ class ULightComponent;
 class UMaterialInterface;
 class UModel;
 class UMovieScene;
+class UMovieSceneSequence;
 class UMovieSceneSkeletalAnimationTrack;
 class UMovieScene3DTransformTrack;
 class UMovieSceneColorTrack;
@@ -69,7 +70,11 @@ namespace UnFbx
 	class FLevelSequenceAnimTrackAdapter : public IAnimTrackAdapter
 	{
 	public:
+
+		UNREALED_API FLevelSequenceAnimTrackAdapter(IMovieScenePlayer* InMovieScenePlayer, UMovieSceneSequence* InMovieSceneSequence, UMovieSceneSequence* InRootMovieSceneSequence, const FMovieSceneSequenceTransform& InRootToLocalTransform, UMovieSceneSkeletalAnimationTrack* InAnimTrack = nullptr);
+		UE_DEPRECATED(5.5, "FLevelSequenceAnimTrackAdapter constructor taking a movie scene is deprecated, use version that takes movie scene sequences")
 		UNREALED_API FLevelSequenceAnimTrackAdapter(IMovieScenePlayer* InMovieScenePlayer, UMovieScene* InMovieScene, const FMovieSceneSequenceTransform& InRootToLocalTransform, UMovieSceneSkeletalAnimationTrack* InAnimTrack = nullptr);
+
 		UNREALED_API virtual int32 GetLocalStartFrame() const override;
 		UNREALED_API virtual int32 GetStartFrame() const override;
 		UNREALED_API virtual int32 GetLength() const override;
@@ -78,8 +83,14 @@ namespace UnFbx
 		UNREALED_API virtual UAnimSequence* GetAnimSequence(int32 LocalFrame) const override;
 		UNREALED_API virtual float GetAnimTime(int32 LocalFrame) const override;
 
+	private:
+		TRange<FFrameNumber> GetSequenceRange() const;
+
 	private:
 		IMovieScenePlayer* MovieScenePlayer;
+		UMovieSceneSequence* MovieSceneSequence;
+		UMovieSceneSequence* RootMovieSceneSequence;
+
 		UMovieScene* MovieScene;
 		FMovieSceneSequenceTransform RootToLocalTransform;
 		UMovieSceneSkeletalAnimationTrack* AnimTrack;
@@ -165,6 +176,9 @@ public:
 	 *
 	 * @return	true, if successful
 	 */
+	UNREALED_API bool ExportLevelSequence(UMovieSceneSequence* MovieSceneSequence, UMovieSceneSequence* RootMovieSceneSequence, const TArray<FGuid>& InBindings, IMovieScenePlayer* MovieScenePlayer, INodeNameAdapter& NodeNameAdapter, FMovieSceneSequenceIDRef SequenceID, const FMovieSceneSequenceTransform& RootToLocalTransform);
+
+	UE_DEPRECATED(5.5, "ExportLevelSequence taking a movie scene is deprecated, use version that takes movie scene sequences")
 	UNREALED_API bool ExportLevelSequence(UMovieScene* MovieScene, const TArray<FGuid>& InBindings, IMovieScenePlayer* MovieScenePlayer, INodeNameAdapter& NodeNameAdapter, FMovieSceneSequenceIDRef SequenceID, const FMovieSceneSequenceTransform& RootToLocalTransform);
 
 	/** Add timecode attributes to the given fbx node and add a key at the sequence playback start, using the timecode of the source section */
@@ -178,8 +192,10 @@ public:
 	 *
 	 * @return	true, if successful
 	 */
-	UNREALED_API bool ExportLevelSequenceTracks(UMovieScene* MovieScene, IMovieScenePlayer* MovieScenePlayer, FMovieSceneSequenceIDRef InSequenceID, FbxNode* FbxActor, UObject* BoundObject, const TArray<UMovieSceneTrack*>& Tracks, const FMovieSceneSequenceTransform& RootToLocalTransform);
+	UNREALED_API bool ExportLevelSequenceTracks(UMovieSceneSequence* MovieSceneSequence, UMovieSceneSequence* RootMovieSceneSequence, IMovieScenePlayer* MovieScenePlayer, FMovieSceneSequenceIDRef InSequenceID, FbxNode* FbxActor, UObject* BoundObject, const TArray<UMovieSceneTrack*>& Tracks, const FMovieSceneSequenceTransform& RootToLocalTransform);
 
+	UE_DEPRECATED(5.5, "ExportLevelSequenceTracks taking a movie scene is deprecated, use version that takes movie scene sequences")
+	UNREALED_API bool ExportLevelSequenceTracks(UMovieScene* MovieScene, IMovieScenePlayer* MovieScenePlayer, FMovieSceneSequenceIDRef InSequenceID, FbxNode* FbxActor, UObject* BoundObject, const TArray<UMovieSceneTrack*>& Tracks, const FMovieSceneSequenceTransform& RootToLocalTransform);
 
 	/**
 	 * Exports the mesh and the actor information for a static mesh actor.
