Sequencer: Spawnable restore deconstructor speculative fix.

diff --git a/Engine/Source/Editor/MovieSceneTools/Private/MovieSceneToolHelpers.cpp b/Engine/Source/Editor/MovieSceneTools/Private/MovieSceneToolHelpers.cpp
--- a/Engine/Source/Editor/MovieSceneTools/Private/MovieSceneToolHelpers.cpp
+++ b/Engine/Source/Editor/MovieSceneTools/Private/MovieSceneToolHelpers.cpp
@@ -4178,7 +4178,10 @@ FSpawnableRestoreState::~FSpawnableRestoreState()
 	for (int32 SpawnableIndex = 0; SpawnableIndex < WeakMovieScene->GetSpawnableCount(); ++SpawnableIndex)
 	{
 		FMovieSceneSpawnable& Spawnable = WeakMovieScene->GetSpawnable(SpawnableIndex);
-		Spawnable.SetSpawnOwnership(SpawnOwnershipMap[Spawnable.GetGuid()]);
+		if (ESpawnOwnership* SpawnOwnership = SpawnOwnershipMap.Find(Spawnable.GetGuid()))
+		{
+			Spawnable.SetSpawnOwnership(*SpawnOwnership);
+		}
 	}
 
 	TSharedPtr<UE::MovieScene::FSharedPlaybackState> SharedPlaybackStateToUse = SharedPlaybackState;
