Sequencer: Spawnable restore deconstructor speculative fix, part 2.

diff --git a/Engine/Source/Editor/MovieSceneTools/Private/MovieSceneToolHelpers.cpp b/Engine/Source/Editor/MovieSceneTools/Private/MovieSceneToolHelpers.cpp
--- a/Engine/Source/Editor/MovieSceneTools/Private/MovieSceneToolHelpers.cpp
+++ b/Engine/Source/Editor/MovieSceneTools/Private/MovieSceneToolHelpers.cpp
@@ -4209,7 +4209,10 @@ FSpawnableRestoreState::~FSpawnableRestoreState()
 				{
 					if (UMovieSceneSpawnableBindingBase* SpawnableBinding = BindingReference.CustomBinding->AsSpawnable(SharedPlaybackStateToUse.ToSharedRef()))
 					{
-						SpawnableBinding->SpawnOwnership = SpawnOwnershipMap[BindingReference.ID];
+						if (ESpawnOwnership* SpawnOwnership = SpawnOwnershipMap.Find(BindingReference.ID))
+						{
+							SpawnableBinding->SpawnOwnership = *SpawnOwnership;
+						}
 					}
 				}
 			}
