Sequencer: Constraints: fix for constraint creation on BPs with multiple SKM components
    - actors with multiple scene components: the binding id must be the scene component handle instead of the owner actor handle if the constraint references a scene component that is not the actor's root component.
    This will also ensure that the binding for the component is created if it hasn't already been.
    - spawnable: use Component.Get(bEvenIfPendingKill) to get a valid component ptr that we can use to resolve the new binding

#rnx

diff --git a/Engine/Source/Editor/MovieSceneTools/Private/Constraints/MovieSceneConstraintChannelHelper.cpp b/Engine/Source/Editor/MovieSceneTools/Private/Constraints/MovieSceneConstraintChannelHelper.cpp
--- a/Engine/Source/Editor/MovieSceneTools/Private/Constraints/MovieSceneConstraintChannelHelper.cpp
+++ b/Engine/Source/Editor/MovieSceneTools/Private/Constraints/MovieSceneConstraintChannelHelper.cpp
@@ -1232,8 +1232,11 @@ void FMovieSceneConstraintChannelHelper::CreateBindingIDForHandle(const TSharedP
 	{
 		return;
 	}
+
+	// make sure object is in sequencer or binding id will be empty we won't resolve the binding
+	static constexpr bool bCreateHandleIfMissing = true;
 	
-	if (const USceneComponent* SceneComponent = Cast<USceneComponent>(InHandle->GetTarget().Get()))
+	if (USceneComponent* SceneComponent = Cast<USceneComponent>(InHandle->GetTarget().Get()))
 	{
 		if (AActor* Actor = SceneComponent->GetTypedOuter<AActor>())
 		{
@@ -1246,9 +1249,21 @@ void FMovieSceneConstraintChannelHelper::CreateBindingIDForHandle(const TSharedP
 			}
 			else
 			{
-				const FGuid Guid = InSequencer->GetHandleToObject(Actor, true); //make sure object is in sequencer or binding id will be empty we won't resolve the binding
+				const FGuid Guid = InSequencer->GetHandleToObject(Actor, bCreateHandleIfMissing);
 				InHandle->ConstraintBindingID = UE::MovieScene::FRelativeObjectBindingID(Guid);
 			}
+
+			// in the context of actors with multiple scene components (such as BPs with multiple skeletal meshes, for example)
+			// the ID must be the SceneComponent handle instead of the Actor handle.
+			// this will also ensure that the binding for the component is created, if this has not yet been done.
+			if (InHandle->ConstraintBindingID.IsValid() && SceneComponent != Actor->GetRootComponent())
+			{
+				const FGuid ComponentHandle = InSequencer->GetHandleToObject(SceneComponent, bCreateHandleIfMissing);
+				if (ComponentHandle.IsValid())
+				{
+					InHandle->ConstraintBindingID = UE::MovieScene::FRelativeObjectBindingID(ComponentHandle);
+				}	
+			}
 		}
 	}
 }
diff --git a/Engine/Source/Runtime/Experimental/Animation/Constraints/Private/TransformableHandle.cpp b/Engine/Source/Runtime/Experimental/Animation/Constraints/Private/TransformableHandle.cpp
--- a/Engine/Source/Runtime/Experimental/Animation/Constraints/Private/TransformableHandle.cpp
+++ b/Engine/Source/Runtime/Experimental/Animation/Constraints/Private/TransformableHandle.cpp
@@ -536,14 +536,19 @@ bool UTransformableComponentHandle::AddTransformKeys(const TArray<FFrameNumber>&
 
 void UTransformableComponentHandle::ResolveBoundObjects(FMovieSceneSequenceID LocalSequenceID, IMovieScenePlayer& Player, UObject* SubObject)
 {
-	for (TWeakObjectPtr<> ParentObject : ConstraintBindingID.ResolveBoundObjects(LocalSequenceID, Player))
+	// in the context of blueprints being recompiled (cf. AActor::RerunConstructionScripts()) or with spawnable objects, the component has to be updated.
+	// as this can be called after it has been destroyed, we get the component even if it's pending kill otherwise Get() will return a nullptr.
+	static constexpr bool bEvenIfPendingKill = true;
+	const USceneComponent* ComponentEvenIfPendingKill = Component.Get(bEvenIfPendingKill);
+	
+	for (const TWeakObjectPtr<>& ParentObject : ConstraintBindingID.ResolveBoundObjects(LocalSequenceID, Player))
 	{
-		if (AActor* Actor = Cast<AActor>(ParentObject.Get()))
+		if (const AActor* Actor = Cast<AActor>(ParentObject.Get()))
 		{
 			const TInlineComponentArray<USceneComponent*> Components(Actor);
-			const int32 Index = Components.IndexOfByPredicate([this](const USceneComponent* SubComponent)
+			const int32 Index = Components.IndexOfByPredicate([this, ComponentEvenIfPendingKill](const USceneComponent* SubComponent)
 			{
-				return Component.IsValid() && (SubComponent->GetFName() == Component->GetFName()); 
+				return ComponentEvenIfPendingKill && (SubComponent->GetFName() == ComponentEvenIfPendingKill->GetFName());
 			});
 			Component = Index != INDEX_NONE ? Components[Index] : Actor->GetRootComponent();
 		}
