Sequencer: Fix customization so that it respects property flags (ie. not editable if VisibleAnywhere)

diff --git a/Engine/Source/Editor/MovieSceneTools/Private/MovieSceneObjectBindingIDCustomization.cpp b/Engine/Source/Editor/MovieSceneTools/Private/MovieSceneObjectBindingIDCustomization.cpp
--- a/Engine/Source/Editor/MovieSceneTools/Private/MovieSceneObjectBindingIDCustomization.cpp
+++ b/Engine/Source/Editor/MovieSceneTools/Private/MovieSceneObjectBindingIDCustomization.cpp
@@ -92,6 +92,7 @@ void FMovieSceneObjectBindingIDCustomization::CustomizeHeader(TSharedRef<IProper
 			.OnIsRecognized_Static(IsAcceptable)
 			[
 				SNew(SComboButton)
+				.IsEnabled_Lambda([this]() { return StructProperty->IsEditable(); })
 				.ToolTipText(this, &FMovieSceneObjectBindingIDCustomization::GetToolTipText)
 				.OnGetMenuContent(this, &FMovieSceneObjectBindingIDCustomization::GetPickerMenu)
 				.ContentPadding(FMargin(4.0, 2.0))
