Sequencer: Fix object property chooser overlap

diff --git a/Engine/Source/Editor/MovieSceneTools/Private/Channels/BuiltInChannelEditors.cpp b/Engine/Source/Editor/MovieSceneTools/Private/Channels/BuiltInChannelEditors.cpp
--- a/Engine/Source/Editor/MovieSceneTools/Private/Channels/BuiltInChannelEditors.cpp
+++ b/Engine/Source/Editor/MovieSceneTools/Private/Channels/BuiltInChannelEditors.cpp
@@ -430,10 +430,14 @@ TSharedRef<SWidget> CreateKeyEditor(const TMovieSceneChannelHandle<FMovieSceneOb
 				+ SHorizontalBox::Slot()
 				.AutoWidth()
 				[
-					SNew(SClassPropertyEntryBox)
-					.MetaClass(PropertyClass)
-					.SelectedClass_Lambda(GetSelectedClassLambda)
-					.OnSetClass_Lambda(OnSetClassLambda)
+					SNew(SBox)
+					.WidthOverride(100.f)
+					[
+						SNew(SClassPropertyEntryBox)
+						.MetaClass(PropertyClass)
+						.SelectedClass_Lambda(GetSelectedClassLambda)
+						.OnSetClass_Lambda(OnSetClassLambda)
+					]
 				]
 
 				+ SHorizontalBox::Slot()
@@ -469,13 +473,17 @@ TSharedRef<SWidget> CreateKeyEditor(const TMovieSceneChannelHandle<FMovieSceneOb
 				+ SHorizontalBox::Slot()
 				.AutoWidth()
 				[
-					SNew(SObjectPropertyEntryBox)
-					.DisplayBrowse(true)
-					.DisplayUseSelected(false)
-					.ObjectPath_Lambda(GetObjectPathLambda)
-					.AllowedClass(RawChannel->GetPropertyClass())
-					.OnObjectChanged_Lambda(OnSetObjectLambda)
-					.OwnerAssetDataArray(AssetDataArray)
+					SNew(SBox)
+					.WidthOverride(100.f)
+					[
+						SNew(SObjectPropertyEntryBox)
+						.DisplayBrowse(true)
+						.DisplayUseSelected(false)
+						.ObjectPath_Lambda(GetObjectPathLambda)
+						.AllowedClass(RawChannel->GetPropertyClass())
+						.OnObjectChanged_Lambda(OnSetObjectLambda)
+						.OwnerAssetDataArray(AssetDataArray)
+					]
 				]
 
 				+ SHorizontalBox::Slot()
@@ -519,19 +527,24 @@ public:
 			SNew(SHorizontalBox)
 
 			+ SHorizontalBox::Slot()
+			.AutoWidth()
 			[
-				SNew(SComboButton)
-				.OnGetMenuContent(this, &SActorReferenceBox::GetPickerMenu)
-				.ContentPadding(FMargin(0.0, 0.0))
-				.ButtonStyle(FAppStyle::Get(), "PropertyEditor.AssetComboStyle")
-				.ForegroundColor(FAppStyle::GetColor("PropertyEditor.AssetName.ColorAndOpacity"))
-				.ButtonContent()
+				SNew(SBox)
+				.WidthOverride(100.f)
 				[
-					GetCurrentItemWidget(
-						SNew(STextBlock)
-						.TextStyle(FAppStyle::Get(), "PropertyEditor.AssetClass")
-						.Font(FCoreStyle::GetDefaultFontStyle("Regular", 8))
-					)
+					SNew(SComboButton)
+					.OnGetMenuContent(this, &SActorReferenceBox::GetPickerMenu)
+					.ContentPadding(FMargin(0.0, 0.0))
+					.ButtonStyle(FAppStyle::Get(), "PropertyEditor.AssetComboStyle")
+					.ForegroundColor(FAppStyle::GetColor("PropertyEditor.AssetName.ColorAndOpacity"))
+					.ButtonContent()
+					[
+						GetCurrentItemWidget(
+							SNew(STextBlock)
+							.TextStyle(FAppStyle::Get(), "PropertyEditor.AssetClass")
+							.Font(FCoreStyle::GetDefaultFontStyle("Regular", 8))
+						)
+					]
 				]
 			]
 
