Fixed so sprinkled GEntityManagerForDebugging is set to __declspec(selectany) to prevent duplicate symbol errors when merging modules

diff --git a/Engine/Plugins/MovieScene/TemplateSequence/Source/TemplateSequence/Private/SequenceCameraShake.cpp b/Engine/Plugins/MovieScene/TemplateSequence/Source/TemplateSequence/Private/SequenceCameraShake.cpp
--- a/Engine/Plugins/MovieScene/TemplateSequence/Source/TemplateSequence/Private/SequenceCameraShake.cpp
+++ b/Engine/Plugins/MovieScene/TemplateSequence/Source/TemplateSequence/Private/SequenceCameraShake.cpp
@@ -10,7 +10,7 @@
 #include UE_INLINE_GENERATED_CPP_BY_NAME(SequenceCameraShake)
 
 #if !IS_MONOLITHIC
-	UE::MovieScene::FEntityManager*& GEntityManagerForDebugging = UE::MovieScene::GEntityManagerForDebuggingVisualizers;
+UE_SELECT_ANY UE::MovieScene::FEntityManager*& GEntityManagerForDebugging = UE::MovieScene::GEntityManagerForDebuggingVisualizers;
 #endif
 
 USequenceCameraShakePattern::USequenceCameraShakePattern(const FObjectInitializer& ObjInit)
diff --git a/Engine/Source/Editor/MovieSceneTools/Private/MovieSceneToolsModule.cpp b/Engine/Source/Editor/MovieSceneTools/Private/MovieSceneToolsModule.cpp
--- a/Engine/Source/Editor/MovieSceneTools/Private/MovieSceneToolsModule.cpp
+++ b/Engine/Source/Editor/MovieSceneTools/Private/MovieSceneToolsModule.cpp
@@ -123,7 +123,7 @@
 TAutoConsoleVariable<bool> CVarDuplicateLinkedAnimSequence(TEXT("Sequencer.DuplicateLinkedAnimSequence"), false, TEXT("When true when we duplicate a level sequence that has a linked anim sequence it will duplicate and link the anim sequencel, if false we leave any link alone."));
 
 #if !IS_MONOLITHIC
-	UE::MovieScene::FEntityManager*& GEntityManagerForDebugging = UE::MovieScene::GEntityManagerForDebuggingVisualizers;
+UE_SELECT_ANY UE::MovieScene::FEntityManager*& GEntityManagerForDebugging = UE::MovieScene::GEntityManagerForDebuggingVisualizers;
 #endif
 
 void FMovieSceneToolsModule::StartupModule()
diff --git a/Engine/Source/Editor/Sequencer/Private/SequencerModule.cpp b/Engine/Source/Editor/Sequencer/Private/SequencerModule.cpp
--- a/Engine/Source/Editor/Sequencer/Private/SequencerModule.cpp
+++ b/Engine/Source/Editor/Sequencer/Private/SequencerModule.cpp
@@ -57,7 +57,7 @@
 #include "PropertyHandle.h"
 
 #if !IS_MONOLITHIC
-	UE::MovieScene::FEntityManager*& GEntityManagerForDebugging = UE::MovieScene::GEntityManagerForDebuggingVisualizers;
+UE_SELECT_ANY UE::MovieScene::FEntityManager*& GEntityManagerForDebugging = UE::MovieScene::GEntityManagerForDebuggingVisualizers;
 #endif
 
 namespace UE::Sequencer::Private
diff --git a/Engine/Source/Runtime/MovieScene/Private/EntitySystem/MovieSceneEntityManager.cpp b/Engine/Source/Runtime/MovieScene/Private/EntitySystem/MovieSceneEntityManager.cpp
--- a/Engine/Source/Runtime/MovieScene/Private/EntitySystem/MovieSceneEntityManager.cpp
+++ b/Engine/Source/Runtime/MovieScene/Private/EntitySystem/MovieSceneEntityManager.cpp
@@ -15,7 +15,7 @@
 
 #include "EntitySystem/EntityAllocationIterator.h"
 
-UE::MovieScene::FEntityManager*& GEntityManagerForDebugging = UE::MovieScene::GEntityManagerForDebuggingVisualizers;
+UE_SELECT_ANY UE::MovieScene::FEntityManager*& GEntityManagerForDebugging = UE::MovieScene::GEntityManagerForDebuggingVisualizers;
 
 #if DO_GUARD_SLOW
 
diff --git a/Engine/Source/Runtime/MovieScene/Public/EntitySystem/MovieSceneEntityManager.h b/Engine/Source/Runtime/MovieScene/Public/EntitySystem/MovieSceneEntityManager.h
--- a/Engine/Source/Runtime/MovieScene/Public/EntitySystem/MovieSceneEntityManager.h
+++ b/Engine/Source/Runtime/MovieScene/Public/EntitySystem/MovieSceneEntityManager.h
@@ -1089,6 +1089,12 @@ private:
 	FEntityManager* EntityManager;
 };
 
+#if PLATFORM_WINDOWS
+#define UE_SELECT_ANY __declspec(selectany)
+#else
+#define UE_SELECT_ANY 
+#endif
+
 extern MOVIESCENE_API FEntityManager* GEntityManagerForDebuggingVisualizers;
 
 
diff --git a/Engine/Source/Runtime/MovieSceneTracks/Private/MovieSceneTracksModule.cpp b/Engine/Source/Runtime/MovieSceneTracks/Private/MovieSceneTracksModule.cpp
--- a/Engine/Source/Runtime/MovieSceneTracks/Private/MovieSceneTracksModule.cpp
+++ b/Engine/Source/Runtime/MovieSceneTracks/Private/MovieSceneTracksModule.cpp
@@ -12,7 +12,7 @@
 #include "MovieSceneTracksComponentTypes.h"
 
 #if !IS_MONOLITHIC
-	UE::MovieScene::FEntityManager*& GEntityManagerForDebugging = UE::MovieScene::GEntityManagerForDebuggingVisualizers;
+UE_SELECT_ANY UE::MovieScene::FEntityManager*& GEntityManagerForDebugging = UE::MovieScene::GEntityManagerForDebuggingVisualizers;
 #endif
 
 DEFINE_STAT(MovieSceneEval_ReinitializeBoundMaterials);
diff --git a/Engine/Source/Runtime/UMG/Private/UMGModule.cpp b/Engine/Source/Runtime/UMG/Private/UMGModule.cpp
--- a/Engine/Source/Runtime/UMG/Private/UMGModule.cpp
+++ b/Engine/Source/Runtime/UMG/Private/UMGModule.cpp
@@ -11,7 +11,7 @@
 DEFINE_LOG_CATEGORY(LogUMG);
 
 #if !IS_MONOLITHIC
-	UE::MovieScene::FEntityManager*& GEntityManagerForDebugging = UE::MovieScene::GEntityManagerForDebuggingVisualizers;
+UE_SELECT_ANY UE::MovieScene::FEntityManager*& GEntityManagerForDebugging = UE::MovieScene::GEntityManagerForDebuggingVisualizers;
 #endif
 
 class FUMGModule : public IUMGModule
