Sequencer: Capture the texture resources and destroy those on the next tick.

This fixes a crash opening a level sequence with an audio track with Take Recorder

diff --git a/Engine/Source/Editor/MovieSceneTools/Private/TrackEditors/AudioTrackEditor.cpp b/Engine/Source/Editor/MovieSceneTools/Private/TrackEditors/AudioTrackEditor.cpp
--- a/Engine/Source/Editor/MovieSceneTools/Private/TrackEditors/AudioTrackEditor.cpp
+++ b/Engine/Source/Editor/MovieSceneTools/Private/TrackEditors/AudioTrackEditor.cpp
@@ -174,6 +174,9 @@ private:
 	void GenerateSpline(int32 NumChannels, int32 SamplePositionOffset);
 
 private:
+
+	void DestroyTexture();
+
 	/** The section we are visualizing */
 	UMovieSceneSection& Section;
 
@@ -240,20 +243,35 @@ FAudioThumbnail::FAudioThumbnail(UMovieSceneSection& InSection, TRange<float> Dr
 
 FAudioThumbnail::~FAudioThumbnail()
 {
-	if (ShouldRender())
+	DestroyTexture();
+}
+
+void
+FAudioThumbnail::DestroyTexture()
+{
+	if (Texture)
 	{
-		BeginReleaseResource( Texture );
+		// UE-114425: Defer the destroy until the next tick to work around the RHI getting destroyed before the render command completes.
+		FSlateTexture2DRHIRef* InTexture = Texture;
 
-		FlushRenderingCommands();
-	}
+		Texture = nullptr;
 
-	if (Texture) 
-	{
-		delete Texture;
+		GEditor->GetTimerManager()->SetTimerForNextTick([this, InTexture]()
+		{
+			ENQUEUE_RENDER_COMMAND(DestroyTexture)(
+				[InTexture](FRHICommandList& RHICmdList)
+				{
+					if (InTexture)
+					{
+						InTexture->ReleaseResource();
+						delete InTexture;
+					}
+				}
+			);
+		});
 	}
 }
 
-
 FIntPoint FAudioThumbnail::GetSize() const {return FIntPoint(TextureSize, Section.GetTypedOuter<UMovieSceneAudioTrack>()->GetRowHeight());}
 FSlateShaderResource* FAudioThumbnail::GetViewportRenderTargetTexture() const {return Texture;}
 bool FAudioThumbnail::RequiresVsync() const {return false;}
