Sequencer: Allow editing multiple selected keyframes only when there's more than one keyframe selected.

Change the foreground color for the key editor to yellow when there are multiple selected keyframes being edited.

diff --git a/Engine/Source/Editor/MovieSceneTools/Private/CurveKeyEditors/SNumericKeyEditor.h b/Engine/Source/Editor/MovieSceneTools/Private/CurveKeyEditors/SNumericKeyEditor.h
--- a/Engine/Source/Editor/MovieSceneTools/Private/CurveKeyEditors/SNumericKeyEditor.h
+++ b/Engine/Source/Editor/MovieSceneTools/Private/CurveKeyEditors/SNumericKeyEditor.h
@@ -92,6 +92,18 @@ public:
 
 private:
 
+	virtual FSlateColor GetForegroundColor() const override
+	{
+		if (KeyEditor.GetEditingKeySelection())
+		{
+			return FLinearColor::Yellow;
+		}
+		else
+		{
+			return FSlateColor::UseForeground();
+		}
+	}
+
 	void OnBeginSliderMovement()
 	{
 		GEditor->BeginTransaction(LOCTEXT("SetNumericKey", "Set Key Value"));
diff --git a/Engine/Source/Editor/MovieSceneTools/Private/CurveKeyEditors/SequencerKeyEditor.h b/Engine/Source/Editor/MovieSceneTools/Private/CurveKeyEditors/SequencerKeyEditor.h
--- a/Engine/Source/Editor/MovieSceneTools/Private/CurveKeyEditors/SequencerKeyEditor.h
+++ b/Engine/Source/Editor/MovieSceneTools/Private/CurveKeyEditors/SequencerKeyEditor.h
@@ -93,6 +93,37 @@ struct TSequencerKeyEditor
 		return Result;
 	}
 
+	bool GetEditingKeySelection() const
+	{
+		using namespace UE::MovieScene;
+		using namespace Sequencer;
+		using namespace UE::Sequencer;
+
+		ChannelType* Channel = ChannelHandle.Get();
+
+		ISequencer* Sequencer = WeakSequencer.Pin().Get();
+		const FKeySelection& KeySelection = Sequencer->GetViewModel()->GetSelection()->KeySelection;
+
+		// Allow editing the key selection if the key editor's channel is one of the selected key's channels and there's more than 1 of those keys selected
+		bool bAllowEditingKeySelection = false;
+		int32 NumSelectedKeys = 0;
+		for (FKeyHandle Key : KeySelection)
+		{
+			// Make sure we only manipulate the values of the channel with the same channel type we're editing
+			TSharedPtr<FChannelModel> ChannelModel = KeySelection.GetModelForKey(Key);
+			if (ChannelModel && ChannelModel->GetChannel() == Channel)
+			{
+				++NumSelectedKeys;
+				if (NumSelectedKeys > 1)
+				{
+					return true;
+				}
+			}
+		}
+
+		return false;
+	}
+
 	void SetValue(const ValueType& InValue)
 	{
 		using namespace UE::MovieScene;
@@ -119,20 +150,8 @@ struct TSequencerKeyEditor
 
 		const FKeySelection& KeySelection = Sequencer->GetViewModel()->GetSelection()->KeySelection;
 
-		// Allow editing the key selection if the key editor's channel is one of the selected key's channels
-		bool bAllowEditingKeySelection = false;
-		for (FKeyHandle Key : KeySelection)
-		{
-			// Make sure we only manipulate the values of the channel with the same channel type we're editing
-			TSharedPtr<FChannelModel> ChannelModel = KeySelection.GetModelForKey(Key);
-			if (ChannelModel && ChannelModel->GetChannel() == Channel)
-			{
-				bAllowEditingKeySelection = true;	
-				break;
-			}
-		}
-
-		if (bAllowEditingKeySelection)
+		const bool bEditingKeySelection = GetEditingKeySelection();
+		if (bEditingKeySelection)
 		{
 			for (FKeyHandle Key : KeySelection)
 			{
