Sequencer: Control Rig: Fix collapsing sections when there are no absoluate sections.

diff --git a/Engine/Source/Editor/MovieSceneTools/Private/MovieSceneToolHelpers.cpp b/Engine/Source/Editor/MovieSceneTools/Private/MovieSceneToolHelpers.cpp
--- a/Engine/Source/Editor/MovieSceneTools/Private/MovieSceneToolHelpers.cpp
+++ b/Engine/Source/Editor/MovieSceneTools/Private/MovieSceneToolHelpers.cpp
@@ -3825,6 +3825,11 @@ bool MovieSceneToolHelpers::BakeToSkelMeshToCallbacks(const FAnimExportSequenceP
 {
 	UMovieScene* MovieScene = AESP.MovieSceneSequence->GetMovieScene();
 	TArray< USkeletalMeshComponent*> SkelMeshComps;
+	if (!ExportOptions)
+	{
+		UE_LOG(LogMovieScene, Warning, TEXT(" MovieSceneToolHelpers::BakeToSkelMesh functions require a valid AnimSeqExportOption"));
+		return false;
+	}
 	if (ExportOptions->bEvaluateAllSkeletalMeshComponents)
 	{
 		AActor* Actor = InSkelMeshComp->GetTypedOuter<AActor>();
@@ -5063,6 +5068,11 @@ static bool MergeChannels(TArray<ChannelType*>& Channels,const  TArray<UMovieSce
 					if (Sections[WeightIndex]->GetRange().Contains(Frame))  
 					{
 						float Weight = Sections[WeightIndex]->GetTotalWeightValue(FrameTime);
+						if (Sections[WeightIndex]->GetBlendType().IsValid() == false ||
+							Sections[WeightIndex]->GetBlendType().Get() != EMovieSceneBlendType::Additive)
+						{
+							Weight = 1.0;
+						}
 						CurveValueType WeightedValue = 0.0;
 						ChannelType* EachChannel = Channels[WeightIndex];
 						EachChannel->Evaluate(FrameTime, WeightedValue);
