Sequencer: Linked Anim Sequences. Allow ability to specify custom range when saving out custom time range.

diff --git a/Engine/Plugins/Animation/ControlRig/Source/ControlRigEditor/Private/Sequencer/ControlRigParameterTrackEditor.cpp b/Engine/Plugins/Animation/ControlRig/Source/ControlRigEditor/Private/Sequencer/ControlRigParameterTrackEditor.cpp
--- a/Engine/Plugins/Animation/ControlRig/Source/ControlRigEditor/Private/Sequencer/ControlRigParameterTrackEditor.cpp
+++ b/Engine/Plugins/Animation/ControlRig/Source/ControlRigEditor/Private/Sequencer/ControlRigParameterTrackEditor.cpp
@@ -1040,6 +1040,7 @@ void FControlRigParameterTrackEditor::BakeToControlRig(UClass* InClass, FGuid Ob
 			UAnimSeqExportOption* AnimSeqExportOption = NewObject<UAnimSeqExportOption>(GetTransientPackage(), NAME_None);
 			UBakeToControlRigSettings* BakeSettings = GetMutableDefault<UBakeToControlRigSettings>();
 			AnimSeqExportOption->bTransactRecording = false;
+			AnimSeqExportOption->CustomDisplayRate = ParentSequencer->GetFocusedDisplayRate();
 
 			TSharedPtr<SWindow> ParentWindow;
 			if (FModuleManager::Get().IsModuleLoaded("MainFrame"))
@@ -1269,6 +1270,7 @@ void FControlRigParameterTrackEditor::BakeInvertedPose(UControlRig* InControlRig
 	UMovieSceneSequence* MovieSceneSequence = GetSequencer()->GetFocusedMovieSceneSequence();
 	UMovieScene* MovieScene = MovieSceneSequence->GetMovieScene();
 	UAnimSeqExportOption* ExportOptions = NewObject<UAnimSeqExportOption>(GetTransientPackage(), NAME_None);
+
 	//@sara to do, not sure if you want to key reduce after, but BakeSettings isn't used
 	//UBakeToControlRigSettings* BakeSettings = GetMutableDefault<UBakeToControlRigSettings>();
 	const TSharedPtr<ISequencer> ParentSequencer = GetSequencer();
diff --git a/Engine/Plugins/MovieScene/SequencerScripting/Source/SequencerScriptingEditor/Private/SequencerTools.cpp b/Engine/Plugins/MovieScene/SequencerScripting/Source/SequencerScriptingEditor/Private/SequencerTools.cpp
--- a/Engine/Plugins/MovieScene/SequencerScripting/Source/SequencerScriptingEditor/Private/SequencerTools.cpp
+++ b/Engine/Plugins/MovieScene/SequencerScripting/Source/SequencerScriptingEditor/Private/SequencerTools.cpp
@@ -455,7 +455,8 @@ bool USequencerToolsFunctionLibrary::LinkAnimSequence(ULevelSequence*  Sequence,
 			{
 				for (FLevelSequenceAnimSequenceLinkItem& LevelAnimLinkItem : LevelAnimLink->AnimSequenceLinks)
 				{
-					if (LevelAnimLinkItem.SkelTrackGuid == Binding.BindingID)
+					if (LevelAnimLinkItem.IsEqual(Binding.BindingID, ExportOptions->bUseCustomTimeRange,
+						ExportOptions->CustomStartFrame, ExportOptions->CustomEndFrame, ExportOptions->CustomDisplayRate))
 					{
 						bAddItem = false;
 						UAnimSequence* OtherAnimSequence = LevelAnimLinkItem.ResolveAnimSequence();
@@ -481,6 +482,15 @@ bool USequencerToolsFunctionLibrary::LinkAnimSequence(ULevelSequence*  Sequence,
 						LevelAnimLinkItem.CurveInterpolation = ExportOptions->CurveInterpolation;
 						LevelAnimLinkItem.bRecordInWorldSpace = ExportOptions->bRecordInWorldSpace;
 						LevelAnimLinkItem.bEvaluateAllSkeletalMeshComponents = ExportOptions->bEvaluateAllSkeletalMeshComponents;
+						
+						LevelAnimLinkItem.IncludeAnimationNames = ExportOptions->IncludeAnimationNames;
+						LevelAnimLinkItem.ExcludeAnimationNames = ExportOptions->ExcludeAnimationNames;
+						LevelAnimLinkItem.WarmUpFrames = ExportOptions->WarmUpFrames;
+						LevelAnimLinkItem.DelayBeforeStart = ExportOptions->DelayBeforeStart;
+						LevelAnimLinkItem.bUseCustomTimeRange = ExportOptions->bUseCustomTimeRange;
+						LevelAnimLinkItem.CustomStartFrame = ExportOptions->CustomStartFrame;
+						LevelAnimLinkItem.CustomEndFrame = ExportOptions->CustomEndFrame;
+						LevelAnimLinkItem.CustomDisplayRate = ExportOptions->CustomDisplayRate;
 
 						break;
 					}
@@ -504,6 +514,15 @@ bool USequencerToolsFunctionLibrary::LinkAnimSequence(ULevelSequence*  Sequence,
 				LevelAnimLinkItem.bRecordInWorldSpace = ExportOptions->bRecordInWorldSpace;
 				LevelAnimLinkItem.bEvaluateAllSkeletalMeshComponents = ExportOptions->bEvaluateAllSkeletalMeshComponents;
 
+				LevelAnimLinkItem.IncludeAnimationNames = ExportOptions->IncludeAnimationNames;
+				LevelAnimLinkItem.ExcludeAnimationNames = ExportOptions->ExcludeAnimationNames;
+				LevelAnimLinkItem.WarmUpFrames = ExportOptions->WarmUpFrames;
+				LevelAnimLinkItem.DelayBeforeStart = ExportOptions->DelayBeforeStart;
+				LevelAnimLinkItem.bUseCustomTimeRange = ExportOptions->bUseCustomTimeRange;
+				LevelAnimLinkItem.CustomStartFrame = ExportOptions->CustomStartFrame;
+				LevelAnimLinkItem.CustomEndFrame = ExportOptions->CustomEndFrame;
+				LevelAnimLinkItem.CustomDisplayRate = ExportOptions->CustomDisplayRate;
+
 				LevelAnimLink->AnimSequenceLinks.Add(LevelAnimLinkItem);
 				AssetUserDataInterface->AddAssetUserData(LevelAnimLink);
 			}
diff --git a/Engine/Source/Editor/MovieSceneTools/Private/MovieSceneToolHelpers.cpp b/Engine/Source/Editor/MovieSceneTools/Private/MovieSceneToolHelpers.cpp
--- a/Engine/Source/Editor/MovieSceneTools/Private/MovieSceneToolHelpers.cpp
+++ b/Engine/Source/Editor/MovieSceneTools/Private/MovieSceneToolHelpers.cpp
@@ -3854,6 +3854,19 @@ bool MovieSceneToolHelpers::BakeToSkelMeshToCallbacks(const FAnimExportSequenceP
 	}
 
 	UnFbx::FLevelSequenceAnimTrackAdapter AnimTrackAdapter(AESP.Player, AESP.MovieSceneSequence, AESP.RootMovieSceneSequence, AESP.RootToLocalTransform);
+	if (ExportOptions->bUseCustomTimeRange)
+	{
+
+		const FFrameRate TickResolution = MovieScene->GetTickResolution();
+		const FFrameRate DisplayResolution = ExportOptions->CustomDisplayRate;
+		const FFrameNumber StartFrameInTick = FFrameRate::TransformTime(FFrameTime(ExportOptions->CustomStartFrame), DisplayResolution, TickResolution).FloorToFrame();
+		FFrameNumber EndFrameInTick = FFrameRate::TransformTime(FFrameTime(ExportOptions->CustomEndFrame), DisplayResolution, TickResolution).CeilToFrame();
+		if (EndFrameInTick < StartFrameInTick)
+		{
+			EndFrameInTick = StartFrameInTick;
+		}
+		AnimTrackAdapter.SetRange(StartFrameInTick,EndFrameInTick);
+	}
 	int32 LocalStartFrame = AnimTrackAdapter.GetLocalStartFrame();
 	int32 StartFrame = AnimTrackAdapter.GetStartFrame();
 	int32 AnimationLength = AnimTrackAdapter.GetLength();
diff --git a/Engine/Source/Editor/MovieSceneTools/Private/TrackEditors/SkeletalAnimationTrackEditor.cpp b/Engine/Source/Editor/MovieSceneTools/Private/TrackEditors/SkeletalAnimationTrackEditor.cpp
--- a/Engine/Source/Editor/MovieSceneTools/Private/TrackEditors/SkeletalAnimationTrackEditor.cpp
+++ b/Engine/Source/Editor/MovieSceneTools/Private/TrackEditors/SkeletalAnimationTrackEditor.cpp
@@ -1147,6 +1147,7 @@ FSkeletalAnimationTrackEditor::FSkeletalAnimationTrackEditor( TSharedRef<ISequen
 	//We use the FGCObject pattern to keep the anim export option alive during the editor session
 
 	AnimSeqExportOption = NewObject<UAnimSeqExportOption>();
+
 }
 
 void FSkeletalAnimationTrackEditor::OnInitialize()
@@ -1316,6 +1317,14 @@ void FSkeletalAnimationTrackEditor::OnSequencerSaved(ISequencer& )
 						const bool bSavedEvaluateAllSkeletalMeshComponents = AnimSeqExportOption->bEvaluateAllSkeletalMeshComponents;
 						const EAnimInterpolationType SavedInterpolationType = AnimSeqExportOption->Interpolation;
 						const ERichCurveInterpMode SavedCurveInterpolationType = AnimSeqExportOption->CurveInterpolation;
+						const TArray<FString> SavedIncludeAnimationNames = AnimSeqExportOption->IncludeAnimationNames;
+						const TArray<FString> SavedExcludeAnimationNames = AnimSeqExportOption->ExcludeAnimationNames;
+						const FFrameNumber SavedWarmUpFrames = AnimSeqExportOption->WarmUpFrames;
+						const FFrameNumber SavedDelayBeforeStart = AnimSeqExportOption->DelayBeforeStart;
+						const bool bSavedUseCustomTimeRange = AnimSeqExportOption->bUseCustomTimeRange;
+						const FFrameNumber SavedCustomStartFrame = AnimSeqExportOption->CustomStartFrame;
+						const FFrameNumber SavedCustomEndFrame = AnimSeqExportOption->CustomEndFrame;
+						const FFrameRate SavedCustomDisplayRate = AnimSeqExportOption->CustomDisplayRate;
 
 						AnimSeqExportOption->bExportMorphTargets = Item.bExportMorphTargets;
 						AnimSeqExportOption->bExportAttributeCurves = Item.bExportAttributeCurves;
@@ -1326,6 +1335,15 @@ void FSkeletalAnimationTrackEditor::OnSequencerSaved(ISequencer& )
 						AnimSeqExportOption->Interpolation = Item.Interpolation;
 						AnimSeqExportOption->CurveInterpolation = Item.CurveInterpolation;
 
+						AnimSeqExportOption->IncludeAnimationNames = Item.IncludeAnimationNames;
+						AnimSeqExportOption->ExcludeAnimationNames = Item.ExcludeAnimationNames;
+						AnimSeqExportOption->WarmUpFrames = Item.WarmUpFrames;
+						AnimSeqExportOption->DelayBeforeStart = Item.DelayBeforeStart;
+						AnimSeqExportOption->bUseCustomTimeRange = Item.bUseCustomTimeRange;
+						AnimSeqExportOption->CustomStartFrame = Item.CustomStartFrame;
+						AnimSeqExportOption->CustomEndFrame = Item.CustomEndFrame;
+						AnimSeqExportOption->CustomDisplayRate = Item.CustomDisplayRate;
+
 						FAnimExportSequenceParameters AESP;
 						AESP.Player = SequencerPtr.Get();
 						AESP.RootToLocalTransform = RootToLocalTransform;
@@ -1342,6 +1360,15 @@ void FSkeletalAnimationTrackEditor::OnSequencerSaved(ISequencer& )
 						AnimSeqExportOption->Interpolation = SavedInterpolationType;
 						AnimSeqExportOption->CurveInterpolation = SavedCurveInterpolationType;
 
+						AnimSeqExportOption->IncludeAnimationNames = SavedIncludeAnimationNames;
+						AnimSeqExportOption->ExcludeAnimationNames = SavedExcludeAnimationNames;
+						AnimSeqExportOption->WarmUpFrames = SavedWarmUpFrames;
+						AnimSeqExportOption->DelayBeforeStart = SavedDelayBeforeStart;
+						AnimSeqExportOption->bUseCustomTimeRange = bSavedUseCustomTimeRange;
+						AnimSeqExportOption->CustomStartFrame = SavedCustomStartFrame;
+						AnimSeqExportOption->CustomEndFrame = SavedCustomEndFrame;
+						AnimSeqExportOption->CustomDisplayRate = SavedCustomDisplayRate;
+
 						//save the anim sequence to disk to make sure they are in sync
 						UPackage* const Package = AnimSequence->GetOutermost();
 						FString const PackageName = Package->GetName();
@@ -1488,7 +1515,7 @@ bool FSkeletalAnimationTrackEditor::CreateAnimationSequence(const TArray<UObject
 				AESP.RootToLocalTransform = RootToLocalTransform;
 				AESP.MovieSceneSequence = MovieSceneSequence;
 				AESP.RootMovieSceneSequence = RootMovieSceneSequence;
-
+				AnimSeqExportOption->CustomDisplayRate = ParentSequencer->GetFocusedDisplayRate();
 				bResult  = MovieSceneToolHelpers::ExportToAnimSequence(AnimSequence, AnimSeqExportOption, AESP, SkelMeshComp);
 			}
 		}
@@ -1522,7 +1549,8 @@ bool FSkeletalAnimationTrackEditor::CreateAnimationSequence(const TArray<UObject
 					{
 						for (FLevelSequenceAnimSequenceLinkItem& LevelAnimLinkItem : LevelAnimLink->AnimSequenceLinks)
 						{
-							if (LevelAnimLinkItem.SkelTrackGuid == Binding)
+							if (LevelAnimLinkItem.IsEqual(Binding, AnimSeqExportOption->bUseCustomTimeRange,
+								AnimSeqExportOption->CustomStartFrame, AnimSeqExportOption->CustomEndFrame, AnimSeqExportOption->CustomDisplayRate))
 							{
 								bAddItem = false;
 								UAnimSequence* OtherAnimSequence = LevelAnimLinkItem.ResolveAnimSequence();
@@ -1547,7 +1575,14 @@ bool FSkeletalAnimationTrackEditor::CreateAnimationSequence(const TArray<UObject
 								LevelAnimLinkItem.bEvaluateAllSkeletalMeshComponents = AnimSeqExportOption->bEvaluateAllSkeletalMeshComponents;
 								LevelAnimLinkItem.Interpolation = AnimSeqExportOption->Interpolation;
 								LevelAnimLinkItem.CurveInterpolation = AnimSeqExportOption->CurveInterpolation;
-
+								LevelAnimLinkItem.IncludeAnimationNames = AnimSeqExportOption->IncludeAnimationNames;
+								LevelAnimLinkItem.ExcludeAnimationNames = AnimSeqExportOption->ExcludeAnimationNames;
+								LevelAnimLinkItem.WarmUpFrames = AnimSeqExportOption->WarmUpFrames;
+								LevelAnimLinkItem.DelayBeforeStart = AnimSeqExportOption->DelayBeforeStart;
+								LevelAnimLinkItem.bUseCustomTimeRange = AnimSeqExportOption->bUseCustomTimeRange;
+								LevelAnimLinkItem.CustomStartFrame = AnimSeqExportOption->CustomStartFrame;
+								LevelAnimLinkItem.CustomEndFrame = AnimSeqExportOption->CustomEndFrame;
+								LevelAnimLinkItem.CustomDisplayRate = AnimSeqExportOption->CustomDisplayRate;
 								break;
 							}
 						}
@@ -1570,6 +1605,14 @@ bool FSkeletalAnimationTrackEditor::CreateAnimationSequence(const TArray<UObject
 						LevelAnimLinkItem.bEvaluateAllSkeletalMeshComponents = AnimSeqExportOption->bEvaluateAllSkeletalMeshComponents;
 						LevelAnimLinkItem.Interpolation = AnimSeqExportOption->Interpolation;
 						LevelAnimLinkItem.CurveInterpolation = AnimSeqExportOption->CurveInterpolation;
+						LevelAnimLinkItem.IncludeAnimationNames = AnimSeqExportOption->IncludeAnimationNames;
+						LevelAnimLinkItem.ExcludeAnimationNames = AnimSeqExportOption->ExcludeAnimationNames;
+						LevelAnimLinkItem.WarmUpFrames = AnimSeqExportOption->WarmUpFrames;
+						LevelAnimLinkItem.DelayBeforeStart = AnimSeqExportOption->DelayBeforeStart;
+						LevelAnimLinkItem.bUseCustomTimeRange = AnimSeqExportOption->bUseCustomTimeRange;
+						LevelAnimLinkItem.CustomStartFrame = AnimSeqExportOption->CustomStartFrame;
+						LevelAnimLinkItem.CustomEndFrame = AnimSeqExportOption->CustomEndFrame;
+						LevelAnimLinkItem.CustomDisplayRate = AnimSeqExportOption->CustomDisplayRate;
 
 						LevelAnimLink->AnimSequenceLinks.Add(LevelAnimLinkItem);
 						AssetUserDataInterface->AddAssetUserData(LevelAnimLink);
@@ -1624,6 +1667,10 @@ bool FSkeletalAnimationTrackEditor::CreateAnimationSequence(const TArray<UObject
 			{
 				Notification->SetCompletionState(SNotificationItem::CS_Success);
 			}
+			if (const TSharedPtr<ISequencer> ParentSequencer = GetSequencer())
+			{
+				ParentSequencer->RequestEvaluate(); 
+			}
 		}
 		else
 		{
diff --git a/Engine/Source/Editor/UnrealEd/Classes/Exporters/AnimSeqExportOption.h b/Engine/Source/Editor/UnrealEd/Classes/Exporters/AnimSeqExportOption.h
--- a/Engine/Source/Editor/UnrealEd/Classes/Exporters/AnimSeqExportOption.h
+++ b/Engine/Source/Editor/UnrealEd/Classes/Exporters/AnimSeqExportOption.h
@@ -70,6 +70,22 @@ public:
 	UPROPERTY(EditAnywhere, BlueprintReadWrite, AdvancedDisplay, Category = Export)
 	bool bTransactRecording = true;
 
+	/** Whether or not to use custom time range */
+	UPROPERTY(EditAnywhere, BlueprintReadWrite, AdvancedDisplay, Category = Export)
+	bool bUseCustomTimeRange = false;
+
+	/** Custom start frame in custom display rate*/
+	UPROPERTY(EditAnywhere, BlueprintReadWrite, AdvancedDisplay, meta = (EditCondition = bUseCustomTimeRange),  Category = Export)
+	FFrameNumber CustomStartFrame = 0;
+
+	/** Custom end frame in custom display rate */
+	UPROPERTY(EditAnywhere, BlueprintReadWrite, AdvancedDisplay, meta = (EditCondition = bUseCustomTimeRange),  Category = Export)
+	FFrameNumber CustomEndFrame = 120;
+
+	/** Custom display rate, should be set from the movie scene/sequencer display rate */
+	UPROPERTY(BlueprintReadWrite, AdvancedDisplay, meta = (EditCondition = bUseCustomTimeRange), Category = Export)
+	FFrameRate CustomDisplayRate = FFrameRate(30,1);
+
 	void ResetToDefault()
 	{
 		bExportTransforms = true;
@@ -83,5 +99,9 @@ public:
 		WarmUpFrames = 0;
 		DelayBeforeStart = 0;
 		bTransactRecording = true;
+		bUseCustomTimeRange = false;
+		CustomStartFrame = 0;
+		CustomEndFrame = 120;
+		CustomDisplayRate = FFrameRate(30, 1);
 	}
 };
diff --git a/Engine/Source/Editor/UnrealEd/Private/Fbx/FbxMainExport.cpp b/Engine/Source/Editor/UnrealEd/Private/Fbx/FbxMainExport.cpp
--- a/Engine/Source/Editor/UnrealEd/Private/Fbx/FbxMainExport.cpp
+++ b/Engine/Source/Editor/UnrealEd/Private/Fbx/FbxMainExport.cpp
@@ -1731,8 +1731,19 @@ FLevelSequenceAnimTrackAdapter::FLevelSequenceAnimTrackAdapter(IMovieScenePlayer
 	}
 }
 
+void FLevelSequenceAnimTrackAdapter::SetRange(const FFrameNumber& StartFrame, const FFrameNumber& EndFrame)
+{
+	TPair<FFrameNumber, FFrameNumber> Range(StartFrame, EndFrame);
+	OptionalRange = Range;
+}
+
 TRange<FFrameNumber> FLevelSequenceAnimTrackAdapter::GetSequenceRange() const
 {
+	if (OptionalRange.IsSet())
+	{
+		TRange<FFrameNumber> SetRange(OptionalRange.GetValue().Key, (OptionalRange.GetValue().Value));
+		return SetRange;
+	}
 	MovieScenePlayer->State.AssignSequence(MovieSceneSequenceID::Root, *RootMovieSceneSequence, *MovieScenePlayer);
 	FMovieSceneSequenceIDRef Template = MovieScenePlayer->State.FindSequenceId(MovieSceneSequence);
 
diff --git a/Engine/Source/Editor/UnrealEd/Private/FbxExporter.h b/Engine/Source/Editor/UnrealEd/Private/FbxExporter.h
--- a/Engine/Source/Editor/UnrealEd/Private/FbxExporter.h
+++ b/Engine/Source/Editor/UnrealEd/Private/FbxExporter.h
@@ -65,6 +65,8 @@ namespace UnFbx
 		virtual UAnimSequence* GetAnimSequence(int32 LocalFrame) const { return nullptr; }
 		/** The time into the anim sequence for the given LocalFrame */
 		virtual float GetAnimTime(int32 LocalFrame) const { return 0.f; }
+		/** Set optional range in frames(tick resolution)*/
+		virtual void SetRange(const FFrameNumber& StartFrame, const FFrameNumber& EndFrame) = 0;
 	};
 
 	/** An anim track adapter for a level sequence. */
@@ -83,6 +85,7 @@ namespace UnFbx
 		UNREALED_API virtual double GetFrameRate() const override;
 		UNREALED_API virtual UAnimSequence* GetAnimSequence(int32 LocalFrame) const override;
 		UNREALED_API virtual float GetAnimTime(int32 LocalFrame) const override;
+		UNREALED_API virtual void SetRange(const FFrameNumber& StartFrame, const FFrameNumber& EndFrame) override;
 
 	private:
 		TRange<FFrameNumber> GetSequenceRange() const;
@@ -95,6 +98,7 @@ namespace UnFbx
 		UMovieScene* MovieScene;
 		FMovieSceneSequenceTransform RootToLocalTransform;
 		UMovieSceneSkeletalAnimationTrack* AnimTrack;
+		TOptional < TPair<FFrameNumber, FFrameNumber>> OptionalRange;
 	};
 /**
  * Main FBX Exporter class.
diff --git a/Engine/Source/Runtime/LevelSequence/Public/LevelSequenceAnimSequenceLink.h b/Engine/Source/Runtime/LevelSequence/Public/LevelSequenceAnimSequenceLink.h
--- a/Engine/Source/Runtime/LevelSequence/Public/LevelSequenceAnimSequenceLink.h
+++ b/Engine/Source/Runtime/LevelSequence/Public/LevelSequenceAnimSequenceLink.h
@@ -45,13 +45,56 @@ public:
 	bool bRecordInWorldSpace = false;
 	UPROPERTY(BlueprintReadWrite, Category = Property)
 	bool bEvaluateAllSkeletalMeshComponents = true;
+	
+	/** Include only the animation bones/curves that match this list */
+	UPROPERTY(BlueprintReadWrite, Category = Property);
+	TArray<FString> IncludeAnimationNames;
+	/** Exclude all animation bones/curves that match this list */
+	UPROPERTY(BlueprintReadWrite, Category = Property);
+	TArray<FString> ExcludeAnimationNames;
+	/** Number of Display Rate frames to evaluate before doing the export. It will evaluate after any Delay. This will use frames before the start frame. Use it if there is some post anim BP effects you want to run before export start time.*/
+	UPROPERTY(BlueprintReadWrite, Category = Property);
+	FFrameNumber WarmUpFrames = 0;
+	/** Number of Display Rate frames to delay at the same frame before doing the export. It will evalaute first, then any warm up, then the export. Use it if there is some post anim BP effects you want to ran repeatedly at the start.*/
+	UPROPERTY(BlueprintReadWrite, AdvancedDisplay, Category = Property);
+	FFrameNumber DelayBeforeStart = 0;
+	/** Whether or not to use custom time range */
+	UPROPERTY(BlueprintReadWrite, Category = Property)
+	bool bUseCustomTimeRange = false;
+	/** Custom start frame in display rate*/
+	UPROPERTY(BlueprintReadWrite, meta = (EditCondition = bUseCustomTimeRange), Category = Property)
+	FFrameNumber CustomStartFrame = 0;
+	/** Custom end frame in display rate */
+	UPROPERTY(BlueprintReadWrite, meta = (EditCondition = bUseCustomTimeRange), Category = Property)
+	FFrameNumber CustomEndFrame = 120;
+	/** Custom display rate, should be set from the movie scene/sequencer display rate */
+	UPROPERTY(EditAnywhere, BlueprintReadWrite, meta = (EditCondition = bUseCustomTimeRange), Category = Property)
+	FFrameRate CustomDisplayRate = FFrameRate(30, 1);
+
 
 	LEVELSEQUENCE_API void SetAnimSequence(UAnimSequence* InAnimSequence);
 	LEVELSEQUENCE_API UAnimSequence* ResolveAnimSequence();
 
+	LEVELSEQUENCE_API bool IsEqual(FGuid InSkelTrackGuid, bool bInUseCustomTimeRange = false,
+		FFrameNumber InCustomStartFrame = 0, FFrameNumber InCustomEndFrame = 120,
+		FFrameRate InCustomDisplayRate = FFrameRate(30,1))
+	{
+		if (InSkelTrackGuid == SkelTrackGuid)
+		{
+			if (bUseCustomTimeRange == bInUseCustomTimeRange)
+			{
+				if (bUseCustomTimeRange == false || (InCustomStartFrame == CustomStartFrame &&
+					InCustomEndFrame == CustomEndFrame && InCustomDisplayRate == CustomDisplayRate))
+				{
+					return true;
+				}
+			}
+		}
+		return false;
+	}
 };
 
-/** Link To Set of Anim Sequences that we may be linked to.*/
+/** Link To Set of Anim Sequences that we may belinked to.*/
 UCLASS(BlueprintType, MinimalAPI)
 class ULevelSequenceAnimSequenceLink : public UAssetUserData
 {
