Tweener & SetKeyPositions improvements
    - SControlRigTweenSlider: defer blend on tick instead of blending for each mouse move
    - FChannelCurveModel::SetKeyPositions only expands the section range if needed

diff --git a/Engine/Plugins/Animation/ControlRig/Source/ControlRigEditor/Private/EditMode/SControlRigTweenWidget.cpp b/Engine/Plugins/Animation/ControlRig/Source/ControlRigEditor/Private/EditMode/SControlRigTweenWidget.cpp
--- a/Engine/Plugins/Animation/ControlRig/Source/ControlRigEditor/Private/EditMode/SControlRigTweenWidget.cpp
+++ b/Engine/Plugins/Animation/ControlRig/Source/ControlRigEditor/Private/EditMode/SControlRigTweenWidget.cpp
@@ -83,7 +83,13 @@ void SControlRigTweenSlider::OnPoseBlendChanged(double ChangedVal)
 	if (WeakSequencer.IsValid() && bIsBlending)
 	{
 		PoseBlendValue = ChangedVal;
-		AnimSlider->Blend(WeakSequencer, ChangedVal);
+
+		// defer blend function on next tick
+		PendingBlendFunction = [this, ChangedVal]()
+		{
+			AnimSlider->Blend(WeakSequencer, ChangedVal);
+		};
+
 		WeakSequencer.Pin()->NotifyMovieSceneDataChanged(EMovieSceneDataChangeType::TrackValueChanged);
 	}
 
@@ -97,6 +103,17 @@ void SControlRigTweenSlider::ResetAnimSlider()
 	}
 }
 
+void SControlRigTweenSlider::Tick(const FGeometry& InAllottedGeometry, const double InCurrentTime, const float InDeltaTime)
+{
+	if (PendingBlendFunction)
+	{
+		PendingBlendFunction();
+		PendingBlendFunction.Reset();
+	}
+	
+	SCompoundWidget::Tick(InAllottedGeometry, InCurrentTime, InDeltaTime);
+}
+
 void SControlRigTweenSlider::DragAnimSliderTool(double Val)
 {
 	if (SpinBox.IsValid())
@@ -124,6 +141,7 @@ void SControlRigTweenSlider::DragAnimSliderTool(double Val)
 
 void SControlRigTweenSlider::OnBeginSliderMovement()
 {
+	PendingBlendFunction.Reset();
 	if (bSliderStartedTransaction == false)
 	{
 		bIsBlending = bSliderStartedTransaction = Setup();
@@ -143,6 +161,12 @@ bool SControlRigTweenSlider::Setup()
 
 void SControlRigTweenSlider::OnEndSliderMovement(double NewValue)
 {
+	if (PendingBlendFunction)
+	{
+		PendingBlendFunction();
+		PendingBlendFunction.Reset();
+	}
+	
 	if (bSliderStartedTransaction)
 	{
 		GEditor->EndTransaction();
diff --git a/Engine/Plugins/Animation/ControlRig/Source/ControlRigEditor/Private/EditMode/SControlRigTweenWidget.h b/Engine/Plugins/Animation/ControlRig/Source/ControlRigEditor/Private/EditMode/SControlRigTweenWidget.h
--- a/Engine/Plugins/Animation/ControlRig/Source/ControlRigEditor/Private/EditMode/SControlRigTweenWidget.h
+++ b/Engine/Plugins/Animation/ControlRig/Source/ControlRigEditor/Private/EditMode/SControlRigTweenWidget.h
@@ -37,6 +37,10 @@ class SControlRigTweenSlider : public SCompoundWidget
 	bool Setup();
 
 	void ResetAnimSlider();
+
+	// SWidget interface
+	virtual void Tick(const FGeometry& InAllottedGeometry, const double InCurrentTime, const float InDeltaTime) override;
+	
 private:
 
 
@@ -58,6 +62,8 @@ private:
 	TSharedPtr<FBaseAnimSlider> AnimSlider;
 	TSharedPtr<SSpinBox<double>> SpinBox;
 
+	// Pending blend function called on tick to avoid blending values for each mouse move
+	TFunction<void()> PendingBlendFunction;
 };
 
 
diff --git a/Engine/Source/Editor/MovieSceneTools/Private/Channels/ChannelCurveModel.cpp b/Engine/Source/Editor/MovieSceneTools/Private/Channels/ChannelCurveModel.cpp
--- a/Engine/Source/Editor/MovieSceneTools/Private/Channels/ChannelCurveModel.cpp
+++ b/Engine/Source/Editor/MovieSceneTools/Private/Channels/ChannelCurveModel.cpp
@@ -223,11 +223,24 @@ void FChannelCurveModel<ChannelType, ChannelValue, KeyType>::SetKeyPositions(TAr
 
 	if (Channel && Section && !IsReadOnly())
 	{
+		TMovieSceneChannelData<ChannelValue> ChannelData = Channel->GetData();
+		const TArrayView<const FFrameNumber> Times = ChannelData.GetTimes();
+		if (Times.IsEmpty())
+		{
+			return;
+		}
+		
 		Section->MarkAsChanged();
 
 		FFrameRate TickResolution = Section->GetTypedOuter<UMovieScene>()->GetTickResolution();
 
-		TMovieSceneChannelData<ChannelValue> ChannelData = Channel->GetData();
+		const int32 LastIndex = Times.Num()-1;
+		const FFrameNumber FirstTime = Times[0];
+		const FFrameNumber LastTime = Times[LastIndex];
+		
+		FFrameNumber Min = TNumericLimits<FFrameNumber>::Max();
+		FFrameNumber Max = TNumericLimits<FFrameNumber>::Lowest();
+		
 		for (int32 Index = 0; Index < InKeys.Num(); ++Index)
 		{
 			int32 KeyIndex = ChannelData.GetIndex(InKeys[Index]);
@@ -238,9 +251,24 @@ void FChannelCurveModel<ChannelType, ChannelValue, KeyType>::SetKeyPositions(TAr
 				KeyIndex = ChannelData.MoveKey(KeyIndex, NewTime);
 				SetKeyValue(KeyIndex, InKeyPositions[Index].OutputValue);
 
-				Section->ExpandToFrame(NewTime);
+				if ((KeyIndex == 0 && NewTime != FirstTime) || (KeyIndex == LastIndex && NewTime != LastTime))
+				{
+					Min = FMath::Min(Min, NewTime);
+					Max = FMath::Max(Max, NewTime);
+				}
+			}
+		}
+
+		// update range if needed
+		if (Min != TNumericLimits<FFrameNumber>::Max())
+		{
+			Section->ExpandToFrame(Min);
+			if (Max != Min)
+			{
+				Section->ExpandToFrame(Max);
 			}
 		}
+		
 		Channel->PostEditChange();
 		if(WeakSequencer.IsValid())
 		{ 
