Sequencer: Fix potential 'check' crash in component material track editor. As material classes implement their own material names, there's not a guarantee material names array matches material slot array, and we must be seeing that here with this material.

diff --git a/Engine/Source/Editor/MovieSceneTools/Private/TrackEditors/MaterialTrackEditor.cpp b/Engine/Source/Editor/MovieSceneTools/Private/TrackEditors/MaterialTrackEditor.cpp
--- a/Engine/Source/Editor/MovieSceneTools/Private/TrackEditors/MaterialTrackEditor.cpp
+++ b/Engine/Source/Editor/MovieSceneTools/Private/TrackEditors/MaterialTrackEditor.cpp
@@ -329,11 +329,13 @@ void FComponentMaterialTrackEditor::BuildTrackContextMenu(FMenuBuilder& MenuBuil
 					if (Material)
 					{
 						// Found material by index, but not by slot name. Auto-rebind would change slot name to the current one
-						FoundMaterialIndex = MaterialInfo.MaterialSlotIndex;
 						TArray<FName> SlotNames = Component->GetMaterialSlotNames();
-						check(SlotNames.IsValidIndex(MaterialInfo.MaterialSlotIndex));
-						FoundSlotName = SlotNames[MaterialInfo.MaterialSlotIndex];
-						AutoRebindTooltip = FText::Format(LOCTEXT("AutoRebindToNewSlotName", "Rebind track to slot {0}, keeping same index {1}"), FText::FromName(FoundSlotName), FText::AsNumber(MaterialInfo.MaterialSlotIndex));
+						if (SlotNames.IsValidIndex(MaterialInfo.MaterialSlotIndex))
+						{
+							FoundMaterialIndex = MaterialInfo.MaterialSlotIndex;
+							FoundSlotName = SlotNames[MaterialInfo.MaterialSlotIndex];
+							AutoRebindTooltip = FText::Format(LOCTEXT("AutoRebindToNewSlotName", "Rebind track to slot {0}, keeping same index {1}"), FText::FromName(FoundSlotName), FText::AsNumber(MaterialInfo.MaterialSlotIndex));
+						}
 					}
 					// If we didn't find a material, we don't create a tooltip, because we won't be able to 'auto' rebind, just manually bind.
 				}
