UE-204136 - X might be missing an import from package Y because of a circular dependency between them

The UK2Node_GetSequenceBinding is a little unusual due to its usage of a FSoftObjectPath for its object reference to the corresponding UMovieSceneSequence. However, the soft reference is actually unnecessary, because we _always_ preload the object reference when its owner Blueprint is loaded. While this isn't necessarily bad, it confuses Zen loader when there's a cyclic dependency between the node and another asset file. The correct fix here is to convert the soft reference to a hard reference.

Additionally, Blueprints and their dependent objects (eg: graphs, nodes, etc) are always preloaded prior to compilation, but not postloaded. This also confuses Zen loader, because it doesn't expect a postload to occur during a preload. Fortunately, we don't actually have to postload UMovieSceneSequence. It's only needed for validating the node's GUID against the UMovieSceneSequence. In this case, its sufficient to preload the UMovieSceneSequence, because we're only checking the GUID against the preloaded Spawnables and Possessables arrays.

Note: content must be resaved with this change to actually fix Zen loader warnings.

diff --git a/Engine/Source/Editor/MovieSceneTools/Private/K2Node_GetSequenceBinding.cpp b/Engine/Source/Editor/MovieSceneTools/Private/K2Node_GetSequenceBinding.cpp
--- a/Engine/Source/Editor/MovieSceneTools/Private/K2Node_GetSequenceBinding.cpp
+++ b/Engine/Source/Editor/MovieSceneTools/Private/K2Node_GetSequenceBinding.cpp
@@ -29,44 +29,42 @@ static const FName SequencePinName(TEXT("Sequence"));
 
 #define LOCTEXT_NAMESPACE "UK2Node_GetSequenceBinding"
 
-void EnsureFullyLoaded(UObject* Object)
+void EnsureFullyPreloaded(UObject* Object)
 {
 	if (!Object)
 	{
 		return;
 	}
 
-	bool bLoadInternalReferences = false;
-
 	if (Object->HasAnyFlags(RF_NeedLoad))
 	{
 		FLinkerLoad* Linker = Object->GetLinker();
 		if (ensure(Linker))
 		{
 			Linker->Preload(Object);
-			bLoadInternalReferences = true;
 			check(!Object->HasAnyFlags(RF_NeedLoad));
 		}
 	}
 
-	bLoadInternalReferences = bLoadInternalReferences || Object->HasAnyFlags(RF_NeedPostLoad | RF_NeedPostLoadSubobjects);
+	check(Object->HasAnyFlags(RF_LoadCompleted));
+
+	TArray<UObject*> ObjectReferences;
+	FReferenceFinder(ObjectReferences, nullptr, false, true, false, true).FindReferences(Object);
 
-	Object->ConditionalPostLoad();
-	Object->ConditionalPostLoadSubobjects();
-	
-	if (bLoadInternalReferences)
+	for (UObject* Reference : ObjectReferences)
 	{
-		// Collect a list of all things this element owns
-		TArray<UObject*> ObjectReferences;
-		FReferenceFinder(ObjectReferences, nullptr, false, true, false, true).FindReferences(Object);
+		check(Reference);
 
-		// Iterate over the list, and preload everything so it is valid for refreshing
-		for (UObject* Reference : ObjectReferences)
+		const bool bIsMovieSceneType =
+			Reference->IsA<UMovieSceneSequence>() ||
+			Reference->IsA<UMovieScene>() ||
+			Reference->IsA<UMovieSceneTrack>() ||
+			Reference->IsA<UMovieSceneSection>()
+		;
+
+		if (bIsMovieSceneType)
 		{
-			if (Reference->IsA<UMovieSceneSequence>() || Reference->IsA<UMovieScene>() || Reference->IsA<UMovieSceneTrack>() || Reference->IsA<UMovieSceneSection>())
-			{
-				EnsureFullyLoaded(Reference);
-			}
+			EnsureFullyPreloaded(Reference);
 		}
 	}
 }
@@ -94,10 +92,26 @@ public:
 	}
 };
 
-UMovieSceneSequence* UK2Node_GetSequenceBinding::GetSequence() const
+void UK2Node_GetSequenceBinding::SetSequence(UMovieSceneSequence* InSequence)
+{
+	SourceMovieSequence = InSequence;
+}
+
+void UK2Node_GetSequenceBinding::Serialize(FArchive& Ar)
 {
-	FLinkerLoad* Linker = GetLinker();
-	return Cast<UMovieSceneSequence>(SourceSequence.TryLoad(Linker ? FUObjectThreadContext::Get().GetSerializeContext() : nullptr));
+	Super::Serialize(Ar);
+
+	const bool bConvertSoftToHardReference =
+		Ar.IsLoading() && 
+		((Ar.GetPortFlags() & PPF_Duplicate) == 0) &&
+		!SourceSequence_DEPRECATED.IsNull()
+	;
+
+	if (bConvertSoftToHardReference)
+	{
+		SourceMovieSequence = Cast<UMovieSceneSequence>(SourceSequence_DEPRECATED.TryLoad());
+		SourceSequence_DEPRECATED.Reset();
+	}
 }
 
 void UK2Node_GetSequenceBinding::ValidateNodeDuringCompilation(FCompilerResultsLog& MessageLog) const
@@ -133,24 +147,20 @@ void UK2Node_GetSequenceBinding::PostPlacedNewNode()
 	// Attempt to assign the sequence asset from our outer if this BP is contained within a sequence
 	if (UMovieSceneSequence* OuterSequence = GetTypedOuter<UMovieSceneSequence>())
 	{
-		SourceSequence = OuterSequence;
+		SourceMovieSequence = OuterSequence;
 	}
 	Super::PostPlacedNewNode();
 }
 
 UMovieScene* UK2Node_GetSequenceBinding::GetObjectMovieScene() const
 {
-	UMovieSceneSequence* Sequence = GetSequence();
-	if (Sequence && Binding.IsValid())
+	if (SourceMovieSequence && Binding.IsValid())
 	{
-		// Ensure that the sequence data is as loaded as it can be - we many only be able to partially load the structural information as part of a blueprint compile as that may happen at Preload time
-		EnsureFullyLoaded(Sequence);
-
 		FMovieSceneSequenceID SequenceID = Binding.GetRelativeSequenceID();
 		if (SequenceID == MovieSceneSequenceID::Root)
 		{
 			// Look it up in the moviescene itself
-			return Sequence->GetMovieScene();
+			return SourceMovieSequence->GetMovieScene();
 		}
 		else
 		{
@@ -186,7 +196,7 @@ UMovieScene* UK2Node_GetSequenceBinding::GetObjectMovieScene() const
 				SequenceSignatureCache.Reset();
 				SequenceHierarchyCache = FMovieSceneSequenceHierarchy();
 
-				UMovieSceneCompiledDataManager::CompileHierarchy(Sequence, &SequenceHierarchyCache, EMovieSceneServerClientMask::All);
+				UMovieSceneCompiledDataManager::CompileHierarchy(SourceMovieSequence, &SequenceHierarchyCache, EMovieSceneServerClientMask::All);
 
 				for (const TTuple<FMovieSceneSequenceID, FMovieSceneSubSequenceData>& Pair : SequenceHierarchyCache.AllSubSequenceData())
 				{
@@ -217,13 +227,12 @@ FNodeHandlingFunctor* UK2Node_GetSequenceBinding::CreateNodeHandler(FKismetCompi
 
 void UK2Node_GetSequenceBinding::PreloadRequiredAssets()
 {
-	EnsureFullyLoaded(GetSequence());
+	EnsureFullyPreloaded(SourceMovieSequence);
 }
 
 FText UK2Node_GetSequenceBinding::GetSequenceName() const
 {
-	UMovieSceneSequence* Sequence = GetSequence();
-	return Sequence ? FText::FromName(Sequence->GetFName()) : LOCTEXT("NoSequence", "No Sequence");
+	return SourceMovieSequence ? FText::FromName(SourceMovieSequence->GetFName()) : LOCTEXT("NoSequence", "No Sequence");
 }
 
 FText UK2Node_GetSequenceBinding::GetBindingName() const
@@ -264,18 +273,16 @@ void UK2Node_GetSequenceBinding::GetNodeContextMenuActions(UToolMenu* Menu, UGra
 		FToolMenuSection& Section = Menu->AddSection("K2NodeGetSequenceBinding", LOCTEXT("ThisNodeHeader", "This Node"));
 		if (!Context->Pin)
 		{
-			UMovieSceneSequence* Sequence = GetSequence();
-
 			Section.AddSubMenu(
 				"SetSequence",
 				LOCTEXT("SetSequence_Text", "Sequence"),
 				LOCTEXT("SetSequence_ToolTip", "Sets the sequence to get a binding from"),
-				FNewToolMenuDelegate::CreateLambda([this, Sequence](UToolMenu* SubMenu)
+				FNewToolMenuDelegate::CreateLambda([this](UToolMenu* SubMenu)
 				{
 					TArray<const UClass*> AllowedClasses({ UMovieSceneSequence::StaticClass() });
 
 					TSharedRef<SWidget> MenuContent = PropertyCustomizationHelpers::MakeAssetPickerWithMenu(
-						FAssetData(Sequence),
+						FAssetData(SourceMovieSequence),
 						true /* bAllowClear */,
 						AllowedClasses,
 						PropertyCustomizationHelpers::GetNewAssetFactoriesForClasses(AllowedClasses),
@@ -296,7 +303,7 @@ void UK2Node_GetSequenceBinding::SetSequence(const FAssetData& InAssetData)
 	const FScopedTransaction Transaction(LOCTEXT("SetSequence", "Set Sequence"));
 	Modify();
 
-	SourceSequence = Cast<UMovieSceneSequence>(InAssetData.GetAsset());
+	SourceMovieSequence = Cast<UMovieSceneSequence>(InAssetData.GetAsset());
 }
 
 void UK2Node_GetSequenceBinding::GetMenuActions(FBlueprintActionDatabaseRegistrar& ActionRegistrar) const
@@ -310,8 +317,6 @@ void UK2Node_GetSequenceBinding::GetMenuActions(FBlueprintActionDatabaseRegistra
 	}
 }
 
-#if WITH_EDITOR
-
 TSharedPtr<SGraphNode> UK2Node_GetSequenceBinding::CreateVisualWidget()
 {
 	class SGraphNodeGetSequenceBinding : public SGraphNode, FMovieSceneObjectBindingIDPicker
@@ -335,9 +340,8 @@ TSharedPtr<SGraphNode> UK2Node_GetSequenceBinding::CreateVisualWidget()
 		virtual void Tick( const FGeometry& AllottedGeometry, const double InCurrentTime, const float InDeltaTime )
 		{
 			UK2Node_GetSequenceBinding* Node = CastChecked<UK2Node_GetSequenceBinding>(GraphNode);
-			UMovieSceneSequence* Sequence = Node->GetSequence();
 
-			if (bNeedsUpdate || Sequence != LastSequence.Get())
+			if (bNeedsUpdate || Node->SourceMovieSequence != LastSequence.Get())
 			{
 				Initialize();
 				UpdateGraphNode();
@@ -345,7 +349,7 @@ TSharedPtr<SGraphNode> UK2Node_GetSequenceBinding::CreateVisualWidget()
 				bNeedsUpdate = false;
 			}
 
-			LastSequence = Sequence;
+			LastSequence = Node->SourceMovieSequence;
 
 			SGraphNode::Tick(AllottedGeometry, InCurrentTime, InDeltaTime);
 		}
@@ -412,7 +416,7 @@ TSharedPtr<SGraphNode> UK2Node_GetSequenceBinding::CreateVisualWidget()
 
 		FReply BrowseToAsset()
 		{
-			UMovieSceneSequence* Sequence = CastChecked<UK2Node_GetSequenceBinding>(GraphNode)->GetSequence();
+			UMovieSceneSequence* Sequence = CastChecked<UK2Node_GetSequenceBinding>(GraphNode)->SourceMovieSequence;
 			if (Sequence)
 			{
 				TArray<UObject*> Objects{ Sequence };
@@ -540,7 +544,7 @@ TSharedPtr<SGraphNode> UK2Node_GetSequenceBinding::CreateVisualWidget()
 		}
 		virtual UMovieSceneSequence* GetSequence() const override
 		{
-			return CastChecked<UK2Node_GetSequenceBinding>(GraphNode)->GetSequence();
+			return CastChecked<UK2Node_GetSequenceBinding>(GraphNode)->SourceMovieSequence;
 		}
 
 		FSlateColor OnGetComboForeground() const
@@ -566,6 +570,4 @@ TSharedPtr<SGraphNode> UK2Node_GetSequenceBinding::CreateVisualWidget()
 	return SNew(SGraphNodeGetSequenceBinding, this);
 }
 
-#endif
-
 #undef LOCTEXT_NAMESPACE
diff --git a/Engine/Source/Editor/MovieSceneTools/Public/K2Node_GetSequenceBinding.h b/Engine/Source/Editor/MovieSceneTools/Public/K2Node_GetSequenceBinding.h
--- a/Engine/Source/Editor/MovieSceneTools/Public/K2Node_GetSequenceBinding.h
+++ b/Engine/Source/Editor/MovieSceneTools/Public/K2Node_GetSequenceBinding.h
@@ -23,18 +23,21 @@ class UK2Node_GetSequenceBinding
 public:
 	GENERATED_BODY()
 
-	/** The sequence from which to choose a binding identifier */
-	UPROPERTY(EditAnywhere, Category="Sequence", meta=(AllowedClasses="/Script/MovieScene.MovieSceneSequence"))
-	FSoftObjectPath SourceSequence;
+	UPROPERTY()
+	FSoftObjectPath SourceSequence_DEPRECATED;
 
 	/** The user-selected literal binding identifier from the sequence to use */
 	UPROPERTY()
 	FMovieSceneObjectBindingID Binding;
 
-	/** Attempt to load the sequence from which to choose a binding */
-	UMovieSceneSequence* GetSequence() const;
+	MOVIESCENETOOLS_API void SetSequence(UMovieSceneSequence* InSequence);
 
 public:
+	// UObject interface
+	virtual void Serialize(FArchive& Ar) override;
+	// End of UObject interface
+
+	// UK2Node interface
 	virtual void ValidateNodeDuringCompilation(FCompilerResultsLog& MessageLog) const override;
 	virtual FText GetNodeTitle(ENodeTitleType::Type TitleType) const override;
 	virtual FText GetTooltipText() const override;
@@ -49,16 +52,20 @@ public:
 	virtual void GetMenuActions(FBlueprintActionDatabaseRegistrar& ActionRegistrar) const override;
 	virtual void AllocateDefaultPins() override;
 	virtual void PostPlacedNewNode() override;
-#if WITH_EDITOR
 	virtual TSharedPtr<SGraphNode> CreateVisualWidget() override;
-#endif
+	// End of UK2Node interface
 
 	void SetSequence(const FAssetData& InAssetData);
 	FText GetSequenceName() const;
 	FText GetBindingName() const;
-	UMovieScene* GetObjectMovieScene() const;
 
 private:
+	UMovieScene* GetObjectMovieScene() const;
+
+	/** The sequence from which to choose a binding identifier */
+	UPROPERTY(EditAnywhere, Category="Sequence", meta=(DisplayName="Source Sequence"))
+	TObjectPtr<UMovieSceneSequence> SourceMovieSequence;
+
 	mutable FMovieSceneSequenceHierarchy SequenceHierarchyCache;
 	mutable TMap<FMovieSceneSequenceID, FGuid> SequenceSignatureCache;
 };
diff --git a/Engine/Source/Editor/Sequencer/Private/SequencerOutlinerItemDragDropOp.cpp b/Engine/Source/Editor/Sequencer/Private/SequencerOutlinerItemDragDropOp.cpp
--- a/Engine/Source/Editor/Sequencer/Private/SequencerOutlinerItemDragDropOp.cpp
+++ b/Engine/Source/Editor/Sequencer/Private/SequencerOutlinerItemDragDropOp.cpp
@@ -119,7 +119,7 @@ FReply FSequencerOutlinerDragDropOp::DroppedOnPanel( const TSharedRef< class SWi
 			if (SequenceModel)
 			{
 				// Fixed bindings always resolve from the root
-				Template->SourceSequence = SequenceModel->GetSequencer()->GetRootMovieSceneSequence();
+				Template->SetSequence(SequenceModel->GetSequencer()->GetRootMovieSceneSequence());
 				Template->Binding = FFixedObjectBindingID(ObjectBinding->GetObjectGuid(), SequenceModel->GetSequenceID());
 				UEdGraphNode* NewNode = Action.PerformAction(&Graph, GetHoveredPin(), GraphPosition, false);
 
