Sequencer: Fix a crash with the binding lifetime track sections where you could get in a situation where a section split was returning nullptr.
To fix this, I've decided to simplify the section adding to either just split if you're on an existing section or else add a new section that fits exactly in the prescribed gap. This means you only get 1 new section for each click of '+' as well which I feel makes more sense.

diff --git a/Engine/Source/Editor/MovieSceneTools/Private/TrackEditors/BindingLifetimeTrackEditor.cpp b/Engine/Source/Editor/MovieSceneTools/Private/TrackEditors/BindingLifetimeTrackEditor.cpp
--- a/Engine/Source/Editor/MovieSceneTools/Private/TrackEditors/BindingLifetimeTrackEditor.cpp
+++ b/Engine/Source/Editor/MovieSceneTools/Private/TrackEditors/BindingLifetimeTrackEditor.cpp
@@ -57,34 +57,99 @@ void FBindingLifetimeTrackEditor::CreateNewSection(UMovieSceneTrack* Track, bool
 		FScopedTransaction Transaction(LOCTEXT("CreateNewSectionTransactionText", "Add Section"));
 
 
-		UMovieSceneSection* NewSection = NewObject<UMovieSceneSection>(Track, UMovieSceneBindingLifetimeSection::StaticClass(), NAME_None, RF_Transactional);
-		check(NewSection);
+		UMovieSceneSection* NewSection = nullptr;
 
 		Track->Modify();
 
 		// first section by default should be infinite
 		if (Track->GetAllSections().Num() == 0)
 		{
+			NewSection = NewObject<UMovieSceneSection>(Track, UMovieSceneBindingLifetimeSection::StaticClass(), NAME_None, RF_Transactional);
 			NewSection->SetRange(TRange<FFrameNumber>::All());
+			Track->AddSection(*NewSection);
 		}
 		else
 		{
-			int32 Duration = 0;
-
-			if (CurrentTime.Time.FrameNumber < FocusedMovieScene->GetPlaybackRange().GetUpperBoundValue())
+			// If the start time overlaps an existing section, split that section at the start time
+			TArray<UMovieSceneSection*> ExistingSections = Track->GetAllSections();
+			for(UMovieSceneSection* Section : ExistingSections)
 			{
-				Duration = FocusedMovieScene->GetPlaybackRange().GetUpperBoundValue().Value - CurrentTime.Time.FrameNumber.Value;
+				TRange<FFrameNumber> SectionRange = Section->GetRange();
+				if (SectionRange.Contains(CurrentTime.Time.FrameNumber))
+				{
+					// Edge case- the section the start is overlapping we're just at the start of it (can happen when adding 2 sections back to back at the same time)
+					// In that case, push back the start of the other section
+					if (SectionRange.GetLowerBound().IsClosed() && SectionRange.GetLowerBoundValue() == CurrentTime.Time.FrameNumber)
+					{
+						FFrameNumber AdjustmentFrames = SequencerPtr->GetFocusedTickResolution().AsFrameNumber(1);
+						if (SectionRange.HasUpperBound())
+						{
+							AdjustmentFrames = (SectionRange.GetUpperBoundValue() - SectionRange.GetLowerBoundValue()) / 2;
+						}
+
+						SectionRange.SetLowerBoundValue(SectionRange.GetLowerBoundValue() + AdjustmentFrames);
+						Section->SetRange(SectionRange);
+
+						NewSection = NewObject<UMovieSceneSection>(Track, UMovieSceneBindingLifetimeSection::StaticClass(), NAME_None, RF_Transactional);
+						TRange<FFrameNumber> NewSectionRange;
+						NewSectionRange.SetLowerBound(TRangeBound<FFrameNumber>(CurrentTime.Time.FrameNumber));
+						TRangeBound<FFrameNumber> NewUpperBound(SectionRange.GetLowerBoundValue());
+						if (NewUpperBound.IsInclusive())
+						{
+							NewUpperBound = TRangeBound<FFrameNumber>::FlipInclusion(NewUpperBound);
+						}
+						NewSectionRange.SetUpperBound(NewUpperBound);
+						NewSection->SetRange(NewSectionRange);
+						Track->AddSection(*NewSection);
+					}
+					else
+					{
+						// Splitting adds the section to the track
+						NewSection = Section->SplitSection(CurrentTime, false);
+					}
+					break;
+				}
 			}
-			else
+
+			// If we didn't overlap anything, add a new section starting at the frame time and ending either at the next section start, or if not existing, make the duration infinite
+			if (!NewSection)
 			{
-				const float DefaultLengthInSeconds = 5.f;
-				Duration = (DefaultLengthInSeconds * SequencerPtr->GetFocusedTickResolution()).FloorToFrame().Value;
+				NewSection = NewObject<UMovieSceneSection>(Track, UMovieSceneBindingLifetimeSection::StaticClass(), NAME_None, RF_Transactional);
+				TRange<FFrameNumber> NewSectionRange;
+				NewSectionRange.SetLowerBound(TRangeBound<FFrameNumber>(CurrentTime.Time.FrameNumber));
+				// By default, set the upper bound to open
+				NewSectionRange.SetUpperBound(TRangeBound<FFrameNumber>::Open());
+
+				UMovieSceneSection* NextSection = nullptr;
+				TRangeBound<FFrameNumber> NextSectionLowerBound = TRangeBound<FFrameNumber>::Open();
+				for (UMovieSceneSection* Section : ExistingSections)
+				{
+					if (Section->GetRange().HasLowerBound() && Section->GetRange().GetLowerBoundValue() > CurrentTime.Time.FrameNumber)
+					{
+						if (NextSectionLowerBound.IsOpen() || Section->GetRange().GetLowerBoundValue() < NextSectionLowerBound.GetValue())
+						{
+							NextSection = Section;
+							NextSectionLowerBound = Section->GetRange().GetLowerBound();
+						}
+					}
+				}
+
+				if (NextSection)
+				{
+					// We found an existing section with a lower bound greater than time, so cap our upper bound value to its lower bound
+					TRangeBound<FFrameNumber> NewUpperBound(NextSectionLowerBound.GetValue());
+					if (NewUpperBound.IsInclusive())
+					{
+						NewUpperBound = TRangeBound<FFrameNumber>::FlipInclusion(NewUpperBound);
+					}
+					NewSectionRange.SetUpperBound(NewUpperBound);
+				}
+
+				NewSection->SetRange(NewSectionRange);
+				Track->AddSection(*NewSection);
 			}
-
-			NewSection->InitialPlacement(Track->GetAllSections(), CurrentTime.Time.FrameNumber.Value, Duration, false);
 		}
 
-		Track->AddSection(*NewSection);
 		Track->UpdateEasing();
 
 		if (bSelect)
diff --git a/Engine/Source/Runtime/MovieScene/Private/Sections/MovieSceneBindingLifetimeSection.cpp b/Engine/Source/Runtime/MovieScene/Private/Sections/MovieSceneBindingLifetimeSection.cpp
--- a/Engine/Source/Runtime/MovieScene/Private/Sections/MovieSceneBindingLifetimeSection.cpp
+++ b/Engine/Source/Runtime/MovieScene/Private/Sections/MovieSceneBindingLifetimeSection.cpp
@@ -64,48 +64,3 @@ void UMovieSceneBindingLifetimeSection::ExternalPopulateEvaluationField(const TR
 	OutFieldBuilder->AddPersistentEntity(EffectiveRange, EntityIndex, MetaDataIndex);
 }
 
-void UMovieSceneBindingLifetimeSection::InitialPlacement(const TArray<UMovieSceneSection*>& Sections, FFrameNumber InStartTime, int32 Duration, bool bAllowMultipleRows)
-{
-	// We don't allow overlap or multi row here, so need a few extra rules for initial placement
-
-	check(Duration >= 0);
-
-	// Inclusive lower, exclusive upper bounds
-	SectionRange = TRange<FFrameNumber>(InStartTime, InStartTime + Duration);
-
-	for (;;)
-	{
-		UMovieSceneSection* OverlappedSection = const_cast<UMovieSceneSection*>(OverlapsWithSections(Sections));
-		if (OverlappedSection == nullptr)
-		{
-			break;
-		}
-
-		TRange<FFrameNumber> OtherRange = OverlappedSection->GetRange();
-
-		if (OtherRange.GetUpperBound().IsClosed())
-		{
-			MoveSection(OtherRange.GetUpperBoundValue() - InStartTime);
-		}
-		else
-		{
-			// Edge case- the other range has an infinite end, but can't be split because the point it would be split at is its start.
-			if (OtherRange.GetLowerBound().IsClosed() && OtherRange.GetLowerBoundValue() == SectionRange.GetLowerBound().GetValue())
-			{
-				// Adjust the lowerbound of OtherRange
-				OtherRange.SetLowerBoundValue(SectionRange.GetUpperBound().GetValue());
-				OverlappedSection->SetRange(OtherRange);
-			}
-			else
-			{
-				// We're overlapping an at-least partially infinite section here. Split it.
-				UMovieSceneSection* NewSection = OverlappedSection->SplitSection(FQualifiedFrameTime(SectionRange.GetLowerBound().GetValue(), GetTypedOuter<UMovieScene>()->GetTickResolution()), false);
-				// The new section will be the one on the right and will still have an open upper bound. Resize it to fit the new section.
-				TRange<FFrameNumber> NewSectionRange = NewSection->GetRange();
-				NewSectionRange.SetLowerBoundValue(SectionRange.GetUpperBound().GetValue());
-				NewSection->SetRange(NewSectionRange);
-			}
-		}
-	}
-}
-
diff --git a/Engine/Source/Runtime/MovieScene/Public/Sections/MovieSceneBindingLifetimeSection.h b/Engine/Source/Runtime/MovieScene/Public/Sections/MovieSceneBindingLifetimeSection.h
--- a/Engine/Source/Runtime/MovieScene/Public/Sections/MovieSceneBindingLifetimeSection.h
+++ b/Engine/Source/Runtime/MovieScene/Public/Sections/MovieSceneBindingLifetimeSection.h
@@ -23,8 +23,6 @@ public:
 	virtual bool PopulateEvaluationFieldImpl(const TRange<FFrameNumber>& EffectiveRange, const FMovieSceneEvaluationFieldEntityMetaData& InMetaData, FMovieSceneEntityComponentFieldBuilder* OutFieldBuilder) override;
 	virtual void ImportEntityImpl(UMovieSceneEntitySystemLinker* EntityLinker, const FEntityImportParams& Params, FImportedEntity* OutImportedEntity) override;
 	void ExternalPopulateEvaluationField(const TRange<FFrameNumber>& EffectiveRange, const FMovieSceneEvaluationFieldEntityMetaData& InMetaData, FMovieSceneEntityComponentFieldBuilder* OutFieldBuilder);
-
-	void InitialPlacement(const TArray<UMovieSceneSection*>& Sections, FFrameNumber InStartTime, int32 InDuration, bool bAllowMultipleRows) override;
 };
 
 
