Sequencer: Expand Meta UPROPERTY tags 'MetaClass', 'MustImplement', 'AllowedClasses', and 'DisallowedClasses' to work on UPROPERTY of type FMovieSceneObjectBindingID via the MovieSceneObjectBindingPicker, using the bound object class of the possessable or spawnable.

Note that this will not prevent objects of other classes to be bound to these bindings in a programmatic way or via dynamic replacement at runtime.

This change is a GitHub pull request from GitHub user laycnc.

#12047

diff --git a/Engine/Source/Editor/MovieSceneTools/Private/MovieSceneObjectBindingIDCustomization.cpp b/Engine/Source/Editor/MovieSceneTools/Private/MovieSceneObjectBindingIDCustomization.cpp
--- a/Engine/Source/Editor/MovieSceneTools/Private/MovieSceneObjectBindingIDCustomization.cpp
+++ b/Engine/Source/Editor/MovieSceneTools/Private/MovieSceneObjectBindingIDCustomization.cpp
@@ -30,6 +30,9 @@
 #include "Widgets/Input/SComboButton.h"
 #include "Widgets/SBoxPanel.h"
 #include "Widgets/Text/STextBlock.h"
+#include "PropertyCustomizationHelpers.h"
+#include "EditorClassUtils.h"
+
 
 struct FGeometry;
 
@@ -61,6 +64,17 @@ void FMovieSceneObjectBindingIDCustomization::CustomizeHeader(TSharedRef<IProper
 
 	StructProperty = PropertyHandle;
 
+	const FString& MetaClassName = PropertyHandle->GetMetaData("MetaClass");
+	const FString& MustImplementName = PropertyHandle->GetMetaData("MustImplement");
+
+	AllowedClassFilters = PropertyCustomizationHelpers::GetClassesFromMetadataString(PropertyHandle->GetMetaData("AllowedClasses"));
+	DisallowedClassFilters = PropertyCustomizationHelpers::GetClassesFromMetadataString(PropertyHandle->GetMetaData("DisallowedClasses"));
+
+	ClassPropertyMetaClass = !MetaClassName.IsEmpty()
+		? FEditorClassUtils::GetClassFromString(MetaClassName)
+		: UObject::StaticClass();
+	InterfaceThatMustBeImplemented = FEditorClassUtils::GetClassFromString(MustImplementName);
+
 	Initialize();
 
 	auto IsAcceptable = [](TSharedPtr<FDragDropOperation> Operation)
diff --git a/Engine/Source/Editor/MovieSceneTools/Private/MovieSceneObjectBindingIDPicker.cpp b/Engine/Source/Editor/MovieSceneTools/Private/MovieSceneObjectBindingIDPicker.cpp
--- a/Engine/Source/Editor/MovieSceneTools/Private/MovieSceneObjectBindingIDPicker.cpp
+++ b/Engine/Source/Editor/MovieSceneTools/Private/MovieSceneObjectBindingIDPicker.cpp
@@ -34,6 +34,9 @@
 #include "Widgets/SBoxPanel.h"
 #include "Widgets/SOverlay.h"
 #include "Widgets/Text/STextBlock.h"
+#include "EditorClassUtils.h"
+#include "MovieScene.h"
+#include "MovieSceneCommonHelpers.h"
 
 struct FMovieSceneSequenceHierarchy;
 
@@ -68,7 +71,7 @@ void FMovieSceneObjectBindingIDPicker::OnGetMenuContent(FMenuBuilder& MenuBuilde
 
 	bool bHadAnyEntries = false;
 
-	if (Node->BindingID.Guid.IsValid())
+	if (Node->BindingID.Guid.IsValid() && IsClassAllowed(Node))
 	{
 		bHadAnyEntries = true;
 		MenuBuilder.AddMenuEntry(
@@ -100,7 +103,7 @@ void FMovieSceneObjectBindingIDPicker::OnGetMenuContent(FMenuBuilder& MenuBuilde
 					);
 			}
 		}
-		else
+		else if (IsClassAllowed(Child))
 		{
 			bHadAnyEntries = true;
 			MenuBuilder.AddMenuEntry(
@@ -210,6 +213,42 @@ FReply FMovieSceneObjectBindingIDPicker::AttemptBindingFixup()
 	return FReply::Handled();
 }
 
+bool FMovieSceneObjectBindingIDPicker::IsClassAllowed(const TSharedPtr<FSequenceBindingNode>& Node) const
+{
+	const UClass* InClass = nullptr;
+	auto BindingID = Node->BindingID;
+	UMovieSceneSequence* BindSequence = WeakSequencer.Pin()->State.FindSequence(BindingID.SequenceID);
+	UMovieScene* BindMovieScene = BindSequence->GetMovieScene();
+
+	InClass = MovieSceneHelpers::GetBoundObjectClass(BindSequence, BindingID.Guid);
+
+	if (InClass == nullptr)
+	{
+		return true;
+	}
+
+	// Abstract types and deprecations are excluded
+	bool bMatchesFlags = !InClass->HasAnyClassFlags(CLASS_Abstract | CLASS_Deprecated);
+
+	if (bMatchesFlags && InClass->IsChildOf(ClassPropertyMetaClass)
+		&& (!InterfaceThatMustBeImplemented || InClass->ImplementsInterface(InterfaceThatMustBeImplemented)))
+	{
+		auto PredicateFn = [InClass](const UClass* Class)
+		{
+			return InClass->IsChildOf(Class);
+		};
+
+		if (DisallowedClassFilters.FindByPredicate(PredicateFn) == nullptr &&
+			(AllowedClassFilters.Num() == 0 || AllowedClassFilters.FindByPredicate(PredicateFn) != nullptr))
+		{
+			return true;
+		}
+	}
+
+	return false;
+
+}
+
 void FMovieSceneObjectBindingIDPicker::SetBindingId(UE::MovieScene::FFixedObjectBindingID InBindingId)
 {
 	SetCurrentValueFromFixed(InBindingId);
diff --git a/Engine/Source/Editor/MovieSceneTools/Public/MovieSceneObjectBindingIDPicker.h b/Engine/Source/Editor/MovieSceneTools/Public/MovieSceneObjectBindingIDPicker.h
--- a/Engine/Source/Editor/MovieSceneTools/Public/MovieSceneObjectBindingIDPicker.h
+++ b/Engine/Source/Editor/MovieSceneTools/Public/MovieSceneObjectBindingIDPicker.h
@@ -133,6 +133,9 @@ private:
 	/** Get the visibility for the warning/fixup button */
 	FReply AttemptBindingFixup();
 
+	/** Checks if a class is allowed by this binding */
+	bool IsClassAllowed(const TSharedPtr<FSequenceBindingNode>& Node) const;
+
 	/** Cached current text and tooltips */
 	FText CurrentText, ToolTipText;
 
@@ -147,4 +150,17 @@ private:
 
 	/** Weak ptr to a widget used to dismiss menus to */
 	TWeakPtr<SWidget> DismissWidget;
+
+protected:
+	/** The meta class for the property that classes must be a child-of. */
+	const UClass* ClassPropertyMetaClass = UObject::StaticClass();
+
+	/** The interface that must be implemented. */
+	const UClass* InterfaceThatMustBeImplemented = nullptr;
+
+	/** Classes that can be picked */
+	TArray<const UClass*> AllowedClassFilters;
+
+	/** Classes that can't be picked */
+	TArray<const UClass*> DisallowedClassFilters;
 };
