Sequencer: Added time-warp capabilities to sequences, sub-sections and skeletal animation sections

diff --git a/Engine/Source/Editor/CurveEditor/Private/SCurveEditorView.cpp b/Engine/Source/Editor/CurveEditor/Private/SCurveEditorView.cpp
--- a/Engine/Source/Editor/CurveEditor/Private/SCurveEditorView.cpp
+++ b/Engine/Source/Editor/CurveEditor/Private/SCurveEditorView.cpp
@@ -6,6 +6,7 @@
 #include "CurveDataAbstraction.h"
 #include "CurveDrawInfo.h"
 #include "CurveEditor.h"
+#include "CurveEditorAxis.h"
 #include "CurveEditorHelpers.h"
 #include "CurveEditorScreenSpace.h"
 #include "CurveEditorSelection.h"
@@ -31,20 +32,67 @@ SCurveEditorView::SCurveEditorView()
 	, bFixedOutputBounds(0)
 	, bAutoSize(1)
 	, bAllowEmpty(0)
+	, bAllowModelViewTransforms(1)
+	, bUpdateModelViewTransforms(0)
+	, bNeedsDefaultGridLinesH(1)
+	, bNeedsDefaultGridLinesV(1)
 {
 	CurveCacheFlags = ECurveCacheFlags::All;
 	CachedValues.CachedActiveCurvesSerialNumber = 0xFFFFFFFF;
 	CachedValues.CachedSelectionSerialNumber = 0xFFFFFFFF;
-	CachedValues.CachedGeometrySize.X = -1.;
+	CachedValues.CachedGeometrySize.X = -1;
 	CachedValues.CachedGeometrySize.Y = -1;
 }
 
+SCurveEditorView::~SCurveEditorView()
+{
+}
+
 FVector2D SCurveEditorView::ComputeDesiredSize(float LayoutScaleMultiplier) const
 {
 	FVector2D ContentDesiredSize = SCompoundWidget::ComputeDesiredSize(LayoutScaleMultiplier);
 	return FVector2D(ContentDesiredSize.X, FixedHeight.Get(ContentDesiredSize.Y));
 }
 
+void SCurveEditorView::Tick(const FGeometry& AllottedGeometry, const double InCurrentTime, const float InDeltaTime)
+{
+	SCompoundWidget::Tick(AllottedGeometry, InCurrentTime, InDeltaTime);
+
+	if (bAllowModelViewTransforms)
+	{
+		UpdateCurveViewTransformsFromModels();
+	}
+}
+
+void SCurveEditorView::UpdateCurveViewTransformsFromModels()
+{
+	TSharedPtr<FCurveEditor> CurveEditor = WeakCurveEditor.Pin();
+
+	if (!CurveEditor)
+	{
+		return;
+	}
+
+	bool bTransformChanged = false;
+	for (TPair<FCurveModelID, FCurveInfo>& CurveInfo : CurveInfoByID)
+	{
+		const FCurveModel* CurveModel = CurveEditor->FindCurve(CurveInfo.Key);
+		if (CurveModel)
+		{
+			FTransform2d PreviousTransform = CurveInfo.Value.ViewToCurveTransform;
+
+			CurveInfo.Value.ViewToCurveTransform = CurveModel->GetCurveTransform();
+
+			bTransformChanged |= (CurveInfo.Value.ViewToCurveTransform != PreviousTransform);
+		}
+	}
+
+	if (bTransformChanged)
+	{
+		CurveCacheFlags = ECurveCacheFlags::All;
+	}
+}
+
 void SCurveEditorView::GetInputBounds(double& OutInputMin, double& OutInputMax) const
 {
 	TSharedPtr<FCurveEditor> CurveEditor = WeakCurveEditor.Pin();
@@ -81,10 +129,96 @@ FCurveEditorScreenSpace SCurveEditorView::GetViewSpace() const
 	return FCurveEditorScreenSpace(GetCachedGeometry().GetLocalSize(), InputMin, InputMax, OutputMin, OutputMax);
 }
 
+FCurveEditorScreenSpace SCurveEditorView::GetCurveSpace(FCurveModelID CurveID) const
+{
+	const_cast<SCurveEditorView*>(this)->UpdateCustomAxes();
+
+	FCurveInfo CurveInfo = CurveInfoByID.FindRef(CurveID);
+
+	FTransform2d CurveTransform = CurveInfo.ViewToCurveTransform;
+	FScale2d     AxisScale(0.0, 0.0);
+
+	double InputMin = 0.0, InputMax = 1.0;
+	GetInputBounds(InputMin, InputMax);
+
+	double AxisOutputMin = OutputMin;
+	double AxisOutputMax = OutputMax;
+	if (CurveInfo.HorizontalAxis)
+	{
+		InputMin = CustomHorizontalAxes[CurveInfo.HorizontalAxis.Index].Min;
+		InputMax = CustomHorizontalAxes[CurveInfo.HorizontalAxis.Index].Max;
+	}
+	if (CurveInfo.VerticalAxis)
+	{
+		AxisOutputMin = CustomVerticalAxes[CurveInfo.VerticalAxis.Index].Min;
+		AxisOutputMax = CustomVerticalAxes[CurveInfo.VerticalAxis.Index].Max;
+	}
+
+	return FCurveEditorScreenSpace(GetCachedGeometry().GetLocalSize(), InputMin, InputMax, AxisOutputMin, AxisOutputMax).ToCurveSpace(CurveTransform);
+}
+
+FCurveEditorScreenSpaceH SCurveEditorView::GetHorizontalAxisSpace(FCurveEditorViewAxisID ID) const
+{
+	double Min = 0.0;
+	double Max = 1.0;
+
+	if (ID)
+	{
+		Min = CustomHorizontalAxes[ID].Min;
+		Max = CustomHorizontalAxes[ID].Max;
+	}
+	else
+	{
+		GetInputBounds(Min, Max);
+	}
+
+	return FCurveEditorScreenSpaceH(GetCachedGeometry().GetLocalSize().X, Min, Max);
+}
+
+FCurveEditorScreenSpaceV SCurveEditorView::GetVerticalAxisSpace(FCurveEditorViewAxisID ID) const
+{
+	double Min = OutputMin;
+	double Max = OutputMax;
+
+	if (ID)
+	{
+		Min = CustomVerticalAxes[ID].Min;
+		Max = CustomVerticalAxes[ID].Max;
+	}
+	return FCurveEditorScreenSpaceV(GetCachedGeometry().GetLocalSize().X, Min, Max);
+}
+
+FCurveEditorViewAxisID SCurveEditorView::GetAxisForCurve(FCurveModelID CurveID, ECurveEditorAxisOrientation Axis) const
+{
+	const_cast<SCurveEditorView*>(this)->UpdateCustomAxes();
+
+	FCurveInfo CurveInfo = CurveInfoByID.FindRef(CurveID);
+
+	return (Axis == ECurveEditorAxisOrientation::Horizontal)
+		? CurveInfo.HorizontalAxis
+		: CurveInfo.VerticalAxis;
+}
+
+TSharedPtr<FCurveEditorAxis> SCurveEditorView::GetAxis(FCurveEditorViewAxisID ID, ECurveEditorAxisOrientation Axis) const
+{
+	if (!ID)
+	{
+		return nullptr;
+	}
+
+	return (Axis == ECurveEditorAxisOrientation::Horizontal)
+		? CustomHorizontalAxes[ID.Index].Axis
+		: CustomVerticalAxes[ID.Index].Axis;
+}
+
 void SCurveEditorView::AddCurve(FCurveModelID CurveID)
 {
 	CurveInfoByID.Add(CurveID, FCurveInfo{CurveInfoByID.Num()});
 	OnCurveListChanged();
+	if (bAllowModelViewTransforms)
+	{
+		bUpdateModelViewTransforms = true;
+	}
 }
 
 void SCurveEditorView::RemoveCurve(FCurveModelID CurveID)
@@ -105,21 +239,262 @@ void SCurveEditorView::RemoveCurve(FCurveModelID CurveID)
 		}
 
 		OnCurveListChanged();
+		if (bAllowModelViewTransforms)
+		{
+			bUpdateModelViewTransforms = true;
+		}
+	}
+}
+
+void SCurveEditorView::UpdateCustomAxes()
+{
+	if (!bUpdateModelViewTransforms)
+	{
+		return;
+	}
+
+	TSharedPtr<FCurveEditor> CurveEditor = WeakCurveEditor.Pin();
+	check(CurveEditor.IsValid());
+
+	bUpdateModelViewTransforms = false;
+
+	bool bHasAnyCustomAxesH = false;
+	bool bHasAnyCustomAxesV = false;
+
+	bool bHasAnyDefaultAxesH = false;
+	bool bHasAnyDefaultAxesV = false;
+
+	TSortedMap<int32, int32> OldHorizontalToNew;
+
+	struct FWorkingAxisInfo
+	{
+		FAxisInfo Info;
+		FCurveEditorViewAxisID ID;
+	};
+	TMap<TSharedPtr<FCurveEditorAxis>, FCurveEditorViewAxisID> NewCustomHorizontalAxes;
+	TMap<TSharedPtr<FCurveEditorAxis>, FCurveEditorViewAxisID> NewCustomVerticalAxes;
+
+	TArray<FWorkingAxisInfo> HorizontalAxisIdMap, VerticalAxisIdMap;
+
+
+	// Populate existing axes with 0 use-counts to preserve scale
+	for (const FAxisInfo& Info : CustomHorizontalAxes)
+	{
+		FWorkingAxisInfo WorkingInfo;
+		WorkingInfo.Info = Info;
+		WorkingInfo.ID = HorizontalAxisIdMap.Num();
+		NewCustomHorizontalAxes.Add(Info.Axis, WorkingInfo.ID);
+		HorizontalAxisIdMap.Add(WorkingInfo);
+	}
+	for (const FAxisInfo& Info : CustomVerticalAxes)
+	{
+		FWorkingAxisInfo WorkingInfo;
+		WorkingInfo.Info = Info;
+		WorkingInfo.ID = VerticalAxisIdMap.Num();
+		NewCustomVerticalAxes.Add(Info.Axis, WorkingInfo.ID);
+		VerticalAxisIdMap.Add(WorkingInfo);
+	}
+
+
+	// Iterate all curves and allocate axes
+	for (TTuple<FCurveModelID, FCurveInfo>& Pair : CurveInfoByID)
+	{
+		FCurveModel* CurveModel = CurveEditor->FindCurve(Pair.Key);
+		FCurveEditorViewAxisID HorizontalAxisID, VerticalAxisID;
+
+		if (ensureAlways(CurveModel))
+		{
+			TSharedPtr<FCurveEditorAxis> HorizontalAxis = nullptr;
+			TSharedPtr<FCurveEditorAxis> VerticalAxis = nullptr;
+
+			CurveModel->AllocateAxes(CurveEditor.Get(), HorizontalAxis, VerticalAxis);
+
+			// Horizontal Axis
+			if (HorizontalAxis)
+			{
+				HorizontalAxisID = NewCustomHorizontalAxes.FindRef(HorizontalAxis);
+				if (!HorizontalAxisID)
+				{
+					HorizontalAxisID = HorizontalAxisIdMap.Num();
+					HorizontalAxisIdMap.Add(FWorkingAxisInfo{ FAxisInfo{HorizontalAxis}, HorizontalAxisID });
+					NewCustomHorizontalAxes.Add(HorizontalAxis, HorizontalAxisID);
+				}
+
+				HorizontalAxisIdMap[HorizontalAxisID.Index].Info.UseCount += 1;
+				bHasAnyCustomAxesH = true;
+			}
+			else
+			{
+				bHasAnyDefaultAxesH = true;
+			}
+
+			// Vertical Axis
+			if (VerticalAxis)
+			{
+				VerticalAxisID = NewCustomVerticalAxes.FindRef(VerticalAxis);
+				if (!VerticalAxisID)
+				{
+					VerticalAxisID = VerticalAxisIdMap.Num();
+					VerticalAxisIdMap.Add(FWorkingAxisInfo{ FAxisInfo{VerticalAxis}, VerticalAxisID });
+					NewCustomVerticalAxes.Add(VerticalAxis, VerticalAxisID);
+				}
+
+				VerticalAxisIdMap[VerticalAxisID.Index].Info.UseCount += 1;
+				bHasAnyCustomAxesV = true;
+			}
+			else
+			{
+				bHasAnyDefaultAxesV = true;
+			}
+		}
+
+		Pair.Value.HorizontalAxis = HorizontalAxisID;
+		Pair.Value.VerticalAxis   = VerticalAxisID;
+	}
+
+	CustomHorizontalAxes.SetNum(HorizontalAxisIdMap.Num());
+	CustomVerticalAxes.SetNum(VerticalAxisIdMap.Num());
+
+	// Sort the IDs by use-count descending
+	Algo::SortBy(HorizontalAxisIdMap, [](const FWorkingAxisInfo& In){ return In.Info.UseCount; }, TGreater<>());
+	Algo::SortBy(VerticalAxisIdMap,   [](const FWorkingAxisInfo& In){ return In.Info.UseCount; }, TGreater<>());
+
+	// Assign new axis IDs
+	TArray<FCurveEditorViewAxisID> ReverseHorizontalLookup, ReverseVerticalLookup;
+
+	ReverseHorizontalLookup.SetNum(HorizontalAxisIdMap.Num());
+	for (int32 Index = 0; Index < HorizontalAxisIdMap.Num(); ++Index)
+	{
+		FCurveEditorViewAxisID OldId = HorizontalAxisIdMap[Index].ID;
+		ReverseHorizontalLookup[OldId.Index] = Index;
+		CustomHorizontalAxes[Index] = HorizontalAxisIdMap[Index].Info;
+	}
+
+	ReverseVerticalLookup.SetNum(VerticalAxisIdMap.Num());
+	for (int32 Index = 0; Index < VerticalAxisIdMap.Num(); ++Index)
+	{
+		FCurveEditorViewAxisID OldId = VerticalAxisIdMap[Index].ID;
+		ReverseVerticalLookup[OldId.Index] = Index;
+		CustomVerticalAxes[Index] = VerticalAxisIdMap[Index].Info;
+	}
+
+	for (TTuple<FCurveModelID, FCurveInfo>& Pair : CurveInfoByID)
+	{
+		if (Pair.Value.HorizontalAxis)
+		{
+			Pair.Value.HorizontalAxis = ReverseHorizontalLookup[Pair.Value.HorizontalAxis.Index];
+		}
+		if (Pair.Value.VerticalAxis)
+		{
+			Pair.Value.VerticalAxis = ReverseVerticalLookup[Pair.Value.VerticalAxis.Index];
+		}
+	}
+
+
+	bNeedsDefaultGridLinesH = bHasAnyDefaultAxesH;
+	bNeedsDefaultGridLinesV = bHasAnyDefaultAxesV;
+}
+
+void SCurveEditorView::FrameVertical(double InOutputMin, double InOutputMax, FCurveEditorViewAxisID AxisID)
+{
+	if (!bFixedOutputBounds && InOutputMin < InOutputMax)
+	{
+		if (AxisID)
+		{
+			FAxisInfo& AxisInfo = CustomVerticalAxes[AxisID];
+			AxisInfo.Min = InOutputMin;
+			AxisInfo.Max = InOutputMax;
+		}
+		else if (TSharedPtr<FCurveEditor> CurveEditor = WeakCurveEditor.Pin())
+		{
+			OutputMin = InOutputMin;
+			OutputMax = InOutputMax;
+		}
 	}
 }
 
-void SCurveEditorView::FrameVertical(double InOutputMin, double InOutputMax)
+void SCurveEditorView::FrameHorizontal(double InInputMin, double InInputMax, FCurveEditorViewAxisID AxisID)
 {
-	//default just set's output if it can
-	SetOutputBounds(InOutputMin, InOutputMax);
+	if (InInputMin < InInputMax)
+	{
+		if (AxisID)
+		{
+			FAxisInfo& AxisInfo = CustomHorizontalAxes[AxisID];
+			AxisInfo.Min = InInputMin;
+			AxisInfo.Max = InInputMax;
+		}
+		else if (TSharedPtr<FCurveEditor> CurveEditor = WeakCurveEditor.Pin())
+		{
+			CurveEditor->GetBounds().SetInputBounds(InInputMin, InInputMax);
+		}
+	}
 }
 
-void SCurveEditorView::SetOutputBounds(double InOutputMin, double InOutputMax)
+void SCurveEditorView::SetOutputBounds(double InOutputMin, double InOutputMax, FCurveEditorViewAxisID AxisID)
 {
 	if (!bFixedOutputBounds && InOutputMin < InOutputMax)
 	{
-		OutputMin = InOutputMin;
-		OutputMax = InOutputMax;
+		if (AxisID)
+		{
+			FAxisInfo& AxisInfo = CustomVerticalAxes[AxisID];
+			AxisInfo.Min = InOutputMin;
+			AxisInfo.Max = InOutputMax;
+		}
+
+		// When no axis ID is specified, we scale all axes based on the change
+		else if (TSharedPtr<FCurveEditor> CurveEditor = WeakCurveEditor.Pin())
+		{
+			double CurrentRange = (OutputMax - OutputMin);
+			double OffsetFactor = (InOutputMin - OutputMin) / CurrentRange;
+			double Scale        = (InOutputMax - InOutputMin) / CurrentRange;
+
+			OutputMin = InOutputMin;
+			OutputMax = InOutputMax;
+
+			// Set all axes scale
+			for (FAxisInfo& AxisInfo : CustomVerticalAxes)
+			{
+				double AxisRange = (AxisInfo.Max - AxisInfo.Min);
+
+				AxisInfo.Min += AxisRange * OffsetFactor;
+				AxisInfo.Max  = AxisInfo.Min + AxisRange * Scale;
+			}
+		}
+	}
+}
+
+void SCurveEditorView::SetInputBounds(double InInputMin, double InInputMax, FCurveEditorViewAxisID AxisID)
+{
+	if (InInputMin < InInputMax)
+	{
+		if (AxisID)
+		{
+			FAxisInfo& AxisInfo = CustomHorizontalAxes[AxisID];
+			AxisInfo.Min = InInputMin;
+			AxisInfo.Max = InInputMax;
+		}
+		// When no axis ID is specified, we scale all axes based on the change
+		else if (TSharedPtr<FCurveEditor> CurveEditor = WeakCurveEditor.Pin())
+		{
+			double CurrentInputMin = 0.0, CurrentInputMax = 1.0;
+			CurveEditor->GetBounds().GetInputBounds(CurrentInputMin, CurrentInputMax);
+
+			double CurrentRange = (CurrentInputMax - CurrentInputMin);
+			double OffsetFactor = (InInputMin - CurrentInputMin) / CurrentRange;
+			double Scale        = (InInputMax - InInputMin) / CurrentRange;
+
+			// Set global scale
+			CurveEditor->GetBounds().SetInputBounds(InInputMin, InInputMax);
+
+			// Set all axes scale
+			for (FAxisInfo& AxisInfo : CustomHorizontalAxes)
+			{
+				double AxisRange = (AxisInfo.Max - AxisInfo.Min);
+
+				AxisInfo.Min += AxisRange * OffsetFactor;
+				AxisInfo.Max  = AxisInfo.Min + AxisRange * Scale;
+			}
+		}
 	}
 }
 
@@ -143,22 +518,41 @@ void SCurveEditorView::ZoomAround(const FVector2D& Amount, double InputOrigin, d
 		double InputMin = 0.0, InputMax = 1.0;
 		CurveEditor->GetBounds().GetInputBounds(InputMin, InputMax);
 
+		const double OriginAlpha = (InputOrigin - InputMin) / (InputMax - InputMin);
+
 		InputMin = InputOrigin - (InputOrigin - InputMin) * Amount.X;
 		InputMax = InputOrigin + (InputMax - InputOrigin) * Amount.X;
 
 		CurveEditor->GetBounds().SetInputBounds(InputMin, InputMax);
+
+		for (FAxisInfo& AxisInfo : CustomHorizontalAxes)
+		{
+			const double AxisInputOrigin = AxisInfo.Min + (AxisInfo.Max - AxisInfo.Min)*OriginAlpha;
+
+			AxisInfo.Min = AxisInputOrigin - (AxisInputOrigin - AxisInfo.Min) * Amount.Y;
+			AxisInfo.Max = AxisInputOrigin + (AxisInfo.Max - AxisInputOrigin) * Amount.Y;
+		}
 	}
 
 	if (Amount.Y != 0.f)
 	{
+		const double OriginAlpha = (OutputOrigin - OutputMin) / (OutputMax - OutputMin);
+
 		OutputMin = OutputOrigin - (OutputOrigin - OutputMin) * Amount.Y;
 		OutputMax = OutputOrigin + (OutputMax - OutputOrigin) * Amount.Y;
+
+		for (FAxisInfo& AxisInfo : CustomVerticalAxes)
+		{
+			const double AxisOutputOrigin = AxisInfo.Min + (AxisInfo.Max - AxisInfo.Min)*OriginAlpha;
+
+			AxisInfo.Min = AxisOutputOrigin - (AxisOutputOrigin - AxisInfo.Min) * Amount.Y;
+			AxisInfo.Max = AxisOutputOrigin + (AxisInfo.Max - AxisOutputOrigin) * Amount.Y;
+		}
 	}
 }
 
 void SCurveEditorView::GetCurveDrawParams(TArray<FCurveDrawParams>& OutDrawParams) 
 {
-
 	TSharedPtr<FCurveEditor> CurveEditor = WeakCurveEditor.Pin();
 	if (!CurveEditor)
 	{
@@ -170,7 +564,7 @@ void SCurveEditorView::GetCurveDrawParams(TArray<FCurveDrawParams>& OutDrawParam
 	GetInputBounds(InputMin, InputMax);
 
 	//make sure the transform is set up
-	UpdateViewToTransformCurves(InputMin,InputMax);
+	UpdateViewToTransformCurves(InputMin, InputMax);
 
 	OutDrawParams.Reset(CurveInfoByID.Num());
 
@@ -184,29 +578,33 @@ void SCurveEditorView::GetCurveDrawParams(TArray<FCurveDrawParams>& OutDrawParam
 			continue;
 		}
 
-		GetCurveDrawParam(CurveEditor, Pair.Key, CurveModel,InputMin, InputMax, Params);
+		GetCurveDrawParam(CurveEditor, Pair.Key, CurveModel, Params);
 		OutDrawParams.Add(MoveTemp(Params));
 	}
 }
 
-void SCurveEditorView::GetCurveDrawParam(TSharedPtr<FCurveEditor>& CurveEditor,const FCurveModelID& ModelID, FCurveModel* CurveModel, 
-	double InputMin, double InputMax, FCurveDrawParams& Params) const
+void SCurveEditorView::GetCurveDrawParam(TSharedPtr<FCurveEditor>& CurveEditor,const FCurveModelID& ModelID, FCurveModel* CurveModel, FCurveDrawParams& Params) const
 {
-
 	FCurveEditorScreenSpace CurveSpace = GetCurveSpace(ModelID);
+
+	const double InputMin = CurveSpace.GetInputMin();
+	const double InputMax = CurveSpace.GetInputMax();
+
 	const double DisplayRatio = (CurveSpace.PixelsPerOutput() / CurveSpace.PixelsPerInput());
 
 	const FKeyHandleSet* SelectedKeys = CurveEditor->GetSelection().GetAll().Find(ModelID);
 
 	// Create a new set of Curve Drawing Parameters to represent this particular Curve
 	Params.Color = CurveModel->GetColor();
+	Params.Thickness = CurveModel->GetThickness();
+	Params.DashLengthPx = CurveModel->GetDashLength();
 	Params.bKeyDrawEnabled = CurveModel->IsKeyDrawEnabled();
 
 	// Gather the display metrics to use for each key type. This allows a Curve Model to override
 	// whether or not the curve supports Keys, Arrive/Leave Tangents, etc. If the Curve Model doesn't
 	// support a particular capability we can skip drawing them.
 	CurveModel->GetKeyDrawInfo(ECurvePointType::ArriveTangent, FKeyHandle::Invalid(), Params.ArriveTangentDrawInfo);
-	CurveModel->GetKeyDrawInfo(ECurvePointType::LeaveTangent, FKeyHandle::Invalid(), Params.LeaveTangentDrawInfo);
+	CurveModel->GetKeyDrawInfo(ECurvePointType::LeaveTangent,  FKeyHandle::Invalid(), Params.LeaveTangentDrawInfo);
 
 	// Gather the interpolating points in input/output space
 	TArray<TTuple<double, double>> InterpolatingPoints;
@@ -408,7 +806,12 @@ void SCurveEditorView::CheckCacheAndInvalidateIfNeeded()
 				if (CurveModel->HasChangedAndResetTest())
 				{
 					bSomethingChanged = true;
-					GetCurveDrawParam(CurveEditor, Params.GetID(), CurveModel, CachedValues.CachedInputMin, CachedValues.CachedInputMax, Params);
+					if (bAllowModelViewTransforms)
+					{
+						bUpdateModelViewTransforms = true;
+						UpdateCustomAxes();
+					}
+					GetCurveDrawParam(CurveEditor, Params.GetID(), CurveModel, Params);
 				}
 			}
 			if (bSomethingChanged)
diff --git a/Engine/Source/Editor/CurveEditor/Private/SCurveViewerPanel.cpp b/Engine/Source/Editor/CurveEditor/Private/SCurveViewerPanel.cpp
--- a/Engine/Source/Editor/CurveEditor/Private/SCurveViewerPanel.cpp
+++ b/Engine/Source/Editor/CurveEditor/Private/SCurveViewerPanel.cpp
@@ -67,18 +67,44 @@ void SCurveViewerPanel::DrawCurves(const FGeometry& AllottedGeometry, FSlateWind
 {
 	const FPaintGeometry PaintGeometry = AllottedGeometry.ToPaintGeometry();
 
+	float DashOffset = static_cast<float>(GetViewSpace().PixelsPerInput() * GetViewSpace().GetInputMin());
+
 	for (const FCurveDrawParams& Params : CachedDrawParams)
 	{
-		FSlateDrawElement::MakeLines(
-			OutDrawElements,
-			BaseLayerId + CurveViewConstants::ELayerOffset::Curves,
-			PaintGeometry,
-			Params.InterpolatingPoints,
-			DrawEffects,
-			Params.Color,
-			CurveViewerConstants::bAntiAliasCurves,
-			CurveThickness.Get()
-		);
+		if (Params.DashLengthPx > 0.f)
+		{
+			TArray<FVector2f> NewVector;
+			NewVector.Reserve(Params.InterpolatingPoints.Num());
+			for (FVector2d Vect : Params.InterpolatingPoints)
+			{
+				NewVector.Add(UE::Slate::CastToVector2f(Vect));
+			}
+
+			FSlateDrawElement::MakeDashedLines(
+				OutDrawElements,
+				BaseLayerId + CurveViewConstants::ELayerOffset::Curves,
+				PaintGeometry,
+				MoveTemp(NewVector),
+				DrawEffects,
+				Params.Color,
+				CurveThickness.Get(),
+				Params.DashLengthPx,
+				DashOffset
+			);
+		}
+		else
+		{
+			FSlateDrawElement::MakeLines(
+				OutDrawElements,
+				BaseLayerId + CurveViewConstants::ELayerOffset::Curves,
+				PaintGeometry,
+				Params.InterpolatingPoints,
+				DrawEffects,
+				Params.Color,
+				CurveViewerConstants::bAntiAliasCurves,
+				CurveThickness.Get()
+			);
+		}
 
 		if (Params.bKeyDrawEnabled)
 		{
diff --git a/Engine/Source/Editor/CurveEditor/Private/Views/SCurveEditorViewNormalized.cpp b/Engine/Source/Editor/CurveEditor/Private/Views/SCurveEditorViewNormalized.cpp
--- a/Engine/Source/Editor/CurveEditor/Private/Views/SCurveEditorViewNormalized.cpp
+++ b/Engine/Source/Editor/CurveEditor/Private/Views/SCurveEditorViewNormalized.cpp
@@ -4,6 +4,7 @@
 
 #include "Containers/SortedMap.h"
 #include "CurveEditor.h"
+#include "CurveEditorAxis.h"
 #include "CurveEditorScreenSpace.h"
 #include "CurveEditorSettings.h"
 #include "CurveModel.h"
@@ -41,6 +42,7 @@ void SCurveEditorViewNormalized::Construct(const FArguments& InArgs, TWeakPtr<FC
 {
 	//when created we set the output bounds to be fixed since otherwise it may get fitted
 	bFixedOutputBounds = true;
+	bAllowModelViewTransforms = false;
 	FrameVertical(-0.1, 1.1);
 
 	SInteractiveCurveEditorView::Construct(InArgs, InCurveEditor);
@@ -57,16 +59,16 @@ void SCurveEditorViewNormalized::Construct(const FArguments& InArgs, TWeakPtr<FC
 	];
 }
 
-FTransform2D CalculateViewToCurveTransform(const double OutputMin, const double OutputMax, const double InCurveOutputMin, const double InCurveOutputMax)
+FTransform2d CalculateViewToCurveTransform(const double OutputMin, const double OutputMax, const double InCurveOutputMin, const double InCurveOutputMax)
 {
 	const double Scale = (InCurveOutputMax - InCurveOutputMin);
 	if (InCurveOutputMax > InCurveOutputMin)
 	{
-		return FTransform2D(FScale2D(1.0, Scale), FVector2D(0.0, InCurveOutputMin));
+		return FTransform2d(FScale2d(1.0, Scale), FVector2d(0.0, InCurveOutputMin));
 	}
 	else
 	{
-		return FVector2D(0.f, InCurveOutputMin - 0.5);
+		return FTransform2d(FVector2d(0.f, InCurveOutputMin - 0.5));
 	}
 }
 
@@ -100,7 +102,7 @@ void SCurveEditorViewNormalized::DrawBufferedCurves(const FGeometry& AllottedGeo
 			continue;
 		}
 
-		FTransform2D ViewToBufferedCurveTransform;
+		FTransform2d ViewToBufferedCurveTransform;
 		double CurveOutputMin = BufferedCurve->GetValueMin(), CurveOutputMax = BufferedCurve->GetValueMax();
 
 		ViewToBufferedCurveTransform = CalculateViewToCurveTransform(OutputMin, OutputMax,CurveOutputMin, CurveOutputMax);
@@ -146,10 +148,22 @@ void SCurveEditorViewNormalized::PaintView(const FPaintArgs& Args, const FGeomet
 	}
 }
 
-void SCurveEditorViewNormalized::FrameVertical(double InOutputMin, double InOutputMax)
+void SCurveEditorViewNormalized::FrameVertical(double InOutputMin, double InOutputMax, FCurveEditorViewAxisID AxisID)
 {
-	OutputMin = -0.1;
-	OutputMax = 1.1;;
+	if (!bFixedOutputBounds && InOutputMin < InOutputMax)
+	{
+		if (AxisID)
+		{
+			FAxisInfo& AxisInfo = GetVerticalAxisInfo(AxisID);
+			AxisInfo.Min = -0.1;
+			AxisInfo.Max =  1.1;
+		}
+		else
+		{
+			OutputMin = -0.1;
+			OutputMax =  1.1;
+		}
+	}
 }
 
 void SCurveEditorViewNormalized::UpdateViewToTransformCurves(double InputMin, double InputMax) 
diff --git a/Engine/Source/Editor/CurveEditor/Private/Views/SCurveEditorViewStacked.cpp b/Engine/Source/Editor/CurveEditor/Private/Views/SCurveEditorViewStacked.cpp
--- a/Engine/Source/Editor/CurveEditor/Private/Views/SCurveEditorViewStacked.cpp
+++ b/Engine/Source/Editor/CurveEditor/Private/Views/SCurveEditorViewStacked.cpp
@@ -40,6 +40,7 @@ struct FSlateBrush;
 void SCurveEditorViewStacked::Construct(const FArguments& InArgs, TWeakPtr<FCurveEditor> InCurveEditor)
 {
 	bFixedOutputBounds = true;
+	bAllowModelViewTransforms = false;
 	OutputMin = 0.0;
 	OutputMax = 1.0;
 
@@ -286,15 +287,15 @@ void SCurveEditorViewStacked::DrawLabels(const FGeometry& AllottedGeometry, cons
 	}
 }
 
-FTransform2D CalculateViewToCurveTransform(const double InCurveOutputMin, double const InCurveOutputMax, const double InValueOffset)
+FTransform2d CalculateViewToCurveTransform(const double InCurveOutputMin, double const InCurveOutputMax, const double InValueOffset)
 {
 	if (InCurveOutputMax > InCurveOutputMin)
 	{
-		return Concatenate(FVector2D(0.f, InValueOffset), Concatenate(FScale2D(1.f, (InCurveOutputMax - InCurveOutputMin)), FVector2D(0.f, InCurveOutputMin)));
+		return Concatenate(FVector2d(0.f, InValueOffset), Concatenate(FScale2d(1.f, (InCurveOutputMax - InCurveOutputMin)), FVector2d(0.f, InCurveOutputMin)));
 	}
 	else
 	{
-		return Concatenate(FVector2D(0.f, InValueOffset - 0.5), FVector2D(0.f, InCurveOutputMin));
+		return Concatenate(FVector2d(0.f, InValueOffset - 0.5), FVector2d(0.f, InCurveOutputMin));
 	}
 }
 
@@ -331,7 +332,7 @@ void SCurveEditorViewStacked::DrawBufferedCurves(const FGeometry& AllottedGeomet
 			continue;
 		}
 
-		FTransform2D ViewToBufferedCurveTransform;
+		FTransform2d ViewToBufferedCurveTransform;
 
 		const int32  CurveIndexFromBottom = CurveInfoByID.Num() - It->Value.CurveIndex - 1;
 		const double PaddingToBottomOfView = (CurveIndexFromBottom + 1)*ValueSpacePadding;
diff --git a/Engine/Source/Editor/CurveEditor/Private/Views/SInteractiveCurveEditorView.cpp b/Engine/Source/Editor/CurveEditor/Private/Views/SInteractiveCurveEditorView.cpp
--- a/Engine/Source/Editor/CurveEditor/Private/Views/SInteractiveCurveEditorView.cpp
+++ b/Engine/Source/Editor/CurveEditor/Private/Views/SInteractiveCurveEditorView.cpp
@@ -11,6 +11,7 @@
 #include "CurveDataAbstraction.h"
 #include "CurveDrawInfo.h"
 #include "CurveEditor.h"
+#include "CurveEditorAxis.h"
 #include "CurveEditorCommands.h"
 #include "CurveEditorContextMenu.h"
 #include "CurveEditorHelpers.h"
@@ -67,6 +68,7 @@
 #include "Widgets/SCompoundWidget.h"
 #include "Widgets/SToolTip.h"
 #include "Widgets/Text/STextBlock.h"
+#include "Widgets/Input/NumericTypeInterface.h"
 
 class FPaintArgs;
 class FWidgetStyle;
@@ -291,18 +293,176 @@ void SInteractiveCurveEditorView::DrawGridLines(TSharedRef<FCurveEditor> CurveEd
 	const FLinearColor   MajorGridColor = CurveEditor->GetPanel()->GetGridLineTint();
 	const FLinearColor   MinorGridColor = MajorGridColor.CopyWithNewOpacity(MajorGridColor.A * .5f);
 	const FPaintGeometry PaintGeometry = AllottedGeometry.ToPaintGeometry();
-	const FLinearColor	 LabelColor = FLinearColor::White.CopyWithNewOpacity(0.65f);
 	const FSlateFontInfo FontInfo = FCoreStyle::Get().GetFontStyle("ToolTip.LargerFont");
 
-	// Get our viewing range bounds. We go through the GetBounds() interface on the curve editor because it's more aware of what our range is than the original widget is.
-	double InputValueMin, InputValueMax;
-	GetInputBounds(InputValueMin, InputValueMax);
+	FCurveEditorScreenSpace ViewSpace = GetViewSpace();
+
+	FCurveEditorScreenSpaceH HorizontalGridSpace = ViewSpace;
+	FCurveEditorScreenSpaceV VerticalGridSpace   = ViewSpace;
 
-	TArray<float> MajorGridLines, MinorGridLines;
-	TArray<FText> MajorGridLabels;
+	TArray<double> MajorGridLines, MinorGridLines;
 
-	GetGridLinesX(CurveEditor, MajorGridLines, MinorGridLines, &MajorGridLabels);
-	ensureMsgf(MajorGridLabels.Num() == 0 || MajorGridLines.Num() == MajorGridLabels.Num(), TEXT("If grid labels are specified, one must be specified for every major grid line, even if it is just an empty FText."));
+	struct FGridLine
+	{
+		double Value;
+		FText Label;
+	};
+
+	struct FGridLineLabels
+	{
+		FGridLineLabels(FCurveEditorViewAxisID AxisID, TSet<FCurveEditorViewAxisID>& HighlightedAxes)
+		{
+			// Show default color if this axis is not highlighted, or its the only one
+			if (HighlightedAxes.Num() == 0)
+			{
+				// Default
+				Color = FLinearColor::White.CopyWithNewOpacity(0.65f);
+			}
+			else if (HighlightedAxes.Contains(AxisID))
+			{
+				// Highlighted
+				Color = FLinearColor::White.CopyWithNewOpacity(0.95f);
+			}
+			else
+			{
+				// Subdued
+				Color = FLinearColor::White.CopyWithNewOpacity(0.15f);
+			}
+		}
+		TArray<FText> Labels;
+		TArray<FVector2f> Sizes;
+		FVector2f MaxSize;
+		FLinearColor Color;
+	};
+	TArray<FGridLineLabels> MajorGridLabels;
+
+	FCurveEditorSelection& Selection = CurveEditor->GetSelection();
+	TSet<FCurveEditorViewAxisID> HighlightedHorizontalAxes, HighlightedVerticalAxes;
+
+	TOptional<FCurveModelID> HoveredCurve = GetHoveredCurve();
+	if (HoveredCurve)
+	{
+		if (const FCurveInfo* CurveInfo = CurveInfoByID.Find(HoveredCurve.GetValue()))
+		{
+			HighlightedHorizontalAxes.Add(CurveInfo->HorizontalAxis);
+			HighlightedVerticalAxes.Add(CurveInfo->VerticalAxis);
+		}
+	}
+	else for (const TTuple<FCurveModelID, FCurveInfo>& Pair : CurveInfoByID)
+	{
+		if (const FKeyHandleSet* SelectionSet = Selection.FindForCurve(Pair.Key))
+		{
+			HighlightedHorizontalAxes.Add(Pair.Value.HorizontalAxis);
+			HighlightedVerticalAxes.Add(Pair.Value.VerticalAxis);
+		}
+	}
+
+	// Ask Custom axes to draw until we find something that does
+	for (int32 Index = 0; Index < CustomHorizontalAxes.Num(); ++Index)
+	{
+		const FAxisInfo& AxisInfo = CustomHorizontalAxes[Index];
+		AxisInfo.Axis->GetGridLines(*CurveEditor, *this, FCurveEditorViewAxisID(Index), MajorGridLines, MinorGridLines, ECurveEditorAxisOrientation::Horizontal);
+		if (MajorGridLines.Num() || MinorGridLines.Num())
+		{
+			HorizontalGridSpace = FCurveEditorScreenSpaceH(HorizontalGridSpace.GetPhysicalWidth(), AxisInfo.Min, AxisInfo.Max);
+			break;
+		}
+	}
+
+	if (MajorGridLines.Num() == 0)
+	{
+		TArray<float> MajorGridLinesFloat, MinorGridLinesFloat;
+		if (bNeedsDefaultGridLinesH)
+		{
+			// Auto populate the major grid labels
+			FGridLineLabels& Labels = MajorGridLabels.Emplace_GetRef(FCurveEditorViewAxisID(), HighlightedHorizontalAxes);
+			GetGridLinesX(CurveEditor, MajorGridLinesFloat, MinorGridLinesFloat, &Labels.Labels);
+		}
+		else
+		{
+			GetGridLinesX(CurveEditor, MajorGridLinesFloat, MinorGridLinesFloat, nullptr);
+		}
+
+		// This legacy API defined grid lines in screen space
+		if (MajorGridLinesFloat.Num())
+		{
+			MajorGridLines.SetNum(MajorGridLinesFloat.Num());
+			for (int32 Index = 0; Index < MajorGridLines.Num(); ++Index)
+			{
+				MajorGridLines[Index] = ViewSpace.ScreenToSeconds(MajorGridLinesFloat[Index]);
+			}
+		}
+		if (MinorGridLinesFloat.Num())
+		{
+			MinorGridLines.SetNum(MinorGridLinesFloat.Num());
+			for (int32 Index = 0; Index < MinorGridLines.Num(); ++Index)
+			{
+				MinorGridLines[Index] = ViewSpace.ScreenToSeconds(MinorGridLinesFloat[Index]);
+			}
+		}
+	}
+	else if (bNeedsDefaultGridLinesH)
+	{
+		FText DefaultFormat = CurveEditor->GetGridLineLabelFormatXAttribute().Get();
+
+		if (!DefaultFormat.IsEmpty())
+		{
+			FGridLineLabels& DefaultGridLabels = MajorGridLabels.Emplace_GetRef(FCurveEditorViewAxisID(), HighlightedHorizontalAxes);
+
+			const int32 Num = MajorGridLines.Num();
+			DefaultGridLabels.Labels.SetNum(Num);
+			for (int32 GridLineIndex = 0; GridLineIndex < Num; ++GridLineIndex)
+			{
+				// Put the grid line from HorizontalGridSpace into the default ViewSpace
+				double GridLine = HorizontalGridSpace.SecondsToScreen(MajorGridLines[GridLineIndex]);
+				GridLine = ViewSpace.ScreenToSeconds(GridLine);
+
+				DefaultGridLabels.Labels[GridLineIndex] = FText::Format(DefaultFormat, GridLine);
+			}
+		}
+	}
+
+	// Populate grid labels for custom axes
+	if (MajorGridLines.Num() > 0)
+	{
+		const int32 NumLabels = MajorGridLines.Num();
+
+		for (int32 Index = 0; Index < CustomHorizontalAxes.Num(); ++Index)
+		{
+			const FAxisInfo& AxisInfo = CustomHorizontalAxes[Index];
+			if (!AxisInfo.Axis->HasLabels())
+			{
+				continue;
+			}
+
+			FGridLineLabels& Entry = MajorGridLabels.Emplace_GetRef(FCurveEditorViewAxisID(Index), HighlightedHorizontalAxes);
+
+			Entry.Labels.SetNum(NumLabels);
+			for (int32 GridLineIndex = 0; GridLineIndex < NumLabels; ++GridLineIndex)
+			{
+				Entry.Labels[GridLineIndex] = AxisInfo.Axis->MakeLabel(MajorGridLines[GridLineIndex]);
+			}
+		}
+	}
+
+	const TSharedRef<FSlateFontMeasure> FontMeasure = FSlateApplication::Get().GetRenderer()->GetFontMeasureService();
+
+	// Compute sizing
+	for (FGridLineLabels& Entry : MajorGridLabels)
+	{
+		Entry.Sizes.SetNum(Entry.Labels.Num());
+
+		FVector2f MaxSize(0.f, 0.f);
+		for (int32 Index = 0; Index < Entry.Labels.Num(); ++Index)
+		{
+			const FVector2f LabelSize = FontMeasure->Measure(Entry.Labels[Index], FontInfo);
+
+			Entry.Sizes[Index] = LabelSize;
+			MaxSize.X = FMath::Max(MaxSize.X, LabelSize.X);
+			MaxSize.Y = FMath::Max(MaxSize.Y, LabelSize.Y);
+		}
+		Entry.MaxSize = MaxSize;
+	}
 
 	// Pre-allocate an array of line points to draw our vertical lines. Each major grid line
 	// will overwrite the X value of both points but leave the Y value untouched so they draw from the bottom to the top.
@@ -310,12 +470,10 @@ void SInteractiveCurveEditorView::DrawGridLines(TSharedRef<FCurveEditor> CurveEd
 	LinePoints.Add(FVector2D(0.f, 0.f));
 	LinePoints.Add(FVector2D(0.f, Height));
 
-	const TSharedRef<FSlateFontMeasure> FontMeasure = FSlateApplication::Get().GetRenderer()->GetFontMeasureService();
-
 	// Draw major vertical grid lines
 	for (int32 i = 0; i < MajorGridLines.Num(); i++)
 	{
-		const float RoundedLine = FMath::RoundToFloat(MajorGridLines[i]);
+		const float RoundedLine = FMath::RoundToFloat(HorizontalGridSpace.SecondsToScreen(MajorGridLines[i]));
 		if (RoundedLine < 0 || RoundedLine > RoundedWidth)
 		{
 			continue;
@@ -324,23 +482,40 @@ void SInteractiveCurveEditorView::DrawGridLines(TSharedRef<FCurveEditor> CurveEd
 		// Vertical Grid Line
 		LinePoints[0].X = LinePoints[1].X = RoundedLine;
 
-		if (MajorGridLabels.IsValidIndex(i))
+		// Offset for all labels
+		if (MajorGridLabels[0].Labels.Num() > 0)
 		{
-			FText Label = MajorGridLabels[i];
-			const FVector2D LabelSize = FontMeasure->Measure(Label, FontInfo);
-			const FPaintGeometry LabelGeometry = AllottedGeometry.ToPaintGeometry(FSlateLayoutTransform(FVector2D(LinePoints[0].X - LabelSize.X*.5f, CurveViewConstants::LabelOffsetPixels)));
-
-			LinePoints[0].Y = LabelSize.Y + CurveViewConstants::LabelOffsetPixels*2.f;
-
-			FSlateDrawElement::MakeText(
-				OutDrawElements,
-				BaseLayerId + CurveViewConstants::ELayerOffset::GridLabels,
-				LabelGeometry,
-				Label,
-				FontInfo,
-				DrawEffects,
-				LabelColor
-			);
+			FVector2f LabelOffset(0.f, 0.f);
+
+			// Compute size of all labels
+			for (FGridLineLabels& Entry : MajorGridLabels)
+			{
+				LabelOffset.Y += CurveViewConstants::LabelOffsetPixels;
+
+				FVector2f LabelSize = Entry.Sizes[i];
+				const FPaintGeometry LabelGeometry = AllottedGeometry.ToPaintGeometry(
+					FSlateLayoutTransform(
+						FVector2f(
+							LinePoints[0].X - LabelSize.X*.5f,                 // Center horizontally on the grid line
+							LabelOffset.Y + (Entry.MaxSize.Y-LabelSize.Y)*.5f  // Center vertically within the axis
+						)
+					)
+				);
+
+				FSlateDrawElement::MakeText(
+					OutDrawElements,
+					BaseLayerId + CurveViewConstants::ELayerOffset::GridLabels,
+					LabelGeometry,
+					Entry.Labels[i],
+					FontInfo,
+					DrawEffects,
+					Entry.Color
+				);
+
+				LabelOffset.Y += Entry.MaxSize.Y + CurveViewConstants::LabelOffsetPixels;
+			}
+
+			LinePoints[0].Y = LabelOffset.Y;
 		}
 		else
 		{
@@ -381,11 +556,130 @@ void SInteractiveCurveEditorView::DrawGridLines(TSharedRef<FCurveEditor> CurveEd
 		);
 	}
 
+
+
+
+
+
+
 	MajorGridLines.Reset();
 	MinorGridLines.Reset();
 	MajorGridLabels.Reset();
-	GetGridLinesY(CurveEditor, MajorGridLines, MinorGridLines, &MajorGridLabels);
-	ensureMsgf(MajorGridLabels.Num() == 0 || MajorGridLines.Num() == MajorGridLabels.Num(), TEXT("If grid labels are specified, one must be specified for every major grid line, even if it is just an empty FText."));
+
+	// Ask Custom axes to draw until we find something that does
+	for (int32 Index = 0; Index < CustomVerticalAxes.Num(); ++Index)
+	{
+		const FAxisInfo& AxisInfo = CustomVerticalAxes[Index];
+		AxisInfo.Axis->GetGridLines(*CurveEditor, *this, FCurveEditorViewAxisID(Index), MajorGridLines, MinorGridLines, ECurveEditorAxisOrientation::Vertical);
+		if (MajorGridLines.Num() || MinorGridLines.Num())
+		{
+			VerticalGridSpace = FCurveEditorScreenSpaceV(ViewSpace.GetPhysicalHeight(), AxisInfo.Min, AxisInfo.Max);
+			break;
+		}
+	}
+
+	if (MajorGridLines.Num() == 0)
+	{
+		// Auto populate the major grid labels
+		TArray<float> MajorGridLinesFloat, MinorGridLinesFloat;
+		if (bNeedsDefaultGridLinesV)
+		{
+			// Auto populate the major grid labels
+			FGridLineLabels& Labels = MajorGridLabels.Emplace_GetRef(FCurveEditorViewAxisID(), HighlightedVerticalAxes);
+			GetGridLinesY(CurveEditor, MajorGridLinesFloat, MinorGridLinesFloat, &Labels.Labels);
+		}
+		else
+		{
+			GetGridLinesY(CurveEditor, MajorGridLinesFloat, MinorGridLinesFloat, nullptr);
+		}
+
+		// This legacy API defined grid lines in screen space
+		if (MajorGridLinesFloat.Num())
+		{
+			MajorGridLines.SetNum(MajorGridLinesFloat.Num());
+			for (int32 Index = 0; Index < MajorGridLines.Num(); ++Index)
+			{
+				MajorGridLines[Index] = ViewSpace.ScreenToValue(MajorGridLinesFloat[Index]);
+			}
+		}
+		if (MinorGridLinesFloat.Num())
+		{
+			MinorGridLines.SetNum(MinorGridLinesFloat.Num());
+			for (int32 Index = 0; Index < MinorGridLines.Num(); ++Index)
+			{
+				MinorGridLines[Index] = ViewSpace.ScreenToValue(MinorGridLinesFloat[Index]);
+			}
+		}
+	}
+	else if (bNeedsDefaultGridLinesV)
+	{
+		FText DefaultFormat = CurveEditor->GetGridLineLabelFormatYAttribute().Get();
+		if (!DefaultFormat.IsEmpty())
+		{
+			FGridLineLabels& DefaultGridLabels = MajorGridLabels.Emplace_GetRef(FCurveEditorViewAxisID(), HighlightedVerticalAxes);
+
+			const int32 Num = MajorGridLines.Num();
+			DefaultGridLabels.Labels.SetNum(Num);
+			for (int32 GridLineIndex = 0; GridLineIndex < Num; ++GridLineIndex)
+			{
+				double GridLine = VerticalGridSpace.ValueToScreen(MajorGridLines[GridLineIndex]);
+				GridLine = ViewSpace.ScreenToValue(GridLine);
+
+				DefaultGridLabels.Labels[GridLineIndex] = FText::Format(DefaultFormat, GridLine);
+			}
+		}
+	}
+
+	// Populate grid labels for custom axes
+	if (MajorGridLines.Num() > 0)
+	{
+		const int32 NumLabels = MajorGridLines.Num();
+
+		for (int32 Index = 0; Index < CustomVerticalAxes.Num(); ++Index)
+		{
+			const FAxisInfo& AxisInfo = CustomVerticalAxes[Index];
+			if (!AxisInfo.Axis->HasLabels())
+			{
+				continue;
+			}
+
+			FGridLineLabels& Entry = MajorGridLabels.Emplace_GetRef(FCurveEditorViewAxisID(Index), HighlightedVerticalAxes);
+
+			Entry.Labels.SetNum(NumLabels);
+			for (int32 GridLineIndex = 0; GridLineIndex < NumLabels; ++GridLineIndex)
+			{
+				Entry.Labels[GridLineIndex] = AxisInfo.Axis->MakeLabel(MajorGridLines[GridLineIndex]);
+			}
+		}
+	}
+
+
+
+
+
+
+
+
+
+
+
+	// Compute sizing
+	for (FGridLineLabels& Entry : MajorGridLabels)
+	{
+		Entry.Sizes.SetNum(Entry.Labels.Num());
+
+		FVector2f MaxSize(0.f, 0.f);
+		for (int32 Index = 0; Index < Entry.Labels.Num(); ++Index)
+		{
+			const FVector2f LabelSize = FontMeasure->Measure(Entry.Labels[Index], FontInfo);
+
+			Entry.Sizes[Index] = LabelSize;
+			MaxSize.X = FMath::Max(MaxSize.X, LabelSize.X);
+			MaxSize.Y = FMath::Max(MaxSize.Y, LabelSize.Y);
+		}
+		Entry.MaxSize = MaxSize;
+	}
+
 
 	// Reset our cached Line to draw from left to right
 	LinePoints[0].X = 0.f;
@@ -394,7 +688,7 @@ void SInteractiveCurveEditorView::DrawGridLines(TSharedRef<FCurveEditor> CurveEd
 	// Draw our major horizontal lines
 	for (int32 i = 0; i < MajorGridLines.Num(); i++)
 	{
-		const float RoundedLine = FMath::RoundToFloat(MajorGridLines[i]);
+		const float RoundedLine = FMath::RoundToFloat(VerticalGridSpace.ValueToScreen(MajorGridLines[i]));
 		if (RoundedLine < 0 || RoundedLine > RoundedHeight)
 		{
 			continue;
@@ -403,29 +697,47 @@ void SInteractiveCurveEditorView::DrawGridLines(TSharedRef<FCurveEditor> CurveEd
 		// Overwrite the height of the line we're drawing to draw the different grid lines.
 		LinePoints[0].Y = LinePoints[1].Y = RoundedLine;
 
-		if (MajorGridLabels.IsValidIndex(i))
+		// Offset for all labels
+		if (MajorGridLabels[0].Labels.Num() > 0)
 		{
-			FText Label = MajorGridLabels[i];
-			const FVector2D LabelSize = FontMeasure->Measure(Label, FontInfo);
-			const FPaintGeometry LabelGeometry = AllottedGeometry.ToPaintGeometry(FSlateLayoutTransform(FVector2D(CurveViewConstants::LabelOffsetPixels, LinePoints[0].Y - LabelSize.Y*.5f)));
-
-			LinePoints[0].X = LabelSize.X + CurveViewConstants::LabelOffsetPixels*2.f;
-
-			FSlateDrawElement::MakeText(
-				OutDrawElements,
-				BaseLayerId + CurveViewConstants::ELayerOffset::GridLabels,
-				LabelGeometry,
-				Label,
-				FontInfo,
-				DrawEffects,
-				LabelColor
-			);
+			FVector2f LabelOffset(0.f, 0.f);
+
+			// Compute size of all labels
+			for (FGridLineLabels& Entry : MajorGridLabels)
+			{
+				LabelOffset.X += CurveViewConstants::LabelOffsetPixels;
+
+				FVector2f LabelSize = Entry.Sizes[i];
+				const FPaintGeometry LabelGeometry = AllottedGeometry.ToPaintGeometry(
+					FSlateLayoutTransform(
+						FVector2D(
+							LabelOffset.X + (Entry.MaxSize.X-LabelSize.X)*.5f,  // Center horizontally within the axis
+							LinePoints[0].Y - LabelSize.Y*.5f                   // Center vertically on the grid line
+						)
+					)
+				);
+
+				FSlateDrawElement::MakeText(
+					OutDrawElements,
+					BaseLayerId + CurveViewConstants::ELayerOffset::GridLabels,
+					LabelGeometry,
+					Entry.Labels[i],
+					FontInfo,
+					DrawEffects,
+					Entry.Color
+				);
+
+				LabelOffset.X += Entry.MaxSize.X + CurveViewConstants::LabelOffsetPixels;
+			}
+
+			LinePoints[0].X = LabelOffset.X;
 		}
 		else
 		{
 			LinePoints[0].X = 0.f;
 		}
 
+		// Draw the grid line
 		FSlateDrawElement::MakeLines(
 			OutDrawElements,
 			BaseLayerId + CurveViewConstants::ELayerOffset::GridLines,
@@ -469,8 +781,7 @@ void SInteractiveCurveEditorView::DrawCurves(TSharedRef<FCurveEditor> CurveEdito
 	const FVector2D      VisibleSize = AllottedGeometry.GetLocalSize();
 	const FPaintGeometry PaintGeometry = AllottedGeometry.ToPaintGeometry();
 
-	const float HoveredCurveThickness = 3.5f;
-	const float UnHoveredCurveThickness = 2.f;
+	const float HoverThicknessOffset = 1.5f;
 	const bool  bAntiAliasCurves = true;
 
 	const bool bDrawLines = CVarDrawCurveLines.GetValueOnGameThread();
@@ -480,22 +791,47 @@ void SInteractiveCurveEditorView::DrawCurves(TSharedRef<FCurveEditor> CurveEdito
 	for (const FCurveDrawParams& Params : CachedDrawParams)
 	{
 		const bool bIsCurveHovered = HoveredCurve.IsSet() && HoveredCurve.GetValue() == Params.GetID();
-		const float Thickness = bIsCurveHovered ? HoveredCurveThickness : UnHoveredCurveThickness;
+		const float Thickness = bIsCurveHovered ? Params.Thickness+HoverThicknessOffset : Params.Thickness;
 		const int32 CurveLayerId = bIsCurveHovered ? BaseLayerId + CurveViewConstants::ELayerOffset::Curves : BaseLayerId + CurveViewConstants::ELayerOffset::HoveredCurves;
 
 		if (bDrawLines)
 		{
+			if (Params.DashLengthPx > 0.f)
+			{
+				float DashOffset = static_cast<float>(GetViewSpace().PixelsPerInput() * GetViewSpace().GetInputMin());
 
-			FSlateDrawElement::MakeLines(
-				OutDrawElements,
-				CurveLayerId,
-				PaintGeometry,
-				Params.InterpolatingPoints,
-				DrawEffects,
-				Params.Color,
-				bAntiAliasCurves,
-				Thickness
-			);
+				TArray<FVector2f> NewVector;
+				NewVector.Reserve(Params.InterpolatingPoints.Num());
+				for (FVector2d Vect : Params.InterpolatingPoints)
+				{
+					NewVector.Add(UE::Slate::CastToVector2f(Vect));
+				}
+
+				FSlateDrawElement::MakeDashedLines(
+					OutDrawElements,
+					CurveLayerId,
+					PaintGeometry,
+					MoveTemp(NewVector),
+					DrawEffects,
+					Params.Color,
+					Thickness,
+					Params.DashLengthPx,
+					DashOffset
+				);
+			}
+			else
+			{
+				FSlateDrawElement::MakeLines(
+					OutDrawElements,
+					CurveLayerId,
+					PaintGeometry,
+					Params.InterpolatingPoints,
+					DrawEffects,
+					Params.Color,
+					bAntiAliasCurves,
+					Thickness
+				);
+			}
 		}
 		
 		TArray<FVector2D> LinePoints;
@@ -696,6 +1032,8 @@ bool SInteractiveCurveEditorView::GetCurveWithinWidgetRange(const FSlateRect& Wi
 
 void SInteractiveCurveEditorView::UpdateCurveProximities(FVector2D MousePixel)
 {
+	TOptional<FCurveModelID> PreviouslyHovered = GetHoveredCurve();
+
 	CurveProximities.Reset();
 	CachedToolTipData.Reset();
 
@@ -747,9 +1085,12 @@ void SInteractiveCurveEditorView::UpdateCurveProximities(FVector2D MousePixel)
 
 	Algo::SortBy(CurveProximities, [](TTuple<FCurveModelID, float> In) { return In.Get<1>(); });
 
+	TOptional<FCurveModelID> NewHovered;
 	// Also, set the cached tooltips if dragging because the curve proximity might not be updated during the drag
 	if (CurveProximities.Num() > 0 && (CurveProximities[0].Get<1>() < CurveViewConstants::HoverProximityThresholdPx || DragOperation.IsSet()))
 	{
+		NewHovered = CurveProximities[0].Get<0>();
+
 		const FCurveModel* HoveredCurve = CurveEditor->FindCurve(CurveProximities[0].Get<0>());
 		if (HoveredCurve)
 		{
@@ -768,6 +1109,11 @@ void SInteractiveCurveEditorView::UpdateCurveProximities(FVector2D MousePixel)
 			CachedToolTipData = ToolTipData;
 		}
 	}
+
+	if (PreviouslyHovered != NewHovered)
+	{
+		RefreshRetainer();
+	}
 }
 
 FText SInteractiveCurveEditorView::FormatToolTipCurveName(const FCurveModel& CurveModel) const
diff --git a/Engine/Source/Editor/CurveEditor/Public/CurveDataAbstraction.h b/Engine/Source/Editor/CurveEditor/Public/CurveDataAbstraction.h
--- a/Engine/Source/Editor/CurveEditor/Public/CurveDataAbstraction.h
+++ b/Engine/Source/Editor/CurveEditor/Public/CurveDataAbstraction.h
@@ -5,6 +5,7 @@
 #include "CoreTypes.h"
 #include "CurveEditorTypes.h"
 #include "Curves/RichCurve.h"
+#include "Math/TransformCalculus2D.h"
 #include "CurveDataAbstraction.generated.h"
 
 struct FSlateBrush;
@@ -25,6 +26,13 @@ struct FKeyPosition
 		: InputValue(Input), OutputValue(Output)
 	{}
 
+	FKeyPosition Transform(const FTransform2d& InTransform) const
+	{
+		FVector2d Point(InputValue, OutputValue);
+		Point = InTransform.TransformPoint(Point);
+		return FKeyPosition(Point.X, Point.Y);
+	}
+
 	/** The key's input (x-axis) position (i.e. it's time) */
 	UPROPERTY()
 	double InputValue;
diff --git a/Engine/Source/Editor/CurveEditor/Public/CurveDrawInfo.h b/Engine/Source/Editor/CurveEditor/Public/CurveDrawInfo.h
--- a/Engine/Source/Editor/CurveEditor/Public/CurveDrawInfo.h
+++ b/Engine/Source/Editor/CurveEditor/Public/CurveDrawInfo.h
@@ -81,6 +81,12 @@ struct FCurveDrawParams
 	/** Value defining how to draw keys of type ECurvePointType::LeaveTangent. */
 	FKeyDrawInfo LeaveTangentDrawInfo;
 
+	/** Thickness of this curve */
+	float Thickness = 2.f;
+
+	/** The color to draw this curve */
+	float DashLengthPx = 0.f;
+
 	/** Whether or not to draw keys on the curve. */
 	uint8 bKeyDrawEnabled : 1;
 
diff --git a/Engine/Source/Editor/CurveEditor/Public/CurveEditor.h b/Engine/Source/Editor/CurveEditor/Public/CurveEditor.h
--- a/Engine/Source/Editor/CurveEditor/Public/CurveEditor.h
+++ b/Engine/Source/Editor/CurveEditor/Public/CurveEditor.h
@@ -215,7 +215,32 @@ public:
 	ICurveEditorToolExtension* GetCurrentTool() const;
 
 	FCurveEditorToolID AddTool(TUniquePtr<ICurveEditorToolExtension>&& InTool);
-	
+
+	/**
+	 * Adds a new axis to this curve editor with the specified name.
+	 * @note Calling this function with the same identifier as a previously registered axis will
+	 *       overwrite the axis with the new one.
+	 * 
+	 * @param InIdentifier      A unique identifer for this axis that can be subsequently used to look up this axis
+	 * @param InAxis            The axis definition
+	 */
+	void AddAxis(const FName& InIdentifier, TSharedPtr<FCurveEditorAxis> InAxis);
+
+	/**
+	 * Find an axis by its ID
+	 */
+	TSharedPtr<FCurveEditorAxis> FindAxis(const FName& InIdentifier) const;
+
+	/**
+	 * Remove the axis with the specified ID
+	 */
+	void RemoveAxis(const FName& InIdentifier);
+
+	/**
+	 * Remove all custom axes definitions
+	 */
+	void ClearAxes();
+
 	/** Nudge left or right*/
 	void TranslateSelectedKeys(double SecondsToAdd);
 	void TranslateSelectedKeysLeft();
@@ -607,6 +632,11 @@ protected:
 	TMap<FCurveModelID, TUniquePtr<FCurveModel>> CurveData;
 	/** Map from curve model ID to its originating tree item */
 	TMap<FCurveModelID, FCurveEditorTreeItemID> TreeIDByCurveID;
+	/** Map of child curves */
+	TMultiMap<FCurveModelID, FCurveModelID> ChildCurves;
+
+	/** Map of all axes */
+	TSortedMap<FName, TSharedPtr<FCurveEditorAxis>, FDefaultAllocator, FNameFastLess> CustomAxes;
 
 	/** Set of pinned curve models */
 	TSet<FCurveModelID> PinnedCurves;
diff --git a/Engine/Source/Editor/CurveEditor/Public/CurveEditorHelpers.h b/Engine/Source/Editor/CurveEditor/Public/CurveEditorHelpers.h
--- a/Engine/Source/Editor/CurveEditor/Public/CurveEditorHelpers.h
+++ b/Engine/Source/Editor/CurveEditor/Public/CurveEditorHelpers.h
@@ -20,6 +20,8 @@ namespace CurveEditor
 
 	FVector2D GetVectorFromSlopeAndLength(float Slope, float Length);
 
+	void PopulateGridLineValues(float PhysicalSize, double ViewMin, double ViewMax, uint8 InMinorDivisions, TArray<double>& OutMajorGridLines, TArray<double>& OutMinorGridLines);
+
 	void ConstructYGridLines(const FCurveEditorScreenSpace& ViewSpace, uint8 InMinorDivisions, TArray<float>& OutMajorGridLines, TArray<float>& OutMinorGridLines, FText GridLineLabelFormatY, TArray<FText>* OutMajorGridLabels);
 
 	void ConstructFixedYGridLines(const FCurveEditorScreenSpace& ViewSpace, uint8 InMinorDivisions, double InMinorGridStep, TArray<float>& OutMajorGridLines, TArray<float>& OutMinorGridLines, FText GridLineLabelFormatY, 
diff --git a/Engine/Source/Editor/CurveEditor/Public/CurveEditorScreenSpace.h b/Engine/Source/Editor/CurveEditor/Public/CurveEditorScreenSpace.h
--- a/Engine/Source/Editor/CurveEditor/Public/CurveEditorScreenSpace.h
+++ b/Engine/Source/Editor/CurveEditor/Public/CurveEditorScreenSpace.h
@@ -66,7 +66,7 @@ public:
 	/**
 	 * Transform this screen space into a curve space using the specified transform
 	 */
-	FCurveEditorScreenSpaceH ToCurveSpace(const FTransform2D& CurveTransform) const
+	FCurveEditorScreenSpaceH ToCurveSpace(const FTransform2d& CurveTransform) const
 	{
 		const FVector2D& T = CurveTransform.GetTranslation();
 
@@ -100,7 +100,7 @@ struct FCurveEditorScreenSpaceV
 	{
 		if (!ensure(OutputMax > OutputMin))
 		{
-			OutputMax = OutputMin + KINDA_SMALL_NUMBER;
+			OutputMax = OutputMin + 1.0;
 		}
 	}
 
@@ -141,7 +141,7 @@ public:
 	/**
 	 * Transform this screen space into a curve space using the specified transform
 	 */
-	FCurveEditorScreenSpaceV ToCurveSpace(const FTransform2D& CurveTransform) const
+	FCurveEditorScreenSpaceV ToCurveSpace(const FTransform2d& CurveTransform) const
 	{
 		const FVector2D& T = CurveTransform.GetTranslation();
 
@@ -183,7 +183,7 @@ public:
 	/**
 	 * Transform this screen space into a curve space using the specified transform
 	 */
-	FCurveEditorScreenSpace ToCurveSpace(const FTransform2D& CurveTransform) const
+	FCurveEditorScreenSpace ToCurveSpace(const FTransform2d& CurveTransform) const
 	{
 		FCurveEditorScreenSpace New = *this;
 		FCurveEditorScreenSpaceH& H = static_cast<FCurveEditorScreenSpaceH&>(New);
diff --git a/Engine/Source/Editor/CurveEditor/Public/CurveModel.h b/Engine/Source/Editor/CurveEditor/Public/CurveModel.h
--- a/Engine/Source/Editor/CurveEditor/Public/CurveModel.h
+++ b/Engine/Source/Editor/CurveEditor/Public/CurveModel.h
@@ -21,6 +21,7 @@
 #include "UObject/UnrealType.h"
 
 class FCurveEditor;
+class FCurveEditorAxis;
 class FName;
 class IBufferedCurveModel;
 class SCurveEditorView;
@@ -35,6 +36,7 @@ struct FKeyDrawInfo;
 struct FKeyHandle;
 struct FKeyPosition;
 
+enum class ECurveEditorAxisOrientation : uint8;
 enum class ECurvePointType : uint8;
 
 /**
@@ -78,7 +80,6 @@ public:
 		GetKeys(MinTime, MaxTime, MinValue, MaxValue, OutKeyHandles);
 	}
 
-
 	/**
 	 * Retrieve all keys that lie in the specified time and value range
 	 *
@@ -249,6 +250,13 @@ public:
 	virtual void CreateKeyProxies(TArrayView<const FKeyHandle> InKeyHandles, TArrayView<UObject*> OutObjects)
 	{}
 
+	/**
+	 * Called when this curve model is added to a curve editor in order to construct additional 'child'
+	 *     curves that are owned by this curve
+	 */
+	virtual void MakeChildCurves(TArray<TUniquePtr<FCurveModel>>& OutChildCurves) const
+	{}
+
 	/**
 	 * Creates a copy of this curve, stored in a minimal buffered curve object.
 	 * Buffered curves are used to cache the positions and attributes of a curve's keys. After creation, a buffered curve 
@@ -261,6 +269,14 @@ public:
 		return nullptr;
 	}
 
+	/**
+	 * Return a per-curve transform for this model, if used by the view
+	 */
+	virtual FTransform2d GetCurveTransform() const
+	{
+		return FTransform2d();
+	}
+
 	/** 
 	 * Returns whether the curve model should be edited or not
 	 */
@@ -282,6 +298,13 @@ public:
 		return true;
 	}
 
+	/**
+	 * Allocate and/or assign axes for the this curve
+	 */
+	virtual void AllocateAxes(FCurveEditor* InCurveEditor, TSharedPtr<FCurveEditorAxis>& OutHorizontalAxis, TSharedPtr<FCurveEditorAxis>& OutVerticalAxis) const
+	{
+	}
+
 	/**
 	* Get the Object and the name to be used to store the curve model color (see UCurveEditorSettings). By default
 	* this is the owning object and the intent name, but it can be overriden, for example for Sequencer it may be the bound object
@@ -310,6 +333,21 @@ public:
 		return CurveModifiedDelegate;
 	}
 
+	/**
+	 * Get the owning object of this curve as the specified template type, or search is outer chain if it is not the specified type.
+	 **/
+	template<typename T>
+	T* GetOwningObjectOrOuter() const
+	{
+		UObject* Obj = GetOwningObject();
+		if (T* CastResult = Cast<T>(Obj))
+		{
+			return CastResult;
+		}
+
+		return Obj ? Obj->GetTypedOuter<T>() : nullptr;
+	}
+
 public:
 
 	/**
@@ -412,8 +450,40 @@ public:
 	}
 
 	/**
- * Retrieves whether or not to disable drawing keys
- */
+	 * Get the dash length used for rendering this curve (0 indicates a solid line should be used)
+	 */
+	FORCEINLINE float GetDashLength() const
+	{
+		return DashLengthPx;
+	}
+
+	/**
+	 * Set the dash length used for rendering this curve (0 indicates a solid line should be used)
+	 */
+	FORCEINLINE void SetDashLength(float InDashLengthPx)
+	{
+		DashLengthPx = InDashLengthPx;
+	}
+
+	/**
+	 * Get this curve's visual thickness in screen space
+	 */
+	FORCEINLINE float GetThickness() const
+	{
+		return Thickness;
+	}
+
+	/**
+	 * Set this curve's visual thickness in screen space
+	 */
+	FORCEINLINE void SetThickness(float InThickness)
+	{
+		Thickness = InThickness;
+	}
+
+	/**
+	 * Retrieves whether or not to disable drawing keys
+	 */
 	FORCEINLINE bool IsKeyDrawEnabled() const
 	{
 		return bKeyDrawEnabled.Get();
@@ -463,6 +533,12 @@ protected:
 	/** Whether or not to draw curve's keys */
 	TAttribute<bool> bKeyDrawEnabled;
 
+	/** When this curve is to be drawn dashed, the length of each dash (or 0.f if solid) */
+	float DashLengthPx = 0.f;
+
+	/** Default thickness for this curve */
+	float Thickness = 2.f;
+
 	/** A set of views supported by this curve */
 	ECurveEditorViewID SupportedViews;
 
diff --git a/Engine/Source/Editor/CurveEditor/Public/SCurveEditorView.h b/Engine/Source/Editor/CurveEditor/Public/SCurveEditorView.h
--- a/Engine/Source/Editor/CurveEditor/Public/SCurveEditorView.h
+++ b/Engine/Source/Editor/CurveEditor/Public/SCurveEditorView.h
@@ -25,7 +25,72 @@ class FSlateRect;
 class FText;
 class SCurveEditorPanel;
 class FCurveModel;
+class FCurveEditorAxis;
 class SRetainerWidget;
+class SCurveEditorView;
+
+enum class ECurveEditorAxisOrientation : uint8;
+
+
+/**
+ * Identifier for a specific axis on a view. These identifiers are transient and should not be stored externally.
+ */
+struct FCurveEditorViewAxisID
+{
+	FCurveEditorViewAxisID()
+		: Index(std::numeric_limits<uint8>::max())
+	{}
+	explicit FCurveEditorViewAxisID(uint8 InIndex)
+		: Index(InIndex)
+	{
+		check(InIndex < std::numeric_limits<uint8>::max());
+	}
+
+	explicit operator bool() const
+	{
+		return Index != std::numeric_limits<uint8>::max();
+	}
+
+	friend uint32 GetTypeHash(FCurveEditorViewAxisID In)
+	{
+		return GetTypeHash(In.Index);
+	}
+
+	friend bool operator==(FCurveEditorViewAxisID A, FCurveEditorViewAxisID B)
+	{
+		return A.Index == B.Index;
+	}
+
+	friend bool operator!=(FCurveEditorViewAxisID A, FCurveEditorViewAxisID B)
+	{
+		return A.Index != B.Index;
+	}
+
+	friend bool operator<(FCurveEditorViewAxisID A, FCurveEditorViewAxisID B)
+	{
+		return A.Index < B.Index;
+	}
+
+private:
+
+	operator int32() const
+	{
+		check(Index != std::numeric_limits<uint8>::max());
+		return Index;
+	}
+
+	FCurveEditorViewAxisID& operator=(int32 InIndex)
+	{
+		check(InIndex >= 0 && InIndex < 255);
+		Index = static_cast<uint8>(InIndex);
+		return *this;
+	}
+
+	friend SCurveEditorView;
+	uint8 Index;
+};
+
+
 
 /**
  * This is the base widget type for all views that exist on a curve editor panel. A view may contain 0 or more curves (stored in CurveInfoByID).
@@ -57,20 +122,43 @@ public:
 	 * Default constructor
 	 */
 	SCurveEditorView();
+	~SCurveEditorView();
 
 	/**
-	 * Get the screen space ultility that defines this view's input, output and pixel metrics
+	 * Get the default screen space ultility that defines this view's input, output and pixel metrics
 	 */
 	FCurveEditorScreenSpace GetViewSpace() const;
 
+	/**
+	 * Get the screen space ultility that defines this view's input, output and pixel metrics for the specified axis combination
+	 */
+	FCurveEditorScreenSpace GetViewSpace(const FName& InHorizontalAxis, const FName& InVerticalAxis) const;
+
 	/**
 	 * Get the screen space ultility that defines the specified curves's input, output and pixel metrics.
 	 * The resulting struct defines the transformation from this view's widget pixel space to the curve input/output space
 	 */
-	FCurveEditorScreenSpace GetCurveSpace(FCurveModelID CurveID) const
-	{
-		return GetViewSpace().ToCurveSpace(CurveInfoByID.FindRef(CurveID).ViewToCurveTransform);
-	}
+	FCurveEditorScreenSpace GetCurveSpace(FCurveModelID CurveID) const;
+
+	/**
+	 * Retrieve the horizontal screen-space information for the specified axis
+	 */
+	FCurveEditorScreenSpaceH GetHorizontalAxisSpace(FCurveEditorViewAxisID ID) const;
+
+	/**
+	 * Retrieve the vertical screen-space information for the specified axis
+	 */
+	FCurveEditorScreenSpaceV GetVerticalAxisSpace(FCurveEditorViewAxisID ID) const;
+
+	/**
+	 * Retrueve the axis ID assigned to the specified curve and orientation
+	 */
+	FCurveEditorViewAxisID GetAxisForCurve(FCurveModelID CurveID, ECurveEditorAxisOrientation AxisOrientation) const;
+
+	/**
+	 * Retrueve the axis associated with the specified ID and orientation
+	 */
+	TSharedPtr<FCurveEditorAxis> GetAxis(FCurveEditorViewAxisID ID, ECurveEditorAxisOrientation AxisOrientation) const;
 
 	/**
 	 * Check whether this view should auto size to fit its FixedHeight or child slot content
@@ -128,7 +216,12 @@ public:
 	/**
 	 * Set this view's output bounds
 	 */
-	void SetOutputBounds(double InOutputMin, double InOutputMax);
+	void SetOutputBounds(double InOutputMin, double InOutputMax, FCurveEditorViewAxisID AxisID = FCurveEditorViewAxisID());
+
+	/**
+	 * Set this view's input bounds
+	 */
+	void SetInputBounds(double InInputMin, double InInputMax, FCurveEditorViewAxisID AxisID = FCurveEditorViewAxisID());
 
 	/**
 	 * Zoom this view in or out around its center point
@@ -156,7 +249,11 @@ public:
 	virtual void UpdateViewToTransformCurves(double InputMin, double InputMax) {};
 
 	/** Frame the view vertially by the input and output bounds, peformaing any custom clipping as needed */
-	virtual void FrameVertical(double InOutputMin, double InOutputMax);
+	virtual void FrameVertical(double InOutputMin, double InOutputMax, FCurveEditorViewAxisID AxisID = FCurveEditorViewAxisID());
+
+	/** Frame the view horizontall by the input and output bounds, peformaing any custom clipping as needed */
+	virtual void FrameHorizontal(double InInputMin, double InInputMax, FCurveEditorViewAxisID AxisID = FCurveEditorViewAxisID());
+
 public:
 
 	/**
@@ -222,14 +319,17 @@ protected:
 	void GetCurveDrawParams(TArray<FCurveDrawParams>& OutDrawParams);
 
 	/** Get it for just one curve*/
-	void GetCurveDrawParam(TSharedPtr<FCurveEditor>& CurveEditor, const FCurveModelID& ModelID, FCurveModel* CurveModel,
-		double InputMin, double InputMax, FCurveDrawParams& OutDrawParam) const;
+	void GetCurveDrawParam(TSharedPtr<FCurveEditor>& CurveEditor, const FCurveModelID& ModelID, FCurveModel* CurveModel, FCurveDrawParams& OutDrawParam) const;
 
 	/** Request a new render from the retainer widget */
 	void RefreshRetainer();
 
+	/** Update all the curve view transforms from curve models */
+	void UpdateCurveViewTransformsFromModels();
+
 	// ~SWidget interface
 	virtual FVector2D ComputeDesiredSize(float LayoutScaleMultiplier) const override;
+	virtual void Tick(const FGeometry& AllottedGeometry, const double InCurrentTime, const float InDeltaTime) override;
 
 private:
 
@@ -242,6 +342,9 @@ private:
 	/** Add a curve from this view */
 	void RemoveCurve(FCurveModelID CurveID);
 
+	/** Update the custom axes if necessary */
+	void UpdateCustomAxes();
+
 public:
 
 	/** (Default: false) When true, this view has been pinned to the top of the view container. Only takes affect from the next call to SCurveEditorPanel::RebuildCurveViews. */
@@ -254,6 +357,14 @@ public:
 	uint8 bAutoSize : 1;
 	/** (Default: false) Defines whether this view should remain on the panel UI even if it does not represent any curves. */
 	uint8 bAllowEmpty : 1;
+	/** (Default: false) When true, custom axes need to be rebuilt before use. */
+	uint8 bAllowModelViewTransforms : 1;
+	/** (Default: false) When true, custom axes need to be rebuilt before use. */
+	uint8 bUpdateModelViewTransforms : 1;
+	/** (Default: true) When true, this view has models that need the default grid lines drawing. */
+	uint8 bNeedsDefaultGridLinesH : 1;
+	/** (Default: true) When true, this view has models that need the default grid lines drawing. */
+	uint8 bNeedsDefaultGridLinesV : 1;
 
 	/** (Default: 0) Should be assigned on Construction of derived types. Defines a custom sort bias to use when sorting the stack of views (before sorting by pinned/unpinned state) */
 	int8 SortBias = 0;
@@ -285,13 +396,38 @@ protected:
 	{
 		/** The linear index of the curve within this view determined by the order curves were added. */
 		int32 CurveIndex;
-		FTransform2D ViewToCurveTransform;
+		FTransform2d ViewToCurveTransform;
+		FCurveEditorViewAxisID HorizontalAxis;
+		FCurveEditorViewAxisID VerticalAxis;
 	};
+
 	/** Map from curve identifier to specific info pertaining to that curve for this view.
 	 * @note: Should only be added to or removed from in AddCurve/RemoveCurve. Derived types must only change the FCurveInfo contained within this map.
 	 */
 	TSortedMap<FCurveModelID, FCurveInfo, TInlineAllocator<1>> CurveInfoByID;
 
+	struct FAxisInfo
+	{
+		TSharedPtr<FCurveEditorAxis> Axis;
+		double Min = 0.0;
+		double Max = 1.0;
+		int32 UseCount = 0;
+	};
+
+	TArray<FAxisInfo> CustomHorizontalAxes;
+	TArray<FAxisInfo> CustomVerticalAxes;
+
+	FAxisInfo& GetHorizontalAxisInfo(FCurveEditorViewAxisID ID)
+	{
+		check(ID);
+		return CustomHorizontalAxes[ID];
+	}
+	FAxisInfo& GetVerticalAxisInfo(FCurveEditorViewAxisID ID)
+	{
+		check(ID);
+		return CustomVerticalAxes[ID];
+	}
+
 	/** Flag enum signifying how the curve cache has changed since it was last generated
 	Note for a data change it may only effect certain data(curves) not every drawn curve*/
 	enum class ECurveCacheFlags : uint8
diff --git a/Engine/Source/Editor/CurveEditor/Public/Views/SCurveEditorViewNormalized.h b/Engine/Source/Editor/CurveEditor/Public/Views/SCurveEditorViewNormalized.h
--- a/Engine/Source/Editor/CurveEditor/Public/Views/SCurveEditorViewNormalized.h
+++ b/Engine/Source/Editor/CurveEditor/Public/Views/SCurveEditorViewNormalized.h
@@ -30,7 +30,7 @@ public:
 
 	virtual void UpdateViewToTransformCurves(double InputMin, double InputMax) override;
 
-	void FrameVertical(double InOutputMin, double InOutputMax) override;
+	void FrameVertical(double InOutputMin, double InOutputMax, FCurveEditorViewAxisID AxisID = FCurveEditorViewAxisID()) override;
 
 protected:
 
diff --git a/Engine/Source/Editor/MovieSceneTools/Private/AutomatedLevelSequenceCapture.cpp b/Engine/Source/Editor/MovieSceneTools/Private/AutomatedLevelSequenceCapture.cpp
--- a/Engine/Source/Editor/MovieSceneTools/Private/AutomatedLevelSequenceCapture.cpp
+++ b/Engine/Source/Editor/MovieSceneTools/Private/AutomatedLevelSequenceCapture.cpp
@@ -423,7 +423,8 @@ bool UAutomatedLevelSequenceCapture::InitializeShots()
 		{
 			// Expand the inner shot section range by the handle size, multiplied by the difference between the outer and inner tick resolutions (and factoring in the time scale)
 			const float OuterToInnerRateDilation = (MovieScene->GetTickResolution() == ShotMovieScene->GetTickResolution()) ? 1.f : (ShotMovieScene->GetTickResolution() / MovieScene->GetTickResolution()).AsDecimal();
-			const float OuterToInnerScale = OuterToInnerRateDilation * ShotSection->Parameters.TimeScale;
+			// coderot: this is not considering timewarping
+			const float OuterToInnerScale = OuterToInnerRateDilation;
 
 			CachedShotStates.Add(FCinematicShotCache(ShotSection->IsActive(), ShotSection->IsLocked(), ShotSection->GetRange(), ShotMovieScene ? ShotMovieScene->GetPlaybackRange() : TRange<FFrameNumber>::Empty()));
 
diff --git a/Engine/Source/Editor/MovieSceneTools/Private/Channels/BezierChannelCurveModel.cpp b/Engine/Source/Editor/MovieSceneTools/Private/Channels/BezierChannelCurveModel.cpp
--- a/Engine/Source/Editor/MovieSceneTools/Private/Channels/BezierChannelCurveModel.cpp
+++ b/Engine/Source/Editor/MovieSceneTools/Private/Channels/BezierChannelCurveModel.cpp
@@ -72,11 +72,23 @@ FBezierChannelCurveModel<ChannelType, ChannelValue, KeyType>::FBezierChannelCurv
 	}
 }
 
+template<typename ChannelType, typename ChannelValue, typename KeyType> 
+FBezierChannelCurveModel<ChannelType, ChannelValue, KeyType>::FBezierChannelCurveModel(TMovieSceneChannelHandle<ChannelType> InChannel, UMovieSceneSection* OwningSection, UObject* InOwningObject, TWeakPtr<ISequencer> InWeakSequencer)
+	: FChannelCurveModel<ChannelType, ChannelValue, KeyType>(InChannel, OwningSection, InOwningObject, InWeakSequencer)
+{
+	ChannelType* Channel = InChannel.Get();
+
+	if (Channel && OwningSection && OwningSection->GetTypedOuter<UMovieScene>())
+	{
+		Channel->SetTickResolution(OwningSection->GetTypedOuter<UMovieScene>()->GetTickResolution());
+	}
+}
+
 template<typename ChannelType, typename ChannelValue, typename KeyType> 
 void FBezierChannelCurveModel<ChannelType, ChannelValue, KeyType>::DrawCurve(const FCurveEditor& CurveEditor, const FCurveEditorScreenSpace& ScreenSpace, TArray<TTuple<double, double>>& InterpolatingPoints) const
 {
 	ChannelType* Channel		= this->GetChannelHandle().Get();
-	UMovieSceneSection* Section = Cast<UMovieSceneSection>(this->GetOwningObject());
+	UMovieSceneSection* Section = this->template GetOwningObjectOrOuter<UMovieSceneSection>();
 
 	if (Channel && Section && Section->GetTypedOuter<UMovieScene>())
 	{
@@ -156,7 +168,7 @@ template <class ChannelType, class ChannelValue, class KeyType>
 TPair<ERichCurveInterpMode, ERichCurveTangentMode> FBezierChannelCurveModel<ChannelType, ChannelValue, KeyType>::GetInterpolationMode(const double& InTime, ERichCurveInterpMode DefaultInterpolationMode, ERichCurveTangentMode DefaultTangentMode) const
 {
 	ChannelType* Channel = this->GetChannelHandle().Get();
-	UMovieSceneSection* Section = Cast<UMovieSceneSection>(this->GetOwningObject());
+	UMovieSceneSection* Section = this->template GetOwningObjectOrOuter<UMovieSceneSection>();
 	if (Channel && Section)
 	{
 		FFrameRate TickResolution = Section->GetTypedOuter<UMovieScene>()->GetTickResolution();
@@ -197,15 +209,17 @@ TPair<ERichCurveInterpMode, ERichCurveTangentMode> FBezierChannelCurveModel<Chan
 template<typename ChannelType, typename ChannelValue, typename KeyType> 
 void FBezierChannelCurveModel<ChannelType, ChannelValue, KeyType>::GetKeyAttributes(TArrayView<const FKeyHandle> InKeys, TArrayView<FKeyAttributes> OutAttributes) const
 {
-	ChannelType* Channel		= this->GetChannelHandle().Get();
-	UMovieSceneSection* Section = Cast<UMovieSceneSection>(this->GetOwningObject());
-	if (Channel && Section && Section->GetTypedOuter<UMovieScene>())
+	ChannelType*        Channel    = this->GetChannelHandle().Get();
+	UMovieSceneSection* Section    = this->template GetOwningObjectOrOuter<UMovieSceneSection>();
+	UMovieScene*        MovieScene = Section ? Section->GetTypedOuter<UMovieScene>() : nullptr;
+
+	if (Channel && Section && MovieScene)
 	{
 		TMovieSceneChannelData<ChannelValue> ChannelData = Channel->GetData();
 		TArrayView<const FFrameNumber>    Times  = ChannelData.GetTimes();
 		TArrayView<ChannelValue> Values = ChannelData.GetValues();
 
-		float TimeInterval = Section->GetTypedOuter<UMovieScene>()->GetTickResolution().AsInterval();
+		float TimeInterval = MovieScene->GetTickResolution().AsInterval();
 
 		for (int32 Index = 0; Index < InKeys.Num(); ++Index)
 		{
@@ -253,17 +267,20 @@ static bool IsAuto(ERichCurveTangentMode TangentMode)
 template<typename ChannelType, typename ChannelValue, typename KeyType> 
 void FBezierChannelCurveModel<ChannelType, ChannelValue, KeyType>::SetKeyAttributes(TArrayView<const FKeyHandle> InKeys, TArrayView<const FKeyAttributes> InAttributes, EPropertyChangeType::Type ChangeType)
 {
-	ChannelType* Channel		= this->GetChannelHandle().Get();
-	UMovieSceneSection* Section = Cast<UMovieSceneSection>(this->GetOwningObject());
-	if (Channel && Section && Section->GetTypedOuter<UMovieScene>() && !this->IsReadOnly())
+	UE::MovieScene::FScopedSignedObjectModifyDefer Defer(true);
+
+	ChannelType*             Channel     = this->GetChannelHandle().Get();
+	UMovieSceneSignedObject* SignedOwner = this->template GetOwningObjectOrOuter<UMovieSceneSignedObject>();
+	UMovieScene*             MovieScene  = SignedOwner ? SignedOwner->GetTypedOuter<UMovieScene>() : nullptr;
+
+	if (Channel && SignedOwner && MovieScene && !this->IsReadOnly())
 	{
 		bool bAutoSetTangents = false;
-		Section->MarkAsChanged();
 
 		TMovieSceneChannelData<ChannelValue> ChannelData = Channel->GetData();
 		TArrayView<ChannelValue> Values = ChannelData.GetValues();
 
-		FFrameRate TickResolution = Section->GetTypedOuter<UMovieScene>()->GetTickResolution();
+		FFrameRate TickResolution = MovieScene->GetTickResolution();
 		float TimeInterval = TickResolution.AsInterval();
 
 		for (int32 Index = 0; Index < InKeys.Num(); ++Index)
@@ -391,6 +408,7 @@ void FBezierChannelCurveModel<ChannelType, ChannelValue, KeyType>::SetKeyAttribu
 			Channel->AutoSetTangents();
 		}
 
+		SignedOwner->MarkAsChanged();
 		this->CurveModifiedDelegate.Broadcast();
 	}
 }
@@ -409,12 +427,11 @@ void FBezierChannelCurveModel<ChannelType, ChannelValue, KeyType>::GetCurveAttri
 template<typename ChannelType, typename ChannelValue, typename KeyType> 
 void FBezierChannelCurveModel<ChannelType, ChannelValue, KeyType>::SetCurveAttributes(const FCurveAttributes& InCurveAttributes)
 {
-	ChannelType* Channel		= this->GetChannelHandle().Get();
-	UMovieSceneSection* Section = Cast<UMovieSceneSection>(this->GetOwningObject());
-	if (Channel && Section && !this->IsReadOnly())
-	{
-		Section->MarkAsChanged();
+	ChannelType*             Channel     = this->GetChannelHandle().Get();
+	UMovieSceneSignedObject* SignedOwner = this->template GetOwningObjectOrOuter<UMovieSceneSignedObject>();
 
+	if (Channel && SignedOwner && !this->IsReadOnly())
+	{
 		if (InCurveAttributes.HasPreExtrapolation())
 		{
 			Channel->PreInfinityExtrap = InCurveAttributes.GetPreExtrapolation();
@@ -425,6 +442,7 @@ void FBezierChannelCurveModel<ChannelType, ChannelValue, KeyType>::SetCurveAttri
 			Channel->PostInfinityExtrap = InCurveAttributes.GetPostExtrapolation();
 		}
 
+		SignedOwner->MarkAsChanged();
 		this->CurveModifiedDelegate.Broadcast();
 	}
 }
@@ -470,9 +488,11 @@ template<typename ChannelType, typename ChannelValue, typename KeyType>
 void FBezierChannelCurveModel<ChannelType, ChannelValue, KeyType>::GetValueRange(double InMinTime, double InMaxTime, double& MinValue, double& MaxValue) const
 {
 	ChannelType* Channel = this->GetChannelHandle().Get();
-	UMovieSceneSection* Section = Cast<UMovieSceneSection>(this->GetOwningObject());
 
-	if (Channel && Section && Section->GetTypedOuter<UMovieScene>())
+	UMovieSceneSignedObject* SignedOwner = this->template GetOwningObjectOrOuter<UMovieSceneSignedObject>();
+	UMovieScene*             MovieScene  = SignedOwner ? SignedOwner->GetTypedOuter<UMovieScene>() : nullptr;
+
+	if (Channel && SignedOwner && MovieScene && !this->IsReadOnly())
 	{
 		TArrayView<const FFrameNumber> Times = Channel->GetData().GetTimes();
 		TArrayView<const ChannelValue> Values = Channel->GetData().GetValues();
@@ -484,7 +504,7 @@ void FBezierChannelCurveModel<ChannelType, ChannelValue, KeyType>::GetValueRange
 		}
 		else
 		{
-			FFrameRate TickResolution = Section->GetTypedOuter<UMovieScene>()->GetTickResolution();
+			FFrameRate TickResolution = MovieScene->GetTickResolution();
 			double ToTime = TickResolution.AsInterval();
 			int32 LastKeyIndex = Values.Num() - 1;
 			MinValue = TNumericLimits<double>::Max();
diff --git a/Engine/Source/Editor/MovieSceneTools/Private/Channels/BoolChannelCurveModel.cpp b/Engine/Source/Editor/MovieSceneTools/Private/Channels/BoolChannelCurveModel.cpp
--- a/Engine/Source/Editor/MovieSceneTools/Private/Channels/BoolChannelCurveModel.cpp
+++ b/Engine/Source/Editor/MovieSceneTools/Private/Channels/BoolChannelCurveModel.cpp
@@ -91,7 +91,7 @@ void FBoolChannelCurveModel::CreateKeyProxies(TArrayView<const FKeyHandle> InKey
 	{
 		UBoolChannelKeyProxy* NewProxy = NewObject<UBoolChannelKeyProxy>(GetTransientPackage(), NAME_None);
 
-		NewProxy->Initialize(InKeyHandles[Index], GetChannelHandle(), Cast<UMovieSceneSection>(GetOwningObject()));
+		NewProxy->Initialize(InKeyHandles[Index], GetChannelHandle(), this->GetOwningObjectOrOuter<UMovieSceneSection>());
 		OutObjects[Index] = NewProxy;
 	}
 }
@@ -117,7 +117,7 @@ TUniquePtr<IBufferedCurveModel> FBoolChannelCurveModel::CreateBufferedCurveCopy(
 		double ValueMin = 0.f, ValueMax = 1.f;
 		GetValueRange(ValueMin, ValueMax);
 
-		return MakeUnique<FBoolChannelBufferedCurveModel>(Channel, Cast<UMovieSceneSection>(GetOwningObject()), MoveTemp(KeyPositions), MoveTemp(KeyAttributes), GetLongDisplayName().ToString(), ValueMin, ValueMax);
+		return MakeUnique<FBoolChannelBufferedCurveModel>(Channel, this->GetOwningObjectOrOuter<UMovieSceneSection>(), MoveTemp(KeyPositions), MoveTemp(KeyAttributes), GetLongDisplayName().ToString(), ValueMin, ValueMax);
 	}
 	return nullptr;
 }
@@ -135,7 +135,7 @@ void FBoolChannelCurveModel::GetCurveAttributes(FCurveAttributes& OutCurveAttrib
 void FBoolChannelCurveModel::SetCurveAttributes(const FCurveAttributes& InCurveAttributes)
 {
 	FMovieSceneBoolChannel* Channel = GetChannelHandle().Get();
-	UMovieSceneSection* Section = Cast<UMovieSceneSection>(GetOwningObject());
+	UMovieSceneSection* Section = this->GetOwningObjectOrOuter<UMovieSceneSection>();
 	if (Channel && Section && !IsReadOnly())
 	{
 		Section->MarkAsChanged();
diff --git a/Engine/Source/Editor/MovieSceneTools/Private/Channels/BuiltInChannelEditors.cpp b/Engine/Source/Editor/MovieSceneTools/Private/Channels/BuiltInChannelEditors.cpp
--- a/Engine/Source/Editor/MovieSceneTools/Private/Channels/BuiltInChannelEditors.cpp
+++ b/Engine/Source/Editor/MovieSceneTools/Private/Channels/BuiltInChannelEditors.cpp
@@ -13,6 +13,7 @@
 #include "IKeyArea.h"
 #include "ISequencer.h"
 #include "SequencerSettings.h"
+#include "CurveEditor.h"
 #include "MovieSceneCommonHelpers.h"
 #include "GameFramework/Actor.h"
 #include "Styling/AppStyle.h"
@@ -34,6 +35,9 @@
 #include "Channels/FloatChannelCurveModel.h"
 #include "Channels/IntegerChannelCurveModel.h"
 #include "Channels/BoolChannelCurveModel.h"
+#include "Channels/TimeWarpChannelCurveModel.h"
+#include "Channels/PiecewiseCurveModel.h"
+#include "Variants/MovieScenePlayRateCurve.h"
 #include "EventChannelCurveModel.h"
 #include "PropertyCustomizationHelpers.h"
 #include "MovieSceneObjectBindingIDCustomization.h"
@@ -53,6 +57,7 @@
 #include "Widgets/Input/SComboButton.h"
 #include "MovieSceneSpawnableAnnotation.h"
 #include "ISequencerModule.h"
+#include "MovieScene.h"
 #include "MovieSceneTracksComponentTypes.h"
 
 #define LOCTEXT_NAMESPACE "BuiltInChannelEditors"
@@ -212,6 +217,12 @@ FKeyHandle AddOrUpdateKey(FMovieSceneDoubleChannel* Channel, UMovieSceneSection*
 	return AddOrUpdateKeyImpl<FMovieSceneDoubleChannel, double>(Channel, SectionToKey, ExternalValue, InTime, Sequencer, InObjectBindingID, PropertyBindings);
 }
 
+FKeyHandle AddOrUpdateKey(FMovieSceneTimeWarpChannel* Channel, UMovieSceneSection* SectionToKey, FFrameNumber InTime, ISequencer& Sequencer, const FGuid& InObjectBindingID, FTrackInstancePropertyBindings* PropertyBindings)
+{
+	TMovieSceneExternalValue<double> ExternalValue;
+	return AddOrUpdateKey(Channel, SectionToKey, ExternalValue, InTime, Sequencer, InObjectBindingID, PropertyBindings);
+}
+
 FKeyHandle AddOrUpdateKey(FMovieSceneActorReferenceData* Channel, UMovieSceneSection* SectionToKey, FFrameNumber InTime, ISequencer& Sequencer, const FGuid& InObjectBindingID, FTrackInstancePropertyBindings* PropertyBindings)
 {
 	if (PropertyBindings && InObjectBindingID.IsValid())
@@ -287,7 +298,7 @@ bool CanCreateKeyEditor(const FMovieSceneActorReferenceData* Channel)
 	return true;
 }
 
-TSharedRef<SWidget> CreateKeyEditor(const TMovieSceneChannelHandle<FMovieSceneBoolChannel>&    Channel, UMovieSceneSection* Section, const FGuid& InObjectBindingID, TWeakPtr<FTrackInstancePropertyBindings> PropertyBindings, TWeakPtr<ISequencer> InSequencer)
+TSharedRef<SWidget> CreateKeyEditor(const TMovieSceneChannelHandle<FMovieSceneBoolChannel>&    Channel, const UE::Sequencer::FCreateKeyEditorParams& Params)
 {
 	const TMovieSceneExternalValue<bool>* ExternalValue = Channel.GetExtendedEditorData();
 	if (!ExternalValue)
@@ -296,15 +307,15 @@ TSharedRef<SWidget> CreateKeyEditor(const TMovieSceneChannelHandle<FMovieSceneBo
 	}
 
 	TSequencerKeyEditor<FMovieSceneBoolChannel, bool> KeyEditor(
-		InObjectBindingID, Channel,
-		Section, InSequencer, PropertyBindings, ExternalValue->OnGetExternalValue
+		Params.ObjectBindingID, Channel,
+		Params.OwningSection, Params.Sequencer, Params.PropertyBindings, ExternalValue->OnGetExternalValue
 		);
 
 	return SNew(SBoolCurveKeyEditor, KeyEditor);
 }
 
 
-TSharedRef<SWidget> CreateKeyEditor(const TMovieSceneChannelHandle<FMovieSceneIntegerChannel>& Channel, UMovieSceneSection* Section, const FGuid& InObjectBindingID, TWeakPtr<FTrackInstancePropertyBindings> PropertyBindings, TWeakPtr<ISequencer> InSequencer)
+TSharedRef<SWidget> CreateKeyEditor(const TMovieSceneChannelHandle<FMovieSceneIntegerChannel>& Channel, const UE::Sequencer::FCreateKeyEditorParams& Params)
 {
 	const TMovieSceneExternalValue<int32>* ExternalValue = Channel.GetExtendedEditorData();
 	if (!ExternalValue)
@@ -313,8 +324,8 @@ TSharedRef<SWidget> CreateKeyEditor(const TMovieSceneChannelHandle<FMovieSceneIn
 	}
 
 	TSequencerKeyEditor<FMovieSceneIntegerChannel, int32> KeyEditor(
-		InObjectBindingID, Channel,
-		Section, InSequencer, PropertyBindings, ExternalValue->OnGetExternalValue
+		Params.ObjectBindingID, Channel,
+		Params.OwningSection, Params.Sequencer, Params.PropertyBindings, ExternalValue->OnGetExternalValue
 		);
 
 	typedef SNumericKeyEditor<FMovieSceneIntegerChannel, int32> KeyEditorType;
@@ -322,7 +333,7 @@ TSharedRef<SWidget> CreateKeyEditor(const TMovieSceneChannelHandle<FMovieSceneIn
 }
 
 
-TSharedRef<SWidget> CreateKeyEditor(const TMovieSceneChannelHandle<FMovieSceneFloatChannel>&   Channel, UMovieSceneSection* Section, const FGuid& InObjectBindingID, TWeakPtr<FTrackInstancePropertyBindings> PropertyBindings, TWeakPtr<ISequencer> InSequencer)
+TSharedRef<SWidget> CreateKeyEditor(const TMovieSceneChannelHandle<FMovieSceneFloatChannel>&   Channel, const UE::Sequencer::FCreateKeyEditorParams& Params)
 {
 	const TMovieSceneExternalValue<float>* ExternalValue = Channel.GetExtendedEditorData();
 	if (!ExternalValue)
@@ -331,8 +342,8 @@ TSharedRef<SWidget> CreateKeyEditor(const TMovieSceneChannelHandle<FMovieSceneFl
 	}
 
 	TSequencerKeyEditor<FMovieSceneFloatChannel, float> KeyEditor(
-		InObjectBindingID, Channel,
-		Section, InSequencer, PropertyBindings, ExternalValue->OnGetExternalValue
+		Params.ObjectBindingID, Channel,
+		Params.OwningSection, Params.Sequencer, Params.PropertyBindings, ExternalValue->OnGetExternalValue
 		);
 
 	typedef SNumericKeyEditor<FMovieSceneFloatChannel, float> KeyEditorType;
@@ -340,17 +351,17 @@ TSharedRef<SWidget> CreateKeyEditor(const TMovieSceneChannelHandle<FMovieSceneFl
 }
 
 
-TSharedRef<SWidget> CreateKeyEditor(const TMovieSceneChannelHandle<FMovieSceneDoubleChannel>&   Channel, UMovieSceneSection* Section, const FGuid& InObjectBindingID, TWeakPtr<FTrackInstancePropertyBindings> PropertyBindings, TWeakPtr<ISequencer> InSequencer)
+TSharedRef<SWidget> CreateKeyEditor(const TMovieSceneChannelHandle<FMovieSceneDoubleChannel>&   Channel, const UE::Sequencer::FCreateKeyEditorParams& Params)
 {
-	const TMovieSceneExternalValue<double>* ExternalValue = Channel.GetExtendedEditorData();
-	if (!ExternalValue)
+	TFunction<TOptional<double>(UObject&, FTrackInstancePropertyBindings*)> ExternalValue;
+	if (const TMovieSceneExternalValue<double>* ExternalValuePtr = Channel.GetExtendedEditorData())
 	{
-		return SNullWidget::NullWidget;
+		ExternalValue = ExternalValuePtr->OnGetExternalValue;
 	}
 
 	TSequencerKeyEditor<FMovieSceneDoubleChannel, double> KeyEditor(
-		InObjectBindingID, Channel,
-		Section, InSequencer, PropertyBindings, ExternalValue->OnGetExternalValue
+		Params.ObjectBindingID, Channel,
+		Params.OwningSection, Params.Sequencer, Params.PropertyBindings, ExternalValue
 		);
 
 	typedef SNumericKeyEditor<FMovieSceneDoubleChannel, double> KeyEditorType;
@@ -358,7 +369,7 @@ TSharedRef<SWidget> CreateKeyEditor(const TMovieSceneChannelHandle<FMovieSceneDo
 }
 
 
-TSharedRef<SWidget> CreateKeyEditor(const TMovieSceneChannelHandle<FMovieSceneStringChannel>&  Channel, UMovieSceneSection* Section, const FGuid& InObjectBindingID, TWeakPtr<FTrackInstancePropertyBindings> PropertyBindings, TWeakPtr<ISequencer> InSequencer)
+TSharedRef<SWidget> CreateKeyEditor(const TMovieSceneChannelHandle<FMovieSceneStringChannel>&  Channel, const UE::Sequencer::FCreateKeyEditorParams& Params)
 {
 	const TMovieSceneExternalValue<FString>* ExternalValue = Channel.GetExtendedEditorData();
 	if (!ExternalValue)
@@ -367,15 +378,15 @@ TSharedRef<SWidget> CreateKeyEditor(const TMovieSceneChannelHandle<FMovieSceneSt
 	}
 
 	TSequencerKeyEditor<FMovieSceneStringChannel, FString> KeyEditor(
-		InObjectBindingID, Channel,
-		Section, InSequencer, PropertyBindings, ExternalValue->OnGetExternalValue
+		Params.ObjectBindingID, Channel,
+		Params.OwningSection, Params.Sequencer, Params.PropertyBindings, ExternalValue->OnGetExternalValue
 		);
 
 	return SNew(SStringCurveKeyEditor, KeyEditor);
 }
 
 
-TSharedRef<SWidget> CreateKeyEditor(const TMovieSceneChannelHandle<FMovieSceneByteChannel>&    Channel, UMovieSceneSection* Section, const FGuid& InObjectBindingID, TWeakPtr<FTrackInstancePropertyBindings> PropertyBindings, TWeakPtr<ISequencer> InSequencer)
+TSharedRef<SWidget> CreateKeyEditor(const TMovieSceneChannelHandle<FMovieSceneByteChannel>&    Channel, const UE::Sequencer::FCreateKeyEditorParams& Params)
 {
 	const TMovieSceneExternalValue<uint8>* ExternalValue = Channel.GetExtendedEditorData();
 	const FMovieSceneByteChannel* RawChannel = Channel.Get();
@@ -385,8 +396,8 @@ TSharedRef<SWidget> CreateKeyEditor(const TMovieSceneChannelHandle<FMovieSceneBy
 	}
 
 	TSequencerKeyEditor<FMovieSceneByteChannel, uint8> KeyEditor(
-		InObjectBindingID, Channel,
-		Section, InSequencer, PropertyBindings, ExternalValue->OnGetExternalValue
+		Params.ObjectBindingID, Channel,
+		Params.OwningSection, Params.Sequencer, Params.PropertyBindings, ExternalValue->OnGetExternalValue
 		);
 
 	if (UEnum* Enum = RawChannel->GetEnum())
@@ -400,15 +411,15 @@ TSharedRef<SWidget> CreateKeyEditor(const TMovieSceneChannelHandle<FMovieSceneBy
 	}
 }
 
-TSharedRef<SWidget> CreateKeyEditor(const TMovieSceneChannelHandle<FMovieSceneObjectPathChannel>& Channel, UMovieSceneSection* Section, const FGuid& InObjectBindingID, TWeakPtr<FTrackInstancePropertyBindings> PropertyBindings, TWeakPtr<ISequencer> InSequencer)
+TSharedRef<SWidget> CreateKeyEditor(const TMovieSceneChannelHandle<FMovieSceneObjectPathChannel>& Channel, const UE::Sequencer::FCreateKeyEditorParams& Params)
 {
 	const TMovieSceneExternalValue<UObject*>* ExternalValue = Channel.GetExtendedEditorData();
 	const FMovieSceneObjectPathChannel*       RawChannel    = Channel.Get();
-	UMovieSceneObjectPropertyTrack* ObjectPathTrack = Cast<UMovieSceneObjectPropertyTrack>(Section->GetOuter());
+	UMovieSceneObjectPropertyTrack* ObjectPathTrack = Cast<UMovieSceneObjectPropertyTrack>(Params.OwningSection->GetOuter());
 
 	if (ExternalValue && RawChannel)
 	{
-		TSequencerKeyEditor<FMovieSceneObjectPathChannel, UObject*> KeyEditor(InObjectBindingID, Channel, Section, InSequencer, PropertyBindings, ExternalValue->OnGetExternalValue);
+		TSequencerKeyEditor<FMovieSceneObjectPathChannel, UObject*> KeyEditor(Params.ObjectBindingID, Channel, Params.OwningSection, Params.Sequencer, Params.PropertyBindings, ExternalValue->OnGetExternalValue);
 
 		UClass* PropertyClass = ObjectPathTrack ? ObjectPathTrack->PropertyClass : nullptr;
 		const bool bClassPicker = ObjectPathTrack ? ObjectPathTrack->bClassProperty : false;
@@ -462,11 +473,9 @@ TSharedRef<SWidget> CreateKeyEditor(const TMovieSceneChannelHandle<FMovieSceneOb
 			};
 
 			TArray<FAssetData> AssetDataArray;
-			if (InSequencer.IsValid())
-			{
-				UMovieSceneSequence* Sequence = InSequencer.Pin()->GetFocusedMovieSceneSequence();
-				AssetDataArray.Add((FAssetData)Sequence);
-			}
+
+			UMovieSceneSequence* Sequence = Params.Sequencer->GetFocusedMovieSceneSequence();
+			AssetDataArray.Add((FAssetData)Sequence);
 
 			return SNew(SHorizontalBox)
 
@@ -498,6 +507,29 @@ TSharedRef<SWidget> CreateKeyEditor(const TMovieSceneChannelHandle<FMovieSceneOb
 	return SNullWidget::NullWidget;
 }
 
+TSharedRef<SWidget> CreateKeyEditor(const TMovieSceneChannelHandle<FMovieSceneTimeWarpChannel>& Channel, const UE::Sequencer::FCreateKeyEditorParams& Params)
+{
+	TFunction<TOptional<double>(UObject&, FTrackInstancePropertyBindings*)> ExternalValue;
+
+	TSequencerKeyEditor<FMovieSceneTimeWarpChannel, double> KeyEditor(
+		Params.ObjectBindingID, Channel,
+		Params.OwningSection, Params.Sequencer, Params.PropertyBindings, ExternalValue
+		);
+
+	KeyEditor.SetOwningObject(Cast<UMovieSceneSignedObject>(Params.OwningObject));
+
+	const FMovieSceneTimeWarpChannel* ChannelPtr = Channel.Get();
+	if (ChannelPtr && ChannelPtr->Domain == ETimeWarpChannelDomain::Time)
+	{
+		// Set the numeric type interface for frame numbers if the channel is in the time domain
+		KeyEditor.SetNumericTypeInterface(Params.Sequencer->GetNumericTypeInterface());
+	}
+
+	typedef SNumericKeyEditor<FMovieSceneTimeWarpChannel, double> KeyEditorType;
+	return SNew(KeyEditorType, KeyEditor);
+}
+
+
 /** Delegate used to set a class */
 DECLARE_DELEGATE_OneParam(FOnSetActorReferenceKey, FMovieSceneActorReferenceKey);
 
@@ -617,7 +649,7 @@ private:
 };
 
 
-TSharedRef<SWidget> CreateKeyEditor(const TMovieSceneChannelHandle<FMovieSceneActorReferenceData>& Channel, UMovieSceneSection* Section, const FGuid& InObjectBindingID, TWeakPtr<FTrackInstancePropertyBindings> PropertyBindings, TWeakPtr<ISequencer> InSequencer)
+TSharedRef<SWidget> CreateKeyEditor(const TMovieSceneChannelHandle<FMovieSceneActorReferenceData>& Channel, const UE::Sequencer::FCreateKeyEditorParams& Params)
 {
 	const FMovieSceneActorReferenceData* RawChannel = Channel.Get();
 	if (!RawChannel)
@@ -627,7 +659,7 @@ TSharedRef<SWidget> CreateKeyEditor(const TMovieSceneChannelHandle<FMovieSceneAc
 
 	TFunction<TOptional<FMovieSceneActorReferenceKey>(UObject&, FTrackInstancePropertyBindings*)> Func;
 
-	TSequencerKeyEditor<FMovieSceneActorReferenceData, FMovieSceneActorReferenceKey> KeyEditor(InObjectBindingID, Channel, Section, InSequencer, PropertyBindings, Func);
+	TSequencerKeyEditor<FMovieSceneActorReferenceData, FMovieSceneActorReferenceKey> KeyEditor(Params.ObjectBindingID, Channel, Params.OwningSection, Params.Sequencer, Params.PropertyBindings, Func);
 
 	auto OnSetCurrentValueLambda = [KeyEditor](const FMovieSceneActorReferenceKey& ActorKey) mutable
 	{
@@ -730,7 +762,7 @@ TSharedRef<SWidget> CreateKeyEditor(const TMovieSceneChannelHandle<FMovieSceneAc
 		return KeyEditor.GetCurrentValue();
 	};
 
-	return SNew(SActorReferenceBox, InSequencer)
+	return SNew(SActorReferenceBox, Params.Sequencer)
 		.ActorReferenceKey_Lambda(GetCurrentValueLambda)
 		.OnSetActorReferenceKey_Lambda(OnSetCurrentValueLambda);
 }
@@ -773,6 +805,81 @@ UMovieSceneKeyStructType* InstanceGeneratedStruct(FMovieSceneByteChannel* Channe
 	return NewStruct;
 }
 
+UMovieSceneKeyStructType* InstanceGeneratedStruct(FMovieSceneTimeWarpChannel* Channel, FSequencerKeyStructGenerator* Generator)
+{
+	FName GeneratedTypeName("FMovieSceneTimeWarpChannelGeneratedStruct");
+
+	UMovieSceneKeyStructType* Existing = Generator->FindGeneratedStruct(GeneratedTypeName);
+	if (Existing)
+	{
+		return Existing;
+	}
+
+	UMovieSceneKeyStructType* NewStruct = FSequencerKeyStructGenerator::AllocateNewKeyStruct(FMovieSceneTimeWarpChannel::StaticStruct());
+	if (!NewStruct)
+	{
+		return nullptr;
+	}
+
+	FStructProperty* NewValueProperty = new FStructProperty(NewStruct, "Value", RF_NoFlags);
+	NewValueProperty->SetPropertyFlags(CPF_Edit);
+	NewValueProperty->SetMetaData("Category", TEXT("Key"));
+	NewValueProperty->Struct = TBaseStructure<FFrameNumber>::Get();
+
+	NewStruct->AddCppProperty(NewValueProperty);
+	NewStruct->DestValueProperty = NewValueProperty;
+
+	FSequencerKeyStructGenerator::FinalizeNewKeyStruct(NewStruct);
+
+	Generator->AddGeneratedStruct(GeneratedTypeName, NewStruct);
+	return NewStruct;
+}
+
+void PostConstructKeyInstance(const TMovieSceneChannelHandle<FMovieSceneTimeWarpChannel>& ChannelHandle, FKeyHandle InHandle, FStructOnScope* Struct)
+{
+	const UMovieSceneKeyStructType* GeneratedStructType = CastChecked<const UMovieSceneKeyStructType>(Struct->GetStruct());
+
+	uint8* StructMemory = Struct->GetStructMemory();
+
+	FStructProperty*     ValueProperty = CastFieldChecked<FStructProperty>(GeneratedStructType->DestValueProperty.Get());
+	FStructProperty*     TimeProperty  = CastFieldChecked<FStructProperty>(GeneratedStructType->DestTimeProperty.Get());
+
+	const FFrameNumber*  TimeAddress   = TimeProperty->ContainerPtrToValuePtr<FFrameNumber>(StructMemory);
+	FFrameNumber*        ValueAddress  = ValueProperty->ContainerPtrToValuePtr<FFrameNumber>(StructMemory);
+
+	// It is safe to capture the property and address in this lambda because the lambda is owned by the struct itself, so cannot be invoked if the struct has been destroyed
+	auto CopyInstanceToKeyLambda = [ChannelHandle, InHandle, GeneratedStructType, ValueProperty, ValueAddress, TimeAddress](const FPropertyChangedEvent&)
+	{
+		if (FMovieSceneTimeWarpChannel* DestinationChannel = ChannelHandle.Get())
+		{
+			const int32 KeyIndex = DestinationChannel->GetData().GetIndex(InHandle);
+			if (KeyIndex != INDEX_NONE)
+			{
+				DestinationChannel->GetData().GetValues()[KeyIndex].Value = ValueAddress->Value;
+
+				// Set the new key time
+				DestinationChannel->SetKeyTime(InHandle, *TimeAddress);
+			}
+		}
+	};
+
+	FGeneratedMovieSceneKeyStruct* KeyStruct = reinterpret_cast<FGeneratedMovieSceneKeyStruct*>(Struct->GetStructMemory());
+	KeyStruct->OnPropertyChangedEvent = CopyInstanceToKeyLambda;
+
+	// Copy the initial value for the struct
+	FMovieSceneTimeWarpChannel* Channel = ChannelHandle.Get();
+	if (Channel)
+	{
+		// Copy the initial value into the struct
+		const int32 KeyIndex = Channel->GetData().GetIndex(InHandle);
+		if (KeyIndex != INDEX_NONE)
+		{
+			double InitialTime = Channel->GetData().GetValues()[KeyIndex].Value;
+			*ValueAddress = FFrameTime::FromDecimal(InitialTime).FloorToFrame();
+		}
+	}
+}
+
 UMovieSceneKeyStructType* InstanceGeneratedStruct(FMovieSceneObjectPathChannel* Channel, FSequencerKeyStructGenerator* Generator)
 {
 	UClass* PropertyClass = Channel->GetPropertyClass();
@@ -818,7 +925,7 @@ void PostConstructKeyInstance(const TMovieSceneChannelHandle<FMovieSceneObjectPa
 
 	uint8* StructMemory = Struct->GetStructMemory();
 
-	FObjectProperty* ValueProperty = CastFieldChecked<FObjectProperty>(GeneratedStructType->DestValueProperty.Get());
+	FObjectProperty*     ValueProperty = CastFieldChecked<FObjectProperty>(GeneratedStructType->DestValueProperty.Get());
 	FStructProperty*     TimeProperty  = CastFieldChecked<FStructProperty>(GeneratedStructType->DestTimeProperty.Get());
 
 	const FFrameNumber*  TimeAddress   = TimeProperty->ContainerPtrToValuePtr<FFrameNumber>(StructMemory);
@@ -1182,20 +1289,6 @@ private:
 	TArray<TExtendKeyMenuParams<ChannelType>> ChannelAndHandles;
 };
 
-struct FFloatChannelKeyMenuExtension : TCurveChannelKeyMenuExtension<FMovieSceneFloatChannel>
-{
-	FFloatChannelKeyMenuExtension(TWeakPtr<ISequencer> InSequencer, TArray<TExtendKeyMenuParams<FMovieSceneFloatChannel>>&& InChannels)
-		: TCurveChannelKeyMenuExtension<FMovieSceneFloatChannel>(InSequencer, MoveTemp(InChannels))
-	{}
-};
-
-struct FDoubleChannelKeyMenuExtension : TCurveChannelKeyMenuExtension<FMovieSceneDoubleChannel>
-{
-	FDoubleChannelKeyMenuExtension(TWeakPtr<ISequencer> InSequencer, TArray<TExtendKeyMenuParams<FMovieSceneDoubleChannel>>&& InChannels)
-		: TCurveChannelKeyMenuExtension<FMovieSceneDoubleChannel>(InSequencer, MoveTemp(InChannels))
-	{}
-};
-
 struct FCurveChannelSectionMenuExtension : TSharedFromThis<FCurveChannelSectionMenuExtension>
 {
 	static TSharedRef<FCurveChannelSectionMenuExtension> GetOrCreate(TWeakPtr<ISequencer> InSequencer)
@@ -1824,41 +1917,56 @@ void ExtendSectionMenu(FMenuBuilder& OuterMenuBuilder, TSharedPtr<FExtender> Men
 
 void ExtendKeyMenu(FMenuBuilder& OuterMenuBuilder, TSharedPtr<FExtender> MenuExtender, TArray<TExtendKeyMenuParams<FMovieSceneFloatChannel>>&& Channels, TWeakPtr<ISequencer> InSequencer)
 {
-	TSharedRef<FFloatChannelKeyMenuExtension> Extension = MakeShared<FFloatChannelKeyMenuExtension>(InSequencer, MoveTemp(Channels));
+	using ExtensionType = TCurveChannelKeyMenuExtension<FMovieSceneFloatChannel>;
+	TSharedRef<ExtensionType> Extension = MakeShared<ExtensionType>(InSequencer, MoveTemp(Channels));
 
 	MenuExtender->AddMenuExtension("SequencerKeyEdit", EExtensionHook::After, nullptr, FMenuExtensionDelegate::CreateLambda([Extension](FMenuBuilder& MenuBuilder) { Extension->ExtendMenu(MenuBuilder); }));
 }
 
 void ExtendKeyMenu(FMenuBuilder& OuterMenuBuilder, TSharedPtr<FExtender> MenuExtender, TArray<TExtendKeyMenuParams<FMovieSceneDoubleChannel>>&& Channels, TWeakPtr<ISequencer> InSequencer)
 {
-	TSharedRef<FDoubleChannelKeyMenuExtension> Extension = MakeShared<FDoubleChannelKeyMenuExtension>(InSequencer, MoveTemp(Channels));
+	using ExtensionType = TCurveChannelKeyMenuExtension<FMovieSceneDoubleChannel>;
+	TSharedRef<ExtensionType> Extension = MakeShared<ExtensionType>(InSequencer, MoveTemp(Channels));
+
+	MenuExtender->AddMenuExtension("SequencerKeyEdit", EExtensionHook::After, nullptr, FMenuExtensionDelegate::CreateLambda([Extension](FMenuBuilder& MenuBuilder) { Extension->ExtendMenu(MenuBuilder); }));
+}
+
+void ExtendKeyMenu(FMenuBuilder& OuterMenuBuilder, TSharedPtr<FExtender> MenuExtender, TArray<TExtendKeyMenuParams<FMovieSceneTimeWarpChannel>>&& Channels, TWeakPtr<ISequencer> InSequencer)
+{
+	using ExtensionType = TCurveChannelKeyMenuExtension<FMovieSceneTimeWarpChannel>;
+	TSharedRef<ExtensionType> Extension = MakeShared<ExtensionType>(InSequencer, MoveTemp(Channels));
 
 	MenuExtender->AddMenuExtension("SequencerKeyEdit", EExtensionHook::After, nullptr, FMenuExtensionDelegate::CreateLambda([Extension](FMenuBuilder& MenuBuilder) { Extension->ExtendMenu(MenuBuilder); }));
 }
 
-TUniquePtr<FCurveModel> CreateCurveEditorModel(const TMovieSceneChannelHandle<FMovieSceneFloatChannel>& FloatChannel, UMovieSceneSection* OwningSection, TSharedRef<ISequencer> InSequencer)
+TUniquePtr<FCurveModel> CreateCurveEditorModel(const TMovieSceneChannelHandle<FMovieSceneFloatChannel>& FloatChannel, const UE::Sequencer::FCreateCurveEditorModelParams& Params)
+{
+	return MakeUnique<FFloatChannelCurveModel>(FloatChannel, Params.OwningSection, Params.Sequencer);
+}
+
+TUniquePtr<FCurveModel> CreateCurveEditorModel(const TMovieSceneChannelHandle<FMovieSceneDoubleChannel>& DoubleChannel, const UE::Sequencer::FCreateCurveEditorModelParams& Params)
 {
-	return MakeUnique<FFloatChannelCurveModel>(FloatChannel, OwningSection, InSequencer);
+	return MakeUnique<FDoubleChannelCurveModel>(DoubleChannel, Params.OwningSection, Params.OwningObject, Params.Sequencer);
 }
 
-TUniquePtr<FCurveModel> CreateCurveEditorModel(const TMovieSceneChannelHandle<FMovieSceneDoubleChannel>& DoubleChannel, UMovieSceneSection* OwningSection, TSharedRef<ISequencer> InSequencer)
+TUniquePtr<FCurveModel> CreateCurveEditorModel(const TMovieSceneChannelHandle<FMovieSceneIntegerChannel>& IntegerChannel, const UE::Sequencer::FCreateCurveEditorModelParams& Params)
 {
-	return MakeUnique<FDoubleChannelCurveModel>(DoubleChannel, OwningSection, InSequencer);
+	return MakeUnique<FIntegerChannelCurveModel>(IntegerChannel, Params.OwningSection, Params.Sequencer);
 }
 
-TUniquePtr<FCurveModel> CreateCurveEditorModel(const TMovieSceneChannelHandle<FMovieSceneIntegerChannel>& IntegerChannel, UMovieSceneSection* OwningSection, TSharedRef<ISequencer> InSequencer)
+TUniquePtr<FCurveModel> CreateCurveEditorModel(const TMovieSceneChannelHandle<FMovieSceneBoolChannel>& BoolChannel, const UE::Sequencer::FCreateCurveEditorModelParams& Params)
 {
-	return MakeUnique<FIntegerChannelCurveModel>(IntegerChannel, OwningSection, InSequencer);
+	return MakeUnique<FBoolChannelCurveModel>(BoolChannel, Params.OwningSection, Params.Sequencer);
 }
 
-TUniquePtr<FCurveModel> CreateCurveEditorModel(const TMovieSceneChannelHandle<FMovieSceneBoolChannel>& BoolChannel, UMovieSceneSection* OwningSection, TSharedRef<ISequencer> InSequencer)
+TUniquePtr<FCurveModel> CreateCurveEditorModel(const TMovieSceneChannelHandle<FMovieSceneEventChannel>& EventChannel, const UE::Sequencer::FCreateCurveEditorModelParams& Params)
 {
-	return MakeUnique<FBoolChannelCurveModel>(BoolChannel, OwningSection, InSequencer);
+	return MakeUnique<FEventChannelCurveModel>(EventChannel, Params.OwningSection, Params.Sequencer);
 }
 
-TUniquePtr<FCurveModel> CreateCurveEditorModel(const TMovieSceneChannelHandle<FMovieSceneEventChannel>& EventChannel, UMovieSceneSection* OwningSection, TSharedRef<ISequencer> InSequencer)
+TUniquePtr<FCurveModel> CreateCurveEditorModel(const TMovieSceneChannelHandle<FMovieSceneTimeWarpChannel>& TimeWarpChannel, const UE::Sequencer::FCreateCurveEditorModelParams& Params)
 {
-	return MakeUnique<FEventChannelCurveModel>(EventChannel, OwningSection, InSequencer);
+	return MakeUnique<FTimeWarpChannelCurveModel>(TimeWarpChannel, Params.OwningSection, Params.OwningObject, Params.Sequencer);
 }
 
 bool ShouldShowCurve(const FMovieSceneFloatChannel* Channel, UMovieSceneSection* InSection)
