Sequencer: Added time-warp capabilities to sequences, sub-sections and skeletal animation sections

diff --git a/Engine/Source/Editor/MovieSceneTools/Private/Channels/BuiltInChannelEditors.h b/Engine/Source/Editor/MovieSceneTools/Private/Channels/BuiltInChannelEditors.h
--- a/Engine/Source/Editor/MovieSceneTools/Private/Channels/BuiltInChannelEditors.h
+++ b/Engine/Source/Editor/MovieSceneTools/Private/Channels/BuiltInChannelEditors.h
@@ -20,6 +20,7 @@
 #include "Channels/MovieSceneIntegerChannel.h"
 #include "Channels/MovieSceneEventChannel.h"
 #include "Channels/MovieSceneObjectPathChannel.h"
+#include "Channels/MovieSceneTimeWarpChannel.h"
 #include "Sections/MovieSceneStringSection.h"
 #include "Sections/MovieSceneParticleSection.h"
 #include "Sections/MovieSceneActorReferenceSection.h"
@@ -37,6 +38,7 @@ class ISectionLayoutBuilder;
 FKeyHandle AddOrUpdateKey(FMovieSceneDoubleChannel* Channel, UMovieSceneSection* SectionToKey,  const TMovieSceneExternalValue<double>& EditorData, FFrameNumber InTime, ISequencer& Sequencer, const FGuid& InObjectBindingID, FTrackInstancePropertyBindings* PropertyBindings);
 FKeyHandle AddOrUpdateKey(FMovieSceneFloatChannel* Channel, UMovieSceneSection* SectionToKey,  const TMovieSceneExternalValue<float>& EditorData, FFrameNumber InTime, ISequencer& Sequencer, const FGuid& InObjectBindingID, FTrackInstancePropertyBindings* PropertyBindings);
 FKeyHandle AddOrUpdateKey(FMovieSceneActorReferenceData* Channel, UMovieSceneSection* SectionToKey, FFrameNumber InTime, ISequencer& Sequencer, const FGuid& InObjectBindingID, FTrackInstancePropertyBindings* PropertyBindings);
+FKeyHandle AddOrUpdateKey(FMovieSceneTimeWarpChannel* Channel, UMovieSceneSection* SectionToKey, FFrameNumber InTime, ISequencer& Sequencer, const FGuid& InObjectBindingID, FTrackInstancePropertyBindings* PropertyBindings);
 
 /** Key editor overrides */
 bool CanCreateKeyEditor(const FMovieSceneBoolChannel*       Channel);
@@ -48,22 +50,26 @@ bool CanCreateKeyEditor(const FMovieSceneStringChannel*     Channel);
 bool CanCreateKeyEditor(const FMovieSceneObjectPathChannel* Channel);
 bool CanCreateKeyEditor(const FMovieSceneActorReferenceData* Channel);
 
-TSharedRef<SWidget> CreateKeyEditor(const TMovieSceneChannelHandle<FMovieSceneBoolChannel>&        Channel, UMovieSceneSection* Section, const FGuid& InObjectBindingID, TWeakPtr<FTrackInstancePropertyBindings> PropertyBindings, TWeakPtr<ISequencer> InSequencer);
-TSharedRef<SWidget> CreateKeyEditor(const TMovieSceneChannelHandle<FMovieSceneByteChannel>&        Channel, UMovieSceneSection* Section, const FGuid& InObjectBindingID, TWeakPtr<FTrackInstancePropertyBindings> PropertyBindings, TWeakPtr<ISequencer> InSequencer);
-TSharedRef<SWidget> CreateKeyEditor(const TMovieSceneChannelHandle<FMovieSceneIntegerChannel>&     Channel, UMovieSceneSection* Section, const FGuid& InObjectBindingID, TWeakPtr<FTrackInstancePropertyBindings> PropertyBindings, TWeakPtr<ISequencer> InSequencer);
-TSharedRef<SWidget> CreateKeyEditor(const TMovieSceneChannelHandle<FMovieSceneDoubleChannel>&       Channel, UMovieSceneSection* Section, const FGuid& InObjectBindingID, TWeakPtr<FTrackInstancePropertyBindings> PropertyBindings, TWeakPtr<ISequencer> InSequencer);
-TSharedRef<SWidget> CreateKeyEditor(const TMovieSceneChannelHandle<FMovieSceneFloatChannel>&       Channel, UMovieSceneSection* Section, const FGuid& InObjectBindingID, TWeakPtr<FTrackInstancePropertyBindings> PropertyBindings, TWeakPtr<ISequencer> InSequencer);
-TSharedRef<SWidget> CreateKeyEditor(const TMovieSceneChannelHandle<FMovieSceneStringChannel>&      Channel, UMovieSceneSection* Section, const FGuid& InObjectBindingID, TWeakPtr<FTrackInstancePropertyBindings> PropertyBindings, TWeakPtr<ISequencer> InSequencer);
-TSharedRef<SWidget> CreateKeyEditor(const TMovieSceneChannelHandle<FMovieSceneObjectPathChannel>&  Channel, UMovieSceneSection* Section, const FGuid& InObjectBindingID, TWeakPtr<FTrackInstancePropertyBindings> PropertyBindings, TWeakPtr<ISequencer> InSequencer);
-TSharedRef<SWidget> CreateKeyEditor(const TMovieSceneChannelHandle<FMovieSceneActorReferenceData>&  Channel, UMovieSceneSection* Section, const FGuid& InObjectBindingID, TWeakPtr<FTrackInstancePropertyBindings> PropertyBindings, TWeakPtr<ISequencer> InSequencer);
+TSharedRef<SWidget> CreateKeyEditor(const TMovieSceneChannelHandle<FMovieSceneBoolChannel>&        Channel, const UE::Sequencer::FCreateKeyEditorParams& Params);
+TSharedRef<SWidget> CreateKeyEditor(const TMovieSceneChannelHandle<FMovieSceneByteChannel>&        Channel, const UE::Sequencer::FCreateKeyEditorParams& Params);
+TSharedRef<SWidget> CreateKeyEditor(const TMovieSceneChannelHandle<FMovieSceneIntegerChannel>&     Channel, const UE::Sequencer::FCreateKeyEditorParams& Params);
+TSharedRef<SWidget> CreateKeyEditor(const TMovieSceneChannelHandle<FMovieSceneDoubleChannel>&      Channel, const UE::Sequencer::FCreateKeyEditorParams& Params);
+TSharedRef<SWidget> CreateKeyEditor(const TMovieSceneChannelHandle<FMovieSceneFloatChannel>&       Channel, const UE::Sequencer::FCreateKeyEditorParams& Params);
+TSharedRef<SWidget> CreateKeyEditor(const TMovieSceneChannelHandle<FMovieSceneStringChannel>&      Channel, const UE::Sequencer::FCreateKeyEditorParams& Params);
+TSharedRef<SWidget> CreateKeyEditor(const TMovieSceneChannelHandle<FMovieSceneObjectPathChannel>&  Channel, const UE::Sequencer::FCreateKeyEditorParams& Params);
+TSharedRef<SWidget> CreateKeyEditor(const TMovieSceneChannelHandle<FMovieSceneActorReferenceData>& Channel, const UE::Sequencer::FCreateKeyEditorParams& Params);
+TSharedRef<SWidget> CreateKeyEditor(const TMovieSceneChannelHandle<FMovieSceneTimeWarpChannel>&    Channel, const UE::Sequencer::FCreateKeyEditorParams& Params);
 
 UMovieSceneKeyStructType* InstanceGeneratedStruct(FMovieSceneByteChannel* Channel,       FSequencerKeyStructGenerator* Generator);
+UMovieSceneKeyStructType* InstanceGeneratedStruct(FMovieSceneTimeWarpChannel* Channel,   FSequencerKeyStructGenerator* Generator);
 UMovieSceneKeyStructType* InstanceGeneratedStruct(FMovieSceneObjectPathChannel* Channel, FSequencerKeyStructGenerator* Generator);
 
+void PostConstructKeyInstance(const TMovieSceneChannelHandle<FMovieSceneTimeWarpChannel>& ChannelHandle, FKeyHandle InHandle, FStructOnScope* Struct);
 void PostConstructKeyInstance(const TMovieSceneChannelHandle<FMovieSceneObjectPathChannel>& ChannelHandle, FKeyHandle InHandle, FStructOnScope* Struct);
 
+
 /** Key drawing overrides */
-void DrawKeys(FMovieSceneDoubleChannel*    Channel, TArrayView<const FKeyHandle> InKeyHandles, const UMovieSceneSection* InOwner, TArrayView<FKeyDrawParams> OutKeyDrawParams);
+void DrawKeys(FMovieSceneDoubleChannel*   Channel, TArrayView<const FKeyHandle> InKeyHandles, const UMovieSceneSection* InOwner, TArrayView<FKeyDrawParams> OutKeyDrawParams);
 void DrawKeys(FMovieSceneFloatChannel*    Channel, TArrayView<const FKeyHandle> InKeyHandles, const UMovieSceneSection* InOwner, TArrayView<FKeyDrawParams> OutKeyDrawParams);
 void DrawKeys(FMovieSceneParticleChannel* Channel, TArrayView<const FKeyHandle> InKeyHandles, const UMovieSceneSection* InOwner, TArrayView<FKeyDrawParams> OutKeyDrawParams);
 void DrawKeys(FMovieSceneEventChannel*    Channel, TArrayView<const FKeyHandle> InKeyHandles, const UMovieSceneSection* InOwner, TArrayView<FKeyDrawParams> OutKeyDrawParams);
@@ -75,6 +81,7 @@ void ExtendSectionMenu(FMenuBuilder& OuterMenuBuilder, TSharedPtr<FExtender> Men
 void ExtendSectionMenu(FMenuBuilder& OuterMenuBuilder, TSharedPtr<FExtender> MenuExtender, TArray<TMovieSceneChannelHandle<FMovieSceneBoolChannel>>&& Channels, TArrayView<UMovieSceneSection* const> Sections, TWeakPtr<ISequencer> InSequencer);
 void ExtendKeyMenu(FMenuBuilder& OuterMenuBuilder, TSharedPtr<FExtender> MenuExtender, TArray<TExtendKeyMenuParams<FMovieSceneDoubleChannel>>&& Channels, TWeakPtr<ISequencer> InSequencer);
 void ExtendKeyMenu(FMenuBuilder& OuterMenuBuilder, TSharedPtr<FExtender> MenuExtender, TArray<TExtendKeyMenuParams<FMovieSceneFloatChannel>>&& Channels, TWeakPtr<ISequencer> InSequencer);
+void ExtendKeyMenu(FMenuBuilder& OuterMenuBuilder, TSharedPtr<FExtender> MenuExtender, TArray<TExtendKeyMenuParams<FMovieSceneTimeWarpChannel>>&& Channels, TWeakPtr<ISequencer> InSequencer);
 
 /** Curve editor models */
 inline bool SupportsCurveEditorModels(const TMovieSceneChannelHandle<FMovieSceneDoubleChannel>& DoubleChannel) { return true; }
@@ -82,12 +89,14 @@ inline bool SupportsCurveEditorModels(const TMovieSceneChannelHandle<FMovieScene
 inline bool SupportsCurveEditorModels(const TMovieSceneChannelHandle<FMovieSceneIntegerChannel>& IntegerChannel) { return true; }
 inline bool SupportsCurveEditorModels(const TMovieSceneChannelHandle<FMovieSceneBoolChannel>& BoolChannel) { return true; }
 inline bool SupportsCurveEditorModels(const TMovieSceneChannelHandle<FMovieSceneEventChannel>& EventChannel) { return true; }
-
-TUniquePtr<FCurveModel> CreateCurveEditorModel(const TMovieSceneChannelHandle<FMovieSceneDoubleChannel>& DoubleChannel, UMovieSceneSection* OwningSection, TSharedRef<ISequencer> InSequencer);
-TUniquePtr<FCurveModel> CreateCurveEditorModel(const TMovieSceneChannelHandle<FMovieSceneFloatChannel>& FloatChannel, UMovieSceneSection* OwningSection, TSharedRef<ISequencer> InSequencer);
-TUniquePtr<FCurveModel> CreateCurveEditorModel(const TMovieSceneChannelHandle<FMovieSceneIntegerChannel>& IntegerChannel, UMovieSceneSection* OwningSection, TSharedRef<ISequencer> InSequencer);
-TUniquePtr<FCurveModel> CreateCurveEditorModel(const TMovieSceneChannelHandle<FMovieSceneBoolChannel>& BoolChannel, UMovieSceneSection* OwningSection, TSharedRef<ISequencer> InSequencer);
-TUniquePtr<FCurveModel> CreateCurveEditorModel(const TMovieSceneChannelHandle<FMovieSceneEventChannel>& EventChannel, UMovieSceneSection* OwningSection, TSharedRef<ISequencer> InSequencer);
+inline bool SupportsCurveEditorModels(const TMovieSceneChannelHandle<FMovieSceneTimeWarpChannel>& TimeWarpChannel) { return true; }
+
+TUniquePtr<FCurveModel> CreateCurveEditorModel(const TMovieSceneChannelHandle<FMovieSceneDoubleChannel>& DoubleChannel, const UE::Sequencer::FCreateCurveEditorModelParams& Params);
+TUniquePtr<FCurveModel> CreateCurveEditorModel(const TMovieSceneChannelHandle<FMovieSceneFloatChannel>& FloatChannel, const UE::Sequencer::FCreateCurveEditorModelParams& Params);
+TUniquePtr<FCurveModel> CreateCurveEditorModel(const TMovieSceneChannelHandle<FMovieSceneIntegerChannel>& IntegerChannel, const UE::Sequencer::FCreateCurveEditorModelParams& Params);
+TUniquePtr<FCurveModel> CreateCurveEditorModel(const TMovieSceneChannelHandle<FMovieSceneBoolChannel>& BoolChannel, const UE::Sequencer::FCreateCurveEditorModelParams& Params);
+TUniquePtr<FCurveModel> CreateCurveEditorModel(const TMovieSceneChannelHandle<FMovieSceneEventChannel>& EventChannel, const UE::Sequencer::FCreateCurveEditorModelParams& Params);
+TUniquePtr<FCurveModel> CreateCurveEditorModel(const TMovieSceneChannelHandle<FMovieSceneTimeWarpChannel>& TimeWarpChannel, const UE::Sequencer::FCreateCurveEditorModelParams& Params);
 
 bool ShouldShowCurve(const FMovieSceneFloatChannel* Channel, UMovieSceneSection* InSection);
 bool ShouldShowCurve(const FMovieSceneDoubleChannel* Channel, UMovieSceneSection* InSection);
\ No newline at end of file
diff --git a/Engine/Source/Editor/MovieSceneTools/Private/Channels/ChannelCurveModel.cpp b/Engine/Source/Editor/MovieSceneTools/Private/Channels/ChannelCurveModel.cpp
--- a/Engine/Source/Editor/MovieSceneTools/Private/Channels/ChannelCurveModel.cpp
+++ b/Engine/Source/Editor/MovieSceneTools/Private/Channels/ChannelCurveModel.cpp
@@ -48,6 +48,13 @@ FChannelCurveModel<ChannelType, ChannelValue, KeyType>::FChannelCurveModel(TMovi
 	SupportedViews = ECurveEditorViewID::Absolute | ECurveEditorViewID::Normalized | ECurveEditorViewID::Stacked;
 }
 
+template <class ChannelType, class ChannelValue, class KeyType>
+FChannelCurveModel<ChannelType, ChannelValue, KeyType>::FChannelCurveModel(TMovieSceneChannelHandle<ChannelType> InChannel, UMovieSceneSection* InOwningSection, UObject* InOwningObject, TWeakPtr<ISequencer> InWeakSequencer)
+	: FChannelCurveModel(InChannel, InOwningSection, InWeakSequencer)
+{
+	WeakOwningObject = InOwningObject;
+}
+
 template <class ChannelType, class ChannelValue, class KeyType>
 FChannelCurveModel<ChannelType, ChannelValue, KeyType>::~FChannelCurveModel()
 {
@@ -57,6 +64,27 @@ FChannelCurveModel<ChannelType, ChannelValue, KeyType>::~FChannelCurveModel()
 	}
 }
 
+template <class ChannelType, class ChannelValue, class KeyType>
+FTransform2d FChannelCurveModel<ChannelType, ChannelValue, KeyType>::GetCurveTransform() const
+{
+	FTransform2d Transform;
+
+	const FMovieSceneChannelMetaData* MetaData = ChannelHandle.GetMetaData();
+	const UMovieSceneSection*         Section  = WeakSection.Get();
+
+	if (MetaData && Section)
+	{
+		FFrameNumber Offset = MetaData->GetOffsetTime(Section);
+		if (Offset != 0)
+		{
+			FFrameRate TickResolution = Section->GetTypedOuter<UMovieScene>()->GetTickResolution();
+			Transform = Concatenate(Transform, FVector2d(-Offset / TickResolution, 0.0));
+		}
+	}
+	return Transform;
+}
+
+
 template <class ChannelType, class ChannelValue, class KeyType>
 const void* FChannelCurveModel<ChannelType, ChannelValue, KeyType>::GetCurve() const
 {
@@ -66,9 +94,9 @@ const void* FChannelCurveModel<ChannelType, ChannelValue, KeyType>::GetCurve() c
 template <class ChannelType, class ChannelValue, class KeyType>
 void FChannelCurveModel<ChannelType, ChannelValue, KeyType>::Modify()
 {
-	if (UMovieSceneSection* Section = WeakSection.Get())
+	if (UObject* Object = GetOwningObject())
 	{
-		Section->Modify();
+		Object->Modify();
 	}
 	LastSignature.Invalidate();
 }
@@ -218,10 +246,19 @@ void FChannelCurveModel<ChannelType, ChannelValue, KeyType>::GetKeyPositions(TAr
 template <class ChannelType, class ChannelValue, class KeyType>
 void FChannelCurveModel<ChannelType, ChannelValue, KeyType>::SetKeyPositions(TArrayView<const FKeyHandle> InKeys, TArrayView<const FKeyPosition> InKeyPositions, EPropertyChangeType::Type ChangeType)
 {
+	UE::MovieScene::FScopedSignedObjectModifyDefer Defer(true);
+
 	ChannelType* Channel = ChannelHandle.Get();
 	UMovieSceneSection* Section = WeakSection.Get();
+	const FMovieSceneChannelMetaData* MetaData = ChannelHandle.GetMetaData();
 
-	if (Channel && Section && !IsReadOnly())
+	UMovieSceneSignedObject* SignedOwner = Cast<UMovieSceneSignedObject>(WeakOwningObject.Get());
+	if (!SignedOwner)
+	{
+		SignedOwner = WeakSection.Get();
+	}
+
+	if (Channel && MetaData && Section && !IsReadOnly())
 	{
 		TMovieSceneChannelData<ChannelValue> ChannelData = Channel->GetData();
 		const TArrayView<const FFrameNumber> Times = ChannelData.GetTimes();
@@ -229,53 +266,64 @@ void FChannelCurveModel<ChannelType, ChannelValue, KeyType>::SetKeyPositions(TAr
 		{
 			return;
 		}
-		
+
 		Section->MarkAsChanged();
 
-		FFrameRate TickResolution = Section->GetTypedOuter<UMovieScene>()->GetTickResolution();
+		FFrameRate   TickResolution = Section->GetTypedOuter<UMovieScene>()->GetTickResolution();
+		FFrameNumber KeyOffset      = 0;
 
 		const int32 LastIndex = Times.Num()-1;
 		const FFrameNumber FirstTime = Times[0];
 		const FFrameNumber LastTime = Times[LastIndex];
-		
-		FFrameNumber Min = TNumericLimits<FFrameNumber>::Max();
-		FFrameNumber Max = TNumericLimits<FFrameNumber>::Lowest();
-		
-		for (int32 Index = 0; Index < InKeys.Num(); ++Index)
+
+		// Expand to frame first, then offset
+		if (InKeyPositions.Num() > 0)
 		{
-			int32 KeyIndex = ChannelData.GetIndex(InKeys[Index]);
-			if (KeyIndex != INDEX_NONE)
+			double Min = TNumericLimits<double>::Max();
+			double Max = TNumericLimits<double>::Lowest();
+
+			for (const FKeyPosition& Value : InKeyPositions)
 			{
-				FFrameNumber NewTime = (InKeyPositions[Index].InputValue * TickResolution).RoundToFrame();
+				Min = FMath::Min(Min, Value.InputValue);
+				Max = FMath::Max(Max, Value.InputValue);
+			}
 
-				KeyIndex = ChannelData.MoveKey(KeyIndex, NewTime);
-				SetKeyValue(KeyIndex, InKeyPositions[Index].OutputValue);
+			FFrameNumber Offset  = MetaData->GetOffsetTime(Section);
+			FFrameNumber MinTime = (Min * TickResolution).RoundToFrame() + Offset;
+			FFrameNumber MaxTime = (Max * TickResolution).RoundToFrame() + Offset;
 
-				if ((KeyIndex == 0 && NewTime != FirstTime) || (KeyIndex == LastIndex && NewTime != LastTime))
-				{
-					Min = FMath::Min(Min, NewTime);
-					Max = FMath::Max(Max, NewTime);
-				}
+			if (!Section->GetRange().Contains(MinTime))
+			{
+				Section->ExpandToFrame(MinTime);
+				KeyOffset += MetaData->GetOffsetTime(Section) - Offset;
+			}
+
+			if (Min != Max && !Section->GetRange().Contains(MaxTime))
+			{
+				Section->ExpandToFrame(MaxTime);
 			}
 		}
 
-		// update range if needed
-		if (Min != TNumericLimits<FFrameNumber>::Max())
+		for (int32 Index = 0; Index < InKeys.Num(); ++Index)
 		{
-			Section->ExpandToFrame(Min);
-			if (Max != Min)
+			int32 KeyIndex = ChannelData.GetIndex(InKeys[Index]);
+			if (KeyIndex != INDEX_NONE)
 			{
-				Section->ExpandToFrame(Max);
+				FFrameNumber NewTime = (InKeyPositions[Index].InputValue * TickResolution).RoundToFrame() - KeyOffset;
+
+				KeyIndex = ChannelData.MoveKey(KeyIndex, NewTime);
+				SetKeyValue(KeyIndex, InKeyPositions[Index].OutputValue);
 			}
 		}
 		
 		Channel->PostEditChange();
 		if(WeakSequencer.IsValid())
 		{ 
-			const FMovieSceneChannelMetaData* MetaData = ChannelHandle.GetMetaData();
 			WeakSequencer.Pin()->OnChannelChanged().Broadcast(MetaData, Section);
 		}
 		CurveModifiedDelegate.Broadcast();
+
+		SignedOwner->MarkAsChanged();
 	}
 }
 
@@ -413,9 +461,17 @@ void FChannelCurveModel<ChannelType, ChannelValue, KeyType>::AddKeys(TArrayView<
 
 	ChannelType* Channel = ChannelHandle.Get();
 	UMovieSceneSection* Section = WeakSection.Get();
-	if (Channel && Section && !IsReadOnly())
+	const FMovieSceneChannelMetaData* MetaData = ChannelHandle.GetMetaData();
+
+	UMovieSceneSignedObject* SignedOwner = Cast<UMovieSceneSignedObject>(WeakOwningObject.Get());
+	if (!SignedOwner)
 	{
-		Section->Modify();
+		SignedOwner = WeakSection.Get();
+	}
+
+	if (Channel && MetaData && Section && !IsReadOnly())
+	{
+		SignedOwner->Modify();
 
 		TMovieSceneChannelData<ChannelValue> ChannelData = Channel->GetData();
 		FFrameRate TickResolution = Section->GetTypedOuter<UMovieScene>()->GetTickResolution();
@@ -449,8 +505,8 @@ void FChannelCurveModel<ChannelType, ChannelValue, KeyType>::AddKeys(TArrayView<
 
 		if (InKeyPositions.Num() > 0)
 		{
-			Section->ExpandToFrame(MinFrame);
-			Section->ExpandToFrame(MaxFrame);
+			FFrameNumber Offset = MetaData->GetOffsetTime(Section);
+			Section->ExpandToFrame(MinFrame + Offset);
 		}
 
 		// We reuse SetKeyAttributes here as there is complex logic determining which parts of the attributes are valid to set.
@@ -459,7 +515,6 @@ void FChannelCurveModel<ChannelType, ChannelValue, KeyType>::AddKeys(TArrayView<
 		Channel->PostEditChange();
 		if (WeakSequencer.IsValid())
 		{
-			const FMovieSceneChannelMetaData* MetaData = ChannelHandle.GetMetaData();
 			WeakSequencer.Pin()->OnChannelChanged().Broadcast(MetaData, Section);
 		}
 		CurveModifiedDelegate.Broadcast();
@@ -471,9 +526,15 @@ void FChannelCurveModel<ChannelType, ChannelValue, KeyType>::RemoveKeys(TArrayVi
 {
 	ChannelType* Channel = ChannelHandle.Get();
 	UMovieSceneSection* Section = WeakSection.Get();
+	UMovieSceneSignedObject* SignedOwner = Cast<UMovieSceneSignedObject>(WeakOwningObject.Get());
+	if (!SignedOwner)
+	{
+		SignedOwner = WeakSection.Get();
+	}
+
 	if (Channel && Section && !IsReadOnly())
 	{
-		Section->Modify();
+		SignedOwner->Modify();
 
 		TMovieSceneChannelData<ChannelValue> ChannelData = Channel->GetData();
 
@@ -513,7 +574,7 @@ void FChannelCurveModel<ChannelType, ChannelValue, KeyType>::FixupCurve()
 	if (UMovieSceneSection* Section = WeakSection.Get())
 	{
 		FMovieSceneChannelProxy* NewChannelProxy = &Section->GetChannelProxy();
-		ChannelHandle = NewChannelProxy->MakeHandle<ChannelType>(ChannelHandle.GetChannelIndex());
+		ChannelHandle = NewChannelProxy->CopyHandle(ChannelHandle);
 		OnDestroyHandle = NewChannelProxy->OnDestroy.AddRaw(this, &FChannelCurveModel<ChannelType, ChannelValue, KeyType>::FixupCurve);
 	}
 }
diff --git a/Engine/Source/Editor/MovieSceneTools/Private/Channels/ConstraintChannelEditor.cpp b/Engine/Source/Editor/MovieSceneTools/Private/Channels/ConstraintChannelEditor.cpp
--- a/Engine/Source/Editor/MovieSceneTools/Private/Channels/ConstraintChannelEditor.cpp
+++ b/Engine/Source/Editor/MovieSceneTools/Private/Channels/ConstraintChannelEditor.cpp
@@ -352,17 +352,12 @@ bool CanCreateKeyEditor(const FMovieSceneConstraintChannel* InChannel)
 	return false;
 }
 
-TSharedRef<SWidget> CreateKeyEditor(
-	const TMovieSceneChannelHandle<FMovieSceneConstraintChannel>& InChannel,
-	UMovieSceneSection* InSection,
-	const FGuid& InObjectBindingID,
-	TWeakPtr<FTrackInstancePropertyBindings> PropertyBindings,
-	TWeakPtr<ISequencer> InSequencer)
+TSharedRef<SWidget> CreateKeyEditor(const TMovieSceneChannelHandle<FMovieSceneConstraintChannel>& InChannel, const UE::Sequencer::FCreateKeyEditorParams& Params)
 {
 	return SNullWidget::NullWidget;
 }
 
-TUniquePtr<FCurveModel> CreateCurveEditorModel(const TMovieSceneChannelHandle<FMovieSceneConstraintChannel>& ChannelHandle, UMovieSceneSection* OwningSection, TSharedRef<ISequencer> InSequencer)
+TUniquePtr<FCurveModel> CreateCurveEditorModel(const TMovieSceneChannelHandle<FMovieSceneConstraintChannel>& ChannelHandle, const UE::Sequencer::FCreateCurveEditorModelParams& Params)
 {
 	if (FMovieSceneConstraintChannel* Channel = ChannelHandle.Get())
 	{
@@ -371,7 +366,7 @@ TUniquePtr<FCurveModel> CreateCurveEditorModel(const TMovieSceneChannelHandle<FM
 			const UCurveEditorSettings* Settings = GetDefault<UCurveEditorSettings>();
 			if (Settings == nullptr || Settings->GetShowBars())
 			{
-				return MakeUnique<FConstraintChannelCurveModel>(ChannelHandle, OwningSection, InSequencer);
+				return MakeUnique<FConstraintChannelCurveModel>(ChannelHandle, Params.OwningSection, Params.Sequencer);
 			}
 		}
 	}
diff --git a/Engine/Source/Editor/MovieSceneTools/Private/Channels/ConstraintChannelEditor.h b/Engine/Source/Editor/MovieSceneTools/Private/Channels/ConstraintChannelEditor.h
--- a/Engine/Source/Editor/MovieSceneTools/Private/Channels/ConstraintChannelEditor.h
+++ b/Engine/Source/Editor/MovieSceneTools/Private/Channels/ConstraintChannelEditor.h
@@ -65,11 +65,8 @@ bool CanCreateKeyEditor(const FMovieSceneConstraintChannel* InChannel);
 
 TSharedRef<SWidget> CreateKeyEditor(
 	const TMovieSceneChannelHandle<FMovieSceneConstraintChannel>& InChannel,
-	UMovieSceneSection* InSection,
-	const FGuid& InObjectBindingID,
-	TWeakPtr<FTrackInstancePropertyBindings> PropertyBindings,
-	TWeakPtr<ISequencer> Sequencer);
+	const UE::Sequencer::FCreateKeyEditorParams& Params);
 
 /** Curve editor models */
 inline bool SupportsCurveEditorModels(const TMovieSceneChannelHandle<FMovieSceneConstraintChannel>& Channel) { return true; }
-TUniquePtr<FCurveModel> CreateCurveEditorModel(const TMovieSceneChannelHandle<FMovieSceneConstraintChannel>& Channel, UMovieSceneSection* OwningSection, TSharedRef<ISequencer> InSequencer);
+TUniquePtr<FCurveModel> CreateCurveEditorModel(const TMovieSceneChannelHandle<FMovieSceneConstraintChannel>& Channel, const UE::Sequencer::FCreateCurveEditorModelParams& Params);
diff --git a/Engine/Source/Editor/MovieSceneTools/Private/Channels/DoubleChannelCurveModel.cpp b/Engine/Source/Editor/MovieSceneTools/Private/Channels/DoubleChannelCurveModel.cpp
--- a/Engine/Source/Editor/MovieSceneTools/Private/Channels/DoubleChannelCurveModel.cpp
+++ b/Engine/Source/Editor/MovieSceneTools/Private/Channels/DoubleChannelCurveModel.cpp
@@ -6,6 +6,7 @@
 #include "Channels/MovieSceneChannelData.h"
 #include "Channels/MovieSceneChannelHandle.h"
 #include "Channels/MovieSceneDoubleChannel.h"
+#include "Channels/MovieSceneInterpolation.h"
 #include "Containers/Array.h"
 #include "Containers/UnrealString.h"
 #include "CurveDataAbstraction.h"
@@ -73,13 +74,64 @@ FDoubleChannelCurveModel::FDoubleChannelCurveModel(TMovieSceneChannelHandle<FMov
 {
 }
 
+FDoubleChannelCurveModel::FDoubleChannelCurveModel(TMovieSceneChannelHandle<FMovieSceneDoubleChannel> InChannel, UMovieSceneSection* OwningSection, UObject* InOwningObject, TWeakPtr<ISequencer> InWeakSequencer)
+	: FBezierChannelCurveModel<FMovieSceneDoubleChannel, FMovieSceneDoubleValue, double>(InChannel, OwningSection, InOwningObject, InWeakSequencer)
+{
+}
+
+void FDoubleChannelCurveModel::GetValueRange(double& MinValue, double& MaxValue) const
+{
+	FMovieSceneDoubleChannel* Channel = GetChannelHandle().Get();
+	UMovieSceneSection* Section = WeakSection.Get();
+
+	if (Channel && Section)
+	{
+		TArrayView<const FFrameNumber> Times = Channel->GetData().GetTimes();
+		FFrameRate TickResolution = Section->GetTypedOuter<UMovieScene>()->GetTickResolution();
+
+		if (Times.Num() > 0)
+		{
+			UE::MovieScene::Interpolation::FInterpolationExtents Extents =
+				Channel->ComputeExtents(Times[0].Value, Times.Last().Value);
+			MinValue = Extents.MinValue;
+			MaxValue = Extents.MaxValue;
+		}
+		else
+		{
+			UE::MovieScene::Interpolation::FInterpolationExtents Extents =
+				Channel->ComputeExtents(0, 1);
+			MinValue = Extents.MinValue;
+			MaxValue = Extents.MaxValue;
+		}
+	}
+}
+
+void FDoubleChannelCurveModel::GetValueRange(double InMinTime, double InMaxTime, double& MinValue, double& MaxValue) const
+{
+	FMovieSceneDoubleChannel* Channel = GetChannelHandle().Get();
+	UMovieSceneSection* Section = WeakSection.Get();
+
+	if (Channel && Section)
+	{
+		TArrayView<const FFrameNumber> Times = Channel->GetData().GetTimes();
+
+		FFrameRate TickResolution = Section->GetTypedOuter<UMovieScene>()->GetTickResolution();
+
+		UE::MovieScene::Interpolation::FInterpolationExtents Extents =
+			Channel->ComputeExtents(InMinTime * TickResolution, InMaxTime * TickResolution);
+
+		MinValue = Extents.MinValue;
+		MaxValue = Extents.MaxValue;
+	}
+}
+
 void FDoubleChannelCurveModel::CreateKeyProxies(TArrayView<const FKeyHandle> InKeyHandles, TArrayView<UObject*> OutObjects)
 {
 	for (int32 Index = 0; Index < InKeyHandles.Num(); ++Index)
 	{
 		UDoubleChannelKeyProxy* NewProxy = NewObject<UDoubleChannelKeyProxy>(GetTransientPackage(), NAME_None);
 
-		NewProxy->Initialize(InKeyHandles[Index], GetChannelHandle(), Cast<UMovieSceneSection>(GetOwningObject()));
+		NewProxy->Initialize(InKeyHandles[Index], GetChannelHandle(), this->GetOwningObjectOrOuter<UMovieSceneSignedObject>());
 		OutObjects[Index] = NewProxy;
 	}
 }
@@ -105,7 +157,7 @@ TUniquePtr<IBufferedCurveModel> FDoubleChannelCurveModel::CreateBufferedCurveCop
 		double ValueMin = 0.f, ValueMax = 1.f;
 		GetValueRange(ValueMin, ValueMax);
 
-		return MakeUnique<FDoubleChannelBufferedCurveModel>(Channel, Cast<UMovieSceneSection>(GetOwningObject()), MoveTemp(KeyPositions), MoveTemp(KeyAttributes), GetLongDisplayName().ToString(), ValueMin, ValueMax);
+		return MakeUnique<FDoubleChannelBufferedCurveModel>(Channel, this->GetOwningObjectOrOuter<UMovieSceneSection>(), MoveTemp(KeyPositions), MoveTemp(KeyAttributes), GetLongDisplayName().ToString(), ValueMin, ValueMax);
 	}
 	return nullptr;
 }
diff --git a/Engine/Source/Editor/MovieSceneTools/Private/Channels/DoubleChannelKeyProxy.cpp b/Engine/Source/Editor/MovieSceneTools/Private/Channels/DoubleChannelKeyProxy.cpp
--- a/Engine/Source/Editor/MovieSceneTools/Private/Channels/DoubleChannelKeyProxy.cpp
+++ b/Engine/Source/Editor/MovieSceneTools/Private/Channels/DoubleChannelKeyProxy.cpp
@@ -6,11 +6,11 @@
 
 struct FPropertyChangedEvent;
 
-void UDoubleChannelKeyProxy::Initialize(FKeyHandle InKeyHandle, TMovieSceneChannelHandle<FMovieSceneDoubleChannel> InChannelHandle, TWeakObjectPtr<UMovieSceneSection> InWeakSection)
+void UDoubleChannelKeyProxy::Initialize(FKeyHandle InKeyHandle, TMovieSceneChannelHandle<FMovieSceneDoubleChannel> InChannelHandle, TWeakObjectPtr<UMovieSceneSignedObject> InWeakSignedObject)
 {
-	KeyHandle     = InKeyHandle;
-	ChannelHandle = InChannelHandle;
-	WeakSection   = InWeakSection;
+	KeyHandle          = InKeyHandle;
+	ChannelHandle      = InChannelHandle;
+	WeakSignedObject   = InWeakSignedObject;
 }
 
 
@@ -18,7 +18,7 @@ void UDoubleChannelKeyProxy::PostEditChangeProperty(FPropertyChangedEvent& Prope
 {
 	Super::PostEditChangeProperty(PropertyChangedEvent);
 
-	OnProxyValueChanged(ChannelHandle, WeakSection.Get(), KeyHandle, Value, Time);
+	OnProxyValueChanged(ChannelHandle, WeakSignedObject.Get(), KeyHandle, Value, Time);
 }
 
 void UDoubleChannelKeyProxy::UpdateValuesFromRawData()
diff --git a/Engine/Source/Editor/MovieSceneTools/Private/Channels/DoubleChannelKeyProxy.h b/Engine/Source/Editor/MovieSceneTools/Private/Channels/DoubleChannelKeyProxy.h
--- a/Engine/Source/Editor/MovieSceneTools/Private/Channels/DoubleChannelKeyProxy.h
+++ b/Engine/Source/Editor/MovieSceneTools/Private/Channels/DoubleChannelKeyProxy.h
@@ -26,7 +26,7 @@ public:
 	/**
 	 * Initialize this key proxy object by caching the underlying key object, and retrieving the time/value each tick
 	 */
-	void Initialize(FKeyHandle InKeyHandle, TMovieSceneChannelHandle<FMovieSceneDoubleChannel> InChannelHandle, TWeakObjectPtr<UMovieSceneSection> InWeakSection);
+	void Initialize(FKeyHandle InKeyHandle, TMovieSceneChannelHandle<FMovieSceneDoubleChannel> InChannelHandle, TWeakObjectPtr<UMovieSceneSignedObject> InWeakSignedObject);
 
 private:
 
@@ -52,6 +52,6 @@ private:
 	FKeyHandle KeyHandle;
 	/** Cached channel in which the key resides */
 	TMovieSceneChannelHandle<FMovieSceneDoubleChannel> ChannelHandle;
-	/** Cached section in which the channel resides */
-	TWeakObjectPtr<UMovieSceneSection> WeakSection;
+	/** Cached signed object in which the channel resides */
+	TWeakObjectPtr<UMovieSceneSignedObject> WeakSignedObject;
 };
diff --git a/Engine/Source/Editor/MovieSceneTools/Private/Channels/FloatChannelCurveModel.cpp b/Engine/Source/Editor/MovieSceneTools/Private/Channels/FloatChannelCurveModel.cpp
--- a/Engine/Source/Editor/MovieSceneTools/Private/Channels/FloatChannelCurveModel.cpp
+++ b/Engine/Source/Editor/MovieSceneTools/Private/Channels/FloatChannelCurveModel.cpp
@@ -79,7 +79,7 @@ void FFloatChannelCurveModel::CreateKeyProxies(TArrayView<const FKeyHandle> InKe
 	{
 		UFloatChannelKeyProxy* NewProxy = NewObject<UFloatChannelKeyProxy>(GetTransientPackage(), NAME_None);
 
-		NewProxy->Initialize(InKeyHandles[Index], GetChannelHandle(), Cast<UMovieSceneSection>(GetOwningObject()));
+		NewProxy->Initialize(InKeyHandles[Index], GetChannelHandle(), this->GetOwningObjectOrOuter<UMovieSceneSection>());
 		OutObjects[Index] = NewProxy;
 	}
 }
@@ -105,7 +105,7 @@ TUniquePtr<IBufferedCurveModel> FFloatChannelCurveModel::CreateBufferedCurveCopy
 		double ValueMin = 0.f, ValueMax = 1.f;
 		GetValueRange(ValueMin, ValueMax);
 
-		return MakeUnique<FFloatChannelBufferedCurveModel>(Channel, Cast<UMovieSceneSection>(GetOwningObject()), MoveTemp(KeyPositions), MoveTemp(KeyAttributes), GetLongDisplayName().ToString(), ValueMin, ValueMax);
+		return MakeUnique<FFloatChannelBufferedCurveModel>(Channel, this->GetOwningObjectOrOuter<UMovieSceneSection>(), MoveTemp(KeyPositions), MoveTemp(KeyAttributes), GetLongDisplayName().ToString(), ValueMin, ValueMax);
 	}
 	return nullptr;
 }
diff --git a/Engine/Source/Editor/MovieSceneTools/Private/Channels/IntegerChannelCurveModel.cpp b/Engine/Source/Editor/MovieSceneTools/Private/Channels/IntegerChannelCurveModel.cpp
--- a/Engine/Source/Editor/MovieSceneTools/Private/Channels/IntegerChannelCurveModel.cpp
+++ b/Engine/Source/Editor/MovieSceneTools/Private/Channels/IntegerChannelCurveModel.cpp
@@ -91,7 +91,7 @@ void FIntegerChannelCurveModel::CreateKeyProxies(TArrayView<const FKeyHandle> In
 	{
 		UIntegerChannelKeyProxy* NewProxy = NewObject<UIntegerChannelKeyProxy>(GetTransientPackage(), NAME_None);
 
-		NewProxy->Initialize(InKeyHandles[Index], GetChannelHandle(), Cast<UMovieSceneSection>(GetOwningObject()));
+		NewProxy->Initialize(InKeyHandles[Index], GetChannelHandle(), this->GetOwningObjectOrOuter<UMovieSceneSection>());
 		OutObjects[Index] = NewProxy;
 	}
 }
@@ -117,7 +117,7 @@ TUniquePtr<IBufferedCurveModel> FIntegerChannelCurveModel::CreateBufferedCurveCo
 		double ValueMin = 0.f, ValueMax = 1.f;
 		GetValueRange(ValueMin, ValueMax);
 
-		return MakeUnique<FIntegerChannelBufferedCurveModel>(Channel, Cast<UMovieSceneSection>(GetOwningObject()), MoveTemp(KeyPositions), MoveTemp(KeyAttributes), GetLongDisplayName().ToString(), ValueMin, ValueMax);
+		return MakeUnique<FIntegerChannelBufferedCurveModel>(Channel, this->GetOwningObjectOrOuter<UMovieSceneSection>(), MoveTemp(KeyPositions), MoveTemp(KeyAttributes), GetLongDisplayName().ToString(), ValueMin, ValueMax);
 	}
 	return nullptr;
 }
@@ -135,7 +135,7 @@ void FIntegerChannelCurveModel::GetCurveAttributes(FCurveAttributes& OutCurveAtt
 void FIntegerChannelCurveModel::SetCurveAttributes(const FCurveAttributes& InCurveAttributes)
 {
 	FMovieSceneIntegerChannel* Channel = GetChannelHandle().Get();
-	UMovieSceneSection* Section = Cast<UMovieSceneSection>(GetOwningObject());
+	UMovieSceneSection* Section = this->GetOwningObjectOrOuter<UMovieSceneSection>();
 	if (Channel && Section && !IsReadOnly())
 	{
 		Section->MarkAsChanged();
diff --git a/Engine/Source/Editor/MovieSceneTools/Private/Channels/SCurveEditorEventChannelView.cpp b/Engine/Source/Editor/MovieSceneTools/Private/Channels/SCurveEditorEventChannelView.cpp
--- a/Engine/Source/Editor/MovieSceneTools/Private/Channels/SCurveEditorEventChannelView.cpp
+++ b/Engine/Source/Editor/MovieSceneTools/Private/Channels/SCurveEditorEventChannelView.cpp
@@ -34,6 +34,7 @@ float SCurveEditorEventChannelView::TrackHeight = 24.f;
 void SCurveEditorEventChannelView::Construct(const FArguments& InArgs, TWeakPtr<FCurveEditor> InCurveEditor)
 {
 	bFixedOutputBounds = true;
+	bAllowModelViewTransforms = false;
 	OutputMin = -0.5;
 	OutputMax =  0.5;
 	WeakCurveEditor = InCurveEditor;
@@ -64,7 +65,7 @@ void SCurveEditorEventChannelView::Tick(const FGeometry& AllottedGeometry, const
 		FCurveModel* Curve = CurveEditor->FindCurve(It.Key());
 		if (ensureAlways(Curve))
 		{
-			It->Value.ViewToCurveTransform = FTransform2D(FVector2D(0.f, Count));
+			It->Value.ViewToCurveTransform = FTransform2d(FVector2D(0.f, Count));
 		}
 
 		Count += 1.0;
diff --git a/Engine/Source/Editor/MovieSceneTools/Private/Channels/SCurveEditorKeyBarView.cpp b/Engine/Source/Editor/MovieSceneTools/Private/Channels/SCurveEditorKeyBarView.cpp
--- a/Engine/Source/Editor/MovieSceneTools/Private/Channels/SCurveEditorKeyBarView.cpp
+++ b/Engine/Source/Editor/MovieSceneTools/Private/Channels/SCurveEditorKeyBarView.cpp
@@ -40,6 +40,8 @@ float SCurveEditorKeyBarView::TrackHeight = 24.f;
 void SCurveEditorKeyBarView::Construct(const FArguments& InArgs, TWeakPtr<FCurveEditor> InCurveEditor)
 {
 	bFixedOutputBounds = true;
+	bAllowModelViewTransforms = false;
+
 	OutputMin = -0.5;
 	OutputMax =  0.5;
 	WeakCurveEditor = InCurveEditor;
@@ -70,7 +72,7 @@ void SCurveEditorKeyBarView::Tick(const FGeometry& AllottedGeometry, const doubl
 		FCurveModel* Curve = CurveEditor->FindCurve(It.Key());
 		if (ensureAlways(Curve))
 		{
-			It->Value.ViewToCurveTransform = FTransform2D(FVector2D(0.f, Count));
+			It->Value.ViewToCurveTransform = FTransform2d(FVector2D(0.f, Count));
 		}
 
 		Count += 1.0;
diff --git a/Engine/Source/Editor/MovieSceneTools/Private/Constraints/ComponentConstraintChannelInterface.cpp b/Engine/Source/Editor/MovieSceneTools/Private/Constraints/ComponentConstraintChannelInterface.cpp
--- a/Engine/Source/Editor/MovieSceneTools/Private/Constraints/ComponentConstraintChannelInterface.cpp
+++ b/Engine/Source/Editor/MovieSceneTools/Private/Constraints/ComponentConstraintChannelInterface.cpp
@@ -3,11 +3,13 @@
 
 #include "ComponentConstraintChannelInterface.h"
 
+#include "Evaluation/MovieSceneEvaluationTemplateInstance.h"
 #include "MovieSceneToolHelpers.h"
 #include "TransformableHandle.h"
 #include "TransformConstraint.h"
 #include "Constraints/MovieSceneConstraintChannelHelper.inl"
 #include "EntitySystem/MovieSceneDecompositionQuery.h"
+#include "Evaluation/MovieSceneEvaluationTemplateInstance.h"
 #include "Sections/MovieSceneConstrainedSection.h"
 #include "Sections/MovieScene3DTransformSection.h"
 #include "Systems/MovieScenePropertyInstantiator.h"
@@ -382,12 +384,18 @@ void FComponentConstraintChannelInterface::RecomposeTransforms(
 		Query.Object   = SceneComponent;
 		Query.bConvertFromSourceEntityIDs = false;  // We already pass the children entity IDs
 		FMovieSceneSequenceTransform RootToLocalTransform = InSequencer->GetFocusedMovieSceneSequenceTransform();
+		FMovieSceneInverseSequenceTransform LocalToRootTransform = RootToLocalTransform.Inverse();
 
 		// add keys
 		for (int32 Index = 0; Index < InFrames.Num(); ++Index)
 		{
-			const FFrameNumber& FrameNumber = (FFrameTime(InFrames[Index]) * RootToLocalTransform.InverseNoLooping()).GetFrame();
-			const FMovieSceneEvaluationRange EvaluationRange = FMovieSceneEvaluationRange(FFrameTime(FrameNumber), TickResolution);
+			TOptional<FFrameTime> RootTime = LocalToRootTransform.TryTransformTime(InFrames[Index]);
+			if (!RootTime)
+			{
+				continue;
+			}
+
+			const FMovieSceneEvaluationRange EvaluationRange = FMovieSceneEvaluationRange(RootTime.GetValue(), TickResolution);
 			const FMovieSceneContext Context = FMovieSceneContext(EvaluationRange, PlaybackStatus).SetHasJumped(true);
 
 			EvaluationTemplate.EvaluateSynchronousBlocking(Context);
diff --git a/Engine/Source/Editor/MovieSceneTools/Private/Constraints/MovieSceneConstraintChannelHelper.cpp b/Engine/Source/Editor/MovieSceneTools/Private/Constraints/MovieSceneConstraintChannelHelper.cpp
--- a/Engine/Source/Editor/MovieSceneTools/Private/Constraints/MovieSceneConstraintChannelHelper.cpp
+++ b/Engine/Source/Editor/MovieSceneTools/Private/Constraints/MovieSceneConstraintChannelHelper.cpp
@@ -3,6 +3,7 @@
 #include "Constraints/MovieSceneConstraintChannelHelper.inl"
 
 #include "ISequencer.h"
+#include "Evaluation/MovieSceneEvaluationTemplateInstance.h"
 
 #include "TransformableHandle.h"
 #include "LevelSequence.h"
@@ -88,12 +89,20 @@ void FCompensationEvaluator::ComputeLocalTransforms(
 	const FConstraintsManagerController& Controller = FConstraintsManagerController::Get(InWorld);
 	static constexpr bool bSorted = true;
 	const TArray<ConstraintPtr> AllConstraints = Controller.GetAllConstraints(bSorted);
-	FMovieSceneSequenceTransform RootToLocalTransform = InSequencer->GetFocusedMovieSceneSequenceTransform();
+
+	FMovieSceneInverseSequenceTransform LocalToRootTransform = InSequencer->GetFocusedMovieSceneSequenceTransform().Inverse();
 
 	for (int32 Index = 0; Index < NumFrames + 1; ++Index)
 	{
 		FFrameNumber FrameNumber = (Index == 0) ? InFrames[0] - 1 : InFrames[Index - 1];
-		FrameNumber = (FFrameTime(FrameNumber) * RootToLocalTransform.InverseNoLooping()).GetFrame();
+		TOptional<FFrameTime> RootTime = LocalToRootTransform.TryTransformTime(FrameNumber);
+		if (!RootTime)
+		{
+			continue;
+		}
+
+		FrameNumber = RootTime->GetFrame();
+
 		// evaluate animation
 		const FMovieSceneEvaluationRange EvaluationRange = FMovieSceneEvaluationRange(FFrameTime(FrameNumber), TickResolution);
 		const FMovieSceneContext Context = FMovieSceneContext(EvaluationRange, PlaybackStatus).SetHasJumped(true);
@@ -253,11 +262,17 @@ void FCompensationEvaluator::ComputeLocalTransformsForBaking(UWorld* InWorld, co
 			BakeHelper->StartBaking(MovieScene);
 		}
 	}
-	FMovieSceneSequenceTransform RootToLocalTransform = InSequencer->GetFocusedMovieSceneSequenceTransform();
+	FMovieSceneInverseSequenceTransform LocalToRootTransform = InSequencer->GetFocusedMovieSceneSequenceTransform().Inverse();
 
 	for (int32 Index = 0; Index < NumFrames; ++Index)
 	{
-		const FFrameNumber& FrameNumber = (FFrameTime(InFrames[Index]) * RootToLocalTransform.InverseNoLooping()).GetFrame();
+		TOptional<FFrameTime> RootTime = LocalToRootTransform.TryTransformTime(InFrames[Index]);
+		if (!RootTime)
+		{
+			continue;
+		}
+
+		FFrameNumber FrameNumber = RootTime->GetFrame();
 
 		// evaluate animation
 		const FMovieSceneEvaluationRange EvaluationRange = FMovieSceneEvaluationRange(FFrameTime(FrameNumber), TickResolution);
@@ -378,11 +393,17 @@ void FCompensationEvaluator::ComputeLocalTransformsBeforeDeletion(
 			BakeHelper->StartBaking(MovieScene);
 		}
 	}
-	FMovieSceneSequenceTransform RootToLocalTransform = InSequencer->GetFocusedMovieSceneSequenceTransform();
+	FMovieSceneInverseSequenceTransform LocalToRootTransform = InSequencer->GetFocusedMovieSceneSequenceTransform().Inverse();
 
 	for (int32 Index = 0; Index < NumFrames; ++Index)
 	{
-		const FFrameNumber& FrameNumber = (FFrameTime(InFrames[Index]) * RootToLocalTransform.InverseNoLooping()).GetFrame();
+		TOptional<FFrameTime> RootTime = LocalToRootTransform.TryTransformTime(InFrames[Index]);
+		if (!RootTime)
+		{
+			continue;
+		}
+
+		FFrameNumber FrameNumber = RootTime->GetFrame();
 
 		// evaluate animation
 		const FMovieSceneEvaluationRange EvaluationRange = FMovieSceneEvaluationRange(FFrameTime(FrameNumber), TickResolution);
@@ -478,11 +499,18 @@ void FCompensationEvaluator::ComputeCompensation(UWorld* InWorld, const TSharedP
 			BakeHelper->StartBaking(MovieScene);
 		}
 	}
-	
-	auto EvaluateAt = [Handle = this->Handle, InSequencer, &AllConstraints, &BakeHelpers](FFrameNumber InFrame)
+
+	FMovieSceneInverseSequenceTransform LocalToRootTransform = InSequencer->GetFocusedMovieSceneSequenceTransform().Inverse();
+
+	auto EvaluateAt = [Handle = this->Handle, InSequencer, &AllConstraints, &BakeHelpers, &LocalToRootTransform](FFrameNumber InFrame)
 	{
-		FMovieSceneSequenceTransform RootToLocalTransform = InSequencer->GetFocusedMovieSceneSequenceTransform();
-		InFrame = (FFrameTime(InFrame) * RootToLocalTransform.InverseNoLooping()).GetFrame();
+		TOptional<FFrameTime> RootTime = LocalToRootTransform.TryTransformTime(InFrame);
+		if (!RootTime)
+		{
+			return;
+		}
+
+		InFrame = RootTime->GetFrame();
 
 		UMovieScene* MovieScene = InSequencer->GetFocusedMovieSceneSequence()->GetMovieScene();
 		const FFrameRate TickResolution = MovieScene->GetTickResolution();
@@ -580,9 +608,9 @@ void FCompensationEvaluator::CacheTransforms(UWorld* InWorld, const TSharedPtr<I
 	const FConstraintsManagerController& Controller = FConstraintsManagerController::Get(InWorld);
 	static constexpr bool bSorted = true;
 	const TArray<ConstraintPtr> AllConstraints = Controller.GetAllConstraints(bSorted);
-	
+
 	UMovieScene* MovieScene = InSequencer->GetFocusedMovieSceneSequence()->GetMovieScene();
-	const FFrameRate TickResolution = MovieScene->GetTickResolution();
+	const FFrameRate TickResolution = InSequencer->GetRootMovieSceneSequence()->GetMovieScene()->GetTickResolution();
 	const EMovieScenePlayerStatus::Type PlaybackStatus = InSequencer->GetPlaybackStatus();
 
 	const int32 NumFrames = InTime.Num();
@@ -592,13 +620,19 @@ void FCompensationEvaluator::CacheTransforms(UWorld* InWorld, const TSharedPtr<I
 	SpaceGlobals.SetNum(NumFrames);
 
 	const TArray<IMovieSceneToolsAnimationBakeHelper*>& BakeHelpers = FMovieSceneToolsModule::Get().GetAnimationBakeHelpers();
-	
+
+	FMovieSceneInverseSequenceTransform LocalToRootTransform = InSequencer->GetFocusedMovieSceneSequenceTransform().Inverse();
+
 	auto EvaluateAt = [&](FFrameNumber InFrame)
 	{
-		const FMovieSceneEvaluationRange EvaluationRange = FMovieSceneEvaluationRange(FFrameTime(InFrame), TickResolution);
+		TOptional<FFrameTime> RootTime = LocalToRootTransform.TryTransformTime(InFrame);
+		if (!RootTime)
+		{
+			return;
+		}
+
+		const FMovieSceneEvaluationRange EvaluationRange = FMovieSceneEvaluationRange(RootTime.GetValue(), TickResolution);
 		const FMovieSceneContext Context = FMovieSceneContext(EvaluationRange, PlaybackStatus).SetHasJumped(true);
-		FMovieSceneSequenceTransform RootToLocalTransform = InSequencer->GetFocusedMovieSceneSequenceTransform();
-		InFrame = (FFrameTime(InFrame) * RootToLocalTransform.InverseNoLooping()).GetFrame();
 
 		for (IMovieSceneToolsAnimationBakeHelper* BakeHelper : BakeHelpers)
 		{
diff --git a/Engine/Source/Editor/MovieSceneTools/Private/CurveKeyEditors/SNumericKeyEditor.h b/Engine/Source/Editor/MovieSceneTools/Private/CurveKeyEditors/SNumericKeyEditor.h
--- a/Engine/Source/Editor/MovieSceneTools/Private/CurveKeyEditors/SNumericKeyEditor.h
+++ b/Engine/Source/Editor/MovieSceneTools/Private/CurveKeyEditors/SNumericKeyEditor.h
@@ -10,7 +10,6 @@
 #include "CurveKeyEditors/SequencerKeyEditor.h"
 #include "Widgets/Input/SSpinBox.h"
 #include "Styling/AppStyle.h"
-#include "CurveKeyEditors/SequencerKeyEditor.h"
 #include "NumericPropertyParams.h"
 
 #define LOCTEXT_NAMESPACE "NumericKeyEditor"
@@ -32,23 +31,32 @@ struct SNonThrottledSpinBox : SSpinBox<T>
 /**
  * A widget for editing a curve representing integer keys.
  */
-template<typename ChannelType, typename NumericType>
-class SNumericKeyEditor : public SCompoundWidget
+template<typename NumericType>
+class SNumericKeyEditorWidget : public SCompoundWidget
 {
 public:
-	SLATE_BEGIN_ARGS(SNumericKeyEditor){}
+	SLATE_BEGIN_ARGS(SNumericKeyEditorWidget){}
 	SLATE_END_ARGS()
 
+	template<typename ChannelType>
 	void Construct(const FArguments& InArgs, const TSequencerKeyEditor<ChannelType, NumericType>& InKeyEditor)
+	{
+		TSharedRef<TSequencerKeyEditorWrapper<ChannelType, NumericType>> KeyEditorWrapper =
+			MakeShared<TSequencerKeyEditorWrapper<ChannelType, NumericType>>(InKeyEditor);
+
+		Construct(InArgs, KeyEditorWrapper);
+	}
+
+	void Construct(const FArguments& InArgs, const TSharedRef<ISequencerKeyEditor<NumericType>>& InKeyEditor)
 	{
 		KeyEditor = InKeyEditor;
 
 		const FProperty* Property = nullptr;
-		ISequencer* Sequencer = InKeyEditor.GetSequencer();
-		FTrackInstancePropertyBindings* PropertyBindings = InKeyEditor.GetPropertyBindings();
+		ISequencer* Sequencer = InKeyEditor->GetSequencer();
+		FTrackInstancePropertyBindings* PropertyBindings = InKeyEditor->GetPropertyBindings();
 		if (Sequencer && PropertyBindings)
 		{
-			for (TWeakObjectPtr<> WeakObject : Sequencer->FindBoundObjects(InKeyEditor.GetObjectBindingID(), Sequencer->GetFocusedTemplateID()))
+			for (TWeakObjectPtr<> WeakObject : Sequencer->FindBoundObjects(InKeyEditor->GetObjectBindingID(), Sequencer->GetFocusedTemplateID()))
 			{
 				if (UObject* Object = WeakObject.Get())
 				{
@@ -63,7 +71,7 @@ public:
 
 		const typename TNumericPropertyParams<NumericType>::FMetaDataGetter MetaDataGetter = TNumericPropertyParams<NumericType>::FMetaDataGetter::CreateLambda([&](const FName& Key)
 		{
-			return InKeyEditor.GetMetaData(Key);
+			return InKeyEditor->GetMetaData(Key);
 		});
 
 		TNumericPropertyParams<NumericType> NumericPropertyParams(Property, MetaDataGetter);
@@ -75,6 +83,7 @@ public:
 			.Font(FAppStyle::GetFontStyle("Sequencer.AnimationOutliner.RegularFont"))
 			.MinValue(NumericPropertyParams.MinValue)
 			.MaxValue(NumericPropertyParams.MaxValue)
+			.TypeInterface(InKeyEditor->GetNumericTypeInterface())
 			.MinSliderValue(NumericPropertyParams.MinSliderValue)
 			.MaxSliderValue(NumericPropertyParams.MaxSliderValue)
 			.SliderExponent(NumericPropertyParams.SliderExponent)
@@ -82,11 +91,11 @@ public:
 			// LinearDeltaSensitivity needs to be left unset if not provided, rather than being set to some default
 			.LinearDeltaSensitivity(NumericPropertyParams.GetLinearDeltaSensitivityAttribute())
 			.WheelStep(NumericPropertyParams.WheelStep)
-			.Value_Raw(&KeyEditor, &decltype(KeyEditor)::GetCurrentValue)
-			.OnValueChanged(this, &SNumericKeyEditor::OnValueChanged)
-			.OnValueCommitted(this, &SNumericKeyEditor::OnValueCommitted)
-			.OnBeginSliderMovement(this, &SNumericKeyEditor::OnBeginSliderMovement)
-			.OnEndSliderMovement(this, &SNumericKeyEditor::OnEndSliderMovement)
+			.Value_Raw(KeyEditor.Get(), &ISequencerKeyEditor<NumericType>::GetCurrentValue)
+			.OnValueChanged(this, &SNumericKeyEditorWidget::OnValueChanged)
+			.OnValueCommitted(this, &SNumericKeyEditorWidget::OnValueCommitted)
+			.OnBeginSliderMovement(this, &SNumericKeyEditorWidget::OnBeginSliderMovement)
+			.OnEndSliderMovement(this, &SNumericKeyEditorWidget::OnEndSliderMovement)
 		];
 	}
 
@@ -94,7 +103,7 @@ private:
 
 	virtual FSlateColor GetForegroundColor() const override
 	{
-		if (KeyEditor.GetEditingKeySelection())
+		if (KeyEditor->GetEditingKeySelection())
 		{
 			return FLinearColor::Yellow;
 		}
@@ -113,14 +122,18 @@ private:
 	{
 		if (GEditor->IsTransactionActive())
 		{
-			KeyEditor.SetValue(Value);
+			KeyEditor->SetValue(Value);
 			GEditor->EndTransaction();
 		}
 	}
 
 	void OnValueChanged(NumericType Value)
 	{
-		KeyEditor.SetValueWithNotify(Value, EMovieSceneDataChangeType::TrackValueChanged);
+		if (!bCommitted)
+		{
+			KeyEditor->SetValueWithNotify(Value, EMovieSceneDataChangeType::TrackValueChanged);
+		}
+		bCommitted = false;
 	}
 
 	void OnValueCommitted(NumericType Value, ETextCommit::Type CommitInfo)
@@ -128,13 +141,18 @@ private:
 		if (CommitInfo == ETextCommit::OnEnter)
 		{
 			const FScopedTransaction Transaction( LOCTEXT("SetNumericKey", "Set Key Value") );
-			KeyEditor.SetValueWithNotify(Value, EMovieSceneDataChangeType::TrackValueChangedRefreshImmediately);
+			KeyEditor->SetValueWithNotify(Value, EMovieSceneDataChangeType::TrackValueChangedRefreshImmediately);
+			bCommitted = true;
 		}
 	}
 
 private:
 
-	TSequencerKeyEditor<ChannelType, NumericType> KeyEditor;
+	TSharedPtr<ISequencerKeyEditor<NumericType>> KeyEditor;
+	bool bCommitted = false;
 };
 
+template<typename ChannelType, typename NumericType>
+using SNumericKeyEditor = SNumericKeyEditorWidget<NumericType>;
+
 #undef LOCTEXT_NAMESPACE
diff --git a/Engine/Source/Editor/MovieSceneTools/Private/CurveKeyEditors/SequencerKeyEditor.h b/Engine/Source/Editor/MovieSceneTools/Private/CurveKeyEditors/SequencerKeyEditor.h
--- a/Engine/Source/Editor/MovieSceneTools/Private/CurveKeyEditors/SequencerKeyEditor.h
+++ b/Engine/Source/Editor/MovieSceneTools/Private/CurveKeyEditors/SequencerKeyEditor.h
@@ -15,6 +15,26 @@
 #include "MVVM/Selection/Selection.h"
 #include "MVVM/ViewModels/SequencerEditorViewModel.h"
 
+template<typename NumericType>
+struct INumericTypeInterface;
+
+template<typename ValueType>
+struct ISequencerKeyEditor
+{
+	virtual ~ISequencerKeyEditor(){}
+
+	virtual TSharedPtr<INumericTypeInterface<ValueType>> GetNumericTypeInterface() const = 0;
+	virtual TOptional<ValueType> GetExternalValue() const = 0;
+	virtual ValueType GetCurrentValue() const = 0;
+	virtual void SetValue(const ValueType& InValue) = 0;
+	virtual void SetValueWithNotify(const ValueType& InValue, EMovieSceneDataChangeType NotifyType = EMovieSceneDataChangeType::TrackValueChanged) = 0;
+	virtual const FGuid& GetObjectBindingID() const = 0;
+	virtual ISequencer* GetSequencer() const = 0;
+	virtual FTrackInstancePropertyBindings* GetPropertyBindings() const = 0;
+	virtual FString GetMetaData(const FName& Key) const = 0;
+	virtual bool GetEditingKeySelection() const = 0;
+};
+
 template<typename ChannelType, typename ValueType>
 struct TSequencerKeyEditor
 {
@@ -59,6 +79,21 @@ struct TSequencerKeyEditor
 		return TOptional<ValueType>();
 	}
 
+	void SetOwningObject(TWeakObjectPtr<UMovieSceneSignedObject> InWeakOwningObject)
+	{
+		WeakOwningObject = InWeakOwningObject;
+	}
+
+	void SetNumericTypeInterface(TSharedPtr<INumericTypeInterface<ValueType>> InNumericTypeInterface)
+	{
+		NumericTypeInterface = InNumericTypeInterface;
+	}
+
+	TSharedPtr<INumericTypeInterface<ValueType>> GetNumericTypeInterface() const
+	{
+		return NumericTypeInterface;
+	}
+
 	TOptional<ValueType> GetExternalValue() const
 	{
 		return Get(ObjectBindingID, WeakSequencer.Pin().Get(), WeakPropertyBindings.Pin().Get(), OnGetExternalValue);
@@ -71,12 +106,14 @@ struct TSequencerKeyEditor
 		ChannelType* Channel = ChannelHandle.Get();
 		ISequencer* Sequencer = WeakSequencer.Pin().Get();
 		UMovieSceneSection* OwningSection = WeakSection.Get();
+		const FMovieSceneChannelMetaData* ChannelMetaData = ChannelHandle.GetMetaData();
 
 		ValueType Result{};
 
-		if (Channel && Sequencer && OwningSection)
+		if (Channel && ChannelMetaData && Sequencer && OwningSection)
 		{
-			const FFrameTime CurrentTime = UE::MovieScene::ClampToDiscreteRange(Sequencer->GetLocalTime().Time, OwningSection->GetRange());
+			const FFrameTime CurrentTime = UE::MovieScene::ClampToDiscreteRange(Sequencer->GetLocalTime().Time, OwningSection->GetRange()) - ChannelMetaData->GetOffsetTime(OwningSection);
+
 			//If we have no keys and no default, key with the external value if it exists
 			if (!EvaluateChannel(OwningSection, Channel, CurrentTime, Result))
 			{
@@ -136,16 +173,24 @@ struct TSequencerKeyEditor
 			return;
 		}
 
-		OwningSection->SetFlags(RF_Transactional);
-
 		ChannelType* Channel = ChannelHandle.Get();
 		ISequencer* Sequencer = WeakSequencer.Pin().Get();
+		const FMovieSceneChannelMetaData* ChannelMetaData = ChannelHandle.GetMetaData();
 
-		if (!OwningSection->TryModify() || !Channel || !Sequencer)
+		if (OwningSection->IsReadOnly() || !Channel || !Sequencer || !ChannelMetaData)
 		{
 			return;
 		}
 
+		UMovieSceneSignedObject* Owner = WeakOwningObject.Get();
+		if (!Owner)
+		{
+			Owner = OwningSection;
+		}
+
+		Owner->Modify();
+		Owner->SetFlags(RF_Transactional);
+
 		const bool  bAutoSetTrackDefaults = Sequencer->GetAutoSetTrackDefaults();
 
 		const FKeySelection& KeySelection = Sequencer->GetViewModel()->GetSelection()->KeySelection;
@@ -169,7 +214,7 @@ struct TSequencerKeyEditor
 		}
 		else
 		{
-			const FFrameNumber CurrentTime = Sequencer->GetLocalTime().Time.FloorToFrame();
+			const FFrameNumber CurrentTime = Sequencer->GetLocalTime().Time.FloorToFrame() - ChannelMetaData->GetOffsetTime(OwningSection);
 
 			EMovieSceneKeyInterpolation Interpolation = GetInterpolationMode(Channel, CurrentTime, Sequencer->GetKeyInterpolation());
 
@@ -275,7 +320,65 @@ private:
 	FGuid ObjectBindingID;
 	TMovieSceneChannelHandle<ChannelType> ChannelHandle;
 	TWeakObjectPtr<UMovieSceneSection> WeakSection;
+	TWeakObjectPtr<UMovieSceneSignedObject> WeakOwningObject;
 	TWeakPtr<ISequencer> WeakSequencer;
 	TWeakPtr<FTrackInstancePropertyBindings> WeakPropertyBindings;
 	TFunction<TOptional<ValueType>(UObject&, FTrackInstancePropertyBindings*)> OnGetExternalValue;
+	TSharedPtr<INumericTypeInterface<ValueType>> NumericTypeInterface;
 };
+
+
+
+
+template<typename ChannelType, typename ValueType>
+struct TSequencerKeyEditorWrapper : ISequencerKeyEditor<ValueType>
+{
+	TSequencerKeyEditorWrapper(const TSequencerKeyEditor<ChannelType, ValueType>& InKeyEditor)
+		: Impl(InKeyEditor)
+	{}
+
+	TSharedPtr<INumericTypeInterface<ValueType>> GetNumericTypeInterface() const override
+	{
+		return Impl.GetNumericTypeInterface();
+	}
+	TOptional<ValueType> GetExternalValue() const override
+	{
+		return Impl.GetExternalValue();
+	}
+	ValueType GetCurrentValue() const override
+	{
+		return Impl.GetCurrentValue();
+	}
+	void SetValue(const ValueType& InValue) override
+	{
+		return Impl.SetValue(InValue);
+	}
+	void SetValueWithNotify(const ValueType& InValue, EMovieSceneDataChangeType NotifyType) override
+	{
+		return Impl.SetValueWithNotify(InValue, NotifyType);
+	}
+	const FGuid& GetObjectBindingID() const override
+	{
+		return Impl.GetObjectBindingID();
+	}
+	ISequencer* GetSequencer() const override
+	{
+		return Impl.GetSequencer();
+	}
+	FTrackInstancePropertyBindings* GetPropertyBindings() const override
+	{
+		return Impl.GetPropertyBindings();
+	}
+	FString GetMetaData(const FName& Key) const override
+	{
+		return Impl.GetMetaData(Key);
+	}
+	bool GetEditingKeySelection() const override
+	{
+		return Impl.GetEditingKeySelection();
+	}
+
+private:
+
+	TSequencerKeyEditor<ChannelType, ValueType> Impl;
+};
\ No newline at end of file
diff --git a/Engine/Source/Editor/MovieSceneTools/Private/MovieSceneToolHelpers.cpp b/Engine/Source/Editor/MovieSceneTools/Private/MovieSceneToolHelpers.cpp
--- a/Engine/Source/Editor/MovieSceneTools/Private/MovieSceneToolHelpers.cpp
+++ b/Engine/Source/Editor/MovieSceneTools/Private/MovieSceneToolHelpers.cpp
@@ -4426,9 +4426,8 @@ static void GetSequencerActorWorldTransforms(IMovieScenePlayer* Player, FMovieSc
 
 			for (int32 Index = 0; Index < Frames.Num(); ++Index)
 			{
-				const FFrameNumber& FrameNumber = Frames[Index];
-				FFrameTime GlobalTime(FrameNumber);
-				GlobalTime = GlobalTime * RootToLocalTransform.InverseNoLooping(); //player evals in root time so need to go back to it.
+				FFrameNumber FrameNumber = Frames[Index];
+				FFrameTime GlobalTime = RootToLocalTransform.Inverse().TryTransformTime(FrameNumber).Get(FrameNumber); //player evals in root time so need to go back to it.
 
 				for (IMovieSceneToolsAnimationBakeHelper* BakeHelper : BakeHelpers)
 				{
diff --git a/Engine/Source/Editor/MovieSceneTools/Private/MovieSceneToolsModule.cpp b/Engine/Source/Editor/MovieSceneTools/Private/MovieSceneToolsModule.cpp
--- a/Engine/Source/Editor/MovieSceneTools/Private/MovieSceneToolsModule.cpp
+++ b/Engine/Source/Editor/MovieSceneTools/Private/MovieSceneToolsModule.cpp
@@ -55,6 +55,7 @@
 #include "TrackEditors/CVarTrackEditor.h"
 #include "TrackEditors/CustomPrimitiveDataTrackEditor.h"
 #include "TrackEditors/BindingLifetimeTrackEditor.h"
+#include "TrackEditors/TimeWarpTrackEditor.h"
 
 #include "Channels/PerlinNoiseChannelInterface.h"
 
@@ -67,6 +68,7 @@
 #include "MovieSceneObjectBindingIDCustomization.h"
 #include "MovieSceneDynamicBindingCustomization.h"
 #include "MovieSceneEventCustomization.h"
+#include "MovieSceneTimeWarpVariantCustomization.h"
 #include "SequencerClipboardReconciler.h"
 #include "ClipboardTypes.h"
 #include "ISettingsModule.h"
@@ -81,6 +83,7 @@
 #include "Channels/MovieSceneByteChannel.h"
 #include "Channels/MovieSceneChannel.h"
 #include "Channels/MovieSceneEventChannel.h"
+#include "Channels/MovieSceneTimeWarpChannel.h"
 #include "Channels/MovieSceneObjectPathChannel.h"
 #include "Channels/MovieSceneCameraShakeSourceTriggerChannel.h"
 #include "Channels/MovieSceneStringChannel.h"
@@ -189,6 +192,7 @@ void FMovieSceneToolsModule::StartupModule()
 		CVarTrackCreateEditorHandle = SequencerModule.RegisterTrackEditor(FOnCreateTrackEditor::CreateStatic(&FCVarTrackEditor::CreateTrackEditor));
 		CustomPrimitiveDataTrackCreateEditorHandle = SequencerModule.RegisterTrackEditor(FOnCreateTrackEditor::CreateStatic(&FCustomPrimitiveDataTrackEditor::CreateTrackEditor));
 		BindingLifetimeTrackCreateEditorHandle = SequencerModule.RegisterTrackEditor(FOnCreateTrackEditor::CreateStatic(&FBindingLifetimeTrackEditor::CreateTrackEditor));
+		TimeWarpTrackCreateEditorHandle = SequencerModule.RegisterTrackEditor(FOnCreateTrackEditor::CreateStatic(&FTimeWarpTrackEditor::CreateTrackEditor));
 
 		// register track models
 		CameraCutTrackModelHandle = SequencerModule.RegisterTrackModel(FOnCreateTrackModel::CreateStatic(&FCameraCutTrackModel::CreateTrackModel));
@@ -208,6 +212,7 @@ void FMovieSceneToolsModule::StartupModule()
 		PropertyModule.RegisterCustomPropertyTypeLayout("MovieSceneDynamicBinding", FOnGetPropertyTypeCustomizationInstance::CreateStatic(&FMovieSceneDynamicBindingCustomization::MakeInstance));
 		PropertyModule.RegisterCustomPropertyTypeLayout("MovieSceneEvent", FOnGetPropertyTypeCustomizationInstance::CreateStatic(&FMovieSceneEventCustomization::MakeInstance));
 		PropertyModule.RegisterCustomPropertyTypeLayout("MovieSceneCVarOverrides", FOnGetPropertyTypeCustomizationInstance::CreateStatic(&UE::MovieScene::FCVarOverridesPropertyTypeCustomization::MakeInstance));
+		PropertyModule.RegisterCustomPropertyTypeLayout("MovieSceneTimeWarpVariant", FOnGetPropertyTypeCustomizationInstance::CreateLambda(&MakeShared<UE::MovieScene::FMovieSceneTimeWarpVariantCustomization>));
 
 		SequencerModule.RegisterChannelInterface<FMovieSceneBoolChannel>();
 		SequencerModule.RegisterChannelInterface<FMovieSceneByteChannel>();
@@ -219,6 +224,7 @@ void FMovieSceneToolsModule::StartupModule()
 		SequencerModule.RegisterChannelInterface<FMovieSceneActorReferenceData>();
 		SequencerModule.RegisterChannelInterface<FMovieSceneEventSectionData>();
 		SequencerModule.RegisterChannelInterface<FMovieSceneObjectPathChannel>();
+		SequencerModule.RegisterChannelInterface<FMovieSceneTimeWarpChannel>();
 
 		SequencerModule.RegisterChannelInterface<FMovieSceneAudioTriggerChannel>();
 
@@ -359,6 +365,7 @@ void FMovieSceneToolsModule::ShutdownModule()
 	SequencerModule.UnRegisterTrackEditor( CVarTrackCreateEditorHandle );
 	SequencerModule.UnRegisterTrackEditor( CustomPrimitiveDataTrackCreateEditorHandle );
 	SequencerModule.UnRegisterTrackEditor( BindingLifetimeTrackCreateEditorHandle );
+	SequencerModule.UnRegisterTrackEditor( TimeWarpTrackCreateEditorHandle );
 
 	// unregister track models
 	SequencerModule.UnregisterTrackModel( CameraCutTrackModelHandle );
diff --git a/Engine/Source/Editor/MovieSceneTools/Private/Sections/CinematicShotSection.cpp b/Engine/Source/Editor/MovieSceneTools/Private/Sections/CinematicShotSection.cpp
--- a/Engine/Source/Editor/MovieSceneTools/Private/Sections/CinematicShotSection.cpp
+++ b/Engine/Source/Editor/MovieSceneTools/Private/Sections/CinematicShotSection.cpp
@@ -31,7 +31,6 @@ FCinematicShotSection::FCinematicSectionCache::FCinematicSectionCache(UMovieScen
 	: InnerFrameRate(1, 1)
 	, InnerFrameOffset(0)
 	, SectionStartFrame(0)
-	, TimeScale(1.f)
 {
 	if (Section)
 	{
@@ -230,14 +229,14 @@ int32 FCinematicShotSection::OnPaintSection(FSequencerSectionPainter& InPainter)
 
 	InPainter.LayerId = InPainter.PaintSectionBackground();
 
-	FVector2D LocalSectionSize = InPainter.SectionGeometry.GetLocalSize();
+	FVector2D LocalHeaderSize = InPainter.HeaderGeometry.GetLocalSize();
 	const UMovieSceneCinematicShotSection& SectionObject = GetSectionObjectAs<UMovieSceneCinematicShotSection>();
 
 	// Paint fancy-looking film border.
 	FSlateDrawElement::MakeBox(
 		InPainter.DrawElements,
 		InPainter.LayerId++,
-		InPainter.SectionGeometry.ToPaintGeometry(FVector2D(LocalSectionSize.X-2.f, 7.f), FSlateLayoutTransform(FVector2D(1.f, 4.f))),
+		InPainter.SectionGeometry.ToPaintGeometry(FVector2D(LocalHeaderSize.X-2.f, 7.f), FSlateLayoutTransform(FVector2D(1.f, 4.f))),
 		FilmBorder,
 		InPainter.bParentEnabled ? ESlateDrawEffect::None : ESlateDrawEffect::DisabledEffect
 	);
@@ -245,7 +244,7 @@ int32 FCinematicShotSection::OnPaintSection(FSequencerSectionPainter& InPainter)
 	FSlateDrawElement::MakeBox(
 		InPainter.DrawElements,
 		InPainter.LayerId++,
-		InPainter.SectionGeometry.ToPaintGeometry(FVector2D(LocalSectionSize.X-2.f, 7.f), FSlateLayoutTransform(FVector2D(1.f, LocalSectionSize.Y - 11.f))),
+		InPainter.SectionGeometry.ToPaintGeometry(FVector2D(LocalHeaderSize.X-2.f, 7.f), FSlateLayoutTransform(FVector2D(1.f, LocalHeaderSize.Y - 11.f))),
 		FilmBorder,
 		InPainter.bParentEnabled ? ESlateDrawEffect::None : ESlateDrawEffect::DisabledEffect
 	);
diff --git a/Engine/Source/Editor/MovieSceneTools/Private/Sections/CinematicShotSection.h b/Engine/Source/Editor/MovieSceneTools/Private/Sections/CinematicShotSection.h
--- a/Engine/Source/Editor/MovieSceneTools/Private/Sections/CinematicShotSection.h
+++ b/Engine/Source/Editor/MovieSceneTools/Private/Sections/CinematicShotSection.h
@@ -63,7 +63,7 @@ private:
 		FFrameRate   InnerFrameRate;
 		FFrameNumber InnerFrameOffset;
 		FFrameNumber SectionStartFrame;
-		float        TimeScale;
+		FMovieSceneTimeWarpVariant TimeScale;
 	};
 
 	/** Cached section thumbnail data */
diff --git a/Engine/Source/Editor/MovieSceneTools/Private/Sections/ThumbnailSection.cpp b/Engine/Source/Editor/MovieSceneTools/Private/Sections/ThumbnailSection.cpp
--- a/Engine/Source/Editor/MovieSceneTools/Private/Sections/ThumbnailSection.cpp
+++ b/Engine/Source/Editor/MovieSceneTools/Private/Sections/ThumbnailSection.cpp
@@ -7,6 +7,9 @@
 #include "Framework/Commands/UIAction.h"
 #include "Framework/MultiBox/MultiBoxBuilder.h"
 #include "GameFramework/Actor.h"
+#include "ISequencer.h"
+#include "MVVM/ViewModels/SequencerEditorViewModel.h"
+#include "MVVM/ViewModels/ViewDensity.h"
 #include "Modules/ModuleManager.h"
 #include "Application/ThrottleManager.h"
 #include "Widgets/Layout/SBox.h"
@@ -218,20 +221,20 @@ int32 FThumbnailSection::OnPaintSection( FSequencerSectionPainter& InPainter ) c
 
 	int32 LayerId = InPainter.LayerId;
 
-	const FGeometry& SectionGeometry = InPainter.SectionGeometry;
+	const FGeometry& HeaderGeometry = InPainter.HeaderGeometry;
 
 	// @todo Sequencer: Need a way to visualize the key here
 
 	const TRange<double> VisibleRange = GetVisibleRange();
 	const TRange<double> GenerationRange = GetTotalRange();
 
-	const float TimePerPx = GenerationRange.Size<double>() / InPainter.SectionGeometry.GetLocalSize().X;
+	const float TimePerPx = GenerationRange.Size<double>() / HeaderGeometry.GetLocalSize().X;
 
 	const FFrameRate TickResolution = Section->GetTypedOuter<UMovieScene>()->GetTickResolution();
 	const double SectionEaseInDuration = TickResolution.AsSeconds(Section->Easing.GetEaseInDuration()) / TimePerPx;
 	const double SectionEaseOutDuration = TickResolution.AsSeconds(Section->Easing.GetEaseOutDuration()) / TimePerPx;
 
-	const FSlateRect ThumbnailClipRect = SectionGeometry.GetLayoutBoundingRect()
+	const FSlateRect ThumbnailClipRect = HeaderGeometry.GetLayoutBoundingRect()
 		.InsetBy(FMargin(SectionThumbnailPadding, 0.f))
 		.InsetBy(FMargin(SectionEaseInDuration, 0.f, SectionEaseOutDuration, 0.f))
 		.IntersectionWith(InPainter.SectionClippingRect);
@@ -259,9 +262,9 @@ int32 FThumbnailSection::OnPaintSection( FSequencerSectionPainter& InPainter ) c
 			? FMath::Max(float(VisibleRange.GetLowerBoundValue() - GenerationRange.GetLowerBoundValue()) / TimePerPx, 0.f) + SectionThumbnailPadding
 			: (Thumbnail->GetTimeRange().GetLowerBoundValue() - GenerationRange.GetLowerBoundValue()) / TimePerPx;
 
-		const float PositionY = (SectionGeometry.GetLocalSize().Y - ThumbnailCropSize.Y)*.5f;
+		const float PositionY = (HeaderGeometry.GetLocalSize().Y - ThumbnailCropSize.Y)*.5f;
 
-		FPaintGeometry PaintGeometry = SectionGeometry.ToPaintGeometry(
+		FPaintGeometry PaintGeometry = HeaderGeometry.ToPaintGeometry(
 			ThumbnailRTSize,
 			FSlateLayoutTransform(ThumbnailScale, FVector2D(PositionX-HorizontalCropOffset, PositionY))
 		);
@@ -282,7 +285,7 @@ int32 FThumbnailSection::OnPaintSection( FSequencerSectionPainter& InPainter ) c
 			DrawEffects |= ESlateDrawEffect::IgnoreTextureAlpha;
 		}
 
-		FGeometry ClipGeometry = SectionGeometry.MakeChild(
+		FGeometry ClipGeometry = HeaderGeometry.MakeChild(
 			ThumbnailCropSize,
 			FSlateLayoutTransform(
 				FVector2D(PositionX, PositionY)
@@ -349,12 +352,19 @@ TRange<double> FThumbnailSection::GetTotalRange() const
 
 void FThumbnailSection::Tick(const FGeometry& AllottedGeometry, const FGeometry& ParentGeometry, const double InCurrentTime, const float InDeltaTime)
 {
-	if (FSlateThrottleManager::Get().IsAllowingExpensiveTasks() && GetDefault<UMovieSceneUserThumbnailSettings>()->bDrawThumbnails)
+	using namespace UE::Sequencer;
+
+	TSharedPtr<ISequencer> Sequencer = SequencerPtr.Pin();
+	if (Sequencer && FSlateThrottleManager::Get().IsAllowingExpensiveTasks() && GetDefault<UMovieSceneUserThumbnailSettings>()->bDrawThumbnails)
 	{
 		const UMovieSceneUserThumbnailSettings* Settings = GetDefault<UMovieSceneUserThumbnailSettings>();
 
+		FViewDensityInfo ViewDensity = Sequencer->GetViewModel()->GetViewDensity();
+		const float Height = GetSectionHeight(ViewDensity);
+
 		FIntPoint AllocatedSize = AllottedGeometry.GetLocalSize().IntPoint();
 		AllocatedSize.X = FMath::Max(AllocatedSize.X, 1);
+		AllocatedSize.Y = FMath::RoundToInt(Height);
 
 		ThumbnailCache.Update(GetTotalRange(), GetVisibleRange(), AllocatedSize, Settings->ThumbnailSize, Settings->Quality, InCurrentTime);
 	}
diff --git a/Engine/Source/Editor/MovieSceneTools/Private/TrackEditors/SkeletalAnimationTrackEditor.cpp b/Engine/Source/Editor/MovieSceneTools/Private/TrackEditors/SkeletalAnimationTrackEditor.cpp
--- a/Engine/Source/Editor/MovieSceneTools/Private/TrackEditors/SkeletalAnimationTrackEditor.cpp
+++ b/Engine/Source/Editor/MovieSceneTools/Private/TrackEditors/SkeletalAnimationTrackEditor.cpp
@@ -13,6 +13,8 @@
 #include "Framework/Application/SlateApplication.h"
 #include "Widgets/Layout/SBox.h"
 #include "SequencerSectionPainter.h"
+#include "Channels/MovieSceneTimeWarpChannel.h"
+#include "Variants/MovieSceneTimeWarpGetter.h"
 #include "Components/SceneComponent.h"
 #include "Components/SkeletalMeshComponent.h"
 #include "Editor/UnrealEdEngine.h"
@@ -37,6 +39,7 @@
 #include "FrameNumberDisplayFormat.h"
 #include "FrameNumberNumericInterface.h"
 #include "AnimationBlueprintLibrary.h"
+#include "MovieSceneTransformTypes.h"
 #include "AnimationEditorUtils.h"
 #include "Factories/PoseAssetFactory.h"
 #include "Misc/MessageDialog.h"
@@ -599,19 +602,50 @@ private:
 FSkeletalAnimationSection::FSkeletalAnimationSection( UMovieSceneSection& InSection, TWeakPtr<ISequencer> InSequencer)
 	: Section(*CastChecked<UMovieSceneSkeletalAnimationSection>(&InSection))
 	, Sequencer(InSequencer)
-	, InitialFirstLoopStartOffsetDuringResize(0)
-	, InitialStartTimeDuringResize(0)
-{ 
+	, PreDilatePlayRate(1.0)
+{
+}
+
+FSkeletalAnimationSection::~FSkeletalAnimationSection()
+{
 }
 
 void FSkeletalAnimationSection::BeginDilateSection()
 {
-	Section.PreviousPlayRate = Section.Params.PlayRate; //make sure to cache the play rate
+	if (Section.Params.PlayRate.GetType() == EMovieSceneTimeWarpType::FixedPlayRate)
+	{
+		PreDilatePlayRate = Section.Params.PlayRate.AsFixedPlayRate(); //make sure to cache the play rate
+	}
+	else if (Section.Params.PlayRate.GetType() == EMovieSceneTimeWarpType::Custom)
+	{
+		FMovieSceneTimeWarpChannel* Channel = Section.GetChannelProxy().GetChannel<FMovieSceneTimeWarpChannel>(0);
+		if (Channel)
+		{
+			Section.Params.PlayRate.AsCustom()->Modify();
+			PreDilateChannel = MakeUnique<FMovieSceneTimeWarpChannel>(*Channel);
+		}
+	}
 }
 
 void FSkeletalAnimationSection::DilateSection(const TRange<FFrameNumber>& NewRange, float DilationFactor)
 {
-	Section.Params.PlayRate = Section.PreviousPlayRate / DilationFactor;
+	if (Section.Params.PlayRate.GetType() == EMovieSceneTimeWarpType::FixedPlayRate)
+	{
+		Section.Params.PlayRate.Set(PreDilatePlayRate / DilationFactor);
+	}
+	else if (Section.Params.PlayRate.GetType() == EMovieSceneTimeWarpType::Custom)
+	{
+		FMovieSceneTimeWarpChannel* Channel = Section.GetChannelProxy().GetChannel<FMovieSceneTimeWarpChannel>(0);
+		if (Channel)
+		{
+			*Channel = *PreDilateChannel;
+
+			// Dilate the times
+			Dilate(Channel, FFrameNumber(0), DilationFactor);
+
+			Section.Params.PlayRate.AsCustom()->MarkAsChanged();
+		}
+	}
 	Section.SetRange(NewRange);
 }
 
@@ -643,13 +677,19 @@ FText FSkeletalAnimationSection::GetSectionToolTip() const
 		UMovieScene* MovieScene = Section.GetTypedOuter<UMovieScene>();
 		FFrameRate TickResolution = MovieScene->GetTickResolution();
 
-		const float AnimPlayRate = FMath::IsNearlyZero(Section.Params.PlayRate) || Section.Params.Animation == nullptr ? 1.0f : Section.Params.PlayRate * Section.Params.Animation->RateScale;
-		const float StartOffset = TickResolution.AsSeconds(Section.Params.FirstLoopStartFrameOffset);
-		const float SectionLength = Section.GetRange().Size<FFrameTime>() / TickResolution;
+		FMovieSceneSequenceTransform Transform = Section.Params.MakeTransform(TickResolution, Section.GetRange());
+
+		const double StartTime     = Transform.TransformTime(0).AsDecimal();
+		const float  SectionLength = Section.GetRange().Size<FFrameTime>() / TickResolution;
 
-		if (!FMath::IsNearlyZero(SectionLength, KINDA_SMALL_NUMBER) && SectionLength > 0)
+		if (Section.Params.PlayRate.GetType() == EMovieSceneTimeWarpType::FixedPlayRate)
 		{
-			return FText::Format(LOCTEXT("ToolTipContentFormat", "Start: {0}s\nDuration: {1}s\nPlay Rate:{2}"), StartOffset, SectionLength, AnimPlayRate);
+			const double PlayRate = Section.Params.PlayRate.AsFixedPlayRate();
+			return FText::Format(LOCTEXT("ToolTipContentFormat_FixedPlayRate", "Start: {0}s\nDuration: {1}s\nPlay Rate: {2}x"), StartTime, SectionLength, PlayRate);
+		}
+		else if (Section.Params.PlayRate.GetType() == EMovieSceneTimeWarpType::Custom)
+		{
+			return FText::Format(LOCTEXT("ToolTipContentFormat_TimwWarp", "Start: {0}s\nDuration: {1}s\nPlay Rate: Variable"), StartTime, SectionLength);
 		}
 	}
 	return FText::GetEmpty();
@@ -733,14 +773,14 @@ FMargin FSkeletalAnimationSection::GetContentPadding() const
 
 int32 FSkeletalAnimationSection::OnPaintSection( FSequencerSectionPainter& Painter ) const
 {
+	using namespace UE::MovieScene;
+
 	const ESlateDrawEffect DrawEffects = Painter.bParentEnabled ? ESlateDrawEffect::None : ESlateDrawEffect::DisabledEffect;
 	
 	const FTimeToPixel& TimeToPixelConverter = Painter.GetTimeConverter();
 
 	int32 LayerId = Painter.PaintSectionBackground();
 
-	static const FSlateBrush* GenericDivider = FAppStyle::GetBrush("Sequencer.GenericDivider");
-
 	if (!Section.HasStartFrame() || !Section.HasEndFrame())
 	{
 		return LayerId;
@@ -748,34 +788,50 @@ int32 FSkeletalAnimationSection::OnPaintSection( FSequencerSectionPainter& Paint
 
 	FFrameRate TickResolution = TimeToPixelConverter.GetTickResolution();
 
-	// Add lines where the animation starts and ends/loops
-	const float AnimPlayRate = FMath::IsNearlyZero(Section.Params.PlayRate) || Section.Params.Animation == nullptr ? 1.0f : Section.Params.PlayRate * Section.Params.Animation->RateScale;
-	const float SeqLength = (Section.Params.GetSequenceLength() - TickResolution.AsSeconds(Section.Params.StartFrameOffset + Section.Params.EndFrameOffset)) / AnimPlayRate;
-	const float FirstLoopSeqLength = SeqLength - TickResolution.AsSeconds(Section.Params.FirstLoopStartFrameOffset) / AnimPlayRate;
+	const FFrameNumber StartFrame = Section.GetInclusiveStartFrame();
+	const FFrameNumber EndFrame   = Section.GetExclusiveEndFrame();
 
-	if (!FMath::IsNearlyZero(SeqLength, KINDA_SMALL_NUMBER) && SeqLength > 0)
+	if (UMovieScene* MovieScene = Section.GetTypedOuter<UMovieScene>())
 	{
-		float MaxOffset  = Section.GetRange().Size<FFrameTime>() / TickResolution;
-		float OffsetTime = FirstLoopSeqLength;
-		float StartTime  = Section.GetInclusiveStartFrame() / TickResolution;
+		FFrameRate FrameRate = MovieScene->GetTickResolution();
+		FMovieSceneSequenceTransform OuterToInnerTransform = Section.Params.MakeTransform(FrameRate, Section.GetRange());
+
+		// As seconds represented as a FFrameTime
+		FFrameTime LoopStart = FFrameTime::FromDecimal(Section.Params.StartFrameOffset / FrameRate);
+
+		FLinearColor SectionTint = Painter.GetSectionColor().LinearRGBToHSV();
+		SectionTint.B *= 0.1f;
+		SectionTint = SectionTint.HSVToLinearRGB();
 
-		while (OffsetTime < MaxOffset)
+		auto PaintTime = [&](FFrameTime Time)
 		{
-			float OffsetPixel = TimeToPixelConverter.SecondsToPixel(StartTime + OffsetTime) - TimeToPixelConverter.SecondsToPixel(StartTime);
+			float OffsetPixel = TimeToPixelConverter.FrameToPixel(Time);
 
-			FSlateDrawElement::MakeBox(
+			TArray<FVector2f> NewVector;
+			NewVector.Reserve(2);
+
+			NewVector.Add(FVector2f(OffsetPixel, 1.f));
+			NewVector.Add(FVector2f(OffsetPixel, Painter.SectionGeometry.Size.Y-2.f));
+
+			constexpr float Thickness = 1.f;
+			constexpr float DashLengthPx = 3.f;
+			FSlateDrawElement::MakeDashedLines(
 				Painter.DrawElements,
-				LayerId,
-				Painter.SectionGeometry.MakeChild(
-					FVector2D(2.f, Painter.SectionGeometry.Size.Y-2.f),
-					FSlateLayoutTransform(FVector2D(OffsetPixel, 1.f))
-				).ToPaintGeometry(),
-				GenericDivider,
-				DrawEffects
+				Painter.LayerId++,
+				Painter.SectionGeometry.ToPaintGeometry(),
+				MoveTemp(NewVector),
+				DrawEffects,
+				SectionTint,
+				Thickness,
+				DashLengthPx
 			);
+			return true;
+		};
 
-			OffsetTime += SeqLength;
-		}
+		OuterToInnerTransform.ExtractBoundariesWithinRange(StartFrame, EndFrame, [&PaintTime](FFrameTime StartTime){
+			PaintTime(StartTime);
+			return true;
+		});
 	}
 
 	return LayerId;
@@ -783,36 +839,16 @@ int32 FSkeletalAnimationSection::OnPaintSection( FSequencerSectionPainter& Paint
 
 void FSkeletalAnimationSection::BeginResizeSection()
 {
-	InitialFirstLoopStartOffsetDuringResize = Section.Params.FirstLoopStartFrameOffset;
-	InitialStartTimeDuringResize = Section.HasStartFrame() ? Section.GetInclusiveStartFrame() : 0;
+	const FFrameRate FrameRate = Section.GetTypedOuter<UMovieScene>()->GetTickResolution();
+	InitialDragTransform = MakeUnique<FMovieSceneSequenceTransform>(Section.Params.MakeTransform(FrameRate, Section.GetRange()));
 }
 
 void FSkeletalAnimationSection::ResizeSection(ESequencerSectionResizeMode ResizeMode, FFrameNumber ResizeTime)
 {
-	// Adjust the start offset when resizing from the beginning
-	if (ResizeMode == SSRM_LeadingEdge)
+	if (ResizeMode == SSRM_LeadingEdge && Section.Params.PlayRate.GetType() != EMovieSceneTimeWarpType::Custom)
 	{
-		// Get the effective animation length, in frames (rounded up), after taking into account start/end trimming.
 		const FFrameRate FrameRate = Section.GetTypedOuter<UMovieScene>()->GetTickResolution();
-		const FFrameNumber SeqLength = FrameRate.AsFrameTime(Section.Params.GetSequenceLength()).CeilToFrame() - Section.Params.StartFrameOffset - Section.Params.EndFrameOffset;
-
-		// Note that there is no scalar-multiplication support for frame numbers, so we need to multiply Value directly.
-		FFrameNumber ResizeAmount = (ResizeTime - InitialStartTimeDuringResize);
-		ResizeAmount.Value = (int32)FMath::Floor(ResizeAmount.Value * Section.Params.PlayRate);
-		FFrameNumber NewFirstLoopStartFrameOffset = InitialFirstLoopStartOffsetDuringResize + ResizeAmount;
-
-		// If the start offset exceeds the length of one loop, trim it back.
-		if (SeqLength > 0)
-		{
-			NewFirstLoopStartFrameOffset = NewFirstLoopStartFrameOffset % SeqLength;
-		}
-		// If the start offset is negative, add an extra loop at the beginning by making this start offset the complement.
-		if (NewFirstLoopStartFrameOffset < 0)
-		{
-			NewFirstLoopStartFrameOffset = SeqLength + NewFirstLoopStartFrameOffset;
-		}
-
-		Section.Params.FirstLoopStartFrameOffset = NewFirstLoopStartFrameOffset;
+		Section.Params.FirstLoopStartFrameOffset = (InitialDragTransform->TransformTime(ResizeTime).AsDecimal() * FrameRate).RoundToFrame();
 	}
 
 	ISequencerSection::ResizeSection(ResizeMode, ResizeTime);
@@ -825,30 +861,16 @@ void FSkeletalAnimationSection::BeginSlipSection()
 
 void FSkeletalAnimationSection::SlipSection(FFrameNumber SlipTime)
 {
-	FFrameRate FrameRate = Section.GetTypedOuter<UMovieScene>()->GetTickResolution();
-	FFrameNumber StartOffset = FrameRate.AsFrameNumber((SlipTime - InitialStartTimeDuringResize) / FrameRate * Section.Params.PlayRate);
-
-	StartOffset += InitialFirstLoopStartOffsetDuringResize;
-
-	if (StartOffset < 0)
-	{
-		// Ensure start offset is not less than 0 and adjust ResizeTime
-		SlipTime = SlipTime - StartOffset;
-
-		StartOffset = FFrameNumber(0);
-	}
-	else
-	{
-		// If the start offset exceeds the length of one loop, trim it back.
-		const FFrameNumber SeqLength = FrameRate.AsFrameNumber(Section.Params.GetSequenceLength()) - Section.Params.StartFrameOffset - Section.Params.EndFrameOffset;
-		StartOffset = StartOffset % SeqLength;
-	}
-
-	Section.Params.FirstLoopStartFrameOffset = StartOffset;
-
+	const FFrameRate FrameRate = Section.GetTypedOuter<UMovieScene>()->GetTickResolution();
+	Section.Params.FirstLoopStartFrameOffset = (InitialDragTransform->TransformTime(SlipTime).AsDecimal() * FrameRate).RoundToFrame();
 	ISequencerSection::SlipSection(SlipTime);
 }
 
+bool FSkeletalAnimationSection::RequestDeleteKeyArea(const TArray<FName>& KeyAreaNamePath)
+{
+	Section.DeleteChannels(KeyAreaNamePath);
+	return true;
+}
 
 void FSkeletalAnimationSection::CustomizePropertiesDetailsView(TSharedRef<IDetailsView> DetailsView, const FSequencerSectionPropertyDetailsViewCustomizationParams& InParams) const
 {
diff --git a/Engine/Source/Editor/MovieSceneTools/Private/TrackEditors/SkeletalAnimationTrackEditor.h b/Engine/Source/Editor/MovieSceneTools/Private/TrackEditors/SkeletalAnimationTrackEditor.h
--- a/Engine/Source/Editor/MovieSceneTools/Private/TrackEditors/SkeletalAnimationTrackEditor.h
+++ b/Engine/Source/Editor/MovieSceneTools/Private/TrackEditors/SkeletalAnimationTrackEditor.h
@@ -5,6 +5,7 @@
 #include "CoreMinimal.h"
 #include "Misc/Guid.h"
 #include "Templates/SubclassOf.h"
+#include "Templates/UniquePtr.h"
 #include "Widgets/SWidget.h"
 #include "ISequencer.h"
 #include "MovieSceneTrack.h"
@@ -14,6 +15,8 @@
 #include "EditModes/SkeletalAnimationTrackEditMode.h"
 
 struct FAssetData;
+struct FMovieSceneTimeWarpChannel;
+struct FMovieSceneSequenceTransform;
 class FMenuBuilder;
 class FSequencerSectionPainter;
 class UMovieSceneSkeletalAnimationSection;
@@ -142,7 +145,7 @@ public:
 	FSkeletalAnimationSection( UMovieSceneSection& InSection, TWeakPtr<ISequencer> InSequencer);
 
 	/** Virtual destructor. */
-	virtual ~FSkeletalAnimationSection() { }
+	virtual ~FSkeletalAnimationSection();
 
 public:
 
@@ -163,6 +166,7 @@ public:
 	virtual void BuildSectionContextMenu(FMenuBuilder& MenuBuilder, const FGuid& InObjectBinding) override;
 	virtual void BeginDilateSection() override;
 	virtual void DilateSection(const TRange<FFrameNumber>& NewRange, float DilationFactor) override;
+	virtual bool RequestDeleteKeyArea(const TArray<FName>& KeyAreaNamePath) override;
 
 
 private:
@@ -175,9 +179,7 @@ private:
 	/** Used to draw animation frame, need selection state and local time*/
 	TWeakPtr<ISequencer> Sequencer;
 
-	/** Cached first loop start offset value valid only during resize */
-	FFrameNumber InitialFirstLoopStartOffsetDuringResize;
-
-	/** Cached start time valid only during resize */
-	FFrameNumber InitialStartTimeDuringResize;
+	TUniquePtr<FMovieSceneSequenceTransform> InitialDragTransform;
+	TUniquePtr<FMovieSceneTimeWarpChannel> PreDilateChannel;
+	double PreDilatePlayRate;
 };
diff --git a/Engine/Source/Editor/MovieSceneTools/Private/TrackEditors/SubTrackEditor.cpp b/Engine/Source/Editor/MovieSceneTools/Private/TrackEditors/SubTrackEditor.cpp
--- a/Engine/Source/Editor/MovieSceneTools/Private/TrackEditors/SubTrackEditor.cpp
+++ b/Engine/Source/Editor/MovieSceneTools/Private/TrackEditors/SubTrackEditor.cpp
@@ -441,8 +441,11 @@ void FSubTrackEditor::Resize(float NewSize, UMovieSceneTrack* InTrack)
 	}
 }
 
-/* FSubTrackEditor
- *****************************************************************************/
+bool FSubTrackEditor::GetDefaultExpansionState(UMovieSceneTrack* InTrack) const
+{
+	return true;
+}
+
 
 void FSubTrackEditor::InsertSection(UMovieSceneTrack* Track)
 {
@@ -546,7 +549,6 @@ void FSubTrackEditor::CreateNewTake(UMovieSceneSubSection* Section)
 
 		TRange<FFrameNumber> NewSectionRange         = Section->GetRange();
 		FFrameNumber         NewSectionStartOffset   = Section->Parameters.StartFrameOffset;
-		float                NewSectionTimeScale     = Section->Parameters.TimeScale;
 		int32                NewSectionPrerollFrames = Section->GetPreRollFrames();
 		int32                NewRowIndex          = Section->GetRowIndex();
 		FFrameNumber         NewSectionStartTime     = NewSectionRange.GetLowerBound().IsClosed() ? UE::MovieScene::DiscreteInclusiveLower(NewSectionRange) : 0;
@@ -565,7 +567,7 @@ void FSubTrackEditor::CreateNewTake(UMovieSceneSubSection* Section)
 
 			NewSection->SetRange(NewSectionRange);
 			NewSection->Parameters.StartFrameOffset = NewSectionStartOffset;
-			NewSection->Parameters.TimeScale = NewSectionTimeScale;
+			NewSection->Parameters.TimeScale = Section->Parameters.TimeScale.DeepCopy(NewSection);
 			NewSection->SetPreRollFrames(NewSectionPrerollFrames);
 			NewSection->SetRowIndex(NewRowIndex);
 			NewSection->SetColorTint(NewSectionColorTint);
@@ -615,7 +617,6 @@ void FSubTrackEditor::ChangeTake(UMovieSceneSequence* Sequence)
 
 		TRange<FFrameNumber> NewSectionRange = Section->GetRange();
 		FFrameNumber		 NewSectionStartOffset = Section->Parameters.StartFrameOffset;
-		float                NewSectionTimeScale = Section->Parameters.TimeScale;
 		int32                NewSectionPrerollFrames = Section->GetPreRollFrames();
 		int32                NewRowIndex = Section->GetRowIndex();
 		FFrameNumber         NewSectionStartTime = NewSectionRange.GetLowerBound().IsClosed() ? UE::MovieScene::DiscreteInclusiveLower(NewSectionRange) : 0;
@@ -631,7 +632,7 @@ void FSubTrackEditor::ChangeTake(UMovieSceneSequence* Sequence)
 
 			NewSection->SetRange(NewSectionRange);
 			NewSection->Parameters.StartFrameOffset = NewSectionStartOffset;
-			NewSection->Parameters.TimeScale = NewSectionTimeScale;
+			NewSection->Parameters.TimeScale = Section->Parameters.TimeScale.DeepCopy(NewSection);
 			NewSection->SetPreRollFrames(NewSectionPrerollFrames);
 			NewSection->SetRowIndex(NewSectionRowIndex);
 			NewSection->SetColorTint(NewSectionColorTint);
