Sequencer: Added time-warp capabilities to sequences, sub-sections and skeletal animation sections

diff --git a/Engine/Source/Editor/MovieSceneTools/Private/TrackEditors/SubTrackEditorBase.cpp b/Engine/Source/Editor/MovieSceneTools/Private/TrackEditors/SubTrackEditorBase.cpp
--- a/Engine/Source/Editor/MovieSceneTools/Private/TrackEditors/SubTrackEditorBase.cpp
+++ b/Engine/Source/Editor/MovieSceneTools/Private/TrackEditors/SubTrackEditorBase.cpp
@@ -10,435 +10,397 @@
 #include "SequencerSettings.h"
 #include "Tracks/MovieSceneCinematicShotTrack.h"
 #include "Tracks/MovieSceneSubTrack.h"
+#include "Channels/MovieSceneTimeWarpChannel.h"
+#include "Channels/MovieSceneChannelProxy.h"
+#include "Variants/MovieSceneTimeWarpGetter.h"
+#include "TimeToPixel.h"
 
 #define LOCTEXT_NAMESPACE "FSubTrackEditorBase"
 
 FSubSectionPainterResult FSubSectionPainterUtil::PaintSection(TSharedPtr<const ISequencer> Sequencer, const UMovieSceneSubSection& SectionObject, FSequencerSectionPainter& InPainter, FSubSectionPainterParams Params)
 {
-    const TRange<FFrameNumber> SectionRange = SectionObject.GetRange();
-    if (SectionRange.GetLowerBound().IsOpen() || SectionRange.GetUpperBound().IsOpen())
-    {
-        return FSSPR_InvalidSection;
-    }
-
-    const int32 SectionSize = UE::MovieScene::DiscreteSize(SectionRange);
-    if (SectionSize <= 0)
-    {
-        return FSSPR_InvalidSection;
-    }
-
-    UMovieSceneSequence* InnerSequence = SectionObject.GetSequence();
-    if (InnerSequence == nullptr || InnerSequence->GetMovieScene() == nullptr)
-    {
-        return FSSPR_NoInnerSequence;
-    }
-
-    const ESlateDrawEffect DrawEffects = InPainter.bParentEnabled ? ESlateDrawEffect::None : ESlateDrawEffect::DisabledEffect;
-    if (SectionObject.Parameters.bCanLoop && !FMath::IsNearlyZero(SectionObject.OuterToInnerTransform().GetTimeScale()))
-    {
-        DoPaintLoopingSection(SectionObject, *InnerSequence, InPainter, DrawEffects);
-    }
-    else
-    {
-        DoPaintNonLoopingSection(SectionObject, *InnerSequence, InPainter, DrawEffects);
-    }
-
-    UMovieScene* MovieScene = InnerSequence->GetMovieScene();
-    const int32 NumTracks = MovieScene->GetPossessableCount() + MovieScene->GetSpawnableCount() + MovieScene->GetTracks().Num();
-
-    FVector2D TopLeft = InPainter.SectionGeometry.AbsoluteToLocal(InPainter.SectionClippingRect.GetTopLeft()) + FVector2D(1.f, -1.f);
-
-    FSlateFontInfo FontInfo = FAppStyle::GetFontStyle("NormalFont");
-
-    TSharedRef<FSlateFontCache> FontCache = FSlateApplication::Get().GetRenderer()->GetFontCache();
-
-    auto GetFontHeight = [&]
-    {
-        return FontCache->GetMaxCharacterHeight(FontInfo, 1.f) + FontCache->GetBaseline(FontInfo, 1.f);
-    };
-    while (GetFontHeight() > InPainter.SectionGeometry.Size.Y && FontInfo.Size > 11)
-    {
-        FontInfo.Size = FMath::Max(FMath::FloorToInt(FontInfo.Size - 6.f), 11);
-    }
-
-    uint32 LayerId = InPainter.LayerId;
-    FMargin ContentPadding = Params.ContentPadding;
-
-    FText TrackNumText, NetworkText;
-    if (Params.bShowTrackNum)
-    {
-        TrackNumText = FText::Format(LOCTEXT("NumTracksFormat", "{0} track(s)"), FText::AsNumber(NumTracks));
-    }
-
-    EMovieSceneServerClientMask NetworkMask = SectionObject.GetNetworkMask();
-    if (NetworkMask == EMovieSceneServerClientMask::Client)
-    {
-        NetworkText = LOCTEXT("SubSectionClientOnlyText", "(client only)");
-    }
-    else if (NetworkMask == EMovieSceneServerClientMask::Server)
-    {
-        NetworkText = LOCTEXT("SubSectionServerOnlyText", "(server only)");
-    }
-
-    FText SectionText;
-    if (!TrackNumText.IsEmpty() || !NetworkText.IsEmpty())
-    {
-        SectionText = FText::Format(LOCTEXT("SectionTextFormat", "{0} {1}"), TrackNumText, NetworkText);
-    }
-    else if (!TrackNumText.IsEmpty())
-    {
-        SectionText = TrackNumText;
-    }
-    else
-    {
-        SectionText = NetworkText;
-    }
-
-
-    if (!SectionText.IsEmpty())
-    {
+	using namespace UE::MovieScene;
+
+	const TRange<FFrameNumber> SectionRange = SectionObject.GetRange();
+	if (SectionRange.GetLowerBound().IsOpen() || SectionRange.GetUpperBound().IsOpen())
+	{
+		return FSSPR_InvalidSection;
+	}
+
+	const int32 SectionSize = UE::MovieScene::DiscreteSize(SectionRange);
+	if (SectionSize <= 0)
+	{
+		return FSSPR_InvalidSection;
+	}
+
+	UMovieSceneSequence* InnerSequence = SectionObject.GetSequence();
+	if (InnerSequence == nullptr || InnerSequence->GetMovieScene() == nullptr)
+	{
+		return FSSPR_NoInnerSequence;
+	}
+
+	const ESlateDrawEffect DrawEffects = InPainter.bParentEnabled ? ESlateDrawEffect::None : ESlateDrawEffect::DisabledEffect;
+
+	PaintSectionBounds(SectionObject, *InnerSequence, InPainter, DrawEffects);
+
+	UMovieScene* MovieScene = InnerSequence->GetMovieScene();
+	const int32 NumTracks = MovieScene->GetPossessableCount() + MovieScene->GetSpawnableCount() + MovieScene->GetTracks().Num();
+
+	FVector2D TopLeft = InPainter.SectionGeometry.AbsoluteToLocal(InPainter.SectionClippingRect.GetTopLeft()) + FVector2D(1.f, -1.f);
+
+	FSlateFontInfo FontInfo = FAppStyle::GetFontStyle("NormalFont");
+
+	TSharedRef<FSlateFontCache> FontCache = FSlateApplication::Get().GetRenderer()->GetFontCache();
+
+	auto GetFontHeight = [&]
+	{
+		return FontCache->GetMaxCharacterHeight(FontInfo, 1.f) + FontCache->GetBaseline(FontInfo, 1.f);
+	};
+	while (GetFontHeight() > InPainter.SectionGeometry.Size.Y && FontInfo.Size > 11)
+	{
+		FontInfo.Size = FMath::Max(FMath::FloorToInt(FontInfo.Size - 6.f), 11);
+	}
+
+	uint32 LayerId = InPainter.LayerId;
+	FMargin ContentPadding = Params.ContentPadding;
+
+	FText TrackNumText, NetworkText;
+	if (Params.bShowTrackNum)
+	{
+		TrackNumText = FText::Format(LOCTEXT("NumTracksFormat", "{0} track(s)"), FText::AsNumber(NumTracks));
+	}
+
+	EMovieSceneServerClientMask NetworkMask = SectionObject.GetNetworkMask();
+	if (NetworkMask == EMovieSceneServerClientMask::Client)
+	{
+		NetworkText = LOCTEXT("SubSectionClientOnlyText", "(client only)");
+	}
+	else if (NetworkMask == EMovieSceneServerClientMask::Server)
+	{
+		NetworkText = LOCTEXT("SubSectionServerOnlyText", "(server only)");
+	}
+
+	FText SectionText;
+	if (!TrackNumText.IsEmpty() || !NetworkText.IsEmpty())
+	{
+		SectionText = FText::Format(LOCTEXT("SectionTextFormat", "{0} {1}"), TrackNumText, NetworkText);
+	}
+	else if (!TrackNumText.IsEmpty())
+	{
+		SectionText = TrackNumText;
+	}
+	else
+	{
+		SectionText = NetworkText;
+	}
+
+
+	if (!SectionText.IsEmpty())
+	{
 		FSlateClippingZone ClippingZone(InPainter.SectionClippingRect.InsetBy(FMargin(1.0f)));
 
 		InPainter.DrawElements.PushClip(ClippingZone);
 		
 		FSlateDrawElement::MakeText(
-            InPainter.DrawElements,
-            ++LayerId,
-            InPainter.SectionGeometry.ToPaintGeometry(
-                FVector2D(InPainter.SectionGeometry.Size.X, GetFontHeight()),
-                FSlateLayoutTransform(TopLeft + FVector2D(ContentPadding.Left, ContentPadding.Top) + FVector2D(11.f, GetFontHeight()*2.f))
-            ),
-            SectionText,
-            FontInfo,
-            DrawEffects,
-            FColor(200, 200, 200, static_cast<uint8>(255 * InPainter.GhostAlpha))
-        );
+			InPainter.DrawElements,
+			++LayerId,
+			InPainter.SectionGeometry.ToPaintGeometry(
+				FVector2D(InPainter.SectionGeometry.Size.X, GetFontHeight()),
+				FSlateLayoutTransform(TopLeft + FVector2D(ContentPadding.Left, ContentPadding.Top) + FVector2D(11.f, GetFontHeight()*2.f))
+			),
+			SectionText,
+			FontInfo,
+			DrawEffects,
+			FColor(200, 200, 200, static_cast<uint8>(255 * InPainter.GhostAlpha))
+		);
 
 		InPainter.DrawElements.PopClip();
 	}
 
-    InPainter.LayerId = LayerId;
+	InPainter.LayerId = LayerId;
 
-    return FSSPR_Success;
+	return FSSPR_Success;
 }
 
-void FSubSectionPainterUtil::DoPaintNonLoopingSection(const UMovieSceneSubSection& SectionObject, const UMovieSceneSequence& InnerSequence, FSequencerSectionPainter& InPainter, ESlateDrawEffect DrawEffects)
+void FSubSectionPainterUtil::PaintSectionBounds(const UMovieSceneSubSection& SectionObject, const UMovieSceneSequence& InnerSequence, FSequencerSectionPainter& InPainter, ESlateDrawEffect DrawEffects)
 {
-    const uint8 GhostAlpha = static_cast<uint8>(255 * InPainter.GhostAlpha);
+	const uint8        GhostAlpha        = static_cast<uint8>(255 * InPainter.GhostAlpha);
+	const FFrameNumber SectionStartFrame = SectionObject.GetInclusiveStartFrame();
+	const UMovieScene* MovieScene        = InnerSequence.GetMovieScene();
+
+	const FMovieSceneSequenceTransform        OuterToInnerTransform = SectionObject.OuterToInnerTransform();
+	const FMovieSceneInverseSequenceTransform InnerToOuterTransform = OuterToInnerTransform.Inverse();
 
-    const FFrameNumber SectionStartFrame = SectionObject.GetInclusiveStartFrame();
+	const TRange<FFrameTime> SectionTimeRange = UE::MovieScene::ConvertToFrameTimeRange(SectionObject.GetRange());
+	const int32              SectionSize      = UE::MovieScene::DiscreteSize(SectionObject.GetRange());
 
-    const TRange<FFrameNumber> SectionRange = SectionObject.GetRange();
-    const int32 SectionSize = UE::MovieScene::DiscreteSize(SectionRange);
-    const float PixelsPerFrame = InPainter.SectionGeometry.Size.X / float(SectionSize);
 
-    UMovieScene* MovieScene = InnerSequence.GetMovieScene();
-    TRange<FFrameNumber> PlaybackRange = MovieScene->GetPlaybackRange();
+	TOptional<FFrameTime> FirstTime, LastTime;
+	auto PaintDashedLine = [&](FFrameTime InTime, FColor Tint, float LineNudge)
+	{
+		float Offset = InPainter.GetTimeConverter().FrameToPixel(InTime) + LineNudge;
 
-    // We're in the non-looping case so we know we have a purely linear transform (unless zero timescale)
-	const FMovieSceneSequenceTransform OuterToInnerTransform = SectionObject.OuterToInnerTransform();
-    const FMovieSceneSequenceTransform InnerToOuterTransform = OuterToInnerTransform.InverseNoLooping();
-    const FFrameNumber PlaybackStart = (UE::MovieScene::DiscreteInclusiveLower(PlaybackRange) * InnerToOuterTransform).FloorToFrame();
-	const FFrameNumber PlaybackEnd = (UE::MovieScene::DiscreteExclusiveUpper(PlaybackRange) * InnerToOuterTransform).FloorToFrame();
-	const TRange<FFrameNumber> PlaybackRangeTransformed(PlaybackStart, PlaybackEnd);
-	bool bZeroTimescale = !FMath::IsFinite(InnerToOuterTransform.GetTimeScale());
+		if (!FirstTime)
+		{
+			FirstTime = InTime;
+		}
+		LastTime = InTime;
 
-	// Section contains the beginning but not the end- draw beginning, draw dark tint before it
-	// Section contains the end but not the beginning- draw dark tint, draw dark tint after it
-	// Section contains both end and beginning- draw both, dark tint before beginning, dark tint after end
-	// Section contains neither or zero-timescale and the frame isn't in the section range- draw dark tint over entire section
-	// Zero-timescale and the frame is in the sectionrange- don't draw anything extra
-	const int32 StartOffset = (PlaybackStart - SectionStartFrame).Value;
-	const int32 EndOffset = (PlaybackEnd - SectionStartFrame).Value;
+		TArray<FVector2f> NewVector;
+		NewVector.Reserve(2);
 
-	if (bZeroTimescale)
+		NewVector.Add(FVector2f(Offset, 0.f));
+		NewVector.Add(FVector2f(Offset, InPainter.SectionGeometry.Size.Y));
+
+		constexpr float Thickness = 1.f;
+		constexpr float DashLengthPx = 3.f;
+		FSlateDrawElement::MakeDashedLines(
+			InPainter.DrawElements,
+			InPainter.LayerId++,
+			InPainter.SectionGeometry.ToPaintGeometry(),
+			MoveTemp(NewVector),
+			DrawEffects,
+			Tint,
+			Thickness,
+			DashLengthPx
+		);
+		return true;
+	};
+
+
+
+	auto PaintLine = [&](FFrameTime InTime, FColor Tint, float LineNudge)
 	{
-		// If we're zero-timescale we can't use the inverse transform to determine bounds, as the transformation is non-deterministic.
-		// Instead, we just ask whether the playback start is in the interior range, and if it doesn't, draw dark tint over the entire section
-		
-		const FFrameTime InnerFrame = SectionStartFrame * OuterToInnerTransform;
-		if (!PlaybackRange.Contains(InnerFrame.FloorToFrame()))
+		float Offset = InPainter.GetTimeConverter().FrameToPixel(InTime) + LineNudge;
+
+		if (!FirstTime)
 		{
-			// add dark tint over the entire section
-			FSlateDrawElement::MakeBox(
-				InPainter.DrawElements,
-				InPainter.LayerId++,
-				InPainter.SectionGeometry.ToPaintGeometry(
-					FVector2f(InPainter.SectionGeometry.Size.X, InPainter.SectionGeometry.Size.Y),
-					FSlateLayoutTransform()
-				),
-				FAppStyle::GetBrush("WhiteBrush"),
-				DrawEffects,
-				FLinearColor::Black.CopyWithNewOpacity(0.5f)
-			);
+			FirstTime = InTime;
 		}
-	}
-	else if (TRange<FFrameNumber>::Intersection(PlaybackRangeTransformed, SectionRange).IsEmpty())
-	{
-		// add dark tint over the entire section
+		LastTime = InTime;
+
+		// add green line for playback start
 		FSlateDrawElement::MakeBox(
 			InPainter.DrawElements,
 			InPainter.LayerId++,
 			InPainter.SectionGeometry.ToPaintGeometry(
-				FVector2f(InPainter.SectionGeometry.Size.X, InPainter.SectionGeometry.Size.Y),
-				FSlateLayoutTransform()
+				FVector2f(1.0f, InPainter.SectionGeometry.Size.Y),
+				FSlateLayoutTransform(FVector2f(Offset, 0.f))
 			),
 			FAppStyle::GetBrush("WhiteBrush"),
 			DrawEffects,
-			FLinearColor::Black.CopyWithNewOpacity(0.5f)
+			Tint
 		);
+		return true;
+	};
+
+
+	FColor GreenTint(32, 128, 32, GhostAlpha);	// 120, 75, 50 (HSV)
+	FColor RedTint(128, 32, 32, GhostAlpha);	// 0,   75, 50 (HSV)
+	FLinearColor SectionTint = InPainter.GetSectionColor().LinearRGBToHSV();
+	SectionTint.B *= 0.1f;
+	SectionTint = SectionTint.HSVToLinearRGB();
+
+	auto PaintBoundary = [PaintDashedLine, SectionTint](FFrameTime Time)
+	{
+		PaintDashedLine(Time, SectionTint.ToFColor(true), 0.f);
+		return true;
+	};
+
+	if (!OuterToInnerTransform.ExtractBoundariesWithinRange(SectionStartFrame, SectionStartFrame+SectionSize, PaintBoundary))
+	{
+		// Just use the playback range
+		TOptional<FFrameTime> StartBound = InnerToOuterTransform.TryTransformTime(MovieScene->GetPlaybackRange().GetLowerBoundValue());
+		TOptional<FFrameTime> EndBound   = InnerToOuterTransform.TryTransformTime(MovieScene->GetPlaybackRange().GetUpperBoundValue());
+
+		if (StartBound)
+		{
+			PaintLine(StartBound.GetValue(), GreenTint, 0.f);
+		}
+		if (EndBound)
+		{
+			PaintLine(EndBound.GetValue(), RedTint, 1.f);
+		}
 	}
-	else
+
+	// Paint the dimmed section range before the first valid time
+	if (FirstTime && FirstTime.GetValue() > SectionTimeRange.GetLowerBoundValue())
 	{
-		if (SectionRange.Contains(PlaybackStart))
+		TOptional<FFrameTime> PreceedingInvalidTime;
+		auto SetPreceedingInvalidTime = [&PreceedingInvalidTime](FFrameTime InTime)
+		{
+			PreceedingInvalidTime = InTime;
+			return false;
+		};
+
+		OuterToInnerTransform.ExtractBoundariesWithinRange(MIN_int32, SectionTimeRange.GetLowerBoundValue(), SetPreceedingInvalidTime);
+		if (!PreceedingInvalidTime)
 		{
-			// add dark tint for left out-of-bounds
+			const float WidthFactor = static_cast<float>((FirstTime.GetValue().AsDecimal() - SectionTimeRange.GetLowerBoundValue().AsDecimal()) / SectionTimeRange.Size<FFrameTime>().AsDecimal());
+
 			FSlateDrawElement::MakeBox(
 				InPainter.DrawElements,
 				InPainter.LayerId++,
 				InPainter.SectionGeometry.ToPaintGeometry(
-					FVector2f(StartOffset * PixelsPerFrame, InPainter.SectionGeometry.Size.Y),
+					FVector2f(InPainter.SectionGeometry.Size.X * WidthFactor-1, InPainter.SectionGeometry.Size.Y),
 					FSlateLayoutTransform()
 				),
 				FAppStyle::GetBrush("WhiteBrush"),
 				DrawEffects,
 				FLinearColor::Black.CopyWithNewOpacity(0.5f)
 			);
-
-			// add green line for playback start
-			FSlateDrawElement::MakeBox(
-				InPainter.DrawElements,
-				InPainter.LayerId++,
-				InPainter.SectionGeometry.ToPaintGeometry(
-					FVector2f(1.0f, InPainter.SectionGeometry.Size.Y),
-					FSlateLayoutTransform(FVector2f(StartOffset * PixelsPerFrame, 0.f))
-				),
-				FAppStyle::GetBrush("WhiteBrush"),
-				DrawEffects,
-				FColor(32, 128, 32, GhostAlpha)	// 120, 75, 50 (HSV)
-			);
 		}
+	}
 
-		if (SectionRange.Contains(PlaybackEnd))
+	// Paint the dimmed section range after the last valid time
+	if (LastTime && LastTime.GetValue() < SectionTimeRange.GetUpperBoundValue())
+	{
+		TOptional<FFrameTime> ProceedingInvalidTime;
+		auto SetProceedingInvalidTime = [&ProceedingInvalidTime](FFrameTime InTime)
 		{
-			// add dark tint for right out-of-bounds
-			FSlateDrawElement::MakeBox(
-				InPainter.DrawElements,
-				InPainter.LayerId++,
-				InPainter.SectionGeometry.ToPaintGeometry(
-					FVector2f((SectionSize - EndOffset) * PixelsPerFrame, InPainter.SectionGeometry.Size.Y),
-					FSlateLayoutTransform(FVector2f(EndOffset * PixelsPerFrame, 0.f))
-				),
-				FAppStyle::GetBrush("WhiteBrush"),
-				DrawEffects,
-				FLinearColor::Black.CopyWithNewOpacity(0.5f)
-			);
+			ProceedingInvalidTime = InTime;
+			return false;
+		};
 
+		OuterToInnerTransform.ExtractBoundariesWithinRange(SectionTimeRange.GetUpperBoundValue(), MAX_int32, SetProceedingInvalidTime);
+		if (!ProceedingInvalidTime)
+		{
+			const float Offset      = InPainter.GetTimeConverter().FrameToPixel(LastTime.GetValue()) + 1;
+			const float WidthFactor = 1.f - static_cast<float>((LastTime.GetValue().AsDecimal() - SectionTimeRange.GetLowerBoundValue().AsDecimal()) / SectionTimeRange.Size<FFrameTime>().AsDecimal());
 
-			// add red line for playback end
 			FSlateDrawElement::MakeBox(
 				InPainter.DrawElements,
 				InPainter.LayerId++,
 				InPainter.SectionGeometry.ToPaintGeometry(
-					FVector2f(1.0f, InPainter.SectionGeometry.Size.Y),
-					FSlateLayoutTransform(FVector2f(EndOffset * PixelsPerFrame, 0.f))
+					FVector2f(InPainter.SectionGeometry.Size.X * WidthFactor-1, InPainter.SectionGeometry.Size.Y),
+					FSlateLayoutTransform(FVector2f(Offset, 0.f))
 				),
 				FAppStyle::GetBrush("WhiteBrush"),
 				DrawEffects,
-				FColor(128, 32, 32, GhostAlpha)	// 0, 75, 50 (HSV)
+				FLinearColor::Black.CopyWithNewOpacity(0.5f)
 			);
 		}
 	}
 }
 
-void FSubSectionPainterUtil::DoPaintLoopingSection(const UMovieSceneSubSection& SectionObject, const UMovieSceneSequence& InnerSequence, FSequencerSectionPainter& InPainter, ESlateDrawEffect DrawEffects)
-{
-    const TRange<FFrameNumber> SectionRange = SectionObject.GetRange();
-    const int32 SectionSize = UE::MovieScene::DiscreteSize(SectionRange);
-    const float PixelsPerFrame = InPainter.SectionGeometry.Size.X / float(SectionSize);
-
-    const float InvTimeScale = FMath::IsNearlyZero(SectionObject.Parameters.TimeScale) ? 1.0f : 1.0f / SectionObject.Parameters.TimeScale;
-
-    const UMovieScene* MovieScene = InnerSequence.GetMovieScene();
-    const TRange<FFrameNumber> InnerPlaybackRange = UMovieSceneSubSection::GetValidatedInnerPlaybackRange(SectionObject.Parameters, *MovieScene);
-
-    const FFrameNumber InnerSubSeqLength = UE::MovieScene::DiscreteSize(InnerPlaybackRange);
-    const FFrameNumber InnerSubSeqFirstLoopLength = InnerSubSeqLength - SectionObject.Parameters.FirstLoopStartFrameOffset;
-
-    const FFrameRate OuterFrameRate   = SectionObject.GetTypedOuter<UMovieScene>()->GetTickResolution();
-    const FFrameRate InnerFrameRate   = MovieScene->GetTickResolution();
-    const FFrameNumber OuterSubSeqLength = (ConvertFrameTime(InnerSubSeqLength, InnerFrameRate, OuterFrameRate) * InvTimeScale).FrameNumber;
-    const FFrameNumber OuterSubSeqFirstLoopLength = (ConvertFrameTime(InnerSubSeqFirstLoopLength, InnerFrameRate, OuterFrameRate) * InvTimeScale).FrameNumber;
-
-    FFrameNumber CurOffsetFrame = FMath::Max(OuterSubSeqFirstLoopLength, FFrameNumber(0));
-
-    const uint8 GhostAlpha = static_cast<uint8>(255 * InPainter.GhostAlpha);
-
-    // Draw separators where the sub-sequence is looping. To be consistent with the non-looping case, we draw a red and green
-    // separator back to back.
-    uint32 MaxLoopBoundaries = 100;
-    while (CurOffsetFrame < SectionSize)
-    {
-        const int32 CurOffset = CurOffsetFrame.Value;
-        FSlateDrawElement::MakeBox(
-            InPainter.DrawElements,
-            InPainter.LayerId++,
-            InPainter.SectionGeometry.ToPaintGeometry(
-                FVector2f(1.0f, InPainter.SectionGeometry.Size.Y),
-                FSlateLayoutTransform(FVector2f(CurOffset * PixelsPerFrame, 0.f))
-            ),
-            FAppStyle::GetBrush("WhiteBrush"),
-            DrawEffects,
-            FColor(32, 128, 32, GhostAlpha)	// 120, 75, 50 (HSV)
-        );
-        if (CurOffset > 0)
-        {
-            FSlateDrawElement::MakeBox(
-                InPainter.DrawElements,
-                InPainter.LayerId++,
-                InPainter.SectionGeometry.ToPaintGeometry(
-                    FVector2f(1.0f, InPainter.SectionGeometry.Size.Y),
-                    FSlateLayoutTransform(FVector2f(CurOffset * PixelsPerFrame - 1.f, 0.f))
-                ),
-                FAppStyle::GetBrush("WhiteBrush"),
-                DrawEffects,
-                FColor(128, 32, 32, GhostAlpha)	// 0, 75, 50 (HSV)
-            );
-        }
-
-        CurOffsetFrame += OuterSubSeqLength;
-
-        if ((--MaxLoopBoundaries) == 0)
-        {
-            break;
-        }
-    }
-}
 
 FSubSectionEditorUtil::FSubSectionEditorUtil(UMovieSceneSubSection& InSection)
-    : SectionObject(InSection)
-    , InitialStartOffsetDuringResize(0)
-    , InitialStartTimeDuringResize(0)
-	, PreviousTimeScale(1.f)
+	: SectionObject(InSection)
+	, PreDilateTimeScale(1.f)
 {
 }
 
+FSubSectionEditorUtil::~FSubSectionEditorUtil()
+{}
+
 void FSubSectionEditorUtil::BeginResizeSection()
 {
-    InitialStartOffsetDuringResize = SectionObject.Parameters.bCanLoop ? SectionObject.Parameters.FirstLoopStartFrameOffset : SectionObject.Parameters.StartFrameOffset;
-    InitialStartTimeDuringResize = SectionObject.HasStartFrame() ? SectionObject.GetInclusiveStartFrame() : 0;
+	InitialDragTransform = MakeUnique<FMovieSceneSequenceTransform>(SectionObject.OuterToInnerTransform());
 }
 
 FFrameNumber FSubSectionEditorUtil::ResizeSection(ESequencerSectionResizeMode ResizeMode, FFrameNumber ResizeTime)
 {
-    UMovieSceneSequence* InnerSequence = SectionObject.GetSequence();
-    UMovieScene* InnerMovieScene = InnerSequence ? InnerSequence->GetMovieScene() : nullptr;
-
-    if (ResizeMode == SSRM_LeadingEdge && InnerMovieScene != nullptr)
-    {
-        // Adjust the start offset when resizing from the beginning
-        FMovieSceneSectionParameters& SectionParameters = SectionObject.Parameters;
-
-        const FFrameRate    OuterFrameRate   = SectionObject.GetTypedOuter<UMovieScene>()->GetTickResolution();
-        const FFrameRate    InnerFrameRate   = InnerSequence->GetMovieScene()->GetTickResolution();
-        const FFrameNumber  ResizeDifference = ResizeTime - InitialStartTimeDuringResize;
-
-        const FFrameNumber InnerResizeDifference = (ConvertFrameTime(ResizeDifference, OuterFrameRate, InnerFrameRate) * SectionParameters.TimeScale).FrameNumber;
-        FFrameNumber NewStartOffset = InitialStartOffsetDuringResize + InnerResizeDifference;
-
-        const int32 InnerPlaybackLength = UE::MovieScene::DiscreteSize(InnerMovieScene->GetPlaybackRange());
-        const float InvTimeScale = FMath::IsNearlyZero(SectionParameters.TimeScale) ? 1.0f : 1.0f / SectionParameters.TimeScale;
-        const FFrameNumber InnerLoopLength = InnerPlaybackLength - SectionParameters.StartFrameOffset - SectionParameters.EndFrameOffset;
-
-        const bool bCanLoop = SectionParameters.bCanLoop;
-        if (!bCanLoop)
-        {
-            if (NewStartOffset < 0)
-            {
-                // Ensure start offset is not less than 0, clamp resize time.
-                const FFrameTime OuterFrameTimeOver = ConvertFrameTime(FFrameTime::FromDecimal(NewStartOffset.Value / SectionParameters.TimeScale), InnerFrameRate, OuterFrameRate);
-                ResizeTime = ResizeTime - OuterFrameTimeOver.GetFrame(); 
-                NewStartOffset = 0;
-            }
-            SectionParameters.StartFrameOffset = NewStartOffset;
-        }
-        else
-        {
-			// If the start offset exceeds the length of one loop, trim it back.
-			NewStartOffset = NewStartOffset % InnerLoopLength;
-            if (NewStartOffset < 0)
-            {
-                // Move the first loop offset forward so that it's positive into the new first loop
-                // that we just "revealed" by pulling the left edge back.
-                NewStartOffset += InnerLoopLength;
-            }
-            SectionParameters.FirstLoopStartFrameOffset = NewStartOffset;
-        }
-    }
-
-    return ResizeTime;
+	UMovieSceneSequence* InnerSequence = SectionObject.GetSequence();
+	UMovieScene* InnerMovieScene = InnerSequence ? InnerSequence->GetMovieScene() : nullptr;
+
+	if (SectionObject.Parameters.TimeScale.GetType() == EMovieSceneTimeWarpType::Custom)
+	{
+		return ResizeTime;
+	}
+
+	if (ResizeMode == SSRM_LeadingEdge && InnerMovieScene != nullptr)
+	{
+		TRange<FFrameNumber> Range = InnerMovieScene->GetPlaybackRange();
+
+		// Find the new inner offset as an absolute time
+		FFrameNumber NewStartTime = InitialDragTransform->TransformTime(ResizeTime).RoundToFrame();
+
+		// Inner offset must be relative to playback start
+		NewStartTime = NewStartTime - Range.GetLowerBoundValue();
+
+		if (SectionObject.Parameters.bCanLoop)
+		{
+			SectionObject.Parameters.FirstLoopStartFrameOffset = NewStartTime;
+		}
+		else
+		{
+			SectionObject.Parameters.StartFrameOffset = NewStartTime;
+		}
+	}
+
+	return ResizeTime;
 }
 
 void FSubSectionEditorUtil::BeginSlipSection()
 {
-    // Cache the same values as when resizing.
-    BeginResizeSection();
+	// Cache the same values as when resizing.
+	BeginResizeSection();
 }
 
 FFrameNumber FSubSectionEditorUtil::SlipSection(FFrameNumber SlipTime)
 {
-    UMovieSceneSequence* InnerSequence = SectionObject.GetSequence();
-    UMovieScene* InnerMovieScene = InnerSequence ? InnerSequence->GetMovieScene() : nullptr;
-
-    if (InnerMovieScene != nullptr)
-    {
-        FMovieSceneSectionParameters& SectionParameters = SectionObject.Parameters;
-
-        const FFrameRate    OuterFrameRate   = SectionObject.GetTypedOuter<UMovieScene>()->GetTickResolution();
-        const FFrameRate    InnerFrameRate   = InnerSequence->GetMovieScene()->GetTickResolution();
-        const FFrameNumber  ResizeDifference = SlipTime - InitialStartTimeDuringResize;
-
-        const FFrameNumber InnerResizeDifference = (ConvertFrameTime(ResizeDifference, OuterFrameRate, InnerFrameRate) * SectionParameters.TimeScale).FrameNumber;
-        FFrameNumber NewStartOffset = InitialStartOffsetDuringResize + InnerResizeDifference;
-
-        const int32 InnerPlaybackLength = UE::MovieScene::DiscreteSize(InnerMovieScene->GetPlaybackRange());
-        const float InvTimeScale = FMath::IsNearlyZero(SectionParameters.TimeScale) ? 1.0f : 1.0f / SectionParameters.TimeScale;
-        const FFrameNumber InnerLoopLength = InnerPlaybackLength - SectionParameters.StartFrameOffset - SectionParameters.EndFrameOffset;
-
-        const bool bCanLoop = SectionParameters.bCanLoop;
-        if (!bCanLoop)
-        {
-            // Ensure start offset is not less than 0
-            SectionParameters.StartFrameOffset = FFrameNumber(FMath::Max(NewStartOffset.Value, 0));
-        }
-        else
-        {
-			// If the start offset exceeds the length of one loop, trim it back.
-            NewStartOffset = NewStartOffset % InnerLoopLength;
-            if (NewStartOffset < 0)
-            {
-                // Move the first loop offset forward so that it's positive into the new first loop
-                // that we just "revealed" by slipping too much to the side.
-                NewStartOffset += InnerLoopLength;
-            }
-            SectionParameters.FirstLoopStartFrameOffset = NewStartOffset;
-        }
-    }
-
-    return SlipTime;
+	UMovieSceneSequence* InnerSequence = SectionObject.GetSequence();
+	UMovieScene* InnerMovieScene = InnerSequence ? InnerSequence->GetMovieScene() : nullptr;
+
+	if (InnerMovieScene != nullptr)
+	{
+		TRange<FFrameNumber> Range = InnerMovieScene->GetPlaybackRange();
+
+		// Find the new inner offset as an absolute time
+		FFrameNumber NewStartTime = InitialDragTransform->TransformTime(SlipTime).RoundToFrame();
+
+		// Inner offset must be relative to playback start
+		NewStartTime = NewStartTime - Range.GetLowerBoundValue();
+
+		if (SectionObject.Parameters.bCanLoop)
+		{
+			SectionObject.Parameters.FirstLoopStartFrameOffset = NewStartTime;
+		}
+		else
+		{
+			SectionObject.Parameters.StartFrameOffset = NewStartTime;
+		}
+	}
+
+	return SlipTime;
 }
 
 void FSubSectionEditorUtil::BeginDilateSection()
 {
-    FMovieSceneSectionParameters& SectionParameters = SectionObject.Parameters;
-	PreviousTimeScale = SectionParameters.TimeScale;
+	if (SectionObject.Parameters.TimeScale.GetType() == EMovieSceneTimeWarpType::FixedPlayRate)
+	{
+		PreDilateTimeScale = SectionObject.Parameters.TimeScale.AsFixedPlayRate(); //make sure to cache the play rate
+	}
+	else if (SectionObject.Parameters.TimeScale.GetType() == EMovieSceneTimeWarpType::Custom)
+	{
+		FMovieSceneTimeWarpChannel* Channel = SectionObject.GetChannelProxy().GetChannel<FMovieSceneTimeWarpChannel>(0);
+		if (Channel)
+		{
+			SectionObject.Parameters.TimeScale.AsCustom()->Modify();
+			PreDilateChannel = MakeUnique<FMovieSceneTimeWarpChannel>(*Channel);
+		}
+	}
 }
 
 void FSubSectionEditorUtil::DilateSection(const TRange<FFrameNumber>& NewRange, float DilationFactor)
 {
-    FMovieSceneSectionParameters& SectionParameters = SectionObject.Parameters;
-	SectionParameters.TimeScale = PreviousTimeScale / DilationFactor;
-	SectionObject.SetRange(NewRange);
+	if (SectionObject.Parameters.TimeScale.GetType() == EMovieSceneTimeWarpType::FixedPlayRate)
+	{
+		SectionObject.Parameters.TimeScale.Set(PreDilateTimeScale / DilationFactor);
+	}
+	else if (SectionObject.Parameters.TimeScale.GetType() == EMovieSceneTimeWarpType::Custom)
+	{
+		FMovieSceneTimeWarpChannel* Channel = SectionObject.GetChannelProxy().GetChannel<FMovieSceneTimeWarpChannel>(0);
+		if (Channel)
+		{
+			*Channel = *PreDilateChannel;
+
+			// Dilate the times
+			Dilate(Channel, FFrameNumber(0), DilationFactor);
+
+			SectionObject.Parameters.TimeScale.AsCustom()->MarkAsChanged();
+		}
+	}
 }
 
 
diff --git a/Engine/Source/Editor/MovieSceneTools/Private/TrackEditors/TimeWarpTrackEditor.cpp b/Engine/Source/Editor/MovieSceneTools/Private/TrackEditors/TimeWarpTrackEditor.cpp
new file mode 100644
--- /dev/null
+++ b/Engine/Source/Editor/MovieSceneTools/Private/TrackEditors/TimeWarpTrackEditor.cpp
@@ -0,0 +1,137 @@
+// Copyright Epic Games, Inc. All Rights Reserved.
+
+#include "TrackEditors/TimeWarpTrackEditor.h"
+#include "ISequencer.h"
+#include "ISequencerSection.h"
+#include "SequencerSectionPainter.h"
+#include "Framework/MultiBox/MultiBoxBuilder.h"
+
+#include "Sections/MovieSceneTimeWarpSection.h"
+
+#include "MVVM/ViewModels/SectionModel.h"
+#include "MVVM/ViewModels/EditorSharedViewModelData.h"
+#include "MVVM/ViewModels/SequencerEditorViewModel.h"
+
+#define LOCTEXT_NAMESPACE "TimeWarpTrackEditor"
+
+namespace UE::Sequencer
+{
+
+struct FTimeWarpSection : FSequencerSection
+{
+	FTimeWarpSection(UMovieSceneSection& InSection)
+		: FSequencerSection(InSection)
+	{}
+
+	int32 OnPaintSection(FSequencerSectionPainter& InPainter) const override
+	{
+		TSharedPtr<FEditorSharedViewModelData> SharedEditorData = CastViewModel<FEditorSharedViewModelData>(InPainter.SectionModel->GetSharedData());
+		TSharedPtr<FSequencerEditorViewModel>  SequencerEditor  = SharedEditorData ? CastViewModel<FSequencerEditorViewModel>(SharedEditorData->GetEditor()) : nullptr;
+		if (!SequencerEditor)
+		{
+			return InPainter.LayerId;
+		}
+
+		InPainter.LayerId = InPainter.PaintSectionBackground();
+
+		TSharedPtr<ISequencer> Sequencer = SequencerEditor->GetSequencer();
+
+		UMovieSceneTimeWarpSection* TimeWarpSection = Cast<UMovieSceneTimeWarpSection>(InPainter.SectionModel->GetSection());
+		if (TimeWarpSection)
+		{
+			// Paint the unwarped current time
+			FFrameTime LocalTime = Sequencer->GetLocalTime().Time;
+			FFrameTime NewLocalTime = LocalTime;
+
+			FMovieSceneInverseNestedSequenceTransform Inverse = TimeWarpSection->GenerateTransform().Inverse();
+
+			if (Inverse.IsLinear())
+			{
+				NewLocalTime = LocalTime * Inverse.AsLinear();
+			}
+			else
+			{
+				FMovieSceneSequenceTransform Transform = Sequencer->GetFocusedMovieSceneSequenceTransform();
+
+				// Time warp track transforms are always added last
+				if (Transform.NestedTransforms.Num() > 0)
+				{
+					Transform.NestedTransforms.RemoveAt(Transform.NestedTransforms.Num()-1);
+				}
+
+				NewLocalTime = Sequencer->GetGlobalTime().Time * Transform;
+			}
+
+			if (LocalTime != NewLocalTime)
+			{
+				const ESlateDrawEffect DrawEffects = InPainter.bParentEnabled ? ESlateDrawEffect::None : ESlateDrawEffect::DisabledEffect;
+				float PixelPosition = InPainter.GetTimeConverter().SecondsToPixel(NewLocalTime / Sequencer->GetFocusedTickResolution());
+
+				FSlateDrawElement::MakeBox(
+					InPainter.DrawElements,
+					InPainter.LayerId++,
+					InPainter.SectionGeometry.ToPaintGeometry(
+						FVector2f(1.0f, InPainter.SectionGeometry.Size.Y),
+						FSlateLayoutTransform(FVector2f(PixelPosition, 0.f))
+					),
+					FAppStyle::GetBrush("WhiteBrush"),
+					DrawEffects,
+					FColor(255, 255, 255, 128)	// 0, 75, 50 (HSV)
+				);
+			}
+		}
+
+		return InPainter.LayerId;
+	}
+};
+
+} // namespace UE::Sequencer
+
+void FTimeWarpTrackEditor::BuildAddTrackMenu(FMenuBuilder& MenuBuilder)
+{
+	MenuBuilder.AddMenuEntry(
+		LOCTEXT("AddTimeWarpTrack", "Time Warp"),
+		LOCTEXT("AddTimeWarpTrackTooltip", "Adds a new track that manipulates the time of the current sequence."),
+		FSlateIcon(FAppStyle::GetAppStyleSetName(), "Sequencer.Tracks.Slomo"),
+		FUIAction(
+			FExecuteAction::CreateRaw(this, &FTimeWarpTrackEditor::HandleAddTimeWarpTrack)
+		)
+	);
+}
+
+TSharedRef<ISequencerSection> FTimeWarpTrackEditor::MakeSectionInterface(UMovieSceneSection& SectionObject, UMovieSceneTrack& Track, FGuid ObjectBinding)
+{
+	return MakeShared<UE::Sequencer::FTimeWarpSection>(SectionObject);
+}
+
+void FTimeWarpTrackEditor::HandleAddTimeWarpTrack()
+{
+	UMovieScene* FocusedMovieScene = GetFocusedMovieScene();
+	if (FocusedMovieScene == nullptr)
+	{
+		return;
+	}
+
+	if (FocusedMovieScene->IsReadOnly())
+	{
+		return;
+	}
+
+	TSharedPtr<ISequencer> SequencerPtr = GetSequencer();
+	if (!SequencerPtr)
+	{
+		return;
+	}
+
+	const FScopedTransaction Transaction(LOCTEXT("AddTimeWarpTrack_Transaction", "Add Time Warp Track"));
+
+	FocusedMovieScene->Modify();
+
+	UMovieSceneTimeWarpTrack* NewTrack = NewObject<UMovieSceneTimeWarpTrack>(FocusedMovieScene, NAME_None, RF_Transactional);
+	NewTrack->AddSection(*NewTrack->CreateNewSection());
+
+	FocusedMovieScene->AddGivenTrack(NewTrack);
+	SequencerPtr->OnAddTrack(NewTrack, FGuid());
+}
+
+#undef LOCTEXT_NAMESPACE
\ No newline at end of file
diff --git a/Engine/Source/Editor/MovieSceneTools/Public/Channels/BezierChannelCurveModel.h b/Engine/Source/Editor/MovieSceneTools/Public/Channels/BezierChannelCurveModel.h
--- a/Engine/Source/Editor/MovieSceneTools/Public/Channels/BezierChannelCurveModel.h
+++ b/Engine/Source/Editor/MovieSceneTools/Public/Channels/BezierChannelCurveModel.h
@@ -55,6 +55,7 @@ class FBezierChannelCurveModel : public FChannelCurveModel<ChannelType, ChannelV
 {
 public:
 	FBezierChannelCurveModel(TMovieSceneChannelHandle<ChannelType> InChannel, UMovieSceneSection* InOwningSection, TWeakPtr<ISequencer> InWeakSequencer);
+	FBezierChannelCurveModel(TMovieSceneChannelHandle<ChannelType> InChannel, UMovieSceneSection* InOwningSection, UObject* InOwningObject, TWeakPtr<ISequencer> InWeakSequencer);
 
 	// FCurveModel
 	virtual void DrawCurve(const FCurveEditor& CurveEditor, const FCurveEditorScreenSpace& ScreenSpace, TArray<TTuple<double, double>>& InterpolatingPoints) const override;
diff --git a/Engine/Source/Editor/MovieSceneTools/Public/Channels/ChannelCurveModel.h b/Engine/Source/Editor/MovieSceneTools/Public/Channels/ChannelCurveModel.h
--- a/Engine/Source/Editor/MovieSceneTools/Public/Channels/ChannelCurveModel.h
+++ b/Engine/Source/Editor/MovieSceneTools/Public/Channels/ChannelCurveModel.h
@@ -40,6 +40,8 @@ class FChannelCurveModel : public FCurveModel
 {
 public:
 	FChannelCurveModel(TMovieSceneChannelHandle<ChannelType> InChannel, UMovieSceneSection* InOwningSection, TWeakPtr<ISequencer> InWeakSequencer);
+	FChannelCurveModel(TMovieSceneChannelHandle<ChannelType> InChannel, UMovieSceneSection* InOwningSection, UObject* InOwningObject, TWeakPtr<ISequencer> InWeakSequencer);
+
 	~FChannelCurveModel();
 
 	virtual const void* GetCurve() const override;
@@ -62,18 +64,29 @@ public:
 	virtual void AddKeys(TArrayView<const FKeyPosition> InKeyPositions, TArrayView<const FKeyAttributes> InAttributes, TArrayView<TOptional<FKeyHandle>>* OutKeyHandles) override;
 	virtual void RemoveKeys(TArrayView<const FKeyHandle> InKeys) override;
 
+	virtual FTransform2d GetCurveTransform() const override;
 	virtual bool IsReadOnly() const override;
 	virtual UObject* GetOwningObject() const override
 	{
+		if (UObject* OwningObject = WeakOwningObject.Get())
+		{
+			return OwningObject;
+		}
 		return WeakSection.Get();
 	}
 	virtual bool HasChangedAndResetTest() override
 	{
-		if (UMovieSceneSection* Section = WeakSection.Get())
+		UMovieSceneSignedObject* SignedOwner = Cast<UMovieSceneSignedObject>(WeakOwningObject.Get());
+		if (!SignedOwner)
+		{
+			SignedOwner = WeakSection.Get();
+		}
+
+		if (SignedOwner)
 		{
-			if (Section->GetSignature() != LastSignature)
+			if (SignedOwner->GetSignature() != LastSignature)
 			{
-				LastSignature = Section->GetSignature();
+				LastSignature = SignedOwner->GetSignature();
 				return true;
 			}
 			return false;
@@ -96,10 +109,11 @@ private:
 
 	void FixupCurve();
 
-private:
+protected:
 
 	TMovieSceneChannelHandle<ChannelType> ChannelHandle;
 	TWeakObjectPtr<UMovieSceneSection> WeakSection;
+	TWeakObjectPtr<> WeakOwningObject;
 	TWeakPtr<ISequencer> WeakSequencer;
 	FDelegateHandle OnDestroyHandle;
 	FGuid LastSignature;
diff --git a/Engine/Source/Editor/MovieSceneTools/Public/Channels/DoubleChannelCurveModel.h b/Engine/Source/Editor/MovieSceneTools/Public/Channels/DoubleChannelCurveModel.h
--- a/Engine/Source/Editor/MovieSceneTools/Public/Channels/DoubleChannelCurveModel.h
+++ b/Engine/Source/Editor/MovieSceneTools/Public/Channels/DoubleChannelCurveModel.h
@@ -21,8 +21,12 @@ class FDoubleChannelCurveModel : public FBezierChannelCurveModel<FMovieSceneDoub
 {
 public:
 	FDoubleChannelCurveModel(TMovieSceneChannelHandle<FMovieSceneDoubleChannel> InChannel, UMovieSceneSection* InOwningSection, TWeakPtr<ISequencer> InWeakSequencer);
+	FDoubleChannelCurveModel(TMovieSceneChannelHandle<FMovieSceneDoubleChannel> InChannel, UMovieSceneSection* InOwningSection, UObject* InOwningObject, TWeakPtr<ISequencer> InWeakSequencer);
 
 	// FCurveModel
 	virtual void CreateKeyProxies(TArrayView<const FKeyHandle> InKeyHandles, TArrayView<UObject*> OutObjects) override;
 	virtual TUniquePtr<IBufferedCurveModel> CreateBufferedCurveCopy() const override;
+
+	virtual void GetValueRange(double& MinValue, double& MaxValue) const override;
+	virtual void GetValueRange(double InMinTime, double InMaxTime, double& MinValue, double& MaxValue) const override;
 };
diff --git a/Engine/Source/Editor/MovieSceneTools/Public/Channels/PerlinNoiseChannelInterface.h b/Engine/Source/Editor/MovieSceneTools/Public/Channels/PerlinNoiseChannelInterface.h
--- a/Engine/Source/Editor/MovieSceneTools/Public/Channels/PerlinNoiseChannelInterface.h
+++ b/Engine/Source/Editor/MovieSceneTools/Public/Channels/PerlinNoiseChannelInterface.h
@@ -103,7 +103,7 @@ struct TPerlinNoiseChannelInterface : ISequencerChannelInterface
 		return true;
 	}
 
-	virtual TSharedRef<SWidget> CreateKeyEditor_Raw(const FMovieSceneChannelHandle& Channel, UMovieSceneSection* Section, const FGuid& InObjectBindingID, TWeakPtr<FTrackInstancePropertyBindings> PropertyBindings, TWeakPtr<ISequencer> Sequencer) const override
+	virtual TSharedRef<SWidget> CreateKeyEditor_Raw(const FMovieSceneChannelHandle& Channel, const UE::Sequencer::FCreateKeyEditorParams& Params) const override
 	{
 		const TMovieSceneExternalValue<typename ChannelType::CurveValueType>* ExternalValue = Channel.Cast<ChannelType>().GetExtendedEditorData();
 		if (!ExternalValue)
@@ -112,8 +112,8 @@ struct TPerlinNoiseChannelInterface : ISequencerChannelInterface
 		}
 
 		TSequencerKeyEditor<ChannelType, typename ChannelType::CurveValueType> KeyEditor(
-			InObjectBindingID, Channel.Cast<ChannelType>(),
-			Section, Sequencer, PropertyBindings, ExternalValue->OnGetExternalValue
+			Params.ObjectBindingID, Channel.Cast<ChannelType>(),
+			Params.OwningSection, Params.Sequencer, Params.PropertyBindings, ExternalValue->OnGetExternalValue
 			);
 
 		using KeyEditorType = SNumericTextBlockKeyEditor<ChannelType, typename ChannelType::CurveValueType>;
@@ -145,7 +145,7 @@ struct TPerlinNoiseChannelInterface : ISequencerChannelInterface
 		return false;
 	}
 
-	virtual TUniquePtr<FCurveModel> CreateCurveEditorModel_Raw(const FMovieSceneChannelHandle& Channel, UMovieSceneSection* OwningSection, TSharedRef<ISequencer> InSequencer) const override
+	virtual TUniquePtr<FCurveModel> CreateCurveEditorModel_Raw(const FMovieSceneChannelHandle& Channel, const UE::Sequencer::FCreateCurveEditorModelParams& Params) const override
 	{
 		return nullptr;
 	}
diff --git a/Engine/Source/Editor/MovieSceneTools/Public/MovieSceneToolsModule.h b/Engine/Source/Editor/MovieSceneTools/Public/MovieSceneToolsModule.h
--- a/Engine/Source/Editor/MovieSceneTools/Public/MovieSceneToolsModule.h
+++ b/Engine/Source/Editor/MovieSceneTools/Public/MovieSceneToolsModule.h
@@ -141,6 +141,7 @@ private:
 	FDelegateHandle CVarTrackCreateEditorHandle;
 	FDelegateHandle CustomPrimitiveDataTrackCreateEditorHandle;
 	FDelegateHandle BindingLifetimeTrackCreateEditorHandle;
+	FDelegateHandle TimeWarpTrackCreateEditorHandle;
 
 	FDelegateHandle CameraCutTrackModelHandle;
 	FDelegateHandle CinematicShotTrackModelHandle;
diff --git a/Engine/Source/Editor/MovieSceneTools/Public/TrackEditors/SubTrackEditor.h b/Engine/Source/Editor/MovieSceneTools/Public/TrackEditors/SubTrackEditor.h
--- a/Engine/Source/Editor/MovieSceneTools/Public/TrackEditors/SubTrackEditor.h
+++ b/Engine/Source/Editor/MovieSceneTools/Public/TrackEditors/SubTrackEditor.h
@@ -58,6 +58,7 @@ public:
 	virtual FReply OnDrop(const FDragDropEvent& DragDropEvent, const FSequencerDragDropParams& DragDropParams) override;
 	virtual bool IsResizable(UMovieSceneTrack* InTrack) const override;
 	virtual void Resize(float NewSize, UMovieSceneTrack* InTrack) override;
+	virtual bool GetDefaultExpansionState(UMovieSceneTrack* InTrack) const override;
 
 public:
 	
diff --git a/Engine/Source/Editor/MovieSceneTools/Public/TrackEditors/SubTrackEditorBase.h b/Engine/Source/Editor/MovieSceneTools/Public/TrackEditors/SubTrackEditorBase.h
--- a/Engine/Source/Editor/MovieSceneTools/Public/TrackEditors/SubTrackEditorBase.h
+++ b/Engine/Source/Editor/MovieSceneTools/Public/TrackEditors/SubTrackEditorBase.h
@@ -5,6 +5,7 @@
 #include "CoreMinimal.h"
 #include "Editor.h"
 #include "ISequencerSection.h"
+#include "Templates/UniquePtr.h"
 #include "Sections/MovieSceneSubSection.h"
 #include "Styling/AppStyle.h"
 #include "Framework/Application/SlateApplication.h"
@@ -20,6 +21,8 @@
 #include "Subsystems/AssetEditorSubsystem.h"
 #include "Tracks/MovieSceneSubTrack.h"
 
+struct FMovieSceneTimeWarpChannel;
+struct FMovieSceneSequenceTransform;
 class ISequencer;
 
 /**
@@ -57,8 +60,7 @@ public:
     static FSubSectionPainterResult PaintSection(TSharedPtr<const ISequencer> Sequencer, const UMovieSceneSubSection& SectionObject, FSequencerSectionPainter& InPainter, FSubSectionPainterParams Params);
 
 private:
-    static void DoPaintNonLoopingSection(const UMovieSceneSubSection& SectionObject, const UMovieSceneSequence& InnerSequence, FSequencerSectionPainter& InPainter, ESlateDrawEffect DrawEffects);
-    static void DoPaintLoopingSection(const UMovieSceneSubSection& SectionObject, const UMovieSceneSequence& InnerSequence, FSequencerSectionPainter& InPainter, ESlateDrawEffect DrawEffects);
+    static void PaintSectionBounds(const UMovieSceneSubSection& SectionObject, const UMovieSceneSequence& InnerSequence, FSequencerSectionPainter& InPainter, ESlateDrawEffect DrawEffects);
 };
 
 /**
@@ -68,6 +70,7 @@ class MOVIESCENETOOLS_API FSubSectionEditorUtil
 {
 public:
     FSubSectionEditorUtil(UMovieSceneSubSection& InSection);
+    ~FSubSectionEditorUtil();
 
     /** Starts a resize operation */
     void BeginResizeSection();
@@ -88,14 +91,12 @@ private:
     /** The section object this utility class is editing */
     UMovieSceneSubSection& SectionObject;
 
-    /** Cached start offset value valid only during resize */
-    FFrameNumber InitialStartOffsetDuringResize;
+	TUniquePtr<FMovieSceneSequenceTransform> InitialDragTransform;
 
-    /** Cached start time valid only during resize */
-    FFrameNumber InitialStartTimeDuringResize;
+	TUniquePtr<FMovieSceneTimeWarpChannel> PreDilateChannel;
 
 	/** Cached time scale valid only during dilate */
-	float PreviousTimeScale;
+	double PreDilateTimeScale;
 };
 
 class MOVIESCENETOOLS_API FSubTrackEditorUtil
@@ -147,6 +148,7 @@ public:
     virtual void SlipSection(FFrameNumber SlipTime) override;
 	virtual void BeginDilateSection() override;
 	virtual void DilateSection(const TRange<FFrameNumber>& NewRange, float DilationFactor) override;
+	virtual bool RequestDeleteKeyArea( const TArray<FName>& KeyAreaNamePath ) override;
 
 protected:
 	static const float TrackHeight;
@@ -376,3 +378,9 @@ void TSubSectionMixin<ParentSectionClass>::DilateSection(const TRange<FFrameNumb
     ParentSectionClass::DilateSection(NewRange, DilationFactor);
 }
 
+template<typename ParentSectionClass>
+bool TSubSectionMixin<ParentSectionClass>::RequestDeleteKeyArea( const TArray<FName>& KeyAreaNamePath )
+{
+	SubSectionObject.DeleteChannels(KeyAreaNamePath);
+	return true;
+}
\ No newline at end of file
diff --git a/Engine/Source/Editor/Sequencer/Private/IKeyArea.cpp b/Engine/Source/Editor/Sequencer/Private/IKeyArea.cpp
--- a/Engine/Source/Editor/Sequencer/Private/IKeyArea.cpp
+++ b/Engine/Source/Editor/Sequencer/Private/IKeyArea.cpp
@@ -55,6 +55,18 @@ UMovieSceneSection* IKeyArea::GetOwningSection() const
 	return SectionInterface ? SectionInterface->GetSectionObject() : nullptr;;
 }
 
+UObject* IKeyArea::GetOwningObject() const
+{
+	const FMovieSceneChannelMetaData* MetaData = ChannelHandle.GetMetaData();
+	UObject* OwningObject = MetaData ? MetaData->WeakOwningObject.Get() : nullptr;
+	if (OwningObject)
+	{
+		return OwningObject;
+	}
+
+	return GetOwningSection();
+}
+
 TSharedPtr<ISequencerSection> IKeyArea::GetSectionInterface() const
 {
 	return WeakSection.Pin();
@@ -84,6 +96,11 @@ FKeyHandle IKeyArea::AddOrUpdateKey(FFrameNumber Time, const FGuid& ObjectBindin
 	FMovieSceneChannel* Channel = ChannelHandle.Get();
 	UMovieSceneSection* Section = GetOwningSection();
 
+	if (const FMovieSceneChannelMetaData* MetaData = ChannelHandle.GetMetaData())
+	{
+		Time -= MetaData->GetOffsetTime(Section);
+	}
+
 	// The extended editor data may be null, but is passed to the interface regardless
 	const void* RawExtendedData = ChannelHandle.GetExtendedEditorData();
 
@@ -112,7 +129,29 @@ void IKeyArea::SetKeyTimes(TArrayView<const FKeyHandle> InKeyHandles, TArrayView
 {
 	check(InKeyHandles.Num() == InKeyTimes.Num());
 
-	if (FMovieSceneChannel* Channel = ChannelHandle.Get())
+	FMovieSceneChannel*               Channel  = ChannelHandle.Get();
+	UMovieSceneSection*               Section  = GetOwningSection();
+	const FMovieSceneChannelMetaData* MetaData = ChannelHandle.GetMetaData();
+
+	if (!Channel || !MetaData || !Section)
+	{
+		return;
+	}
+
+	FFrameNumber KeyOffset = MetaData->GetOffsetTime(Section);
+	if (KeyOffset != 0)
+	{
+		// Need to copy the array to offset
+		TArray<FFrameNumber> OffsetKeyTimes;
+		OffsetKeyTimes.SetNumUninitialized(InKeyTimes.Num());
+		for (int32 i = 0; i < InKeyTimes.Num(); ++i)
+		{
+			OffsetKeyTimes[i] = InKeyTimes[i] - KeyOffset;
+		}
+
+		Channel->SetKeyTimes(InKeyHandles, OffsetKeyTimes);
+	}
+	else
 	{
 		Channel->SetKeyTimes(InKeyHandles, InKeyTimes);
 	}
@@ -120,17 +159,62 @@ void IKeyArea::SetKeyTimes(TArrayView<const FKeyHandle> InKeyHandles, TArrayView
 
 void IKeyArea::GetKeyTimes(TArrayView<const FKeyHandle> InKeyHandles, TArrayView<FFrameNumber> OutTimes) const
 {
-	if (FMovieSceneChannel* Channel = ChannelHandle.Get())
+	FMovieSceneChannel*               Channel  = ChannelHandle.Get();
+	UMovieSceneSection*               Section  = GetOwningSection();
+	const FMovieSceneChannelMetaData* MetaData = ChannelHandle.GetMetaData();
+
+	if (!Channel || !MetaData || !Section)
 	{
-		Channel->GetKeyTimes(InKeyHandles, OutTimes);
+		return;
+	}
+
+	Channel->GetKeyTimes(InKeyHandles, OutTimes);
+
+	FFrameNumber KeyOffset = MetaData->GetOffsetTime(Section);
+	if (KeyOffset != 0)
+	{
+		for (int32 i = 0; i < OutTimes.Num(); ++i)
+		{
+			OutTimes[i] += KeyOffset;
+		}
 	}
 }
 
 void IKeyArea::GetKeyInfo(TArray<FKeyHandle>* OutHandles, TArray<FFrameNumber>* OutTimes, const TRange<FFrameNumber>& WithinRange) const
 {
-	if (FMovieSceneChannel* Channel = ChannelHandle.Get())
+	FMovieSceneChannel*               Channel  = ChannelHandle.Get();
+	UMovieSceneSection*               Section  = GetOwningSection();
+	const FMovieSceneChannelMetaData* MetaData = ChannelHandle.GetMetaData();
+
+	if (!Channel || !MetaData || !Section)
 	{
-		Channel->GetKeys(WithinRange, OutTimes, OutHandles);
+		return;
+	}
+
+	FFrameNumber KeyOffset = MetaData->GetOffsetTime(Section);
+
+	TRange<FFrameNumber> OffsetRange = WithinRange;
+	if (KeyOffset != 0)
+	{
+		if (OffsetRange.HasLowerBound())
+		{
+			OffsetRange.SetLowerBoundValue(OffsetRange.GetLowerBoundValue() - KeyOffset);
+		}
+		if (OffsetRange.HasUpperBound())
+		{
+			OffsetRange.SetUpperBoundValue(OffsetRange.GetUpperBoundValue() - KeyOffset);
+		}
+	}
+
+	Channel->GetKeys(OffsetRange, OutTimes, OutHandles);
+
+	if (OutTimes && KeyOffset != 0)
+	{
+		const int32 Num = OutTimes->Num();
+		for (int32 i = 0; i < Num; ++i)
+		{
+			(*OutTimes)[i] += KeyOffset;
+		}
 	}
 }
 
@@ -180,6 +264,8 @@ bool IKeyArea::CanCreateKeyEditor() const
 
 TSharedRef<SWidget> IKeyArea::CreateKeyEditor(TWeakPtr<ISequencer> Sequencer, const FGuid& ObjectBindingID)
 {
+	using namespace UE::Sequencer;
+
 	ISequencerChannelInterface* EditorInterface = FindChannelEditorInterface();
 	UMovieSceneSection* OwningSection = GetOwningSection();
 
@@ -189,9 +275,17 @@ TSharedRef<SWidget> IKeyArea::CreateKeyEditor(TWeakPtr<ISequencer> Sequencer, co
 		PropertyBindingsPtr = TSharedPtr<FTrackInstancePropertyBindings>(AsShared(), &PropertyBindings.GetValue());
 	}
 
-	if (EditorInterface && OwningSection)
+	const FMovieSceneChannelMetaData* MetaData = ChannelHandle.GetMetaData();
+	if (EditorInterface && OwningSection && MetaData)
 	{
-		return EditorInterface->CreateKeyEditor_Raw(ChannelHandle, OwningSection, ObjectBindingID, PropertyBindingsPtr, Sequencer);
+		FCreateKeyEditorParams Params{
+			OwningSection,
+			MetaData->WeakOwningObject.Get(),
+			Sequencer.Pin().ToSharedRef(),
+			ObjectBindingID,
+			PropertyBindingsPtr
+		};
+		return EditorInterface->CreateKeyEditor_Raw(ChannelHandle, Params);
 	}
 	return SNullWidget::NullWidget;
 }
@@ -201,9 +295,12 @@ void IKeyArea::CopyKeys(FMovieSceneClipboardBuilder& ClipboardBuilder, TArrayVie
 	ISequencerChannelInterface* EditorInterface = FindChannelEditorInterface();
 	FMovieSceneChannel* Channel = ChannelHandle.Get();
 	UMovieSceneSection* OwningSection = GetOwningSection();
+	const FMovieSceneChannelMetaData* MetaData = ChannelHandle.GetMetaData();
 
-	if (EditorInterface && Channel && OwningSection)
+	if (EditorInterface && Channel && OwningSection && MetaData)
 	{
+		TGuardValue<FFrameNumber> GuardKeyOffset(ClipboardBuilder.KeyOffset, ClipboardBuilder.KeyOffset + MetaData->GetOffsetTime(OwningSection));
+
 		EditorInterface->CopyKeys_Raw(Channel, OwningSection, ChannelName, ClipboardBuilder, KeyMask);
 	}
 }
@@ -215,9 +312,33 @@ TArray<FKeyHandle> IKeyArea::PasteKeys(const FMovieSceneClipboardKeyTrack& KeyTr
 	ISequencerChannelInterface* EditorInterface = FindChannelEditorInterface();
 	FMovieSceneChannel* Channel = ChannelHandle.Get();
 	UMovieSceneSection* OwningSection = GetOwningSection();
-	if (EditorInterface && Channel && OwningSection)
+	UObject*            OwningObject = GetOwningObject();
+	const FMovieSceneChannelMetaData* MetaData = ChannelHandle.GetMetaData();
+
+	if (EditorInterface && Channel && OwningSection && OwningObject && MetaData)
 	{
+		if (OwningSection->IsReadOnly())
+		{
+			return PastedKeys;
+		}
+
+		OwningObject->Modify();
+
 		EditorInterface->PasteKeys_Raw(Channel, OwningSection, KeyTrack, SrcEnvironment, DstEnvironment, PastedKeys);
+
+		FFrameNumber KeyOffset = MetaData->GetOffsetTime(OwningSection);
+		if (KeyOffset != 0)
+		{
+			TArray<FFrameNumber> KeyTimes;
+			KeyTimes.SetNumUninitialized(PastedKeys.Num());
+
+			Channel->GetKeyTimes(PastedKeys, KeyTimes);
+			for (int32 i = 0; i < KeyTimes.Num(); ++i)
+			{
+				KeyTimes[i] += KeyOffset;
+			}
+			Channel->SetKeyTimes(PastedKeys, KeyTimes);
+		}
 	}
 
 	return PastedKeys;
@@ -272,7 +393,15 @@ TUniquePtr<FCurveModel> IKeyArea::CreateCurveEditorModel(TSharedRef<ISequencer>
 	UMovieSceneSection* OwningSection = GetOwningSection();
 	if (EditorInterface && OwningSection && ChannelHandle.Get() != nullptr)
 	{
-		TUniquePtr<FCurveModel> CurveModel = EditorInterface->CreateCurveEditorModel_Raw(ChannelHandle, OwningSection, InSequencer);
+		const FMovieSceneChannelMetaData* MetaData = ChannelHandle.GetMetaData();
+
+		UE::Sequencer::FCreateCurveEditorModelParams Params = {
+			OwningSection,
+			MetaData->WeakOwningObject.Get(),
+			InSequencer
+		};
+
+		TUniquePtr<FCurveModel> CurveModel = EditorInterface->CreateCurveEditorModel_Raw(ChannelHandle, Params);
 		if (CurveModel.IsValid())
 		{
 			// Build long, short and context names for this curve to maximize information shown in the Curve Editor UI.
@@ -295,7 +424,7 @@ TUniquePtr<FCurveModel> IKeyArea::CreateCurveEditorModel(TSharedRef<ISequencer>
 			// Not all tracks have all the information so we need to format it differently depending on how many are valid.
 			FormatArgs.Add(TEXT("ObjectBindingName"), ObjectBindingName);
 			FormatArgs.Add(TEXT("OwningTrackName"), OwningTrackName);
-			FormatArgs.Add(TEXT("GroupName"), ChannelHandle.GetMetaData()->Group);
+			FormatArgs.Add(TEXT("GroupName"), MetaData->Group);
 			FormatArgs.Add(TEXT("DisplayName"), DisplayText);
 
 			FText FormatText = NSLOCTEXT("SequencerIKeyArea", "CurveLongDisplayNameFormat", "{ObjectBindingName}.{OwningTrackName}.{GroupName}.{DisplayName}");
@@ -305,13 +434,13 @@ TUniquePtr<FCurveModel> IKeyArea::CreateCurveEditorModel(TSharedRef<ISequencer>
 
 			FText LongDisplayName = FText::FromString(FormattedText);
 			const FText ShortDisplayName = DisplayText;
-			FString IntentName = ChannelHandle.GetMetaData()->IntentName.ToString();
+			FString IntentName = MetaData->IntentName.ToString();
 			if (IntentName.IsEmpty())
 			{
-				IntentName = ChannelHandle.GetMetaData()->Group.IsEmptyOrWhitespace() ? DisplayText.ToString() : FString::Format(TEXT("{0}.{1}"), { ChannelHandle.GetMetaData()->Group.ToString(), DisplayText.ToString() });
+				IntentName = MetaData->Group.IsEmptyOrWhitespace() ? DisplayText.ToString() : FString::Format(TEXT("{0}.{1}"), { MetaData->Group.ToString(), DisplayText.ToString() });
 			}
 
-			FText LongIntentNameFormat = ChannelHandle.GetMetaData()->LongIntentNameFormat;
+			FText LongIntentNameFormat = MetaData->LongIntentNameFormat;
 			if (LongIntentNameFormat.IsEmpty())
 			{
 				LongIntentNameFormat = NSLOCTEXT("SequencerIKeyArea", "LongIntentNameFormat", "{ObjectBindingName}.{GroupName}.{DisplayName}");
@@ -328,7 +457,7 @@ TUniquePtr<FCurveModel> IKeyArea::CreateCurveEditorModel(TSharedRef<ISequencer>
 			CurveModel->SetLongDisplayName(LongDisplayName);
 			CurveModel->SetIntentionName(IntentName);
 			CurveModel->SetLongIntentionName(LongIntentName);
-			CurveModel->SetChannelName(ChannelHandle.GetMetaData()->Name);
+			CurveModel->SetChannelName(MetaData->Name);
 			if (Color.IsSet())
 			{
 				CurveModel->SetColor(Color.GetValue(),false);
diff --git a/Engine/Source/Editor/Sequencer/Private/MVVM/CurveEditorExtension.cpp b/Engine/Source/Editor/Sequencer/Private/MVVM/CurveEditorExtension.cpp
--- a/Engine/Source/Editor/Sequencer/Private/MVVM/CurveEditorExtension.cpp
+++ b/Engine/Source/Editor/Sequencer/Private/MVVM/CurveEditorExtension.cpp
@@ -9,6 +9,8 @@
 #include "IPropertyRowGenerator.h"
 #include "IStructureDetailsView.h"
 #include "MVVM/ViewModels/SequencerEditorViewModel.h"
+#include "CurveEditorAxis.h"
+#include "SCurveEditorView.h"
 #include "SCurveEditorPanel.h"
 #include "SCurveEditorToolProperties.h"
 #include "SCurveKeyDetailPanel.h"
@@ -33,14 +35,94 @@ namespace UE
 namespace Sequencer
 {
 
+/** Custom curve editor axis that displays the 'current time' in display rate */
+class FSequencerTimeCurveEditorAxis : public FLinearCurveEditorAxis
+{
+public:
+
+	TWeakPtr<FSequencer> WeakSequencer;
+	FSequencerTimeCurveEditorAxis(TWeakPtr<FSequencer> InWeakSequencer)
+		: WeakSequencer(InWeakSequencer)
+	{
+	}
+
+	void GetGridLines(const FCurveEditor& CurveEditor, const SCurveEditorView& View, FCurveEditorViewAxisID AxisID, TArray<double>& OutMajorGridLines, TArray<double>& OutMinorGridLines, ECurveEditorAxisOrientation Axis) const override
+	{
+		TSharedPtr<FSequencer> Sequencer = WeakSequencer.Pin();
+
+		if (!Sequencer.IsValid())
+		{
+			return;
+		}
+
+		double ToSeconds = Sequencer->GetFocusedTickResolution().AsInterval();
+
+		double MajorGridStep = 0.0;
+		int32  MinorDivisions = 0;
+
+
+		float Size = 1.0;
+		float Min  = 0.0;
+		float Max  = 1.0;
+
+		if (Axis == ECurveEditorAxisOrientation::Horizontal)
+		{
+			FCurveEditorScreenSpaceH AxisSpace = View.GetHorizontalAxisSpace(AxisID);
+			Size = AxisSpace.GetPhysicalWidth();
+			Min  = AxisSpace.GetInputMin();
+			Max  = AxisSpace.GetInputMax();
+		}
+		else
+		{
+			FCurveEditorScreenSpaceV AxisSpace = View.GetVerticalAxisSpace(AxisID);
+			Size = AxisSpace.GetPhysicalHeight();
+			Min = AxisSpace.GetOutputMin();
+			Max = AxisSpace.GetOutputMax();
+		}
+
+		if (Sequencer.IsValid() && Sequencer->GetGridMetrics(Size, Min, Max, MajorGridStep, MinorDivisions))
+		{
+			double FirstMajorLine = FMath::FloorToDouble(Min / MajorGridStep) * MajorGridStep;
+			double LastMajorLine = FMath::CeilToDouble(Max / MajorGridStep) * MajorGridStep;
+
+			for (double CurrentMajorLine = FirstMajorLine; CurrentMajorLine < LastMajorLine; CurrentMajorLine += MajorGridStep)
+			{
+				OutMajorGridLines.Add(CurrentMajorLine);
+
+				for (int32 Step = 1; Step < MinorDivisions; ++Step)
+				{
+					double MinorLine = CurrentMajorLine + Step * MajorGridStep / MinorDivisions;
+					OutMinorGridLines.Add(MinorLine);
+				}
+			}
+		}
+	}
+};
+
 class FSequencerCurveEditor : public FCurveEditor
 {
 public:
 	TWeakPtr<FSequencer> WeakSequencer;
+	TSharedPtr<FLinearCurveEditorAxis> FocusedTimeAxis;
 
-	FSequencerCurveEditor(TWeakPtr<FSequencer> InSequencer)
+	FSequencerCurveEditor(TWeakPtr<FSequencer> InSequencer, TSharedPtr<INumericTypeInterface<double>> InNumericTypeInterface)
 		: WeakSequencer(InSequencer)
-	{}
+	{
+		FocusedTimeAxis = MakeShared<FSequencerTimeCurveEditorAxis>(InSequencer);
+		FocusedTimeAxis->NumericTypeInterface = InNumericTypeInterface;
+
+		InSequencer.Pin()->OnActivateSequence().AddRaw(this, &FSequencerCurveEditor::HandleSequenceActivated);
+
+		AddAxis("FocusedSequenceTime", FocusedTimeAxis);
+	}
+
+	~FSequencerCurveEditor()
+	{
+		if (TSharedPtr< FSequencer> Sequencer = WeakSequencer.Pin())
+		{
+			Sequencer->OnActivateSequence().AddRaw(this, &FSequencerCurveEditor::HandleSequenceActivated);
+		}
+	}
 
 	virtual void GetGridLinesX(TArray<float>& MajorGridLines, TArray<float>& MinorGridLines, TArray<FText>* MajorGridLabels) const override
 	{
@@ -76,6 +158,11 @@ public:
 			(int32)ECurveEditorTangentTypes::InterpolationCubicWeighted |
 			(int32)ECurveEditorTangentTypes::InterpolationCubicSmartAuto);
 	}
+
+	void HandleSequenceActivated(FMovieSceneSequenceIDRef NewSequenceID)
+	{
+		FocusedTimeAxis->NumericTypeInterface = WeakSequencer.Pin()->GetNumericTypeInterface();
+	}
 };
 
 struct FSequencerCurveEditorBounds : ICurveEditorBounds
@@ -182,7 +269,7 @@ void FCurveEditorExtension::CreateCurveEditor(const FTimeSliderArgs& TimeSliderA
 
 		USequencerSettings* SequencerSettings = Sequencer->GetSequencerSettings();
 
-		CurveEditorModel = MakeShared<FSequencerCurveEditor>(Sequencer);
+		CurveEditorModel = MakeShared<FSequencerCurveEditor>(Sequencer, TimeSliderArgs.NumericTypeInterface);
 		CurveEditorModel->SetBounds(MakeUnique<FSequencerCurveEditorBounds>(Sequencer.ToSharedRef()));
 		CurveEditorModel->InitCurveEditor(CurveEditorInitParams);
 
diff --git a/Engine/Source/Editor/Sequencer/Private/MVVM/ViewModels/ChannelModel.cpp b/Engine/Source/Editor/Sequencer/Private/MVVM/ViewModels/ChannelModel.cpp
--- a/Engine/Source/Editor/Sequencer/Private/MVVM/ViewModels/ChannelModel.cpp
+++ b/Engine/Source/Editor/Sequencer/Private/MVVM/ViewModels/ChannelModel.cpp
@@ -97,6 +97,11 @@ UMovieSceneSection* FChannelModel::GetSection() const
 	return KeyArea->GetOwningSection();
 }
 
+UObject* FChannelModel::GetOwningObject() const
+{
+	return KeyArea->GetOwningObject();
+}
+
 FOutlinerSizing FChannelModel::GetDesiredSizing() const
 {
 	FOutlinerSizing Sizing(15.f);
