Sequencer: Prevent playback bounds being drawn outside of sub-sequences and shots

diff --git a/Engine/Source/Editor/MovieSceneTools/Private/TrackEditors/SubTrackEditorBase.cpp b/Engine/Source/Editor/MovieSceneTools/Private/TrackEditors/SubTrackEditorBase.cpp
--- a/Engine/Source/Editor/MovieSceneTools/Private/TrackEditors/SubTrackEditorBase.cpp
+++ b/Engine/Source/Editor/MovieSceneTools/Private/TrackEditors/SubTrackEditorBase.cpp
@@ -146,6 +146,12 @@ void FSubSectionPainterUtil::PaintSectionBounds(const UMovieSceneSubSection& Sec
 		}
 		LastTime = InTime;
 
+		// Don't actually draw it if it is out of bounds
+		if (Offset < 0 || Offset >= InPainter.SectionGeometry.Size.X)
+		{
+			return;
+		}
+
 		TArray<FVector2f> NewVector;
 		NewVector.Reserve(2);
 
@@ -164,7 +170,6 @@ void FSubSectionPainterUtil::PaintSectionBounds(const UMovieSceneSubSection& Sec
 			Thickness,
 			DashLengthPx
 		);
-		return true;
 	};
 
 
@@ -179,6 +184,12 @@ void FSubSectionPainterUtil::PaintSectionBounds(const UMovieSceneSubSection& Sec
 		}
 		LastTime = InTime;
 
+		// Don't actually draw it if it is out of bounds
+		if (Offset < 0 || Offset >= InPainter.SectionGeometry.Size.X)
+		{
+			return;
+		}
+
 		// add green line for playback start
 		FSlateDrawElement::MakeBox(
 			InPainter.DrawElements,
@@ -191,7 +202,6 @@ void FSubSectionPainterUtil::PaintSectionBounds(const UMovieSceneSubSection& Sec
 			DrawEffects,
 			Tint
 		);
-		return true;
 	};
 
 
