Sequencer- Fix bug with FSpawnableRestoreState- it was assuming spawnables had Lifetime Binding Tracks, but they can instead have spawn tracks. This was causing an issue with the Animation Browser.

diff --git a/Engine/Plugins/MovieScene/SequencerScripting/Source/SequencerScriptingEditor/Private/SequencerTools.cpp b/Engine/Plugins/MovieScene/SequencerScripting/Source/SequencerScriptingEditor/Private/SequencerTools.cpp
--- a/Engine/Plugins/MovieScene/SequencerScripting/Source/SequencerScriptingEditor/Private/SequencerTools.cpp
+++ b/Engine/Plugins/MovieScene/SequencerScripting/Source/SequencerScriptingEditor/Private/SequencerTools.cpp
@@ -356,7 +356,7 @@ bool USequencerToolsFunctionLibrary::ExportAnimSequence(UWorld* World, ULevelSeq
 	{
 		FSpawnableRestoreState SpawnableRestoreState(MovieScene, Player->GetSharedPlaybackState().ToSharedPtr());
  
-		if (1 || SpawnableRestoreState.bWasChanged)
+		if (SpawnableRestoreState.bWasChanged)
 		{
 			// Evaluate at the beginning of the subscene time to ensure that spawnables are created before export
 			FFrameTime StartTime = FFrameRate::TransformTime(UE::MovieScene::DiscreteInclusiveLower(MovieScene->GetPlaybackRange()).Value, MovieScene->GetTickResolution(), MovieScene->GetDisplayRate());
diff --git a/Engine/Source/Editor/MovieSceneTools/Private/MovieSceneToolHelpers.cpp b/Engine/Source/Editor/MovieSceneTools/Private/MovieSceneToolHelpers.cpp
--- a/Engine/Source/Editor/MovieSceneTools/Private/MovieSceneToolHelpers.cpp
+++ b/Engine/Source/Editor/MovieSceneTools/Private/MovieSceneToolHelpers.cpp
@@ -4060,11 +4060,14 @@ FSpawnableRestoreState::FSpawnableRestoreState(UMovieScene* MovieScene, TSharedP
 	, WeakMovieScene(MovieScene)
 	, SharedPlaybackState(InSharedPlaybackState)
 {
-	for (int32 SpawnableIndex = 0; SpawnableIndex < WeakMovieScene->GetSpawnableCount(); ++SpawnableIndex)
+	auto SaveStateForSpawnable = [this](FGuid BindingGuid, ESpawnOwnership PreviousSpawnOwnership, TFunction<void(ESpawnOwnership)> SpawnOwnershipSetter)
 	{
-		FMovieSceneSpawnable& Spawnable = WeakMovieScene->GetSpawnable(SpawnableIndex);
+		// Spawnable could be in a subscene, so temporarily override it to persist throughout
+		SpawnOwnershipMap.Add(BindingGuid, PreviousSpawnOwnership);
+
+		SpawnOwnershipSetter(ESpawnOwnership::RootSequence);
 
-		UMovieSceneSpawnTrack* SpawnTrack = WeakMovieScene->FindTrack<UMovieSceneSpawnTrack>(Spawnable.GetGuid());
+		UMovieSceneSpawnTrack* SpawnTrack = WeakMovieScene->FindTrack<UMovieSceneSpawnTrack>(BindingGuid);
 
 		if (SpawnTrack && SpawnTrack->GetAllSections().Num() > 0)
 		{
@@ -4075,31 +4078,47 @@ FSpawnableRestoreState::FSpawnableRestoreState(UMovieScene* MovieScene, TSharedP
 			}
 
 			bWasChanged = true;
-			
-			// Spawnable could be in a subscene, so temporarily override it to persist throughout
-			SpawnOwnershipMap.Add(Spawnable.GetGuid(), Spawnable.GetSpawnOwnership());
-			Spawnable.SetSpawnOwnership(ESpawnOwnership::RootSequence);
 
 			UMovieSceneSpawnSection* SpawnSection = Cast<UMovieSceneSpawnSection>(SpawnTrack->GetAllSections()[0]);
 			SpawnSection->Modify();
 			SpawnSection->GetChannel().Reset();
 			SpawnSection->GetChannel().SetDefault(true);
 		}
+
+		UMovieSceneBindingLifetimeTrack* BindingLifetimeTrack = WeakMovieScene->FindTrack<UMovieSceneBindingLifetimeTrack>(BindingGuid);
+		if (BindingLifetimeTrack && BindingLifetimeTrack->GetAllSections().Num() > 0)
+		{
+			// Start a transaction that will be undone later for the modifications to the binding lifetime track
+			if (!bWasChanged)
+			{
+				GEditor->BeginTransaction(NSLOCTEXT("MovieSceneToolHelpers", "SpwanableRestoreState", "SpawnableRestoreState"));
+			}
+
+			bWasChanged = true;
+
+			BindingLifetimeTrack->Modify();
+			BindingLifetimeTrack->RemoveAllAnimationData();
+			// Add a single infinite section
+			UMovieSceneSection* NewSection = NewObject<UMovieSceneSection>(BindingLifetimeTrack, UMovieSceneBindingLifetimeSection::StaticClass(), NAME_None, RF_Transactional);
+			check(NewSection);
+			NewSection->SetRange(TRange<FFrameNumber>::All());
+			BindingLifetimeTrack->AddSection(*NewSection);
+		}
+	};
+	
+	for (int32 SpawnableIndex = 0; SpawnableIndex < WeakMovieScene->GetSpawnableCount(); ++SpawnableIndex)
+	{
+		FMovieSceneSpawnable& Spawnable = WeakMovieScene->GetSpawnable(SpawnableIndex);
+
+		SaveStateForSpawnable(Spawnable.GetGuid(), Spawnable.GetSpawnOwnership(), [&Spawnable](ESpawnOwnership SpawnOwnership) { Spawnable.SetSpawnOwnership(SpawnOwnership);});
 	}
 
 	TSharedPtr<UE::MovieScene::FSharedPlaybackState> SharedPlaybackStateToUse = SharedPlaybackState;
 	if (!SharedPlaybackStateToUse)
 	{
-		// Supporting deprecated path where we have no SharedPlaybackState
-		UE::MovieScene::FSharedPlaybackStateCreateParams CreateParams;
-		CreateParams.PlaybackContext = GEditor->GetEditorWorldContext().World();
-		UMovieSceneSequence* ThisSequence = WeakMovieScene->GetTypedOuter<UMovieSceneSequence>();
-		SharedPlaybackStateToUse = MakeShared<UE::MovieScene::FSharedPlaybackState>(*ThisSequence, CreateParams);
-
-		FMovieSceneEvaluationState State;
-		SharedPlaybackStateToUse->AddCapabilityRaw(&State);
-		State.AssignSequence(MovieSceneSequenceID::Root, *ThisSequence, SharedPlaybackStateToUse.ToSharedRef());
+		SharedPlaybackStateToUse = MovieSceneHelpers::CreateTransientSharedPlaybackState(GEditor->GetEditorWorldContext().World(), WeakMovieScene->GetTypedOuter<UMovieSceneSequence>());
 	}
+
 	ensure(SharedPlaybackStateToUse.IsValid());
 	if (UMovieSceneSequence* Sequence = WeakMovieScene->GetTypedOuter<UMovieSceneSequence>())
 	{
@@ -4111,30 +4130,8 @@ FSpawnableRestoreState::FSpawnableRestoreState(UMovieScene* MovieScene, TSharedP
 				{
 					if (UMovieSceneSpawnableBindingBase* SpawnableBinding = BindingReference.CustomBinding->AsSpawnable(SharedPlaybackStateToUse.ToSharedRef()))
 					{
-						UMovieSceneBindingLifetimeTrack* BindingLifetimeTrack = WeakMovieScene->FindTrack<UMovieSceneBindingLifetimeTrack>(BindingReference.ID);
-						if (BindingLifetimeTrack && BindingLifetimeTrack->GetAllSections().Num() > 0)
-						{
-							// Start a transaction that will be undone later for the modifications to the binding lifetime track
-							if (!bWasChanged)
-							{
-								GEditor->BeginTransaction(NSLOCTEXT("MovieSceneToolHelpers", "SpwanableRestoreState", "SpawnableRestoreState"));
-							}
-
-							bWasChanged = true;
-
-							// Spawnable could be in a subscene, so temporarily override it to persist throughout
-							SpawnOwnershipMap.Add(BindingReference.ID, SpawnableBinding->SpawnOwnership);
-							SpawnableBinding->SpawnOwnership = ESpawnOwnership::RootSequence;
-
-							BindingLifetimeTrack->Modify();
-							BindingLifetimeTrack->RemoveAllAnimationData();
-							// Add a single infinite section
-							UMovieSceneSection* NewSection = NewObject<UMovieSceneSection>(BindingLifetimeTrack, UMovieSceneBindingLifetimeSection::StaticClass(), NAME_None, RF_Transactional);
-							check(NewSection);
-							NewSection->SetRange(TRange<FFrameNumber>::All());
-							BindingLifetimeTrack->AddSection(*NewSection);
-						}
-					}
+						SaveStateForSpawnable(BindingReference.ID, SpawnableBinding->SpawnOwnership, [SpawnableBinding](ESpawnOwnership SpawnOwnership) { SpawnableBinding->SpawnOwnership = SpawnOwnership; });
+					}								
 				}
 			}
 		}
