Sequencer: Fixed dilation of sub sections

diff --git a/Engine/Source/Editor/MovieSceneTools/Private/TrackEditors/SubTrackEditorBase.cpp b/Engine/Source/Editor/MovieSceneTools/Private/TrackEditors/SubTrackEditorBase.cpp
--- a/Engine/Source/Editor/MovieSceneTools/Private/TrackEditors/SubTrackEditorBase.cpp
+++ b/Engine/Source/Editor/MovieSceneTools/Private/TrackEditors/SubTrackEditorBase.cpp
@@ -396,7 +396,9 @@ void FSubSectionEditorUtil::DilateSection(const TRange<FFrameNumber>& NewRange,
 {
 	if (SectionObject.Parameters.TimeScale.GetType() == EMovieSceneTimeWarpType::FixedPlayRate)
 	{
-		SectionObject.Parameters.TimeScale.Set(PreDilateTimeScale / DilationFactor);
+		// Clamp dilation to a 'sensible' range
+		double NewDilation = FMath::Clamp(PreDilateTimeScale / DilationFactor, -1000000.0, 1000000);
+		SectionObject.Parameters.TimeScale.Set(NewDilation);
 	}
 	else if (SectionObject.Parameters.TimeScale.GetType() == EMovieSceneTimeWarpType::Custom)
 	{
@@ -411,6 +413,8 @@ void FSubSectionEditorUtil::DilateSection(const TRange<FFrameNumber>& NewRange,
 			SectionObject.Parameters.TimeScale.AsCustom()->MarkAsChanged();
 		}
 	}
+
+	SectionObject.SetRange(NewRange);
 }
 
 
