Sequencer: Fix parsing of shot name so that it takes into account the number of characters for the take separator.

Previously, the take separator was limited to a single character and parsing the take number was limited to a single character. This now properly allows for multiple characters. For example, if the take separator is "_tk", "myshot_tk001" will parse "001" as the take number. Previously, it would get stuck trying to parse "tk001".

diff --git a/Engine/Source/Editor/MovieSceneTools/Private/MovieSceneToolHelpers.cpp b/Engine/Source/Editor/MovieSceneTools/Private/MovieSceneToolHelpers.cpp
--- a/Engine/Source/Editor/MovieSceneTools/Private/MovieSceneToolHelpers.cpp
+++ b/Engine/Source/Editor/MovieSceneTools/Private/MovieSceneToolHelpers.cpp
@@ -402,12 +402,12 @@ bool MovieSceneToolHelpers::ParseShotName(const FString& InShotName, FString& Sh
 				ShotPrefix = ShotName.Left(LastSlashPos);
 			}
 			
-			FString TakeStr = ShotName.RightChop(LastSlashPos + 1);
+			FString TakeStr = ShotName.RightChop(LastSlashPos + ProjectSettings->TakeSeparator.Len());
 			if (TakeStr.IsNumeric())
 			{
 				ShotNumber = INDEX_NONE; // Nullify the shot number since we only have a shot prefix
 				TakeNumber = FCString::Atoi(*TakeStr);
-				TakeNumberDigits = ShotName.Len() - (LastSlashPos+1);
+				TakeNumberDigits = ShotName.Len() - (LastSlashPos + ProjectSettings->TakeSeparator.Len());
 				return true;
 			}
 		}
