Sequencer: Protect against invalid Sequencer in FMovieSceneObjectBindingIDPicker

#rnx

diff --git a/Engine/Source/Editor/MovieSceneTools/Private/MovieSceneObjectBindingIDPicker.cpp b/Engine/Source/Editor/MovieSceneTools/Private/MovieSceneObjectBindingIDPicker.cpp
--- a/Engine/Source/Editor/MovieSceneTools/Private/MovieSceneObjectBindingIDPicker.cpp
+++ b/Engine/Source/Editor/MovieSceneTools/Private/MovieSceneObjectBindingIDPicker.cpp
@@ -217,8 +217,10 @@ bool FMovieSceneObjectBindingIDPicker::IsClassAllowed(const TSharedPtr<FSequence
 {
 	const UClass* InClass = nullptr;
 	auto BindingID = Node->BindingID;
-	UMovieSceneSequence* BindSequence = WeakSequencer.Pin()->State.FindSequence(BindingID.SequenceID);
-	UMovieScene* BindMovieScene = BindSequence->GetMovieScene();
+
+	TSharedPtr<ISequencer> Sequencer = WeakSequencer.Pin();
+	UMovieSceneSequence* BindSequence = Sequencer.IsValid() ? Sequencer->State.FindSequence(BindingID.SequenceID) : nullptr;
+	UMovieScene* BindMovieScene = BindSequence ? BindSequence->GetMovieScene() : nullptr;
 
 	InClass = MovieSceneHelpers::GetBoundObjectClass(BindSequence, BindingID.Guid);
 
