UE-220740 - TrackEvaluation level crashes upon opening

The check added in CL 34646598 is a bit overzealous. It's entirely possible that PostLoad might allocate a new object that won't have RF_LoadCompleted set. The point of the check is to verify the correctness of loaded objects, and not dynamic ones.

#rnx

diff --git a/Engine/Source/Editor/MovieSceneTools/Private/K2Node_GetSequenceBinding.cpp b/Engine/Source/Editor/MovieSceneTools/Private/K2Node_GetSequenceBinding.cpp
--- a/Engine/Source/Editor/MovieSceneTools/Private/K2Node_GetSequenceBinding.cpp
+++ b/Engine/Source/Editor/MovieSceneTools/Private/K2Node_GetSequenceBinding.cpp
@@ -36,9 +36,9 @@ void EnsureFullyPreloaded(UObject* Object)
 		return;
 	}
 
+	FLinkerLoad* Linker = Object->GetLinker();
 	if (Object->HasAnyFlags(RF_NeedLoad))
 	{
-		FLinkerLoad* Linker = Object->GetLinker();
 		if (ensure(Linker))
 		{
 			Linker->Preload(Object);
@@ -46,7 +46,12 @@ void EnsureFullyPreloaded(UObject* Object)
 		}
 	}
 
-	check(Object->HasAnyFlags(RF_LoadCompleted));
+	// We only want to ensure that _loaded_ objects have RF_LoadCompleted set.
+	// Some objects can be created during postload, so we don't need to verify RF_LoadCompleted in those cases.
+	if (Linker)
+	{
+		check(Object->HasAnyFlags(RF_LoadCompleted));
+	}
 
 	TArray<UObject*> ObjectReferences;
 	FReferenceFinder(ObjectReferences, nullptr, false, true, false, true).FindReferences(Object);
