Constraints: fixed compensation computation

#rnx

diff --git a/Engine/Source/Editor/MovieSceneTools/Private/Constraints/MovieSceneConstraintChannelHelper.cpp b/Engine/Source/Editor/MovieSceneTools/Private/Constraints/MovieSceneConstraintChannelHelper.cpp
--- a/Engine/Source/Editor/MovieSceneTools/Private/Constraints/MovieSceneConstraintChannelHelper.cpp
+++ b/Engine/Source/Editor/MovieSceneTools/Private/Constraints/MovieSceneConstraintChannelHelper.cpp
@@ -87,7 +87,7 @@ void FCompensationEvaluator::ComputeLocalTransforms(
 
 	// get all constraints to evaluate
 	const FConstraintsManagerController& Controller = FConstraintsManagerController::Get(InWorld);
-	static constexpr bool bSorted = true;
+	static constexpr bool bSorted = true, bTickHandles = true;
 	const TArray<ConstraintPtr> AllConstraints = Controller.GetAllConstraints(bSorted);
 
 	FMovieSceneInverseSequenceTransform LocalToRootTransform = InSequencer->GetFocusedMovieSceneSequenceTransform().Inverse();
@@ -121,7 +121,7 @@ void FCompensationEvaluator::ComputeLocalTransforms(
 		{
 			if (InConstraint.IsValid())
 			{
-				InConstraint->Evaluate();
+				InConstraint->Evaluate(bTickHandles);
 			}
 		}
 
@@ -132,8 +132,6 @@ void FCompensationEvaluator::ComputeLocalTransforms(
 				BakeHelper->PostEvaluation(MovieScene, FrameNumber);
 			}
 		}
-		// evaluate ControlRig?
-		// ControlRig->Evaluate_AnyThread();
 
 		FTransform& ChildLocal = ChildLocals[Index];
 		FTransform& ChildGlobal = ChildGlobals[Index];
@@ -141,6 +139,7 @@ void FCompensationEvaluator::ComputeLocalTransforms(
 
 		// store child transforms        	
 		ChildLocal = Handle->GetLocalTransform();
+		Handle->PreEvaluate();
 		ChildGlobal = Handle->GetGlobalTransform();
 
 		const UTickableTransformConstraint* LastConstraint = GetLastActiveConstraint();
