Sequencer: Refresh on property reset to default.

This is a minor UI fix for a bug that has been there for a while - the binding data doesn't clear when the reset to default button is pushed.

diff --git a/Engine/Source/Editor/MovieSceneTools/Private/MovieSceneObjectBindingIDCustomization.cpp b/Engine/Source/Editor/MovieSceneTools/Private/MovieSceneObjectBindingIDCustomization.cpp
--- a/Engine/Source/Editor/MovieSceneTools/Private/MovieSceneObjectBindingIDCustomization.cpp
+++ b/Engine/Source/Editor/MovieSceneTools/Private/MovieSceneObjectBindingIDCustomization.cpp
@@ -63,6 +63,8 @@ void FMovieSceneObjectBindingIDCustomization::CustomizeHeader(TSharedRef<IProper
 	using namespace UE::Sequencer;
 
 	StructProperty = PropertyHandle;
+	StructProperty->SetOnPropertyResetToDefault(
+		FSimpleDelegate::CreateSP(this, &FMovieSceneObjectBindingIDCustomization::OnResetToDefault));
 
 	const FString& MetaClassName = PropertyHandle->GetMetaData("MetaClass");
 	const FString& MustImplementName = PropertyHandle->GetMetaData("MustImplement");
@@ -223,4 +225,9 @@ void FMovieSceneObjectBindingIDCustomization::SetCurrentValue(const FMovieSceneO
 	StructProperty->NotifyFinishedChangingProperties();
 }
 
+void FMovieSceneObjectBindingIDCustomization::OnResetToDefault()
+{
+	StructProperty->RequestRebuildChildren();
+}
+
 #undef LOCTEXT_NAMESPACE
diff --git a/Engine/Source/Editor/MovieSceneTools/Public/MovieSceneObjectBindingIDCustomization.h b/Engine/Source/Editor/MovieSceneTools/Public/MovieSceneObjectBindingIDCustomization.h
--- a/Engine/Source/Editor/MovieSceneTools/Public/MovieSceneObjectBindingIDCustomization.h
+++ b/Engine/Source/Editor/MovieSceneTools/Public/MovieSceneObjectBindingIDCustomization.h
@@ -57,5 +57,7 @@ private:
 
 	FReply OnDrop(const FGeometry& InGeometry, const FDragDropEvent& InDragDropEvent);
 
+	void OnResetToDefault();
+
 	TSharedPtr<IPropertyHandle> StructProperty;
 };
\ No newline at end of file
