Sequencer: Round to frame number when assigning values from key editors

diff --git a/Engine/Source/Editor/MovieSceneTools/Private/CurveKeyEditors/SequencerKeyEditor.h b/Engine/Source/Editor/MovieSceneTools/Private/CurveKeyEditors/SequencerKeyEditor.h
--- a/Engine/Source/Editor/MovieSceneTools/Private/CurveKeyEditors/SequencerKeyEditor.h
+++ b/Engine/Source/Editor/MovieSceneTools/Private/CurveKeyEditors/SequencerKeyEditor.h
@@ -214,7 +214,7 @@ struct TSequencerKeyEditor
 		}
 		else
 		{
-			const FFrameNumber CurrentTime = Sequencer->GetLocalTime().Time.FloorToFrame() - ChannelMetaData->GetOffsetTime(OwningSection);
+			const FFrameNumber CurrentTime = Sequencer->GetLocalTime().Time.RoundToFrame() - ChannelMetaData->GetOffsetTime(OwningSection);
 
 			EMovieSceneKeyInterpolation Interpolation = GetInterpolationMode(Channel, CurrentTime, Sequencer->GetKeyInterpolation());
 
