Sequencer: Fix crash in Trim Left/Right and Preserve for Attach Track.

Check for upper/lower bounds

diff --git a/Engine/Source/Editor/MovieSceneTools/Private/TrackEditors/AttachTrackEditor.cpp b/Engine/Source/Editor/MovieSceneTools/Private/TrackEditors/AttachTrackEditor.cpp
--- a/Engine/Source/Editor/MovieSceneTools/Private/TrackEditors/AttachTrackEditor.cpp
+++ b/Engine/Source/Editor/MovieSceneTools/Private/TrackEditors/AttachTrackEditor.cpp
@@ -786,8 +786,10 @@ struct FAttachRevertModifier
 
 		TransformEvaluator = FWorldTransformEvaluator(InWeakAttachTrackEditor, ConstraintObject, InSocketName, InComponentName);
 
-		BeginConstraintTransform = TransformEvaluator(InRevertRange.GetLowerBoundValue());
-
+		if (InRevertRange.HasLowerBound())
+		{
+			BeginConstraintTransform = TransformEvaluator(InRevertRange.GetLowerBoundValue());
+		}
 	}
 
 	/**
@@ -796,9 +798,13 @@ struct FAttachRevertModifier
 	FAttachRevertModifier(TSharedPtr<F3DAttachTrackEditor> InWeakAttachTrackEditor, const TRange<FFrameNumber>& InRevertRange, const FWorldTransformEvaluator& InTransformEvaluator, bool bInFullRevert)
 		: bFullRevert(bInFullRevert)
 		, TransformEvaluator(InTransformEvaluator)
-		, BeginConstraintTransform(InTransformEvaluator(InRevertRange.GetLowerBoundValue()))
 		, RevertRange(InRevertRange)
-	{}
+	{
+		if (InRevertRange.HasLowerBound())
+		{
+			BeginConstraintTransform = InTransformEvaluator(InRevertRange.GetLowerBoundValue());
+		}
+	}
 
 	/** Reverts a transform in relative space to world space */
 	FTransform operator()(const FTransform& InTransform, const FFrameNumber& InTime)
@@ -868,12 +874,18 @@ void UpdateChannelTransforms(const TRange<FFrameNumber>& InAttachRange, TMap<FFr
 
 		// Get the keys contained in before attach and after attach ranges
 		InChannels[ChannelIndex]->GetKeys(ExcludedRanges[0], &LowerKeyTimes, &LowerKeyHandles);
-		ExcludedRanges.Top().SetLowerBound(TRangeBound<FFrameNumber>::Exclusive(ExcludedRanges.Top().GetLowerBoundValue()));
+		if (ExcludedRanges.Top().HasLowerBound())
+		{
+			ExcludedRanges.Top().SetLowerBound(TRangeBound<FFrameNumber>::Exclusive(ExcludedRanges.Top().GetLowerBoundValue()));
+		}
 		InChannels[ChannelIndex]->GetKeys(ExcludedRanges.Top(), &UpperKeyTimes, &UpperKeyHandles);
 
 		// Add all keys before attach range if they exist
 		int32 ValueIndex = 0;
-		if (ExcludedRanges.Num() > 0 && ExcludedRanges[0].GetUpperBoundValue() <= InAttachRange.GetLowerBoundValue() && LowerKeyTimes.Num() > 0)
+		if (ExcludedRanges.Num() > 0 && 
+			ExcludedRanges[0].HasUpperBound() &&
+			InAttachRange.HasLowerBound() && 
+			ExcludedRanges[0].GetUpperBoundValue() <= InAttachRange.GetLowerBoundValue() && LowerKeyTimes.Num() > 0)
 		{
 			InChannels[ChannelIndex]->DeleteKeys(LowerKeyHandles);
 			TArray<FMovieSceneDoubleValue> ValuesToAdd;
@@ -886,7 +898,9 @@ void UpdateChannelTransforms(const TRange<FFrameNumber>& InAttachRange, TMap<FFr
 		InChannels[ChannelIndex]->AddKeys(NewKeyFrames, NewKeyValues);
 
 		// Add all keys after attach range if they exist
-		if (ExcludedRanges.Num() > 0 && ExcludedRanges.Top().GetLowerBoundValue() >= InAttachRange.GetUpperBoundValue() && ValueIndex < PrevKeyValues.Num() && UpperKeyTimes.Num() > 0)
+		if (ExcludedRanges.Num() > 0 && 
+			ExcludedRanges.Top().HasLowerBound() && InAttachRange.HasUpperBound() &&
+			ExcludedRanges.Top().GetLowerBoundValue() >= InAttachRange.GetUpperBoundValue() && ValueIndex < PrevKeyValues.Num() && UpperKeyTimes.Num() > 0)
 		{
 			InChannels[ChannelIndex]->DeleteKeys(UpperKeyHandles);
 			TArray<FMovieSceneDoubleValue> ValuesToAdd;
@@ -1021,9 +1035,20 @@ void F3DAttachTrackEditor::TrimAndPreserve(FGuid InObjectBinding, UMovieSceneSec
 		{
 			TArray<FKeyHandle> KeyAtTime;
 
-			bInTrimLeft ? 
-			Channel->GetKeys(TRange<FFrameNumber>::Inclusive(ExcludedRange.GetLowerBoundValue() - 1, ExcludedRange.GetLowerBoundValue() - 1), nullptr, &KeyAtTime) : 
-			Channel->GetKeys(TRange<FFrameNumber>::Inclusive(ExcludedRange.GetUpperBoundValue() - 1, ExcludedRange.GetUpperBoundValue() - 1), nullptr, &KeyAtTime);
+			if (bInTrimLeft)
+			{
+				if (ExcludedRange.HasLowerBound())
+				{
+					Channel->GetKeys(TRange<FFrameNumber>::Inclusive(ExcludedRange.GetLowerBoundValue() - 1, ExcludedRange.GetLowerBoundValue() - 1), nullptr, &KeyAtTime);
+				}
+			}
+			else
+			{
+				if (ExcludedRange.HasUpperBound())
+				{
+					Channel->GetKeys(TRange<FFrameNumber>::Inclusive(ExcludedRange.GetUpperBoundValue() - 1, ExcludedRange.GetUpperBoundValue() - 1), nullptr, &KeyAtTime);
+				}
+			}
 
 			Channel->DeleteKeys(KeyAtTime);
 			Channel->AutoSetTangents();
