Curve Editor: avoid forcing flush for every modifications

#rnx

diff --git a/Engine/Source/Editor/MovieSceneTools/Private/Channels/ChannelCurveModel.cpp b/Engine/Source/Editor/MovieSceneTools/Private/Channels/ChannelCurveModel.cpp
--- a/Engine/Source/Editor/MovieSceneTools/Private/Channels/ChannelCurveModel.cpp
+++ b/Engine/Source/Editor/MovieSceneTools/Private/Channels/ChannelCurveModel.cpp
@@ -246,7 +246,7 @@ void FChannelCurveModel<ChannelType, ChannelValue, KeyType>::GetKeyPositions(TAr
 template <class ChannelType, class ChannelValue, class KeyType>
 void FChannelCurveModel<ChannelType, ChannelValue, KeyType>::SetKeyPositions(TArrayView<const FKeyHandle> InKeys, TArrayView<const FKeyPosition> InKeyPositions, EPropertyChangeType::Type ChangeType)
 {
-	UE::MovieScene::FScopedSignedObjectModifyDefer Defer(true);
+	UE::MovieScene::FScopedSignedObjectModifyDefer Defer;
 
 	ChannelType* Channel = ChannelHandle.Get();
 	UMovieSceneSection* Section = WeakSection.Get();
