Sequencer: Fix can split/trim so that it only returns true if the resulting section is not empty.

diff --git a/Engine/Source/Editor/MovieSceneTools/Private/MovieSceneToolHelpers.cpp b/Engine/Source/Editor/MovieSceneTools/Private/MovieSceneToolHelpers.cpp
--- a/Engine/Source/Editor/MovieSceneTools/Private/MovieSceneToolHelpers.cpp
+++ b/Engine/Source/Editor/MovieSceneTools/Private/MovieSceneToolHelpers.cpp
@@ -145,14 +145,27 @@ bool MovieSceneToolHelpers::CanTrimSectionLeft(const TSet<UMovieSceneSection*>&
 {
 	for (UMovieSceneSection* Section : Sections)
 	{
-		if (Section && Section->HasStartFrame() && Time.Time.FrameNumber > Section->GetInclusiveStartFrame())
+		if (Section)
 		{
-			// Don't allow an empty section
-			if (Section->HasEndFrame() && Time.Time.FrameNumber >= Section->GetExclusiveEndFrame())
+			TRange<FFrameNumber> Range = Section->GetRange();
+			TArray<TRange<FFrameNumber> > Splits = Range.Split(Time.Time.FrameNumber);
+
+			if (Splits.Num() > 0)
 			{
-				continue;
+				if (Splits[0] == Range) // can't split
+				{
+					continue;
+				}
+				else
+				{
+					TRange<FFrameNumber> Head = Splits[0];
+
+					if (!Head.IsEmpty())
+					{
+						return true;
+					}
+				}
 			}
-			return true;
 		}
 	}
 	return false;
@@ -162,14 +175,27 @@ bool MovieSceneToolHelpers::CanTrimSectionRight(const TSet<UMovieSceneSection*>&
 {
 	for (UMovieSceneSection* Section : Sections)
 	{
-		if (Section && Section->HasEndFrame() && Time.Time.FrameNumber < Section->GetExclusiveEndFrame())
+		if (Section)
 		{
-			// Don't allow an empty section
-			if (Section->HasStartFrame() && Time.Time.FrameNumber <= Section->GetInclusiveStartFrame())
+			TRange<FFrameNumber> Range = Section->GetRange();
+			TArray<TRange<FFrameNumber> > Splits = Range.Split(Time.Time.FrameNumber);
+
+			if (Splits.Num() > 0)
 			{
-				continue;
+				if (Splits[0] == Range) // can't split
+				{
+					continue;
+				}
+				else
+				{
+					TRange<FFrameNumber> Tail = Splits.Last();
+
+					if (!Tail.IsEmpty())
+					{
+						return true;
+					}
+				}
 			}
-			return true;
 		}
 	}
 	return false;
@@ -272,17 +298,28 @@ bool MovieSceneToolHelpers::CanSplitSection(const TSet<UMovieSceneSection*>& Sec
 {
 	for (UMovieSceneSection* Section : Sections)
 	{
-		if (Section && Section->IsTimeWithinSection(Time.Time.FrameNumber))
+		if (Section)
 		{
-			if (Section->HasStartFrame() && Time.Time.FrameNumber <= Section->GetInclusiveStartFrame())
-			{
-				continue;
-			}
-			if (Section->HasEndFrame() && Time.Time.FrameNumber >= Section->GetExclusiveEndFrame())
+			TRange<FFrameNumber> Range = Section->GetRange();
+			TArray<TRange<FFrameNumber> > Splits = Range.Split(Time.Time.FrameNumber);
+
+			if (Splits.Num() > 0)
 			{
-				continue;
+				if (Splits[0] == Range) // can't split
+				{
+					continue;
+				}
+				else if (Splits.Num() == 2)
+				{
+					TRange<FFrameNumber> Head = Splits[0];
+					TRange<FFrameNumber> Tail = Splits[1];
+
+					if (!Head.IsEmpty() && !Tail.IsEmpty())
+					{
+						return true;
+					}
+				}
 			}
-			return true;
 		}
 	}
 	return false;
