Sequencer: Fix crash in Trim Left/Right and Preserve for Attach Track.

Check for upper/lower bounds

diff --git a/Engine/Source/Editor/MovieSceneTools/Private/TrackEditors/AttachTrackEditor.cpp b/Engine/Source/Editor/MovieSceneTools/Private/TrackEditors/AttachTrackEditor.cpp
--- a/Engine/Source/Editor/MovieSceneTools/Private/TrackEditors/AttachTrackEditor.cpp
+++ b/Engine/Source/Editor/MovieSceneTools/Private/TrackEditors/AttachTrackEditor.cpp
@@ -94,14 +94,30 @@ public:
 			LOCTEXT("TrimRightPreserve", "Trim Right and Preserve"),
 			LOCTEXT("TrimRightPreserveToolTip", "Trims the right side of this attach at the current time and preserves the last key's world coordinates"),
 			FSlateIcon(),
-			FUIAction(FExecuteAction::CreateRaw(AttachTrackEditor, &F3DAttachTrackEditor::TrimAndPreserve, ObjectBinding, &Section, false))
+			FUIAction(
+				FExecuteAction::CreateRaw(AttachTrackEditor, &F3DAttachTrackEditor::TrimAndPreserve, ObjectBinding, &Section, false),
+				FCanExecuteAction::CreateLambda([this] { 
+					TSet<UMovieSceneSection*> Sections;
+					Sections.Add(&Section);
+					TSharedPtr<ISequencer> Sequencer = AttachTrackEditor->GetSequencer();
+					FQualifiedFrameTime Time = Sequencer->GetLocalTime();
+					return MovieSceneToolHelpers::CanTrimSectionRight(Sections, Time);
+				}))
 		);
 
 		MenuBuilder.AddMenuEntry(
 			LOCTEXT("TrimLeftPreserve", "Trim Left and Preserve"),
 			LOCTEXT("TrimLeftPreserveToolTip", "Trims the left side of this attach at the current time and preserves the first key's world coordinates"),
 			FSlateIcon(),
-			FUIAction(FExecuteAction::CreateRaw(AttachTrackEditor, &F3DAttachTrackEditor::TrimAndPreserve, ObjectBinding, &Section, true))
+			FUIAction(
+				FExecuteAction::CreateRaw(AttachTrackEditor, &F3DAttachTrackEditor::TrimAndPreserve, ObjectBinding, &Section, true),
+				FCanExecuteAction::CreateLambda([this] {
+					TSet<UMovieSceneSection*> Sections;
+					Sections.Add(&Section);
+					TSharedPtr<ISequencer> Sequencer = AttachTrackEditor->GetSequencer();
+					FQualifiedFrameTime Time = Sequencer->GetLocalTime();
+					return MovieSceneToolHelpers::CanTrimSectionLeft(Sections, Time);
+				}))
 		);
 
 		MenuBuilder.EndSection();
@@ -989,7 +1005,14 @@ void F3DAttachTrackEditor::TrimAndPreserve(FGuid InObjectBinding, UMovieSceneSec
 			FFrameNumber PreserveEdgeTime;
 			if (bInTrimLeft)
 			{
-				PreserveEdgeTime = ExcludedRange.GetUpperBoundValue();
+				if (ExcludedRange.HasUpperBound())
+				{
+					PreserveEdgeTime = ExcludedRange.GetUpperBoundValue();
+				}
+				else
+				{
+					PreserveEdgeTime = TNumericLimits<FFrameNumber>::Max();
+				}
 				RevertEdgeTime = PreserveEdgeTime.Value - 1;
 				EdgeKeys = { PreserveEdgeTime, RevertEdgeTime };
 				ResizeAndAddKey(PreserveEdgeTime, Channels.Num(), TransformMap, nullptr);
@@ -997,7 +1020,14 @@ void F3DAttachTrackEditor::TrimAndPreserve(FGuid InObjectBinding, UMovieSceneSec
 			}
 			else
 			{
-				RevertEdgeTime = ExcludedRange.GetLowerBoundValue();
+				if (ExcludedRange.HasLowerBound())
+				{
+					RevertEdgeTime = ExcludedRange.GetLowerBoundValue();
+				}
+				else
+				{
+					RevertEdgeTime = TNumericLimits<FFrameNumber>::Min();
+				}
 				PreserveEdgeTime = RevertEdgeTime.Value - 1;
 				EdgeKeys = { RevertEdgeTime, PreserveEdgeTime };
 				ResizeAndAddKey(RevertEdgeTime, Channels.Num(), TransformMap, &KeyTimesToCompensate);
