Sequencer- SoftObjectPtr parameters to Director Blueprint bindings would be lost after reloading editor.

diff --git a/Engine/Source/Editor/MovieSceneTools/Private/MovieSceneDirectorBlueprintUtils.cpp b/Engine/Source/Editor/MovieSceneTools/Private/MovieSceneDirectorBlueprintUtils.cpp
--- a/Engine/Source/Editor/MovieSceneTools/Private/MovieSceneDirectorBlueprintUtils.cpp
+++ b/Engine/Source/Editor/MovieSceneTools/Private/MovieSceneDirectorBlueprintUtils.cpp
@@ -311,7 +311,14 @@ FMovieSceneDirectorBlueprintEntrypointResult FMovieSceneDirectorBlueprintUtils::
 				bool bMarkAsModified = false;
 				if (PayloadVariable->ObjectValue.IsValid())
 				{
-					Schema->TrySetDefaultObject(*Pin, PayloadVariable->ObjectValue.ResolveObject(), bMarkAsModified);
+					if (UObject* ResolvedObject = PayloadVariable->ObjectValue.ResolveObject())
+					{
+						Schema->TrySetDefaultObject(*Pin, ResolvedObject, bMarkAsModified);
+					}
+					else
+					{
+						Schema->TrySetDefaultValue(*Pin, PayloadVariable->ObjectValue.ToString(), bMarkAsModified);
+					}
 				}
 				else if (!PayloadVariable->Value.IsEmpty())
 				{
