Fix potential viewport thumbnail rendering crash: Hook ViewState to referenced object collection.

Code was added in CL 17019155 for UE-108255.

diff --git a/Engine/Source/Editor/MovieSceneTools/Private/MovieSceneToolsUserSettings.cpp b/Engine/Source/Editor/MovieSceneTools/Private/MovieSceneToolsUserSettings.cpp
--- a/Engine/Source/Editor/MovieSceneTools/Private/MovieSceneToolsUserSettings.cpp
+++ b/Engine/Source/Editor/MovieSceneTools/Private/MovieSceneToolsUserSettings.cpp
@@ -1,6 +1,7 @@
 // Copyright Epic Games, Inc. All Rights Reserved.
 
 #include "MovieSceneToolsUserSettings.h"
+#include "SceneManagement.h"
 #include "UObject/UnrealType.h"
 
 
@@ -28,6 +29,16 @@ void UMovieSceneUserThumbnailSettings::PostEditChangeProperty(FPropertyChangedEv
 	SaveConfig();
 }
 
+void UMovieSceneUserThumbnailSettings::AddReferencedObjects(UObject* InThis, FReferenceCollector& InCollector)
+{
+	UMovieSceneUserThumbnailSettings* This = CastChecked<UMovieSceneUserThumbnailSettings>(InThis);
+	if (This->ViewState.GetReference())
+	{
+		This->ViewState.GetReference()->AddReferencedObjects(InCollector);
+	}
+	Super::AddReferencedObjects(InThis, InCollector);
+}
+
 UMovieSceneUserImportFBXSettings::UMovieSceneUserImportFBXSettings(const FObjectInitializer& Initializer)
 	: Super(Initializer)
 {
diff --git a/Engine/Source/Editor/MovieSceneTools/Public/MovieSceneToolsUserSettings.h b/Engine/Source/Editor/MovieSceneTools/Public/MovieSceneToolsUserSettings.h
--- a/Engine/Source/Editor/MovieSceneTools/Public/MovieSceneToolsUserSettings.h
+++ b/Engine/Source/Editor/MovieSceneTools/Public/MovieSceneToolsUserSettings.h
@@ -53,6 +53,8 @@ public:
 
 	virtual void PostEditChangeProperty( struct FPropertyChangedEvent& PropertyChangedEvent) override;
 
+	static void AddReferencedObjects(UObject* InThis, FReferenceCollector& InCollector);
+
 private:
 	FOnForceRedraw OnForceRedrawEvent;
 };
