[Sequencer] Sidebar Details:
 - Fix broken lambdas for channel extensions (Perlin Noise, etc)
 - Fix broken lambdas for ContextMenu functions, which the sidebar uses to create the display widgets
 - Fix crash on load of level with sidebar section selected
 - Fix "SequencerAddConvertBindings" automation test not passing
 - Create ISidebarChannelExtension to allow saving a reference for the sidebar

diff --git a/Engine/Plugins/Animation/LiveLink/Source/LiveLinkSequencer/Private/Sequencer/LiveLinkPropertyTrackEditor.cpp b/Engine/Plugins/Animation/LiveLink/Source/LiveLinkSequencer/Private/Sequencer/LiveLinkPropertyTrackEditor.cpp
--- a/Engine/Plugins/Animation/LiveLink/Source/LiveLinkSequencer/Private/Sequencer/LiveLinkPropertyTrackEditor.cpp
+++ b/Engine/Plugins/Animation/LiveLink/Source/LiveLinkSequencer/Private/Sequencer/LiveLinkPropertyTrackEditor.cpp
@@ -43,41 +43,49 @@ protected:
 
 void FLiveLinkSection::BuildSectionContextMenu(FMenuBuilder& MenuBuilder, const FGuid& InObjectBinding)
 {
-	UMovieSceneLiveLinkSection* LiveLinkSection = CastChecked<UMovieSceneLiveLinkSection>(WeakSection.Get());
-	TSharedPtr<ISequencer> SequencerPtr = WeakSequencer.Pin();
-
-	auto MakeUIAction = [=](int32 Index)
+	auto MakeUIAction = [this](int32 Index)
 	{
 		return FUIAction(
-			FExecuteAction::CreateLambda([=]
-			{
-				FScopedTransaction Transaction(LOCTEXT("SetLiveLinkActiveChannelsTransaction", "Set Live LinkActive Channels"));
-				LiveLinkSection->Modify();
-				TArray<bool> ChannelMask = LiveLinkSection->ChannelMask;
-				ChannelMask[Index] = !ChannelMask[Index];
-				LiveLinkSection->SetMask(ChannelMask);
-
-				SequencerPtr->NotifyMovieSceneDataChanged(EMovieSceneDataChangeType::MovieSceneStructureItemsChanged);
-			}
-			),
-			FCanExecuteAction(),
-			FGetActionCheckState::CreateLambda([=]
-			{
-				TArray<bool> ChannelMask = LiveLinkSection->ChannelMask;
-				if (ChannelMask[Index])
+			FExecuteAction::CreateLambda([this, Index]
 				{
-					return ECheckBoxState::Checked;
-				}
-				else 
+					const TSharedPtr<ISequencer> SequencerPtr = WeakSequencer.Pin();
+					if (!SequencerPtr.IsValid())
+					{
+						return;
+					}
+
+					UMovieSceneLiveLinkSection* LiveLinkSection = CastChecked<UMovieSceneLiveLinkSection>(WeakSection.Get());
+					
+					FScopedTransaction Transaction(LOCTEXT("SetLiveLinkActiveChannelsTransaction", "Set Live LinkActive Channels"));
+					LiveLinkSection->Modify();
+					TArray<bool> ChannelMask = LiveLinkSection->ChannelMask;
+					ChannelMask[Index] = !ChannelMask[Index];
+					LiveLinkSection->SetMask(ChannelMask);
+
+					SequencerPtr->NotifyMovieSceneDataChanged(EMovieSceneDataChangeType::MovieSceneStructureItemsChanged);
+				}),
+			FCanExecuteAction(),
+			FGetActionCheckState::CreateLambda([this, Index]
 				{
-					return ECheckBoxState::Unchecked;
-				}
-			})
+					UMovieSceneLiveLinkSection* LiveLinkSection = CastChecked<UMovieSceneLiveLinkSection>(WeakSection.Get());
+
+					TArray<bool> ChannelMask = LiveLinkSection->ChannelMask;
+					if (ChannelMask[Index])
+					{
+						return ECheckBoxState::Checked;
+					}
+					else 
+					{
+						return ECheckBoxState::Unchecked;
+					}
+				})
 		);
 	};
 
 	MenuBuilder.BeginSection(NAME_None, LOCTEXT("LiveLinkChannelsText", "Active Live Link Channels"));
 	{
+		UMovieSceneLiveLinkSection* LiveLinkSection = CastChecked<UMovieSceneLiveLinkSection>(WeakSection.Get());
+
 		TArray<FMovieSceneChannelMetaData> AllMetaData;
 		for (const FMovieSceneChannelEntry& Entry : LiveLinkSection->GetChannelProxy().GetAllEntries())
 		{
diff --git a/Engine/Plugins/Experimental/PropertyAnimator/Source/PropertyAnimatorEditor/Private/Sequencer/PropertyAnimatorEditorCurveChannelInterface.h b/Engine/Plugins/Experimental/PropertyAnimator/Source/PropertyAnimatorEditor/Private/Sequencer/PropertyAnimatorEditorCurveChannelInterface.h
--- a/Engine/Plugins/Experimental/PropertyAnimator/Source/PropertyAnimatorEditor/Private/Sequencer/PropertyAnimatorEditorCurveChannelInterface.h
+++ b/Engine/Plugins/Experimental/PropertyAnimator/Source/PropertyAnimatorEditor/Private/Sequencer/PropertyAnimatorEditorCurveChannelInterface.h
@@ -34,8 +34,8 @@ struct TPropertyAnimatorEditorCurveChannelInterface : ISequencerChannelInterface
 	virtual TUniquePtr<FCurveModel> CreateCurveEditorModel_Raw(const FMovieSceneChannelHandle&, const UE::Sequencer::FCreateCurveEditorModelParams& Params) const override { return nullptr; }
 	virtual TSharedPtr<UE::Sequencer::FChannelModel> CreateChannelModel_Raw(const FMovieSceneChannelHandle&, FName) const override { return nullptr; }
 	virtual TSharedPtr<UE::Sequencer::STrackAreaLaneView> CreateChannelView_Raw(const FMovieSceneChannelHandle&, TWeakPtr<UE::Sequencer::FChannelModel>, const UE::Sequencer::FCreateTrackLaneViewParams&) const override { return nullptr; }
-	virtual void ExtendSectionMenu_Raw(FMenuBuilder& InMenuBuilder, TSharedPtr<FExtender> InMenuExtender, TConstArrayView<FMovieSceneChannelHandle> InChannels, TConstArrayView<UMovieSceneSection*> InSections, TWeakPtr<ISequencer> InSequencer) const override;
-	virtual void ExtendSidebarMenu_Raw(FMenuBuilder& InMenuBuilder, TSharedPtr<FExtender> InMenuExtender, TConstArrayView<FMovieSceneChannelHandle> InChannels, TConstArrayView<UMovieSceneSection*> InSections, TWeakPtr<ISequencer> InSequencer) const override;
+	virtual void ExtendSectionMenu_Raw(FMenuBuilder& MenuBuilder, TSharedPtr<FExtender> InMenuExtender, TConstArrayView<FMovieSceneChannelHandle> InChannels, const TArray<TWeakObjectPtr<UMovieSceneSection>>& InWeakSections, TWeakPtr<ISequencer> InWeakSequencer) const override;
+	virtual TSharedPtr<ISidebarChannelExtension> ExtendSidebarMenu_Raw(FMenuBuilder& MenuBuilder, TSharedPtr<FExtender> InMenuExtender, TConstArrayView<FMovieSceneChannelHandle> InChannels, const TArray<TWeakObjectPtr<UMovieSceneSection>>& InWeakSections, TWeakPtr<ISequencer> InWeakSequencer) const override;
 	virtual int32 DrawExtra_Raw(FMovieSceneChannel* InChannel, const UMovieSceneSection* InOwner, const FSequencerChannelPaintArgs& InPaintArgs, int32 InLayerId) const override;
 	virtual TSharedRef<SWidget> CreateKeyEditor_Raw(const FMovieSceneChannelHandle& InChannel, const UE::Sequencer::FCreateKeyEditorParams& Params) const override;
 	//~ End ISequencerChannelInterface
@@ -45,10 +45,10 @@ template<typename InChannelType, typename InMenuExtensionType>
 void TPropertyAnimatorEditorCurveChannelInterface<InChannelType, InMenuExtensionType>::ExtendSectionMenu_Raw(FMenuBuilder& MenuBuilder
 	, TSharedPtr<FExtender> InMenuExtender
 	, TConstArrayView<FMovieSceneChannelHandle> InChannels
-	, TConstArrayView<UMovieSceneSection*> InSections
-	, TWeakPtr<ISequencer> InSequencer) const
+	, const TArray<TWeakObjectPtr<UMovieSceneSection>>& InWeakSections
+	, TWeakPtr<ISequencer> InWeakSequencer) const
 {
-	TSharedRef<FPropertyAnimatorEditorCurveSectionMenuExtension> Extension = MakeShared<FMenuExtensionType>(InChannels, InSections);
+	TSharedRef<FPropertyAnimatorEditorCurveSectionMenuExtension> Extension = MakeShared<FMenuExtensionType>(InChannels, InWeakSections);
 
 	InMenuExtender->AddMenuExtension(TEXT("SequencerChannels"), EExtensionHook::First, nullptr
 		, FMenuExtensionDelegate::CreateLambda([Extension](FMenuBuilder& InInnerMenuBuilder)
@@ -58,14 +58,15 @@ void TPropertyAnimatorEditorCurveChannelInterface<InChannelType, InMenuExtension
 }
 
 template<typename InChannelType, typename InMenuExtensionType>
-void TPropertyAnimatorEditorCurveChannelInterface<InChannelType, InMenuExtensionType>::ExtendSidebarMenu_Raw(FMenuBuilder& MenuBuilder
+TSharedPtr<ISidebarChannelExtension> TPropertyAnimatorEditorCurveChannelInterface<InChannelType, InMenuExtensionType>::ExtendSidebarMenu_Raw(FMenuBuilder& MenuBuilder
 	, TSharedPtr<FExtender> InMenuExtender
 	, TConstArrayView<FMovieSceneChannelHandle> InChannels
-	, TConstArrayView<UMovieSceneSection*> InSections
-	, TWeakPtr<ISequencer> InSequencer) const
+	, const TArray<TWeakObjectPtr<UMovieSceneSection>>& InWeakSections
+	, TWeakPtr<ISequencer> InWeakSequencer) const
 {
-	TSharedRef<FPropertyAnimatorEditorCurveSectionMenuExtension> Extension = MakeShared<FMenuExtensionType>(InChannels, InSections);
+	const TSharedRef<FMenuExtensionType> Extension = MakeShared<FMenuExtensionType>(InChannels, InWeakSections);
 	Extension->ExtendMenu(MenuBuilder, false);
+	return Extension;
 }
 
 template<typename InChannelType, typename InMenuExtensionType>
diff --git a/Engine/Plugins/Experimental/PropertyAnimator/Source/PropertyAnimatorEditor/Private/Sequencer/PropertyAnimatorEditorCurveSectionMenuExtension.cpp b/Engine/Plugins/Experimental/PropertyAnimator/Source/PropertyAnimatorEditor/Private/Sequencer/PropertyAnimatorEditorCurveSectionMenuExtension.cpp
--- a/Engine/Plugins/Experimental/PropertyAnimator/Source/PropertyAnimatorEditor/Private/Sequencer/PropertyAnimatorEditorCurveSectionMenuExtension.cpp
+++ b/Engine/Plugins/Experimental/PropertyAnimator/Source/PropertyAnimatorEditor/Private/Sequencer/PropertyAnimatorEditorCurveSectionMenuExtension.cpp
@@ -38,14 +38,14 @@ void FPropertyAnimatorEditorCurveSectionMenuExtension::FChannelNotifyHook::Notif
 	}
 }
 
-FPropertyAnimatorEditorCurveSectionMenuExtension::FPropertyAnimatorEditorCurveSectionMenuExtension(TConstArrayView<FMovieSceneChannelHandle> InChannelHandles, TConstArrayView<UMovieSceneSection*> InSections)
+FPropertyAnimatorEditorCurveSectionMenuExtension::FPropertyAnimatorEditorCurveSectionMenuExtension(TConstArrayView<FMovieSceneChannelHandle> InChannelHandles, const TConstArrayView<TWeakObjectPtr<UMovieSceneSection>>& InWeakSections)
 	: ChannelHandles(InChannelHandles)
-	, Sections(InSections)
+	, WeakSections(InWeakSections)
 {
 	Initialize();
 }
 
-void FPropertyAnimatorEditorCurveSectionMenuExtension::ExtendMenu(FMenuBuilder& InMenuBuilder, const bool bSubMenu)
+TSharedPtr<ISidebarChannelExtension> FPropertyAnimatorEditorCurveSectionMenuExtension::ExtendMenu(FMenuBuilder& InMenuBuilder, const bool bInSubMenu)
 {
 	TSharedRef<FPropertyAnimatorEditorCurveSectionMenuExtension> This = SharedThis(this);
 
@@ -54,7 +54,7 @@ void FPropertyAnimatorEditorCurveSectionMenuExtension::ExtendMenu(FMenuBuilder&
 
 	if (ChannelHandles.Num() > 1)
 	{
-		if (bSubMenu)
+		if (bInSubMenu)
 		{
 			InMenuBuilder.AddSubMenu(MenuTitle, MenuTooltip, FNewMenuDelegate::CreateLambda([This](FMenuBuilder& InInnerMenuBuilder)
 				{
@@ -68,7 +68,7 @@ void FPropertyAnimatorEditorCurveSectionMenuExtension::ExtendMenu(FMenuBuilder&
 	}
 	else if (ChannelHandles.Num() == 1)
 	{
-		if (bSubMenu)
+		if (bInSubMenu)
 		{
 			InMenuBuilder.AddSubMenu(MenuTitle, MenuTooltip, FNewMenuDelegate::CreateLambda([This](FMenuBuilder& InInnerMenuBuilder)
 				{
@@ -80,6 +80,8 @@ void FPropertyAnimatorEditorCurveSectionMenuExtension::ExtendMenu(FMenuBuilder&
 			This->BuildParametersMenu(InMenuBuilder, 0);
 		}
 	}
+
+	return This;
 }
 
 void FPropertyAnimatorEditorCurveSectionMenuExtension::Initialize()
@@ -88,9 +90,9 @@ void FPropertyAnimatorEditorCurveSectionMenuExtension::Initialize()
 	// Also, create the notify hooks. Normal channels need to modify the section, but overriden channels
 	// need to modify their override channel container.
 	TArray<FMovieSceneChannelProxy*> ChannelProxies;
-	for (UMovieSceneSection* Section : Sections)
+	for (const TWeakObjectPtr<UMovieSceneSection>& WeakSection : WeakSections)
 	{
-		ChannelProxies.Add(&Section->GetChannelProxy());
+		ChannelProxies.Add(&WeakSection->GetChannelProxy());
 	}
 
 	for (const FMovieSceneChannelHandle& ChannelHandle : ChannelHandles)
@@ -98,11 +100,15 @@ void FPropertyAnimatorEditorCurveSectionMenuExtension::Initialize()
 		int32 SectionIndex = ChannelProxies.Find(ChannelHandle.GetChannelProxy());
 		ChannelHandleSectionIndexes.Add(SectionIndex);
 
-		UMovieSceneSection* Section = Sections[SectionIndex];
+		TWeakObjectPtr<UMovieSceneSection> WeakSection = WeakSections[SectionIndex];
+		if (!WeakSection.IsValid())
+		{
+			continue;
+		}
 
-		UObject* ObjectToModify = Section;
+		UObject* ObjectToModify = WeakSection.Get();
 
-		IMovieSceneChannelOverrideProvider* ChannelOverrideProvider = Cast<IMovieSceneChannelOverrideProvider>(Section);
+		IMovieSceneChannelOverrideProvider* ChannelOverrideProvider = Cast<IMovieSceneChannelOverrideProvider>(WeakSection);
 
 		UMovieSceneSectionChannelOverrideRegistry* ChannelOverrideRegistry = ChannelOverrideProvider
 			? ChannelOverrideProvider->GetChannelOverrideRegistry(false)
@@ -125,7 +131,7 @@ void FPropertyAnimatorEditorCurveSectionMenuExtension::Initialize()
 
 void FPropertyAnimatorEditorCurveSectionMenuExtension::BuildChannelsMenu(FMenuBuilder& InMenuBuilder)
 {
-	const bool bMultipleSections = Sections.Num() > 1;
+	const bool bMultipleSections = WeakSections.Num() > 1;
 
 	TSharedRef<FPropertyAnimatorEditorCurveSectionMenuExtension> This = SharedThis(this);
 
diff --git a/Engine/Plugins/Experimental/PropertyAnimator/Source/PropertyAnimatorEditor/Private/Sequencer/PropertyAnimatorEditorCurveSectionMenuExtension.h b/Engine/Plugins/Experimental/PropertyAnimator/Source/PropertyAnimatorEditor/Private/Sequencer/PropertyAnimatorEditorCurveSectionMenuExtension.h
--- a/Engine/Plugins/Experimental/PropertyAnimator/Source/PropertyAnimatorEditor/Private/Sequencer/PropertyAnimatorEditorCurveSectionMenuExtension.h
+++ b/Engine/Plugins/Experimental/PropertyAnimator/Source/PropertyAnimatorEditor/Private/Sequencer/PropertyAnimatorEditorCurveSectionMenuExtension.h
@@ -2,6 +2,7 @@
 
 #pragma once
 
+#include "ISequencerChannelInterface.h"
 #include "Channels/MovieSceneChannelHandle.h"
 #include "Containers/ArrayView.h"
 #include "Misc/NotifyHook.h"
@@ -11,7 +12,7 @@
 class FMenuBuilder;
 class UMovieSceneSection;
 
-class FPropertyAnimatorEditorCurveSectionMenuExtension : public TSharedFromThis<FPropertyAnimatorEditorCurveSectionMenuExtension>
+class FPropertyAnimatorEditorCurveSectionMenuExtension : public TSharedFromThis<FPropertyAnimatorEditorCurveSectionMenuExtension>, public ISidebarChannelExtension
 {
 	struct FChannelNotifyHook : FNotifyHook
 	{
@@ -30,13 +31,12 @@ class FPropertyAnimatorEditorCurveSectionMenuExtension : public TSharedFromThis<
 	};
 
 public:
-	FPropertyAnimatorEditorCurveSectionMenuExtension(TConstArrayView<FMovieSceneChannelHandle> InChannelHandles, TConstArrayView<UMovieSceneSection*> InSections);
-
+	FPropertyAnimatorEditorCurveSectionMenuExtension(TConstArrayView<FMovieSceneChannelHandle> InChannelHandles, const TConstArrayView<TWeakObjectPtr<UMovieSceneSection>>& InWeakSections);
 	virtual ~FPropertyAnimatorEditorCurveSectionMenuExtension() = default;
 
 	virtual bool GetParameterStructData(FMovieSceneChannelHandle InChannelHandle, UStruct*& OutStruct, uint8*& OutData) const = 0;
 
-	void ExtendMenu(FMenuBuilder& InMenuBuilder, const bool bSubMenu = true);
+	virtual TSharedPtr<ISidebarChannelExtension> ExtendMenu(FMenuBuilder& InMenuBuilder, const bool bInSubMenu = true) override;
 
 private:
 	void Initialize();
@@ -49,7 +49,7 @@ private:
 
 	TArray<int32> ChannelHandleSectionIndexes;
 
-	TArray<UMovieSceneSection*> Sections;
+	TArray<TWeakObjectPtr<UMovieSceneSection>> WeakSections;
 
 	TArray<FChannelNotifyHook> NotifyHooks;
 };
@@ -58,8 +58,8 @@ template<typename InChannelType>
 class TPropertyAnimatorEditorCurveSectionMenuExtension : public FPropertyAnimatorEditorCurveSectionMenuExtension
 {
 public:
-	TPropertyAnimatorEditorCurveSectionMenuExtension(const TConstArrayView<FMovieSceneChannelHandle>& InChannelHandles, const TConstArrayView<UMovieSceneSection*>& InSections)
-		: FPropertyAnimatorEditorCurveSectionMenuExtension(InChannelHandles, InSections)
+	TPropertyAnimatorEditorCurveSectionMenuExtension(const TConstArrayView<FMovieSceneChannelHandle>& InChannelHandles, const TConstArrayView<TWeakObjectPtr<UMovieSceneSection>>& InWeakSections)
+		: FPropertyAnimatorEditorCurveSectionMenuExtension(InChannelHandles, InWeakSections)
 	{
 	}
 
diff --git a/Engine/Source/Editor/MovieSceneTools/Private/Channels/BuiltInChannelEditors.cpp b/Engine/Source/Editor/MovieSceneTools/Private/Channels/BuiltInChannelEditors.cpp
--- a/Engine/Source/Editor/MovieSceneTools/Private/Channels/BuiltInChannelEditors.cpp
+++ b/Engine/Source/Editor/MovieSceneTools/Private/Channels/BuiltInChannelEditors.cpp
@@ -1291,7 +1291,7 @@ private:
 	TArray<TExtendKeyMenuParams<ChannelType>> ChannelAndHandles;
 };
 
-struct FCurveChannelSectionMenuExtension : TSharedFromThis<FCurveChannelSectionMenuExtension>
+struct FCurveChannelSectionMenuExtension : TSharedFromThis<FCurveChannelSectionMenuExtension>, ISidebarChannelExtension
 {
 	static TSharedRef<FCurveChannelSectionMenuExtension> GetOrCreate(TWeakPtr<ISequencer> InSequencer)
 	{
@@ -1315,24 +1315,23 @@ struct FCurveChannelSectionMenuExtension : TSharedFromThis<FCurveChannelSectionM
 		, bMenusAdded(false)
 	{
 	}
+	virtual ~FCurveChannelSectionMenuExtension() {}
 
-	void AddSections(TArrayView<UMovieSceneSection* const> InSections)
+	void AddSections(const TArray<TWeakObjectPtr<UMovieSceneSection>>& InWeakSections)
 	{
-		for (UMovieSceneSection* Section : InSections)
-		{
-			Sections.Add(Section);
-		}
+		WeakSections = TSet(InWeakSections);
+
 		++NumCurveChannelTypes;
 	}
 
-	void ExtendMenu(FMenuBuilder& MenuBuilder)
+	virtual TSharedPtr<ISidebarChannelExtension> ExtendMenu(FMenuBuilder& MenuBuilder, const bool bInSubMenu = false) override
 	{
 		--NumCurveChannelTypes;
 
 		if (bMenusAdded)
 		{
 			// Only add menus once -- not once per curve channel type (float, double, etc)
-			return;
+			return nullptr;
 		}
 
 		bMenusAdded = true;
@@ -1340,7 +1339,7 @@ struct FCurveChannelSectionMenuExtension : TSharedFromThis<FCurveChannelSectionM
 		ISequencer* SequencerPtr = WeakSequencer.Pin().Get();
 		if (!SequencerPtr)
 		{
-			return;
+			return nullptr;
 		}
 
 		TSharedRef<FCurveChannelSectionMenuExtension> SharedThis = this->AsShared();
@@ -1368,6 +1367,8 @@ struct FCurveChannelSectionMenuExtension : TSharedFromThis<FCurveChannelSectionM
 						FNewMenuDelegate::CreateLambda([SharedThis](FMenuBuilder& SubMenuBuilder){ SharedThis->AddDisplayOptionsMenu(SubMenuBuilder); })
 						);
 			}));
+
+		return SharedThis;
 	}
 
 	void AddDisplayOptionsMenu(FMenuBuilder& MenuBuilder)
@@ -1620,7 +1621,7 @@ struct FCurveChannelSectionMenuExtension : TSharedFromThis<FCurveChannelSectionM
 		// Otherwise, the channels of all the sections
 		if (FloatChannels.Num() + DoubleChannels.Num() + IntegerChannels.Num() + BoolChannels.Num() == 0)
 		{
-			for (TWeakObjectPtr<UMovieSceneSection> WeakSection : Sections)
+			for (TWeakObjectPtr<UMovieSceneSection> WeakSection : WeakSections)
 			{
 				if (UMovieSceneSection* Section = WeakSection.Get())
 				{
@@ -1665,7 +1666,7 @@ struct FCurveChannelSectionMenuExtension : TSharedFromThis<FCurveChannelSectionM
 		bool bAnythingChanged = false;
 
 		// Modify all sections
-		for (TWeakObjectPtr<UMovieSceneSection> WeakSection : Sections)
+		for (TWeakObjectPtr<UMovieSceneSection> WeakSection : WeakSections)
 		{
 			if (UMovieSceneSection* Section = WeakSection.Get())
 			{
@@ -1768,7 +1769,7 @@ struct FCurveChannelSectionMenuExtension : TSharedFromThis<FCurveChannelSectionM
 
 		// Modify all sections
 
-		for (TWeakObjectPtr<UMovieSceneSection> WeakSection : Sections)
+		for (TWeakObjectPtr<UMovieSceneSection> WeakSection : WeakSections)
 		{
 			if (UMovieSceneSection* Section = WeakSection.Get())
 			{
@@ -1777,7 +1778,7 @@ struct FCurveChannelSectionMenuExtension : TSharedFromThis<FCurveChannelSectionM
 		}
 
 		// Apply to all channels
-		for (TWeakObjectPtr<UMovieSceneSection> WeakSection : Sections)
+		for (TWeakObjectPtr<UMovieSceneSection> WeakSection : WeakSections)
 		{
 			if (UMovieSceneSection* Section = WeakSection.Get())
 			{
@@ -1810,7 +1811,7 @@ struct FCurveChannelSectionMenuExtension : TSharedFromThis<FCurveChannelSectionM
 	ECheckBoxState IsShowCurve() const
 	{
 		int32 NumShowedAndHidden[2] = { 0, 0 };
-		for (TWeakObjectPtr<UMovieSceneSection> WeakSection : Sections)
+		for (TWeakObjectPtr<UMovieSceneSection> WeakSection : WeakSections)
 		{
 			if (UMovieSceneSection* Section = WeakSection.Get())
 			{
@@ -1845,7 +1846,7 @@ struct FCurveChannelSectionMenuExtension : TSharedFromThis<FCurveChannelSectionM
 
 	bool IsAnyShowCurve() const
 	{
-		for (TWeakObjectPtr<UMovieSceneSection> WeakSection : Sections)
+		for (TWeakObjectPtr<UMovieSceneSection> WeakSection : WeakSections)
 		{
 			if (UMovieSceneSection* Section = WeakSection.Get())
 			{
@@ -1878,75 +1879,83 @@ private:
 	static TWeakPtr<FCurveChannelSectionMenuExtension> WeakCurrentExtension;
 
 	TWeakPtr<ISequencer> WeakSequencer;
-	TSet<TWeakObjectPtr<UMovieSceneSection>> Sections;
+	TSet<TWeakObjectPtr<UMovieSceneSection>> WeakSections;
 	int32 NumCurveChannelTypes;
 	bool bMenusAdded;
 };
 
 TWeakPtr<FCurveChannelSectionMenuExtension> FCurveChannelSectionMenuExtension::WeakCurrentExtension;
 
-void ExtendSectionMenu(FMenuBuilder& OuterMenuBuilder, TSharedPtr<FExtender> MenuExtender, TArray<TMovieSceneChannelHandle<FMovieSceneFloatChannel>>&& Channels, TArrayView<UMovieSceneSection* const> Sections, TWeakPtr<ISequencer> InSequencer)
+void ExtendSectionMenu(FMenuBuilder& OuterMenuBuilder, TSharedPtr<FExtender> MenuExtender, TArray<TMovieSceneChannelHandle<FMovieSceneFloatChannel>>&& Channels, const TArray<TWeakObjectPtr<UMovieSceneSection>>& InWeakSections, TWeakPtr<ISequencer> InWeakSequencer)
 {
-	TSharedRef<FCurveChannelSectionMenuExtension> Extension = FCurveChannelSectionMenuExtension::GetOrCreate(InSequencer);
-	Extension->AddSections(Sections);
+	TSharedRef<FCurveChannelSectionMenuExtension> Extension = FCurveChannelSectionMenuExtension::GetOrCreate(InWeakSequencer);
+	Extension->AddSections(InWeakSections);
 
 	MenuExtender->AddMenuExtension("SequencerChannels", EExtensionHook::First, nullptr, FMenuExtensionDelegate::CreateLambda([Extension](FMenuBuilder& MenuBuilder) { Extension->ExtendMenu(MenuBuilder); }));
 }
 
-void ExtendSectionMenu(FMenuBuilder& OuterMenuBuilder, TSharedPtr<FExtender> MenuExtender, TArray<TMovieSceneChannelHandle<FMovieSceneDoubleChannel>>&& Channels, TArrayView<UMovieSceneSection* const> Sections, TWeakPtr<ISequencer> InSequencer)
+void ExtendSectionMenu(FMenuBuilder& OuterMenuBuilder, TSharedPtr<FExtender> MenuExtender, TArray<TMovieSceneChannelHandle<FMovieSceneDoubleChannel>>&& Channels, const TArray<TWeakObjectPtr<UMovieSceneSection>>& InWeakSections, TWeakPtr<ISequencer> InWeakSequencer)
 {
-	TSharedRef<FCurveChannelSectionMenuExtension> Extension = FCurveChannelSectionMenuExtension::GetOrCreate(InSequencer);
-	Extension->AddSections(Sections);
+	TSharedRef<FCurveChannelSectionMenuExtension> Extension = FCurveChannelSectionMenuExtension::GetOrCreate(InWeakSequencer);
+	Extension->AddSections(InWeakSections);
 
 	MenuExtender->AddMenuExtension("SequencerChannels", EExtensionHook::First, nullptr, FMenuExtensionDelegate::CreateLambda([Extension](FMenuBuilder& MenuBuilder) { Extension->ExtendMenu(MenuBuilder); }));
 }
 
-void ExtendSectionMenu(FMenuBuilder& OuterMenuBuilder, TSharedPtr<FExtender> MenuExtender, TArray<TMovieSceneChannelHandle<FMovieSceneIntegerChannel>>&& Channels, TArrayView<UMovieSceneSection* const> Sections, TWeakPtr<ISequencer> InSequencer)
+void ExtendSectionMenu(FMenuBuilder& OuterMenuBuilder, TSharedPtr<FExtender> MenuExtender, TArray<TMovieSceneChannelHandle<FMovieSceneIntegerChannel>>&& Channels, const TArray<TWeakObjectPtr<UMovieSceneSection>>& InWeakSections, TWeakPtr<ISequencer> InWeakSequencer)
 {
-	TSharedRef<FCurveChannelSectionMenuExtension> Extension = FCurveChannelSectionMenuExtension::GetOrCreate(InSequencer);
-	Extension->AddSections(Sections);
+	TSharedRef<FCurveChannelSectionMenuExtension> Extension = FCurveChannelSectionMenuExtension::GetOrCreate(InWeakSequencer);
+	Extension->AddSections(InWeakSections);
 
 	MenuExtender->AddMenuExtension("SequencerChannels", EExtensionHook::First, nullptr, FMenuExtensionDelegate::CreateLambda([Extension](FMenuBuilder& MenuBuilder) { Extension->ExtendMenu(MenuBuilder); }));
 }
 
-void ExtendSectionMenu(FMenuBuilder& OuterMenuBuilder, TSharedPtr<FExtender> MenuExtender, TArray<TMovieSceneChannelHandle<FMovieSceneBoolChannel>>&& Channels, TArrayView<UMovieSceneSection* const> Sections, TWeakPtr<ISequencer> InSequencer)
+void ExtendSectionMenu(FMenuBuilder& OuterMenuBuilder, TSharedPtr<FExtender> MenuExtender, TArray<TMovieSceneChannelHandle<FMovieSceneBoolChannel>>&& Channels, const TArray<TWeakObjectPtr<UMovieSceneSection>>& InWeakSections, TWeakPtr<ISequencer> InWeakSequencer)
 {
-	TSharedRef<FCurveChannelSectionMenuExtension> Extension = FCurveChannelSectionMenuExtension::GetOrCreate(InSequencer);
-	Extension->AddSections(Sections);
+	TSharedRef<FCurveChannelSectionMenuExtension> Extension = FCurveChannelSectionMenuExtension::GetOrCreate(InWeakSequencer);
+	Extension->AddSections(InWeakSections);
 
 	MenuExtender->AddMenuExtension("SequencerChannels", EExtensionHook::First, nullptr, FMenuExtensionDelegate::CreateLambda([Extension](FMenuBuilder& MenuBuilder) { Extension->ExtendMenu(MenuBuilder); }));
 }
 
-void ExtendSidebarMenu(FMenuBuilder& OuterMenuBuilder, TSharedPtr<FExtender> MenuExtender, TArray<TMovieSceneChannelHandle<FMovieSceneFloatChannel>>&& Channels, TArrayView<UMovieSceneSection* const> Sections, TWeakPtr<ISequencer> InSequencer)
+TSharedPtr<ISidebarChannelExtension> ExtendSidebarMenu(FMenuBuilder& OuterMenuBuilder, TSharedPtr<FExtender> InMenuExtender, TArray<TMovieSceneChannelHandle<FMovieSceneFloatChannel>>&& Channels, const TArray<TWeakObjectPtr<UMovieSceneSection>>& InWeakSections, TWeakPtr<ISequencer> InWeakSequencer)
 {
-	TSharedRef<FCurveChannelSectionMenuExtension> Extension = FCurveChannelSectionMenuExtension::GetOrCreate(InSequencer);
-	Extension->AddSections(Sections);
+	TSharedRef<FCurveChannelSectionMenuExtension> Extension = FCurveChannelSectionMenuExtension::GetOrCreate(InWeakSequencer);
+	Extension->AddSections(InWeakSections);
 
-	MenuExtender->AddMenuExtension("SequencerChannels", EExtensionHook::First, nullptr, FMenuExtensionDelegate::CreateLambda([Extension](FMenuBuilder& MenuBuilder) { Extension->ExtendMenu(MenuBuilder); }));
+	InMenuExtender->AddMenuExtension("SequencerChannels", EExtensionHook::First, nullptr, FMenuExtensionDelegate::CreateLambda([Extension](FMenuBuilder& MenuBuilder) { Extension->ExtendMenu(MenuBuilder); }));
+
+	return Extension;
 }
 
-void ExtendSidebarMenu(FMenuBuilder& OuterMenuBuilder, TSharedPtr<FExtender> MenuExtender, TArray<TMovieSceneChannelHandle<FMovieSceneDoubleChannel>>&& Channels, TArrayView<UMovieSceneSection* const> Sections, TWeakPtr<ISequencer> InSequencer)
+TSharedPtr<ISidebarChannelExtension> ExtendSidebarMenu(FMenuBuilder& OuterMenuBuilder, TSharedPtr<FExtender> InMenuExtender, TArray<TMovieSceneChannelHandle<FMovieSceneDoubleChannel>>&& Channels, const TArray<TWeakObjectPtr<UMovieSceneSection>>& InWeakSections, TWeakPtr<ISequencer> InWeakSequencer)
 {
-	TSharedRef<FCurveChannelSectionMenuExtension> Extension = FCurveChannelSectionMenuExtension::GetOrCreate(InSequencer);
-	Extension->AddSections(Sections);
+	TSharedRef<FCurveChannelSectionMenuExtension> Extension = FCurveChannelSectionMenuExtension::GetOrCreate(InWeakSequencer);
+	Extension->AddSections(InWeakSections);
 
-	MenuExtender->AddMenuExtension("SequencerChannels", EExtensionHook::First, nullptr, FMenuExtensionDelegate::CreateLambda([Extension](FMenuBuilder& MenuBuilder) { Extension->ExtendMenu(MenuBuilder); }));
+	InMenuExtender->AddMenuExtension("SequencerChannels", EExtensionHook::First, nullptr, FMenuExtensionDelegate::CreateLambda([Extension](FMenuBuilder& MenuBuilder) { Extension->ExtendMenu(MenuBuilder); }));
+
+	return Extension;
 }
 
-void ExtendSidebarMenu(FMenuBuilder& OuterMenuBuilder, TSharedPtr<FExtender> MenuExtender, TArray<TMovieSceneChannelHandle<FMovieSceneIntegerChannel>>&& Channels, TArrayView<UMovieSceneSection* const> Sections, TWeakPtr<ISequencer> InSequencer)
+TSharedPtr<ISidebarChannelExtension> ExtendSidebarMenu(FMenuBuilder& OuterMenuBuilder, TSharedPtr<FExtender> InMenuExtender, TArray<TMovieSceneChannelHandle<FMovieSceneIntegerChannel>>&& Channels, const TArray<TWeakObjectPtr<UMovieSceneSection>>& InWeakSections, TWeakPtr<ISequencer> InWeakSequencer)
 {
-	TSharedRef<FCurveChannelSectionMenuExtension> Extension = FCurveChannelSectionMenuExtension::GetOrCreate(InSequencer);
-	Extension->AddSections(Sections);
+	TSharedRef<FCurveChannelSectionMenuExtension> Extension = FCurveChannelSectionMenuExtension::GetOrCreate(InWeakSequencer);
+	Extension->AddSections(InWeakSections);
 
-	MenuExtender->AddMenuExtension("SequencerChannels", EExtensionHook::First, nullptr, FMenuExtensionDelegate::CreateLambda([Extension](FMenuBuilder& MenuBuilder) { Extension->ExtendMenu(MenuBuilder); }));
+	InMenuExtender->AddMenuExtension("SequencerChannels", EExtensionHook::First, nullptr, FMenuExtensionDelegate::CreateLambda([Extension](FMenuBuilder& MenuBuilder) { Extension->ExtendMenu(MenuBuilder); }));
+
+	return Extension;
 }
 
-void ExtendSidebarMenu(FMenuBuilder& OuterMenuBuilder, TSharedPtr<FExtender> MenuExtender, TArray<TMovieSceneChannelHandle<FMovieSceneBoolChannel>>&& Channels, TArrayView<UMovieSceneSection* const> Sections, TWeakPtr<ISequencer> InSequencer)
+TSharedPtr<ISidebarChannelExtension> ExtendSidebarMenu(FMenuBuilder& OuterMenuBuilder, TSharedPtr<FExtender> InMenuExtender, TArray<TMovieSceneChannelHandle<FMovieSceneBoolChannel>>&& Channels, const TArray<TWeakObjectPtr<UMovieSceneSection>>& InWeakSections, TWeakPtr<ISequencer> InWeakSequencer)
 {
-	TSharedRef<FCurveChannelSectionMenuExtension> Extension = FCurveChannelSectionMenuExtension::GetOrCreate(InSequencer);
-	Extension->AddSections(Sections);
+	TSharedRef<FCurveChannelSectionMenuExtension> Extension = FCurveChannelSectionMenuExtension::GetOrCreate(InWeakSequencer);
+	Extension->AddSections(InWeakSections);
 
-	MenuExtender->AddMenuExtension("SequencerChannels", EExtensionHook::First, nullptr, FMenuExtensionDelegate::CreateLambda([Extension](FMenuBuilder& MenuBuilder) { Extension->ExtendMenu(MenuBuilder); }));
+	InMenuExtender->AddMenuExtension("SequencerChannels", EExtensionHook::First, nullptr, FMenuExtensionDelegate::CreateLambda([Extension](FMenuBuilder& MenuBuilder) { Extension->ExtendMenu(MenuBuilder); }));
+
+	return Extension;
 }
 
 void ExtendKeyMenu(FMenuBuilder& OuterMenuBuilder, TSharedPtr<FExtender> MenuExtender, TArray<TExtendKeyMenuParams<FMovieSceneFloatChannel>>&& Channels, TWeakPtr<ISequencer> InSequencer)
diff --git a/Engine/Source/Editor/MovieSceneTools/Private/Channels/BuiltInChannelEditors.h b/Engine/Source/Editor/MovieSceneTools/Private/Channels/BuiltInChannelEditors.h
--- a/Engine/Source/Editor/MovieSceneTools/Private/Channels/BuiltInChannelEditors.h
+++ b/Engine/Source/Editor/MovieSceneTools/Private/Channels/BuiltInChannelEditors.h
@@ -75,15 +75,15 @@ void DrawKeys(FMovieSceneParticleChannel* Channel, TArrayView<const FKeyHandle>
 void DrawKeys(FMovieSceneEventChannel*    Channel, TArrayView<const FKeyHandle> InKeyHandles, const UMovieSceneSection* InOwner, TArrayView<FKeyDrawParams> OutKeyDrawParams);
 
 /** Context menu overrides */
-void ExtendSectionMenu(FMenuBuilder& OuterMenuBuilder, TSharedPtr<FExtender> MenuExtender, TArray<TMovieSceneChannelHandle<FMovieSceneDoubleChannel>>&& Channels, TArrayView<UMovieSceneSection* const> Sections, TWeakPtr<ISequencer> InSequencer);
-void ExtendSectionMenu(FMenuBuilder& OuterMenuBuilder, TSharedPtr<FExtender> MenuExtender, TArray<TMovieSceneChannelHandle<FMovieSceneFloatChannel>>&& Channels, TArrayView<UMovieSceneSection* const> Sections, TWeakPtr<ISequencer> InSequencer);
-void ExtendSectionMenu(FMenuBuilder& OuterMenuBuilder, TSharedPtr<FExtender> MenuExtender, TArray<TMovieSceneChannelHandle<FMovieSceneIntegerChannel>>&& Channels, TArrayView<UMovieSceneSection* const> Sections, TWeakPtr<ISequencer> InSequencer);
-void ExtendSectionMenu(FMenuBuilder& OuterMenuBuilder, TSharedPtr<FExtender> MenuExtender, TArray<TMovieSceneChannelHandle<FMovieSceneBoolChannel>>&& Channels, TArrayView<UMovieSceneSection* const> Sections, TWeakPtr<ISequencer> InSequencer);
-
-void ExtendSidebarMenu(FMenuBuilder& OuterMenuBuilder, TSharedPtr<FExtender> MenuExtender, TArray<TMovieSceneChannelHandle<FMovieSceneDoubleChannel>>&& Channels, TArrayView<UMovieSceneSection* const> Sections, TWeakPtr<ISequencer> InSequencer);
-void ExtendSidebarMenu(FMenuBuilder& OuterMenuBuilder, TSharedPtr<FExtender> MenuExtender, TArray<TMovieSceneChannelHandle<FMovieSceneFloatChannel>>&& Channels, TArrayView<UMovieSceneSection* const> Sections, TWeakPtr<ISequencer> InSequencer);
-void ExtendSidebarMenu(FMenuBuilder& OuterMenuBuilder, TSharedPtr<FExtender> MenuExtender, TArray<TMovieSceneChannelHandle<FMovieSceneIntegerChannel>>&& Channels, TArrayView<UMovieSceneSection* const> Sections, TWeakPtr<ISequencer> InSequencer);
-void ExtendSidebarMenu(FMenuBuilder& OuterMenuBuilder, TSharedPtr<FExtender> MenuExtender, TArray<TMovieSceneChannelHandle<FMovieSceneBoolChannel>>&& Channels, TArrayView<UMovieSceneSection* const> Sections, TWeakPtr<ISequencer> InSequencer);
+void ExtendSectionMenu(FMenuBuilder& OuterMenuBuilder, TSharedPtr<FExtender> MenuExtender, TArray<TMovieSceneChannelHandle<FMovieSceneDoubleChannel>>&& Channels, const TArray<TWeakObjectPtr<UMovieSceneSection>>& InWeakSections, TWeakPtr<ISequencer> InWeakSequencer);
+void ExtendSectionMenu(FMenuBuilder& OuterMenuBuilder, TSharedPtr<FExtender> MenuExtender, TArray<TMovieSceneChannelHandle<FMovieSceneFloatChannel>>&& Channels, const TArray<TWeakObjectPtr<UMovieSceneSection>>& InWeakSections, TWeakPtr<ISequencer> InWeakSequencer);
+void ExtendSectionMenu(FMenuBuilder& OuterMenuBuilder, TSharedPtr<FExtender> MenuExtender, TArray<TMovieSceneChannelHandle<FMovieSceneIntegerChannel>>&& Channels, const TArray<TWeakObjectPtr<UMovieSceneSection>>& InWeakSections, TWeakPtr<ISequencer> InWeakSequencer);
+void ExtendSectionMenu(FMenuBuilder& OuterMenuBuilder, TSharedPtr<FExtender> MenuExtender, TArray<TMovieSceneChannelHandle<FMovieSceneBoolChannel>>&& Channels, const TArray<TWeakObjectPtr<UMovieSceneSection>>& InWeakSections, TWeakPtr<ISequencer> InWeakSequencer);
+
+TSharedPtr<ISidebarChannelExtension> ExtendSidebarMenu(FMenuBuilder& OuterMenuBuilder, TSharedPtr<FExtender> MenuExtender, TArray<TMovieSceneChannelHandle<FMovieSceneDoubleChannel>>&& InChannels, const TArray<TWeakObjectPtr<UMovieSceneSection>>& InWeakSections, TWeakPtr<ISequencer> InWeakSequencer);
+TSharedPtr<ISidebarChannelExtension> ExtendSidebarMenu(FMenuBuilder& OuterMenuBuilder, TSharedPtr<FExtender> MenuExtender, TArray<TMovieSceneChannelHandle<FMovieSceneFloatChannel>>&& InChannels,const  TArray<TWeakObjectPtr<UMovieSceneSection>>& InWeakSections, TWeakPtr<ISequencer> InWeakSequencer);
+TSharedPtr<ISidebarChannelExtension> ExtendSidebarMenu(FMenuBuilder& OuterMenuBuilder, TSharedPtr<FExtender> MenuExtender, TArray<TMovieSceneChannelHandle<FMovieSceneIntegerChannel>>&& InChannels, const TArray<TWeakObjectPtr<UMovieSceneSection>>& InWeakSections, TWeakPtr<ISequencer> InWeakSequencer);
+TSharedPtr<ISidebarChannelExtension> ExtendSidebarMenu(FMenuBuilder& OuterMenuBuilder, TSharedPtr<FExtender> MenuExtender, TArray<TMovieSceneChannelHandle<FMovieSceneBoolChannel>>&& InChannels, const TArray<TWeakObjectPtr<UMovieSceneSection>>& InWeakSections, TWeakPtr<ISequencer> InWeakSequencer);
 
 void ExtendKeyMenu(FMenuBuilder& OuterMenuBuilder, TSharedPtr<FExtender> MenuExtender, TArray<TExtendKeyMenuParams<FMovieSceneDoubleChannel>>&& Channels, TWeakPtr<ISequencer> InSequencer);
 void ExtendKeyMenu(FMenuBuilder& OuterMenuBuilder, TSharedPtr<FExtender> MenuExtender, TArray<TExtendKeyMenuParams<FMovieSceneFloatChannel>>&& Channels, TWeakPtr<ISequencer> InSequencer);
diff --git a/Engine/Source/Editor/MovieSceneTools/Private/Channels/PerlinNoiseChannelInterface.cpp b/Engine/Source/Editor/MovieSceneTools/Private/Channels/PerlinNoiseChannelInterface.cpp
--- a/Engine/Source/Editor/MovieSceneTools/Private/Channels/PerlinNoiseChannelInterface.cpp
+++ b/Engine/Source/Editor/MovieSceneTools/Private/Channels/PerlinNoiseChannelInterface.cpp
@@ -15,47 +15,53 @@
 
 #define LOCTEXT_NAMESPACE "PerlinNoiseChannelInterface"
 
-FPerlinNoiseChannelSectionMenuExtension::FPerlinNoiseChannelSectionMenuExtension(TArrayView<const FMovieSceneChannelHandle> InChannelHandles, TArrayView<UMovieSceneSection* const> InSections)
+FPerlinNoiseChannelSectionMenuExtension::FPerlinNoiseChannelSectionMenuExtension(TArrayView<const FMovieSceneChannelHandle> InChannelHandles, const TArray<TWeakObjectPtr<UMovieSceneSection>>& InWeakSections)
 	: ChannelHandles(InChannelHandles)
-	, Sections(InSections)
+	, WeakSections(InWeakSections)
 {
 	Initialize();
 }
 
-void FPerlinNoiseChannelSectionMenuExtension::ExtendMenu(FMenuBuilder& MenuBuilder, const bool bSubMenu)
+TSharedPtr<ISidebarChannelExtension> FPerlinNoiseChannelSectionMenuExtension::ExtendMenu(FMenuBuilder& MenuBuilder, const bool bInSubMenu)
 {
-	TSharedRef<FPerlinNoiseChannelSectionMenuExtension> SharedThis = this->AsShared();
-
 	if (ChannelHandles.Num() > 1)
 	{
-		if (bSubMenu)
+		if (bInSubMenu)
 		{
 			MenuBuilder.AddSubMenu(
 				LOCTEXT("PerlinNoiseChannelsMenu", "Perlin Noise Channels"),
 				LOCTEXT("PerlinNoiseChannelsMenuToolTip", "Edit parameters for Perlin Noise channels"),
-				FNewMenuDelegate::CreateLambda([SharedThis](FMenuBuilder& InnerMenuBuilder) { SharedThis->BuildChannelsMenu(InnerMenuBuilder); })
+				FNewMenuDelegate::CreateLambda([this](FMenuBuilder& InnerMenuBuilder)
+					{
+						BuildChannelsMenu(InnerMenuBuilder);
+					})
 			);
 		}
 		else
 		{
-			SharedThis->BuildChannelsMenu(MenuBuilder);
+			BuildChannelsMenu(MenuBuilder);
 		}
 	}
 	else if (ChannelHandles.Num() == 1)
 	{
-		if (bSubMenu)
+		if (bInSubMenu)
 		{
 			MenuBuilder.AddSubMenu(
 				LOCTEXT("PerlinNoiseChannelsMenu", "Perlin Noise Channels"),
 				LOCTEXT("PerlinNoiseChannelsMenuToolTip", "Edit parameters for Perlin Noise channels"),
-				FNewMenuDelegate::CreateLambda([SharedThis](FMenuBuilder& InnerMenuBuilder) { SharedThis->BuildParametersMenu(InnerMenuBuilder, 0); })
+				FNewMenuDelegate::CreateLambda([this](FMenuBuilder& InnerMenuBuilder)
+					{
+						BuildParametersMenu(InnerMenuBuilder, 0);
+					})
 			);
 		}
 		else
 		{
-			SharedThis->BuildParametersMenu(MenuBuilder, 0);
+			BuildParametersMenu(MenuBuilder, 0);
 		}
 	}
+
+	return AsShared();
 }
 
 void FPerlinNoiseChannelSectionMenuExtension::Initialize()
@@ -64,7 +70,7 @@ void FPerlinNoiseChannelSectionMenuExtension::Initialize()
 	// Also, create the notify hooks. Normal channels need to modify the section, but overriden channels
 	// need to modify their override channel container.
 	TArray<FMovieSceneChannelProxy*> ChannelProxies;
-	for (UMovieSceneSection* Section : Sections)
+	for (const TWeakObjectPtr<UMovieSceneSection>& Section : WeakSections)
 	{
 		ChannelProxies.Add(&Section->GetChannelProxy());
 	}
@@ -74,11 +80,15 @@ void FPerlinNoiseChannelSectionMenuExtension::Initialize()
 		int32 SectionIndex = ChannelProxies.Find(ChannelHandle.GetChannelProxy());
 		ChannelHandleSectionIndexes.Add(SectionIndex);
 
-		UMovieSceneSection* Section = Sections[SectionIndex];
+		TWeakObjectPtr<UMovieSceneSection> WeakSection = WeakSections[SectionIndex];
+		if (!WeakSection.IsValid())
+		{
+			continue;
+		}
 
-		UObject* ObjectToModify = Section;
+		UObject* ObjectToModify = WeakSection.Get();
 
-		IMovieSceneChannelOverrideProvider* ChannelOverrideProvider = Cast<IMovieSceneChannelOverrideProvider>(Section);
+		IMovieSceneChannelOverrideProvider* ChannelOverrideProvider = Cast<IMovieSceneChannelOverrideProvider>(WeakSection);
 		UMovieSceneSectionChannelOverrideRegistry* ChannelOverrideRegistry = ChannelOverrideProvider ?
 			ChannelOverrideProvider->GetChannelOverrideRegistry(false) : nullptr;
 		if (ChannelOverrideRegistry)
@@ -99,7 +109,7 @@ void FPerlinNoiseChannelSectionMenuExtension::Initialize()
 
 void FPerlinNoiseChannelSectionMenuExtension::BuildChannelsMenu(FMenuBuilder& MenuBuilder)
 {
-	const bool bMultipleSections = Sections.Num() > 1;
+	const bool bMultipleSections = WeakSections.Num() > 1;
 	TSharedRef<FPerlinNoiseChannelSectionMenuExtension> SharedThis = this->AsShared();
 
 	for (int32 Index = 0; Index < ChannelHandles.Num(); ++Index)
@@ -112,7 +122,10 @@ void FPerlinNoiseChannelSectionMenuExtension::BuildChannelsMenu(FMenuBuilder& Me
 			MenuBuilder.AddSubMenu(
 				FText::Format(LOCTEXT("PerlinNoiseChannelAndSectionSelectMenu", "Section{0}.{1}"), SectionIndex + 1, FText::FromName(ChannelHandle.GetMetaData()->Name)),
 				LOCTEXT("PerlinNoiseChannelAndSectionSelectMenuToolTip", "Edit parameters for this Perlin Noise channel"),
-				FNewMenuDelegate::CreateLambda([SharedThis, Index](FMenuBuilder& InnerMenuBuilder) { SharedThis->BuildParametersMenu(InnerMenuBuilder, Index); })
+				FNewMenuDelegate::CreateLambda([this, Index](FMenuBuilder& InnerMenuBuilder)
+					{
+						BuildParametersMenu(InnerMenuBuilder, Index);
+					})
 			);
 		}
 		else
@@ -120,7 +133,10 @@ void FPerlinNoiseChannelSectionMenuExtension::BuildChannelsMenu(FMenuBuilder& Me
 			MenuBuilder.AddSubMenu(
 				FText::FromName(ChannelHandle.GetMetaData()->Name),
 				LOCTEXT("PerlinNoiseChannelSelectMenuToolTip", "Edit parameters for this Perlin Noise channel"),
-				FNewMenuDelegate::CreateLambda([SharedThis, Index](FMenuBuilder& InnerMenuBuilder) { SharedThis->BuildParametersMenu(InnerMenuBuilder, Index); })
+				FNewMenuDelegate::CreateLambda([this, Index](FMenuBuilder& InnerMenuBuilder)
+					{
+						BuildParametersMenu(InnerMenuBuilder, Index);
+					})
 			);
 		}
 	}
@@ -180,7 +196,7 @@ void FPerlinNoiseChannelSectionMenuExtension::FChannelNotifyHook::NotifyPreChang
 {
 	GEditor->BeginTransaction(FText::Format(LOCTEXT("EditProperty", "Edit {0}"), PropertyAboutToChange->GetDisplayNameText()));
 
-	ObjectToModify->Modify();
+	WeakObjectToModify->Modify();
 }
 
 void FPerlinNoiseChannelSectionMenuExtension::FChannelNotifyHook::NotifyPostChange(const FPropertyChangedEvent& PropertyChangedEvent, FProperty* PropertyThatChanged)
diff --git a/Engine/Source/Editor/MovieSceneTools/Private/TrackEditors/SkeletalAnimationTrackEditor.cpp b/Engine/Source/Editor/MovieSceneTools/Private/TrackEditors/SkeletalAnimationTrackEditor.cpp
--- a/Engine/Source/Editor/MovieSceneTools/Private/TrackEditors/SkeletalAnimationTrackEditor.cpp
+++ b/Engine/Source/Editor/MovieSceneTools/Private/TrackEditors/SkeletalAnimationTrackEditor.cpp
@@ -881,71 +881,115 @@ void FSkeletalAnimationSection::CustomizePropertiesDetailsView(TSharedRef<IDetai
 
 void FSkeletalAnimationSection::BuildSectionContextMenu(FMenuBuilder& MenuBuilder, const FGuid& ObjectBinding)
 {
+	const TSharedPtr<ISequencer> SequencerPtr = Sequencer.Pin();
+	if (!SequencerPtr.IsValid())
+	{
+		return;
+	}
+
+	UMovieSceneSkeletalAnimationTrack* const Track = Section.GetTypedOuter<UMovieSceneSkeletalAnimationTrack>();
+	if (!IsValid(Track))
+	{
+		return;
+	}
+
 	// Can't pick the object that this track binds
-	TSharedPtr<ISequencer> SequencerPtr = Sequencer.Pin();
-	USkeleton* Skeleton = AcquireSkeletonFromObjectGuid(ObjectBinding, SequencerPtr);
-	USkeletalMeshComponent* SkelMeshComp = AcquireSkeletalMeshFromObjectGuid(ObjectBinding, SequencerPtr);
-	UMovieSceneSkeletalAnimationTrack* Track = Section.GetTypedOuter<UMovieSceneSkeletalAnimationTrack>();
+	USkeleton* const Skeleton = AcquireSkeletonFromObjectGuid(ObjectBinding, SequencerPtr);
+	if (!IsValid(Skeleton))
+	{
+		return;
+	}
 
-	if (Track && Skeleton)
+	const int32 NumBones = Skeleton->GetReferenceSkeleton().GetNum();
+	TArray<FName> BoneNames;
+	for (int32 BoneIndex = 0; BoneIndex < NumBones; ++BoneIndex)
 	{
-		const int32 NumBones = Skeleton->GetReferenceSkeleton().GetNum();
-		TArray<FName> BoneNames;
-		for (int32 BoneIndex = 0; BoneIndex < NumBones; ++BoneIndex)
-		{
-			BoneNames.Add(Skeleton->GetReferenceSkeleton().GetBoneName(BoneIndex));
-		}
+		BoneNames.Add(Skeleton->GetReferenceSkeleton().GetBoneName(BoneIndex));
+	}
 
-		auto MatchToBone = [=, this](bool bMatchPrevious, int32 Index)
-		{
-			return FUIAction(
-				FExecuteAction::CreateLambda([=, this]
+	auto MatchToBone = [this, ObjectBinding](bool bMatchPrevious, int32 Index)
+	{
+		return FUIAction(
+			FExecuteAction::CreateLambda([this, ObjectBinding, bMatchPrevious, Index]
+				{
+					const TSharedPtr<ISequencer> SequencerPtr = Sequencer.Pin();
+					if (!SequencerPtr.IsValid())
 					{
-						FScopedTransaction MatchSection(LOCTEXT("MatchSectionByBone_Transaction", "Match Section By Bone"));
-						Section.Modify();	
-						Section.bMatchWithPrevious = bMatchPrevious;
-						if (Index >= 0)
-						{
-							FName Name = Skeleton->GetReferenceSkeleton().GetBoneName(Index);
-							Section.MatchSectionByBoneTransform(SkelMeshComp, SequencerPtr->GetLocalTime().Time, SequencerPtr->GetLocalTime().Rate, Name);
-						}
-						else
-						{
-							Section.ClearMatchedOffsetTransforms();
-						}
-						SequencerPtr->NotifyMovieSceneDataChanged(EMovieSceneDataChangeType::TrackValueChanged);
+						return;
+					}
 
+					UMovieSceneSkeletalAnimationTrack* const Track = Section.GetTypedOuter<UMovieSceneSkeletalAnimationTrack>();
+					if (!IsValid(Track))
+					{
+						return;
 					}
-				),
-				FCanExecuteAction::CreateLambda([=]()->bool 
-					{ 
-						return SequencerPtr.IsValid(); 
-					}),
-				FIsActionChecked::CreateLambda([=, this]()->bool
+
+					// Can't pick the object that this track binds
+					USkeleton* const Skeleton = AcquireSkeletonFromObjectGuid(ObjectBinding, SequencerPtr);
+					if (!IsValid(Skeleton))
 					{
-						if (Index >= 0)
-						{
-							FName Name = Skeleton->GetReferenceSkeleton().GetBoneName(Index);
-							return (Section.MatchedBoneName == Name);
-						}
-						return (Section.MatchedBoneName == NAME_None);
-					})
-				);
-		};
+						return;
+					}
 
+					FScopedTransaction MatchSection(LOCTEXT("MatchSectionByBone_Transaction", "Match Section By Bone"));
+					Section.Modify();	
+					Section.bMatchWithPrevious = bMatchPrevious;
+					if (Index >= 0)
+					{
+						USkeletalMeshComponent* const SkelMeshComp = AcquireSkeletalMeshFromObjectGuid(ObjectBinding, SequencerPtr);
 
-		MenuBuilder.BeginSection(NAME_None, LOCTEXT("MotionBlendingOptions", "Motion Blending Options"));
-		{
-			MenuBuilder.AddSubMenu(
-				LOCTEXT("MatchWithThisBoneInPreviousClip", "Match With This Bone In Previous Clip"), LOCTEXT("MatchWithThisBoneInPreviousClip_Tooltip", "Match This Bone With Previous Clip At Current Frame"),
-				FNewMenuDelegate::CreateLambda([=](FMenuBuilder& SubMenuBuilder) {
+						FName Name = Skeleton->GetReferenceSkeleton().GetBoneName(Index);
+						Section.MatchSectionByBoneTransform(SkelMeshComp, SequencerPtr->GetLocalTime().Time, SequencerPtr->GetLocalTime().Rate, Name);
+					}
+					else
+					{
+						Section.ClearMatchedOffsetTransforms();
+					}
+					SequencerPtr->NotifyMovieSceneDataChanged(EMovieSceneDataChangeType::TrackValueChanged);
+
+				}
+			),
+			FCanExecuteAction::CreateLambda([this]() -> bool 
+				{ 
+					return Sequencer.IsValid(); 
+				}),
+			FIsActionChecked::CreateLambda([this, ObjectBinding, Index]() -> bool
+				{
+					const TSharedPtr<ISequencer> SequencerPtr = Sequencer.Pin();
+					if (!SequencerPtr.IsValid())
+					{
+						return false;
+					}
+
+					USkeleton* const Skeleton = AcquireSkeletonFromObjectGuid(ObjectBinding, SequencerPtr);
+					if (!IsValid(Skeleton))
+					{
+						return false;
+					}
+
+					if (Index >= 0)
+					{
+						FName Name = Skeleton->GetReferenceSkeleton().GetBoneName(Index);
+						return (Section.MatchedBoneName == Name);
+					}
+
+					return (Section.MatchedBoneName == NAME_None);
+				})
+			);
+	};
+
+	MenuBuilder.BeginSection(NAME_None, LOCTEXT("MotionBlendingOptions", "Motion Blending Options"));
+	{
+		MenuBuilder.AddSubMenu(
+			LOCTEXT("MatchWithThisBoneInPreviousClip", "Match With This Bone In Previous Clip"), LOCTEXT("MatchWithThisBoneInPreviousClip_Tooltip", "Match This Bone With Previous Clip At Current Frame"),
+			FNewMenuDelegate::CreateLambda([MatchToBone, BoneNames](FMenuBuilder& SubMenuBuilder)
+				{
 					int32 Index = -1;
 					FText NoNameText = LOCTEXT("TurnOffBoneMatching", "Turn Off Matching");
 					FText NoNameTooltipText = LOCTEXT("TurnOffMatchingTooltip", "Turn Off Any Bone Matching");
 					SubMenuBuilder.AddMenuEntry(
 						NoNameText, NoNameTooltipText,
 						FSlateIcon(), MatchToBone(true, Index++), NAME_None, EUserInterfaceActionType::RadioButton);
-					
 
 					for (const FName& BoneName : BoneNames)
 					{
@@ -956,11 +1000,12 @@ void FSkeletalAnimationSection::BuildSectionContextMenu(FMenuBuilder& MenuBuilde
 							Text, TooltipText,
 							FSlateIcon(), MatchToBone(true, Index++), NAME_None, EUserInterfaceActionType::RadioButton);
 					}
-					}));
+				}));
 
-			MenuBuilder.AddSubMenu(
-				LOCTEXT("MatchWithThisBoneInNextClip", "Match With This Bone In Next Clip"), LOCTEXT("MatchWithThisBoneInNextClip_Tooltip", "Match This Bone With Next Clip At Current Frame"),
-				FNewMenuDelegate::CreateLambda([=](FMenuBuilder& SubMenuBuilder) {
+		MenuBuilder.AddSubMenu(
+			LOCTEXT("MatchWithThisBoneInNextClip", "Match With This Bone In Next Clip"), LOCTEXT("MatchWithThisBoneInNextClip_Tooltip", "Match This Bone With Next Clip At Current Frame"),
+			FNewMenuDelegate::CreateLambda([MatchToBone, BoneNames](FMenuBuilder& SubMenuBuilder)
+				{
 					int32 Index = -1;
 					
 					FText NoNameText = LOCTEXT("TurnOffBoneMatching", "Turn Off Matching");
@@ -968,7 +1013,6 @@ void FSkeletalAnimationSection::BuildSectionContextMenu(FMenuBuilder& MenuBuilde
 					SubMenuBuilder.AddMenuEntry(
 						NoNameText, NoNameTooltipText,
 						FSlateIcon(), MatchToBone(false, Index++), NAME_None, EUserInterfaceActionType::RadioButton);
-					
 
 					for (const FName& BoneName : BoneNames)
 					{
@@ -979,118 +1023,200 @@ void FSkeletalAnimationSection::BuildSectionContextMenu(FMenuBuilder& MenuBuilde
 							Text, TooltipText,
 							FSlateIcon(), MatchToBone(false, Index++), NAME_None, EUserInterfaceActionType::RadioButton);
 					}
-					}));
-			MenuBuilder.AddMenuEntry(
-				LOCTEXT("MatchTranslation", "Match X and Y Translation"),
-				LOCTEXT("MatchTranslationTooltip", "Match the Translation to the Specified Bone"),
-				FSlateIcon(),
-				FUIAction(
-					FExecuteAction::CreateLambda([=, this]()->void {
+				}));
+
+		MenuBuilder.AddMenuEntry(
+			LOCTEXT("MatchTranslation", "Match X and Y Translation"),
+			LOCTEXT("MatchTranslationTooltip", "Match the Translation to the Specified Bone"),
+			FSlateIcon(),
+			FUIAction(
+				FExecuteAction::CreateLambda([this, ObjectBinding]()
+					{
+						const TSharedPtr<ISequencer> SequencerPtr = Sequencer.Pin();
+						if (!SequencerPtr.IsValid())
+						{
+							return;
+						}
+
+						USkeletalMeshComponent* const SkelMeshComp = AcquireSkeletalMeshFromObjectGuid(ObjectBinding, SequencerPtr);
+
 						FScopedTransaction MatchTransaction(LOCTEXT("MatchTranslation_Transaction", "Match Translation"));
 						Section.Modify();	
 						Section.ToggleMatchTranslation();
 						Section.MatchSectionByBoneTransform(SkelMeshComp, SequencerPtr->GetLocalTime().Time, SequencerPtr->GetLocalTime().Rate, Section.MatchedBoneName);
-						SequencerPtr->NotifyMovieSceneDataChanged(EMovieSceneDataChangeType::TrackValueChanged);}),
-						FCanExecuteAction::CreateLambda([]()->bool { return true; }),
-						FIsActionChecked::CreateLambda([this]()->bool { return Section.bMatchTranslation; })),
-				NAME_None,
-				EUserInterfaceActionType::ToggleButton
-			);
+						SequencerPtr->NotifyMovieSceneDataChanged(EMovieSceneDataChangeType::TrackValueChanged);
+					}),
+				FCanExecuteAction::CreateLambda([]() -> bool
+					{
+						return true;
+					}),
+				FIsActionChecked::CreateLambda([this]() -> bool
+					{
+						return Section.bMatchTranslation;
+					})),
+			NAME_None,
+			EUserInterfaceActionType::ToggleButton);
+
+		MenuBuilder.AddMenuEntry(
+			LOCTEXT("MatchZHeight", "Match Z Height"),
+			LOCTEXT("MatchZHeightTooltip", "Match the Z Height, may want this off for better matching"),
+			FSlateIcon(),
+			FUIAction(
+				FExecuteAction::CreateLambda([this, ObjectBinding]()
+					{
+						const TSharedPtr<ISequencer> SequencerPtr = Sequencer.Pin();
+						if (!SequencerPtr.IsValid())
+						{
+							return;
+						}
+
+						USkeletalMeshComponent* const SkelMeshComp = AcquireSkeletalMeshFromObjectGuid(ObjectBinding, SequencerPtr);
 
-			MenuBuilder.AddMenuEntry(
-				LOCTEXT("MatchZHeight", "Match Z Height"),
-				LOCTEXT("MatchZHeightTooltip", "Match the Z Height, may want this off for better matching"),
-				FSlateIcon(),
-				FUIAction(
-					FExecuteAction::CreateLambda([=, this]()->void { 
 						FScopedTransaction MatchTransaction(LOCTEXT("MatchZHeight_Transaction", "Match Z Height"));
 						Section.Modify();
 						Section.ToggleMatchIncludeZHeight(); 
 						Section.MatchSectionByBoneTransform(SkelMeshComp, SequencerPtr->GetLocalTime().Time, SequencerPtr->GetLocalTime().Rate, Section.MatchedBoneName);
-						SequencerPtr->NotifyMovieSceneDataChanged(EMovieSceneDataChangeType::TrackValueChanged); }),
-						FCanExecuteAction::CreateLambda([]()->bool { return true; }),
-						FIsActionChecked::CreateLambda([this]()->bool { return Section.bMatchIncludeZHeight;  })),
-				NAME_None,
-				EUserInterfaceActionType::ToggleButton
-			);
+						SequencerPtr->NotifyMovieSceneDataChanged(EMovieSceneDataChangeType::TrackValueChanged);
+					}),
+					FCanExecuteAction::CreateLambda([]() -> bool
+						{
+							return true;
+						}),
+					FIsActionChecked::CreateLambda([this]() -> bool
+						{
+							return Section.bMatchIncludeZHeight;
+						})),
+			NAME_None,
+			EUserInterfaceActionType::ToggleButton
+		);
+
+		MenuBuilder.AddMenuEntry(
+			LOCTEXT("MatchYawRotation", "Match Yaw Rotation"),
+			LOCTEXT("MatchYawRotationTooltip", "Match the Yaw Rotation, may want this off for better matching"),
+			FSlateIcon(),
+			FUIAction(
+				FExecuteAction::CreateLambda([this, ObjectBinding]()
+					{
+						const TSharedPtr<ISequencer> SequencerPtr = Sequencer.Pin();
+						if (!SequencerPtr.IsValid())
+						{
+							return;
+						}
+
+						USkeletalMeshComponent* const SkelMeshComp = AcquireSkeletalMeshFromObjectGuid(ObjectBinding, SequencerPtr);
 
-			MenuBuilder.AddMenuEntry(
-				LOCTEXT("MatchYawRotation", "Match Yaw Rotation"),
-				LOCTEXT("MatchYawRotationTooltip", "Match the Yaw Rotation, may want this off for better matching"),
-				FSlateIcon(),
-				FUIAction(
-					FExecuteAction::CreateLambda([=, this]()->void {
 						FScopedTransaction MatchTransaction(LOCTEXT("MatchYawRotation_Transaction", "Match Yaw Rotation"));
 						Section.Modify();
 						Section.ToggleMatchIncludeYawRotation();
 						Section.MatchSectionByBoneTransform(SkelMeshComp, SequencerPtr->GetLocalTime().Time, SequencerPtr->GetLocalTime().Rate, Section.MatchedBoneName);
 						SequencerPtr->NotifyMovieSceneDataChanged(EMovieSceneDataChangeType::TrackValueChanged); }),
-						FCanExecuteAction::CreateLambda([]()->bool { return true; }),
-						FIsActionChecked::CreateLambda([this]()->bool { return Section.bMatchRotationYaw; })),
+					FCanExecuteAction::CreateLambda([]() -> bool
+						{
+							return true;
+						}),
+					FIsActionChecked::CreateLambda([this]() -> bool
+						{
+							return Section.bMatchRotationYaw;
+						})),
 				NAME_None,
-							EUserInterfaceActionType::ToggleButton
-							);
+				EUserInterfaceActionType::ToggleButton);
+
+		MenuBuilder.AddMenuEntry(
+			LOCTEXT("MatchPitchRotation", "Match Pitch Rotation"),
+			LOCTEXT("MatchPitchRotationTooltip", "Match the Pitch Rotation, may want this off for better matching"),
+			FSlateIcon(),
+			FUIAction(
+				FExecuteAction::CreateLambda([this, ObjectBinding]()
+					{
+						const TSharedPtr<ISequencer> SequencerPtr = Sequencer.Pin();
+						if (!SequencerPtr.IsValid())
+						{
+							return;
+						}
+
+						USkeletalMeshComponent* const SkelMeshComp = AcquireSkeletalMeshFromObjectGuid(ObjectBinding, SequencerPtr);
 
-			MenuBuilder.AddMenuEntry(
-				LOCTEXT("MatchPitchRotation", "Match Pitch Rotation"),
-				LOCTEXT("MatchPitchRotationTooltip", "Match the Pitch Rotation, may want this off for better matching"),
-				FSlateIcon(),
-				FUIAction(
-					FExecuteAction::CreateLambda([=, this]()->void { 
 						FScopedTransaction MatchTransaction(LOCTEXT("MatchPitchRotation_Transaction", "Match Pitch Rotation"));
 						Section.Modify();
 						Section.ToggleMatchIncludePitchRotation();
 						Section.MatchSectionByBoneTransform(SkelMeshComp, SequencerPtr->GetLocalTime().Time, SequencerPtr->GetLocalTime().Rate, Section.MatchedBoneName);
-						SequencerPtr->NotifyMovieSceneDataChanged(EMovieSceneDataChangeType::TrackValueChanged); }),
-						FCanExecuteAction::CreateLambda([]()->bool { return true; }),
-						FIsActionChecked::CreateLambda([this]()->bool { return Section.bMatchRotationPitch;})),
-				NAME_None,
-				EUserInterfaceActionType::ToggleButton
-			);
+						SequencerPtr->NotifyMovieSceneDataChanged(EMovieSceneDataChangeType::TrackValueChanged);
+					}),
+					FCanExecuteAction::CreateLambda([]() -> bool
+						{
+							return true;
+						}),
+					FIsActionChecked::CreateLambda([this]() -> bool
+						{
+							return Section.bMatchRotationPitch;
+						})),
+			NAME_None,
+			EUserInterfaceActionType::ToggleButton);
+
+		MenuBuilder.AddMenuEntry(
+			LOCTEXT("MatchRollRotation", "Match Roll Rotation"),
+			LOCTEXT("MatchRollRotationTooltip", "Match the Roll Rotation, may want this off for better matching"),
+			FSlateIcon(),
+			FUIAction(
+				FExecuteAction::CreateLambda([this, ObjectBinding]()
+					{
+						const TSharedPtr<ISequencer> SequencerPtr = Sequencer.Pin();
+						if (!SequencerPtr.IsValid())
+						{
+							return;
+						}
+
+						USkeletalMeshComponent* const SkelMeshComp = AcquireSkeletalMeshFromObjectGuid(ObjectBinding, SequencerPtr);
 
-			MenuBuilder.AddMenuEntry(
-				LOCTEXT("MatchRollRotation", "Match Roll Rotation"),
-				LOCTEXT("MatchRollRotationTooltip", "Match the Roll Rotation, may want this off for better matching"),
-				FSlateIcon(),
-				FUIAction(
-					FExecuteAction::CreateLambda([=, this]()->void {
 						FScopedTransaction MatchTransaction(LOCTEXT("MatchRollRotation_Transaction", "Match Roll Rotation"));
 						Section.Modify();
 						Section.ToggleMatchIncludeRollRotation();
 						Section.MatchSectionByBoneTransform(SkelMeshComp, SequencerPtr->GetLocalTime().Time, SequencerPtr->GetLocalTime().Rate, Section.MatchedBoneName);
-						SequencerPtr->NotifyMovieSceneDataChanged(EMovieSceneDataChangeType::TrackValueChanged); }),
-						FCanExecuteAction::CreateLambda([]()->bool { return true; }),
-						FIsActionChecked::CreateLambda([this]()->bool { return Section.bMatchRotationRoll; })),
-				NAME_None,
-							EUserInterfaceActionType::ToggleButton
-				);
-
+						SequencerPtr->NotifyMovieSceneDataChanged(EMovieSceneDataChangeType::TrackValueChanged);
+					}),
+				FCanExecuteAction::CreateLambda([]() -> bool
+					{
+						return true;
+					}),
+				FIsActionChecked::CreateLambda([this]() -> bool
+					{
+						return Section.bMatchRotationRoll;
+					})),
+			NAME_None,
+			EUserInterfaceActionType::ToggleButton);
+	}
+	MenuBuilder.EndSection();
 
-			MenuBuilder.EndSection();
-		}
-		MenuBuilder.BeginSection(NAME_None, LOCTEXT("SkelAnimSectionDisplay", "Display"));
-		{
+	MenuBuilder.BeginSection(NAME_None, LOCTEXT("SkelAnimSectionDisplay", "Display"));
+	{
+		MenuBuilder.AddMenuEntry(
+			NSLOCTEXT("Sequencer", "ShowSkeletons", "Show Skeleton"),
+			NSLOCTEXT("Sequencer", "ShowSkeletonsTooltip", "Show A Skeleton for this Section."),
+			FSlateIcon(),
+			FUIAction(
+				FExecuteAction::CreateLambda([this]()
+					{
+						const TSharedPtr<ISequencer> SequencerPtr = Sequencer.Pin();
+						if (!SequencerPtr.IsValid())
+						{
+							return;
+						}
 
-			MenuBuilder.AddMenuEntry(
-				NSLOCTEXT("Sequencer", "ShowSkeletons", "Show Skeleton"),
-				NSLOCTEXT("Sequencer", "ShowSkeletonsTooltip", "Show A Skeleton for this Section."),
-				FSlateIcon(),
-				FUIAction(
-					FExecuteAction::CreateLambda([this, SequencerPtr]()->void {
 						Section.ToggleShowSkeleton();
 						SequencerPtr->NotifyMovieSceneDataChanged(EMovieSceneDataChangeType::TrackValueChanged);
-
-						}),
-					FCanExecuteAction::CreateLambda([=]()->bool { return SequencerPtr != nullptr; }),
-						FIsActionChecked::CreateLambda([this]()->bool { return Section.bShowSkeleton; })),
-						NAME_None,
-						EUserInterfaceActionType::ToggleButton
-					);
-
-			MenuBuilder.EndSection();
-
-		}
+					}),
+				FCanExecuteAction::CreateLambda([this]() -> bool
+					{
+						return Sequencer.IsValid();
+					}),
+				FIsActionChecked::CreateLambda([this]() -> bool
+					{
+						return Section.bShowSkeleton;
+					})),
+				NAME_None,
+				EUserInterfaceActionType::ToggleButton);
 	}
+	MenuBuilder.EndSection();
 }
 
 
diff --git a/Engine/Source/Editor/MovieSceneTools/Public/Channels/PerlinNoiseChannelInterface.h b/Engine/Source/Editor/MovieSceneTools/Public/Channels/PerlinNoiseChannelInterface.h
--- a/Engine/Source/Editor/MovieSceneTools/Public/Channels/PerlinNoiseChannelInterface.h
+++ b/Engine/Source/Editor/MovieSceneTools/Public/Channels/PerlinNoiseChannelInterface.h
@@ -44,11 +44,12 @@ class SNumericTextBlockKeyEditor : public SCompoundWidget
 	}
 };
 
-struct FPerlinNoiseChannelSectionMenuExtension : TSharedFromThis<FPerlinNoiseChannelSectionMenuExtension>
+struct FPerlinNoiseChannelSectionMenuExtension : TSharedFromThis<FPerlinNoiseChannelSectionMenuExtension>, ISidebarChannelExtension
 {
-	FPerlinNoiseChannelSectionMenuExtension(TArrayView<const FMovieSceneChannelHandle> InChannelHandles, TArrayView<UMovieSceneSection* const> InSections);
+	FPerlinNoiseChannelSectionMenuExtension(TArrayView<const FMovieSceneChannelHandle> InChannelHandles, const TArray<TWeakObjectPtr<UMovieSceneSection>>& InWeakSections);
+	virtual ~FPerlinNoiseChannelSectionMenuExtension() {}
 
-	void ExtendMenu(FMenuBuilder& MenuBuilder, const bool bSubMenu = true);
+	virtual TSharedPtr<ISidebarChannelExtension> ExtendMenu(FMenuBuilder& MenuBuilder, const bool bInSubMenu = true) override;
 
 private:
 
@@ -61,9 +62,9 @@ private:
 
 	struct FChannelNotifyHook : FNotifyHook
 	{
-		UObject* ObjectToModify = nullptr;
+		TWeakObjectPtr<UObject> WeakObjectToModify;
 
-		FChannelNotifyHook(UObject* InObjectToModify) : ObjectToModify(InObjectToModify) {}
+		FChannelNotifyHook(UObject* InObjectToModify) : WeakObjectToModify(InObjectToModify) {}
 
 		virtual void NotifyPreChange(FProperty* PropertyAboutToChange) override;
 		virtual void NotifyPostChange(const FPropertyChangedEvent& PropertyChangedEvent, FProperty* PropertyThatChanged) override;
@@ -71,7 +72,7 @@ private:
 
 	TArray<FMovieSceneChannelHandle> ChannelHandles;
 	TArray<int32> ChannelHandleSectionIndexes;
-	TArray<UMovieSceneSection*> Sections;
+	TArray<TWeakObjectPtr<UMovieSceneSection>> WeakSections;
 	TArray<FChannelNotifyHook> NotifyHooks;
 };
 
@@ -124,21 +125,22 @@ struct TPerlinNoiseChannelInterface : ISequencerChannelInterface
 	{
 	}
 
-	virtual void ExtendSectionMenu_Raw(FMenuBuilder& MenuBuilder, TSharedPtr<FExtender> MenuExtender, TArrayView<const FMovieSceneChannelHandle> Channels, TArrayView<UMovieSceneSection* const> Sections, TWeakPtr<ISequencer> InSequencer) const override
+	virtual void ExtendSectionMenu_Raw(FMenuBuilder& MenuBuilder, TSharedPtr<FExtender> InMenuExtender, TArrayView<const FMovieSceneChannelHandle> InChannels, const TArray<TWeakObjectPtr<UMovieSceneSection>>& InWeakSections, TWeakPtr<ISequencer> InWeakSequencer) const override
 	{
-		TSharedRef<FPerlinNoiseChannelSectionMenuExtension> Extension = MakeShared<FPerlinNoiseChannelSectionMenuExtension>(Channels, Sections);
+		TSharedRef<FPerlinNoiseChannelSectionMenuExtension> Extension = MakeShared<FPerlinNoiseChannelSectionMenuExtension>(InChannels, InWeakSections);
 
-		MenuExtender->AddMenuExtension(TEXT("SequencerChannels"), EExtensionHook::First, nullptr
+		InMenuExtender->AddMenuExtension(TEXT("SequencerChannels"), EExtensionHook::First, nullptr
 			, FMenuExtensionDelegate::CreateLambda([Extension](FMenuBuilder& MenuBuilder)
 				{
 					Extension->ExtendMenu(MenuBuilder);
 				}));
 	}
 
-	virtual void ExtendSidebarMenu_Raw(FMenuBuilder& MenuBuilder, TSharedPtr<FExtender> MenuExtender, TArrayView<const FMovieSceneChannelHandle> Channels, TArrayView<UMovieSceneSection* const> Sections, TWeakPtr<ISequencer> InSequencer) const override
+	virtual TSharedPtr<ISidebarChannelExtension> ExtendSidebarMenu_Raw(FMenuBuilder& MenuBuilder, TSharedPtr<FExtender> InMenuExtender, TArrayView<const FMovieSceneChannelHandle> InChannels, const TArray<TWeakObjectPtr<UMovieSceneSection>>& InWeakSections, TWeakPtr<ISequencer> InWeakSequencer) const override
 	{
-		TSharedRef<FPerlinNoiseChannelSectionMenuExtension> Extension = MakeShared<FPerlinNoiseChannelSectionMenuExtension>(Channels, Sections);
+		const TSharedRef<FPerlinNoiseChannelSectionMenuExtension> Extension = MakeShared<FPerlinNoiseChannelSectionMenuExtension>(InChannels, InWeakSections);
 		Extension->ExtendMenu(MenuBuilder, false);
+		return Extension;
 	}
 
 	virtual void DrawKeys_Raw(FMovieSceneChannel* Channel, TArrayView<const FKeyHandle> InKeyHandles, const UMovieSceneSection* InOwner, TArrayView<FKeyDrawParams> OutKeyDrawParams) const override
diff --git a/Engine/Source/Editor/Sequencer/Private/MVVM/ViewModels/CategoryModel.cpp b/Engine/Source/Editor/Sequencer/Private/MVVM/ViewModels/CategoryModel.cpp
--- a/Engine/Source/Editor/Sequencer/Private/MVVM/ViewModels/CategoryModel.cpp
+++ b/Engine/Source/Editor/Sequencer/Private/MVVM/ViewModels/CategoryModel.cpp
@@ -316,12 +316,12 @@ void FCategoryGroupModel::BuildSidebarMenu(FMenuBuilder& MenuBuilder)
 			{
 				return InSection;
 			});
-		SequencerHelpers::BuildEditSectionMenu(*Sequencer.Get(), WeakSectionObjects, MenuBuilder, false);
+		SequencerHelpers::BuildEditSectionMenu(Sequencer, WeakSectionObjects, MenuBuilder, false);
 	}
 
 	FOutlinerItemModel::BuildSidebarMenu(MenuBuilder);
 }
-	
+
 } // namespace UE::Sequencer
 
 #undef LOCTEXT_NAMESPACE
diff --git a/Engine/Source/Editor/Sequencer/Private/MVVM/ViewModels/TrackModel.cpp b/Engine/Source/Editor/Sequencer/Private/MVVM/ViewModels/TrackModel.cpp
--- a/Engine/Source/Editor/Sequencer/Private/MVVM/ViewModels/TrackModel.cpp
+++ b/Engine/Source/Editor/Sequencer/Private/MVVM/ViewModels/TrackModel.cpp
@@ -669,8 +669,6 @@ void FTrackModel::BuildContextMenu(FMenuBuilder& MenuBuilder)
 		return;
 	}
 
-	FSequencer& SequencerRef = *Sequencer.Get();
-
 	if (TrackEditor)
 	{
 		TrackEditor->BuildTrackContextMenu(MenuBuilder, Track);
@@ -678,13 +676,13 @@ void FTrackModel::BuildContextMenu(FMenuBuilder& MenuBuilder)
 
 	if (Track->GetSupportedBlendTypes().Num() > 0)
 	{
-		SequencerHelpers::BuildNewSectionMenu(SequencerRef, GetRowIndex() + 1, GetTrack(), MenuBuilder);
+		SequencerHelpers::BuildNewSectionMenu(Sequencer, GetRowIndex() + 1, GetTrack(), MenuBuilder);
 	}
 
-	SequencerHelpers::BuildBlendingMenu(SequencerRef, Track, MenuBuilder);
+	SequencerHelpers::BuildBlendingMenu(Sequencer, Track, MenuBuilder);
 
 	const TArray<TWeakObjectPtr<>> TrackAreaModels = SequencerHelpers::GetSectionObjectsFromTrackAreaModels(GetTrackAreaModelList());
-	SequencerHelpers::BuildEditSectionMenu(SequencerRef, TrackAreaModels, MenuBuilder, true);
+	SequencerHelpers::BuildEditSectionMenu(Sequencer, TrackAreaModels, MenuBuilder, true);
 
 	if (const TViewModelPtr<FChannelGroupModel> ChannelGroup = TopLevelChannelList.GetHead().ImplicitCast())
 	{
@@ -714,8 +712,6 @@ void FTrackModel::BuildSidebarMenu(FMenuBuilder& MenuBuilder)
 		return;
 	}
 
-	FSequencer& SequencerRef = *Sequencer.Get();
-
 	if (TrackEditor)
 	{
 		TrackEditor->BuildTrackSidebarMenu(MenuBuilder, Track);
@@ -723,13 +719,13 @@ void FTrackModel::BuildSidebarMenu(FMenuBuilder& MenuBuilder)
 
 	if (Track->GetSupportedBlendTypes().Num() > 0)
 	{
-		SequencerHelpers::BuildNewSectionMenu(SequencerRef, GetRowIndex() + 1, GetTrack(), MenuBuilder);
+		SequencerHelpers::BuildNewSectionMenu(Sequencer, GetRowIndex() + 1, GetTrack(), MenuBuilder);
 	}
 
-	SequencerHelpers::BuildBlendingMenu(SequencerRef, Track, MenuBuilder);
+	SequencerHelpers::BuildBlendingMenu(Sequencer, Track, MenuBuilder);
 
 	const TArray<TWeakObjectPtr<>> TrackAreaModels = SequencerHelpers::GetSectionObjectsFromTrackAreaModels(GetTrackAreaModelList());
-	SequencerHelpers::BuildEditSectionMenu(SequencerRef, TrackAreaModels, MenuBuilder, false);
+	SequencerHelpers::BuildEditSectionMenu(Sequencer, TrackAreaModels, MenuBuilder, false);
 
 	if (const TViewModelPtr<FChannelGroupModel> ChannelGroup = TopLevelChannelList.GetHead().ImplicitCast())
 	{
diff --git a/Engine/Source/Editor/Sequencer/Private/MVVM/ViewModels/TrackRowModel.cpp b/Engine/Source/Editor/Sequencer/Private/MVVM/ViewModels/TrackRowModel.cpp
--- a/Engine/Source/Editor/Sequencer/Private/MVVM/ViewModels/TrackRowModel.cpp
+++ b/Engine/Source/Editor/Sequencer/Private/MVVM/ViewModels/TrackRowModel.cpp
@@ -323,8 +323,6 @@ void FTrackRowModel::BuildContextMenu(FMenuBuilder& MenuBuilder)
 		return;
 	}
 
-	FSequencer& SequencerRef = *Sequencer.Get();
-
 	if (TrackEditor)
 	{
 		TrackEditor->BuildTrackContextMenu(MenuBuilder, Track);
@@ -332,13 +330,13 @@ void FTrackRowModel::BuildContextMenu(FMenuBuilder& MenuBuilder)
 
 	if (Track->GetSupportedBlendTypes().Num() > 0)
 	{
-		SequencerHelpers::BuildNewSectionMenu(SequencerRef, GetRowIndex() + 1, GetTrack(), MenuBuilder);
+		SequencerHelpers::BuildNewSectionMenu(Sequencer, GetRowIndex() + 1, GetTrack(), MenuBuilder);
 	}
 
-	SequencerHelpers::BuildBlendingMenu(SequencerRef, Track, MenuBuilder);
+	SequencerHelpers::BuildBlendingMenu(Sequencer, Track, MenuBuilder);
 
 	const TArray<TWeakObjectPtr<>> TrackAreaModels = SequencerHelpers::GetSectionObjectsFromTrackAreaModels(GetTrackAreaModelList());
-	SequencerHelpers::BuildEditSectionMenu(SequencerRef, TrackAreaModels, MenuBuilder, true);
+	SequencerHelpers::BuildEditSectionMenu(Sequencer, TrackAreaModels, MenuBuilder, true);
 
 	if (const TViewModelPtr<FChannelGroupModel> ChannelGroup = TopLevelChannelList.GetHead().ImplicitCast())
 	{
@@ -368,8 +366,6 @@ void FTrackRowModel::BuildSidebarMenu(FMenuBuilder& MenuBuilder)
 		return;
 	}
 
-	FSequencer& SequencerRef = *Sequencer.Get();
-
 	if (TrackEditor)
 	{
 		TrackEditor->BuildTrackSidebarMenu(MenuBuilder, Track);
@@ -377,13 +373,13 @@ void FTrackRowModel::BuildSidebarMenu(FMenuBuilder& MenuBuilder)
 
 	if (Track->GetSupportedBlendTypes().Num() > 0)
 	{
-		SequencerHelpers::BuildNewSectionMenu(SequencerRef, GetRowIndex() + 1, GetTrack(), MenuBuilder);
+		SequencerHelpers::BuildNewSectionMenu(Sequencer, GetRowIndex() + 1, GetTrack(), MenuBuilder);
 	}
 
-	SequencerHelpers::BuildBlendingMenu(SequencerRef, Track, MenuBuilder);
+	SequencerHelpers::BuildBlendingMenu(Sequencer, Track, MenuBuilder);
 
 	const TArray<TWeakObjectPtr<>> TrackAreaModels = SequencerHelpers::GetSectionObjectsFromTrackAreaModels(GetTrackAreaModelList());
-	SequencerHelpers::BuildEditSectionMenu(SequencerRef, TrackAreaModels, MenuBuilder, false);
+	SequencerHelpers::BuildEditSectionMenu(Sequencer, TrackAreaModels, MenuBuilder, false);
 
 	if (const TViewModelPtr<FChannelGroupModel> ChannelGroup = TopLevelChannelList.GetHead().ImplicitCast())
 	{
diff --git a/Engine/Source/Editor/Sequencer/Private/Menus/CurveChannelSectionSidebarExtension.cpp b/Engine/Source/Editor/Sequencer/Private/Menus/CurveChannelSectionSidebarExtension.cpp
--- a/Engine/Source/Editor/Sequencer/Private/Menus/CurveChannelSectionSidebarExtension.cpp
+++ b/Engine/Source/Editor/Sequencer/Private/Menus/CurveChannelSectionSidebarExtension.cpp
@@ -30,29 +30,28 @@
 
 #define LOCTEXT_NAMESPACE "CurveChannelSectionSidebarExtension"
 
-FCurveChannelSectionSidebarExtension::FCurveChannelSectionSidebarExtension(const TWeakPtr<ISequencer>& InSequencerWeak)
-	: SequencerWeak(InSequencerWeak)
+FCurveChannelSectionSidebarExtension::FCurveChannelSectionSidebarExtension(const TWeakPtr<ISequencer>& InWeakSequencer)
+	: WeakSequencer(InWeakSequencer)
 {
 }
 
-void FCurveChannelSectionSidebarExtension::AddSections(TArrayView<UMovieSceneSection* const> InSections)
+void FCurveChannelSectionSidebarExtension::AddSections(const TArray<TWeakObjectPtr<UMovieSceneSection>>& InWeakSections)
 {
-	for (UMovieSceneSection* const Section : InSections)
-	{
-		Sections.Add(Section);
-	}
+	WeakSections = TSet(InWeakSections);
 }
 
-void FCurveChannelSectionSidebarExtension::ExtendMenu(FMenuBuilder& MenuBuilder)
+TSharedPtr<ISidebarChannelExtension> FCurveChannelSectionSidebarExtension::ExtendMenu(FMenuBuilder& MenuBuilder, const bool bInSubMenu)
 {
 	AddExtrapolationMenu(MenuBuilder, true);
 	AddExtrapolationMenu(MenuBuilder, false);
 	AddDisplayOptionsMenu(MenuBuilder);
+
+	return SharedThis(this);
 }
 
 void FCurveChannelSectionSidebarExtension::AddDisplayOptionsMenu(FMenuBuilder& MenuBuilder)
 {
-	const TSharedPtr<ISequencer> Sequencer = SequencerWeak.Pin();
+	const TSharedPtr<ISequencer> Sequencer = WeakSequencer.Pin();
 	if (!Sequencer)
 	{
 		return;
@@ -110,7 +109,7 @@ void FCurveChannelSectionSidebarExtension::AddDisplayOptionsMenu(FMenuBuilder& M
 			.WidthOverride(50.f)
 			.IsEnabled_Lambda([this, KeyAreaName]()
 				{
-					if (const TSharedPtr<ISequencer> Sequencer = SequencerWeak.Pin())
+					if (const TSharedPtr<ISequencer> Sequencer = WeakSequencer.Pin())
 					{
 						return IsAnyShowCurve() && Sequencer->GetSequencerSettings()->HasKeyAreaCurveExtents(KeyAreaName);
 					}
@@ -134,7 +133,7 @@ void FCurveChannelSectionSidebarExtension::AddDisplayOptionsMenu(FMenuBuilder& M
 			.WidthOverride(50.f)
 			.IsEnabled_Lambda([this, KeyAreaName]()
 				{
-					if (const TSharedPtr<ISequencer> Sequencer = SequencerWeak.Pin())
+					if (const TSharedPtr<ISequencer> Sequencer = WeakSequencer.Pin())
 					{
 						return IsAnyShowCurve() && Sequencer->GetSequencerSettings()->HasKeyAreaCurveExtents(KeyAreaName);
 					}
@@ -250,7 +249,7 @@ void FCurveChannelSectionSidebarExtension::AddExtrapolationMenu(FMenuBuilder& Me
 void FCurveChannelSectionSidebarExtension::GetChannels(TArray<FMovieSceneFloatChannel*>& FloatChannels, TArray<FMovieSceneDoubleChannel*>& DoubleChannels,
 	TArray<FMovieSceneIntegerChannel*>& IntegerChannels, TArray<FMovieSceneBoolChannel*>& BoolChannels) const
 {
-	ISequencer* const Sequencer = SequencerWeak.Pin().Get();
+	const TSharedPtr<ISequencer> Sequencer = WeakSequencer.Pin();
 	if (!Sequencer)
 	{
 		return;
@@ -287,7 +286,7 @@ void FCurveChannelSectionSidebarExtension::GetChannels(TArray<FMovieSceneFloatCh
 	// Otherwise, the channels of all the sections
 	if (FloatChannels.Num() + DoubleChannels.Num() + IntegerChannels.Num() + BoolChannels.Num() == 0)
 	{
-		for (const TWeakObjectPtr<UMovieSceneSection>& WeakSection : Sections)
+		for (const TWeakObjectPtr<UMovieSceneSection>& WeakSection : WeakSections)
 		{
 			if (const UMovieSceneSection* const Section = WeakSection.Get())
 			{
@@ -332,7 +331,7 @@ void FCurveChannelSectionSidebarExtension::SetExtrapolationMode(const ERichCurve
 	bool bAnythingChanged = false;
 
 	// Modify all sections
-	for (TWeakObjectPtr<UMovieSceneSection> WeakSection : Sections)
+	for (TWeakObjectPtr<UMovieSceneSection> WeakSection : WeakSections)
 	{
 		UMovieSceneSection* const Section = WeakSection.Get();
 		if (IsValid(Section))
@@ -369,7 +368,7 @@ void FCurveChannelSectionSidebarExtension::SetExtrapolationMode(const ERichCurve
 
 	if (bAnythingChanged)
 	{
-		if (ISequencer* const Sequencer = SequencerWeak.Pin().Get())
+		if (const TSharedPtr<ISequencer> Sequencer = WeakSequencer.Pin())
 		{
 			Sequencer->NotifyMovieSceneDataChanged(EMovieSceneDataChangeType::TrackValueChanged);
 		}
@@ -435,7 +434,7 @@ void FCurveChannelSectionSidebarExtension::ToggleShowCurve()
 	bool bAnythingChanged = false;
 
 	// Modify all sections
-	for (const TWeakObjectPtr<UMovieSceneSection>& WeakSection : Sections)
+	for (const TWeakObjectPtr<UMovieSceneSection>& WeakSection : WeakSections)
 	{
 		UMovieSceneSection* const Section = WeakSection.Get();
 		if (IsValid(Section))
@@ -445,7 +444,7 @@ void FCurveChannelSectionSidebarExtension::ToggleShowCurve()
 	}
 
 	// Apply to all channels
-	for (const TWeakObjectPtr<UMovieSceneSection>& WeakSection : Sections)
+	for (const TWeakObjectPtr<UMovieSceneSection>& WeakSection : WeakSections)
 	{
 		if (const UMovieSceneSection* const Section = WeakSection.Get())
 		{
@@ -480,7 +479,7 @@ void FCurveChannelSectionSidebarExtension::ToggleShowCurve()
 ECheckBoxState FCurveChannelSectionSidebarExtension::IsShowCurve() const
 {
 	int32 NumShowedAndHidden[2] = { 0, 0 };
-	for (const TWeakObjectPtr<UMovieSceneSection>& WeakSection : Sections)
+	for (const TWeakObjectPtr<UMovieSceneSection>& WeakSection : WeakSections)
 	{
 		if (const UMovieSceneSection* const Section = WeakSection.Get())
 		{
@@ -517,7 +516,7 @@ ECheckBoxState FCurveChannelSectionSidebarExtension::IsShowCurve() const
 
 bool FCurveChannelSectionSidebarExtension::IsAnyShowCurve() const
 {
-	for (const TWeakObjectPtr<UMovieSceneSection>& WeakSection : Sections)
+	for (const TWeakObjectPtr<UMovieSceneSection>& WeakSection : WeakSections)
 	{
 		const UMovieSceneSection* const Section = WeakSection.Get();
 		if (IsValid(Section))
@@ -547,7 +546,7 @@ bool FCurveChannelSectionSidebarExtension::IsAnyShowCurve() const
 
 int32 FCurveChannelSectionSidebarExtension::GetKeyAreaHeight() const
 {
-	if (const TSharedPtr<ISequencer> Sequencer = SequencerWeak.Pin())
+	if (const TSharedPtr<ISequencer> Sequencer = WeakSequencer.Pin())
 	{
 		return (int32)Sequencer->GetSequencerSettings()->GetKeyAreaHeightWithCurves();
 	}
@@ -654,7 +653,7 @@ void FCurveChannelSectionSidebarExtension::OnKeyAreaCurveMaxChanged(const double
 
 USequencerSettings* FCurveChannelSectionSidebarExtension::GetSequencerSettings() const
 {
-	if (const TSharedPtr<ISequencer> Sequencer = SequencerWeak.Pin())
+	if (const TSharedPtr<ISequencer> Sequencer = WeakSequencer.Pin())
 	{
 		return Sequencer->GetSequencerSettings();
 	}
diff --git a/Engine/Source/Editor/Sequencer/Private/Menus/CurveChannelSectionSidebarExtension.h b/Engine/Source/Editor/Sequencer/Private/Menus/CurveChannelSectionSidebarExtension.h
--- a/Engine/Source/Editor/Sequencer/Private/Menus/CurveChannelSectionSidebarExtension.h
+++ b/Engine/Source/Editor/Sequencer/Private/Menus/CurveChannelSectionSidebarExtension.h
@@ -12,13 +12,14 @@ struct FMovieSceneIntegerChannel;
 enum class ECheckBoxState : uint8;
 enum ERichCurveExtrapolation : int;
 
-struct FCurveChannelSectionSidebarExtension : TSharedFromThis<FCurveChannelSectionSidebarExtension>
+struct FCurveChannelSectionSidebarExtension : TSharedFromThis<FCurveChannelSectionSidebarExtension>, ISidebarChannelExtension
 {
-	FCurveChannelSectionSidebarExtension(const TWeakPtr<ISequencer>& InSequencerWeak);
+	FCurveChannelSectionSidebarExtension(const TWeakPtr<ISequencer>& InWeakSequencer);
+	virtual ~FCurveChannelSectionSidebarExtension() {}
 
-	void AddSections(TArrayView<UMovieSceneSection* const> InSections);
+	void AddSections(const TArray<TWeakObjectPtr<UMovieSceneSection>>& InWeakSections);
 
-	void ExtendMenu(FMenuBuilder& MenuBuilder);
+	virtual TSharedPtr<ISidebarChannelExtension> ExtendMenu(FMenuBuilder& MenuBuilder, const bool bInSubMenu = true) override;
 
 private:
 	void AddDisplayOptionsMenu(FMenuBuilder& MenuBuilder);
@@ -48,7 +49,7 @@ private:
 
 	USequencerSettings* GetSequencerSettings() const;
 
-	TWeakPtr<ISequencer> SequencerWeak;
+	TWeakPtr<ISequencer> WeakSequencer;
 
-	TSet<TWeakObjectPtr<UMovieSceneSection>> Sections;
+	TSet<TWeakObjectPtr<UMovieSceneSection>> WeakSections;
 };
diff --git a/Engine/Source/Editor/Sequencer/Private/Sequencer.cpp b/Engine/Source/Editor/Sequencer/Private/Sequencer.cpp
--- a/Engine/Source/Editor/Sequencer/Private/Sequencer.cpp
+++ b/Engine/Source/Editor/Sequencer/Private/Sequencer.cpp
@@ -670,8 +670,6 @@ void FSequencer::InitSequencer(const FSequencerInitParams& InitParams, const TSh
 	OnActivateSequenceEvent.Broadcast(ActiveTemplateIDs[0]);
 
 	{ // Selection Details Drawer
-		FSequencerSelectionDrawer SelectionDrawer(SharedThis(this));
-
 		FSidebarDrawerConfig DetailsDrawer;
 		DetailsDrawer.UniqueId = SelectionDrawerId;
 		DetailsDrawer.ButtonText = LOCTEXT("SelectionDetailsPanelLabel", "Selection");
diff --git a/Engine/Source/Editor/Sequencer/Private/SequencerCommonHelpers.cpp b/Engine/Source/Editor/Sequencer/Private/SequencerCommonHelpers.cpp
--- a/Engine/Source/Editor/Sequencer/Private/SequencerCommonHelpers.cpp
+++ b/Engine/Source/Editor/Sequencer/Private/SequencerCommonHelpers.cpp
@@ -311,28 +311,26 @@ private:
 	TWeakPtr<ISequencer> WeakSequencer;
 };
 
-void SequencerHelpers::BuildNewSectionMenu(FSequencer& Sequencer
+void SequencerHelpers::BuildNewSectionMenu(const TWeakPtr<FSequencer>& InWeakSequencer
 	, const int32 InRowIndex
 	, const TWeakObjectPtr<UMovieSceneTrack>& InTrackWeak
 	, FMenuBuilder& MenuBuilder)
 {
 	using namespace UE::Sequencer;
 
-	const TWeakPtr<ISequencer> WeakSequencer = Sequencer.AsWeak();
-
 	MenuBuilder.AddSubMenu(
 		NSLOCTEXT("Sequencer", "AddSection", "Add Section"),
 		FText(),
-		FNewMenuDelegate::CreateLambda([WeakSequencer, InRowIndex, InTrackWeak](FMenuBuilder& SubMenuBuilder)
+		FNewMenuDelegate::CreateLambda([InWeakSequencer, InRowIndex, InTrackWeak](FMenuBuilder& SubMenuBuilder)
 		{
-			if (const TSharedPtr<ISequencer> Sequencer = WeakSequencer.Pin())
+			if (const TSharedPtr<ISequencer> Sequencer = InWeakSequencer.Pin())
 			{
 				FSequencerUtilities::PopulateMenu_CreateNewSection(SubMenuBuilder, InRowIndex, InTrackWeak.Get(), Sequencer);
 			}
 		}));
 }
 
-void SequencerHelpers::BuildEditSectionMenu(FSequencer& Sequencer
+void SequencerHelpers::BuildEditSectionMenu(const TWeakPtr<FSequencer>& InWeakSequencer
 	, const TArray<TWeakObjectPtr<>>& InWeakSections
 	, FMenuBuilder& MenuBuilder
 	, const bool bInSubMenu)
@@ -344,16 +342,17 @@ void SequencerHelpers::BuildEditSectionMenu(FSequencer& Sequencer
 		return;
 	}
 
-	const TWeakPtr<FSequencer> WeakSequencer = StaticCastWeakPtr<FSequencer>(Sequencer.AsWeak());
-
-	auto BuildSection = [WeakSequencer, InWeakSections](FMenuBuilder& LambdaMenuBuilder)
+	auto BuildSection = [InWeakSequencer, InWeakSections](FMenuBuilder& LambdaMenuBuilder)
 	{
-		const TSharedPtr<FSequencer> Sequencer = WeakSequencer.Pin();
+		const TSharedPtr<FSequencer> Sequencer = InWeakSequencer.Pin();
 		if (!Sequencer)
 		{
 			return;
 		}
 
+		const TWeakObjectPtr<UMovieScene> CurrentScene = Sequencer->GetFocusedMovieSceneSequence()->GetMovieScene();
+		const TWeakPtr<INumericTypeInterface<double>> NumericTypeInterfafce = Sequencer->GetNumericTypeInterface();
+
 		TSharedRef<SSectionDetailsNotifyHookWrapper> DetailsNotifyWrapper = SNew(SSectionDetailsNotifyHookWrapper);
 		FDetailsViewArgs DetailsViewArgs;
 		{
@@ -373,21 +372,14 @@ void SequencerHelpers::BuildEditSectionMenu(FSequencer& Sequencer
 
 		TSharedRef<IDetailsView> DetailsView = FModuleManager::GetModuleChecked<FPropertyEditorModule>("PropertyEditor").CreateDetailView(DetailsViewArgs);
 		DetailsView->RegisterInstancedCustomPropertyTypeLayout("FrameNumber",
-			FOnGetPropertyTypeCustomizationInstance::CreateLambda([WeakSequencer]()
+			FOnGetPropertyTypeCustomizationInstance::CreateLambda([NumericTypeInterfafce]()
 			{
-				const TSharedPtr<FSequencer> Sequencer = WeakSequencer.Pin();
-				check(Sequencer);
-				return MakeShared<FFrameNumberDetailsCustomization>(Sequencer->GetNumericTypeInterface());
+				return MakeShared<FFrameNumberDetailsCustomization>(NumericTypeInterfafce.Pin());
 			}));
 		DetailsView->RegisterInstancedCustomPropertyLayout(UMovieSceneSection::StaticClass(),
-			FOnGetDetailCustomizationInstance::CreateLambda([WeakSequencer]()
+			FOnGetDetailCustomizationInstance::CreateLambda([NumericTypeInterfafce, CurrentScene]()
 			{
-				const TSharedPtr<FSequencer> Sequencer = WeakSequencer.Pin();
-				check(Sequencer);
-				// We pass the current scene to the UMovieSceneSection customization so we can get the overall bounds of the section when we change a section from infinite->bounded.
-				UMovieScene* const CurrentScene = Sequencer->GetFocusedMovieSceneSequence()->GetMovieScene();
-				check(CurrentScene);
-				return MakeShared<FMovieSceneSectionDetailsCustomization>(Sequencer->GetNumericTypeInterface(), CurrentScene);
+				return MakeShared<FMovieSceneSectionDetailsCustomization>(NumericTypeInterfafce.Pin(), CurrentScene.Get());
 			}));
 	
 		DetailsView->SetIsPropertyVisibleDelegate(FIsPropertyVisible::CreateLambda([](const FPropertyAndParent& PropertyAndParent)
@@ -407,7 +399,7 @@ void SequencerHelpers::BuildEditSectionMenu(FSequencer& Sequencer
 				{
 					TSharedPtr<ISequencerSection> SectionInterface = SectionHandle->GetSectionInterface();
 					FSequencerSectionPropertyDetailsViewCustomizationParams CustomizationDetails(
-						SectionInterface.ToSharedRef(), WeakSequencer, *SectionHandle->GetParentTrackExtension()->GetTrackEditor().Get());
+						SectionInterface.ToSharedRef(), InWeakSequencer, *SectionHandle->GetParentTrackExtension()->GetTrackEditor().Get());
 					TSharedPtr<FObjectBindingModel> ParentObjectBindingNode = SectionHandle->FindAncestorOfType<FObjectBindingModel>();
 					if (ParentObjectBindingNode.IsValid())
 					{
@@ -421,7 +413,7 @@ void SequencerHelpers::BuildEditSectionMenu(FSequencer& Sequencer
 		Sequencer->OnInitializeDetailsPanel().Broadcast(DetailsView, Sequencer.ToSharedRef());
 		DetailsView->SetObjects(InWeakSections);
 
-		DetailsNotifyWrapper->SetDetailsAndSequencer(DetailsView, WeakSequencer);
+		DetailsNotifyWrapper->SetDetailsAndSequencer(DetailsView, InWeakSequencer);
 		DetailsNotifyWrapper->SetEnabled(!Sequencer->IsReadOnly());
 
 		LambdaMenuBuilder.BeginSection(TEXT("TrackSection"), NSLOCTEXT("Sequencer", "TrackSectionMenuSection", "Track Section"));
@@ -451,7 +443,7 @@ void SequencerHelpers::BuildEditSectionMenu(FSequencer& Sequencer
 	}
 }
 
-void SequencerHelpers::BuildBlendingMenu(FSequencer& Sequencer
+void SequencerHelpers::BuildBlendingMenu(const TWeakPtr<FSequencer>& InWeakSequencer
 	, const TWeakObjectPtr<UMovieSceneTrack>& InTrackWeak
 	, FMenuBuilder& MenuBuilder)
 {
@@ -473,14 +465,12 @@ void SequencerHelpers::BuildBlendingMenu(FSequencer& Sequencer
 		return;
 	}
 
-	const TWeakPtr<ISequencer> WeakSequencer = Sequencer.AsWeak();
-
 	MenuBuilder.AddSubMenu(
 		NSLOCTEXT("Sequencer", "BlendingAlgorithmSubMenu", "Blending Algorithm"),
 		FText(),
-		FNewMenuDelegate::CreateLambda([WeakSequencer, InTrackWeak](FMenuBuilder& SubMenuBuilder)
+		FNewMenuDelegate::CreateLambda([InWeakSequencer, InTrackWeak](FMenuBuilder& SubMenuBuilder)
 		{
-			if (const TSharedPtr<ISequencer> Sequencer = WeakSequencer.Pin())
+			if (const TSharedPtr<ISequencer> Sequencer = InWeakSequencer.Pin())
 			{
 				FSequencerUtilities::PopulateMenu_BlenderSubMenu(SubMenuBuilder, InTrackWeak.Get(), Sequencer);
 			}
diff --git a/Engine/Source/Editor/Sequencer/Private/SequencerCommonHelpers.h b/Engine/Source/Editor/Sequencer/Private/SequencerCommonHelpers.h
--- a/Engine/Source/Editor/Sequencer/Private/SequencerCommonHelpers.h
+++ b/Engine/Source/Editor/Sequencer/Private/SequencerCommonHelpers.h
@@ -71,7 +71,7 @@ public:
 	/*
 	 * Build a sub menu for adding a new track section
 	 */
-	static void BuildNewSectionMenu(FSequencer& Sequencer
+	static void BuildNewSectionMenu(const TWeakPtr<FSequencer>& InWeakSequencer
 		, const int32 InRowIndex
 		, const TWeakObjectPtr<UMovieSceneTrack>& InTrackWeak
 		, FMenuBuilder& MenuBuilder);
@@ -79,7 +79,7 @@ public:
 	/*
 	 * Build an inline menu or sub menu for editing track section(s)
 	 */
-	static void BuildEditSectionMenu(FSequencer& Sequencer
+	static void BuildEditSectionMenu(const TWeakPtr<FSequencer>& InWeakSequencer
 		, const TArray<TWeakObjectPtr<>>& InWeakSections
 		, FMenuBuilder& MenuBuilder
 		, const bool bInSubMenu);
@@ -87,7 +87,7 @@ public:
 	/*
 	 * Build a menu for selection the blend algorithm
 	 */
-	static void BuildBlendingMenu(FSequencer& Sequencer
+	static void BuildBlendingMenu(const TWeakPtr<FSequencer>& InWeakSequencer
 		, const TWeakObjectPtr<UMovieSceneTrack>& InTrackWeak
 		, FMenuBuilder& MenuBuilder);
 
diff --git a/Engine/Source/Editor/Sequencer/Private/SequencerContextMenus.cpp b/Engine/Source/Editor/Sequencer/Private/SequencerContextMenus.cpp
--- a/Engine/Source/Editor/Sequencer/Private/SequencerContextMenus.cpp
+++ b/Engine/Source/Editor/Sequencer/Private/SequencerContextMenus.cpp
@@ -289,7 +289,12 @@ void FSectionContextMenu::PopulateMenu(FMenuBuilder& MenuBuilder, TSharedPtr<FEx
 		ISequencerChannelInterface* ChannelInterface = SequencerModule.FindChannelEditorInterface(Pair.Key);
 		if (ChannelInterface)
 		{
-			ChannelInterface->ExtendSectionMenu_Raw(MenuBuilder, MenuExtender, Pair.Value, Sections, WeakSequencer);
+			TArray<TWeakObjectPtr<UMovieSceneSection>> WeakSections;
+			Algo::Transform(Sections, WeakSections, [](UMovieSceneSection* const InSection)
+				{
+					return InSection;
+				});
+			ChannelInterface->ExtendSectionMenu_Raw(MenuBuilder, MenuExtender, Pair.Value, WeakSections, WeakSequencer);
 		}
 	}
 
@@ -313,7 +318,7 @@ void FSectionContextMenu::PopulateMenu(FMenuBuilder& MenuBuilder, TSharedPtr<FEx
 				}
 			}
 
-			SequencerHelpers::BuildEditSectionMenu(*Sequencer.Get(), Sections, SubMenuBuilder, /*bInSubMenu=*/true);
+			SequencerHelpers::BuildEditSectionMenu(Sequencer, Sections, SubMenuBuilder, /*bInSubMenu=*/true);
 		})
 	);
 
diff --git a/Engine/Source/Editor/Sequencer/Private/Widgets/Sidebar/SequencerSelectionDrawer.cpp b/Engine/Source/Editor/Sequencer/Private/Widgets/Sidebar/SequencerSelectionDrawer.cpp
--- a/Engine/Source/Editor/Sequencer/Private/Widgets/Sidebar/SequencerSelectionDrawer.cpp
+++ b/Engine/Source/Editor/Sequencer/Private/Widgets/Sidebar/SequencerSelectionDrawer.cpp
@@ -129,7 +129,7 @@ TSharedRef<SWidget> FSequencerSelectionDrawer::CreateContentWidget()
 		if (const TSharedPtr<FSequencerSelection> SequencerSelection = Private::GetSelection(*Sequencer.Get()))
 		{
 			SequencerSelection->OnChanged.AddSP(this, &FSequencerSelectionDrawer::OnSequencerSelectionChanged);
-			
+
 			OnSequencerSelectionChanged();
 		}
 	}
@@ -153,99 +153,92 @@ void FSequencerSelectionDrawer::OnSequencerSelectionChanged()
 	{
 		return;
 	}
-	const TSharedPtr<FSequencer> Sequencer = WeakSequencer.Pin();
-	if (!WeakSequencer.IsValid())
-	{
-		return;
-	}
 
-	const TSharedPtr<FSequencerSelection> SequencerSelection = Private::GetSelection(*Sequencer.Get());
-	if (!SequencerSelection.IsValid())
+	if (bWaitingToHandleSelectionChanged)
 	{
 		return;
 	}
-	static bool bWaitingToHandleSelectionChanged = false;
 
-	if (bWaitingToHandleSelectionChanged == false)
-	{
-		bWaitingToHandleSelectionChanged = true;
-		GEditor->GetTimerManager()->SetTimerForNextTick([this]()
+	bWaitingToHandleSelectionChanged = true;
+
+	GEditor->GetTimerManager()->SetTimerForNextTick([this]()
 		{
-			bWaitingToHandleSelectionChanged = false;
-			if (!ContentBox.IsValid())
-			{
-				return;
-			}
-			ContentBox->ClearChildren();
+			UpdateFromSelectionNextFrame();
+		});
+}
 
-			CurveChannelExtension.Reset();
+void FSequencerSelectionDrawer::UpdateFromSelectionNextFrame()
+{
+	bWaitingToHandleSelectionChanged = false;
 
-			const TSharedPtr<FSequencer> Sequencer = WeakSequencer.Pin();
-			if (!WeakSequencer.IsValid())
-			{
-				return;
-			}
+	ContentBox->ClearChildren();
 
-			const TSharedPtr<FSequencerSelection> SequencerSelection = Private::GetSelection(*Sequencer.Get());
-			if (!SequencerSelection.IsValid())
-			{
-				return;
-			}
+	CurveChannelExtension.Reset();
 
-			const TSharedRef<FSequencerSelection> SelectionRef = SequencerSelection.ToSharedRef();
+	const TSharedPtr<FSequencer> Sequencer = WeakSequencer.Pin();
+	if (!Sequencer.IsValid())
+	{
+		return;
+	}
 
-			auto AddToContent = [this](const TSharedRef<SWidget>& InWidget)
-				{
-					ContentBox->AddSlot()
-						.AutoHeight()
-						[
-							InWidget
-						];
-				};
-
-			const ISequencerModule& SequencerModule = FModuleManager::Get().LoadModuleChecked<ISequencerModule>(TEXT("Sequencer"));
-			const TSharedPtr<FExtensibilityManager> SidebarExtensibilityManager = SequencerModule.GetSidebarExtensibilityManager();
-
-			FMenuBuilder MenuBuilder(/*bInShouldCloseWindowAfterMenuSelection=*/false
-				, Sequencer->GetCommandBindings()
-				, SidebarExtensibilityManager->GetAllExtenders()
-				, /*bInCloseSelfOnly=*/true, &FCoreStyle::Get(), /*bInSearchable=*/true, TEXT("Sequencer.Sidebar"));
-
-			/**
-				* Selection details display order preference:
-				*  1) Key items
-				*  2) Track area items (if no key selected)
-				*  3) Outliner items (if no key or track area selected)
-				*  4) Marked frames
-				*/
-
-				// 1) Key items
-			BuildKeySelectionDetails(SelectionRef, MenuBuilder);
-
-			// Early out for key selections
-			const bool bIsKeySelected = SequencerSelection->KeySelection.Num() > 0;
-			if (bIsKeySelected)
-			{
-				AddToContent(MenuBuilder.MakeWidget());
-				return;
-			}
+	const TSharedPtr<FSequencerSelection> SequencerSelection = Private::GetSelection(*Sequencer.Get());
+	if (!SequencerSelection.IsValid())
+	{
+		return;
+	}
 
-			// 2) Track area items
-			BuildTrackAreaDetails(*Sequencer, SelectionRef, MenuBuilder);
+	const TSharedRef<FSequencerSelection> SelectionRef = SequencerSelection.ToSharedRef();
 
-			// 3) Outliner items
-			const bool bIsTrackAreaSelected = SequencerSelection->TrackArea.Num() > 0;
-			if (!bIsTrackAreaSelected)
-			{
-				BuildOutlinerDetails(*Sequencer, SelectionRef, MenuBuilder);
-			}
+	auto AddToContent = [this](const TSharedRef<SWidget>& InWidget)
+		{
+			ContentBox->AddSlot()
+				.AutoHeight()
+				[
+					InWidget
+				];
+		};
+
+	const ISequencerModule& SequencerModule = FModuleManager::Get().LoadModuleChecked<ISequencerModule>(TEXT("Sequencer"));
+	const TSharedPtr<FExtensibilityManager> SidebarExtensibilityManager = SequencerModule.GetSidebarExtensibilityManager();
+
+	FMenuBuilder MenuBuilder(/*bInShouldCloseWindowAfterMenuSelection=*/false
+		, Sequencer->GetCommandBindings()
+		, SidebarExtensibilityManager->GetAllExtenders()
+		, /*bInCloseSelfOnly=*/true, &FCoreStyle::Get(), /*bInSearchable=*/true, TEXT("Sequencer.Sidebar"));
+
+	/**
+	 * Selection details display order preference:
+	 *  1) Key items
+	 *  2) Track area items (if no key selected)
+	 *  3) Outliner items (if no key or track area selected)
+	 *  4) Marked frames
+	 */
+
+	// 1) Key items
+	BuildKeySelectionDetails(SelectionRef, MenuBuilder);
+
+	// Early out for key selections
+	const bool bIsKeySelected = SequencerSelection->KeySelection.Num() > 0;
+	if (bIsKeySelected)
+	{
+		AddToContent(MenuBuilder.MakeWidget());
+		return;
+	}
 
-			// 4) Marked frames
-			BuildMarkedFrameDetails(SelectionRef, MenuBuilder);
+	// 2) Track area items
+	BuildTrackAreaDetails(SelectionRef, MenuBuilder);
 
-			AddToContent(MenuBuilder.MakeWidget());
-		});
+	// 3) Outliner items
+	const bool bIsTrackAreaSelected = SequencerSelection->TrackArea.Num() > 0;
+	if (!bIsTrackAreaSelected)
+	{
+		BuildOutlinerDetails(SelectionRef, MenuBuilder);
 	}
+
+	// 4) Marked frames
+	BuildMarkedFrameDetails(SelectionRef, MenuBuilder);
+
+	AddToContent(MenuBuilder.MakeWidget());
 }
 
 void FSequencerSelectionDrawer::BuildKeySelectionDetails(const TSharedRef<FSequencerSelection>& InSelection, FMenuBuilder& MenuBuilder)
@@ -276,11 +269,11 @@ void FSequencerSelectionDrawer::BuildKeySelectionDetails(const TSharedRef<FSeque
 	}
 	if (Channels.Num() == 1)
 	{
-		FSectionContextMenu::BuildKeyEditMenu(MenuBuilder, WeakSequencer, Sequencer->GetLastEvaluatedLocalTime());
+		FSectionContextMenu::BuildKeyEditMenu(MenuBuilder, Sequencer, Sequencer->GetLastEvaluatedLocalTime());
 	}
 }
 
-void FSequencerSelectionDrawer::BuildTrackAreaDetails(FSequencer& InSequencer, const TSharedRef<FSequencerSelection>& InSelection, FMenuBuilder& MenuBuilder)
+void FSequencerSelectionDrawer::BuildTrackAreaDetails(const TSharedRef<FSequencerSelection>& InSelection, FMenuBuilder& MenuBuilder)
 {
 	TArray<TWeakObjectPtr<>> AllSectionObjects;
 
@@ -299,8 +292,6 @@ void FSequencerSelectionDrawer::BuildTrackAreaDetails(FSequencer& InSequencer, c
 		{
 			if (InSelection->TrackArea.Num() == 1)
 			{
-				const TViewModelPtr<IOutlinerExtension> LinkedOutlinerItem = SectionModel->GetLinkedOutlinerItem();
-
 				if (const TSharedPtr<ISequencerSection> SectionInterface = SectionModel->GetSectionInterface())
 				{
 					const TSharedPtr<IObjectBindingExtension> ObjectBinding = SectionModel->FindAncestorOfType<IObjectBindingExtension>();
@@ -314,17 +305,23 @@ void FSequencerSelectionDrawer::BuildTrackAreaDetails(FSequencer& InSequencer, c
 
 	if (!AllSectionObjects.IsEmpty())
 	{
-		SequencerHelpers::BuildEditSectionMenu(InSequencer, AllSectionObjects, MenuBuilder, false);
+		SequencerHelpers::BuildEditSectionMenu(WeakSequencer, AllSectionObjects, MenuBuilder, false);
 	}
 }
 
-void FSequencerSelectionDrawer::BuildOutlinerDetails(FSequencer& InSequencer, const TSharedRef<FSequencerSelection>& InSelection, FMenuBuilder& MenuBuilder)
+void FSequencerSelectionDrawer::BuildOutlinerDetails(const TSharedRef<FSequencerSelection>& InSelection, FMenuBuilder& MenuBuilder)
 {
 	if (InSelection->Outliner.Num() == 0)
 	{
 		return;
 	}
 
+	const TSharedPtr<FSequencer> Sequencer = WeakSequencer.Pin();
+	if (!Sequencer.IsValid())
+	{
+		return;
+	}
+
 	TSet<TViewModelPtr<FChannelGroupOutlinerModel>> ChannelGroups;
 
 	for (const FViewModelPtr OutlinerItem : InSelection->Outliner)
@@ -350,7 +347,7 @@ void FSequencerSelectionDrawer::BuildOutlinerDetails(FSequencer& InSequencer, co
 		TArray<FName> ChannelTypeNames;
 		TArray<ISequencerChannelInterface*> ChannelInterfaces;
 		TArray<FMovieSceneChannelHandle> ChannelHandles;
-		TArray<UMovieSceneSection*> SceneSections;
+		TArray<TWeakObjectPtr<UMovieSceneSection>> WeakSceneSections;
 
 		for (const TViewModelPtr<FChannelGroupOutlinerModel>& ChannelModel : ChannelGroups)
 		{
@@ -363,7 +360,7 @@ void FSequencerSelectionDrawer::BuildOutlinerDetails(FSequencer& InSequencer, co
 					ChannelTypeNames.Add(Channel.GetChannelTypeName());
 					ChannelInterfaces.Add(SequencerChannelIterface);
 					ChannelHandles.Add(Channel);
-					SceneSections.Add(KeyArea->GetOwningSection());
+					WeakSceneSections.Add(KeyArea->GetOwningSection());
 				}
 			}
 		}
@@ -385,23 +382,31 @@ void FSequencerSelectionDrawer::BuildOutlinerDetails(FSequencer& InSequencer, co
 		// Channel Interface Extensions (Perlin Noise, Easing, Wave)
 		if (ChannelInterfaces.Num() > 0)
 		{
+			ChannelExtensions.Empty();
+
 			if (bAllChannelNamesEqual)
 			{
-				ChannelInterfaces[0]->ExtendSidebarMenu_Raw(MenuBuilder, Extender, ChannelHandles, SceneSections, WeakSequencer);
+				if (const TSharedPtr<ISidebarChannelExtension> ChannelExtension = ChannelInterfaces[0]->ExtendSidebarMenu_Raw(MenuBuilder, Extender, ChannelHandles, WeakSceneSections, Sequencer))
+				{
+					ChannelExtensions.Add(ChannelExtension);
+				}
 			}
 			else
 			{
 				// Display different channels separately and don't allow to edit "all-in-one"
 				for (int32 Index = 0; Index < ChannelInterfaces.Num(); ++Index)
 				{
-					ChannelInterfaces[Index]->ExtendSidebarMenu_Raw(MenuBuilder, Extender, { ChannelHandles[Index] }, { SceneSections[Index] }, WeakSequencer);
+					if (const TSharedPtr<ISidebarChannelExtension> ChannelExtension = ChannelInterfaces[Index]->ExtendSidebarMenu_Raw(MenuBuilder, Extender, { ChannelHandles[Index] }, { WeakSceneSections[Index] }, Sequencer))
+					{
+						ChannelExtensions.Add(ChannelExtension);
+					}
 				}
 			}
 		}
 
 		// Curve Channel Options (Pre-Finity, Post-Finity, etc.)
-		CurveChannelExtension = MakeShared<FCurveChannelSectionSidebarExtension>(WeakSequencer);
-		CurveChannelExtension->AddSections(SceneSections);
+		CurveChannelExtension = MakeShared<FCurveChannelSectionSidebarExtension>(Sequencer);
+		CurveChannelExtension->AddSections(WeakSceneSections);
 		CurveChannelExtension->ExtendMenu(MenuBuilder);
 	}
 }
diff --git a/Engine/Source/Editor/Sequencer/Private/Widgets/Sidebar/SequencerSelectionDrawer.h b/Engine/Source/Editor/Sequencer/Private/Widgets/Sidebar/SequencerSelectionDrawer.h
--- a/Engine/Source/Editor/Sequencer/Private/Widgets/Sidebar/SequencerSelectionDrawer.h
+++ b/Engine/Source/Editor/Sequencer/Private/Widgets/Sidebar/SequencerSelectionDrawer.h
@@ -9,6 +9,7 @@ class FMenuBuilder;
 class FName;
 class FSequencer;
 class FText;
+class ISidebarChannelExtension;
 class SVerticalBox;
 class SWidget;
 struct FCurveChannelSectionSidebarExtension;
@@ -35,10 +36,11 @@ public:
 
 protected:
 	void OnSequencerSelectionChanged();
+	void UpdateFromSelectionNextFrame();
 
 	void BuildKeySelectionDetails(const TSharedRef<UE::Sequencer::FSequencerSelection>& InSelection, FMenuBuilder& MenuBuilder);
-	void BuildTrackAreaDetails(FSequencer& InSequencer, const TSharedRef<UE::Sequencer::FSequencerSelection>& InSelection, FMenuBuilder& MenuBuilder);
-	void BuildOutlinerDetails(FSequencer& InSequencer, const TSharedRef<UE::Sequencer::FSequencerSelection>& InSelection, FMenuBuilder& MenuBuilder);
+	void BuildTrackAreaDetails(const TSharedRef<UE::Sequencer::FSequencerSelection>& InSelection, FMenuBuilder& MenuBuilder);
+	void BuildOutlinerDetails(const TSharedRef<UE::Sequencer::FSequencerSelection>& InSelection, FMenuBuilder& MenuBuilder);
 	void BuildMarkedFrameDetails(const TSharedRef<UE::Sequencer::FSequencerSelection>& InSelection, FMenuBuilder& MenuBuilder);
 
 	TSharedRef<SWidget> CreateHintText(const FText& InMessage);
@@ -54,4 +56,8 @@ protected:
 	TSharedPtr<SVerticalBox> ContentBox;
 
 	TSharedPtr<FCurveChannelSectionSidebarExtension> CurveChannelExtension;
+
+	TArray<TSharedPtr<ISidebarChannelExtension>> ChannelExtensions;
+
+	bool bWaitingToHandleSelectionChanged = false;
 };
diff --git a/Engine/Source/Editor/Sequencer/Public/ISequencerChannelInterface.h b/Engine/Source/Editor/Sequencer/Public/ISequencerChannelInterface.h
--- a/Engine/Source/Editor/Sequencer/Public/ISequencerChannelInterface.h
+++ b/Engine/Source/Editor/Sequencer/Public/ISequencerChannelInterface.h
@@ -65,6 +65,11 @@ struct FExtendKeyMenuParams
 	TArray<FKeyHandle> Handles;
 };
 
+class ISidebarChannelExtension
+{
+	virtual TSharedPtr<ISidebarChannelExtension> ExtendMenu(FMenuBuilder& MenuBuilder, const bool bInSubMenu = true) = 0;
+};
+
 /**
  * Abstract interface that defines all sequencer interactions for any channel type
  * Channels are stored internally as FMovieSceneChannel*, with this interface providing a common set of operations for all channels through a safe cast from the FMovieSceneChannel*.
@@ -153,20 +158,28 @@ struct ISequencerChannelInterface
 	 * @param MenuExtender          The menu extender to use
 	 * @param Channels              Array of type specific channels that exist in the selected sections
 	 * @param Sections              Array of sections being shown on the context menu
-	 * @param InSequencer           The currently active sequencer
+	 * @param InWeakSequencer       The currently active sequencer
 	 */
-	virtual void ExtendSectionMenu_Raw(FMenuBuilder& MenuBuilder, TSharedPtr<FExtender> MenuExtender, TArrayView<const FMovieSceneChannelHandle> Channels, TArrayView<UMovieSceneSection* const> Sections, TWeakPtr<ISequencer> InSequencer) const = 0;
+	virtual void ExtendSectionMenu_Raw(FMenuBuilder& MenuBuilder
+		, TSharedPtr<FExtender> MenuExtender
+		, TArrayView<const FMovieSceneChannelHandle> InChannels
+		, const TArray<TWeakObjectPtr<UMovieSceneSection>>& InWeakSections
+		, TWeakPtr<ISequencer> InWeakSequencer) const = 0;
 
 	/**
 	 * Extend the section sidebar menu
 	 *
 	 * @param MenuBuilder           The menu builder used to create this context menu
-	 * @param MenuExtender          The menu extender to use
-	 * @param Channels              Array of type specific channels that exist in the selected sections
-	 * @param Sections              Array of sections being shown on the context menu
-	 * @param InSequencer           The currently active sequencer
-	 */
-	virtual void ExtendSidebarMenu_Raw(FMenuBuilder& MenuBuilder, TSharedPtr<FExtender> MenuExtender, TArrayView<const FMovieSceneChannelHandle> Channels, TArrayView<UMovieSceneSection* const> Sections, TWeakPtr<ISequencer> InSequencer) const = 0;
+	 * @param InMenuExtender        The menu extender to use
+	 * @param InChannels            Array of type specific channels that exist in the selected sections
+	 * @param InWeakSections        Array of sections being shown on the context menu
+	 * @param InWeakSequencer       The currently active sequencer
+	 */
+	virtual TSharedPtr<ISidebarChannelExtension> ExtendSidebarMenu_Raw(FMenuBuilder& MenuBuilder
+		, TSharedPtr<FExtender> InMenuExtender
+		, TArrayView<const FMovieSceneChannelHandle> InChannels
+		, const TArray<TWeakObjectPtr<UMovieSceneSection>>& InWeakSections
+		, TWeakPtr<ISequencer> InWeakSequencer) const = 0;
 
 	/**
 	 * Gather information on how to draw the specified keys
diff --git a/Engine/Source/Editor/Sequencer/Public/SequencerChannelInterface.h b/Engine/Source/Editor/Sequencer/Public/SequencerChannelInterface.h
--- a/Engine/Source/Editor/Sequencer/Public/SequencerChannelInterface.h
+++ b/Engine/Source/Editor/Sequencer/Public/SequencerChannelInterface.h
@@ -145,44 +145,52 @@ struct TSequencerChannelInterfaceCommon : ISequencerChannelInterface
 	 * Extend the section context menu
 	 *
 	 * @param MenuBuilder           The menu builder used to create this context menu
-	 * @param MenuExtender          The menu extender to use
-	 * @param Channels              Array of type specific channels that exist in the selected sections
-	 * @param Sections              Array of sections being shown on the context menu
-	 * @param InSequencer           The currently active sequencer
+	 * @param InMenuExtender        The menu extender to use
+	 * @param InChannels            Array of type specific channels that exist in the selected sections
+	 * @param InWeakSections        Array of sections being shown on the context menu
+	 * @param InWeakSequencer       The currently active sequencer
 	 */
-	virtual void ExtendSectionMenu_Raw(FMenuBuilder& MenuBuilder, TSharedPtr<FExtender> MenuExtender, TArrayView<const FMovieSceneChannelHandle> Channels, TArrayView<UMovieSceneSection* const> Sections, TWeakPtr<ISequencer> InSequencer) const override
+	virtual void ExtendSectionMenu_Raw(FMenuBuilder& MenuBuilder
+		, TSharedPtr<FExtender> InMenuExtender
+		, TArrayView<const FMovieSceneChannelHandle> InChannels
+		, const TArray<TWeakObjectPtr<UMovieSceneSection>>& InWeakSections
+		, TWeakPtr<ISequencer> InWeakSequencer) const override
 	{
 		using namespace Sequencer;
 		TArray<TMovieSceneChannelHandle<ChannelType>> TypedChannels;
 
-		for (const FMovieSceneChannelHandle& RawHandle : Channels)
+		for (const FMovieSceneChannelHandle& RawHandle : InChannels)
 		{
 			TypedChannels.Add(RawHandle.Cast<ChannelType>());
 		}
 
-		ExtendSectionMenu(MenuBuilder, MenuExtender, MoveTemp(TypedChannels), Sections, InSequencer);
+		ExtendSectionMenu(MenuBuilder, InMenuExtender, MoveTemp(TypedChannels), InWeakSections, InWeakSequencer);
 	}
 
 	/**
 	 * Extend the section sidebar menu
 	 *
 	 * @param MenuBuilder           The menu builder used to create this context menu
-	 * @param MenuExtender          The menu extender to use
-	 * @param Channels              Array of type specific channels that exist in the selected sections
-	 * @param Sections              Array of sections being shown on the context menu
-	 * @param InSequencer           The currently active sequencer
+	 * @param InMenuExtender        The menu extender to use
+	 * @param InChannels            Array of type specific channels that exist in the selected sections
+	 * @param InWeakSections        Array of sections being shown on the context menu
+	 * @param InWeakSequencer       The currently active sequencer
 	 */
-	virtual void ExtendSidebarMenu_Raw(FMenuBuilder& MenuBuilder, TSharedPtr<FExtender> MenuExtender, TArrayView<const FMovieSceneChannelHandle> Channels, TArrayView<UMovieSceneSection* const> Sections, TWeakPtr<ISequencer> InSequencer) const override
+	virtual TSharedPtr<ISidebarChannelExtension> ExtendSidebarMenu_Raw(FMenuBuilder& MenuBuilder
+		, TSharedPtr<FExtender> InMenuExtender
+		, TArrayView<const FMovieSceneChannelHandle> InChannels
+		, const TArray<TWeakObjectPtr<UMovieSceneSection>>& InWeakSections
+		, TWeakPtr<ISequencer> InWeakSequencer) const override
 	{
 		using namespace Sequencer;
 		TArray<TMovieSceneChannelHandle<ChannelType>> TypedChannels;
 
-		for (const FMovieSceneChannelHandle& RawHandle : Channels)
+		for (const FMovieSceneChannelHandle& RawHandle : InChannels)
 		{
 			TypedChannels.Add(RawHandle.Cast<ChannelType>());
 		}
 
-		ExtendSidebarMenu(MenuBuilder, MenuExtender, MoveTemp(TypedChannels), Sections, InSequencer);
+		return ExtendSidebarMenu(MenuBuilder, InMenuExtender, MoveTemp(TypedChannels), InWeakSections, InWeakSequencer);
 	}
 
 	/**
diff --git a/Engine/Source/Editor/Sequencer/Public/SequencerChannelTraits.h b/Engine/Source/Editor/Sequencer/Public/SequencerChannelTraits.h
--- a/Engine/Source/Editor/Sequencer/Public/SequencerChannelTraits.h
+++ b/Engine/Source/Editor/Sequencer/Public/SequencerChannelTraits.h
@@ -103,20 +103,31 @@ namespace Sequencer
 	 * @param InSequencer    The sequencer that is currently active
 	 */
 	template<typename ChannelType>
-	void ExtendSectionMenu(FMenuBuilder& MenuBuilder, TSharedPtr<FExtender> MenuExtender, TArray<TMovieSceneChannelHandle<ChannelType>>&& Channels, TArrayView<UMovieSceneSection* const> Sections, TWeakPtr<ISequencer> InSequencer)
+	void ExtendSectionMenu(FMenuBuilder& MenuBuilder
+		, TSharedPtr<FExtender> InMenuExtender
+		, TArray<TMovieSceneChannelHandle<ChannelType>>&& InChannels
+		, const TArray<TWeakObjectPtr<UMovieSceneSection>>& InWeakSections
+		, TWeakPtr<ISequencer> InWeakSequencer)
 	{}
 
 	/**
 	 * Extend the specified selected sidebar menu
 	 *
-	 * @param MenuBuilder    The menu builder that will construct the section context menu
-	 * @param Channels       An array of all channels that are currently selected, in no particular order
-	 * @param Sections       An array of all sections that the selected channels reside in
-	 * @param InSequencer    The sequencer that is currently active
+	 * @param MenuBuilder     The menu builder that will construct the section context menu
+	 * @param InMenuExtender  The extender for the menu
+	 * @param InChannels      An array of all channels that are currently selected, in no particular order
+	 * @param InWeakSections  An array of all sections that the selected channels reside in
+	 * @param InWeakSequencer The sequencer that is currently active
 	 */
 	template<typename ChannelType>
-	void ExtendSidebarMenu(FMenuBuilder& MenuBuilder, TSharedPtr<FExtender> MenuExtender, TArray<TMovieSceneChannelHandle<ChannelType>>&& Channels, TArrayView<UMovieSceneSection* const> Sections, TWeakPtr<ISequencer> InSequencer)
-	{}
+	TSharedPtr<ISidebarChannelExtension> ExtendSidebarMenu(FMenuBuilder& MenuBuilder
+		, TSharedPtr<FExtender> InMenuExtender
+		, TArray<TMovieSceneChannelHandle<ChannelType>>&& InChannels
+		, const TArray<TWeakObjectPtr<UMovieSceneSection>>& InWeakSections
+		, TWeakPtr<ISequencer> InWeakSequencer)
+	{
+		return nullptr;
+	}
 
 	/**
 	 * Extend the specified selected key context menu
