[Sequencer] Fix Horde complaining about default virtual function value

diff --git a/Engine/Source/Editor/MovieSceneTools/Private/Channels/BuiltInChannelEditors.cpp b/Engine/Source/Editor/MovieSceneTools/Private/Channels/BuiltInChannelEditors.cpp
--- a/Engine/Source/Editor/MovieSceneTools/Private/Channels/BuiltInChannelEditors.cpp
+++ b/Engine/Source/Editor/MovieSceneTools/Private/Channels/BuiltInChannelEditors.cpp
@@ -1324,7 +1324,7 @@ struct FCurveChannelSectionMenuExtension : TSharedFromThis<FCurveChannelSectionM
 		++NumCurveChannelTypes;
 	}
 
-	virtual TSharedPtr<ISidebarChannelExtension> ExtendMenu(FMenuBuilder& MenuBuilder, const bool bInSubMenu = false) override
+	virtual TSharedPtr<ISidebarChannelExtension> ExtendMenu(FMenuBuilder& MenuBuilder, const bool bInSubMenu) override
 	{
 		--NumCurveChannelTypes;
 
@@ -1891,7 +1891,7 @@ void ExtendSectionMenu(FMenuBuilder& OuterMenuBuilder, TSharedPtr<FExtender> Men
 	TSharedRef<FCurveChannelSectionMenuExtension> Extension = FCurveChannelSectionMenuExtension::GetOrCreate(InWeakSequencer);
 	Extension->AddSections(InWeakSections);
 
-	MenuExtender->AddMenuExtension("SequencerChannels", EExtensionHook::First, nullptr, FMenuExtensionDelegate::CreateLambda([Extension](FMenuBuilder& MenuBuilder) { Extension->ExtendMenu(MenuBuilder); }));
+	MenuExtender->AddMenuExtension("SequencerChannels", EExtensionHook::First, nullptr, FMenuExtensionDelegate::CreateLambda([Extension](FMenuBuilder& MenuBuilder) { Extension->ExtendMenu(MenuBuilder, true); }));
 }
 
 void ExtendSectionMenu(FMenuBuilder& OuterMenuBuilder, TSharedPtr<FExtender> MenuExtender, TArray<TMovieSceneChannelHandle<FMovieSceneDoubleChannel>>&& Channels, const TArray<TWeakObjectPtr<UMovieSceneSection>>& InWeakSections, TWeakPtr<ISequencer> InWeakSequencer)
@@ -1899,7 +1899,7 @@ void ExtendSectionMenu(FMenuBuilder& OuterMenuBuilder, TSharedPtr<FExtender> Men
 	TSharedRef<FCurveChannelSectionMenuExtension> Extension = FCurveChannelSectionMenuExtension::GetOrCreate(InWeakSequencer);
 	Extension->AddSections(InWeakSections);
 
-	MenuExtender->AddMenuExtension("SequencerChannels", EExtensionHook::First, nullptr, FMenuExtensionDelegate::CreateLambda([Extension](FMenuBuilder& MenuBuilder) { Extension->ExtendMenu(MenuBuilder); }));
+	MenuExtender->AddMenuExtension("SequencerChannels", EExtensionHook::First, nullptr, FMenuExtensionDelegate::CreateLambda([Extension](FMenuBuilder& MenuBuilder) { Extension->ExtendMenu(MenuBuilder, true); }));
 }
 
 void ExtendSectionMenu(FMenuBuilder& OuterMenuBuilder, TSharedPtr<FExtender> MenuExtender, TArray<TMovieSceneChannelHandle<FMovieSceneIntegerChannel>>&& Channels, const TArray<TWeakObjectPtr<UMovieSceneSection>>& InWeakSections, TWeakPtr<ISequencer> InWeakSequencer)
@@ -1907,7 +1907,7 @@ void ExtendSectionMenu(FMenuBuilder& OuterMenuBuilder, TSharedPtr<FExtender> Men
 	TSharedRef<FCurveChannelSectionMenuExtension> Extension = FCurveChannelSectionMenuExtension::GetOrCreate(InWeakSequencer);
 	Extension->AddSections(InWeakSections);
 
-	MenuExtender->AddMenuExtension("SequencerChannels", EExtensionHook::First, nullptr, FMenuExtensionDelegate::CreateLambda([Extension](FMenuBuilder& MenuBuilder) { Extension->ExtendMenu(MenuBuilder); }));
+	MenuExtender->AddMenuExtension("SequencerChannels", EExtensionHook::First, nullptr, FMenuExtensionDelegate::CreateLambda([Extension](FMenuBuilder& MenuBuilder) { Extension->ExtendMenu(MenuBuilder, true); }));
 }
 
 void ExtendSectionMenu(FMenuBuilder& OuterMenuBuilder, TSharedPtr<FExtender> MenuExtender, TArray<TMovieSceneChannelHandle<FMovieSceneBoolChannel>>&& Channels, const TArray<TWeakObjectPtr<UMovieSceneSection>>& InWeakSections, TWeakPtr<ISequencer> InWeakSequencer)
@@ -1915,7 +1915,7 @@ void ExtendSectionMenu(FMenuBuilder& OuterMenuBuilder, TSharedPtr<FExtender> Men
 	TSharedRef<FCurveChannelSectionMenuExtension> Extension = FCurveChannelSectionMenuExtension::GetOrCreate(InWeakSequencer);
 	Extension->AddSections(InWeakSections);
 
-	MenuExtender->AddMenuExtension("SequencerChannels", EExtensionHook::First, nullptr, FMenuExtensionDelegate::CreateLambda([Extension](FMenuBuilder& MenuBuilder) { Extension->ExtendMenu(MenuBuilder); }));
+	MenuExtender->AddMenuExtension("SequencerChannels", EExtensionHook::First, nullptr, FMenuExtensionDelegate::CreateLambda([Extension](FMenuBuilder& MenuBuilder) { Extension->ExtendMenu(MenuBuilder, true); }));
 }
 
 TSharedPtr<ISidebarChannelExtension> ExtendSidebarMenu(FMenuBuilder& OuterMenuBuilder, TSharedPtr<FExtender> InMenuExtender, TArray<TMovieSceneChannelHandle<FMovieSceneFloatChannel>>&& Channels, const TArray<TWeakObjectPtr<UMovieSceneSection>>& InWeakSections, TWeakPtr<ISequencer> InWeakSequencer)
@@ -1923,7 +1923,7 @@ TSharedPtr<ISidebarChannelExtension> ExtendSidebarMenu(FMenuBuilder& OuterMenuBu
 	TSharedRef<FCurveChannelSectionMenuExtension> Extension = FCurveChannelSectionMenuExtension::GetOrCreate(InWeakSequencer);
 	Extension->AddSections(InWeakSections);
 
-	InMenuExtender->AddMenuExtension("SequencerChannels", EExtensionHook::First, nullptr, FMenuExtensionDelegate::CreateLambda([Extension](FMenuBuilder& MenuBuilder) { Extension->ExtendMenu(MenuBuilder); }));
+	InMenuExtender->AddMenuExtension("SequencerChannels", EExtensionHook::First, nullptr, FMenuExtensionDelegate::CreateLambda([Extension](FMenuBuilder& MenuBuilder) { Extension->ExtendMenu(MenuBuilder, false); }));
 
 	return Extension;
 }
@@ -1933,7 +1933,7 @@ TSharedPtr<ISidebarChannelExtension> ExtendSidebarMenu(FMenuBuilder& OuterMenuBu
 	TSharedRef<FCurveChannelSectionMenuExtension> Extension = FCurveChannelSectionMenuExtension::GetOrCreate(InWeakSequencer);
 	Extension->AddSections(InWeakSections);
 
-	InMenuExtender->AddMenuExtension("SequencerChannels", EExtensionHook::First, nullptr, FMenuExtensionDelegate::CreateLambda([Extension](FMenuBuilder& MenuBuilder) { Extension->ExtendMenu(MenuBuilder); }));
+	InMenuExtender->AddMenuExtension("SequencerChannels", EExtensionHook::First, nullptr, FMenuExtensionDelegate::CreateLambda([Extension](FMenuBuilder& MenuBuilder) { Extension->ExtendMenu(MenuBuilder, false); }));
 
 	return Extension;
 }
@@ -1943,7 +1943,7 @@ TSharedPtr<ISidebarChannelExtension> ExtendSidebarMenu(FMenuBuilder& OuterMenuBu
 	TSharedRef<FCurveChannelSectionMenuExtension> Extension = FCurveChannelSectionMenuExtension::GetOrCreate(InWeakSequencer);
 	Extension->AddSections(InWeakSections);
 
-	InMenuExtender->AddMenuExtension("SequencerChannels", EExtensionHook::First, nullptr, FMenuExtensionDelegate::CreateLambda([Extension](FMenuBuilder& MenuBuilder) { Extension->ExtendMenu(MenuBuilder); }));
+	InMenuExtender->AddMenuExtension("SequencerChannels", EExtensionHook::First, nullptr, FMenuExtensionDelegate::CreateLambda([Extension](FMenuBuilder& MenuBuilder) { Extension->ExtendMenu(MenuBuilder, false); }));
 
 	return Extension;
 }
@@ -1953,7 +1953,7 @@ TSharedPtr<ISidebarChannelExtension> ExtendSidebarMenu(FMenuBuilder& OuterMenuBu
 	TSharedRef<FCurveChannelSectionMenuExtension> Extension = FCurveChannelSectionMenuExtension::GetOrCreate(InWeakSequencer);
 	Extension->AddSections(InWeakSections);
 
-	InMenuExtender->AddMenuExtension("SequencerChannels", EExtensionHook::First, nullptr, FMenuExtensionDelegate::CreateLambda([Extension](FMenuBuilder& MenuBuilder) { Extension->ExtendMenu(MenuBuilder); }));
+	InMenuExtender->AddMenuExtension("SequencerChannels", EExtensionHook::First, nullptr, FMenuExtensionDelegate::CreateLambda([Extension](FMenuBuilder& MenuBuilder) { Extension->ExtendMenu(MenuBuilder, false); }));
 
 	return Extension;
 }
diff --git a/Engine/Source/Editor/Sequencer/Public/ISequencerChannelInterface.h b/Engine/Source/Editor/Sequencer/Public/ISequencerChannelInterface.h
--- a/Engine/Source/Editor/Sequencer/Public/ISequencerChannelInterface.h
+++ b/Engine/Source/Editor/Sequencer/Public/ISequencerChannelInterface.h
@@ -67,7 +67,7 @@ struct FExtendKeyMenuParams
 
 class ISidebarChannelExtension
 {
-	virtual TSharedPtr<ISidebarChannelExtension> ExtendMenu(FMenuBuilder& MenuBuilder, const bool bInSubMenu = false) = 0;
+	virtual TSharedPtr<ISidebarChannelExtension> ExtendMenu(FMenuBuilder& MenuBuilder, const bool bInSubMenu) = 0;
 };
 
 /**
