Sequencer: Update menu to read Data Layer Track

#rnx

diff --git a/Engine/Source/Editor/MovieSceneTools/Private/TrackEditors/DataLayerTrackEditor.cpp b/Engine/Source/Editor/MovieSceneTools/Private/TrackEditors/DataLayerTrackEditor.cpp
--- a/Engine/Source/Editor/MovieSceneTools/Private/TrackEditors/DataLayerTrackEditor.cpp
+++ b/Engine/Source/Editor/MovieSceneTools/Private/TrackEditors/DataLayerTrackEditor.cpp
@@ -293,7 +293,7 @@ TSharedRef<ISequencerSection> FDataLayerTrackEditor::MakeSectionInterface(UMovie
 void FDataLayerTrackEditor::BuildAddTrackMenu(FMenuBuilder& MenuBuilder)
 {
 	MenuBuilder.AddMenuEntry(
-		LOCTEXT("AddTrack", "Data Layer"),
+		LOCTEXT("AddTrack", "Data Layer Track"),
 		LOCTEXT("AddTrackToolTip", "Adds a new track that can load, activate or unload Data Layers in a World Partition world."),
 		FSlateIcon(FAppStyle::GetAppStyleSetName(), "Sequencer.Tracks.DataLayer"),
 		FUIAction(FExecuteAction::CreateRaw(this, &FDataLayerTrackEditor::HandleAddTrack)));
