[Sequencer] Remove default value for virtual function

diff --git a/Engine/Plugins/Experimental/PropertyAnimator/Source/PropertyAnimatorEditor/Private/Sequencer/PropertyAnimatorEditorCurveChannelInterface.h b/Engine/Plugins/Experimental/PropertyAnimator/Source/PropertyAnimatorEditor/Private/Sequencer/PropertyAnimatorEditorCurveChannelInterface.h
--- a/Engine/Plugins/Experimental/PropertyAnimator/Source/PropertyAnimatorEditor/Private/Sequencer/PropertyAnimatorEditorCurveChannelInterface.h
+++ b/Engine/Plugins/Experimental/PropertyAnimator/Source/PropertyAnimatorEditor/Private/Sequencer/PropertyAnimatorEditorCurveChannelInterface.h
@@ -53,7 +53,7 @@ void TPropertyAnimatorEditorCurveChannelInterface<InChannelType, InMenuExtension
 	InMenuExtender->AddMenuExtension(TEXT("SequencerChannels"), EExtensionHook::First, nullptr
 		, FMenuExtensionDelegate::CreateLambda([Extension](FMenuBuilder& InInnerMenuBuilder)
 			{
-				Extension->ExtendMenu(InInnerMenuBuilder);
+				Extension->ExtendMenu(InInnerMenuBuilder, true);
 			}));
 }
 
diff --git a/Engine/Plugins/Experimental/PropertyAnimator/Source/PropertyAnimatorEditor/Private/Sequencer/PropertyAnimatorEditorCurveSectionMenuExtension.h b/Engine/Plugins/Experimental/PropertyAnimator/Source/PropertyAnimatorEditor/Private/Sequencer/PropertyAnimatorEditorCurveSectionMenuExtension.h
--- a/Engine/Plugins/Experimental/PropertyAnimator/Source/PropertyAnimatorEditor/Private/Sequencer/PropertyAnimatorEditorCurveSectionMenuExtension.h
+++ b/Engine/Plugins/Experimental/PropertyAnimator/Source/PropertyAnimatorEditor/Private/Sequencer/PropertyAnimatorEditorCurveSectionMenuExtension.h
@@ -36,7 +36,7 @@ public:
 
 	virtual bool GetParameterStructData(FMovieSceneChannelHandle InChannelHandle, UStruct*& OutStruct, uint8*& OutData) const = 0;
 
-	virtual TSharedPtr<ISidebarChannelExtension> ExtendMenu(FMenuBuilder& InMenuBuilder, const bool bInSubMenu = true) override;
+	virtual TSharedPtr<ISidebarChannelExtension> ExtendMenu(FMenuBuilder& InMenuBuilder, const bool bInSubMenu) override;
 
 private:
 	void Initialize();
diff --git a/Engine/Source/Editor/MovieSceneTools/Public/Channels/PerlinNoiseChannelInterface.h b/Engine/Source/Editor/MovieSceneTools/Public/Channels/PerlinNoiseChannelInterface.h
--- a/Engine/Source/Editor/MovieSceneTools/Public/Channels/PerlinNoiseChannelInterface.h
+++ b/Engine/Source/Editor/MovieSceneTools/Public/Channels/PerlinNoiseChannelInterface.h
@@ -49,7 +49,7 @@ struct FPerlinNoiseChannelSectionMenuExtension : TSharedFromThis<FPerlinNoiseCha
 	FPerlinNoiseChannelSectionMenuExtension(TArrayView<const FMovieSceneChannelHandle> InChannelHandles, const TArray<TWeakObjectPtr<UMovieSceneSection>>& InWeakSections);
 	virtual ~FPerlinNoiseChannelSectionMenuExtension() {}
 
-	virtual TSharedPtr<ISidebarChannelExtension> ExtendMenu(FMenuBuilder& MenuBuilder, const bool bInSubMenu = true) override;
+	virtual TSharedPtr<ISidebarChannelExtension> ExtendMenu(FMenuBuilder& MenuBuilder, const bool bInSubMenu) override;
 
 private:
 
@@ -132,7 +132,7 @@ struct TPerlinNoiseChannelInterface : ISequencerChannelInterface
 		InMenuExtender->AddMenuExtension(TEXT("SequencerChannels"), EExtensionHook::First, nullptr
 			, FMenuExtensionDelegate::CreateLambda([Extension](FMenuBuilder& MenuBuilder)
 				{
-					Extension->ExtendMenu(MenuBuilder);
+					Extension->ExtendMenu(MenuBuilder, true);
 				}));
 	}
 
diff --git a/Engine/Source/Editor/Sequencer/Private/Menus/CurveChannelSectionSidebarExtension.h b/Engine/Source/Editor/Sequencer/Private/Menus/CurveChannelSectionSidebarExtension.h
--- a/Engine/Source/Editor/Sequencer/Private/Menus/CurveChannelSectionSidebarExtension.h
+++ b/Engine/Source/Editor/Sequencer/Private/Menus/CurveChannelSectionSidebarExtension.h
@@ -19,7 +19,7 @@ struct FCurveChannelSectionSidebarExtension : TSharedFromThis<FCurveChannelSecti
 
 	void AddSections(const TArray<TWeakObjectPtr<UMovieSceneSection>>& InWeakSections);
 
-	virtual TSharedPtr<ISidebarChannelExtension> ExtendMenu(FMenuBuilder& MenuBuilder, const bool bInSubMenu = true) override;
+	virtual TSharedPtr<ISidebarChannelExtension> ExtendMenu(FMenuBuilder& MenuBuilder, const bool bInSubMenu) override;
 
 private:
 	void AddDisplayOptionsMenu(FMenuBuilder& MenuBuilder);
diff --git a/Engine/Source/Editor/Sequencer/Private/Widgets/Sidebar/SequencerSelectionDrawer.cpp b/Engine/Source/Editor/Sequencer/Private/Widgets/Sidebar/SequencerSelectionDrawer.cpp
--- a/Engine/Source/Editor/Sequencer/Private/Widgets/Sidebar/SequencerSelectionDrawer.cpp
+++ b/Engine/Source/Editor/Sequencer/Private/Widgets/Sidebar/SequencerSelectionDrawer.cpp
@@ -407,7 +407,7 @@ void FSequencerSelectionDrawer::BuildOutlinerDetails(const TSharedRef<FSequencer
 		// Curve Channel Options (Pre-Finity, Post-Finity, etc.)
 		CurveChannelExtension = MakeShared<FCurveChannelSectionSidebarExtension>(Sequencer);
 		CurveChannelExtension->AddSections(WeakSceneSections);
-		CurveChannelExtension->ExtendMenu(MenuBuilder);
+		CurveChannelExtension->ExtendMenu(MenuBuilder, false);
 	}
 }
 
