Sequencer: Fix MacArm64-specific compiler error from 35590453.

diff --git a/Engine/Source/Editor/MovieSceneTools/Private/Conditions/MovieSceneConditionCustomization.cpp b/Engine/Source/Editor/MovieSceneTools/Private/Conditions/MovieSceneConditionCustomization.cpp
--- a/Engine/Source/Editor/MovieSceneTools/Private/Conditions/MovieSceneConditionCustomization.cpp
+++ b/Engine/Source/Editor/MovieSceneTools/Private/Conditions/MovieSceneConditionCustomization.cpp
@@ -262,17 +262,17 @@ void FMovieSceneConditionCustomization::FillDirectorBlueprintConditionSubMenu(FM
 			FUIAction(
 				FExecuteAction::CreateSPLambda(this, [SharedThis = StaticCastSharedRef<FMovieSceneConditionCustomization>(AsShared())]()
 				{
-					UMovieSceneSequence* Sequence = SharedThis->Sequence.Get();
-					if (Sequence)
+					UMovieSceneSequence* ThisSequence = SharedThis->Sequence.Get();
+					if (ThisSequence)
 					{
 						// Create a new director blueprint condition and set it in the details view. Use 'interactive change' so we don't early fire the property finished changing event and reset the details view mid-change
 						PropertyCustomizationHelpers::CreateNewInstanceOfEditInlineObjectClass(SharedThis->ConditionPropertyHandle.ToSharedRef(), UMovieSceneDirectorBlueprintCondition::StaticClass(), EPropertyValueSetFlags::InteractiveChange);
 						TSharedPtr<IPropertyHandle> DirectorBlueprintConditionHandle = SharedThis->ConditionPropertyHandle->GetChildHandle(TEXT("DirectorBlueprintConditionData"));
-						TSharedPtr<FMovieSceneDirectorBlueprintConditionCustomization> BlueprintConditionCustomization = FMovieSceneDirectorBlueprintConditionCustomization::MakeInstance(Sequence->GetMovieScene(), DirectorBlueprintConditionHandle, SharedThis->PropertyUtilities);
+						TSharedPtr<FMovieSceneDirectorBlueprintConditionCustomization> BlueprintConditionCustomization = FMovieSceneDirectorBlueprintConditionCustomization::MakeInstance(ThisSequence->GetMovieScene(), DirectorBlueprintConditionHandle, SharedThis->PropertyUtilities);
 						BlueprintConditionCustomization->CreateEndpoint();
 						SharedThis->PropertyUtilities->NotifyFinishedChangingProperties(FPropertyChangedEvent(SharedThis->ConditionPropertyHandle->GetProperty()));
 						SharedThis->PropertyUtilities->ForceRefresh();
-						Sequence->Modify();
+						ThisSequence->Modify();
 					}
 				}
 			)
