Sequencer: Removed deprecated methods from CinematicShotTrackEditor.h

diff --git a/Engine/Source/Editor/MovieSceneTools/Private/TrackEditors/CinematicShotTrackEditor.cpp b/Engine/Source/Editor/MovieSceneTools/Private/TrackEditors/CinematicShotTrackEditor.cpp
--- a/Engine/Source/Editor/MovieSceneTools/Private/TrackEditors/CinematicShotTrackEditor.cpp
+++ b/Engine/Source/Editor/MovieSceneTools/Private/TrackEditors/CinematicShotTrackEditor.cpp
@@ -214,28 +214,11 @@ void FCinematicShotTrackEditor::BuildTrackContextMenu(FMenuBuilder& MenuBuilder,
 	MenuBuilder.EndSection();
 }
 
-void FCinematicShotTrackEditor::InsertShot()
-{
-	UMovieScene* FocusedMovieScene = GetFocusedMovieScene();
-	UMovieSceneCinematicShotTrack* CinematicShotTrack = FocusedMovieScene ? FocusedMovieScene->FindTrack<UMovieSceneCinematicShotTrack>() : nullptr;
-	FSubTrackEditor::InsertSection(CinematicShotTrack);
-}
-
-void FCinematicShotTrackEditor::DuplicateShot(UMovieSceneCinematicShotSection* Section)
-{
-	FSubTrackEditor::DuplicateSection(Section);
-}
-
 void FCinematicShotTrackEditor::RenderShots(const TArray<UMovieSceneCinematicShotSection*>& Sections)
 {
 	GetSequencer()->RenderMovie(Sections);
 }
 
-void FCinematicShotTrackEditor::NewTake(UMovieSceneCinematicShotSection* Section)
-{
-	FSubTrackEditor::CreateNewTake(Section);
-}
-
 /* FSubTrackEditor
  *****************************************************************************/
 
@@ -276,20 +259,6 @@ TSubclassOf<UMovieSceneSubTrack> FCinematicShotTrackEditor::GetSubTrackClass() c
 	return UMovieSceneCinematicShotTrack::StaticClass();
 }
 
-TSharedRef<SWidget> FCinematicShotTrackEditor::HandleAddCinematicShotComboButtonGetMenuContent()
-{
-	UMovieScene* FocusedMovieScene = GetFocusedMovieScene();
-	UMovieSceneCinematicShotTrack* CinematicShotTrack = FocusedMovieScene ? FocusedMovieScene->FindTrack<UMovieSceneCinematicShotTrack>() : nullptr;
-	return HandleAddSubSequenceComboButtonGetMenuContent(CinematicShotTrack);
-}
-
-UMovieSceneCinematicShotTrack* FCinematicShotTrackEditor::FindOrCreateCinematicShotTrack()
-{
-	UMovieScene* FocusedMovieScene = GetFocusedMovieScene();
-	UMovieSceneCinematicShotTrack* CinematicShotTrack = FocusedMovieScene ? FocusedMovieScene->FindTrack<UMovieSceneCinematicShotTrack>() : nullptr;
-	return Cast<UMovieSceneCinematicShotTrack>(FindOrCreateSubTrack(FocusedMovieScene, CinematicShotTrack));
-}
-
 bool FCinematicShotTrackEditor::HandleAddSubTrackMenuEntryCanExecute() const
 {
 	UMovieScene* FocusedMovieScene = GetFocusedMovieScene();
diff --git a/Engine/Source/Editor/MovieSceneTools/Public/TrackEditors/CinematicShotTrackEditor.h b/Engine/Source/Editor/MovieSceneTools/Public/TrackEditors/CinematicShotTrackEditor.h
--- a/Engine/Source/Editor/MovieSceneTools/Public/TrackEditors/CinematicShotTrackEditor.h
+++ b/Engine/Source/Editor/MovieSceneTools/Public/TrackEditors/CinematicShotTrackEditor.h
@@ -57,17 +57,6 @@ public:
 	virtual void Tick(float DeltaTime) override;
 	virtual void BuildTrackContextMenu(FMenuBuilder& MenuBuilder, UMovieSceneTrack* Track) override;
 
-	UE_DEPRECATED(5.3, "InsertShot has been deprecated in favor of FSubTrackEditor::InsertSection(UMovieSceneTrack*)")
-	void InsertShot();
-	UE_DEPRECATED(5.3, "DuplicateShot has been deprecated in favor of FSubTrackEditor::DuplicateSection(UMovieSceneSubSection*)")
-	void DuplicateShot(UMovieSceneCinematicShotSection* Section);
-	UE_DEPRECATED(5.3, "RenameShot has been removed because it was unused")
-	void RenameShot(UMovieSceneCinematicShotSection* Section) {}
-	UE_DEPRECATED(5.3, "NewTake has been deprecated in favor of FSubTrackEditor::CreateNewTake(UMovieSceneSubSection*)")
-	void NewTake(UMovieSceneCinematicShotSection* Section);
-	UE_DEPRECATED(5.3, "InsertFiller has been removed because it is obsolete")
-	void InsertFiller() {}
-
 	/*
 	 * Render shots. 
 	 *
@@ -88,11 +77,6 @@ public:
 
 protected:
 
-	UE_DEPRECATED(5.3, "HandleAddCinematicShotComboButtonGetMenuContent has been deprecated. Please implement FSubTrackEditor::HandleAddSubSequenceComboButtonGetMenuContent(UMovieSceneTrack*) instead")
-	virtual TSharedRef<SWidget> HandleAddCinematicShotComboButtonGetMenuContent();
-	UE_DEPRECATED(5.3, "FindOrCreateCinematicShotTrack has been deprecated in favor of FSubTrackEditor::FindOrCreateSubTrack(UMovieScene* MovieScene, UMovieSceneTrack*)")
-	virtual UMovieSceneCinematicShotTrack* FindOrCreateCinematicShotTrack();
-
 	virtual bool HandleAddSubTrackMenuEntryCanExecute() const override;
 	virtual bool CanHandleAssetAdded(UMovieSceneSequence* Sequence) const override;
 
