Sequencer: Fixed deleting timewarp not modifying the section correctly

diff --git a/Engine/Source/Editor/MovieSceneTools/Private/TrackEditors/SkeletalAnimationTrackEditor.cpp b/Engine/Source/Editor/MovieSceneTools/Private/TrackEditors/SkeletalAnimationTrackEditor.cpp
--- a/Engine/Source/Editor/MovieSceneTools/Private/TrackEditors/SkeletalAnimationTrackEditor.cpp
+++ b/Engine/Source/Editor/MovieSceneTools/Private/TrackEditors/SkeletalAnimationTrackEditor.cpp
@@ -868,6 +868,7 @@ void FSkeletalAnimationSection::SlipSection(FFrameNumber SlipTime)
 
 bool FSkeletalAnimationSection::RequestDeleteKeyArea(const TArray<FName>& KeyAreaNamePath)
 {
+	Section.Modify();
 	Section.DeleteChannels(KeyAreaNamePath);
 	return true;
 }
diff --git a/Engine/Source/Editor/MovieSceneTools/Public/TrackEditors/SubTrackEditorBase.h b/Engine/Source/Editor/MovieSceneTools/Public/TrackEditors/SubTrackEditorBase.h
--- a/Engine/Source/Editor/MovieSceneTools/Public/TrackEditors/SubTrackEditorBase.h
+++ b/Engine/Source/Editor/MovieSceneTools/Public/TrackEditors/SubTrackEditorBase.h
@@ -381,6 +381,7 @@ void TSubSectionMixin<ParentSectionClass>::DilateSection(const TRange<FFrameNumb
 template<typename ParentSectionClass>
 bool TSubSectionMixin<ParentSectionClass>::RequestDeleteKeyArea( const TArray<FName>& KeyAreaNamePath )
 {
+	SubSectionObject.Modify();
 	SubSectionObject.DeleteChannels(KeyAreaNamePath);
 	return true;
 }
\ No newline at end of file
