Curve Editor: Add byte channel curve model and extrapolation for byte channel

diff --git a/Engine/Source/Editor/MovieSceneTools/Private/Channels/BuiltInChannelEditors.cpp b/Engine/Source/Editor/MovieSceneTools/Private/Channels/BuiltInChannelEditors.cpp
--- a/Engine/Source/Editor/MovieSceneTools/Private/Channels/BuiltInChannelEditors.cpp
+++ b/Engine/Source/Editor/MovieSceneTools/Private/Channels/BuiltInChannelEditors.cpp
@@ -32,11 +32,13 @@
 #include "Channels/MovieSceneFloatChannel.h"
 #include "Channels/MovieSceneDoubleChannel.h"
 #include "Channels/MovieSceneBoolChannel.h"
+#include "Channels/MovieSceneByteChannel.h"
 #include "Channels/MovieSceneIntegerChannel.h"
 #include "Channels/DoubleChannelCurveModel.h"
 #include "Channels/FloatChannelCurveModel.h"
 #include "Channels/IntegerChannelCurveModel.h"
 #include "Channels/BoolChannelCurveModel.h"
+#include "Channels/ByteChannelCurveModel.h"
 #include "Channels/TimeWarpChannelCurveModel.h"
 #include "Channels/PiecewiseCurveModel.h"
 #include "Variants/MovieScenePlayRateCurve.h"
@@ -1584,7 +1586,8 @@ struct FCurveChannelSectionMenuExtension : TSharedFromThis<FCurveChannelSectionM
 	}
 
 	void GetChannels(TArray<FMovieSceneFloatChannel*>& FloatChannels, TArray<FMovieSceneDoubleChannel*>& DoubleChannels,
-		TArray<FMovieSceneIntegerChannel*>& IntegerChannels, TArray<FMovieSceneBoolChannel*>& BoolChannels) const
+		TArray<FMovieSceneIntegerChannel*>& IntegerChannels, TArray<FMovieSceneBoolChannel*>& BoolChannels,
+		TArray<FMovieSceneByteChannel*>& ByteChannels) const
 	{
 		ISequencer* Sequencer = WeakSequencer.Pin().Get();
 		if (!Sequencer)
@@ -1618,10 +1621,15 @@ struct FCurveChannelSectionMenuExtension : TSharedFromThis<FCurveChannelSectionM
 				FMovieSceneBoolChannel* Channel = static_cast<FMovieSceneBoolChannel*>(Handle.Get());
 				BoolChannels.Add(Channel);
 			}
+			else if (Handle.GetChannelTypeName() == FMovieSceneByteChannel::StaticStruct()->GetFName())
+			{
+				FMovieSceneByteChannel* Channel = static_cast<FMovieSceneByteChannel*>(Handle.Get());
+				ByteChannels.Add(Channel);
+			}
 		}
 
 		// Otherwise, the channels of all the sections
-		if (FloatChannels.Num() + DoubleChannels.Num() + IntegerChannels.Num() + BoolChannels.Num() == 0)
+		if (FloatChannels.Num() + DoubleChannels.Num() + IntegerChannels.Num() + BoolChannels.Num() + ByteChannels.Num() == 0)
 		{
 			for (TWeakObjectPtr<UMovieSceneSection> WeakSection : WeakSections)
 			{
@@ -1644,6 +1652,10 @@ struct FCurveChannelSectionMenuExtension : TSharedFromThis<FCurveChannelSectionM
 					{
 						BoolChannels.Add(Channel);
 					}
+					for (FMovieSceneByteChannel* Channel : ChannelProxy.GetChannels<FMovieSceneByteChannel>())
+					{
+						ByteChannels.Add(Channel);
+					}
 				}
 			}
 		}
@@ -1655,10 +1667,11 @@ struct FCurveChannelSectionMenuExtension : TSharedFromThis<FCurveChannelSectionM
 		TArray<FMovieSceneDoubleChannel*> DoubleChannels;
 		TArray<FMovieSceneIntegerChannel*> IntegerChannels;
 		TArray<FMovieSceneBoolChannel*> BoolChannels;
+		TArray<FMovieSceneByteChannel*> ByteChannels;
 
-		GetChannels(FloatChannels, DoubleChannels,IntegerChannels,BoolChannels);
+		GetChannels(FloatChannels, DoubleChannels, IntegerChannels, BoolChannels, ByteChannels);
 
-		if (FloatChannels.Num() + DoubleChannels.Num() + IntegerChannels.Num() + BoolChannels.Num() == 0)
+		if (FloatChannels.Num() + DoubleChannels.Num() + IntegerChannels.Num() + BoolChannels.Num() + ByteChannels.Num() == 0)
 		{
 			return;
 		}
@@ -1701,6 +1714,12 @@ struct FCurveChannelSectionMenuExtension : TSharedFromThis<FCurveChannelSectionM
 			DestExtrap = ExtrapMode;
 			bAnythingChanged = true;
 		}
+		for (FMovieSceneByteChannel* Channel : ByteChannels)
+		{
+			TEnumAsByte<ERichCurveExtrapolation>& DestExtrap = bPreInfinity ? Channel->PreInfinityExtrap : Channel->PostInfinityExtrap;
+			DestExtrap = ExtrapMode;
+			bAnythingChanged = true;
+		}
 
 		if (bAnythingChanged)
 		{
@@ -1721,8 +1740,9 @@ struct FCurveChannelSectionMenuExtension : TSharedFromThis<FCurveChannelSectionM
 		TArray<FMovieSceneDoubleChannel*> DoubleChannels;
 		TArray<FMovieSceneIntegerChannel*> IntegerChannels;
 		TArray<FMovieSceneBoolChannel*> BoolChannels;
+		TArray<FMovieSceneByteChannel*> ByteChannels;
 
-		GetChannels(FloatChannels, DoubleChannels, IntegerChannels, BoolChannels);
+		GetChannels(FloatChannels, DoubleChannels, IntegerChannels, BoolChannels, ByteChannels);
 
 		for (FMovieSceneFloatChannel* Channel : FloatChannels)
 		{
@@ -1756,6 +1776,14 @@ struct FCurveChannelSectionMenuExtension : TSharedFromThis<FCurveChannelSectionM
 				return false;
 			}
 		}
+		for (FMovieSceneByteChannel* Channel : ByteChannels)
+		{
+			ERichCurveExtrapolation SourceExtrap = bPreInfinity ? Channel->PreInfinityExtrap : Channel->PostInfinityExtrap;
+			if (SourceExtrap != ExtrapMode)
+			{
+				return false;
+			}
+		}
 
 		return true;
 	}
@@ -1920,6 +1948,14 @@ void ExtendSectionMenu(FMenuBuilder& OuterMenuBuilder, TSharedPtr<FExtender> Men
 	MenuExtender->AddMenuExtension("SequencerChannels", EExtensionHook::First, nullptr, FMenuExtensionDelegate::CreateLambda([Extension](FMenuBuilder& MenuBuilder) { Extension->ExtendMenu(MenuBuilder, true); }));
 }
 
+void ExtendSectionMenu(FMenuBuilder& OuterMenuBuilder, TSharedPtr<FExtender> MenuExtender, TArray<TMovieSceneChannelHandle<FMovieSceneByteChannel>>&& Channels, const TArray<TWeakObjectPtr<UMovieSceneSection>>& InWeakSections, TWeakPtr<ISequencer> InWeakSequencer)
+{
+	TSharedRef<FCurveChannelSectionMenuExtension> Extension = FCurveChannelSectionMenuExtension::GetOrCreate(InWeakSequencer);
+	Extension->AddSections(InWeakSections);
+
+	MenuExtender->AddMenuExtension("SequencerChannels", EExtensionHook::First, nullptr, FMenuExtensionDelegate::CreateLambda([Extension](FMenuBuilder& MenuBuilder) { Extension->ExtendMenu(MenuBuilder, true); }));
+}
+
 TSharedPtr<ISidebarChannelExtension> ExtendSidebarMenu(FMenuBuilder& OuterMenuBuilder, TSharedPtr<FExtender> InMenuExtender, TArray<TMovieSceneChannelHandle<FMovieSceneFloatChannel>>&& Channels, const TArray<TWeakObjectPtr<UMovieSceneSection>>& InWeakSections, TWeakPtr<ISequencer> InWeakSequencer)
 {
 	TSharedRef<FCurveChannelSectionMenuExtension> Extension = FCurveChannelSectionMenuExtension::GetOrCreate(InWeakSequencer);
@@ -1960,6 +1996,16 @@ TSharedPtr<ISidebarChannelExtension> ExtendSidebarMenu(FMenuBuilder& OuterMenuBu
 	return Extension;
 }
 
+TSharedPtr<ISidebarChannelExtension> ExtendSidebarMenu(FMenuBuilder& OuterMenuBuilder, TSharedPtr<FExtender> InMenuExtender, TArray<TMovieSceneChannelHandle<FMovieSceneByteChannel>>&& Channels, const TArray<TWeakObjectPtr<UMovieSceneSection>>& InWeakSections, TWeakPtr<ISequencer> InWeakSequencer)
+{
+	TSharedRef<FCurveChannelSectionMenuExtension> Extension = FCurveChannelSectionMenuExtension::GetOrCreate(InWeakSequencer);
+	Extension->AddSections(InWeakSections);
+
+	InMenuExtender->AddMenuExtension("SequencerChannels", EExtensionHook::First, nullptr, FMenuExtensionDelegate::CreateLambda([Extension](FMenuBuilder& MenuBuilder) { Extension->ExtendMenu(MenuBuilder, false); }));
+
+	return Extension;
+}
+
 void ExtendKeyMenu(FMenuBuilder& OuterMenuBuilder, TSharedPtr<FExtender> MenuExtender, TArray<TExtendKeyMenuParams<FMovieSceneFloatChannel>>&& Channels, TWeakPtr<ISequencer> InSequencer)
 {
 	using ExtensionType = TCurveChannelKeyMenuExtension<FMovieSceneFloatChannel>;
@@ -2004,6 +2050,11 @@ TUniquePtr<FCurveModel> CreateCurveEditorModel(const TMovieSceneChannelHandle<FM
 	return MakeUnique<FBoolChannelCurveModel>(BoolChannel, Params.OwningSection, Params.Sequencer);
 }
 
+TUniquePtr<FCurveModel> CreateCurveEditorModel(const TMovieSceneChannelHandle<FMovieSceneByteChannel>& ByteChannel, const UE::Sequencer::FCreateCurveEditorModelParams& Params)
+{
+	return MakeUnique<FByteChannelCurveModel>(ByteChannel, Params.OwningSection, Params.Sequencer);
+}
+
 TUniquePtr<FCurveModel> CreateCurveEditorModel(const TMovieSceneChannelHandle<FMovieSceneEventChannel>& EventChannel, const UE::Sequencer::FCreateCurveEditorModelParams& Params)
 {
 	return MakeUnique<FEventChannelCurveModel>(EventChannel, Params.OwningSection, Params.Sequencer);
diff --git a/Engine/Source/Editor/MovieSceneTools/Private/Channels/BuiltInChannelEditors.h b/Engine/Source/Editor/MovieSceneTools/Private/Channels/BuiltInChannelEditors.h
--- a/Engine/Source/Editor/MovieSceneTools/Private/Channels/BuiltInChannelEditors.h
+++ b/Engine/Source/Editor/MovieSceneTools/Private/Channels/BuiltInChannelEditors.h
@@ -79,11 +79,13 @@ void ExtendSectionMenu(FMenuBuilder& OuterMenuBuilder, TSharedPtr<FExtender> Men
 void ExtendSectionMenu(FMenuBuilder& OuterMenuBuilder, TSharedPtr<FExtender> MenuExtender, TArray<TMovieSceneChannelHandle<FMovieSceneFloatChannel>>&& Channels, const TArray<TWeakObjectPtr<UMovieSceneSection>>& InWeakSections, TWeakPtr<ISequencer> InWeakSequencer);
 void ExtendSectionMenu(FMenuBuilder& OuterMenuBuilder, TSharedPtr<FExtender> MenuExtender, TArray<TMovieSceneChannelHandle<FMovieSceneIntegerChannel>>&& Channels, const TArray<TWeakObjectPtr<UMovieSceneSection>>& InWeakSections, TWeakPtr<ISequencer> InWeakSequencer);
 void ExtendSectionMenu(FMenuBuilder& OuterMenuBuilder, TSharedPtr<FExtender> MenuExtender, TArray<TMovieSceneChannelHandle<FMovieSceneBoolChannel>>&& Channels, const TArray<TWeakObjectPtr<UMovieSceneSection>>& InWeakSections, TWeakPtr<ISequencer> InWeakSequencer);
+void ExtendSectionMenu(FMenuBuilder& OuterMenuBuilder, TSharedPtr<FExtender> MenuExtender, TArray<TMovieSceneChannelHandle<FMovieSceneByteChannel>>&& Channels, const TArray<TWeakObjectPtr<UMovieSceneSection>>& InWeakSections, TWeakPtr<ISequencer> InWeakSequencer);
 
 TSharedPtr<ISidebarChannelExtension> ExtendSidebarMenu(FMenuBuilder& OuterMenuBuilder, TSharedPtr<FExtender> MenuExtender, TArray<TMovieSceneChannelHandle<FMovieSceneDoubleChannel>>&& InChannels, const TArray<TWeakObjectPtr<UMovieSceneSection>>& InWeakSections, TWeakPtr<ISequencer> InWeakSequencer);
 TSharedPtr<ISidebarChannelExtension> ExtendSidebarMenu(FMenuBuilder& OuterMenuBuilder, TSharedPtr<FExtender> MenuExtender, TArray<TMovieSceneChannelHandle<FMovieSceneFloatChannel>>&& InChannels,const  TArray<TWeakObjectPtr<UMovieSceneSection>>& InWeakSections, TWeakPtr<ISequencer> InWeakSequencer);
 TSharedPtr<ISidebarChannelExtension> ExtendSidebarMenu(FMenuBuilder& OuterMenuBuilder, TSharedPtr<FExtender> MenuExtender, TArray<TMovieSceneChannelHandle<FMovieSceneIntegerChannel>>&& InChannels, const TArray<TWeakObjectPtr<UMovieSceneSection>>& InWeakSections, TWeakPtr<ISequencer> InWeakSequencer);
 TSharedPtr<ISidebarChannelExtension> ExtendSidebarMenu(FMenuBuilder& OuterMenuBuilder, TSharedPtr<FExtender> MenuExtender, TArray<TMovieSceneChannelHandle<FMovieSceneBoolChannel>>&& InChannels, const TArray<TWeakObjectPtr<UMovieSceneSection>>& InWeakSections, TWeakPtr<ISequencer> InWeakSequencer);
+TSharedPtr<ISidebarChannelExtension> ExtendSidebarMenu(FMenuBuilder& OuterMenuBuilder, TSharedPtr<FExtender> MenuExtender, TArray<TMovieSceneChannelHandle<FMovieSceneByteChannel>>&& InChannels, const TArray<TWeakObjectPtr<UMovieSceneSection>>& InWeakSections, TWeakPtr<ISequencer> InWeakSequencer);
 
 void ExtendKeyMenu(FMenuBuilder& OuterMenuBuilder, TSharedPtr<FExtender> MenuExtender, TArray<TExtendKeyMenuParams<FMovieSceneDoubleChannel>>&& Channels, TWeakPtr<ISequencer> InSequencer);
 void ExtendKeyMenu(FMenuBuilder& OuterMenuBuilder, TSharedPtr<FExtender> MenuExtender, TArray<TExtendKeyMenuParams<FMovieSceneFloatChannel>>&& Channels, TWeakPtr<ISequencer> InSequencer);
@@ -94,6 +96,7 @@ inline bool SupportsCurveEditorModels(const TMovieSceneChannelHandle<FMovieScene
 inline bool SupportsCurveEditorModels(const TMovieSceneChannelHandle<FMovieSceneFloatChannel>& FloatChannel) { return true; }
 inline bool SupportsCurveEditorModels(const TMovieSceneChannelHandle<FMovieSceneIntegerChannel>& IntegerChannel) { return true; }
 inline bool SupportsCurveEditorModels(const TMovieSceneChannelHandle<FMovieSceneBoolChannel>& BoolChannel) { return true; }
+inline bool SupportsCurveEditorModels(const TMovieSceneChannelHandle<FMovieSceneByteChannel>& ByteChannel) { return true; }
 inline bool SupportsCurveEditorModels(const TMovieSceneChannelHandle<FMovieSceneEventChannel>& EventChannel) { return true; }
 inline bool SupportsCurveEditorModels(const TMovieSceneChannelHandle<FMovieSceneTimeWarpChannel>& TimeWarpChannel) { return true; }
 
@@ -101,6 +104,7 @@ TUniquePtr<FCurveModel> CreateCurveEditorModel(const TMovieSceneChannelHandle<FM
 TUniquePtr<FCurveModel> CreateCurveEditorModel(const TMovieSceneChannelHandle<FMovieSceneFloatChannel>& FloatChannel, const UE::Sequencer::FCreateCurveEditorModelParams& Params);
 TUniquePtr<FCurveModel> CreateCurveEditorModel(const TMovieSceneChannelHandle<FMovieSceneIntegerChannel>& IntegerChannel, const UE::Sequencer::FCreateCurveEditorModelParams& Params);
 TUniquePtr<FCurveModel> CreateCurveEditorModel(const TMovieSceneChannelHandle<FMovieSceneBoolChannel>& BoolChannel, const UE::Sequencer::FCreateCurveEditorModelParams& Params);
+TUniquePtr<FCurveModel> CreateCurveEditorModel(const TMovieSceneChannelHandle<FMovieSceneByteChannel>& ByteChannel, const UE::Sequencer::FCreateCurveEditorModelParams& Params);
 TUniquePtr<FCurveModel> CreateCurveEditorModel(const TMovieSceneChannelHandle<FMovieSceneEventChannel>& EventChannel, const UE::Sequencer::FCreateCurveEditorModelParams& Params);
 TUniquePtr<FCurveModel> CreateCurveEditorModel(const TMovieSceneChannelHandle<FMovieSceneTimeWarpChannel>& TimeWarpChannel, const UE::Sequencer::FCreateCurveEditorModelParams& Params);
 
diff --git a/Engine/Source/Editor/MovieSceneTools/Private/Channels/ByteChannelCurveModel.cpp b/Engine/Source/Editor/MovieSceneTools/Private/Channels/ByteChannelCurveModel.cpp
new file mode 100644
--- /dev/null
+++ b/Engine/Source/Editor/MovieSceneTools/Private/Channels/ByteChannelCurveModel.cpp
@@ -0,0 +1,171 @@
+// Copyright Epic Games, Inc. All Rights Reserved.
+
+#include "Channels/ByteChannelCurveModel.h"
+
+#include "Algo/BinarySearch.h"
+#include "Channels/ByteChannelKeyProxy.h"
+#include "Channels/MovieSceneChannelData.h"
+#include "Channels/MovieSceneChannelHandle.h"
+#include "Channels/MovieSceneByteChannel.h"
+#include "Containers/Array.h"
+#include "Containers/UnrealString.h"
+#include "CurveDataAbstraction.h"
+#include "CurveEditorScreenSpace.h"
+#include "Curves/KeyHandle.h"
+#include "HAL/PlatformCrt.h"
+#include "IBufferedCurveModel.h"
+#include "Internationalization/Text.h"
+#include "Math/Range.h"
+#include "Misc/FrameNumber.h"
+#include "Misc/FrameRate.h"
+#include "Misc/FrameTime.h"
+#include "MovieScene.h"
+#include "MovieSceneSection.h"
+#include "Templates/Casts.h"
+#include "Templates/Tuple.h"
+#include "Templates/UnrealTemplate.h"
+#include "UObject/Package.h"
+#include "UObject/UObjectGlobals.h"
+#include "UObject/UnrealNames.h"
+#include "UObject/WeakObjectPtr.h"
+#include "UObject/WeakObjectPtrTemplates.h"
+
+class FCurveEditor;
+class UObject;
+
+/**
+ * Buffered curve implementation for a byte channel curve model, stores a copy of the byte channel in order to draw itself.
+ */
+class FByteChannelBufferedCurveModel : public IBufferedCurveModel
+{
+public:
+	/** Create a copy of the byte channel while keeping the reference to the section */
+	FByteChannelBufferedCurveModel(const FMovieSceneByteChannel* InMovieSceneByteChannel, TWeakObjectPtr<UMovieSceneSection> InWeakSection,
+		TArray<FKeyPosition>&& InKeyPositions, TArray<FKeyAttributes>&& InKeyAttributes, const FString& InLongDisplayName, const double InValueMin, const double InValueMax)
+		: IBufferedCurveModel(MoveTemp(InKeyPositions), MoveTemp(InKeyAttributes), InLongDisplayName, InValueMin, InValueMax)
+		, Channel(*InMovieSceneByteChannel)
+		, WeakSection(InWeakSection)
+	{}
+
+	virtual void DrawCurve(const FCurveEditor& InCurveEditor, const FCurveEditorScreenSpace& InScreenSpace, TArray<TTuple<double, double>>& OutInterpolatingPoints) const override
+	{
+		UMovieSceneSection* Section = WeakSection.Get();
+
+		if (Section)
+		{
+			FFrameRate TickResolution = Section->GetTypedOuter<UMovieScene>()->GetTickResolution();
+
+			TMovieSceneChannelData<const uint8> ChannelData = Channel.GetData();
+			TArrayView<const FFrameNumber> Times = ChannelData.GetTimes();
+			TArrayView<const uint8> Values = ChannelData.GetValues();
+
+			const double StartTimeSeconds = InScreenSpace.GetInputMin();
+			const double EndTimeSeconds = InScreenSpace.GetInputMax();
+
+			const FFrameNumber StartFrame = (StartTimeSeconds * TickResolution).FloorToFrame();
+			const FFrameNumber EndFrame = (EndTimeSeconds * TickResolution).CeilToFrame();
+
+			const int32 StartingIndex = Algo::UpperBound(Times, StartFrame);
+			const int32 EndingIndex = Algo::LowerBound(Times, EndFrame);
+
+			for (int32 KeyIndex = StartingIndex; KeyIndex < EndingIndex; ++KeyIndex)
+			{
+				OutInterpolatingPoints.Add(MakeTuple(Times[KeyIndex] / TickResolution, double(Values[KeyIndex])));
+			}
+		}
+	}
+
+private:
+	FMovieSceneByteChannel Channel;
+	TWeakObjectPtr<UMovieSceneSection> WeakSection;
+};
+
+FByteChannelCurveModel::FByteChannelCurveModel(TMovieSceneChannelHandle<FMovieSceneByteChannel> InChannel, UMovieSceneSection* OwningSection, TWeakPtr<ISequencer> InWeakSequencer)
+	: FChannelCurveModel<FMovieSceneByteChannel, uint8, uint8>(InChannel, OwningSection, InWeakSequencer)
+{
+}
+
+void FByteChannelCurveModel::CreateKeyProxies(TArrayView<const FKeyHandle> InKeyHandles, TArrayView<UObject*> OutObjects)
+{
+	for (int32 Index = 0; Index < InKeyHandles.Num(); ++Index)
+	{
+		UByteChannelKeyProxy* NewProxy = NewObject<UByteChannelKeyProxy>(GetTransientPackage(), NAME_None);
+
+		NewProxy->Initialize(InKeyHandles[Index], GetChannelHandle(), this->GetOwningObjectOrOuter<UMovieSceneSection>());
+		OutObjects[Index] = NewProxy;
+	}
+}
+
+TUniquePtr<IBufferedCurveModel> FByteChannelCurveModel::CreateBufferedCurveCopy() const
+{
+	FMovieSceneByteChannel* Channel = GetChannelHandle().Get();
+	if (Channel)
+	{
+		TArray<FKeyHandle> TargetKeyHandles;
+		TMovieSceneChannelData<uint8> ChannelData = Channel->GetData();
+
+		TRange<FFrameNumber> TotalRange = ChannelData.GetTotalRange();
+		ChannelData.GetKeys(TotalRange, nullptr, &TargetKeyHandles);
+
+		TArray<FKeyPosition> KeyPositions;
+		KeyPositions.SetNumUninitialized(GetNumKeys());
+		TArray<FKeyAttributes> KeyAttributes;
+		KeyAttributes.SetNumUninitialized(GetNumKeys());
+		GetKeyPositions(TargetKeyHandles, KeyPositions);
+		GetKeyAttributes(TargetKeyHandles, KeyAttributes);
+
+		double ValueMin = 0.f, ValueMax = 1.f;
+		GetValueRange(ValueMin, ValueMax);
+
+		return MakeUnique<FByteChannelBufferedCurveModel>(Channel, this->GetOwningObjectOrOuter<UMovieSceneSection>(), MoveTemp(KeyPositions), MoveTemp(KeyAttributes), GetLongDisplayName().ToString(), ValueMin, ValueMax);
+	}
+	return nullptr;
+}
+
+void FByteChannelCurveModel::GetCurveAttributes(FCurveAttributes& OutCurveAttributes) const
+{
+	FMovieSceneByteChannel* Channel = GetChannelHandle().Get();
+	if (Channel)
+	{
+		OutCurveAttributes.SetPreExtrapolation(Channel->PreInfinityExtrap);
+		OutCurveAttributes.SetPostExtrapolation(Channel->PostInfinityExtrap);
+	}
+}
+
+void FByteChannelCurveModel::SetCurveAttributes(const FCurveAttributes& InCurveAttributes)
+{
+	FMovieSceneByteChannel* Channel = GetChannelHandle().Get();
+	UMovieSceneSection* Section = this->GetOwningObjectOrOuter<UMovieSceneSection>();
+	if (Channel && Section && !IsReadOnly())
+	{
+		Section->MarkAsChanged();
+
+		if (InCurveAttributes.HasPreExtrapolation())
+		{
+			Channel->PreInfinityExtrap = InCurveAttributes.GetPreExtrapolation();
+		}
+
+		if (InCurveAttributes.HasPostExtrapolation())
+		{
+			Channel->PostInfinityExtrap = InCurveAttributes.GetPostExtrapolation();
+		}
+
+		CurveModifiedDelegate.Broadcast();
+	}
+}
+
+double FByteChannelCurveModel::GetKeyValue(TArrayView<const uint8> Values, int32 Index) const
+{
+	return (double)Values[Index];
+}
+
+void FByteChannelCurveModel::SetKeyValue(int32 Index, double KeyValue) const
+{
+	FMovieSceneByteChannel* Channel = GetChannelHandle().Get();
+
+	if (Channel)
+	{
+		TMovieSceneChannelData<uint8> ChannelData = Channel->GetData();
+		ChannelData.GetValues()[Index] = (uint8)KeyValue;
+	}
+}
diff --git a/Engine/Source/Editor/MovieSceneTools/Private/Channels/ByteChannelKeyProxy.cpp b/Engine/Source/Editor/MovieSceneTools/Private/Channels/ByteChannelKeyProxy.cpp
new file mode 100644
--- /dev/null
+++ b/Engine/Source/Editor/MovieSceneTools/Private/Channels/ByteChannelKeyProxy.cpp
@@ -0,0 +1,27 @@
+// Copyright Epic Games, Inc. All Rights Reserved.
+
+#include "Channels/ByteChannelKeyProxy.h"
+
+#include "Channels/MovieSceneByteChannel.h"
+#include "HAL/PlatformCrt.h"
+
+struct FPropertyChangedEvent;
+
+void UByteChannelKeyProxy::Initialize(FKeyHandle InKeyHandle, TMovieSceneChannelHandle<FMovieSceneByteChannel> InChannelHandle, TWeakObjectPtr<UMovieSceneSection> InWeakSection)
+{
+	KeyHandle     = InKeyHandle;
+	ChannelHandle = InChannelHandle;
+	WeakSection   = InWeakSection;
+}
+
+void UByteChannelKeyProxy::PostEditChangeProperty(FPropertyChangedEvent& PropertyChangedEvent)
+{
+	Super::PostEditChangeProperty(PropertyChangedEvent);
+
+	OnProxyValueChanged(ChannelHandle, WeakSection.Get(), KeyHandle, Value, Time);
+}
+
+void UByteChannelKeyProxy::UpdateValuesFromRawData()
+{
+	RefreshCurrentValue(ChannelHandle, KeyHandle, Value, Time);
+}
\ No newline at end of file
diff --git a/Engine/Source/Editor/MovieSceneTools/Private/Channels/ByteChannelKeyProxy.h b/Engine/Source/Editor/MovieSceneTools/Private/Channels/ByteChannelKeyProxy.h
new file mode 100644
--- /dev/null
+++ b/Engine/Source/Editor/MovieSceneTools/Private/Channels/ByteChannelKeyProxy.h
@@ -0,0 +1,58 @@
+// Copyright Epic Games, Inc. All Rights Reserved.
+
+#pragma once
+
+#include "Channels/MovieSceneChannelHandle.h"
+#include "Channels/MovieSceneByteChannel.h"
+#include "CoreTypes.h"
+#include "CurveEditorKeyProxy.h"
+#include "Curves/KeyHandle.h"
+#include "Misc/FrameNumber.h"
+#include "MovieSceneKeyProxy.h"
+#include "UObject/Object.h"
+#include "UObject/UObjectGlobals.h"
+#include "UObject/WeakObjectPtrTemplates.h"
+
+#include "ByteChannelKeyProxy.generated.h"
+
+class UMovieSceneSection;
+struct FPropertyChangedEvent;
+
+UCLASS()
+class UByteChannelKeyProxy : public UObject, public ICurveEditorKeyProxy, public IMovieSceneKeyProxy
+{
+public:
+	GENERATED_BODY()
+
+	/**
+	 * Initialize this key proxy object by caching the underlying key object, and retrieving the time/value each tick
+	 */
+	void Initialize(FKeyHandle InKeyHandle, TMovieSceneChannelHandle<FMovieSceneByteChannel> InChannelHandle, TWeakObjectPtr<UMovieSceneSection> InWeakSection);
+
+private:
+
+	/** Apply this class's properties to the underlying key */
+	virtual void PostEditChangeProperty(FPropertyChangedEvent& PropertyChangedEvent) override;
+
+	/** Update this class's properties from the underlying key */
+	virtual void UpdateValuesFromRawData() override;
+
+private:
+
+	/** User-facing time of the key, applied to the actual key on PostEditChange, and updated every tick */
+	UPROPERTY(EditAnywhere, Category="Key")
+	FFrameNumber Time;
+
+	/** User-facing value of the key, applied to the actual key on PostEditChange, and updated every tick */
+	UPROPERTY(EditAnywhere, Category="Key", meta=(ShowOnlyInnerProperties))
+	uint8 Value;
+
+private:
+
+	/** Cached key handle that this key proxy relates to */
+	FKeyHandle KeyHandle;
+	/** Cached channel in which the key resides */
+	TMovieSceneChannelHandle<FMovieSceneByteChannel> ChannelHandle;
+	/** Cached section in which the channel resides */
+	TWeakObjectPtr<UMovieSceneSection> WeakSection;
+};
\ No newline at end of file
diff --git a/Engine/Source/Editor/MovieSceneTools/Private/Channels/ChannelCurveModel.cpp b/Engine/Source/Editor/MovieSceneTools/Private/Channels/ChannelCurveModel.cpp
--- a/Engine/Source/Editor/MovieSceneTools/Private/Channels/ChannelCurveModel.cpp
+++ b/Engine/Source/Editor/MovieSceneTools/Private/Channels/ChannelCurveModel.cpp
@@ -7,6 +7,7 @@
 #include "Channels/MovieSceneChannelProxy.h"
 #include "Channels/MovieSceneFloatChannel.h"
 #include "Channels/MovieSceneIntegerChannel.h"
+#include "Channels/MovieSceneByteChannel.h"
 #include "Containers/UnrealString.h"
 #include "CurveDataAbstraction.h"
 #include "CurveDrawInfo.h"
@@ -605,3 +606,4 @@ template class FChannelCurveModel<FMovieSceneDoubleChannel, FMovieSceneDoubleVal
 template class FChannelCurveModel<FMovieSceneFloatChannel, FMovieSceneFloatValue, float>;
 template class FChannelCurveModel<FMovieSceneIntegerChannel, int32, int32>;
 template class FChannelCurveModel<FMovieSceneBoolChannel, bool, bool>;
+template class FChannelCurveModel<FMovieSceneByteChannel, uint8, uint8>;
diff --git a/Engine/Source/Editor/MovieSceneTools/Public/Channels/ByteChannelCurveModel.h b/Engine/Source/Editor/MovieSceneTools/Public/Channels/ByteChannelCurveModel.h
new file mode 100644
--- /dev/null
+++ b/Engine/Source/Editor/MovieSceneTools/Public/Channels/ByteChannelCurveModel.h
@@ -0,0 +1,37 @@
+// Copyright Epic Games, Inc. All Rights Reserved.
+
+#pragma once
+
+#include "Channels/ChannelCurveModel.h"
+#include "Channels/MovieSceneChannelHandle.h"
+#include "Channels/MovieSceneByteChannel.h"
+#include "Containers/ArrayView.h"
+#include "CoreTypes.h"
+#include "IBufferedCurveModel.h"
+#include "Templates/UniquePtr.h"
+
+class IBufferedCurveModel;
+class ISequencer;
+class UMovieSceneSection;
+class UObject;
+struct FKeyHandle;
+struct FMovieSceneByteChannel;
+template <typename ChannelType> struct TMovieSceneChannelHandle;
+
+class FByteChannelCurveModel : public FChannelCurveModel<FMovieSceneByteChannel, uint8, uint8>
+{
+public:
+	FByteChannelCurveModel(TMovieSceneChannelHandle<FMovieSceneByteChannel> InChannel, UMovieSceneSection* InOwningSection, TWeakPtr<ISequencer> InWeakSequencer);
+
+	// FCurveModel
+	virtual void CreateKeyProxies(TArrayView<const FKeyHandle> InKeyHandles, TArrayView<UObject*> OutObjects) override;
+	virtual TUniquePtr<IBufferedCurveModel> CreateBufferedCurveCopy() const override;
+	virtual void GetCurveAttributes(FCurveAttributes& OutAttributes) const override;
+	virtual void SetCurveAttributes(const FCurveAttributes& InAttributes) override;
+protected:
+
+	// FChannelCurveModel
+	virtual double GetKeyValue(TArrayView<const uint8> Values, int32 Index) const override;
+	virtual void SetKeyValue(int32 Index, double KeyValue) const override;
+
+};
\ No newline at end of file
diff --git a/Engine/Source/Editor/Sequencer/Private/Menus/CurveChannelSectionSidebarExtension.cpp b/Engine/Source/Editor/Sequencer/Private/Menus/CurveChannelSectionSidebarExtension.cpp
--- a/Engine/Source/Editor/Sequencer/Private/Menus/CurveChannelSectionSidebarExtension.cpp
+++ b/Engine/Source/Editor/Sequencer/Private/Menus/CurveChannelSectionSidebarExtension.cpp
@@ -17,6 +17,7 @@
 #include "Channels/MovieSceneFloatChannel.h"
 #include "Channels/MovieSceneDoubleChannel.h"
 #include "Channels/MovieSceneBoolChannel.h"
+#include "Channels/MovieSceneByteChannel.h"
 #include "Channels/MovieSceneIntegerChannel.h"
 #include "Curves/RealCurve.h"
 #include "Modules/ModuleManager.h"
@@ -247,7 +248,8 @@ void FCurveChannelSectionSidebarExtension::AddExtrapolationMenu(FMenuBuilder& Me
 }
 
 void FCurveChannelSectionSidebarExtension::GetChannels(TArray<FMovieSceneFloatChannel*>& FloatChannels, TArray<FMovieSceneDoubleChannel*>& DoubleChannels,
-	TArray<FMovieSceneIntegerChannel*>& IntegerChannels, TArray<FMovieSceneBoolChannel*>& BoolChannels) const
+	TArray<FMovieSceneIntegerChannel*>& IntegerChannels, TArray<FMovieSceneBoolChannel*>& BoolChannels,
+	TArray<FMovieSceneByteChannel*>& ByteChannels) const
 {
 	const TSharedPtr<ISequencer> Sequencer = WeakSequencer.Pin();
 	if (!Sequencer)
@@ -281,10 +283,15 @@ void FCurveChannelSectionSidebarExtension::GetChannels(TArray<FMovieSceneFloatCh
 			FMovieSceneBoolChannel* const Channel = static_cast<FMovieSceneBoolChannel*>(Handle.Get());
 			BoolChannels.Add(Channel);
 		}
+		else if (Handle.GetChannelTypeName() == FMovieSceneByteChannel::StaticStruct()->GetFName())
+		{
+			FMovieSceneByteChannel* const Channel = static_cast<FMovieSceneByteChannel*>(Handle.Get());
+			ByteChannels.Add(Channel);
+		}
 	}
 
 	// Otherwise, the channels of all the sections
-	if (FloatChannels.Num() + DoubleChannels.Num() + IntegerChannels.Num() + BoolChannels.Num() == 0)
+	if (FloatChannels.Num() + DoubleChannels.Num() + IntegerChannels.Num() + BoolChannels.Num() + ByteChannels.Num() == 0)
 	{
 		for (const TWeakObjectPtr<UMovieSceneSection>& WeakSection : WeakSections)
 		{
@@ -307,6 +314,10 @@ void FCurveChannelSectionSidebarExtension::GetChannels(TArray<FMovieSceneFloatCh
 				{
 					BoolChannels.Add(Channel);
 				}
+				for (FMovieSceneByteChannel* const Channel : ChannelProxy.GetChannels<FMovieSceneByteChannel>())
+				{
+					ByteChannels.Add(Channel);
+				}
 			}
 		}
 	}
@@ -318,10 +329,11 @@ void FCurveChannelSectionSidebarExtension::SetExtrapolationMode(const ERichCurve
 	TArray<FMovieSceneDoubleChannel*> DoubleChannels;
 	TArray<FMovieSceneIntegerChannel*> IntegerChannels;
 	TArray<FMovieSceneBoolChannel*> BoolChannels;
-	
-	GetChannels(FloatChannels, DoubleChannels,IntegerChannels,BoolChannels);
+	TArray<FMovieSceneByteChannel*> ByteChannels;
 
-	if (FloatChannels.Num() + DoubleChannels.Num() + IntegerChannels.Num() + BoolChannels.Num() == 0)
+	GetChannels(FloatChannels, DoubleChannels, IntegerChannels, BoolChannels, ByteChannels);
+
+	if (FloatChannels.Num() + DoubleChannels.Num() + IntegerChannels.Num() + BoolChannels.Num() + ByteChannels.Num() == 0)
 	{
 		return;
 	}
@@ -365,6 +377,12 @@ void FCurveChannelSectionSidebarExtension::SetExtrapolationMode(const ERichCurve
 		DestExtrap = InExtrapolation;
 		bAnythingChanged = true;
 	}
+	for (FMovieSceneByteChannel* const Channel : ByteChannels)
+	{
+		TEnumAsByte<ERichCurveExtrapolation>& DestExtrap = bInPreInfinity ? Channel->PreInfinityExtrap : Channel->PostInfinityExtrap;
+		DestExtrap = InExtrapolation;
+		bAnythingChanged = true;
+	}
 
 	if (bAnythingChanged)
 	{
@@ -385,8 +403,9 @@ bool FCurveChannelSectionSidebarExtension::IsExtrapolationModeSelected(const ERi
 	TArray<FMovieSceneDoubleChannel*> DoubleChannels;
 	TArray<FMovieSceneIntegerChannel*> IntegerChannels;
 	TArray<FMovieSceneBoolChannel*> BoolChannels;
-	
-	GetChannels(FloatChannels, DoubleChannels, IntegerChannels, BoolChannels);
+	TArray<FMovieSceneByteChannel*> ByteChannels;
+
+	GetChannels(FloatChannels, DoubleChannels, IntegerChannels, BoolChannels, ByteChannels);
 
 	for (FMovieSceneFloatChannel* const Channel : FloatChannels)
 	{
@@ -420,6 +439,14 @@ bool FCurveChannelSectionSidebarExtension::IsExtrapolationModeSelected(const ERi
 			return false;
 		}
 	}
+	for (FMovieSceneByteChannel* const Channel : ByteChannels)
+	{
+		const ERichCurveExtrapolation SourceExtrapolation = bInPreInfinity ? Channel->PreInfinityExtrap : Channel->PostInfinityExtrap;
+		if (SourceExtrapolation != InExtrapolation)
+		{
+			return false;
+		}
+	}
 
 	return true;
 }
diff --git a/Engine/Source/Editor/Sequencer/Private/Menus/CurveChannelSectionSidebarExtension.h b/Engine/Source/Editor/Sequencer/Private/Menus/CurveChannelSectionSidebarExtension.h
--- a/Engine/Source/Editor/Sequencer/Private/Menus/CurveChannelSectionSidebarExtension.h
+++ b/Engine/Source/Editor/Sequencer/Private/Menus/CurveChannelSectionSidebarExtension.h
@@ -6,6 +6,7 @@
 #include "MovieSceneSection.h"
 
 struct FMovieSceneBoolChannel;
+struct FMovieSceneByteChannel;
 struct FMovieSceneDoubleChannel;
 struct FMovieSceneFloatChannel;
 struct FMovieSceneIntegerChannel;
@@ -26,7 +27,8 @@ private:
 	void AddExtrapolationMenu(FMenuBuilder& MenuBuilder, const bool bInPreInfinity);
 
 	void GetChannels(TArray<FMovieSceneFloatChannel*>& FloatChannels, TArray<FMovieSceneDoubleChannel*>& DoubleChannels,
-		TArray<FMovieSceneIntegerChannel*>& IntegerChannels, TArray<FMovieSceneBoolChannel*>& BoolChannels) const;
+		TArray<FMovieSceneIntegerChannel*>& IntegerChannels, TArray<FMovieSceneBoolChannel*>& BoolChannels, 
+		TArray<FMovieSceneByteChannel*>& ByteChannels) const;
 
 	void SetExtrapolationMode(const ERichCurveExtrapolation InExtrapolation, const bool bInPreInfinity);
 	bool IsExtrapolationModeSelected(const ERichCurveExtrapolation InExtrapolation, const bool bInPreInfinity) const;
diff --git a/Engine/Source/Runtime/MovieScene/Private/Channels/MovieSceneByteChannel.cpp b/Engine/Source/Runtime/MovieScene/Private/Channels/MovieSceneByteChannel.cpp
--- a/Engine/Source/Runtime/MovieScene/Private/Channels/MovieSceneByteChannel.cpp
+++ b/Engine/Source/Runtime/MovieScene/Private/Channels/MovieSceneByteChannel.cpp
@@ -48,8 +48,62 @@ bool FMovieSceneByteChannel::Evaluate(FFrameTime InTime, uint8& OutValue) const
 {
 	if (Times.Num())
 	{
-		const int32 Index = FMath::Max(0, Algo::UpperBound(Times, InTime.FrameNumber)-1);
-		OutValue = Values[Index];
+		const FFrameNumber MinFrame = Times[0];
+		const FFrameNumber MaxFrame = Times.Last();
+		//we do None, Constant, and Linear first there is no cycling and so we just exit
+		if (InTime < FFrameTime(MinFrame))
+		{
+			if (PreInfinityExtrap == RCCE_None)
+			{
+				return false;
+			}
+
+			if (PreInfinityExtrap == RCCE_Constant || PreInfinityExtrap == RCCE_Linear)
+			{
+				OutValue = Values[0];
+				return true;
+			}
+		}
+		else if (InTime > FFrameTime(MaxFrame))
+		{
+			if (PostInfinityExtrap == RCCE_None)
+			{
+				return false;
+			}
+
+			if (PostInfinityExtrap == RCCE_Constant || PostInfinityExtrap == RCCE_Linear)
+			{
+				OutValue = Values.Last();
+				return true;
+			}
+		}
+
+		// Compute the cycled time based on extrapolation
+		UE::MovieScene::FCycleParams Params = UE::MovieScene::CycleTime(MinFrame, MaxFrame, InTime);
+
+		// Deal with offset cycles and oscillation
+		// Move int over to double then we will convert back to int if doing an offset
+		const double FirstValue = double(Values[0]);
+		const double LastValue = double(Values.Last());
+		if (InTime < FFrameTime(MinFrame))
+		{
+			switch (PreInfinityExtrap)
+			{
+			case RCCE_CycleWithOffset: Params.ComputePreValueOffset(FirstValue, LastValue); break;
+			case RCCE_Oscillate:       Params.Oscillate(MinFrame.Value, MaxFrame.Value);                       break;
+			}
+		}
+		else if (InTime > FFrameTime(MaxFrame))
+		{
+			switch (PostInfinityExtrap)
+			{
+			case RCCE_CycleWithOffset: Params.ComputePostValueOffset(FirstValue, LastValue); break;
+			case RCCE_Oscillate:       Params.Oscillate(MinFrame.Value, MaxFrame.Value);    break;
+			}
+		}
+
+		const int32 Index = FMath::Max(0, Algo::UpperBound(Times, Params.Time)-1);
+		OutValue = Values[Index] + (uint8)(Params.ValueOffset + 0.5);
 		return true;
 	}
 	else if (bHasDefaultValue)
diff --git a/Engine/Source/Runtime/MovieScene/Public/Channels/MovieSceneByteChannel.h b/Engine/Source/Runtime/MovieScene/Public/Channels/MovieSceneByteChannel.h
--- a/Engine/Source/Runtime/MovieScene/Public/Channels/MovieSceneByteChannel.h
+++ b/Engine/Source/Runtime/MovieScene/Public/Channels/MovieSceneByteChannel.h
@@ -19,6 +19,7 @@
 #include "UObject/Class.h"
 #include "UObject/ObjectMacros.h"
 #include "UObject/ObjectPtr.h"
+#include "Curves/RealCurve.h"
 
 #include "MovieSceneByteChannel.generated.h"
 
@@ -33,7 +34,9 @@ struct FMovieSceneByteChannel : public FMovieSceneChannel
 	GENERATED_BODY()
 
 	FMovieSceneByteChannel()
-		: DefaultValue(0)
+		: PreInfinityExtrap(RCCE_Constant)
+		, PostInfinityExtrap(RCCE_Constant)
+		, DefaultValue(0)
 		, bHasDefaultValue(false)
 		, Enum(nullptr)
 		, KeyHandles()
@@ -170,6 +173,15 @@ public:
 		Enum = InEnum;
 	}
 
+public:
+	/** Pre-infinity extrapolation state, byte channel only supports constant, cycle and oscillate */
+	UPROPERTY()
+	TEnumAsByte<ERichCurveExtrapolation> PreInfinityExtrap;
+
+	/** Post-infinity extrapolation state, byte channel only supports constant, cycle and oscillate */
+	UPROPERTY()
+	TEnumAsByte<ERichCurveExtrapolation> PostInfinityExtrap;
+
 private:
 
 	UPROPERTY(meta=(KeyTimes))
