Display timecode at playhead location on audio sections in sequencer.

diff --git a/Engine/Source/Editor/MovieSceneTools/Private/TrackEditors/AudioTrackEditor.cpp b/Engine/Source/Editor/MovieSceneTools/Private/TrackEditors/AudioTrackEditor.cpp
--- a/Engine/Source/Editor/MovieSceneTools/Private/TrackEditors/AudioTrackEditor.cpp
+++ b/Engine/Source/Editor/MovieSceneTools/Private/TrackEditors/AudioTrackEditor.cpp
@@ -925,6 +925,47 @@ void FAudioSection::SlipSection(FFrameNumber SlipTime)
 	ISequencerSection::SlipSection(SlipTime);
 }
 
+TOptional<FFrameTime> FAudioSection::GetSectionTime(FSequencerSectionPainter& InPainter) const
+{
+	if (!InPainter.bIsSelected || !Sequencer.Pin())
+	{
+		return TOptional<FFrameTime>();
+	}
+
+	const FQualifiedFrameTime CurrentTime = Sequencer.Pin()->GetLocalTime();
+	if (!Section.GetRange().Contains(CurrentTime.Time.FrameNumber))
+	{
+		return TOptional<FFrameTime>();
+	}
+
+	FQualifiedFrameTime HintFrameTime;
+
+	TOptional<FSoundWaveTimecodeInfo> TimecodeInfo;
+	if (UMovieSceneAudioSection* AudioSection = Cast<UMovieSceneAudioSection>(&Section))
+	{
+		if (const USoundWave* SoundWave = DeriveSoundWave(AudioSection))
+		{
+			TimecodeInfo = SoundWave->GetTimecodeInfo();
+		}
+
+		if (TimecodeInfo.IsSet())
+		{
+			const FFrameRate SectionFrameRate = AudioSection->GetTypedOuter<UMovieScene>()->GetTickResolution();
+			const double AudioStartOffset = SectionFrameRate.AsSeconds(AudioSection->GetStartOffset());
+
+			const FFrameRate TickResolution = CurrentTime.Rate;
+			const FFrameTime FramesSinceMidnight = TickResolution.AsFrameTime(TimecodeInfo->GetNumSecondsSinceMidnight() + AudioStartOffset);
+
+			const FFrameTime SectionOffset = CurrentTime.Time - AudioSection->GetInclusiveStartFrame();
+			const FFrameTime CursorFrameTime = FramesSinceMidnight + SectionOffset;
+
+			HintFrameTime = FQualifiedFrameTime(FFrameTime(CursorFrameTime.CeilToFrame()), TickResolution);
+		}
+	}
+
+	return HintFrameTime.Time;
+}
+
 void FAudioSection::RegenerateWaveforms(TRange<float> DrawRange, int32 XOffset, int32 XSize, const FColor& ColorTint, float DisplayScale)
 {
 	UMovieSceneAudioSection* AudioSection = Cast<UMovieSceneAudioSection>(&Section);
diff --git a/Engine/Source/Editor/MovieSceneTools/Public/TrackEditors/AudioTrackEditor.h b/Engine/Source/Editor/MovieSceneTools/Public/TrackEditors/AudioTrackEditor.h
--- a/Engine/Source/Editor/MovieSceneTools/Public/TrackEditors/AudioTrackEditor.h
+++ b/Engine/Source/Editor/MovieSceneTools/Public/TrackEditors/AudioTrackEditor.h
@@ -151,6 +151,7 @@ public:
 	virtual void ResizeSection(ESequencerSectionResizeMode ResizeMode, FFrameNumber ResizeTime) override;
 	virtual void BeginSlipSection() override;
 	virtual void SlipSection(FFrameNumber SlipTime) override;
+	virtual TOptional<FFrameTime> GetSectionTime(FSequencerSectionPainter& InPainter) const override;
 	
 private:
 
