[Sequencer] Fix crash on opening new level with keyframe selected in Sequencer.

diff --git a/Engine/Source/Editor/MovieSceneTools/Private/FrameNumberDetailsCustomization.cpp b/Engine/Source/Editor/MovieSceneTools/Private/FrameNumberDetailsCustomization.cpp
--- a/Engine/Source/Editor/MovieSceneTools/Private/FrameNumberDetailsCustomization.cpp
+++ b/Engine/Source/Editor/MovieSceneTools/Private/FrameNumberDetailsCustomization.cpp
@@ -68,6 +68,11 @@ void FFrameNumberDetailsCustomization::CustomizeChildren(TSharedRef<IPropertyHan
 
 FText FFrameNumberDetailsCustomization::OnGetTimeText() const
 {
+	if (!NumericTypeInterface.IsValid())
+	{
+		return FText();
+	}
+
 	int32 CurrentValue = 0.0;
 	FPropertyAccess::Result Result = FrameNumberProperty->GetValue(CurrentValue);
 
@@ -98,6 +103,11 @@ FText FFrameNumberDetailsCustomization::OnGetTimeToolTipText() const
 
 void FFrameNumberDetailsCustomization::OnTimeTextCommitted(const FText& InText, ETextCommit::Type CommitInfo)
 {
+	if (!NumericTypeInterface.IsValid())
+	{
+		return;
+	}
+
 	TArray<FString> PerObjectValueStrs;
 	FrameNumberProperty->GetPerObjectValues(PerObjectValueStrs);
 
diff --git a/Engine/Source/Editor/MovieSceneTools/Private/SKeyEditInterface.cpp b/Engine/Source/Editor/MovieSceneTools/Private/SKeyEditInterface.cpp
--- a/Engine/Source/Editor/MovieSceneTools/Private/SKeyEditInterface.cpp
+++ b/Engine/Source/Editor/MovieSceneTools/Private/SKeyEditInterface.cpp
@@ -101,15 +101,21 @@ FVector2D SKeyEditInterface::ComputeDesiredSize(float LayoutScaleMultiplier) con
 
 TSharedRef<IPropertyTypeCustomization> SKeyEditInterface::CreateBindingIDCustomization()
 {
-	TSharedPtr<ISequencer> SequencerPtr = WeakSequencer.Pin();
-	FMovieSceneSequenceID SequenceID = SequencerPtr.IsValid() ? SequencerPtr->GetFocusedTemplateID() : MovieSceneSequenceID::Root;
-	return MakeShared<FMovieSceneObjectBindingIDCustomization>(SequenceID, WeakSequencer);
+	if (const TSharedPtr<ISequencer> SequencerPtr = WeakSequencer.Pin())
+	{
+		FMovieSceneSequenceID SequenceID = SequencerPtr.IsValid() ? SequencerPtr->GetFocusedTemplateID() : MovieSceneSequenceID::Root;
+		return MakeShared<FMovieSceneObjectBindingIDCustomization>(SequenceID, WeakSequencer);
+	}
+	return MakeShared<FMovieSceneObjectBindingIDCustomization>(FMovieSceneSequenceID(), WeakSequencer);
 }
 
 TSharedRef<IPropertyTypeCustomization> SKeyEditInterface::CreateFrameNumberCustomization()
 {
-	TSharedPtr<ISequencer> SequencerPtr = WeakSequencer.Pin();
-	return MakeShared<FFrameNumberDetailsCustomization>(SequencerPtr->GetNumericTypeInterface());
+	if (const TSharedPtr<ISequencer> SequencerPtr = WeakSequencer.Pin())
+	{
+		return MakeShared<FFrameNumberDetailsCustomization>(SequencerPtr->GetNumericTypeInterface());
+	}
+	return MakeShared<FFrameNumberDetailsCustomization>(nullptr);
 }
 
 TSharedRef<IPropertyTypeCustomization> SKeyEditInterface::CreateEventCustomization()
@@ -170,4 +176,4 @@ void SKeyEditInterface::PostUndo(bool bSuccess)
 void SKeyEditInterface::PostRedo(bool bSuccess)
 {
 	PostUndo(bSuccess);
-}
\ No newline at end of file
+}
