Sequencer Conditions- Automatically set the 'Condition Context' parameter if one exists.

diff --git a/Engine/Source/Editor/MovieSceneTools/Private/Conditions/MovieSceneDirectorBlueprintConditionCustomization.cpp b/Engine/Source/Editor/MovieSceneTools/Private/Conditions/MovieSceneDirectorBlueprintConditionCustomization.cpp
--- a/Engine/Source/Editor/MovieSceneTools/Private/Conditions/MovieSceneDirectorBlueprintConditionCustomization.cpp
+++ b/Engine/Source/Editor/MovieSceneTools/Private/Conditions/MovieSceneDirectorBlueprintConditionCustomization.cpp
@@ -178,6 +178,23 @@ void FMovieSceneDirectorBlueprintConditionCustomization::OnCreateEndpoint(UMovie
 		}
 
 		FMovieSceneDirectorBlueprintConditionUtils::SetEndpoint(MovieScene, DirectorBlueprintConditionData, NewEndpoint);
+
+		// If we have a candidate for the condition context pin, set it automatically
+		UK2Node* CommonEndpoint = GetCommonEndpoint();
+		if (CommonEndpoint)
+		{
+			TArray<FWellKnownParameterCandidates> WellKnownParameterCandidates;
+			GetWellKnownParameterCandidates(CommonEndpoint, WellKnownParameterCandidates);
+			for(int32 ParameterIndex = 0; ParameterIndex < WellKnownParameterCandidates.Num(); ++ParameterIndex)
+			{
+				const FWellKnownParameterCandidates& Candidate = WellKnownParameterCandidates[ParameterIndex];
+				if (Candidate.CandidatePinNames.Num() > 0)
+				{
+					// Pick the first one
+					SetWellKnownParameterPinName(EditObjects[Index], RawData[Index], ParameterIndex, Candidate.CandidatePinNames[0]);
+				}
+			}
+		}
 	}
 
 	FMovieSceneDirectorBlueprintConditionUtils::EnsureBlueprintExtensionCreated(Sequence, Blueprint);
@@ -194,6 +211,23 @@ void FMovieSceneDirectorBlueprintConditionCustomization::OnSetEndpoint(UMovieSce
 
 		FMovieSceneDirectorBlueprintConditionUtils::SetEndpoint(MovieScene, DirectorBlueprintConditionData, NewEndpoint);
 
+		// If we have a candidate for the condition context pin, set it automatically
+		UK2Node* CommonEndpoint = GetCommonEndpoint();
+		if (CommonEndpoint)
+		{
+			TArray<FWellKnownParameterCandidates> WellKnownParameterCandidates;
+			GetWellKnownParameterCandidates(CommonEndpoint, WellKnownParameterCandidates);
+			for (int32 ParameterIndex = 0; ParameterIndex < WellKnownParameterCandidates.Num(); ++ParameterIndex)
+			{
+				const FWellKnownParameterCandidates& Candidate = WellKnownParameterCandidates[ParameterIndex];
+				if (Candidate.CandidatePinNames.Num() > 0)
+				{
+					// Pick the first one
+					SetWellKnownParameterPinName(EditObjects[Index], RawData[Index], ParameterIndex, Candidate.CandidatePinNames[0]);
+				}
+			}
+		}
+
 		FMovieSceneDirectorBlueprintConditionUtils::EnsureBlueprintExtensionCreated(Sequence, Blueprint);
 	}
 }
