Sequencer- Platform Condition
* Added a platform condition that returns true if the current platform is one of the selected ones on the condition.
* Added a customization for the platform condition to provide a drop-down menu
* Fixed the condition customization allowing for individual conditions to have customization.
* Hide director blueprint condition from the class list as we call it out separately.
* Rename 'Entity' to 'OwnerObject' in condition scope enum for clarity.

diff --git a/Engine/Source/Editor/MovieSceneTools/Private/Conditions/MovieSceneConditionCustomization.cpp b/Engine/Source/Editor/MovieSceneTools/Private/Conditions/MovieSceneConditionCustomization.cpp
--- a/Engine/Source/Editor/MovieSceneTools/Private/Conditions/MovieSceneConditionCustomization.cpp
+++ b/Engine/Source/Editor/MovieSceneTools/Private/Conditions/MovieSceneConditionCustomization.cpp
@@ -36,6 +36,12 @@ public:
 	{
 		if (InClass && InClass->IsChildOf(UMovieSceneCondition::StaticClass()))
 		{
+			// Don't show the director blueprint condition here, as we call it out separately
+			if (InClass == UMovieSceneDirectorBlueprintCondition::StaticClass())
+			{
+				return false;
+			}
+
 			if (MovieScene.IsValid())
 			{
 				return MovieScene->IsConditionClassAllowed(InClass);
@@ -146,19 +152,16 @@ void FMovieSceneConditionCustomization::CustomizeChildren(TSharedRef<IPropertyHa
 	if (NumChildren == 1)
 	{
 		TSharedRef<IPropertyHandle> ObjectHandle = ConditionPropertyHandle->GetChildHandle(0).ToSharedRef();
-		// Now show the properties on the object
-		uint32 NumObjectChildren;
-		ObjectHandle->GetNumChildren(NumObjectChildren);
-		for (uint32 ChildIndex = 0; ChildIndex < NumObjectChildren; ++ChildIndex)
+		TArray<void*> ConditionRawArray;
+		ObjectHandle->AccessRawData(ConditionRawArray);
+		if (ConditionRawArray.Num() > 0)
 		{
-			TSharedRef<IPropertyHandle> ChildHandle = ObjectHandle->GetChildHandle(ChildIndex).ToSharedRef();
-
-			// Skip properties that are already marked hidden
-			if (ChildHandle->IsCustomized())
+			UMovieSceneCondition* Condition = reinterpret_cast<UMovieSceneCondition*>(ConditionRawArray[0]);
 			{
-				continue;
+				TArray<UObject*> ObjectArray;
+				ObjectArray.Add(Condition);
+				IDetailPropertyRow* ExternalRow = ChildBuilder.AddExternalObjects(ObjectArray, FAddPropertyParams().HideRootObjectNode(true).AllowChildren(true));
 			}
-			ChildBuilder.AddProperty(ChildHandle);
 		}
 	}
 }
diff --git a/Engine/Source/Editor/MovieSceneTools/Private/Conditions/MovieScenePlatformConditionCustomization.cpp b/Engine/Source/Editor/MovieSceneTools/Private/Conditions/MovieScenePlatformConditionCustomization.cpp
new file mode 100644
--- /dev/null
+++ b/Engine/Source/Editor/MovieSceneTools/Private/Conditions/MovieScenePlatformConditionCustomization.cpp
@@ -0,0 +1,137 @@
+// Copyright Epic Games, Inc. All Rights Reserved.
+
+#include "Conditions/MovieScenePlatformConditionCustomization.h"
+#include "PropertyHandle.h"
+#include "DetailLayoutBuilder.h"
+#include "Conditions/MovieScenePlatformCondition.h"
+#include "Misc/DataDrivenPlatformInfoRegistry.h"
+#include "Widgets/Input/SComboButton.h"
+#include "DetailCategoryBuilder.h"
+#include "DetailWidgetRow.h"
+#include "Framework/MultiBox/MultiBoxBuilder.h"
+#include "Textures/SlateIcon.h"
+
+#define LOCTEXT_NAMESPACE "MovieSceneDynamicBindingCustomization"
+
+TSharedRef<IDetailCustomization> FMovieScenePlatformConditionCustomization::MakeInstance()
+{
+	return MakeShareable(new FMovieScenePlatformConditionCustomization);
+}
+
+void FMovieScenePlatformConditionCustomization::CustomizeDetails(IDetailLayoutBuilder& DetailBuilder)
+{
+	ValidPlatformsPropertyHandle = DetailBuilder.GetProperty(GET_MEMBER_NAME_CHECKED(UMovieScenePlatformCondition, ValidPlatforms), UMovieScenePlatformCondition::StaticClass());
+	ValidPlatformsPropertyHandle->MarkHiddenByCustomization();
+	IDetailCategoryBuilder& PlatformsCategory = DetailBuilder.EditCategory(TEXT("Valid Platforms"));
+
+	const TArray<const FDataDrivenPlatformInfo*>& PlatformInfos = FDataDrivenPlatformInfoRegistry::GetSortedPlatformInfos(EPlatformInfoType::TruePlatformsOnly);
+
+	const auto& GetCurrentValidPlatformNames = [](TSharedPtr<IPropertyHandle> ValidPlatformsPropertyHandle) -> TArray<FName>
+	{
+		TArray<FName> Names;
+
+		TArray<void*> RawData;
+		ValidPlatformsPropertyHandle->AccessRawData(RawData);
+
+		if (RawData.Num() > 0)
+		{
+			if (TArray<FName>* CurrentValidPlatformNamesPtr = reinterpret_cast<TArray<FName>*>(RawData[0]))
+			{
+				Names = *CurrentValidPlatformNamesPtr;
+			}
+		}
+		return Names;
+	};
+
+	auto GetComboButtonText = [SharedThis = StaticCastSharedRef<FMovieScenePlatformConditionCustomization>(AsShared()), GetCurrentValidPlatformNames]() -> FText
+	{
+		TArray<FName> CurrentValidPlatformNames = GetCurrentValidPlatformNames(SharedThis->ValidPlatformsPropertyHandle);
+		TArray<FText> CurrentValidPlatforms;
+		CurrentValidPlatforms.Reserve(CurrentValidPlatformNames.Num());
+
+		for (const FName& PlatformName : CurrentValidPlatformNames)
+		{
+			CurrentValidPlatforms.Add(FText::FromName(PlatformName));
+		}
+		if (CurrentValidPlatforms.Num() > 3)
+		{
+			CurrentValidPlatforms.SetNum(3);
+			CurrentValidPlatforms.Add(FText::FromString("..."));
+		}
+
+		return FText::Join(FText::FromString(", "), CurrentValidPlatforms);
+	};
+
+	PlatformsCategory.AddCustomRow(LOCTEXT("ValidPlatforms", "Valid Platforms"))
+		.NameContent()
+		[
+			SNew(STextBlock)
+			.Text(LOCTEXT("ValidPlatforms", "Valid Platforms"))
+			.ToolTipText(LOCTEXT("ValidPlatformsTooltip", "Which platforms will pass the condition"))
+			.Font(IDetailLayoutBuilder::GetDetailFont())
+		]
+		.ValueContent()
+		[
+			
+			// Constructs the UI for bitmask property editing.
+			SNew(SComboButton)
+			.ButtonContent()
+				[
+					SNew(STextBlock)
+					.Text_Lambda(GetComboButtonText)
+				]
+			.OnGetMenuContent_Lambda([&DetailBuilder, SharedThis = StaticCastSharedRef<FMovieScenePlatformConditionCustomization>(AsShared()), PlatformInfos, GetCurrentValidPlatformNames]()
+				{
+					TArray<FName> CurrentValidPlatformNames = GetCurrentValidPlatformNames(SharedThis->ValidPlatformsPropertyHandle);
+
+					FMenuBuilder MenuBuilder(false, nullptr);
+
+					for (int32 i = 0; i < PlatformInfos.Num(); ++i)
+					{
+						MenuBuilder.AddMenuEntry(
+							FText::FromName(PlatformInfos[i]->IniPlatformName),
+							FText(),
+							FSlateIcon(FAppStyle::GetAppStyleSetName(), PlatformInfos[i]->GetIconStyleName(EPlatformIconSize::Normal)),
+							FUIAction
+							(
+								FExecuteAction::CreateLambda([SharedThis, i, PlatformInfos, CurrentValidPlatformNames, &DetailBuilder]()
+									{
+										TArray<FName> NewValidPlatformNames = CurrentValidPlatformNames;
+										if (NewValidPlatformNames.Contains(PlatformInfos[i]->IniPlatformName))
+										{
+											NewValidPlatformNames.Remove(PlatformInfos[i]->IniPlatformName);
+										}
+										else
+										{
+											NewValidPlatformNames.Add(PlatformInfos[i]->IniPlatformName);
+										}
+
+										TArray<void*> RawData;
+										SharedThis->ValidPlatformsPropertyHandle->AccessRawData(RawData);
+										if (RawData.Num() == 1)
+										{
+											if (TArray<FName>* CurrentValidPlatformNamesPtr = reinterpret_cast<TArray<FName>*>(RawData[0]))
+											{
+												*CurrentValidPlatformNamesPtr = NewValidPlatformNames;
+												SharedThis->ValidPlatformsPropertyHandle->NotifyPostChange(EPropertyChangeType::ValueSet);
+												DetailBuilder.ForceRefreshDetails();
+											}
+										}
+									}),
+								FCanExecuteAction(),
+								FIsActionChecked::CreateLambda([SharedThis, i, PlatformInfos, CurrentValidPlatformNames]() -> bool
+									{
+										return CurrentValidPlatformNames.Contains(PlatformInfos[i]->IniPlatformName);
+									})
+							),
+							NAME_None,
+							EUserInterfaceActionType::Check);
+					}
+
+					return MenuBuilder.MakeWidget();
+				})
+			
+		]; 
+}
+
+#undef LOCTEXT_NAMESPACE
\ No newline at end of file
diff --git a/Engine/Source/Editor/MovieSceneTools/Private/MovieSceneToolsModule.cpp b/Engine/Source/Editor/MovieSceneTools/Private/MovieSceneToolsModule.cpp
--- a/Engine/Source/Editor/MovieSceneTools/Private/MovieSceneToolsModule.cpp
+++ b/Engine/Source/Editor/MovieSceneTools/Private/MovieSceneToolsModule.cpp
@@ -131,6 +131,7 @@
 #include "Conditions/MovieSceneDirectorBlueprintCondition.h"
 #include "Conditions/MovieSceneDirectorBlueprintConditionUtils.h"
 #include "Conditions/MovieSceneDirectorBlueprintConditionCustomization.h"
+#include "Conditions/MovieScenePlatformConditionCustomization.h"
 
 #define LOCTEXT_NAMESPACE "FMovieSceneToolsModule"
 
@@ -222,6 +223,7 @@ void FMovieSceneToolsModule::StartupModule()
 		PropertyModule.RegisterCustomPropertyTypeLayout("MovieSceneCVarOverrides", FOnGetPropertyTypeCustomizationInstance::CreateStatic(&UE::MovieScene::FCVarOverridesPropertyTypeCustomization::MakeInstance));
 		PropertyModule.RegisterCustomPropertyTypeLayout("MovieSceneTimeWarpVariant", FOnGetPropertyTypeCustomizationInstance::CreateLambda(&MakeShared<UE::MovieScene::FMovieSceneTimeWarpVariantCustomization>));
 		PropertyModule.RegisterCustomPropertyTypeLayout("MovieSceneConditionContainer", FOnGetPropertyTypeCustomizationInstance::CreateStatic(&FMovieSceneConditionCustomization::MakeInstance));
+		PropertyModule.RegisterCustomClassLayout("MovieScenePlatformCondition", FOnGetDetailCustomizationInstance::CreateStatic(&FMovieScenePlatformConditionCustomization::MakeInstance));
 
 
 		SequencerModule.RegisterChannelInterface<FMovieSceneBoolChannel>();
@@ -398,6 +400,7 @@ void FMovieSceneToolsModule::ShutdownModule()
 		PropertyModule.UnregisterCustomPropertyTypeLayout("MovieSceneDynamicBinding");
 		PropertyModule.UnregisterCustomPropertyTypeLayout("MovieSceneDirectorBlueprintConditionData");
 		PropertyModule.UnregisterCustomPropertyTypeLayout("MovieSceneConditionContainer");
+		PropertyModule.UnregisterCustomClassLayout("MovieScenePlatformCondition");
 	}
 
 	FEditorModeRegistry::Get().UnregisterMode(FSkeletalAnimationTrackEditMode::ModeName);
diff --git a/Engine/Source/Editor/MovieSceneTools/Public/Conditions/MovieScenePlatformConditionCustomization.h b/Engine/Source/Editor/MovieSceneTools/Public/Conditions/MovieScenePlatformConditionCustomization.h
new file mode 100644
--- /dev/null
+++ b/Engine/Source/Editor/MovieSceneTools/Public/Conditions/MovieScenePlatformConditionCustomization.h
@@ -0,0 +1,18 @@
+// Copyright Epic Games, Inc. All Rights Reserved.
+
+#pragma once
+
+#include "IDetailCustomization.h"
+#include "Templates/SharedPointer.h"
+#include "PropertyHandle.h"
+
+class IDetailLayoutBuilder;
+
+class MOVIESCENETOOLS_API FMovieScenePlatformConditionCustomization : public IDetailCustomization
+{
+public:
+	static TSharedRef<IDetailCustomization> MakeInstance();
+	virtual void CustomizeDetails(IDetailLayoutBuilder& DetailBuilder) override;
+
+	TSharedPtr<IPropertyHandle> ValidPlatformsPropertyHandle;
+};
diff --git a/Engine/Source/Runtime/MovieScene/Private/Conditions/MovieSceneCondition.cpp b/Engine/Source/Runtime/MovieScene/Private/Conditions/MovieSceneCondition.cpp
--- a/Engine/Source/Runtime/MovieScene/Private/Conditions/MovieSceneCondition.cpp
+++ b/Engine/Source/Runtime/MovieScene/Private/Conditions/MovieSceneCondition.cpp
@@ -33,7 +33,7 @@ uint32 UMovieSceneCondition::ComputeCacheKey(FGuid BindingGuid, FMovieSceneSeque
 	{
 		HashResult = HashCombineFast(HashCombineFast(HashResult, GetTypeHash(BindingGuid)), GetTypeHash(SequenceID));
 	}
-	else if (Scope == EMovieSceneConditionScope::Entity && EntityOwner != nullptr)
+	else if (Scope == EMovieSceneConditionScope::OwnerObject && EntityOwner != nullptr)
 	{
 		HashResult = HashCombineFast(HashResult, GetTypeHash(EntityOwner));
 	}
diff --git a/Engine/Source/Runtime/MovieScene/Public/Conditions/MovieSceneCondition.h b/Engine/Source/Runtime/MovieScene/Public/Conditions/MovieSceneCondition.h
--- a/Engine/Source/Runtime/MovieScene/Public/Conditions/MovieSceneCondition.h
+++ b/Engine/Source/Runtime/MovieScene/Public/Conditions/MovieSceneCondition.h
@@ -28,8 +28,8 @@ enum class EMovieSceneConditionScope : uint8
 	Global,
 	/* Condition may have different results for different object bindings. */
 	Binding,
-	/* Condition may have different results for different entities (tracks, sections) in the Sequence.*/
-	Entity,
+	/* Condition may have different results for each different outer object owner (i.e. track, section) in the Sequence.*/
+	OwnerObject,
 };
 
 /* Defines how often a condition needs to be checked. 
@@ -74,7 +74,7 @@ struct MOVIESCENE_API FMovieSceneConditionContainer
 {
 	GENERATED_BODY()
 
-	UPROPERTY(EditAnywhere, BlueprintReadWrite, Instanced, Category="Sequencer|Condition")
+	UPROPERTY(EditAnywhere, BlueprintReadWrite, Instanced, Category="Sequencer|Condition", meta = (EditInline, AllowEditInlineCustomization))
 	TObjectPtr<UMovieSceneCondition> Condition = nullptr;
 };
 
diff --git a/Engine/Source/Runtime/MovieScene/Public/Conditions/MovieSceneGroupCondition.h b/Engine/Source/Runtime/MovieScene/Public/Conditions/MovieSceneGroupCondition.h
--- a/Engine/Source/Runtime/MovieScene/Public/Conditions/MovieSceneGroupCondition.h
+++ b/Engine/Source/Runtime/MovieScene/Public/Conditions/MovieSceneGroupCondition.h
@@ -29,7 +29,7 @@ enum class EMovieSceneGroupConditionOperator : uint8
 /**
  * Condition class that allows the grouping of other conditions using 'and', 'or', or 'xor'.
  */ 
-UCLASS(BlueprintType, DefaultToInstanced, EditInlineNew, Meta=(DisplayName="Group"))
+UCLASS(BlueprintType, DefaultToInstanced, EditInlineNew, Meta=(DisplayName="Group Condition"))
 class MOVIESCENE_API UMovieSceneGroupCondition
 	: public UMovieSceneCondition
 {
diff --git a/Engine/Source/Runtime/MovieSceneTracks/Private/Conditions/MovieScenePlatformCondition.cpp b/Engine/Source/Runtime/MovieSceneTracks/Private/Conditions/MovieScenePlatformCondition.cpp
new file mode 100644
--- /dev/null
+++ b/Engine/Source/Runtime/MovieSceneTracks/Private/Conditions/MovieScenePlatformCondition.cpp
@@ -0,0 +1,13 @@
+// Copyright Epic Games, Inc. All Rights Reserved.
+
+#include "Conditions/MovieScenePlatformCondition.h"
+#include "CoreMinimal.h"
+#include "Engine/World.h"
+#include "EntitySystem/MovieSceneSharedPlaybackState.h"
+#include "Kismet/GameplayStatics.h"
+
+bool UMovieScenePlatformCondition::EvaluateConditionInternal(FGuid BindingGuid, FMovieSceneSequenceID SequenceID, TSharedRef<const UE::MovieScene::FSharedPlaybackState> SharedPlaybackState) const
+{
+	FString PlatformName = UGameplayStatics::GetPlatformName();
+	return ValidPlatforms.Contains(*PlatformName);
+}
diff --git a/Engine/Source/Runtime/MovieSceneTracks/Public/Conditions/MovieScenePlatformCondition.h b/Engine/Source/Runtime/MovieSceneTracks/Public/Conditions/MovieScenePlatformCondition.h
new file mode 100644
--- /dev/null
+++ b/Engine/Source/Runtime/MovieSceneTracks/Public/Conditions/MovieScenePlatformCondition.h
@@ -0,0 +1,42 @@
+// Copyright Epic Games, Inc. All Rights Reserved.
+
+#pragma once
+
+#include "CoreTypes.h"
+#include "Conditions/MovieSceneCondition.h"
+#include "MovieScenePlatformCondition.generated.h"
+
+namespace UE
+{
+	namespace MovieScene
+	{
+		struct FSharedPlaybackState;
+	}
+}
+
+/**
+ * Condition on whether the platform running the executable matches one of the given platforms.
+ */
+UCLASS(DisplayName="Platform Condition")
+class MOVIESCENETRACKS_API UMovieScenePlatformCondition
+	: public UMovieSceneCondition
+{
+	GENERATED_BODY()
+
+public: 
+
+	UPROPERTY(EditAnywhere, Category="Sequencer|Condition")
+	TArray<FName> ValidPlatforms;
+
+protected:
+
+	/*
+	* UMovieSceneCondition overrides 
+	*/
+	virtual bool EvaluateConditionInternal(FGuid BindingGuid, FMovieSceneSequenceID SequenceID, TSharedRef<const UE::MovieScene::FSharedPlaybackState> SharedPlaybackState) const override;
+	
+	virtual EMovieSceneConditionScope GetScopeInternal() const override { return EMovieSceneConditionScope::Global; }
+	
+	virtual EMovieSceneConditionCheckFrequency GetCheckFrequencyInternal() const override { return EMovieSceneConditionCheckFrequency::Once; }
+
+};
