Sequencer- fix Mac compile issue.

diff --git a/Engine/Source/Editor/MovieSceneTools/Private/Conditions/MovieScenePlatformConditionCustomization.cpp b/Engine/Source/Editor/MovieSceneTools/Private/Conditions/MovieScenePlatformConditionCustomization.cpp
--- a/Engine/Source/Editor/MovieSceneTools/Private/Conditions/MovieScenePlatformConditionCustomization.cpp
+++ b/Engine/Source/Editor/MovieSceneTools/Private/Conditions/MovieScenePlatformConditionCustomization.cpp
@@ -26,12 +26,12 @@ void FMovieScenePlatformConditionCustomization::CustomizeDetails(IDetailLayoutBu
 
 	const TArray<const FDataDrivenPlatformInfo*>& PlatformInfos = FDataDrivenPlatformInfoRegistry::GetSortedPlatformInfos(EPlatformInfoType::TruePlatformsOnly);
 
-	const auto& GetCurrentValidPlatformNames = [](TSharedPtr<IPropertyHandle> ValidPlatformsPropertyHandle) -> TArray<FName>
+	const auto& GetCurrentValidPlatformNames = [](TSharedPtr<IPropertyHandle> InValidPlatformsPropertyHandle) -> TArray<FName>
 	{
 		TArray<FName> Names;
 
 		TArray<void*> RawData;
-		ValidPlatformsPropertyHandle->AccessRawData(RawData);
+		InValidPlatformsPropertyHandle->AccessRawData(RawData);
 
 		if (RawData.Num() > 0)
 		{
