Sequencer- Add feature to allow tracks to disable conditions on themselves.
Make time warp track disallow conditions on itself.

diff --git a/Engine/Plugins/MovieScene/LevelSequenceEditor/Source/LevelSequenceEditor/Private/LevelSequenceEditorSubsystem.cpp b/Engine/Plugins/MovieScene/LevelSequenceEditor/Source/LevelSequenceEditor/Private/LevelSequenceEditorSubsystem.cpp
--- a/Engine/Plugins/MovieScene/LevelSequenceEditor/Source/LevelSequenceEditor/Private/LevelSequenceEditorSubsystem.cpp
+++ b/Engine/Plugins/MovieScene/LevelSequenceEditor/Source/LevelSequenceEditor/Private/LevelSequenceEditorSubsystem.cpp
@@ -711,11 +711,10 @@ void ULevelSequenceEditorSubsystem::AddBindingDetailCustomizations(TSharedRef<ID
 	}
 }
 
-void ULevelSequenceEditorSubsystem::AddTrackRowMetadataCustomizations(TSharedRef<IDetailsView> DetailsView, TSharedPtr<ISequencer> ActiveSequencer)
+void ULevelSequenceEditorSubsystem::AddTrackRowMetadataCustomizations(TSharedRef<IDetailsView> DetailsView, TSharedPtr<ISequencer> ActiveSequencer, UMovieSceneSequence* Sequence)
 {
 	if (ActiveSequencer.IsValid())
 	{
-		UMovieSceneSequence* Sequence = ActiveSequencer->GetFocusedMovieSceneSequence();
 		UMovieScene* MovieScene = Sequence ? Sequence->GetMovieScene() : nullptr;
 		if (MovieScene)
 		{
@@ -2686,7 +2685,7 @@ void ULevelSequenceEditorSubsystem::AddTrackRowMetadataMenu(FMenuBuilder& MenuBu
 
 	TSharedRef<IDetailsView> DetailsView = FModuleManager::GetModuleChecked<FPropertyEditorModule>("PropertyEditor").CreateDetailView(DetailsViewArgs);
 
-	AddTrackRowMetadataCustomizations(DetailsView, Sequencer);
+	AddTrackRowMetadataCustomizations(DetailsView, Sequencer, Sequence);
 
 	RefreshTrackRowMetadataDetails(&DetailsView.Get());
 	DetailsView->OnFinishedChangingProperties().AddUObject(this, &ULevelSequenceEditorSubsystem::OnFinishedChangingTrackRowMetadata, DetailsView);
@@ -3193,6 +3192,7 @@ void ULevelSequenceEditorSubsystem::RefreshTrackRowMetadataDetails(IDetailsView*
 		{
 			if (const FMovieSceneTrackRowMetadata* Metadata = Track->FindTrackRowMetadata(SelectedTrackRows[Index].Value))
 			{
+				Helper->OwnerTrack = Track;
 				Helper->TrackRowMetadata = *Metadata;
 				if (Metadata->ConditionContainer.Condition)
 				{
diff --git a/Engine/Plugins/MovieScene/LevelSequenceEditor/Source/LevelSequenceEditor/Public/LevelSequenceEditorSubsystem.h b/Engine/Plugins/MovieScene/LevelSequenceEditor/Source/LevelSequenceEditor/Public/LevelSequenceEditorSubsystem.h
--- a/Engine/Plugins/MovieScene/LevelSequenceEditor/Source/LevelSequenceEditor/Public/LevelSequenceEditorSubsystem.h
+++ b/Engine/Plugins/MovieScene/LevelSequenceEditor/Source/LevelSequenceEditor/Public/LevelSequenceEditorSubsystem.h
@@ -40,6 +40,7 @@ class IDetailsView;
 class USequencerCurveEditorObject;
 class UMovieSceneCustomBinding;
 class IMenu;
+class UMovieSceneTrackRowMetadataHelper;
 
 USTRUCT(BlueprintType)
 struct FMovieSceneScriptingParams
@@ -81,20 +82,6 @@ public:
 	TArray<FMovieSceneBindingPropertyInfo> Bindings;
 };
 
-
-// Helper UObject for editing optional track row metadata not in-place. A UObject instead of a UStruct because we need to support instanced sub objects (conditions)
-UCLASS(CollapseCategories)
-class LEVELSEQUENCEEDITOR_API UMovieSceneTrackRowMetadataHelper : public UObject
-{
-	GENERATED_BODY()
-
-public:
-
-	UPROPERTY(EditAnywhere, Category = "General", meta=(ShowOnlyInnerProperties))
-	FMovieSceneTrackRowMetadata TrackRowMetadata;
-};
-
-
 /**
 * ULevelSequenceEditorSubsystem
 * Subsystem for level sequence editor related utilities to scripts
@@ -317,7 +304,7 @@ private:
 private:
 
 	void AddBindingDetailCustomizations(TSharedRef<IDetailsView> DetailsView, TSharedPtr<ISequencer> ActiveSequencer, FGuid BindingGuid);
-	void AddTrackRowMetadataCustomizations(TSharedRef<IDetailsView> DetailsView, TSharedPtr<ISequencer> ActiveSequencer);
+	void AddTrackRowMetadataCustomizations(TSharedRef<IDetailsView> DetailsView, TSharedPtr<ISequencer> ActiveSequencer, UMovieSceneSequence* Sequence);
 	void OnBindingPropertyMenuBeingDestroyed(const TSharedRef<IMenu>& Menu, TSharedRef<IDetailsView> DetailsView);
 	void OnTrackRowMetadataMenuBeingDestroyed(const TSharedRef<IMenu>& Menu, TSharedRef<IDetailsView> DetailsView);
 
diff --git a/Engine/Source/Editor/MovieSceneTools/Private/Conditions/MovieSceneConditionCustomization.cpp b/Engine/Source/Editor/MovieSceneTools/Private/Conditions/MovieSceneConditionCustomization.cpp
--- a/Engine/Source/Editor/MovieSceneTools/Private/Conditions/MovieSceneConditionCustomization.cpp
+++ b/Engine/Source/Editor/MovieSceneTools/Private/Conditions/MovieSceneConditionCustomization.cpp
@@ -87,8 +87,13 @@ void FMovieSceneConditionCustomization::CustomizeHeader(TSharedRef<IPropertyHand
 		Sequence = GetCommonSequence();
 	}
 
+	if (!Track.IsValid())
+	{
+		Track = GetCommonTrack();
+	}
+
 	// If conditions not allowed, hide condition property functionality
-	if (!Sequence.IsValid() || !Sequence->GetMovieScene()->IsConditionClassAllowed(UMovieSceneCondition::StaticClass()))
+	if (!Sequence.IsValid() || !Sequence->GetMovieScene()->IsConditionClassAllowed(UMovieSceneCondition::StaticClass()) || (Track.IsValid() && !Track->SupportsConditions()))
 	{
 		ConditionContainerPropertyHandle->MarkHiddenByCustomization();
 		return;
@@ -130,7 +135,7 @@ void FMovieSceneConditionCustomization::CustomizeHeader(TSharedRef<IPropertyHand
 void FMovieSceneConditionCustomization::CustomizeChildren(TSharedRef<IPropertyHandle> InPropertyHandle, IDetailChildrenBuilder& ChildBuilder, IPropertyTypeCustomizationUtils& CustomizationUtils)
 {
 	// If conditions not allowed, hide condition property functionality
-	if (!Sequence.IsValid() || !Sequence->GetMovieScene()->IsConditionClassAllowed(UMovieSceneCondition::StaticClass()))
+	if (!Sequence.IsValid() || !Sequence->GetMovieScene()->IsConditionClassAllowed(UMovieSceneCondition::StaticClass()) || (Track.IsValid() && !Track->SupportsConditions()))
 	{
 		return;
 	}
@@ -391,5 +396,39 @@ UMovieSceneSequence* FMovieSceneConditionCustomization::GetCommonSequence() cons
 	return CommonSequence;
 }
 
+UMovieSceneTrack* FMovieSceneConditionCustomization::GetCommonTrack() const
+{
+	TArray<UObject*> EditObjects;
+	ConditionContainerPropertyHandle->GetOuterObjects(EditObjects);
+
+	UMovieSceneTrack* CommonTrack = nullptr;
+
+	for (UObject* Obj : EditObjects)
+	{
+		UMovieSceneTrack* ThisTrack = Cast<UMovieSceneTrack>(Obj);
+		if (!ThisTrack)
+		{
+			ThisTrack = Obj ? Obj->GetTypedOuter<UMovieSceneTrack>() : nullptr;
+		}
+
+		if (!ThisTrack)
+		{
+			// Special case
+			if (UMovieSceneTrackRowMetadataHelper* TrackRowHelper = Cast<UMovieSceneTrackRowMetadataHelper>(Obj))
+			{
+				ThisTrack = TrackRowHelper->OwnerTrack.Get();
+			}
+		}
+
+		if (CommonTrack && CommonTrack != ThisTrack)
+		{
+			return nullptr;
+		}
+
+		CommonTrack = ThisTrack;
+	}
+	return CommonTrack;
+}
+
 #undef LOCTEXT_NAMESPACE
 
diff --git a/Engine/Source/Editor/MovieSceneTools/Public/Conditions/MovieSceneConditionCustomization.h b/Engine/Source/Editor/MovieSceneTools/Public/Conditions/MovieSceneConditionCustomization.h
--- a/Engine/Source/Editor/MovieSceneTools/Public/Conditions/MovieSceneConditionCustomization.h
+++ b/Engine/Source/Editor/MovieSceneTools/Public/Conditions/MovieSceneConditionCustomization.h
@@ -7,6 +7,8 @@
 #include "IPropertyTypeCustomization.h"
 #include "Widgets/Input/SComboButton.h"
 #include "Conditions/MovieSceneCondition.h"
+#include "MovieSceneTrack.h"
+#include "MovieSceneConditionCustomization.generated.h"
 
 class IDetailLayoutBuilder;
 class IPropertyHandle;
@@ -15,9 +17,25 @@ class FPropertyEditor;
 class SWidget;
 class FMenuBuilder;
 class UMovieSceneSequence;
+class UMovieSceneTrack;
 class FMovieSceneDirectorBlueprintConditionCustomization;
 class IDetailsView;
 
+// Helper UObject for editing optional track row metadata not in-place. A UObject instead of a UStruct because we need to support instanced sub objects (conditions)
+UCLASS(CollapseCategories)
+class MOVIESCENETOOLS_API UMovieSceneTrackRowMetadataHelper : public UObject
+{
+	GENERATED_BODY()
+
+public:
+
+	UPROPERTY(EditAnywhere, Category = "General", meta=(ShowOnlyInnerProperties))
+	FMovieSceneTrackRowMetadata TrackRowMetadata;
+
+	UPROPERTY()
+	TWeakObjectPtr<UMovieSceneTrack> OwnerTrack;
+};
+
 class FMovieSceneConditionCustomization : public IPropertyTypeCustomization
 {
 public:
@@ -75,6 +93,9 @@ private:
 	/* Gets the common sequence for this customization */
 	UMovieSceneSequence* GetCommonSequence() const;
 
+	/* Gets the common track for this customization */
+	UMovieSceneTrack* GetCommonTrack() const;
+
 	/**
 	* Generate the content of the quick bind sub-menu dropdown (shown if the endpoint is not already bound)
 	*/
@@ -103,4 +124,5 @@ private:
 	TSharedPtr<SWidget> OpenMenuWidget;
 
 	TWeakObjectPtr<UMovieSceneSequence> Sequence;
+	TWeakObjectPtr<UMovieSceneTrack> Track;
 };
diff --git a/Engine/Source/Runtime/MovieScene/Private/MovieSceneTrack.cpp b/Engine/Source/Runtime/MovieScene/Private/MovieSceneTrack.cpp
--- a/Engine/Source/Runtime/MovieScene/Private/MovieSceneTrack.cpp
+++ b/Engine/Source/Runtime/MovieScene/Private/MovieSceneTrack.cpp
@@ -36,6 +36,7 @@ UMovieSceneTrack::UMovieSceneTrack(const FObjectInitializer& InInitializer)
 	TrackTint = FColor(127, 127, 127, 0);
 	SortingOrder = -1;
 	bSupportsDefaultSections = true;
+	bSupportsConditions = true;
 #endif
 
 	BuiltInTreePopulationMode = ETreePopulationMode::HighPassPerRow;
diff --git a/Engine/Source/Runtime/MovieScene/Private/Tracks/MovieSceneTimeWarpTrack.cpp b/Engine/Source/Runtime/MovieScene/Private/Tracks/MovieSceneTimeWarpTrack.cpp
--- a/Engine/Source/Runtime/MovieScene/Private/Tracks/MovieSceneTimeWarpTrack.cpp
+++ b/Engine/Source/Runtime/MovieScene/Private/Tracks/MovieSceneTimeWarpTrack.cpp
@@ -5,6 +5,15 @@
 #include "Variants/MovieSceneTimeWarpGetter.h"
 #include "Evaluation/MovieSceneSequenceTransform.h"
 
+
+UMovieSceneTimeWarpTrack::UMovieSceneTimeWarpTrack(const FObjectInitializer& ObjectInitializer)
+	: Super(ObjectInitializer)
+{
+#if WITH_EDITORONLY_DATA
+	bSupportsConditions = false;
+#endif
+}
+
 bool UMovieSceneTimeWarpTrack::SupportsType(TSubclassOf<UMovieSceneSection> SectionClass) const
 {
 	return SectionClass == UMovieSceneTimeWarpSection::StaticClass();
diff --git a/Engine/Source/Runtime/MovieScene/Public/MovieSceneTrack.h b/Engine/Source/Runtime/MovieScene/Public/MovieSceneTrack.h
--- a/Engine/Source/Runtime/MovieScene/Public/MovieSceneTrack.h
+++ b/Engine/Source/Runtime/MovieScene/Public/MovieSceneTrack.h
@@ -556,6 +556,14 @@ public:
 		return bSupportsDefaultSections;
 	}
 
+	/**
+	 * @return Whether or not this track supports conditions.
+	 */
+	virtual bool SupportsConditions() const
+	{
+		return bSupportsConditions;
+	}
+
 	/*
 	* Returns an array of all conditions on track, track row, or section
 	*/
@@ -579,6 +587,10 @@ protected:
 	UPROPERTY()
 	bool bSupportsDefaultSections;
 
+	/** Does this track support conditions */
+	UPROPERTY()
+	bool bSupportsConditions;
+
 public:
 #endif
 
diff --git a/Engine/Source/Runtime/MovieScene/Public/Tracks/MovieSceneTimeWarpTrack.h b/Engine/Source/Runtime/MovieScene/Public/Tracks/MovieSceneTimeWarpTrack.h
--- a/Engine/Source/Runtime/MovieScene/Public/Tracks/MovieSceneTimeWarpTrack.h
+++ b/Engine/Source/Runtime/MovieScene/Public/Tracks/MovieSceneTimeWarpTrack.h
@@ -19,6 +19,13 @@ class UMovieSceneTimeWarpTrack
 
 public:
 
+	/**
+	 * Create and initialize a new instance.
+	 *
+	 * @param ObjectInitializer The object initializer.
+	 */
+	UMovieSceneTimeWarpTrack(const FObjectInitializer& ObjectInitializer);
+
 	/**
 	 * Generate the a sequence transform that perform's this track's time-warp.
 	 * @note: The resulting transform may be the identity transform
