Sequencer- Add 'Use Selected' and 'Browse To' asset buttons to condition picker.

diff --git a/Engine/Source/Editor/MovieSceneTools/Private/Conditions/MovieSceneConditionCustomization.cpp b/Engine/Source/Editor/MovieSceneTools/Private/Conditions/MovieSceneConditionCustomization.cpp
--- a/Engine/Source/Editor/MovieSceneTools/Private/Conditions/MovieSceneConditionCustomization.cpp
+++ b/Engine/Source/Editor/MovieSceneTools/Private/Conditions/MovieSceneConditionCustomization.cpp
@@ -23,6 +23,7 @@
 #include "Framework/Application/SlateApplication.h"
 #include "MovieScene.h"
 #include "ClassViewerFilter.h"
+#include "ScopedTransaction.h"
 
 #define LOCTEXT_NAMESPACE "MovieSceneConditionCustomization"
 
@@ -108,27 +109,49 @@ void FMovieSceneConditionCustomization::CustomizeHeader(TSharedRef<IPropertyHand
 	]
 	.ValueContent()
 	[
-		SAssignNew(ComboButton, SComboButton)
-		.OnGetMenuContent(this, &FMovieSceneConditionCustomization::GenerateConditionPicker)
-		.ContentPadding(0.0f)
-		.ButtonContent()
-		[
-			SNew(SHorizontalBox)
-			+ SHorizontalBox::Slot()
+		SNew(SHorizontalBox)
+		+ SHorizontalBox::Slot()
+			.FillWidth(0.5)
+			.VAlign(VAlign_Center)
+			[
+				SAssignNew(ComboButton, SComboButton)
+				.OnGetMenuContent(this, &FMovieSceneConditionCustomization::GenerateConditionPicker)
+				.ContentPadding(0.0f)
+				.ButtonContent()
+				[
+					SNew(SHorizontalBox)
+					+ SHorizontalBox::Slot()
+					.AutoWidth()
+					.VAlign(VAlign_Center)
+					.Padding(0.0f, 0.0f, 4.0f, 0.0f)
+					[
+						SNew(SImage)
+						.Image(this, &FMovieSceneConditionCustomization::GetDisplayValueIcon)
+					]
+					+ SHorizontalBox::Slot()
+					.VAlign(VAlign_Center)
+					[
+						SNew(STextBlock)
+						.Text(this, &FMovieSceneConditionCustomization::GetDisplayValueAsString)
+					]
+				]
+			]
+		+ SHorizontalBox::Slot()
 			.AutoWidth()
 			.VAlign(VAlign_Center)
-			.Padding(0.0f, 0.0f, 4.0f, 0.0f)
 			[
-				SNew(SImage)
-				.Image(this, &FMovieSceneConditionCustomization::GetDisplayValueIcon)
+				PropertyCustomizationHelpers::MakeUseSelectedButton(FSimpleDelegate::CreateSP(this, &FMovieSceneConditionCustomization::OnUseSelected),
+				LOCTEXT("UseSelectedConditionClass", "Use Selected Condition Class in Content Browser"),
+				TAttribute<bool>::Create(TAttribute<bool>::FGetter::CreateSP(this, &FMovieSceneConditionCustomization::CanUseSelectedAsset)))
 			]
-			+ SHorizontalBox::Slot()
+		+ SHorizontalBox::Slot()
+			.AutoWidth()
 			.VAlign(VAlign_Center)
 			[
-				SNew(STextBlock)
-				.Text(this, &FMovieSceneConditionCustomization::GetDisplayValueAsString)
+				PropertyCustomizationHelpers::MakeBrowseButton(FSimpleDelegate::CreateSP(this, &FMovieSceneConditionCustomization::OnBrowseTo),
+				LOCTEXT("BrowseToConditionClass", "Browse To Condition Class in Content Browser"), 
+				TAttribute<bool>::Create(TAttribute<bool>::FGetter::CreateSP(this, &FMovieSceneConditionCustomization::CanBrowseToAsset)))
 			]
-		]
 	];
 }
 
@@ -221,6 +244,7 @@ void FMovieSceneConditionCustomization::FillConditionClassSubMenu(FMenuBuilder&
 						FString NewConditionName = SharedThis->Sequence->GetName() + TEXT("_Condition");
 						AssetToolsModule.Get().CreateUniqueAssetName(NewConditionPath + TEXT("/") + NewConditionName, TEXT(""), NewConditionPath, NewConditionName);
 
+						const FScopedTransaction Transaction(LOCTEXT("CreateConditionAsset", "Create Condition Asset"));
 						UBlueprint* Blueprint = FKismetEditorUtilities::CreateBlueprintFromClass(LOCTEXT("CreateNewConditionClass", "Create New Condition Class"), UMovieSceneCondition::StaticClass(), NewConditionName);
 
 						if (Blueprint != NULL && Blueprint->GeneratedClass)
@@ -280,14 +304,18 @@ void FMovieSceneConditionCustomization::FillDirectorBlueprintConditionSubMenu(FM
 					UMovieSceneSequence* ThisSequence = SharedThis->Sequence.Get();
 					if (ThisSequence)
 					{
+						const FScopedTransaction Transaction(LOCTEXT("CreateNewConditionEndpoint", "Create New Condition Endpoint"));
+
+						ThisSequence->Modify();
 						// Create a new director blueprint condition and set it in the details view. Use 'interactive change' so we don't early fire the property finished changing event and reset the details view mid-change
 						PropertyCustomizationHelpers::CreateNewInstanceOfEditInlineObjectClass(SharedThis->ConditionPropertyHandle.ToSharedRef(), UMovieSceneDirectorBlueprintCondition::StaticClass(), EPropertyValueSetFlags::InteractiveChange);
 						TSharedPtr<IPropertyHandle> DirectorBlueprintConditionHandle = SharedThis->ConditionPropertyHandle->GetChildHandle(TEXT("DirectorBlueprintConditionData"));
 						TSharedPtr<FMovieSceneDirectorBlueprintConditionCustomization> BlueprintConditionCustomization = FMovieSceneDirectorBlueprintConditionCustomization::MakeInstance(ThisSequence->GetMovieScene(), DirectorBlueprintConditionHandle, SharedThis->PropertyUtilities);
 						BlueprintConditionCustomization->CreateEndpoint();
 						SharedThis->PropertyUtilities->NotifyFinishedChangingProperties(FPropertyChangedEvent(SharedThis->ConditionPropertyHandle->GetProperty()));
+						// Extra end transaction because we use 'Interactive Change' in the CreateNewInstance call
+						GEditor->EndTransaction();
 						SharedThis->PropertyUtilities->ForceRefresh();
-						ThisSequence->Modify();
 					}
 				}
 			)
@@ -323,17 +351,112 @@ void FMovieSceneConditionCustomization::PopulateQuickBindSubMenu(FMenuBuilder& M
 		{
 			if (!SelectedAction.IsEmpty())
 			{
+				const FScopedTransaction Transaction(LOCTEXT("SetConditionEndpoint", "Set Condition Endpoint"));
+				
 				// Create a new director blueprint condition and set it in the details view. Use 'interactive change' so we don't early fire the property finished changing event and reset the details view mid-change
 				PropertyCustomizationHelpers::CreateNewInstanceOfEditInlineObjectClass(SharedThis->ConditionPropertyHandle.ToSharedRef(), UMovieSceneDirectorBlueprintCondition::StaticClass(), EPropertyValueSetFlags::InteractiveChange);
 				TSharedPtr<IPropertyHandle> DirectorBlueprintConditionHandle = SharedThis->ConditionPropertyHandle->GetChildHandle(TEXT("DirectorBlueprintConditionData"));
 				BlueprintConditionCustomization->SetPropertyHandle(DirectorBlueprintConditionHandle);
 				BlueprintConditionCustomization->HandleQuickBindActionSelected(SelectedAction, InSelectionType, Blueprint, EndpointDefinition);
+				// Extra end transaction because we use 'Interactive Change' in the CreateNewInstance call
+				GEditor->EndTransaction();
 				SharedThis->PropertyUtilities->ForceRefresh();
 			}
 		}));
 	}
 }
 
+void FMovieSceneConditionCustomization::OnUseSelected()
+{
+	// Load selected assets
+	FEditorDelegates::LoadSelectedAssetsIfNeeded.Broadcast();
+
+	TArray<FAssetData> SelectedAssets;
+	GEditor->GetContentBrowserSelections(SelectedAssets);
+
+	for (const FAssetData& AssetData : SelectedAssets)
+	{
+		UBlueprint* SelectedBlueprint = Cast<UBlueprint>(AssetData.GetAsset());
+
+		if (SelectedBlueprint)
+		{
+			if (SelectedBlueprint->GeneratedClass && SelectedBlueprint->GeneratedClass->IsChildOf<UMovieSceneCondition>())
+			{
+				const FScopedTransaction Transaction(LOCTEXT("SetConditionClass", "Set Condition Class"));
+
+				Sequence->Modify();
+				PropertyCustomizationHelpers::CreateNewInstanceOfEditInlineObjectClass(ConditionPropertyHandle.ToSharedRef(), SelectedBlueprint->GeneratedClass, EPropertyValueSetFlags::InteractiveChange);
+				PropertyUtilities->NotifyFinishedChangingProperties(FPropertyChangedEvent(ConditionPropertyHandle->GetProperty()));
+				// Extra end transaction because we use 'Interactive Change' in the CreateNewInstance call
+				GEditor->EndTransaction();
+				PropertyUtilities->ForceRefresh();
+				return;
+			}
+		}
+	}
+}
+
+bool FMovieSceneConditionCustomization::CanUseSelectedAsset() const
+{
+	// Load selected assets
+	FEditorDelegates::LoadSelectedAssetsIfNeeded.Broadcast();
+
+	TArray<FAssetData> SelectedAssets;
+	GEditor->GetContentBrowserSelections(SelectedAssets);
+
+	for (const FAssetData& AssetData : SelectedAssets)
+	{
+		UBlueprint* SelectedBlueprint = Cast<UBlueprint>(AssetData.GetAsset());
+
+		if (SelectedBlueprint)
+		{
+			if (SelectedBlueprint->GeneratedClass && SelectedBlueprint->GeneratedClass->IsChildOf<UMovieSceneCondition>())
+			{
+				return true;
+			}
+		}
+	}
+
+	return false;
+}
+
+void FMovieSceneConditionCustomization::OnBrowseTo()
+{
+	UObject* CurrentValue = NULL;
+	FPropertyAccess::Result Result = ConditionPropertyHandle->GetValue(CurrentValue);
+	if (Result == FPropertyAccess::Success && CurrentValue != NULL)
+	{
+		UClass* CurrentClass = CurrentValue->GetClass();
+		if (CurrentClass)
+		{
+			if (TObjectPtr<UObject> Blueprint = CurrentClass->ClassGeneratedBy)
+			{
+				TArray< UObject* > Objects;
+				Objects.Add(Blueprint.Get());
+				GEditor->SyncBrowserToObjects(Objects);
+			}
+		}
+	}
+}
+
+bool FMovieSceneConditionCustomization::CanBrowseToAsset() const
+{
+	UObject* CurrentValue = NULL;
+	FPropertyAccess::Result Result = ConditionPropertyHandle->GetValue(CurrentValue);
+	if (Result == FPropertyAccess::Success && CurrentValue != NULL)
+	{
+		UClass* CurrentClass = CurrentValue->GetClass();
+		if (CurrentClass)
+		{
+			if (TObjectPtr<UObject> Blueprint = CurrentClass->ClassGeneratedBy)
+			{
+				return true;
+			}
+		}
+	}
+	return false;
+}
+
 TSharedRef<SWidget> FMovieSceneConditionCustomization::GenerateConditionPicker()
 {
 	FMenuBuilder MenuBuilder(true, nullptr, nullptr, true);
diff --git a/Engine/Source/Editor/MovieSceneTools/Private/MovieSceneDirectorBlueprintEndpointCustomization.cpp b/Engine/Source/Editor/MovieSceneTools/Private/MovieSceneDirectorBlueprintEndpointCustomization.cpp
--- a/Engine/Source/Editor/MovieSceneTools/Private/MovieSceneDirectorBlueprintEndpointCustomization.cpp
+++ b/Engine/Source/Editor/MovieSceneTools/Private/MovieSceneDirectorBlueprintEndpointCustomization.cpp
@@ -1008,6 +1008,7 @@ void FMovieSceneDirectorBlueprintEndpointCustomization::HandleQuickBindActionSel
 
 		if (Action)
 		{
+			Blueprint->Modify();
 			UK2Node* NewEndpoint = FMovieSceneDirectorBlueprintUtils::CreateEndpoint(Blueprint, EndpointDefinition);
 
 			UEdGraphPin* EndpointThenPin = NewEndpoint->FindPin(UEdGraphSchema_K2::PN_Then, EGPD_Output);
diff --git a/Engine/Source/Editor/MovieSceneTools/Public/Conditions/MovieSceneConditionCustomization.h b/Engine/Source/Editor/MovieSceneTools/Public/Conditions/MovieSceneConditionCustomization.h
--- a/Engine/Source/Editor/MovieSceneTools/Public/Conditions/MovieSceneConditionCustomization.h
+++ b/Engine/Source/Editor/MovieSceneTools/Public/Conditions/MovieSceneConditionCustomization.h
@@ -101,6 +101,18 @@ private:
 	*/
 	void PopulateQuickBindSubMenu(FMenuBuilder& MenuBuilder);
 
+	/* Used by 'Use Selected' button when clicked to change the condition class to the selected condition class in the asset browser*/
+	void OnUseSelected();
+
+	/* Used to enable/disable the selected asset button based on whether a condition class is currently selected */
+	bool CanUseSelectedAsset() const;
+
+	/* Used by 'Browse To' button when clicked to browse to the condition class currently in use in the asset browser*/
+	void OnBrowseTo();
+
+	/* Used to enable/disable the browse to button based on whether a blueprint condition class is currently in use. */
+	bool CanBrowseToAsset() const;
+
 private:
 
 	/** The property handle for the condition container */
