Sequencer: Moved the data layer information into the tooltip since it doesn't fit in the section display

diff --git a/Engine/Source/Editor/MovieSceneTools/Private/TrackEditors/DataLayerTrackEditor.cpp b/Engine/Source/Editor/MovieSceneTools/Private/TrackEditors/DataLayerTrackEditor.cpp
--- a/Engine/Source/Editor/MovieSceneTools/Private/TrackEditors/DataLayerTrackEditor.cpp
+++ b/Engine/Source/Editor/MovieSceneTools/Private/TrackEditors/DataLayerTrackEditor.cpp
@@ -85,17 +85,44 @@ struct FDataLayerSection
 							.Text(this, &FDataLayerSection::GetPrerollText)
 						]
 					]
-
-					+ SVerticalBox::Slot()
-					[
-						SNew(STextBlock)
-						.Text(this, &FDataLayerSection::GetLayerBarText)
-						.AutoWrapText(true)
-					]
 				]
 			];
 	}
 
+	FText GetSectionToolTip() const override
+	{
+		UMovieSceneDataLayerSection* Section = WeakSection.Get();
+		UDataLayerEditorSubsystem* SubSystem = UDataLayerEditorSubsystem::Get();
+
+		if (SubSystem && Section)
+		{
+			FString LayerName;
+
+			const TArray<UDataLayerAsset*>& DataLayerAssets = Section->GetDataLayerAssets();
+			for (int32 Index = 0; Index < DataLayerAssets.Num(); ++Index)
+			{
+				UDataLayerInstance* DataLayerInstance = SubSystem->GetDataLayerInstance(DataLayerAssets[Index]);
+				if (DataLayerInstance)
+				{
+					LayerName += DataLayerInstance->GetDataLayerFullName();
+				}
+				else
+				{
+					LayerName += FText::Format(LOCTEXT("UnknownDataLayer", "**invalid: {0}**"), FText::FromString(DataLayerAssets[Index]->GetFullName())).ToString();
+				}
+
+				if (Index < DataLayerAssets.Num() - 1)
+				{
+					LayerName += TEXT("\n");
+				}
+			}
+
+			return FText::FromString(LayerName);
+		}
+
+		return FText();
+	}
+
 	FText GetVisibilityText() const
 	{
 		UMovieSceneDataLayerSection* Section = WeakSection.Get();
@@ -128,40 +155,6 @@ struct FDataLayerSection
 		return FText();
 	}
 
-	FText GetLayerBarText() const
-	{
-		UMovieSceneDataLayerSection* Section   = WeakSection.Get();
-		UDataLayerEditorSubsystem*   SubSystem = UDataLayerEditorSubsystem::Get();
-
-		if (SubSystem && Section)
-		{
-			FString LayerName;
-
-			const TArray<UDataLayerAsset*>& DataLayerAssets = Section->GetDataLayerAssets();
-			for (int32 Index = 0; Index < DataLayerAssets.Num(); ++Index)
-			{
-				UDataLayerInstance* DataLayerInstance = SubSystem->GetDataLayerInstance(DataLayerAssets[Index]);
-				if (DataLayerInstance)
-				{
-					LayerName += DataLayerInstance->GetDataLayerFullName();
-				}
-				else
-				{
-					LayerName += FText::Format(LOCTEXT("UnknownDataLayer", "**invalid: {0}**"), FText::FromString(DataLayerAssets[Index]->GetFullName())).ToString();
-				}
-
-				if (Index < DataLayerAssets.Num()-1)
-				{
-					LayerName += TEXT(", ");
-				}
-			}
-
-			return FText::FromString(LayerName);
-		}
-
-		return FText();
-	}
-
 	FSlateColor GetTextColor() const
 	{
 		UMovieSceneDataLayerSection* Section = WeakSection.Get();
