Sequencer: Fixed formatting of Export Dialog's filenames to be consistent with File > Export Selected.

diff --git a/Engine/Source/Editor/MovieSceneTools/Private/LevelSequenceFBXInterop.cpp b/Engine/Source/Editor/MovieSceneTools/Private/LevelSequenceFBXInterop.cpp
--- a/Engine/Source/Editor/MovieSceneTools/Private/LevelSequenceFBXInterop.cpp
+++ b/Engine/Source/Editor/MovieSceneTools/Private/LevelSequenceFBXInterop.cpp
@@ -100,7 +100,7 @@ void FLevelSequenceFBXInterop::ExportFBX()
 	bool bExportFileNamePicked = false;
 	if ( DesktopPlatform != NULL )
 	{
-		FString FileTypes = "FBX document|*.fbx";
+		FString FileTypes = "FBX document (*.fbx)|*.fbx";
 		UMovieSceneSequence* Sequence = Sequencer->GetFocusedMovieSceneSequence();
 		for (TObjectIterator<UClass> It; It; ++It)
 		{
@@ -117,16 +117,16 @@ void FLevelSequenceFBXInterop::ExportFBX()
 
 			for (int32 i = 0; i < Default->FormatExtension.Num(); ++i)
 			{
-				const FString& FormatExtension = Default->FormatExtension[i];
+				// We use force-lowercase here to be consistent with "File > Export Selected..."
+				const FString& FormatExtension = Default->FormatExtension[i].ToLower();
 				const FString& FormatDescription = Default->FormatDescription[i];
 
 				if (FileTypes.Len() > 0)
 				{
 					FileTypes += TEXT("|");
 				}
-				FileTypes += FormatDescription;
-				FileTypes += TEXT("|*.");
-				FileTypes += FormatExtension;
+
+				FileTypes += FString::Printf(TEXT("%s (*.%s)|*.%s"), *FormatDescription, *FormatExtension, *FormatExtension);
 			}
 
 			Exporters.Add(Default);
