Sequencer: Add an option to interpolate the integer channel

Defaults to off so there's no change in existing integer channel evaluation

diff --git a/Engine/Plugins/MovieScene/SequencerScripting/Source/SequencerScripting/Private/KeysAndChannels/MovieSceneScriptingInteger.h b/Engine/Plugins/MovieScene/SequencerScripting/Source/SequencerScripting/Private/KeysAndChannels/MovieSceneScriptingInteger.h
--- a/Engine/Plugins/MovieScene/SequencerScripting/Source/SequencerScripting/Private/KeysAndChannels/MovieSceneScriptingInteger.h
+++ b/Engine/Plugins/MovieScene/SequencerScripting/Source/SequencerScripting/Private/KeysAndChannels/MovieSceneScriptingInteger.h
@@ -173,6 +173,25 @@ public:
 	{
 		return Impl::GetDefaultFromChannel(ChannelHandle).IsSet();
 	}
+
+	/**
+	* Set this channel to evaluate linear keys with interpolation.
+	*/
+	UFUNCTION(BlueprintCallable, Category = "Sequencer|Keys", meta = (DisplayName = "Set Interpolate Linear Keys (Integer)"))
+	void SetInterpolateLinearKeys(bool bInInterpolateLinearKeys)
+	{
+		ChannelHandle.Get()->bInterpolateLinearKeys = bInInterpolateLinearKeys;
+	}
+
+	/**
+	* Get whether this channel will evaluate linear keys with interpolation.
+	*/
+	UFUNCTION(BlueprintCallable, Category = "Sequencer|Keys", meta = (DisplayName = "Get Interpolate Linear Keys (Integer)"))
+	bool GetInterpolateLinearKeys() const
+	{
+		return ChannelHandle.Get()->bInterpolateLinearKeys;
+	}
+
 public:
 	TWeakObjectPtr<UMovieSceneSequence> OwningSequence;
 	TMovieSceneChannelHandle<FMovieSceneIntegerChannel> ChannelHandle;
diff --git a/Engine/Source/Editor/MovieSceneTools/Private/Channels/BuiltInChannelEditors.cpp b/Engine/Source/Editor/MovieSceneTools/Private/Channels/BuiltInChannelEditors.cpp
--- a/Engine/Source/Editor/MovieSceneTools/Private/Channels/BuiltInChannelEditors.cpp
+++ b/Engine/Source/Editor/MovieSceneTools/Private/Channels/BuiltInChannelEditors.cpp
@@ -535,6 +535,38 @@ TSharedRef<SWidget> CreateKeyEditor(const TMovieSceneChannelHandle<FMovieSceneTi
 	return SNew(KeyEditorType, KeyEditor);
 }
 
+template<typename ChannelType>
+void GetTypedChannels(ISequencer* Sequencer, const TSet<TWeakObjectPtr<UMovieSceneSection>>& WeakSections, TArray<ChannelType*>& Channels)
+{
+	// Get selected channels
+	TArray<const IKeyArea*> KeyAreas;
+	Sequencer->GetSelectedKeyAreas(KeyAreas);
+	for (const IKeyArea* KeyArea : KeyAreas)
+	{
+		FMovieSceneChannelHandle Handle = KeyArea->GetChannel();
+		if (Handle.GetChannelTypeName() == ChannelType::StaticStruct()->GetFName())
+		{
+			ChannelType* Channel = static_cast<ChannelType*>(Handle.Get());
+			Channels.Add(Channel);
+		}
+	}
+
+	// Otherwise, the channels of all the sections
+	if (Channels.Num() == 0)
+	{
+		for (TWeakObjectPtr<UMovieSceneSection> WeakSection : WeakSections)
+		{
+			if (UMovieSceneSection* Section = WeakSection.Get())
+			{
+				FMovieSceneChannelProxy& ChannelProxy = Section->GetChannelProxy();
+				for (ChannelType* Channel : ChannelProxy.GetChannels<ChannelType>())
+				{
+					Channels.Add(Channel);
+				}
+			}
+		}
+	}
+}
 
 /** Delegate used to set a class */
 DECLARE_DELEGATE_OneParam(FOnSetActorReferenceKey, FMovieSceneActorReferenceKey);
@@ -1370,6 +1402,22 @@ struct FCurveChannelSectionMenuExtension : TSharedFromThis<FCurveChannelSectionM
 						LOCTEXT("DisplayOptionsTooltip", "Display options"),
 						FNewMenuDelegate::CreateLambda([SharedThis](FMenuBuilder& SubMenuBuilder){ SharedThis->AddDisplayOptionsMenu(SubMenuBuilder); })
 						);
+
+				if (SharedThis->CanInterpolateLinearKeys())
+				{
+					SubMenuBuilder.AddMenuEntry(
+							LOCTEXT("InterpolateLinearKeys", "Interpolate Linear Keys"),
+							LOCTEXT("InterpolateLinearKeysTooltip", "Interpolate linear keys"),
+							FSlateIcon(),
+							FUIAction(
+								FExecuteAction::CreateLambda([SharedThis] { SharedThis->ToggleInterpolateLinearKeys(); }),
+								FCanExecuteAction(),
+								FGetActionCheckState::CreateLambda([SharedThis] { return SharedThis->IsInterpolateLinearKeys(); })
+							),
+							NAME_None,
+							EUserInterfaceActionType::ToggleButton
+						);
+				}
 			}));
 
 		return SharedThis;
@@ -1787,6 +1835,85 @@ struct FCurveChannelSectionMenuExtension : TSharedFromThis<FCurveChannelSectionM
 
 		return true;
 	}
+	
+	bool CanInterpolateLinearKeys() const
+	{
+		ISequencer* Sequencer = WeakSequencer.Pin().Get();
+		if (!Sequencer)
+		{
+			return false;
+		}
+
+		TArray<FMovieSceneIntegerChannel*> IntegerChannels;
+		GetTypedChannels<FMovieSceneIntegerChannel>(Sequencer, WeakSections, IntegerChannels);
+
+		return IntegerChannels.Num() > 0;
+	}
+
+	void ToggleInterpolateLinearKeys()
+	{
+		ISequencer* Sequencer = WeakSequencer.Pin().Get();
+		if (!Sequencer)
+		{
+			return;
+		}
+
+		TArray<FMovieSceneIntegerChannel*> IntegerChannels;
+		GetTypedChannels<FMovieSceneIntegerChannel>(Sequencer, WeakSections, IntegerChannels);
+
+		FScopedTransaction Transaction(LOCTEXT("ToggleInterpolateLinearKeys_Transaction", "Toggle Interpolate Linear Keys"));
+
+		bool bAnythingChanged = false;
+
+		// Modify all sections
+		for (TWeakObjectPtr<UMovieSceneSection> WeakSection : WeakSections)
+		{
+			if (UMovieSceneSection* Section = WeakSection.Get())
+			{
+				Section->Modify();
+			}
+		}
+
+		for (FMovieSceneIntegerChannel* Channel : IntegerChannels)
+		{
+			bAnythingChanged = true;
+			Channel->bInterpolateLinearKeys = !Channel->bInterpolateLinearKeys;
+		}
+
+		if (!bAnythingChanged)
+		{
+			Transaction.Cancel();
+		}
+	}
+
+	ECheckBoxState IsInterpolateLinearKeys()
+	{
+		ISequencer* Sequencer = WeakSequencer.Pin().Get();
+		if (!Sequencer)
+		{
+			return ECheckBoxState::Undetermined;
+		}
+
+		TArray<FMovieSceneIntegerChannel*> IntegerChannels;
+		GetTypedChannels<FMovieSceneIntegerChannel>(Sequencer, WeakSections, IntegerChannels);
+
+		int32 NumInterpolatedAndNotInterpolated[2] = { 0, 0 };
+
+		for (FMovieSceneIntegerChannel* Channel : IntegerChannels)
+		{
+			NumInterpolatedAndNotInterpolated[Channel->bInterpolateLinearKeys ? 0 : 1]++;
+		}
+
+		if (NumInterpolatedAndNotInterpolated[0] == 0 && NumInterpolatedAndNotInterpolated[1] > 0)  // No curve showed, some hidden
+		{
+			return ECheckBoxState::Unchecked;
+		}
+		else if (NumInterpolatedAndNotInterpolated[0] > 0 && NumInterpolatedAndNotInterpolated[1] == 0) // Some curves showed, none hidden
+		{
+			return ECheckBoxState::Checked;
+		}
+		return ECheckBoxState::Undetermined;  // Mixed states, or no curves
+	}
 
 	void ToggleShowCurve()
 	{
diff --git a/Engine/Source/Editor/MovieSceneTools/Private/Channels/IntegerChannelCurveModel.cpp b/Engine/Source/Editor/MovieSceneTools/Private/Channels/IntegerChannelCurveModel.cpp
--- a/Engine/Source/Editor/MovieSceneTools/Private/Channels/IntegerChannelCurveModel.cpp
+++ b/Engine/Source/Editor/MovieSceneTools/Private/Channels/IntegerChannelCurveModel.cpp
@@ -33,6 +33,74 @@
 class FCurveEditor;
 class UObject;
 
+void DrawCurveImpl(const FMovieSceneIntegerChannel* Channel, const FCurveEditorScreenSpace& InScreenSpace, FFrameRate DisplayRate, FFrameRate TickResolution, TArray<TTuple<double, double>>& OutInterpolatingPoints)
+{
+	TMovieSceneChannelData<const int32> ChannelData = Channel->GetData();
+	TArrayView<const FFrameNumber> Times = ChannelData.GetTimes();
+	TArrayView<const int32> Values = ChannelData.GetValues();
+
+	const bool bInterpolateLinearKeys = Channel->bInterpolateLinearKeys;
+	const double StartTimeSeconds = InScreenSpace.GetInputMin();
+	const double EndTimeSeconds = InScreenSpace.GetInputMax();
+
+	const FFrameNumber StartFrame = (StartTimeSeconds * TickResolution).FloorToFrame();
+	const FFrameNumber EndFrame = (EndTimeSeconds * TickResolution).CeilToFrame();
+
+	const int32 StartingIndex = Algo::UpperBound(Times, StartFrame);
+	const int32 EndingIndex = Algo::LowerBound(Times, EndFrame);
+
+	// Add the lower bound of the visible space
+	const bool bValidRange = StartingIndex < EndingIndex;
+	//if we aren't just doing the default constant then we need to sample
+	const bool PreNotConstant = (Channel->PreInfinityExtrap != RCCE_None && Channel->PreInfinityExtrap != RCCE_Constant);
+	const bool PostNotConstant = (Channel->PostInfinityExtrap != RCCE_None && Channel->PostInfinityExtrap != RCCE_Constant);
+	if (bValidRange && (PreNotConstant || PostNotConstant))
+	{
+		const FFrameNumber StartTimeInDisplay = FFrameRate::TransformTime(FFrameTime(StartFrame), TickResolution, DisplayRate).FloorToFrame();
+		const FFrameNumber EndTimeInDisplay = FFrameRate::TransformTime(FFrameTime(EndFrame), TickResolution, DisplayRate).CeilToFrame();
+
+		double Value = 0.0;
+		TOptional<double> PreviousValue;
+		for (FFrameNumber DisplayFrameNumber = StartTimeInDisplay; DisplayFrameNumber <= EndTimeInDisplay; ++DisplayFrameNumber)
+		{
+			FFrameTime TickFrameTime = FFrameRate::TransformTime(FFrameTime(DisplayFrameNumber), DisplayRate, TickResolution);
+			Channel->EvaluateInterp(TickFrameTime, Value);
+			if (!bInterpolateLinearKeys && (PreviousValue.IsSet() && PreviousValue.GetValue() != Value))
+			{
+				OutInterpolatingPoints.Add(MakeTuple(TickFrameTime / TickResolution, PreviousValue.GetValue()));
+			}
+			OutInterpolatingPoints.Add(MakeTuple(TickFrameTime / TickResolution, Value));
+			PreviousValue = Value;
+		}
+	}
+	else
+	{
+		if (bValidRange)
+		{
+			OutInterpolatingPoints.Add(MakeTuple(StartFrame / TickResolution, (double)Values[StartingIndex]));
+		}
+
+		TOptional<double> PreviousValue;
+		for (int32 KeyIndex = StartingIndex; KeyIndex < EndingIndex; ++KeyIndex)
+		{
+			double Value = (double)Values[KeyIndex];
+			if (!bInterpolateLinearKeys && (PreviousValue.IsSet() && PreviousValue.GetValue() != Value))
+			{
+				OutInterpolatingPoints.Add(MakeTuple(Times[KeyIndex] / TickResolution, PreviousValue.GetValue()));
+			}
+
+			OutInterpolatingPoints.Add(MakeTuple(Times[KeyIndex] / TickResolution, Value));
+			PreviousValue = Value;
+		}
+
+		// Add the upper bound of the visible space
+		if (bValidRange)
+		{
+			OutInterpolatingPoints.Add(MakeTuple(EndFrame / TickResolution, (double)Values[EndingIndex - 1]));
+		}
+	}
+}
+
 /**
  * Buffered curve implementation for a integer channel curve model, stores a copy of the integer channel in order to draw itself.
  */
@@ -53,25 +121,10 @@ public:
 
 		if (Section)
 		{
+			FFrameRate DisplayRate = Section->GetTypedOuter<UMovieScene>()->GetDisplayRate();
 			FFrameRate TickResolution = Section->GetTypedOuter<UMovieScene>()->GetTickResolution();
 
-			TMovieSceneChannelData<const int32> ChannelData = Channel.GetData();
-			TArrayView<const FFrameNumber> Times = ChannelData.GetTimes();
-			TArrayView<const int32> Values = ChannelData.GetValues();
-
-			const double StartTimeSeconds = InScreenSpace.GetInputMin();
-			const double EndTimeSeconds = InScreenSpace.GetInputMax();
-
-			const FFrameNumber StartFrame = (StartTimeSeconds * TickResolution).FloorToFrame();
-			const FFrameNumber EndFrame = (EndTimeSeconds * TickResolution).CeilToFrame();
-
-			const int32 StartingIndex = Algo::UpperBound(Times, StartFrame);
-			const int32 EndingIndex = Algo::LowerBound(Times, EndFrame);
-
-			for (int32 KeyIndex = StartingIndex; KeyIndex < EndingIndex; ++KeyIndex)
-			{
-				OutInterpolatingPoints.Add(MakeTuple(Times[KeyIndex] / TickResolution, double(Values[KeyIndex])));
-			}
+			DrawCurveImpl(&Channel, InScreenSpace, DisplayRate, TickResolution, OutInterpolatingPoints);
 		}
 	}
 
@@ -85,6 +138,24 @@ FIntegerChannelCurveModel::FIntegerChannelCurveModel(TMovieSceneChannelHandle<FM
 {
 }
 
+void FIntegerChannelCurveModel::DrawCurve(const FCurveEditor& CurveEditor, const FCurveEditorScreenSpace& InScreenSpace, TArray<TTuple<double, double>>& OutInterpolatingPoints) const
+{
+	UMovieSceneSection* Section = this->GetOwningObjectOrOuter<UMovieSceneSection>();
+	FMovieSceneIntegerChannel* Channel = GetChannelHandle().Get();
+
+	if (Section && Channel && Channel->bInterpolateLinearKeys)
+	{
+		FFrameRate DisplayRate = Section->GetTypedOuter<UMovieScene>()->GetDisplayRate();
+		FFrameRate TickResolution = Section->GetTypedOuter<UMovieScene>()->GetTickResolution();
+
+		DrawCurveImpl(Channel, InScreenSpace, DisplayRate, TickResolution, OutInterpolatingPoints);
+	}
+	else
+	{
+		FChannelCurveModel<FMovieSceneIntegerChannel, int32, int32>::DrawCurve(CurveEditor, InScreenSpace, OutInterpolatingPoints);
+	}
+}
+
 void FIntegerChannelCurveModel::CreateKeyProxies(TArrayView<const FKeyHandle> InKeyHandles, TArrayView<UObject*> OutObjects)
 {
 	for (int32 Index = 0; Index < InKeyHandles.Num(); ++Index)
diff --git a/Engine/Source/Editor/MovieSceneTools/Public/Channels/IntegerChannelCurveModel.h b/Engine/Source/Editor/MovieSceneTools/Public/Channels/IntegerChannelCurveModel.h
--- a/Engine/Source/Editor/MovieSceneTools/Public/Channels/IntegerChannelCurveModel.h
+++ b/Engine/Source/Editor/MovieSceneTools/Public/Channels/IntegerChannelCurveModel.h
@@ -28,6 +28,7 @@ public:
 	FIntegerChannelCurveModel(TMovieSceneChannelHandle<FMovieSceneIntegerChannel> InChannel, UMovieSceneSection* InOwningSection, TWeakPtr<ISequencer> InWeakSequencer);
 
 	// FCurveModel
+	virtual void DrawCurve(const FCurveEditor& CurveEditor, const FCurveEditorScreenSpace& ScreenSpace, TArray<TTuple<double, double>>& InterpolatingPoints) const override;
 	virtual void CreateKeyProxies(TArrayView<const FKeyHandle> InKeyHandles, TArrayView<UObject*> OutObjects) override;
 	virtual TUniquePtr<IBufferedCurveModel> CreateBufferedCurveCopy() const override;
 	virtual void GetCurveAttributes(FCurveAttributes& OutAttributes) const override;
diff --git a/Engine/Source/Runtime/MovieScene/Private/Channels/MovieSceneIntegerChannel.cpp b/Engine/Source/Runtime/MovieScene/Private/Channels/MovieSceneIntegerChannel.cpp
--- a/Engine/Source/Runtime/MovieScene/Private/Channels/MovieSceneIntegerChannel.cpp
+++ b/Engine/Source/Runtime/MovieScene/Private/Channels/MovieSceneIntegerChannel.cpp
@@ -44,6 +44,14 @@ bool FMovieSceneIntegerChannel::SerializeFromMismatchedTag(const FPropertyTag& T
 }
 
 bool FMovieSceneIntegerChannel::Evaluate(FFrameTime InTime, int32& OutValue) const
+{
+	double InterpValue = 0.;
+	bool bReturnValue = EvaluateInterp(InTime, InterpValue);
+	OutValue = FMath::FloorToInt(InterpValue);
+	return bReturnValue;
+}
+
+bool FMovieSceneIntegerChannel::EvaluateInterp(FFrameTime InTime, double& OutValue) const
 {
 	if (Times.Num())
 	{
@@ -101,8 +109,24 @@ bool FMovieSceneIntegerChannel::Evaluate(FFrameTime InTime, int32& OutValue) con
 			}
 		}
 
-		const int32 Index = FMath::Max(0, Algo::UpperBound(Times, Params.Time)-1);
-		OutValue = Values[Index] + (int32)(Params.ValueOffset + 0.5);
+		const int32 Index = FMath::Max(0, Algo::UpperBound(Times, Params.Time) - 1);
+
+		if (bInterpolateLinearKeys && Index < Times.Num() - 1)
+		{
+			const double LowerValue = Values[Index];
+			const double UpperValue = Values[Index+1];
+			const double LowerTime = FFrameTime(Times[Index]).AsDecimal();
+			const double UpperTime = FFrameTime(Times[Index+1]).AsDecimal();
+
+			const double Interp = (Params.Time.AsDecimal() - LowerTime) / (UpperTime - LowerTime);
+
+			const double Value = FMath::Lerp(LowerValue, UpperValue, Interp);
+			OutValue = Value + Params.ValueOffset + 0.5;
+		}
+		else
+		{
+			OutValue = Values[Index] + (int32)(Params.ValueOffset + 0.5);
+		}
 		return true;
 	}
 	else if (bHasDefaultValue)
diff --git a/Engine/Source/Runtime/MovieScene/Public/Channels/MovieSceneIntegerChannel.h b/Engine/Source/Runtime/MovieScene/Public/Channels/MovieSceneIntegerChannel.h
--- a/Engine/Source/Runtime/MovieScene/Public/Channels/MovieSceneIntegerChannel.h
+++ b/Engine/Source/Runtime/MovieScene/Public/Channels/MovieSceneIntegerChannel.h
@@ -34,7 +34,11 @@ struct FMovieSceneIntegerChannel : public FMovieSceneChannel
 	GENERATED_BODY()
 
 	FMovieSceneIntegerChannel()
-		: PreInfinityExtrap(RCCE_Constant), PostInfinityExtrap(RCCE_Constant), DefaultValue(0), bHasDefaultValue(false)
+		: PreInfinityExtrap(RCCE_Constant)
+		, PostInfinityExtrap(RCCE_Constant)
+		, bInterpolateLinearKeys(false)
+		, DefaultValue(0)
+		, bHasDefaultValue(false)
 	{}
 
 	/**
@@ -95,6 +99,15 @@ struct FMovieSceneIntegerChannel : public FMovieSceneChannel
 	 */
 	MOVIESCENE_API bool Evaluate(FFrameTime InTime, int32& OutValue) const;
 
+	/**
+	 * Evaluate the interpolated value when bInterpolateLinearKeys is true
+	 *
+	 * @param InTime     The time to evaluate at
+	 * @param OutValue   A value to receive the result
+	 * @return true if the channel was evaluated successfully, false otherwise
+	 */
+	MOVIESCENE_API bool EvaluateInterp(FFrameTime InTime, double& OutValue) const;
+
 	/**
 	 * Set the channel's times and values to the requested values
 	 */
@@ -178,7 +191,9 @@ public:
 	{
 		bHasDefaultValue = false;
 	}
+
 public:
+
 	/** Pre-infinity extrapolation state, integer channel supports them all but linear since that requires a tangent*/
 	UPROPERTY()
 	TEnumAsByte<ERichCurveExtrapolation> PreInfinityExtrap;
@@ -186,6 +201,11 @@ public:
 	/** Post-infinity extrapolation state, integer channel supports them all but linear since that requires a tangent*/
 	UPROPERTY()
 	TEnumAsByte<ERichCurveExtrapolation> PostInfinityExtrap;
+
+	/** Whether to evaluate linear keys as interpolated or not*/
+	UPROPERTY()
+	bool bInterpolateLinearKeys;
+
 private:
 
 	UPROPERTY(meta=(KeyTimes))
