Extract pre-existing thumbnail capturing for camera cuts to MovieSceneTools (TrackEditorThumbnail.h/.cpp).
This is in preparation of UE-218542, which seeks to capture more relevant thumbnails for level sequences upon save.

diff --git a/Engine/Source/Editor/MovieSceneTools/Private/Sections/ThumbnailSection.cpp b/Engine/Source/Editor/MovieSceneTools/Private/Sections/ThumbnailSection.cpp
--- a/Engine/Source/Editor/MovieSceneTools/Private/Sections/ThumbnailSection.cpp
+++ b/Engine/Source/Editor/MovieSceneTools/Private/Sections/ThumbnailSection.cpp
@@ -22,6 +22,7 @@
 #include "IVREditorModule.h"
 #include "MovieScene.h"
 #include "MovieSceneTimeHelpers.h"
+#include "TrackEditorThumbnail/TrackThumbnailUtils.h"
 
 #define LOCTEXT_NAMESPACE "FThumbnailSection"
 
@@ -399,11 +400,7 @@ void FViewportThumbnailSection::PreDraw(FTrackEditorThumbnail& Thumbnail)
 	TSharedPtr<ISequencer> Sequencer = SequencerPtr.Pin();
 	if (Sequencer.IsValid())
 	{
-		Sequencer->EnterSilentMode();
-		SavedPlaybackStatus = Sequencer->GetPlaybackStatus();
-		Sequencer->SetPlaybackStatus(EMovieScenePlayerStatus::Jumping);
-		Sequencer->SetLocalTimeDirectly(Thumbnail.GetEvalPosition() * Sequencer->GetLocalTime().Rate );
-		Sequencer->ForceEvaluate();
+		UE::MoveSceneTools::PreDrawThumbnailSetupSequencer(*Sequencer, Thumbnail.GetEvalPosition());
 	}
 }
 
@@ -414,7 +411,7 @@ void FViewportThumbnailSection::PostDraw(FTrackEditorThumbnail& Thumbnail)
 	if (Sequencer.IsValid())
 	{
 		Thumbnail.SetupFade(Sequencer->GetSequencerWidget());
-		Sequencer->ExitSilentMode();
+		UE::MoveSceneTools::PostDrawThumbnailCleanupSequencer(*Sequencer);
 	}
 }
 
diff --git a/Engine/Source/Editor/MovieSceneTools/Private/TrackEditorThumbnail/TrackEditorThumbnail.cpp b/Engine/Source/Editor/MovieSceneTools/Private/TrackEditorThumbnail/TrackEditorThumbnail.cpp
--- a/Engine/Source/Editor/MovieSceneTools/Private/TrackEditorThumbnail/TrackEditorThumbnail.cpp
+++ b/Engine/Source/Editor/MovieSceneTools/Private/TrackEditorThumbnail/TrackEditorThumbnail.cpp
@@ -21,6 +21,7 @@
 #include "CanvasTypes.h"
 #include "EngineModule.h"
 #include "SceneViewExtension.h"
+#include "TrackEditorThumbnail/TrackThumbnailUtils.h"
 
 namespace TrackEditorThumbnailConstants
 {
@@ -359,12 +360,11 @@ void FTrackEditorThumbnailCache::DrawThumbnail(FTrackEditorThumbnail& TrackEdito
 	else if (ViewportThumbnailClient)
 	{
 		ViewportThumbnailClient->PreDraw(TrackEditorThumbnail);
-
 		DrawViewportThumbnail(TrackEditorThumbnail);
-
 		ViewportThumbnailClient->PostDraw(TrackEditorThumbnail);
 	}
 
+	// The texture will be rendered by a shader so transition it to SRVMask.
 	FSlateTextureRenderTarget2DResource* pSlateResource = TrackEditorThumbnail.GetRenderTarget();
 	if (pSlateResource != nullptr)
 	{
@@ -379,7 +379,7 @@ void FTrackEditorThumbnailCache::DrawThumbnail(FTrackEditorThumbnail& TrackEdito
 		);
 	}
 }
-void FTrackEditorThumbnailCache::DrawViewportThumbnail(FTrackEditorThumbnail& TrackEditorThumbnail)
+void FTrackEditorThumbnailCache::DrawViewportThumbnail(FTrackEditorThumbnail& TrackEditorThumbnail) const
 {
 	check(ViewportThumbnailClient);
 
@@ -391,8 +391,7 @@ void FTrackEditorThumbnailCache::DrawViewportThumbnail(FTrackEditorThumbnail& Tr
 
 	FMinimalViewInfo ViewInfo;
 	PreviewCameraComponent->GetCameraView(FApp::GetDeltaTime(), ViewInfo);
-
-	FIntPoint RTSize = CalculateTextureSize(ViewInfo);
+	const FIntPoint RTSize = CalculateTextureSize(ViewInfo);
 	if (RTSize.X <= 0 || RTSize.Y <= 0)
 	{
 		return;
@@ -400,77 +399,8 @@ void FTrackEditorThumbnailCache::DrawViewportThumbnail(FTrackEditorThumbnail& Tr
 
 	TrackEditorThumbnail.bIgnoreAlpha = true;
 	TrackEditorThumbnail.ResizeRenderTarget(RTSize);
-
 	UWorld* World = PreviewCameraComponent->GetWorld();
-
-	FSceneViewFamilyContext ViewFamily( FSceneViewFamily::ConstructionValues( TrackEditorThumbnail.GetRenderTarget(), World->Scene, FEngineShowFlags(ESFIM_Game) )
-		.SetTime(FGameTime::GetTimeSinceAppStart())
-		.SetResolveScene(true));
-
-	FSceneViewStateInterface* ViewStateInterface = nullptr;
-
-	// Screen percentage is not supported in thumbnail.
-	ViewFamily.EngineShowFlags.ScreenPercentage = false;
-
-	switch (CurrentCache.Quality)
-	{
-	case EThumbnailQuality::Draft:
-		ViewFamily.EngineShowFlags.DisableAdvancedFeatures();
-		ViewFamily.EngineShowFlags.SetPostProcessing(false);
-		break;
-
-	case EThumbnailQuality::Normal:
-	case EThumbnailQuality::Best:
-		ViewFamily.EngineShowFlags.SetMotionBlur(false);
-
-		// Default eye adaptation requires a viewstate.
-		ViewFamily.EngineShowFlags.EyeAdaptation = true;
-		UMovieSceneUserThumbnailSettings* ThumbnailSettings = GetMutableDefault<UMovieSceneUserThumbnailSettings>();
-		FSceneViewStateInterface* Ref = ThumbnailSettings->ViewState.GetReference();
-		if (!Ref)
-		{
-			ThumbnailSettings->ViewState.Allocate(ViewFamily.GetFeatureLevel());
-		}
-		ViewStateInterface = ThumbnailSettings->ViewState.GetReference();
-		break;
-	}
-
-	FSceneViewInitOptions ViewInitOptions;
-
-	// Use target exposure without blend. 
-	ViewInitOptions.bInCameraCut = true;
-	ViewInitOptions.SceneViewStateInterface = ViewStateInterface;
-
-	ViewInitOptions.BackgroundColor = FLinearColor::Black;
-	ViewInitOptions.SetViewRectangle(FIntRect(FIntPoint::ZeroValue, RTSize));
-	ViewInitOptions.ViewFamily = &ViewFamily;
-
-	ViewInitOptions.ViewOrigin = ViewInfo.Location;
-	ViewInitOptions.ViewRotationMatrix = FInverseRotationMatrix(ViewInfo.Rotation) * FMatrix(
-		FPlane(0, 0, 1, 0),
-		FPlane(1, 0, 0, 0),
-		FPlane(0, 1, 0, 0),
-		FPlane(0, 0, 0, 1));
-
-	ViewInitOptions.ProjectionMatrix = ViewInfo.CalculateProjectionMatrix();
-
-	FSceneView* NewView = new FSceneView(ViewInitOptions);
-	ViewFamily.Views.Add(NewView);
-
-	const float GlobalResolutionFraction = 1.f;
-	ViewFamily.SetScreenPercentageInterface(new FLegacyScreenPercentageDriver(ViewFamily, GlobalResolutionFraction));
-
-	FCanvas Canvas(TrackEditorThumbnail.GetRenderTarget(), nullptr, FGameTime::GetTimeSinceAppStart(), World->Scene->GetFeatureLevel());
-	Canvas.Clear(FLinearColor::Transparent);
-
-	ViewFamily.ViewExtensions = GEngine->ViewExtensions->GatherActiveExtensions(FSceneViewExtensionContext(World->Scene));
-	for (const FSceneViewExtensionRef& Extension : ViewFamily.ViewExtensions)
-	{
-		Extension->SetupViewFamily(ViewFamily);
-		Extension->SetupView(ViewFamily, *NewView);
-	}
-
-	GetRendererModule().BeginRenderingViewFamily(&Canvas, &ViewFamily);
+	UE::MoveSceneTools::DrawViewportThumbnail(*TrackEditorThumbnail.GetRenderTarget(), RTSize, *World->Scene, ViewInfo, CurrentCache.Quality);
 }
 
 
diff --git a/Engine/Source/Editor/MovieSceneTools/Private/TrackEditorThumbnail/TrackThumbnailUtils.cpp b/Engine/Source/Editor/MovieSceneTools/Private/TrackEditorThumbnail/TrackThumbnailUtils.cpp
new file mode 100644
--- /dev/null
+++ b/Engine/Source/Editor/MovieSceneTools/Private/TrackEditorThumbnail/TrackThumbnailUtils.cpp
@@ -0,0 +1,104 @@
+﻿// Copyright Epic Games, Inc. All Rights Reserved.
+
+#include "TrackEditorThumbnail/TrackThumbnailUtils.h"
+
+#include "Camera/CameraTypes.h"
+#include "CanvasTypes.h"
+#include "EngineModule.h"
+#include "LegacyScreenPercentageDriver.h"
+#include "MovieSceneToolsUserSettings.h"
+#include "SceneViewExtension.h"
+
+namespace UE::MoveSceneTools
+{
+	void PreDrawThumbnailSetupSequencer(ISequencer& Sequencer, FFrameTime CaptureFrame)
+	{
+		Sequencer.EnterSilentMode();
+		Sequencer.SetPlaybackStatus(EMovieScenePlayerStatus::Jumping);
+		Sequencer.SetLocalTimeDirectly(CaptureFrame);
+		Sequencer.ForceEvaluate();
+	}
+
+	void PostDrawThumbnailCleanupSequencer(ISequencer& Sequencer)
+	{
+		Sequencer.ExitSilentMode();
+	}
+
+	void DrawViewportThumbnail(
+		FRenderTarget& ThumbnailRenderTarget,
+		const FIntPoint& RenderTargetSize,
+		FSceneInterface& Scene,
+		const FMinimalViewInfo& ViewInfo,
+		EThumbnailQuality Quality
+		)
+	{
+		FSceneViewFamilyContext ViewFamily( FSceneViewFamily::ConstructionValues(&ThumbnailRenderTarget, &Scene, FEngineShowFlags(ESFIM_Game))
+		.SetTime(FGameTime::GetTimeSinceAppStart())
+		.SetResolveScene(true));
+
+		FSceneViewStateInterface* ViewStateInterface = nullptr;
+
+		// Screen percentage is not supported in thumbnail.
+		ViewFamily.EngineShowFlags.ScreenPercentage = false;
+
+		switch (Quality)
+		{
+		case EThumbnailQuality::Draft:
+			ViewFamily.EngineShowFlags.DisableAdvancedFeatures();
+			ViewFamily.EngineShowFlags.SetPostProcessing(false);
+			break;
+
+		case EThumbnailQuality::Normal:
+		case EThumbnailQuality::Best:
+			ViewFamily.EngineShowFlags.SetMotionBlur(false);
+
+			// Default eye adaptation requires a viewstate.
+			ViewFamily.EngineShowFlags.EyeAdaptation = true;
+			UMovieSceneUserThumbnailSettings* ThumbnailSettings = GetMutableDefault<UMovieSceneUserThumbnailSettings>();
+			FSceneViewStateInterface* Ref = ThumbnailSettings->ViewState.GetReference();
+			if (!Ref)
+			{
+				ThumbnailSettings->ViewState.Allocate(ViewFamily.GetFeatureLevel());
+			}
+			ViewStateInterface = ThumbnailSettings->ViewState.GetReference();
+			break;
+		}
+
+		FSceneViewInitOptions ViewInitOptions;
+
+		// Use target exposure without blend. 
+		ViewInitOptions.bInCameraCut = true;
+		ViewInitOptions.SceneViewStateInterface = ViewStateInterface;
+
+		ViewInitOptions.BackgroundColor = FLinearColor::Black;
+		ViewInitOptions.SetViewRectangle(FIntRect(FIntPoint::ZeroValue, RenderTargetSize));
+		ViewInitOptions.ViewFamily = &ViewFamily;
+
+		ViewInitOptions.ViewOrigin = ViewInfo.Location;
+		ViewInitOptions.ViewRotationMatrix = FInverseRotationMatrix(ViewInfo.Rotation) * FMatrix(
+			FPlane(0, 0, 1, 0),
+			FPlane(1, 0, 0, 0),
+			FPlane(0, 1, 0, 0),
+			FPlane(0, 0, 0, 1));
+
+		ViewInitOptions.ProjectionMatrix = ViewInfo.CalculateProjectionMatrix();
+
+		FSceneView* NewView = new FSceneView(ViewInitOptions);
+		ViewFamily.Views.Add(NewView);
+
+		const float GlobalResolutionFraction = 1.f;
+		ViewFamily.SetScreenPercentageInterface(new FLegacyScreenPercentageDriver(ViewFamily, GlobalResolutionFraction));
+
+		FCanvas Canvas(&ThumbnailRenderTarget, nullptr, FGameTime::GetTimeSinceAppStart(), Scene.GetFeatureLevel());
+		Canvas.Clear(FLinearColor::Transparent);
+
+		ViewFamily.ViewExtensions = GEngine->ViewExtensions->GatherActiveExtensions(FSceneViewExtensionContext(&Scene));
+		for (const FSceneViewExtensionRef& Extension : ViewFamily.ViewExtensions)
+		{
+			Extension->SetupViewFamily(ViewFamily);
+			Extension->SetupView(ViewFamily, *NewView);
+		}
+
+		GetRendererModule().BeginRenderingViewFamily(&Canvas, &ViewFamily);
+	}
+}
diff --git a/Engine/Source/Editor/MovieSceneTools/Public/TrackEditorThumbnail/TrackEditorThumbnail.h b/Engine/Source/Editor/MovieSceneTools/Public/TrackEditorThumbnail/TrackEditorThumbnail.h
--- a/Engine/Source/Editor/MovieSceneTools/Public/TrackEditorThumbnail/TrackEditorThumbnail.h
+++ b/Engine/Source/Editor/MovieSceneTools/Public/TrackEditorThumbnail/TrackEditorThumbnail.h
@@ -213,7 +213,7 @@ protected:
 
 	void DrawThumbnail(FTrackEditorThumbnail& TrackEditorThumbnail);
 
-	void DrawViewportThumbnail(FTrackEditorThumbnail& TrackEditorThumbnail);
+	void DrawViewportThumbnail(FTrackEditorThumbnail& TrackEditorThumbnail) const;
 
 	void ComputeNewThumbnails();
 
diff --git a/Engine/Source/Editor/MovieSceneTools/Public/TrackEditorThumbnail/TrackThumbnailUtils.h b/Engine/Source/Editor/MovieSceneTools/Public/TrackEditorThumbnail/TrackThumbnailUtils.h
new file mode 100644
--- /dev/null
+++ b/Engine/Source/Editor/MovieSceneTools/Public/TrackEditorThumbnail/TrackThumbnailUtils.h
@@ -0,0 +1,34 @@
+﻿// Copyright Epic Games, Inc. All Rights Reserved.
+
+#pragma once
+
+#include "ISequencer.h"
+#include "Math/MathFwd.h"
+#include "Misc/FrameTime.h"
+
+class FSceneInterface;
+class FRenderTarget;
+enum class EThumbnailQuality : uint8;
+struct FMinimalViewInfo;
+
+namespace UE::MoveSceneTools
+{
+	/** Prepares the sequencer so a thumbnail can be captured from it. */
+	MOVIESCENETOOLS_API void PreDrawThumbnailSetupSequencer(ISequencer& Sequencer, FFrameTime CaptureFrame);
+	inline void PreDrawThumbnailSetupSequencer(ISequencer& Sequencer, double ThumbnailEvalPosition)
+	{
+		PreDrawThumbnailSetupSequencer(Sequencer, ThumbnailEvalPosition * Sequencer.GetLocalTime().Rate);
+	}
+	
+	/** Cleans up the sequencer from the consequences of PreDrawThumbnailSetupSequencer. */
+	MOVIESCENETOOLS_API void PostDrawThumbnailCleanupSequencer(ISequencer& Sequencer);
+
+	/** Renders the thumbnail's render target in the current Scene. */
+	MOVIESCENETOOLS_API void DrawViewportThumbnail(
+		FRenderTarget& ThumbnailRenderTarget,
+		const FIntPoint& RenderTargetSize,
+		FSceneInterface& Scene,
+		const FMinimalViewInfo& ViewInfo,
+		EThumbnailQuality Quality
+		);
+}
