Cache last evaluated view target in new FCameraCutViewTargetCacheCapability.
- Deprecated FCinematicShotTrackEditor::GetCinematicShotCamera
- Added util ISequencer::GetLastEvaluatedCameraCutView

diff --git a/Engine/Source/Editor/MovieSceneTools/Private/TrackEditors/CinematicShotTrackEditor.cpp b/Engine/Source/Editor/MovieSceneTools/Private/TrackEditors/CinematicShotTrackEditor.cpp
--- a/Engine/Source/Editor/MovieSceneTools/Private/TrackEditors/CinematicShotTrackEditor.cpp
+++ b/Engine/Source/Editor/MovieSceneTools/Private/TrackEditors/CinematicShotTrackEditor.cpp
@@ -56,19 +56,10 @@ TSharedRef<ISequencerTrackEditor> FCinematicShotTrackEditor::CreateTrackEditor(T
 	return MakeShareable(new FCinematicShotTrackEditor(InSequencer));
 }
 
-
-void FCinematicShotTrackEditor::OnInitialize()
-{
-	OnCameraCutHandle = GetSequencer()->OnCameraCut().AddSP(this, &FCinematicShotTrackEditor::OnUpdateCameraCut);
-}
-
-
-void FCinematicShotTrackEditor::OnRelease()
+TWeakObjectPtr<AActor> FCinematicShotTrackEditor::GetCinematicShotCamera() const
 {
-	if (OnCameraCutHandle.IsValid() && GetSequencer().IsValid())
-	{
-		GetSequencer()->OnCameraCut().Remove(OnCameraCutHandle);
-	}
+	const UCameraComponent* Camera = GetSequencer()->GetLastEvaluatedCameraCut().Get();
+	return Camera ? Camera->GetOwner() : nullptr;
 }
 
 
@@ -311,12 +302,6 @@ FText FCinematicShotTrackEditor::GetLockShotsToolTip() const
 		LOCTEXT("LockShots", "Lock Viewport to Shots");
 }
 
-void FCinematicShotTrackEditor::OnUpdateCameraCut(UObject* CameraObject, bool bJumpCut)
-{
-	// Keep track of the camera when it switches so that the thumbnail can be drawn with the correct camera
-	CinematicShotCamera = Cast<AActor>(CameraObject);
-}
-
 UAutomatedLevelSequenceCapture* GetMovieSceneCapture()
 {
 	UAutomatedLevelSequenceCapture* MovieSceneCapture = Cast<UAutomatedLevelSequenceCapture>(IMovieSceneCaptureModule::Get().GetFirstActiveMovieSceneCapture());
diff --git a/Engine/Source/Editor/MovieSceneTools/Public/TrackEditors/CinematicShotTrackEditor.h b/Engine/Source/Editor/MovieSceneTools/Public/TrackEditors/CinematicShotTrackEditor.h
--- a/Engine/Source/Editor/MovieSceneTools/Public/TrackEditors/CinematicShotTrackEditor.h
+++ b/Engine/Source/Editor/MovieSceneTools/Public/TrackEditors/CinematicShotTrackEditor.h
@@ -44,13 +44,12 @@ public:
 	 */
 	static TSharedRef<ISequencerTrackEditor> CreateTrackEditor(TSharedRef<ISequencer> OwningSequencer);
 
-	TWeakObjectPtr<AActor> GetCinematicShotCamera() const { return CinematicShotCamera; }
+	UE_DEPRECATED(5.5, "Use FCameraCutPlaybackCapability::LastViewTargetCamera instead.")
+	TWeakObjectPtr<AActor> GetCinematicShotCamera() const;
 
 public:
 
 	// ISequencerTrackEditor interface
-	virtual void OnInitialize() override;
-	virtual void OnRelease() override;
 	virtual TSharedPtr<SWidget> BuildOutlinerColumnWidget(const FBuildColumnWidgetParams& Params, const FName& ColumnName) override;
 	virtual TSharedRef<ISequencerSection> MakeSectionInterface(UMovieSceneSection& SectionObject, UMovieSceneTrack& Track, FGuid ObjectBinding) override;
 	virtual bool SupportsSequence(UMovieSceneSequence* InSequence) const override;
@@ -91,9 +90,6 @@ private:
 	/** Delegate for shots button lock tooltip */
 	FText GetLockShotsToolTip() const;
 
-	/** Called when our sequencer wants to switch cameras */
-	void OnUpdateCameraCut(UObject* CameraObject, bool bJumpCut);
-
 	/** Callback for ImportEDL. */
 	void ImportEDL();
 	
@@ -110,10 +106,4 @@ private:
 
 	/** The Thumbnail pool which draws all the viewport thumbnails for the shot track. */
 	TSharedPtr<FTrackEditorThumbnailPool> ThumbnailPool;
-
-	/** The camera actor for the current cut. */
-	TWeakObjectPtr<AActor> CinematicShotCamera;
-
-	/** Delegate binding handle for ISequencer::OnCameraCut */
-	FDelegateHandle OnCameraCutHandle;
 };
diff --git a/Engine/Source/Editor/Sequencer/Private/Capabilities/CameraCutViewTargetCacheCapability.cpp b/Engine/Source/Editor/Sequencer/Private/Capabilities/CameraCutViewTargetCacheCapability.cpp
new file mode 100644
--- /dev/null
+++ b/Engine/Source/Editor/Sequencer/Private/Capabilities/CameraCutViewTargetCacheCapability.cpp
@@ -0,0 +1,9 @@
+﻿// Copyright Epic Games, Inc. All Rights Reserved.
+
+#include "CameraCutViewTargetCacheCapability.h"
+
+namespace UE::MovieScene
+{
+	const TPlaybackCapabilityID<FCameraCutViewTargetCacheCapability> FCameraCutViewTargetCacheCapability::ID
+		= TPlaybackCapabilityID<FCameraCutViewTargetCacheCapability>::Register();
+}
\ No newline at end of file
diff --git a/Engine/Source/Editor/Sequencer/Private/Capabilities/CameraCutViewTargetCacheCapability.h b/Engine/Source/Editor/Sequencer/Private/Capabilities/CameraCutViewTargetCacheCapability.h
new file mode 100644
--- /dev/null
+++ b/Engine/Source/Editor/Sequencer/Private/Capabilities/CameraCutViewTargetCacheCapability.h
@@ -0,0 +1,23 @@
+﻿// Copyright Epic Games, Inc. All Rights Reserved.
+
+#pragma once
+
+#include "Evaluation/IMovieScenePlaybackCapability.h"
+#include "UObject/WeakObjectPtr.h"
+#include "UObject/WeakObjectPtrTemplates.h"
+
+class UCameraComponent;
+
+namespace UE::MovieScene
+{
+	struct FCameraCutViewTargetCacheCapability
+	{
+		/** Playback capability ID */
+		static const TPlaybackCapabilityID<FCameraCutViewTargetCacheCapability> ID;
+		
+		/** The last evaluated view target. */
+		TWeakObjectPtr<UCameraComponent> LastViewTargetCamera;
+	};
+}
+
+
diff --git a/Engine/Source/Editor/Sequencer/Private/ISequencer.cpp b/Engine/Source/Editor/Sequencer/Private/ISequencer.cpp
--- a/Engine/Source/Editor/Sequencer/Private/ISequencer.cpp
+++ b/Engine/Source/Editor/Sequencer/Private/ISequencer.cpp
@@ -3,12 +3,22 @@
 #include "ISequencer.h"
 
 #include "AnimatedRange.h"
+#include "Camera/CameraComponent.h"
+#include "Capabilities/CameraCutViewTargetCacheCapability.h"
 #include "Misc/AssertionMacros.h"
 #include "MovieScene.h"
 #include "MovieSceneSequence.h"
 #include "ITimeSlider.h"
 #include "SequencerUtilities.h"
 
+TWeakObjectPtr<UCameraComponent> ISequencer::GetLastEvaluatedCameraCut() const
+{
+	const TSharedPtr<const UE::MovieScene::FSharedPlaybackState> PlaybackState = FindSharedPlaybackState();
+	const UE::MovieScene::FCameraCutViewTargetCacheCapability* Capability = PlaybackState
+	   ? PlaybackState->FindCapability<UE::MovieScene::FCameraCutViewTargetCacheCapability>() : nullptr;
+	return Capability ? Capability->LastViewTargetCamera.Get() : nullptr;
+}
+
 FAnimatedRange ISequencer::GetViewRange() const
 {
 	return FAnimatedRange();
diff --git a/Engine/Source/Editor/Sequencer/Private/Sequencer.cpp b/Engine/Source/Editor/Sequencer/Private/Sequencer.cpp
--- a/Engine/Source/Editor/Sequencer/Private/Sequencer.cpp
+++ b/Engine/Source/Editor/Sequencer/Private/Sequencer.cpp
@@ -722,6 +722,7 @@ void FSequencer::InitRootSequenceInstance()
 	if (ensure(SharedPlaybackState.IsValid()))
 	{
 		SharedPlaybackState->SetOrAddCapabilityRaw<FCameraCutPlaybackCapability>((FCameraCutPlaybackCapability*)this);
+		SharedPlaybackState->SetOrAddCapabilityRaw<FCameraCutViewTargetCacheCapability>((FCameraCutViewTargetCacheCapability*)this);
 	}
 }
 
@@ -12106,6 +12107,7 @@ float FSequencer::GetCameraBlendPlayRate()
 
 void FSequencer::OnCameraCutUpdated(const UE::MovieScene::FOnCameraCutUpdatedParams& Params)
 {
+	LastViewTargetCamera = Params.ViewTargetCamera;
 	OnCameraCutEvent.Broadcast(Params.ViewTarget, Params.bIsJumpCut);
 }
 
diff --git a/Engine/Source/Editor/Sequencer/Private/Sequencer.h b/Engine/Source/Editor/Sequencer/Private/Sequencer.h
--- a/Engine/Source/Editor/Sequencer/Private/Sequencer.h
+++ b/Engine/Source/Editor/Sequencer/Private/Sequencer.h
@@ -40,6 +40,7 @@
 #include "Curves/RichCurve.h"
 #include "Sections/MovieScene3DTransformSection.h"
 #include "SequencerTimeChangeUndoRedoProxy.h"
+#include "Capabilities/CameraCutViewTargetCacheCapability.h"
 #include "Filters/ISequencerTrackFilters.h"
 
 class AActor;
@@ -116,6 +117,7 @@ class FSequencer final
 	, public FEditorUndoClient
 	, public FTickableEditorObject
 	, public UE::MovieScene::FCameraCutPlaybackCapability
+	, private UE::MovieScene::FCameraCutViewTargetCacheCapability
 {
 	using FViewModel = UE::Sequencer::FViewModel;
 
diff --git a/Engine/Source/Editor/Sequencer/Public/ISequencer.h b/Engine/Source/Editor/Sequencer/Public/ISequencer.h
--- a/Engine/Source/Editor/Sequencer/Public/ISequencer.h
+++ b/Engine/Source/Editor/Sequencer/Public/ISequencer.h
@@ -51,6 +51,7 @@ class ISidebarDrawerContent;
 class ISequencerTrackEditor;
 class SWidget;
 class UActorFactory;
+class UCameraComponent;
 class UMovieSceneCinematicShotSection;
 class UMovieSceneFolder;
 class UMovieSceneSection;
@@ -469,6 +470,9 @@ public:
 
 	/** Forcefully reevaluate the sequence immediately */
 	virtual void ForceEvaluate() = 0;
+	
+	/** @return The camera cut that was last used by a camera cut; it is the view that the sequence should have at the current time. */
+	SEQUENCER_API TWeakObjectPtr<UCameraComponent> GetLastEvaluatedCameraCut() const;
 
 	/** Reset the timing manager to the clock source specified by the root movie scene */
 	virtual void ResetTimeController() = 0;
