Sequencer: Evaluate when color picker has changed

diff --git a/Engine/Source/Editor/MovieSceneTools/Private/Sections/ColorPropertySection.cpp b/Engine/Source/Editor/MovieSceneTools/Private/Sections/ColorPropertySection.cpp
--- a/Engine/Source/Editor/MovieSceneTools/Private/Sections/ColorPropertySection.cpp
+++ b/Engine/Source/Editor/MovieSceneTools/Private/Sections/ColorPropertySection.cpp
@@ -58,7 +58,7 @@ FReply FColorPropertySection::OnKeyDoubleClicked(const TArray<FKeyHandle>& KeyHa
 	FMovieSceneFloatChannel* BChannel = Proxy.GetChannel<FMovieSceneFloatChannel>(2);
 	FMovieSceneFloatChannel* AChannel = Proxy.GetChannel<FMovieSceneFloatChannel>(3);
 
-	FMovieSceneKeyColorPicker KeyColorPicker(ColorSection, RChannel, GChannel, BChannel, AChannel, KeyHandles);
+	FMovieSceneKeyColorPicker KeyColorPicker(ColorSection, RChannel, GChannel, BChannel, AChannel, KeyHandles, WeakSequencer);
 
 	return FReply::Handled();
 }
diff --git a/Engine/Source/Editor/MovieSceneTools/Private/Sections/ComponentMaterialParameterSection.cpp b/Engine/Source/Editor/MovieSceneTools/Private/Sections/ComponentMaterialParameterSection.cpp
--- a/Engine/Source/Editor/MovieSceneTools/Private/Sections/ComponentMaterialParameterSection.cpp
+++ b/Engine/Source/Editor/MovieSceneTools/Private/Sections/ComponentMaterialParameterSection.cpp
@@ -9,6 +9,7 @@
 #include "IKeyArea.h"
 #include "Internationalization/Internationalization.h"
 #include "Internationalization/Text.h"
+#include "ISequencer.h"
 #include "Layout/Geometry.h"
 #include "Layout/PaintGeometry.h"
 #include "Math/Color.h"
@@ -192,7 +193,7 @@ FReply FComponentMaterialParameterSection::OnKeyDoubleClicked(const TArray<FKeyH
 
 	for (FColorMaterialParameterInfoAndCurves& NameAndCurve : ParameterSection->ColorParameterInfosAndCurves)
 	{
-		FMovieSceneKeyColorPicker KeyColorPicker(ParameterSection, &NameAndCurve.RedCurve, &NameAndCurve.GreenCurve, &NameAndCurve.BlueCurve, &NameAndCurve.AlphaCurve, KeyHandles);
+		FMovieSceneKeyColorPicker KeyColorPicker(ParameterSection, &NameAndCurve.RedCurve, &NameAndCurve.GreenCurve, &NameAndCurve.BlueCurve, &NameAndCurve.AlphaCurve, KeyHandles, WeakSequencer);
 	}
 
 	return FReply::Handled();
diff --git a/Engine/Source/Editor/MovieSceneTools/Private/Sections/ComponentMaterialParameterSection.h b/Engine/Source/Editor/MovieSceneTools/Private/Sections/ComponentMaterialParameterSection.h
--- a/Engine/Source/Editor/MovieSceneTools/Private/Sections/ComponentMaterialParameterSection.h
+++ b/Engine/Source/Editor/MovieSceneTools/Private/Sections/ComponentMaterialParameterSection.h
@@ -9,6 +9,7 @@
 
 class FName;
 class FSequencerSectionPainter;
+class ISequencer;
 class UMovieSceneSection;
 struct FKeyHandle;
 
@@ -20,8 +21,9 @@ class FComponentMaterialParameterSection
 {
 public:
 
-	FComponentMaterialParameterSection(UMovieSceneSection& InSectionObject)
+	FComponentMaterialParameterSection(UMovieSceneSection& InSectionObject, TWeakPtr<ISequencer> InSequencer)
 		: FSequencerSection(InSectionObject)
+		, WeakSequencer(InSequencer)
 	{ }
 
 public:
@@ -32,4 +34,9 @@ public:
 	virtual bool RequestDeleteKeyArea(const TArray<FName>& KeyAreaNamePath) override;
 
 	virtual TSharedPtr<UE::Sequencer::FCategoryModel> ConstructCategoryModel(FName InCategoryName, const FText& InDisplayText, TArrayView<const FChannelData> Channels) const;
+
+private:
+
+	/** Weak pointer to the sequencer this section is for */
+	TWeakPtr<ISequencer> WeakSequencer;
 };
diff --git a/Engine/Source/Editor/MovieSceneTools/Private/Sections/MovieSceneSectionHelpers.cpp b/Engine/Source/Editor/MovieSceneTools/Private/Sections/MovieSceneSectionHelpers.cpp
--- a/Engine/Source/Editor/MovieSceneTools/Private/Sections/MovieSceneSectionHelpers.cpp
+++ b/Engine/Source/Editor/MovieSceneTools/Private/Sections/MovieSceneSectionHelpers.cpp
@@ -6,6 +6,7 @@
 #include "Channels/MovieSceneChannelProxy.h"
 #include "Channels/MovieSceneFloatChannel.h"
 
+#include "ISequencer.h"
 #include "Misc/FrameNumber.h"
 #include "Widgets/Colors/SColorPicker.h"
 #include "ScopedTransaction.h"
@@ -87,7 +88,7 @@ FFrameNumber FMovieSceneKeyColorPicker::KeyTime = FFrameNumber();
 FLinearColor FMovieSceneKeyColorPicker::InitialColor = FLinearColor();
 bool FMovieSceneKeyColorPicker::bColorPickerWasCancelled = false;
 
-FMovieSceneKeyColorPicker::FMovieSceneKeyColorPicker(UMovieSceneSection* Section, FMovieSceneFloatChannel* RChannel, FMovieSceneFloatChannel* GChannel, FMovieSceneFloatChannel* BChannel, FMovieSceneFloatChannel* AChannel, const TArray<FKeyHandle>& KeyHandles)
+FMovieSceneKeyColorPicker::FMovieSceneKeyColorPicker(UMovieSceneSection* Section, FMovieSceneFloatChannel* RChannel, FMovieSceneFloatChannel* GChannel, FMovieSceneFloatChannel* BChannel, FMovieSceneFloatChannel* AChannel, const TArray<FKeyHandle>& KeyHandles, TWeakPtr<ISequencer> InSequencer)
 {
 	// One key was clicked on, find which from the channels
 	int32 KeyIndex = INDEX_NONE;
@@ -141,9 +142,9 @@ FMovieSceneKeyColorPicker::FMovieSceneKeyColorPicker(UMovieSceneSection* Section
 	PickerArgs.ParentWidget = FSlateApplication::Get().GetActiveTopLevelWindow();
 	PickerArgs.OptionalOwningDetailsView = PickerArgs.ParentWidget;
 
-	PickerArgs.OnColorCommitted = FOnLinearColorValueChanged::CreateRaw(this, &FMovieSceneKeyColorPicker::OnColorPickerPicked, RChannel, GChannel, BChannel, AChannel);
-	PickerArgs.OnColorPickerWindowClosed = FOnWindowClosed::CreateRaw(this, &FMovieSceneKeyColorPicker::OnColorPickerClosed, Section, RChannel, GChannel, BChannel, AChannel);
-	PickerArgs.OnColorPickerCancelled = FOnColorPickerCancelled::CreateRaw(this, &FMovieSceneKeyColorPicker::OnColorPickerCancelled, RChannel, GChannel, BChannel, AChannel);
+	PickerArgs.OnColorCommitted = FOnLinearColorValueChanged::CreateRaw(this, &FMovieSceneKeyColorPicker::OnColorPickerPicked, RChannel, GChannel, BChannel, AChannel, InSequencer);
+	PickerArgs.OnColorPickerWindowClosed = FOnWindowClosed::CreateRaw(this, &FMovieSceneKeyColorPicker::OnColorPickerClosed, Section, RChannel, GChannel, BChannel, AChannel, InSequencer);
+	PickerArgs.OnColorPickerCancelled = FOnColorPickerCancelled::CreateRaw(this, &FMovieSceneKeyColorPicker::OnColorPickerCancelled, RChannel, GChannel, BChannel, AChannel, InSequencer);
 
 	OpenColorPicker(PickerArgs);
 }
@@ -161,15 +162,21 @@ void UpdateOrAddKey(FMovieSceneFloatChannel* Channel, FFrameNumber KeyTime, floa
 	}
 }
 
-void FMovieSceneKeyColorPicker::OnColorPickerPicked(FLinearColor NewColor, FMovieSceneFloatChannel* RChannel, FMovieSceneFloatChannel* GChannel, FMovieSceneFloatChannel* BChannel, FMovieSceneFloatChannel* AChannel)
+void FMovieSceneKeyColorPicker::OnColorPickerPicked(FLinearColor NewColor, FMovieSceneFloatChannel* RChannel, FMovieSceneFloatChannel* GChannel, FMovieSceneFloatChannel* BChannel, FMovieSceneFloatChannel* AChannel, TWeakPtr<ISequencer> WeakSequencer)
 {
 	UpdateOrAddKey(RChannel, KeyTime, NewColor.R);
 	UpdateOrAddKey(GChannel, KeyTime, NewColor.G);
 	UpdateOrAddKey(BChannel, KeyTime, NewColor.B);
 	UpdateOrAddKey(AChannel, KeyTime, NewColor.A);
+
+	TSharedPtr<ISequencer> Sequencer = WeakSequencer.Pin();
+	if (Sequencer)
+	{
+		Sequencer->NotifyMovieSceneDataChanged(EMovieSceneDataChangeType::TrackValueChanged);
+	}
 }
 
-void FMovieSceneKeyColorPicker::OnColorPickerClosed(const TSharedRef<SWindow>& Window, UMovieSceneSection* Section, FMovieSceneFloatChannel* RChannel, FMovieSceneFloatChannel* GChannel, FMovieSceneFloatChannel* BChannel, FMovieSceneFloatChannel* AChannel)
+void FMovieSceneKeyColorPicker::OnColorPickerClosed(const TSharedRef<SWindow>& Window, UMovieSceneSection* Section, FMovieSceneFloatChannel* RChannel, FMovieSceneFloatChannel* GChannel, FMovieSceneFloatChannel* BChannel, FMovieSceneFloatChannel* AChannel, TWeakPtr<ISequencer> WeakSequencer)
 {
 	// Under Unreal UX terms, closing the Color Picker (via the UI) is the same as confirming it since we've been live updating
 	// the color. The track already has the latest color change so we undo the change before calling Modify so that Undo sets us
@@ -195,10 +202,16 @@ void FMovieSceneKeyColorPicker::OnColorPickerClosed(const TSharedRef<SWindow>& W
 		UpdateOrAddKey(GChannel, KeyTime, G);
 		UpdateOrAddKey(BChannel, KeyTime, B);
 		UpdateOrAddKey(AChannel, KeyTime, A);
+
+		TSharedPtr<ISequencer> Sequencer = WeakSequencer.Pin();
+		if (Sequencer)
+		{
+			Sequencer->NotifyMovieSceneDataChanged(EMovieSceneDataChangeType::TrackValueChanged);
+		}
 	}
 }
 
-void FMovieSceneKeyColorPicker::OnColorPickerCancelled(FLinearColor NewColor, FMovieSceneFloatChannel* RChannel, FMovieSceneFloatChannel* GChannel, FMovieSceneFloatChannel* BChannel, FMovieSceneFloatChannel* AChannel)
+void FMovieSceneKeyColorPicker::OnColorPickerCancelled(FLinearColor NewColor, FMovieSceneFloatChannel* RChannel, FMovieSceneFloatChannel* GChannel, FMovieSceneFloatChannel* BChannel, FMovieSceneFloatChannel* AChannel, TWeakPtr<ISequencer> WeakSequencer)
 {
 	bColorPickerWasCancelled = true;
 	
@@ -207,4 +220,10 @@ void FMovieSceneKeyColorPicker::OnColorPickerCancelled(FLinearColor NewColor, FM
 	UpdateOrAddKey(GChannel, KeyTime, InitialColor.G);
 	UpdateOrAddKey(BChannel, KeyTime, InitialColor.B);
 	UpdateOrAddKey(AChannel, KeyTime, InitialColor.A);
+
+	TSharedPtr<ISequencer> Sequencer = WeakSequencer.Pin();
+	if (Sequencer)
+	{
+		Sequencer->NotifyMovieSceneDataChanged(EMovieSceneDataChangeType::TrackValueChanged);
+	}
 }
diff --git a/Engine/Source/Editor/MovieSceneTools/Private/Sections/MovieSceneSectionHelpers.h b/Engine/Source/Editor/MovieSceneTools/Private/Sections/MovieSceneSectionHelpers.h
--- a/Engine/Source/Editor/MovieSceneTools/Private/Sections/MovieSceneSectionHelpers.h
+++ b/Engine/Source/Editor/MovieSceneTools/Private/Sections/MovieSceneSectionHelpers.h
@@ -8,6 +8,7 @@
 struct FKeyHandle;
 struct FTimeToPixel;
 struct FMovieSceneFloatChannel;
+class ISequencer;
 class UMovieSceneSection;
 
 class MOVIESCENETOOLS_API MovieSceneSectionHelpers
@@ -21,13 +22,13 @@ public:
 class MOVIESCENETOOLS_API FMovieSceneKeyColorPicker
 {
 public:
-	FMovieSceneKeyColorPicker(UMovieSceneSection* Section, FMovieSceneFloatChannel* RChannel, FMovieSceneFloatChannel* GChannel, FMovieSceneFloatChannel* BChannel, FMovieSceneFloatChannel* AChannel, const TArray<FKeyHandle>& KeyHandles);
+	FMovieSceneKeyColorPicker(UMovieSceneSection* Section, FMovieSceneFloatChannel* RChannel, FMovieSceneFloatChannel* GChannel, FMovieSceneFloatChannel* BChannel, FMovieSceneFloatChannel* AChannel, const TArray<FKeyHandle>& KeyHandles, TWeakPtr<ISequencer> InSequencer);
 
 private:
 		
-	void OnColorPickerPicked(FLinearColor NewFolderColor, FMovieSceneFloatChannel* RChannel, FMovieSceneFloatChannel* GChannel, FMovieSceneFloatChannel* BChannel, FMovieSceneFloatChannel* AChannel);
-	void OnColorPickerClosed(const TSharedRef<SWindow>& Window, UMovieSceneSection* Section, FMovieSceneFloatChannel* RChannel, FMovieSceneFloatChannel* GChannel, FMovieSceneFloatChannel* BChannel, FMovieSceneFloatChannel* AChannel);
-	void OnColorPickerCancelled(FLinearColor NewFolderColor, FMovieSceneFloatChannel* RChannel, FMovieSceneFloatChannel* GChannel, FMovieSceneFloatChannel* BChannel, FMovieSceneFloatChannel* AChannel);
+	void OnColorPickerPicked(FLinearColor NewFolderColor, FMovieSceneFloatChannel* RChannel, FMovieSceneFloatChannel* GChannel, FMovieSceneFloatChannel* BChannel, FMovieSceneFloatChannel* AChannel, TWeakPtr<ISequencer> InSequencer);
+	void OnColorPickerClosed(const TSharedRef<SWindow>& Window, UMovieSceneSection* Section, FMovieSceneFloatChannel* RChannel, FMovieSceneFloatChannel* GChannel, FMovieSceneFloatChannel* BChannel, FMovieSceneFloatChannel* AChannel, TWeakPtr<ISequencer> InSequencer);
+	void OnColorPickerCancelled(FLinearColor NewFolderColor, FMovieSceneFloatChannel* RChannel, FMovieSceneFloatChannel* GChannel, FMovieSceneFloatChannel* BChannel, FMovieSceneFloatChannel* AChannel, TWeakPtr<ISequencer> InSequencer);
 
 private:
 
diff --git a/Engine/Source/Editor/MovieSceneTools/Private/Sections/ParameterSection.cpp b/Engine/Source/Editor/MovieSceneTools/Private/Sections/ParameterSection.cpp
--- a/Engine/Source/Editor/MovieSceneTools/Private/Sections/ParameterSection.cpp
+++ b/Engine/Source/Editor/MovieSceneTools/Private/Sections/ParameterSection.cpp
@@ -9,6 +9,7 @@
 #include "IKeyArea.h"
 #include "Internationalization/Internationalization.h"
 #include "Internationalization/Text.h"
+#include "ISequencer.h"
 #include "Layout/Geometry.h"
 #include "Layout/PaintGeometry.h"
 #include "Math/Color.h"
@@ -195,7 +196,7 @@ FReply FParameterSection::OnKeyDoubleClicked(const TArray<FKeyHandle> &KeyHandle
 
 	for (FColorParameterNameAndCurves& NameAndCurve : ParameterSection->GetColorParameterNamesAndCurves())
 	{
-		FMovieSceneKeyColorPicker KeyColorPicker(ParameterSection, &NameAndCurve.RedCurve, &NameAndCurve.GreenCurve, &NameAndCurve.BlueCurve, &NameAndCurve.AlphaCurve, KeyHandles);
+		FMovieSceneKeyColorPicker KeyColorPicker(ParameterSection, &NameAndCurve.RedCurve, &NameAndCurve.GreenCurve, &NameAndCurve.BlueCurve, &NameAndCurve.AlphaCurve, KeyHandles, WeakSequencer);
 	}
 
 	return FReply::Handled();
diff --git a/Engine/Source/Editor/MovieSceneTools/Private/Sections/ParameterSection.h b/Engine/Source/Editor/MovieSceneTools/Private/Sections/ParameterSection.h
--- a/Engine/Source/Editor/MovieSceneTools/Private/Sections/ParameterSection.h
+++ b/Engine/Source/Editor/MovieSceneTools/Private/Sections/ParameterSection.h
@@ -9,6 +9,7 @@
 
 class FName;
 class FSequencerSectionPainter;
+class ISequencer;
 class UMovieSceneSection;
 struct FKeyHandle;
 
@@ -20,8 +21,9 @@ class FParameterSection
 {
 public:
 
-	FParameterSection(UMovieSceneSection& InSectionObject)
+	FParameterSection(UMovieSceneSection& InSectionObject, TWeakPtr<ISequencer> InSequencer)
 		: FSequencerSection(InSectionObject)
+		, WeakSequencer(InSequencer)
 	{ }
 
 public:
@@ -32,4 +34,9 @@ public:
 	virtual bool RequestDeleteKeyArea(const TArray<FName>& KeyAreaNamePath) override;
 
 	virtual TSharedPtr<UE::Sequencer::FCategoryModel> ConstructCategoryModel(FName InCategoryName, const FText& InDisplayText, TArrayView<const FChannelData> Channels) const;
+
+private:
+
+	/** Weak pointer to the sequencer this section is for */
+	TWeakPtr<ISequencer> WeakSequencer;
 };
diff --git a/Engine/Source/Editor/MovieSceneTools/Private/TrackEditors/CustomPrimitiveDataTrackEditor.cpp b/Engine/Source/Editor/MovieSceneTools/Private/TrackEditors/CustomPrimitiveDataTrackEditor.cpp
--- a/Engine/Source/Editor/MovieSceneTools/Private/TrackEditors/CustomPrimitiveDataTrackEditor.cpp
+++ b/Engine/Source/Editor/MovieSceneTools/Private/TrackEditors/CustomPrimitiveDataTrackEditor.cpp
@@ -34,7 +34,7 @@ TSharedRef<ISequencerSection> FCustomPrimitiveDataTrackEditor::MakeSectionInterf
 	UMovieSceneCustomPrimitiveDataSection* CustomPrimitiveDataSection = Cast<UMovieSceneCustomPrimitiveDataSection>(&SectionObject);
 	checkf(CustomPrimitiveDataSection != nullptr, TEXT("Unsupported section type."));
 
-	return MakeShareable(new FParameterSection(*CustomPrimitiveDataSection));
+	return MakeShareable(new FParameterSection(*CustomPrimitiveDataSection, GetSequencer()));
 }
 
 TSharedPtr<SWidget> FCustomPrimitiveDataTrackEditor::BuildOutlinerEditWidget(const FGuid& ObjectBinding, UMovieSceneTrack* Track, const FBuildEditWidgetParams& Params)
diff --git a/Engine/Source/Editor/MovieSceneTools/Private/TrackEditors/MaterialParameterCollectionTrackEditor.cpp b/Engine/Source/Editor/MovieSceneTools/Private/TrackEditors/MaterialParameterCollectionTrackEditor.cpp
--- a/Engine/Source/Editor/MovieSceneTools/Private/TrackEditors/MaterialParameterCollectionTrackEditor.cpp
+++ b/Engine/Source/Editor/MovieSceneTools/Private/TrackEditors/MaterialParameterCollectionTrackEditor.cpp
@@ -75,7 +75,7 @@ TSharedRef<ISequencerSection> FMaterialParameterCollectionTrackEditor::MakeSecti
 	UMovieSceneParameterSection* ParameterSection = Cast<UMovieSceneParameterSection>(&SectionObject);
 	checkf(ParameterSection != nullptr, TEXT("Unsupported section type."));
 
-	return MakeShareable(new FParameterSection(*ParameterSection));
+	return MakeShareable(new FParameterSection(*ParameterSection, GetSequencer()));
 }
 
 TSharedRef<SWidget> CreateAssetPicker(FOnAssetSelected OnAssetSelected, FOnAssetEnterPressed OnAssetEnterPressed, TWeakPtr<ISequencer> InSequencer)
diff --git a/Engine/Source/Editor/MovieSceneTools/Private/TrackEditors/MaterialTrackEditor.cpp b/Engine/Source/Editor/MovieSceneTools/Private/TrackEditors/MaterialTrackEditor.cpp
--- a/Engine/Source/Editor/MovieSceneTools/Private/TrackEditors/MaterialTrackEditor.cpp
+++ b/Engine/Source/Editor/MovieSceneTools/Private/TrackEditors/MaterialTrackEditor.cpp
@@ -37,11 +37,11 @@ TSharedRef<ISequencerSection> FMaterialTrackEditor::MakeSectionInterface( UMovie
 
 	if (ComponentMaterialParameterSection)
 	{
-		return MakeShareable(new FComponentMaterialParameterSection(*ComponentMaterialParameterSection));
+		return MakeShareable(new FComponentMaterialParameterSection(*ComponentMaterialParameterSection, GetSequencer()));
 	}
 	else
 	{
-		return MakeShareable(new FParameterSection(*ParameterSection));
+		return MakeShareable(new FParameterSection(*ParameterSection, GetSequencer()));
 	}
 }
 
diff --git a/Engine/Source/Editor/MovieSceneTools/Private/TrackEditors/ParticleParameterTrackEditor.cpp b/Engine/Source/Editor/MovieSceneTools/Private/TrackEditors/ParticleParameterTrackEditor.cpp
--- a/Engine/Source/Editor/MovieSceneTools/Private/TrackEditors/ParticleParameterTrackEditor.cpp
+++ b/Engine/Source/Editor/MovieSceneTools/Private/TrackEditors/ParticleParameterTrackEditor.cpp
@@ -31,7 +31,7 @@ TSharedRef<ISequencerSection> FParticleParameterTrackEditor::MakeSectionInterfac
 	UMovieSceneParameterSection* ParameterSection = Cast<UMovieSceneParameterSection>(&SectionObject);
 	checkf( ParameterSection != nullptr, TEXT("Unsupported section type.") );
 
-	return MakeShareable(new FParameterSection( *ParameterSection ));
+	return MakeShareable(new FParameterSection( *ParameterSection, GetSequencer() ));
 }
 
 
