Sequencer: Unify asset pickers in Sequencer to have the same width and height override which is now a user preference

Defaults to 500x400

diff --git a/Engine/Plugins/Animation/ControlRig/Source/ControlRigEditor/Private/Sequencer/ControlRigParameterTrackEditor.cpp b/Engine/Plugins/Animation/ControlRig/Source/ControlRigEditor/Private/Sequencer/ControlRigParameterTrackEditor.cpp
--- a/Engine/Plugins/Animation/ControlRig/Source/ControlRigEditor/Private/Sequencer/ControlRigParameterTrackEditor.cpp
+++ b/Engine/Plugins/Animation/ControlRig/Source/ControlRigEditor/Private/Sequencer/ControlRigParameterTrackEditor.cpp
@@ -5792,9 +5792,12 @@ void FControlRigParameterSection::AddAnimationSubMenuForFK(FMenuBuilder& MenuBui
 
 	FContentBrowserModule& ContentBrowserModule = FModuleManager::Get().LoadModuleChecked<FContentBrowserModule>(TEXT("ContentBrowser"));
 
+	const float WidthOverride = SequencerPtr.IsValid() ? SequencerPtr->GetSequencerSettings()->GetAssetBrowserWidth() : 500.f;
+	const float HeightOverride = SequencerPtr.IsValid() ? SequencerPtr->GetSequencerSettings()->GetAssetBrowserHeight() : 400.f;
+
 	TSharedPtr<SBox> MenuEntry = SNew(SBox)
-		.WidthOverride(300.0f)
-		.HeightOverride(300.f)
+		.WidthOverride(WidthOverride)
+		.HeightOverride(HeightOverride)
 		[
 			ContentBrowserModule.Get().CreateAssetPicker(AssetPickerConfig)
 		];
diff --git a/Engine/Plugins/Compositing/Composure/Source/ComposureEditor/Private/Sequencer/ComposureExportTrackEditor.cpp b/Engine/Plugins/Compositing/Composure/Source/ComposureEditor/Private/Sequencer/ComposureExportTrackEditor.cpp
--- a/Engine/Plugins/Compositing/Composure/Source/ComposureEditor/Private/Sequencer/ComposureExportTrackEditor.cpp
+++ b/Engine/Plugins/Compositing/Composure/Source/ComposureEditor/Private/Sequencer/ComposureExportTrackEditor.cpp
@@ -8,8 +8,10 @@
 #include "DetailsViewArgs.h"
 #include "PropertyEditorModule.h"
 #include "IDetailsView.h"
+#include "ISequencer.h"
 #include "Modules/ModuleManager.h"
 #include "MovieScene/MovieSceneComposureExportTrack.h"
+#include "SequencerSettings.h"
 
 #define LOCTEXT_NAMESPACE "ComposureExportTrackEditor"
 
@@ -51,11 +53,15 @@ void FComposureExportTrackEditor::BuildTrackContextMenu(FMenuBuilder& MenuBuilde
 		// Assign the object
 		DetailsView->SetObject(ExportTrack, true);
 
+		TSharedPtr<ISequencer> SequencerPtr = GetSequencer();
+		const float WidthOverride = SequencerPtr.IsValid() ? SequencerPtr->GetSequencerSettings()->GetAssetBrowserWidth() : 500.f;
+		const float HeightOverride = SequencerPtr.IsValid() ? SequencerPtr->GetSequencerSettings()->GetAssetBrowserHeight() : 400.f;
+
 		// Add it to the menu
 		TSharedRef< SWidget > DetailsViewWidget =
 			SNew(SBox)
-			.MaxDesiredHeight(400.0f)
-			.WidthOverride(450.0f)
+			.WidthOverride(WidthOverride)
+			.HeightOverride(HeightOverride)
 		[
 			DetailsView
 		];
diff --git a/Engine/Plugins/Media/MediaCompositing/Source/MediaCompositingEditor/Private/Sequencer/MediaPlayerPropertyTrackEditor.cpp b/Engine/Plugins/Media/MediaCompositing/Source/MediaCompositingEditor/Private/Sequencer/MediaPlayerPropertyTrackEditor.cpp
--- a/Engine/Plugins/Media/MediaCompositing/Source/MediaCompositingEditor/Private/Sequencer/MediaPlayerPropertyTrackEditor.cpp
+++ b/Engine/Plugins/Media/MediaCompositing/Source/MediaCompositingEditor/Private/Sequencer/MediaPlayerPropertyTrackEditor.cpp
@@ -16,6 +16,7 @@
 #include "Widgets/Layout/SBox.h"
 #include "MediaSource.h"
 #include "LevelSequence.h"
+#include "SequencerSettings.h"
 
 #include "Sequencer/MediaThumbnailSection.h"
 
@@ -147,7 +148,8 @@ TSharedPtr<SWidget> FMediaPlayerPropertyTrackEditor::BuildOutlinerEditWidget(con
 
 	auto CreatePicker = [this, MediaTrack]
 	{
-		UMovieSceneSequence* Sequence = GetSequencer() ? GetSequencer()->GetFocusedMovieSceneSequence() : nullptr;
+		TSharedPtr<ISequencer> SequencerPtr = GetSequencer();
+		UMovieSceneSequence* Sequence = SequencerPtr.IsValid() ? SequencerPtr->GetFocusedMovieSceneSequence() : nullptr;
 
 		FAssetPickerConfig AssetPickerConfig;
 		{
@@ -164,9 +166,12 @@ TSharedPtr<SWidget> FMediaPlayerPropertyTrackEditor::BuildOutlinerEditWidget(con
 
 		FContentBrowserModule& ContentBrowserModule = FModuleManager::Get().LoadModuleChecked<FContentBrowserModule>(TEXT("ContentBrowser"));
 
+		const float WidthOverride = SequencerPtr.IsValid() ? SequencerPtr->GetSequencerSettings()->GetAssetBrowserWidth() : 500.f;
+		const float HeightOverride = SequencerPtr.IsValid() ? SequencerPtr->GetSequencerSettings()->GetAssetBrowserHeight() : 400.f;
+
 		TSharedRef<SBox> Picker = SNew(SBox)
-			.WidthOverride(500.0f)
-			.HeightOverride(400.f)
+			.WidthOverride(WidthOverride)
+			.HeightOverride(HeightOverride)
 			[
 				ContentBrowserModule.Get().CreateAssetPicker(AssetPickerConfig)
 			];
diff --git a/Engine/Plugins/Media/MediaCompositing/Source/MediaCompositingEditor/Private/Sequencer/MediaTrackEditor.cpp b/Engine/Plugins/Media/MediaCompositing/Source/MediaCompositingEditor/Private/Sequencer/MediaTrackEditor.cpp
--- a/Engine/Plugins/Media/MediaCompositing/Source/MediaCompositingEditor/Private/Sequencer/MediaTrackEditor.cpp
+++ b/Engine/Plugins/Media/MediaCompositing/Source/MediaCompositingEditor/Private/Sequencer/MediaTrackEditor.cpp
@@ -22,6 +22,7 @@
 #include "Widgets/Layout/SBox.h"
 #include "Misc/QualifiedFrameTime.h"
 #include "LevelSequence.h"
+#include "SequencerSettings.h"
 
 #include "Sequencer/MediaThumbnailSection.h"
 
@@ -96,7 +97,8 @@ TSharedPtr<SWidget> FMediaTrackEditor::BuildOutlinerEditWidget(const FGuid& Obje
 
 	auto CreatePicker = [this, MediaTrack]
 	{
-		UMovieSceneSequence* Sequence = GetSequencer() ? GetSequencer()->GetFocusedMovieSceneSequence() : nullptr;
+		TSharedPtr<ISequencer> SequencerPtr = GetSequencer();
+		UMovieSceneSequence* Sequence = SequencerPtr.IsValid() ? SequencerPtr->GetFocusedMovieSceneSequence() : nullptr;
 
 		FAssetPickerConfig AssetPickerConfig;
 		{
@@ -113,9 +115,12 @@ TSharedPtr<SWidget> FMediaTrackEditor::BuildOutlinerEditWidget(const FGuid& Obje
 
 		FContentBrowserModule& ContentBrowserModule = FModuleManager::Get().LoadModuleChecked<FContentBrowserModule>(TEXT("ContentBrowser"));
 
+		const float WidthOverride = SequencerPtr.IsValid() ? SequencerPtr->GetSequencerSettings()->GetAssetBrowserWidth() : 500.f;
+		const float HeightOverride = SequencerPtr.IsValid() ? SequencerPtr->GetSequencerSettings()->GetAssetBrowserHeight() : 400.f;
+
 		TSharedRef<SBox> Picker = SNew(SBox)
-			.WidthOverride(300.0f)
-			.HeightOverride(300.f)
+			.WidthOverride(WidthOverride)
+			.HeightOverride(HeightOverride)
 			[
 				ContentBrowserModule.Get().CreateAssetPicker(AssetPickerConfig)
 			];
diff --git a/Engine/Plugins/Media/MediaPlate/Source/MediaPlateEditor/Private/Sequencer/MediaPlateTrackEditor.cpp b/Engine/Plugins/Media/MediaPlate/Source/MediaPlateEditor/Private/Sequencer/MediaPlateTrackEditor.cpp
--- a/Engine/Plugins/Media/MediaPlate/Source/MediaPlateEditor/Private/Sequencer/MediaPlateTrackEditor.cpp
+++ b/Engine/Plugins/Media/MediaPlate/Source/MediaPlateEditor/Private/Sequencer/MediaPlateTrackEditor.cpp
@@ -17,6 +17,7 @@
 #include "MovieSceneMediaTrack.h"
 #include "SceneOutlinerModule.h"
 #include "Sequencer/MediaTrackEditor.h"
+#include "SequencerSettings.h"
 
 #define LOCTEXT_NAMESPACE "FMediaPlateTrackEditor"
 
@@ -248,12 +249,16 @@ void FMediaPlateTrackEditor::OnBuildOutlinerEditWidget(FMenuBuilder& MenuBuilder
 				}
 				));
 
+			TSharedPtr<ISequencer> SequencerPtr = GetSequencer();
+			const float WidthOverride = SequencerPtr.IsValid() ? SequencerPtr->GetSequencerSettings()->GetAssetBrowserWidth() : 500.f;
+			const float HeightOverride = SequencerPtr.IsValid() ? SequencerPtr->GetSequencerSettings()->GetAssetBrowserHeight() : 400.f;
+
 			// Create actor picker.
 			FSceneOutlinerModule& SceneOutlinerModule = 
 				FModuleManager::LoadModuleChecked<FSceneOutlinerModule>("SceneOutliner");
 			TSharedRef<SBox> Picker = SNew(SBox)
-				.WidthOverride(300.0f)
-				.HeightOverride(300.f)
+				.WidthOverride(WidthOverride)
+				.HeightOverride(HeightOverride)
 				[
 					SceneOutlinerModule.CreateActorPicker(InitOptions,
 						FOnActorPicked::CreateLambda([this](AActor* Actor)
diff --git a/Engine/Plugins/MovieScene/LevelSequenceEditor/Source/LevelSequenceEditor/Private/LevelSequenceEditorSubsystem.cpp b/Engine/Plugins/MovieScene/LevelSequenceEditor/Source/LevelSequenceEditor/Private/LevelSequenceEditorSubsystem.cpp
--- a/Engine/Plugins/MovieScene/LevelSequenceEditor/Source/LevelSequenceEditor/Private/LevelSequenceEditorSubsystem.cpp
+++ b/Engine/Plugins/MovieScene/LevelSequenceEditor/Source/LevelSequenceEditor/Private/LevelSequenceEditorSubsystem.cpp
@@ -2512,12 +2512,15 @@ void ULevelSequenceEditorSubsystem::AddAssignActorMenu(FMenuBuilder& MenuBuilder
 		InitOptions.Filters->AddFilterPredicate<FActorTreeItem>(FActorTreeItem::FFilterPredicate::CreateLambda( IsActorValidForAssignment ) );
 	}
 
+	const float WidthOverride = Sequencer.IsValid() ? Sequencer->GetSequencerSettings()->GetAssetBrowserWidth() : 500.f;
+	const float HeightOverride = Sequencer.IsValid() ? Sequencer->GetSequencerSettings()->GetAssetBrowserHeight() : 400.f;
+
 	// actor selector to allow the user to choose an actor
 	FSceneOutlinerModule& SceneOutlinerModule = FModuleManager::LoadModuleChecked<FSceneOutlinerModule>("SceneOutliner");
 	TSharedRef< SWidget > MiniSceneOutliner =
 		SNew( SBox )
-		.MaxDesiredHeight(400.0f)
-		.WidthOverride(300.0f)
+		.WidthOverride(WidthOverride)
+		.HeightOverride(HeightOverride)
 		[
 			SceneOutlinerModule.CreateActorPicker(
 				InitOptions,
diff --git a/Engine/Plugins/MovieScene/LevelSequenceEditor/Source/LevelSequenceEditor/Private/Misc/LevelSequenceEditorActorBinding.cpp b/Engine/Plugins/MovieScene/LevelSequenceEditor/Source/LevelSequenceEditor/Private/Misc/LevelSequenceEditorActorBinding.cpp
--- a/Engine/Plugins/MovieScene/LevelSequenceEditor/Source/LevelSequenceEditor/Private/Misc/LevelSequenceEditorActorBinding.cpp
+++ b/Engine/Plugins/MovieScene/LevelSequenceEditor/Source/LevelSequenceEditor/Private/Misc/LevelSequenceEditorActorBinding.cpp
@@ -15,6 +15,7 @@
 #include "Misc/EditorPathHelper.h"
 #include "Editor.h"
 #include "Selection.h"
+#include "SequencerSettings.h"
 
 #define LOCTEXT_NAMESPACE "LevelSequenceEditorActorBinding"
 
@@ -49,9 +50,10 @@ void FLevelSequenceEditorActorBinding::AddPossessActorMenuExtensions(FMenuBuilde
 	// We don't need to check against Sequencer spawnables as they're not valid for possession.
 	TSet<UObject*> ExistingPossessedObjects;
 	UMovieSceneSequence* MovieSceneSequence = nullptr;
-	if (Sequencer.IsValid())
+	TSharedPtr<ISequencer> SequencerPtr = Sequencer.Pin();
+	if (SequencerPtr.IsValid())
 	{
-		MovieSceneSequence = Sequencer.Pin()->GetFocusedMovieSceneSequence();
+		MovieSceneSequence = SequencerPtr->GetFocusedMovieSceneSequence();
 		UMovieScene* MovieScene = MovieSceneSequence->GetMovieScene();
 		if(MovieScene)
 		{
@@ -63,7 +65,7 @@ void FLevelSequenceEditorActorBinding::AddPossessActorMenuExtensions(FMenuBuilde
 				{
 					// A possession guid can apply to more than one object, so we get all bound objects for the GUID and add them to our set.
 					TArray<UObject*, TInlineAllocator<1>> OutObjects;
-					MovieSceneSequence->LocateBoundObjects(Possessable.GetGuid(), UE::UniversalObjectLocator::FResolveParams(Sequencer.Pin()->GetPlaybackContext()), Sequencer.Pin()->FindSharedPlaybackState(), OutObjects);
+					MovieSceneSequence->LocateBoundObjects(Possessable.GetGuid(), UE::UniversalObjectLocator::FResolveParams(SequencerPtr->GetPlaybackContext()), SequencerPtr->FindSharedPlaybackState(), OutObjects);
 					ExistingPossessedObjects.Append(OutObjects);
 				}
 			}
@@ -133,12 +135,15 @@ void FLevelSequenceEditorActorBinding::AddPossessActorMenuExtensions(FMenuBuilde
 
 	const bool bHideLevelInstanceHierarchy = !FEditorPathHelper::IsEnabled();
 
+	const float WidthOverride = SequencerPtr.IsValid() ? SequencerPtr->GetSequencerSettings()->GetAssetBrowserWidth() : 500.f;
+	const float HeightOverride = SequencerPtr.IsValid() ? SequencerPtr->GetSequencerSettings()->GetAssetBrowserHeight() : 400.f;
+
 	// actor selector to allow the user to choose an actor
 	FSceneOutlinerModule& SceneOutlinerModule = FModuleManager::LoadModuleChecked<FSceneOutlinerModule>("SceneOutliner");
 	TSharedRef< SWidget > MiniSceneOutliner =
 		SNew(SBox)
-		.MaxDesiredHeight(400.0f)
-		.WidthOverride(300.0f)
+		.WidthOverride(WidthOverride)
+		.HeightOverride(HeightOverride)
 		[
 			SceneOutlinerModule.CreateActorPicker(
 				InitOptions,
diff --git a/Engine/Plugins/MovieScene/TemplateSequence/Source/TemplateSequenceEditor/Private/TrackEditors/TemplateSequenceTrackEditor.cpp b/Engine/Plugins/MovieScene/TemplateSequence/Source/TemplateSequenceEditor/Private/TrackEditors/TemplateSequenceTrackEditor.cpp
--- a/Engine/Plugins/MovieScene/TemplateSequence/Source/TemplateSequenceEditor/Private/TrackEditors/TemplateSequenceTrackEditor.cpp
+++ b/Engine/Plugins/MovieScene/TemplateSequence/Source/TemplateSequenceEditor/Private/TrackEditors/TemplateSequenceTrackEditor.cpp
@@ -8,6 +8,7 @@
 #include "Framework/MultiBox/MultiBoxBuilder.h"
 #include "Framework/Notifications/NotificationManager.h"
 #include "IContentBrowserSingleton.h"
+#include "ISequencer.h"
 #include "LevelSequence.h"
 #include "Misc/ConfigCacheIni.h"
 #include "Misc/PackageName.h"
@@ -20,6 +21,7 @@
 #include "Tracks/MovieScene3DTransformTrack.h"
 #include "Tracks/MovieSceneSubTrack.h"
 #include "Sections/TemplateSequenceSection.h"
+#include "SequencerSettings.h"
 #include "Tracks/TemplateSequenceTrack.h"
 #include "Widgets/Input/SButton.h"
 #include "Widgets/Notifications/SNotificationList.h"
@@ -119,10 +121,12 @@ public:
 			CurBaseClass = CurBaseClass->GetSuperClass();
 		}
 
+		TSharedRef<FTemplateSequenceTrackEditor> TrackEditorRef = TrackEditor.ToSharedRef();
+		TSharedPtr<ISequencer> Sequencer = TrackEditorRef->GetSequencer();
+
 		FAssetPickerConfig AssetPickerConfig;
 		{
-			TSharedRef<FTemplateSequenceTrackEditor> TrackEditorRef = TrackEditor.ToSharedRef();
-			UMovieSceneSequence* Sequence = TrackEditorRef->GetSequencer() ? TrackEditorRef->GetSequencer()->GetFocusedMovieSceneSequence() : nullptr;
+			UMovieSceneSequence* Sequence = Sequencer.IsValid() ? Sequencer->GetFocusedMovieSceneSequence() : nullptr;
 
 			AssetPickerConfig.OnAssetSelected = FOnAssetSelected::CreateSP(TrackEditorRef, &FTemplateSequenceTrackEditor::OnTemplateSequenceAssetSelected, ObjectBindings);
 			AssetPickerConfig.OnAssetEnterPressed = FOnAssetEnterPressed::CreateSP(TrackEditorRef, &FTemplateSequenceTrackEditor::OnTemplateSequenceAssetEnterPressed, ObjectBindings);
@@ -175,11 +179,14 @@ public:
 		FContentBrowserModule& ContentBrowserModule = FModuleManager::Get().LoadModuleChecked<FContentBrowserModule>(TEXT("ContentBrowser"));
 		AssetPicker = ContentBrowserModule.Get().CreateAssetPicker(AssetPickerConfig);
 
+		const float WidthOverride = Sequencer.IsValid() ? Sequencer->GetSequencerSettings()->GetAssetBrowserWidth() : 500.f;
+		const float HeightOverride = Sequencer.IsValid() ? Sequencer->GetSequencerSettings()->GetAssetBrowserHeight() : 400.f;
+
 		ChildSlot
 		[
 			SNew(SBox)
-			.WidthOverride(300.0f)
-			.HeightOverride(300.f)
+			.WidthOverride(WidthOverride)
+			.HeightOverride(HeightOverride)
 			[
 				SNew(SVerticalBox)
 				+SVerticalBox::Slot()
diff --git a/Engine/Plugins/VirtualProduction/DMX/DMXEngine/Source/DMXEditor/Private/Sequencer/DMXLibraryTrackEditor.cpp b/Engine/Plugins/VirtualProduction/DMX/DMXEngine/Source/DMXEditor/Private/Sequencer/DMXLibraryTrackEditor.cpp
--- a/Engine/Plugins/VirtualProduction/DMX/DMXEngine/Source/DMXEditor/Private/Sequencer/DMXLibraryTrackEditor.cpp
+++ b/Engine/Plugins/VirtualProduction/DMX/DMXEngine/Source/DMXEditor/Private/Sequencer/DMXLibraryTrackEditor.cpp
@@ -7,6 +7,7 @@
 #include "Framework/Application/SlateApplication.h"
 #include "Framework/MultiBox/MultiBoxBuilder.h"
 #include "IContentBrowserSingleton.h"
+#include "ISequencer.h"
 #include "Library/DMXEntityFixturePatch.h"
 #include "Library/DMXLibrary.h"
 #include "MovieSceneSection.h"
@@ -14,6 +15,7 @@
 #include "Sequencer/DMXLibrarySection.h"
 #include "Sequencer/MovieSceneDMXLibrarySection.h"
 #include "Sequencer/MovieSceneDMXLibraryTrack.h"
+#include "SequencerSettings.h"
 #include "Styling/SlateIconFinder.h"
 #include "Widgets/Layout/SBox.h"
 
@@ -52,9 +54,13 @@ namespace UE::DMX
 
 		FContentBrowserModule& ContentBrowserModule = FModuleManager::Get().LoadModuleChecked<FContentBrowserModule>(TEXT("ContentBrowser"));
 
+		TSharedPtr<ISequencer> Sequencer = InSequencer.Pin();
+		const float WidthOverride = Sequencer.IsValid() ? Sequencer->GetSequencerSettings()->GetAssetBrowserWidth() : 500.f;
+		const float HeightOverride = Sequencer.IsValid() ? Sequencer->GetSequencerSettings()->GetAssetBrowserHeight() : 400.f;
+
 		return SNew(SBox)
-			.WidthOverride(300.0f)
-			.HeightOverride(300.f)
+			.WidthOverride(WidthOverride)
+			.HeightOverride(HeightOverride)
 			[
 				ContentBrowserModule.Get().CreateAssetPicker(AssetPickerConfig)
 			];
diff --git a/Engine/Source/Editor/MovieSceneTools/Private/TrackEditors/ActorPickerTrackEditor.cpp b/Engine/Source/Editor/MovieSceneTools/Private/TrackEditors/ActorPickerTrackEditor.cpp
--- a/Engine/Source/Editor/MovieSceneTools/Private/TrackEditors/ActorPickerTrackEditor.cpp
+++ b/Engine/Source/Editor/MovieSceneTools/Private/TrackEditors/ActorPickerTrackEditor.cpp
@@ -4,6 +4,7 @@
 #include "Widgets/SBoxPanel.h"
 #include "Framework/MultiBox/MultiBoxBuilder.h"
 #include "GameFramework/Actor.h"
+#include "ISequencer.h"
 #include "Modules/ModuleManager.h"
 #include "Layout/WidgetPath.h"
 #include "Framework/Application/MenuStack.h"
@@ -25,6 +26,7 @@
 #include "MovieSceneToolHelpers.h"
 #include "SComponentChooser.h"
 #include "ActorTreeItem.h"
+#include "SequencerSettings.h"
 
 #define LOCTEXT_NAMESPACE "FActorPickerTrackEditor"
 
@@ -79,6 +81,10 @@ void FActorPickerTrackEditor::ShowActorSubMenu(FMenuBuilder& MenuBuilder, TArray
 		// Actor selector to allow the user to choose a parent actor
 		FSceneOutlinerModule& SceneOutlinerModule = FModuleManager::LoadModuleChecked<FSceneOutlinerModule>( "SceneOutliner" );
 
+		TSharedPtr<ISequencer> Sequencer = GetSequencer();
+		const float WidthOverride = Sequencer.IsValid() ? Sequencer->GetSequencerSettings()->GetAssetBrowserWidth() : 500.f;
+		const float HeightOverride = Sequencer.IsValid() ? Sequencer->GetSequencerSettings()->GetAssetBrowserHeight() : 400.f;
+
 		TSharedRef< SWidget > MenuWidget = 
 			SNew(SHorizontalBox)
 
@@ -86,8 +92,8 @@ void FActorPickerTrackEditor::ShowActorSubMenu(FMenuBuilder& MenuBuilder, TArray
 			.AutoWidth()
 			[
 				SNew(SBox)
-				.MaxDesiredHeight(400.0f)
-				.WidthOverride(300.0f)
+				.WidthOverride(WidthOverride)
+				.HeightOverride(HeightOverride)
 				[
 					SceneOutlinerModule.CreateActorPicker(
 						InitOptions,
diff --git a/Engine/Source/Editor/MovieSceneTools/Private/TrackEditors/AudioTrackEditor.cpp b/Engine/Source/Editor/MovieSceneTools/Private/TrackEditors/AudioTrackEditor.cpp
--- a/Engine/Source/Editor/MovieSceneTools/Private/TrackEditors/AudioTrackEditor.cpp
+++ b/Engine/Source/Editor/MovieSceneTools/Private/TrackEditors/AudioTrackEditor.cpp
@@ -34,8 +34,10 @@
 #include "MVVM/Views/ViewUtilities.h"
 #include "AssetRegistry/AssetRegistryModule.h"
 #include "ISectionLayoutBuilder.h"
+#include "ISequencer.h"
 #include "MovieSceneToolHelpers.h"
 #include "Dialogs/Dialogs.h"
+#include "SequencerSettings.h"
 
 #include "DragAndDrop/AssetDragDropOp.h"
 #include "Misc/QualifiedFrameTime.h"
@@ -1370,7 +1372,8 @@ void FAudioTrackEditor::HandleAddAttachedAudioTrackMenuEntryExecute(FMenuBuilder
 
 TSharedRef<SWidget> FAudioTrackEditor::BuildAudioSubMenu(FOnAssetSelected OnAssetSelected, FOnAssetEnterPressed OnAssetEnterPressed)
 {
-	UMovieSceneSequence* Sequence = GetSequencer() ? GetSequencer()->GetFocusedMovieSceneSequence() : nullptr;
+	TSharedPtr<ISequencer> SequencerPtr = GetSequencer();
+	UMovieSceneSequence* Sequence = SequencerPtr.IsValid() ? SequencerPtr->GetFocusedMovieSceneSequence() : nullptr;
 
 	FAssetRegistryModule& AssetRegistryModule = FModuleManager::LoadModuleChecked<FAssetRegistryModule>(TEXT("AssetRegistry"));
 	TArray<FTopLevelAssetPath> ClassNames;
@@ -1397,9 +1400,12 @@ TSharedRef<SWidget> FAudioTrackEditor::BuildAudioSubMenu(FOnAssetSelected OnAsse
 
 	FContentBrowserModule& ContentBrowserModule = FModuleManager::Get().LoadModuleChecked<FContentBrowserModule>(TEXT("ContentBrowser"));
 
+	const float WidthOverride = SequencerPtr.IsValid() ? SequencerPtr->GetSequencerSettings()->GetAssetBrowserWidth() : 500.f;
+	const float HeightOverride = SequencerPtr.IsValid() ? SequencerPtr->GetSequencerSettings()->GetAssetBrowserHeight() : 400.f;
+
 	TSharedPtr<SBox> MenuEntry = SNew(SBox)
-		.WidthOverride(300.0f)
-		.HeightOverride(300.f)
+		.WidthOverride(WidthOverride)
+		.HeightOverride(HeightOverride)
 		[
 			ContentBrowserModule.Get().CreateAssetPicker(AssetPickerConfig)
 		];
diff --git a/Engine/Source/Editor/MovieSceneTools/Private/TrackEditors/CameraCutTrackEditor.cpp b/Engine/Source/Editor/MovieSceneTools/Private/TrackEditors/CameraCutTrackEditor.cpp
--- a/Engine/Source/Editor/MovieSceneTools/Private/TrackEditors/CameraCutTrackEditor.cpp
+++ b/Engine/Source/Editor/MovieSceneTools/Private/TrackEditors/CameraCutTrackEditor.cpp
@@ -22,10 +22,12 @@
 #include "Editor.h"
 #include "Framework/MultiBox/MultiBoxBuilder.h"
 #include "GameFramework/WorldSettings.h"
+#include "ISequencer.h"
 #include "LevelEditorViewport.h"
 #include "Modules/ModuleManager.h"
 #include "SceneOutlinerModule.h"
 #include "SceneOutlinerPublicTypes.h"
+#include "SequencerSettings.h"
 #include "Styling/AppStyle.h"
 #include "Widgets/Input/SCheckBox.h"
 #include "Widgets/Layout/SBox.h"
@@ -491,15 +493,19 @@ TSharedRef<SWidget> FCameraCutTrackEditor::HandleAddCameraCutComboButtonGetMenuC
 			// Actor selector to allow the user to choose a parent actor
 			FSceneOutlinerModule& SceneOutlinerModule = FModuleManager::LoadModuleChecked<FSceneOutlinerModule>( "SceneOutliner" );
 
-			TSharedRef< SWidget > MenuWidget = 
+			TSharedPtr<ISequencer> Sequencer = GetSequencer();
+			const float WidthOverride = Sequencer.IsValid() ? Sequencer->GetSequencerSettings()->GetAssetBrowserWidth() : 500.f;
+			const float HeightOverride = Sequencer.IsValid() ? Sequencer->GetSequencerSettings()->GetAssetBrowserHeight() : 400.f;
+
+			TSharedRef< SWidget > MenuWidget =
 				SNew(SHorizontalBox)
 
 				+SHorizontalBox::Slot()
 				.AutoWidth()
 				[
 					SNew(SBox)
-					.MaxDesiredHeight(400.0f)
-					.WidthOverride(300.0f)
+					.WidthOverride(WidthOverride)
+					.HeightOverride(HeightOverride)
 					[
 						SceneOutlinerModule.CreateActorPicker(
 							InitOptions,
diff --git a/Engine/Source/Editor/MovieSceneTools/Private/TrackEditors/CameraShakeSourceShakeTrackEditor.cpp b/Engine/Source/Editor/MovieSceneTools/Private/TrackEditors/CameraShakeSourceShakeTrackEditor.cpp
--- a/Engine/Source/Editor/MovieSceneTools/Private/TrackEditors/CameraShakeSourceShakeTrackEditor.cpp
+++ b/Engine/Source/Editor/MovieSceneTools/Private/TrackEditors/CameraShakeSourceShakeTrackEditor.cpp
@@ -57,6 +57,7 @@
 #include "Sections/MovieSceneCameraShakeSourceShakeSection.h"
 #include "Sections/MovieSceneCameraShakeSourceTriggerSection.h"
 #include "SequencerSectionPainter.h"
+#include "SequencerSettings.h"
 #include "MVVM/Views/ViewUtilities.h"
 #include "SlotBase.h"
 #include "Styling/AppStyle.h"
@@ -464,7 +465,8 @@ TSharedRef<SWidget> FCameraShakeSourceShakeTrackEditor::BuildCameraShakeTracksMe
 
 void FCameraShakeSourceShakeTrackEditor::AddOtherCameraShakeBrowserSubMenu(FMenuBuilder& MenuBuilder, TArray<FGuid> ObjectBindings)
 {
-	UMovieSceneSequence* Sequence = GetSequencer() ? GetSequencer()->GetFocusedMovieSceneSequence() : nullptr;
+	TSharedPtr<ISequencer> SequencerPtr = GetSequencer();
+	UMovieSceneSequence* Sequence = SequencerPtr.IsValid() ? SequencerPtr->GetFocusedMovieSceneSequence() : nullptr;
 
 	FAssetPickerConfig AssetPickerConfig;
 	{
@@ -502,9 +504,12 @@ void FCameraShakeSourceShakeTrackEditor::AddOtherCameraShakeBrowserSubMenu(FMenu
 
 	FContentBrowserModule& ContentBrowserModule = FModuleManager::Get().LoadModuleChecked<FContentBrowserModule>(TEXT("ContentBrowser"));
 
+	const float WidthOverride = SequencerPtr.IsValid() ? SequencerPtr->GetSequencerSettings()->GetAssetBrowserWidth() : 500.f;
+	const float HeightOverride = SequencerPtr.IsValid() ? SequencerPtr->GetSequencerSettings()->GetAssetBrowserHeight() : 400.f;
+
 	TSharedPtr<SBox> MenuEntry = SNew(SBox)
-		.WidthOverride(300.0f)
-		.HeightOverride(400.f)
+		.WidthOverride(WidthOverride)
+		.HeightOverride(HeightOverride)
 		[
 			ContentBrowserModule.Get().CreateAssetPicker(AssetPickerConfig)
 		];
diff --git a/Engine/Source/Editor/MovieSceneTools/Private/TrackEditors/CameraShakeTrackEditor.cpp b/Engine/Source/Editor/MovieSceneTools/Private/TrackEditors/CameraShakeTrackEditor.cpp
--- a/Engine/Source/Editor/MovieSceneTools/Private/TrackEditors/CameraShakeTrackEditor.cpp
+++ b/Engine/Source/Editor/MovieSceneTools/Private/TrackEditors/CameraShakeTrackEditor.cpp
@@ -41,6 +41,7 @@
 #include "MovieSceneTrack.h"
 #include "ScopedTransaction.h"
 #include "Sections/MovieSceneCameraShakeSection.h"
+#include "SequencerSettings.h"
 #include "MVVM/Views/ViewUtilities.h"
 #include "SlotBase.h"
 #include "Templates/Casts.h"
@@ -158,7 +159,8 @@ TSharedRef<SWidget> FCameraShakeTrackEditor::BuildCameraShakeSubMenu(FGuid Objec
 
 void FCameraShakeTrackEditor::AddCameraShakeSubMenu(FMenuBuilder& MenuBuilder, TArray<FGuid> ObjectBindings)
 {
-	UMovieSceneSequence* Sequence = GetSequencer() ? GetSequencer()->GetFocusedMovieSceneSequence() : nullptr;
+	TSharedPtr<ISequencer> SequencerPtr = GetSequencer();
+	UMovieSceneSequence* Sequence = SequencerPtr ? SequencerPtr->GetFocusedMovieSceneSequence() : nullptr;
 
 	FAssetPickerConfig AssetPickerConfig;
 	{
@@ -194,9 +196,12 @@ void FCameraShakeTrackEditor::AddCameraShakeSubMenu(FMenuBuilder& MenuBuilder, T
 
 	FContentBrowserModule& ContentBrowserModule = FModuleManager::Get().LoadModuleChecked<FContentBrowserModule>(TEXT("ContentBrowser"));
 
+	const float WidthOverride = SequencerPtr.IsValid() ? SequencerPtr->GetSequencerSettings()->GetAssetBrowserWidth() : 500.f;
+	const float HeightOverride = SequencerPtr.IsValid() ? SequencerPtr->GetSequencerSettings()->GetAssetBrowserHeight() : 400.f;
+
 	TSharedPtr<SBox> MenuEntry = SNew(SBox)
-		.WidthOverride(300.0f)
-		.HeightOverride(300.f)
+		.WidthOverride(WidthOverride)
+		.HeightOverride(HeightOverride)
 		[
 			ContentBrowserModule.Get().CreateAssetPicker(AssetPickerConfig)
 		];
diff --git a/Engine/Source/Editor/MovieSceneTools/Private/TrackEditors/EventTrackEditor.cpp b/Engine/Source/Editor/MovieSceneTools/Private/TrackEditors/EventTrackEditor.cpp
--- a/Engine/Source/Editor/MovieSceneTools/Private/TrackEditors/EventTrackEditor.cpp
+++ b/Engine/Source/Editor/MovieSceneTools/Private/TrackEditors/EventTrackEditor.cpp
@@ -6,6 +6,7 @@
 #include "Styling/AppStyle.h"
 #include "UObject/Package.h"
 #include "Tracks/MovieSceneEventTrack.h"
+#include "ISequencer.h"
 #include "ISequencerSection.h"
 #include "MovieSceneObjectBindingIDPicker.h"
 #include "IDetailCustomization.h"
@@ -20,6 +21,7 @@
 #include "Sections/MovieSceneEventSection.h"
 #include "Sections/MovieSceneEventTriggerSection.h"
 #include "Sections/MovieSceneEventRepeaterSection.h"
+#include "SequencerSettings.h"
 #include "MVVM/Views/ViewUtilities.h"
 #include "MovieSceneSequenceEditor.h"
 
@@ -189,6 +191,8 @@ void FEventTrackEditor::BuildTrackContextMenu(FMenuBuilder& MenuBuilder, UMovieS
 
 	auto PopulateSubMenu = [this, EventTrack, ObjectBinding](FMenuBuilder& SubMenuBuilder)
 	{
+		TSharedPtr<ISequencer> Sequencer = GetSequencer();
+			
 		FPropertyEditorModule& PropertyEditor = FModuleManager::Get().LoadModuleChecked<FPropertyEditorModule>("PropertyEditor");
 
 		// Create a details view for the track
@@ -206,16 +210,22 @@ void FEventTrackEditor::BuildTrackContextMenu(FMenuBuilder& MenuBuilder, UMovieS
 		FOnGetDetailCustomizationInstance CreateInstance = FOnGetDetailCustomizationInstance::CreateLambda([ObjectBinding]{ return MakeShared<FEventTrackCustomization>(ObjectBinding); });
 		DetailsView->RegisterInstancedCustomPropertyLayout(UMovieSceneEventTrack::StaticClass(), CreateInstance);
 
-		GetSequencer()->OnInitializeDetailsPanel().Broadcast(DetailsView, GetSequencer().ToSharedRef());
+		if (Sequencer.IsValid())
+		{
+			Sequencer->OnInitializeDetailsPanel().Broadcast(DetailsView, Sequencer.ToSharedRef());
+		}
 
 		// Assign the object
 		DetailsView->SetObject(EventTrack, true);
 
+		const float WidthOverride = Sequencer.IsValid() ? Sequencer->GetSequencerSettings()->GetAssetBrowserWidth() : 500.f;
+		const float HeightOverride = Sequencer.IsValid() ? Sequencer->GetSequencerSettings()->GetAssetBrowserHeight() : 400.f;
+
 		// Add it to the menu
 		TSharedRef< SWidget > DetailsViewWidget =
 			SNew(SBox)
-			.MaxDesiredHeight(400.0f)
-			.WidthOverride(450.0f)
+			.WidthOverride(WidthOverride)
+			.HeightOverride(HeightOverride)
 		[
 			DetailsView
 		];
diff --git a/Engine/Source/Editor/MovieSceneTools/Private/TrackEditors/MaterialParameterCollectionTrackEditor.cpp b/Engine/Source/Editor/MovieSceneTools/Private/TrackEditors/MaterialParameterCollectionTrackEditor.cpp
--- a/Engine/Source/Editor/MovieSceneTools/Private/TrackEditors/MaterialParameterCollectionTrackEditor.cpp
+++ b/Engine/Source/Editor/MovieSceneTools/Private/TrackEditors/MaterialParameterCollectionTrackEditor.cpp
@@ -37,6 +37,7 @@
 #include "Sections/MovieSceneParameterSection.h"
 #include "Sections/ParameterSection.h"
 #include "SequencerUtilities.h"
+#include "SequencerSettings.h"
 #include "MVVM/Views/ViewUtilities.h"
 #include "Styling/SlateIconFinder.h"
 #include "Templates/Casts.h"
@@ -80,7 +81,8 @@ TSharedRef<ISequencerSection> FMaterialParameterCollectionTrackEditor::MakeSecti
 
 TSharedRef<SWidget> CreateAssetPicker(FOnAssetSelected OnAssetSelected, FOnAssetEnterPressed OnAssetEnterPressed, TWeakPtr<ISequencer> InSequencer)
 {
-	UMovieSceneSequence* Sequence = InSequencer.IsValid() ? InSequencer.Pin()->GetFocusedMovieSceneSequence() : nullptr;
+	TSharedPtr<ISequencer> Sequencer = InSequencer.Pin();
+	UMovieSceneSequence* Sequence = Sequencer.IsValid() ? Sequencer->GetFocusedMovieSceneSequence() : nullptr;
 
 	FAssetPickerConfig AssetPickerConfig;
 	{
@@ -96,10 +98,13 @@ TSharedRef<SWidget> CreateAssetPicker(FOnAssetSelected OnAssetSelected, FOnAsset
 	}
 
 	FContentBrowserModule& ContentBrowserModule = FModuleManager::Get().LoadModuleChecked<FContentBrowserModule>(TEXT("ContentBrowser"));
+	
+	const float WidthOverride = Sequencer.IsValid() ? Sequencer->GetSequencerSettings()->GetAssetBrowserWidth() : 500.f;
+	const float HeightOverride = Sequencer.IsValid() ? Sequencer->GetSequencerSettings()->GetAssetBrowserHeight() : 400.f;
 
 	return SNew(SBox)
-		.WidthOverride(300.0f)
-		.HeightOverride(300.f)
+		.WidthOverride(WidthOverride)
+		.HeightOverride(HeightOverride)
 		[
 			ContentBrowserModule.Get().CreateAssetPicker(AssetPickerConfig)
 		];
diff --git a/Engine/Source/Editor/MovieSceneTools/Private/TrackEditors/SkeletalAnimationTrackEditor.cpp b/Engine/Source/Editor/MovieSceneTools/Private/TrackEditors/SkeletalAnimationTrackEditor.cpp
--- a/Engine/Source/Editor/MovieSceneTools/Private/TrackEditors/SkeletalAnimationTrackEditor.cpp
+++ b/Engine/Source/Editor/MovieSceneTools/Private/TrackEditors/SkeletalAnimationTrackEditor.cpp
@@ -29,6 +29,7 @@
 #include "MVVM/ViewModels/ViewDensity.h"
 #include "MVVM/Extensions/ITrackExtension.h"
 #include "ISectionLayoutBuilder.h"
+#include "ISequencer.h"
 #include "Animation/AnimMontage.h"
 #include "Animation/AnimSequence.h"
 #include "Animation/PoseAsset.h"
@@ -45,6 +46,7 @@
 #include "Factories/PoseAssetFactory.h"
 #include "Misc/MessageDialog.h"
 #include "Framework/Notifications/NotificationManager.h"
+#include "SequencerSettings.h"
 #include "Widgets/Notifications/SNotificationList.h"
 
 #include "Engine/SCS_Node.h"
@@ -2057,7 +2059,8 @@ bool FSkeletalAnimationTrackEditor::ShouldFilterAsset(const FAssetData& AssetDat
 
 void FSkeletalAnimationTrackEditor::AddAnimationSubMenu(FMenuBuilder& MenuBuilder, TArray<FGuid> ObjectBindings, USkeleton* Skeleton, UMovieSceneTrack* Track)
 {
-	UMovieSceneSequence* Sequence = GetSequencer() ? GetSequencer()->GetFocusedMovieSceneSequence() : nullptr;
+	TSharedPtr<ISequencer> SequencerPtr = GetSequencer();
+	UMovieSceneSequence* Sequence = SequencerPtr.IsValid() ? SequencerPtr->GetFocusedMovieSceneSequence() : nullptr;
 
 	FAssetPickerConfig AssetPickerConfig;
 	{
@@ -2075,9 +2078,12 @@ void FSkeletalAnimationTrackEditor::AddAnimationSubMenu(FMenuBuilder& MenuBuilde
 
 	FContentBrowserModule& ContentBrowserModule = FModuleManager::Get().LoadModuleChecked<FContentBrowserModule>(TEXT("ContentBrowser"));
 
+	const float WidthOverride = SequencerPtr.IsValid() ? SequencerPtr->GetSequencerSettings()->GetAssetBrowserWidth() : 500.f;
+	const float HeightOverride = SequencerPtr.IsValid() ? SequencerPtr->GetSequencerSettings()->GetAssetBrowserHeight() : 400.f;
+
 	TSharedPtr<SBox> MenuEntry = SNew(SBox)
-		.WidthOverride(300.0f)
-		.HeightOverride(300.f)
+		.WidthOverride(WidthOverride)
+		.HeightOverride(HeightOverride)
 		[
 			ContentBrowserModule.Get().CreateAssetPicker(AssetPickerConfig)
 		];
diff --git a/Engine/Source/Editor/MovieSceneTools/Private/TrackEditors/SubTrackEditor.cpp b/Engine/Source/Editor/MovieSceneTools/Private/TrackEditors/SubTrackEditor.cpp
--- a/Engine/Source/Editor/MovieSceneTools/Private/TrackEditors/SubTrackEditor.cpp
+++ b/Engine/Source/Editor/MovieSceneTools/Private/TrackEditors/SubTrackEditor.cpp
@@ -17,10 +17,12 @@
 #include "Tracks/MovieSceneSubTrack.h"
 #include "Tracks/MovieSceneCinematicShotTrack.h"
 #include "IContentBrowserSingleton.h"
+#include "ISequencer.h"
 #include "ContentBrowserModule.h"
 #include "MVVM/Views/ViewUtilities.h"
 #include "MVVM/Extensions/ITrackExtension.h"
 #include "SequencerSectionPainter.h"
+#include "SequencerSettings.h"
 #include "TrackEditors/SubTrackEditorBase.h"
 #include "DragAndDrop/AssetDragDropOp.h"
 #include "MovieSceneMetaData.h"
@@ -882,8 +884,8 @@ TSharedRef<SWidget> FSubTrackEditor::HandleAddSubSequenceComboButtonGetMenuConte
 			FUIAction(FExecuteAction::CreateSP(this, &FSubTrackEditor::InsertSection, Track))
 		);
 
-
-		UMovieSceneSequence* Sequence = GetSequencer() ? GetSequencer()->GetFocusedMovieSceneSequence() : nullptr;
+		TSharedPtr<ISequencer> SequencerPtr = GetSequencer();
+		UMovieSceneSequence* Sequence = SequencerPtr.IsValid() ? SequencerPtr->GetFocusedMovieSceneSequence() : nullptr;
 		FAssetPickerConfig AssetPickerConfig;
 		{
 			AssetPickerConfig.OnAssetSelected = FOnAssetSelected::CreateRaw( this, &FSubTrackEditor::HandleAddSubSequenceComboButtonMenuEntryExecute, Track);
@@ -898,9 +900,12 @@ TSharedRef<SWidget> FSubTrackEditor::HandleAddSubSequenceComboButtonGetMenuConte
 
 		FContentBrowserModule& ContentBrowserModule = FModuleManager::Get().LoadModuleChecked<FContentBrowserModule>(TEXT("ContentBrowser"));
 
+		const float WidthOverride = SequencerPtr.IsValid() ? SequencerPtr->GetSequencerSettings()->GetAssetBrowserWidth() : 500.f;
+		const float HeightOverride = SequencerPtr.IsValid() ? SequencerPtr->GetSequencerSettings()->GetAssetBrowserHeight() : 400.f;
+
 		TSharedPtr<SBox> MenuEntry = SNew(SBox)
-			.WidthOverride(300.0f)
-			.HeightOverride(300.f)
+			.WidthOverride(WidthOverride)
+			.HeightOverride(HeightOverride)
 			[
 				ContentBrowserModule.Get().CreateAssetPicker(AssetPickerConfig)
 			];
diff --git a/Engine/Source/Editor/MovieSceneTools/Private/TrackEditors/TransformTrackEditor.cpp b/Engine/Source/Editor/MovieSceneTools/Private/TrackEditors/TransformTrackEditor.cpp
--- a/Engine/Source/Editor/MovieSceneTools/Private/TrackEditors/TransformTrackEditor.cpp
+++ b/Engine/Source/Editor/MovieSceneTools/Private/TrackEditors/TransformTrackEditor.cpp
@@ -19,10 +19,12 @@
 #include "ISectionLayoutBuilder.h"
 #include "IKeyArea.h"
 #include "IContentBrowserSingleton.h"
+#include "ISequencer.h"
 #include "ContentBrowserModule.h"
 #include "Editor.h"
 #include "Sections/TransformPropertySection.h"
 #include "SequencerUtilities.h"
+#include "SequencerSettings.h"
 #include "MVVM/Views/ViewUtilities.h"
 #include "MVVM/ViewModels/OutlinerColumns/OutlinerColumnTypes.h"
 #include "MVVM/ViewModels/ObjectBindingModel.h"
@@ -187,15 +189,15 @@ void F3DTransformTrackEditor::BuildTrackContextMenu( FMenuBuilder& MenuBuilder,
 
 void F3DTransformTrackEditor::BuildAssetPickerSubMenu(FMenuBuilder& InMenuBuilder, UMovieScene3DTransformTrack* InTransformTrack)
 {
-	const TSharedPtr<ISequencer> TrackSequencer = GetSequencer();
-	if (!TrackSequencer.IsValid())
+	const TSharedPtr<ISequencer> SequencerPtr = GetSequencer();
+	if (!SequencerPtr.IsValid())
 	{
 		return;
 	}
 
 	TWeakObjectPtr<UMovieScene3DTransformTrack> TransformTrackWeak = InTransformTrack;
 
-	UMovieSceneSequence* Sequence = TrackSequencer->GetFocusedMovieSceneSequence();
+	UMovieSceneSequence* Sequence = SequencerPtr->GetFocusedMovieSceneSequence();
 
 	FAssetPickerConfig AssetPickerConfig;
 	AssetPickerConfig.bAddFilterUI = true;
@@ -216,10 +218,13 @@ void F3DTransformTrackEditor::BuildAssetPickerSubMenu(FMenuBuilder& InMenuBuilde
 
 	FContentBrowserModule& ContentBrowserModule = FModuleManager::LoadModuleChecked<FContentBrowserModule>(TEXT("ContentBrowser"));
 
+	const float WidthOverride = SequencerPtr.IsValid() ? SequencerPtr->GetSequencerSettings()->GetAssetBrowserWidth() : 500.f;
+	const float HeightOverride = SequencerPtr.IsValid() ? SequencerPtr->GetSequencerSettings()->GetAssetBrowserHeight() : 400.f;
+
 	InMenuBuilder.AddWidget(
 		SNew(SBox)
-		.WidthOverride(200.0f)
-		.HeightOverride(400.0f)
+		.WidthOverride(WidthOverride)
+		.HeightOverride(HeightOverride)
 		[
 			ContentBrowserModule.Get().CreateAssetPicker(AssetPickerConfig)
 		], 
@@ -996,7 +1001,6 @@ void F3DTransformTrackEditor::AddTransformKeysForObject( UObject* Object, EMovie
 	}
 }
 
-
 void F3DTransformTrackEditor::AddTransformKeys( UObject* ObjectToKey, const TOptional<FTransformData>& LastTransform, const FTransformData& CurrentTransform, EMovieSceneTransformChannel ChannelsToKey, ESequencerKeyMode KeyMode )
 {
 	if (!GetSequencer()->IsAllowedToChange())
@@ -1041,7 +1045,6 @@ void F3DTransformTrackEditor::AddTransformKeys( UObject* ObjectToKey, const TOpt
 	AnimatablePropertyChanged( FOnKeyProperty::CreateLambda(OnKeyProperty) );
 }
 
-
 FTransformData F3DTransformTrackEditor::RecomposeTransform(const FTransformData& InTransformData, UObject* AnimatedObject, UMovieSceneSection* Section)
 {
 	using namespace UE::MovieScene;
diff --git a/Engine/Source/Editor/Sequencer/Private/SSequencerPlayRateCombo.cpp b/Engine/Source/Editor/Sequencer/Private/SSequencerPlayRateCombo.cpp
--- a/Engine/Source/Editor/Sequencer/Private/SSequencerPlayRateCombo.cpp
+++ b/Engine/Source/Editor/Sequencer/Private/SSequencerPlayRateCombo.cpp
@@ -446,12 +446,16 @@ void SSequencerPlayRateCombo::PopulateCustomClockSourceMenu(FMenuBuilder& MenuBu
 		// Only display actors that are not possessed already
 		InitOptions.Filters->AddFilterPredicate<FActorTreeItem>(FActorTreeItem::FFilterPredicate::CreateLambda( IsActorValid ) );
 
+		TSharedPtr<ISequencer> SequencerPtr = WeakSequencer.Pin();
+		const float WidthOverride = SequencerPtr.IsValid() ? SequencerPtr->GetSequencerSettings()->GetAssetBrowserWidth() : 500.f;
+		const float HeightOverride = SequencerPtr.IsValid() ? SequencerPtr->GetSequencerSettings()->GetAssetBrowserHeight() : 400.f;
+
 		// actor selector to allow the user to choose an actor
 		FSceneOutlinerModule& SceneOutlinerModule = FModuleManager::LoadModuleChecked<FSceneOutlinerModule>("SceneOutliner");
 		ActorMenuBuilder.AddWidget(
 			SNew(SBox)
-			.MaxDesiredHeight(400.0f)
-			.WidthOverride(300.0f)
+			.WidthOverride(WidthOverride)
+			.HeightOverride(HeightOverride)
 			[
 				SceneOutlinerModule.CreateActorPicker(InitOptions, FOnActorPicked::CreateLambda([this](AActor* In){ this->SetCustomClockSource(In); }))
 			],
diff --git a/Engine/Source/Editor/Sequencer/Private/SequencerSettings.cpp b/Engine/Source/Editor/Sequencer/Private/SequencerSettings.cpp
--- a/Engine/Source/Editor/Sequencer/Private/SequencerSettings.cpp
+++ b/Engine/Source/Editor/Sequencer/Private/SequencerSettings.cpp
@@ -72,11 +72,8 @@ USequencerSettings::USequencerSettings( const FObjectInitializer& ObjectInitiali
 	bShowSequencerToolbar = true;
 	bShowMarkedFrames = true;
 	ViewDensity = "Relaxed";
-	bIncludePinnedInFilter = false;
-	bAutoExpandNodesOnFilterPass = false;
-	bUseFilterSubmenusForCategories = false;
-	LastFilterBarLayout = EFilterBarLayout::Horizontal;
-	LastFilterBarSizeCoefficient = 0.1f;
+	AssetBrowserWidth = 500.f;
+	AssetBrowserHeight = 400.f;
 
 	SectionColorTints.Add(FColor(88, 102, 142, 255)); // blue
 	SectionColorTints.Add(FColor(99, 137, 132, 255)); // blue-green
@@ -1051,6 +1048,24 @@ void USequencerSettings::SetViewDensity(FName InViewDensity)
 	}
 }
 
+void USequencerSettings::SetAssetBrowserWidth(float InAssetBrowserWidth)
+{
+	if (InAssetBrowserWidth != AssetBrowserWidth)
+	{
+		AssetBrowserWidth = InAssetBrowserWidth;
+		SaveConfig();
+	}
+}
+
+void USequencerSettings::SetAssetBrowserHeight(float InAssetBrowserHeight)
+{
+	if (InAssetBrowserHeight != AssetBrowserHeight)
+	{
+		AssetBrowserHeight = InAssetBrowserHeight;
+		SaveConfig();
+	}
+}
+
 FSidebarState& USequencerSettings::GetSidebarState()
 {
 	return SidebarState.FindOrAdd(GetFName());
diff --git a/Engine/Source/Editor/Sequencer/Public/SequencerSettings.h b/Engine/Source/Editor/Sequencer/Public/SequencerSettings.h
--- a/Engine/Source/Editor/Sequencer/Public/SequencerSettings.h
+++ b/Engine/Source/Editor/Sequencer/Public/SequencerSettings.h
@@ -482,6 +482,21 @@ public:
 	/** Sets the saved view density */
 	void SetViewDensity(FName InViewDensity);
 
+	/** Gets the asset browser width */
+	float GetAssetBrowserWidth() const { return AssetBrowserWidth; }
+	/** Sets the asset browser width */
+	void SetAssetBrowserWidth(float InAssetBrowserWidth);
+
+	/** Gets the asset browser height */
+	float GetAssetBrowserHeight() const { return AssetBrowserHeight; }
+	/** Sets the asset browser width */
+	void SetAssetBrowserHeight(float InAssetBrowserHeight);
+
+	/** Gets whether the given track filter is enabled */
+	bool IsTrackFilterEnabled(const FString& TrackFilter) const;
+	/** Sets whether the track filter should be enabled/disabled */
+	void SetTrackFilterEnabled(const FString& TrackFilter, bool bEnabled);
+
 	/** Get outliner column visibility in display order */
 	TArray<FColumnVisibilitySetting> GetOutlinerColumnSettings() const { return ColumnVisibilitySettings; }
 	/** Sets the visibility of outliner columns in display order */
@@ -790,6 +805,18 @@ protected:
 	UPROPERTY(config, EditAnywhere, Category = General)
 	FName ViewDensity;
 
+	/** The width for the asset browsers in Sequencer */
+	UPROPERTY(config, EditAnywhere, Category = General)
+	float AssetBrowserWidth;
+
+	/** The height for the asset browsers in Sequencer */
+	UPROPERTY(config, EditAnywhere, Category = General)
+	float AssetBrowserHeight;
+
+	/** The track filters that are enabled */
+	UPROPERTY(config, EditAnywhere, Category = General)
+	TArray<FString> TrackFilters;
+
 	/** List of all columns and their visibility, in the order to be displayed in the outliner view */
 	UPROPERTY(config, EditAnywhere, Category = General)
 	TArray<FColumnVisibilitySetting> ColumnVisibilitySettings;
