[Backout] - CL36140272
Original CL Desc
-----------------------------------------------------------------
Anim layers are editor only asset user data attached to a level sequence that's shown from the animation toolbar. Consist of a set of sections and their corresponding owners(control rig/object) with various properties(mute, layer weight, etc.). Note since some of these values may get set in sequencer we poll some of these properties, using caching as needed.

For control rig we have an idea of multiple sections to key, since we need to have that per control e.g. left arm in left arm section/layer, right arm in right arm section/layer, everything else in base layer.

Weight track for the anim layers shows if it's keyed or not so I reuse control rig editor to cache that data which it also does for anim details proxies.

Move SCollapseControlsWidget to track editor header so it can be shared.

Sections in layers/with overrides can have a Blending Order, this will talked about more in the override code review.

Have a couple todo's in particurly when recompiling control rigs that are in layers

diff --git a/Engine/Plugins/Animation/ControlRig/Source/ControlRig/Private/Sequencer/MovieSceneControlRigParameterTrack.cpp b/Engine/Plugins/Animation/ControlRig/Source/ControlRig/Private/Sequencer/MovieSceneControlRigParameterTrack.cpp
--- a/Engine/Plugins/Animation/ControlRig/Source/ControlRig/Private/Sequencer/MovieSceneControlRigParameterTrack.cpp
+++ b/Engine/Plugins/Animation/ControlRig/Source/ControlRig/Private/Sequencer/MovieSceneControlRigParameterTrack.cpp
@@ -403,76 +403,9 @@ UMovieSceneSection* UMovieSceneControlRigParameterTrack::FindOrAddSection(FFrame
 	return NewSection;
 }
 
-TArray<TWeakObjectPtr<UMovieSceneSection>> UMovieSceneControlRigParameterTrack::GetSectionsToKey() const
-{
-	TArray<TWeakObjectPtr<UMovieSceneSection>> SectionsToKey;
-	if (SectionToKeyPerControl.Num() > 0)
-	{
-		SectionToKeyPerControl.GenerateValueArray(SectionsToKey);
-	}
-	else
-	{
-		SectionsToKey.Add(SectionToKey);
-	}
-	return SectionsToKey;
-}
-
-UMovieSceneSection* UMovieSceneControlRigParameterTrack::GetSectionToKey(const FName& InControlName) const
-{
-	if (const TWeakObjectPtr<UMovieSceneSection>* Section = SectionToKeyPerControl.Find(InControlName))
-	{
-		if (Section->IsValid())
-		{
-			return Section->Get();
-		}
-	}
-	return GetSectionToKey();
-}
-
-void UMovieSceneControlRigParameterTrack::SetSectionToKey(UMovieSceneSection* InSection, const FName& InControlName)
-{
-	if (Sections.Num() < 1 || InSection == nullptr)
-	{
-		return;
-	}
-	Modify();
-	SectionToKeyPerControl.Add(InControlName, InSection);
-	SectionToKey = Sections[0];
-}
-
 void UMovieSceneControlRigParameterTrack::SetSectionToKey(UMovieSceneSection* InSection)
 {
-	if (Sections.Num() < 1 || InSection == nullptr)
-	{
-		return;
-	}
-	if (UMovieSceneControlRigParameterSection* CRSection = Cast<UMovieSceneControlRigParameterSection>(InSection))
-	{
-		Modify();
-		if (SectionToKeyPerControl.Num() > 0) //we have sections that are in layers so need to respect them
-		{
-			if (bSetSectionToKeyPerControl)
-			{
-				for (TPair < FName, TWeakObjectPtr<UMovieSceneSection>>& SectionToKeyItem : SectionToKeyPerControl)
-				{
-					//only set it as the section to key if it's in that section, otherwise leave it alone
-					FChannelMapInfo* pChannelIndex = CRSection->ControlChannelMap.Find(SectionToKeyItem.Key);
-					if (pChannelIndex)
-					{
-						if (CRSection->GetControlsMask(pChannelIndex->MaskIndex))
-						{
-							SectionToKeyItem.Value = InSection;
-						}
-					}
-				}
-			}
-			SectionToKey = Sections[0];
-		}
-		else
-		{
-			SectionToKey = InSection;
-		}
-	}
+	SectionToKey = InSection;
 }
 
 UMovieSceneSection* UMovieSceneControlRigParameterTrack::GetSectionToKey() const
diff --git a/Engine/Plugins/Animation/ControlRig/Source/ControlRig/Public/Sequencer/MovieSceneControlRigParameterTrack.h b/Engine/Plugins/Animation/ControlRig/Source/ControlRig/Public/Sequencer/MovieSceneControlRigParameterTrack.h
--- a/Engine/Plugins/Animation/ControlRig/Source/ControlRig/Public/Sequencer/MovieSceneControlRigParameterTrack.h
+++ b/Engine/Plugins/Animation/ControlRig/Source/ControlRig/Public/Sequencer/MovieSceneControlRigParameterTrack.h
@@ -11,7 +11,6 @@
 #include "INodeAndChannelMappings.h"
 #include "MovieSceneControlRigParameterSection.h"
 #include "EulerTransform.h"
-#include "Tracks/IMovieSceneSectionsToKey.h"
 #include "MovieSceneControlRigParameterTrack.generated.h"
 
 struct FEndLoadPackageContext;
@@ -42,8 +41,6 @@ class UMovieSceneControlRigParameterTrack
 	: public UMovieSceneNameableTrack
 	, public IMovieSceneTrackTemplateProducer
 	, public INodeAndChannelMappings
-	, public IMovieSceneSectionsToKey
-
 {
 	GENERATED_UCLASS_BODY()
 
@@ -147,18 +144,6 @@ public:
 	 */
 	virtual UMovieSceneSection* GetSectionToKey() const override;
 
-	/**
-	* Get multiple sections to key
-	* @return multiple sections to key
-	*/
-	virtual TArray<TWeakObjectPtr<UMovieSceneSection>> GetSectionsToKey() const override;
-
-	/**
-	*  Control rig supports per control sections to key needed for anim layer workflows
-	*/
-	CONTROLRIG_API UMovieSceneSection* GetSectionToKey(const FName& InControlName) const;
-	CONTROLRIG_API void SetSectionToKey(UMovieSceneSection* InSection, const FName& inControlName);
-
 	CONTROLRIG_API void SetTrackName(FName InName) { TrackName = InName; }
 
 	UMovieSceneControlRigParameterSection::FSpaceChannelAddedEvent& SpaceChannelAdded() { return OnSpaceChannelAdded; }
@@ -174,11 +159,6 @@ public:
 
 	CONTROLRIG_API int32 GetPriorityOrder() const;
 	CONTROLRIG_API void SetPriorityOrder(int32 InPriorityIndex);
-private:
-	UPROPERTY()
-	TMap<FName, TWeakObjectPtr<UMovieSceneSection>> SectionToKeyPerControl;
-public:
-	bool bSetSectionToKeyPerControl = true;// set up the map when calling the UMovieSceneTrack::SetSectionToKey()
 private:
 	//get the rotation orderm will not be set if it's default FRotator order. If bCurrent is true, it uses what's set
 	//if false it uses the default setting from the current control rig.
diff --git a/Engine/Plugins/Animation/ControlRig/Source/ControlRigEditor/Private/ControlRigEditorStyle.h b/Engine/Plugins/Animation/ControlRig/Source/ControlRigEditor/Private/ControlRigEditorStyle.h
--- a/Engine/Plugins/Animation/ControlRig/Source/ControlRigEditor/Private/ControlRigEditorStyle.h
+++ b/Engine/Plugins/Animation/ControlRig/Source/ControlRigEditor/Private/ControlRigEditorStyle.h
@@ -90,11 +90,8 @@ public:
 			Set("ControlRig.TweenTool", new IMAGE_BRUSH_SVG("Slate/AnimationTweens", Icon16x16));
 			Set("ControlRig.EditableMotionTrails", new IMAGE_BRUSH_SVG("Slate/EditableMotionTrails", Icon16x16));
 			Set("ControlRig.TemporaryPivot", new IMAGE_BRUSH_SVG("Slate/TemporaryPivot", Icon16x16));
-			Set("ControlRig.AnimLayerSelected", new IMAGE_BRUSH_SVG("Slate/AnimLayerSelected", Icon16x16));
-			Set("ControlRig.FilterAnimLayerSelected", new IMAGE_BRUSH_SVG("Slate/FilterAnimLayerSelected", Icon16x16));
-			Set("ControlRig.AnimLayers", new IMAGE_BRUSH_SVG("Slate/AnimLayers", Icon16x16));
-
 		}
+
 		// Control Rig Editor styles
 		{
 			// tab icons
diff --git a/Engine/Plugins/Animation/ControlRig/Source/ControlRigEditor/Private/ControlRigSequencerEditorLibrary.cpp b/Engine/Plugins/Animation/ControlRig/Source/ControlRigEditor/Private/ControlRigSequencerEditorLibrary.cpp
--- a/Engine/Plugins/Animation/ControlRig/Source/ControlRigEditor/Private/ControlRigSequencerEditorLibrary.cpp
+++ b/Engine/Plugins/Animation/ControlRig/Source/ControlRigEditor/Private/ControlRigSequencerEditorLibrary.cpp
@@ -49,8 +49,6 @@
 #include "Sections/MovieSceneConstrainedSection.h"
 #include "BakingAnimationKeySettings.h"
 #include "EditMode/ControlRigEditMode.h"
-#include "Sequencer/AnimLayers/AnimLayers.h"
-
 
 #include UE_INLINE_GENERATED_CPP_BY_NAME(ControlRigSequencerEditorLibrary)
 
@@ -3289,133 +3287,5 @@ bool UControlRigSequencerEditorLibrary::SetControlRigApplyMode(UControlRig* InCo
 	return false;
 }
 
-bool UControlRigSequencerEditorLibrary::DeleteAnimLayer(int32 Index)
-{
-	TWeakPtr<ISequencer> WeakSequencer = GetSequencerFromAsset();
-	if (WeakSequencer.IsValid() == false)
-	{
-		UE_LOG(LogControlRig, Error, TEXT("DeleteAnimLayer: Need open Sequencer"));
-		return false;
-	}
-	TSharedPtr<ISequencer> SequencerPtr = WeakSequencer.Pin();
-	if (UAnimLayers* AnimLayers = UAnimLayers::GetAnimLayers(SequencerPtr.Get()))
-	{
-		if (Index != INDEX_NONE)
-		{
-			return AnimLayers->DeleteAnimLayer(SequencerPtr.Get(), Index);
-		}
-	}
-	else
-	{
-		UE_LOG(LogControlRig, Error, TEXT("DeleteAnimLayer: No Anim Layers on Level Sequence"));
-	}
-	return false;
-}
-
-int32 UControlRigSequencerEditorLibrary::DuplicateAnimLayer(int32 Index)
-{
-	TWeakPtr<ISequencer> WeakSequencer = GetSequencerFromAsset();
-	if (WeakSequencer.IsValid() == false)
-	{
-		UE_LOG(LogControlRig, Error, TEXT("DuplicateAnimLayer: Need open Sequencer"));
-		return INDEX_NONE;
-	}
-	TSharedPtr<ISequencer> SequencerPtr = WeakSequencer.Pin();
-	if (UAnimLayers* AnimLayers = UAnimLayers::GetAnimLayers(SequencerPtr.Get()))
-	{
-		if (Index != INDEX_NONE)
-		{
-			return AnimLayers->DuplicateAnimLayer(SequencerPtr.Get(), Index);
-		}
-	}
-	else
-	{
-		UE_LOG(LogControlRig, Error, TEXT("DuplicateAnimLayer: No Anim Layers on Level Sequence"));
-	}
-	return INDEX_NONE;
-}
-
-int32 UControlRigSequencerEditorLibrary::AddAnimLayerFromSelection()
-{
-	TWeakPtr<ISequencer> WeakSequencer = GetSequencerFromAsset();
-	if (WeakSequencer.IsValid() == false)
-	{
-		UE_LOG(LogControlRig, Error, TEXT("AddAnimLayerFromSelection: Need open Sequencer"));
-		return INDEX_NONE;
-	}
-	TSharedPtr<ISequencer> SequencerPtr = WeakSequencer.Pin();
-	if (UAnimLayers* AnimLayers = UAnimLayers::GetAnimLayers(SequencerPtr.Get()))
-	{
-		return AnimLayers->AddAnimLayerFromSelection(SequencerPtr.Get());
-	}
-	else
-	{
-		UE_LOG(LogControlRig, Error, TEXT("AddAnimLayerFromSelection: No Anim Layers on Level Sequence"));
-	}
-	return INDEX_NONE;
-
-}
-
-bool UControlRigSequencerEditorLibrary::MergeAnimLayers(const TArray<int32>& Indices)
-{
-	TWeakPtr<ISequencer> WeakSequencer = GetSequencerFromAsset();
-	if (WeakSequencer.IsValid() == false)
-	{
-		UE_LOG(LogControlRig, Error, TEXT("MergeAnimLayers: Need open Sequencer"));
-		return false;
-	}
-	TSharedPtr<ISequencer> SequencerPtr = WeakSequencer.Pin();
-	if (UAnimLayers* AnimLayers = UAnimLayers::GetAnimLayers(SequencerPtr.Get()))
-	{
-		return AnimLayers->MergeAnimLayers(SequencerPtr.Get(), Indices, nullptr);
-	}
-	else
-	{
-		UE_LOG(LogControlRig, Error, TEXT("MergeAnimLayers: No Anim Layers on Level Sequence"));
-	}
-	return false;
-}
-
-TArray<UAnimLayer*> UControlRigSequencerEditorLibrary::GetAnimLayers()
-{
-	TArray<UAnimLayer*> AnimLayersArray;
-	TWeakPtr<ISequencer> WeakSequencer = GetSequencerFromAsset();
-	if (WeakSequencer.IsValid() == false)
-	{
-		UE_LOG(LogControlRig, Error, TEXT("GetAnimLayers: Need open Sequencer"));
-		return AnimLayersArray;
-	}
-	TSharedPtr<ISequencer> SequencerPtr = WeakSequencer.Pin();
-	if (UAnimLayers* AnimLayers = UAnimLayers::GetAnimLayers(SequencerPtr.Get()))
-	{
-		return AnimLayers->AnimLayers;
-	}
-	else
-	{
-		UE_LOG(LogControlRig, Error, TEXT("GetAnimLayers: No Anim Layers on Level Sequence"));
-	}
-	return AnimLayersArray;
-}
-
-int32 UControlRigSequencerEditorLibrary::GetAnimLayerIndex(UAnimLayer* AnimLayer)
-{
-	TWeakPtr<ISequencer> WeakSequencer = GetSequencerFromAsset();
-	if (WeakSequencer.IsValid() == false)
-	{
-		UE_LOG(LogControlRig, Error, TEXT("GetAnimLayerIndex: Need open Sequencer"));
-		return INDEX_NONE;
-	}
-	TSharedPtr<ISequencer> SequencerPtr = WeakSequencer.Pin();
-	if (UAnimLayers* AnimLayers = UAnimLayers::GetAnimLayers(SequencerPtr.Get()))
-	{
-		return AnimLayers->GetAnimLayerIndex(AnimLayer);
-	}
-	else
-	{
-		UE_LOG(LogControlRig, Error, TEXT("GetAnimLayerIndex: No Anim Layers on Level Sequence"));
-	}
-	return INDEX_NONE;
-}
-
 #undef LOCTEXT_NAMESPACE
 
diff --git a/Engine/Plugins/Animation/ControlRig/Source/ControlRigEditor/Private/EditMode/AnimDetailsProxy.cpp b/Engine/Plugins/Animation/ControlRig/Source/ControlRigEditor/Private/EditMode/AnimDetailsProxy.cpp
--- a/Engine/Plugins/Animation/ControlRig/Source/ControlRigEditor/Private/EditMode/AnimDetailsProxy.cpp
+++ b/Engine/Plugins/Animation/ControlRig/Source/ControlRigEditor/Private/EditMode/AnimDetailsProxy.cpp
@@ -133,7 +133,7 @@ static EAnimDetailSelectionState CachePropertySelection(TWeakPtr<FCurveEditor>&
 									}
 									if (UMovieSceneControlRigParameterTrack* Track = Section->GetTypedOuter< UMovieSceneControlRigParameterTrack>())
 									{
-										if (Track->GetSectionToKey(ControlElement->GetFName()) != Section)
+										if (Track->GetSectionToKey() != Section)
 										{
 											continue;
 										}
@@ -3145,29 +3145,21 @@ UControlRigControlsProxy* UControlRigDetailPanelControlProxies::AddProxy(UObject
 				SequencerOnlyProxies.Add(InObject, NewProxies);
 			}
 		}
-		if (!Proxy)
-		{
-			return Proxy;
-		}
-		Proxy->Type = Type;
-		Proxy->OwnerObject = InObject;
-		Proxy->OwnerBindingAndTrack.WeakTrack = Track;
-		Proxy->OwnerBindingAndTrack.Binding = Binding;
-		Proxy->AddSequencerProxyItem(InObject, Proxy->OwnerBindingAndTrack.WeakTrack, Binding);
-		if (Type == ERigControlType::Transform || Type == ERigControlType::TransformNoScale || Type == ERigControlType::EulerTransform)
+		if (Proxy)
 		{
+			Proxy->Type = Type;
+			Proxy->OwnerObject = InObject;
+			Proxy->OwnerBindingAndTrack.WeakTrack = Track;
+			Proxy->OwnerBindingAndTrack.Binding = Binding;
+			Proxy->AddSequencerProxyItem(InObject, Proxy->OwnerBindingAndTrack.WeakTrack, Binding);
 			Proxy->bIsIndividual = false;
+			Proxy->SetFlags(RF_Transactional);
+			Proxy->Modify();
+			UWorld* World = GCurrentLevelEditingViewportClient ? GCurrentLevelEditingViewportClient->GetWorld() : nullptr;
+			const FConstraintsManagerController& Controller = FConstraintsManagerController::Get(World);
+			Controller.EvaluateAllConstraints();
+			Proxy->ValueChanged();
 		}
-		else
-		{
-			Proxy->bIsIndividual = true;
-		}			
-		Proxy->SetFlags(RF_Transactional);
-		Proxy->Modify();
-		UWorld* World = GCurrentLevelEditingViewportClient ? GCurrentLevelEditingViewportClient->GetWorld() : nullptr;
-		const FConstraintsManagerController& Controller = FConstraintsManagerController::Get(World);
-		Controller.EvaluateAllConstraints();
-		Proxy->ValueChanged();
 	}
 	return Proxy;
 }
@@ -3707,7 +3699,7 @@ bool UControlRigDetailPanelControlProxies::SelectPropertyInternal(UControlRigCon
 									{
 										continue;
 									}
-									UMovieSceneSection* SectionToKey = Track->GetSectionToKey(ControlElement->GetFName()) ? Track->GetSectionToKey(ControlElement->GetFName()) : Track->GetAllSections()[0];
+									UMovieSceneSection* SectionToKey = Track->GetSectionToKey() ? Track->GetSectionToKey() : Track->GetAllSections()[0];
 									if (TViewModelPtr<FChannelGroupOutlinerModel> ChannelModel = CastViewModel<FChannelGroupOutlinerModel>(OutlinerExtenstionIt.GetCurrentItem()))
 									{
 										if (ChannelModel->GetChannel(SectionToKey) == nullptr) //if not section to key we also don't select it.
diff --git a/Engine/Plugins/Animation/ControlRig/Source/ControlRigEditor/Private/EditMode/ControlRigEditMode.cpp b/Engine/Plugins/Animation/ControlRig/Source/ControlRigEditor/Private/EditMode/ControlRigEditMode.cpp
--- a/Engine/Plugins/Animation/ControlRig/Source/ControlRigEditor/Private/EditMode/ControlRigEditMode.cpp
+++ b/Engine/Plugins/Animation/ControlRig/Source/ControlRigEditor/Private/EditMode/ControlRigEditMode.cpp
@@ -81,7 +81,6 @@
 #include "Editor/Sequencer/Private/SSequencer.h"
 #include "Slate/SceneViewport.h"
 #include "Tools/BakingHelper.h"
-#include "Sequencer/AnimLayers/AnimLayers.h"
 
 #include UE_INLINE_GENERATED_CPP_BY_NAME(ControlRigEditMode)
 
@@ -812,15 +811,17 @@ TSet<FName> FControlRigEditMode::GetActiveControlsFromSequencer(UControlRig* Con
 				{
 					if (ControlRigParameterTrack->GetControlRig() == ControlRig)
 					{
-						TArray<FRigControlElement*> Controls;
-						ControlRig->GetControlsInOrder(Controls);
-						int Index = 0;
-						for (FRigControlElement* ControlElement : Controls)
+						UMovieSceneControlRigParameterSection* ActiveSection = Cast<UMovieSceneControlRigParameterSection>(ControlRigParameterTrack->GetSectionToKey());
+						if (ActiveSection)
 						{
-							UMovieSceneControlRigParameterSection* ActiveSection = Cast<UMovieSceneControlRigParameterSection>(ControlRigParameterTrack->GetSectionToKey(ControlElement->GetFName()));
-							if (ActiveSection)
+							TArray<FRigControlElement*> Controls;
+							ControlRig->GetControlsInOrder(Controls);
+							TArray<bool> Mask = ActiveSection->GetControlsMask();
+
+							TArray<FName> Names;
+							int Index = 0;
+							for (FRigControlElement* ControlElement : Controls)
 							{
-								TArray<bool> Mask = ActiveSection->GetControlsMask();
 								if (Mask[Index])
 								{
 									ActiveControls.Add(ControlElement->GetFName());
@@ -5785,7 +5786,6 @@ bool FDetailKeyFrameCacheAndHandler::IsPropertyKeyable(const UClass* InObjectCla
 			}
 		}
 	}
-	
 
 	if (InObjectClass
 		&& InObjectClass->IsChildOf(UAnimDetailControlsProxyTransform::StaticClass())
@@ -5801,11 +5801,6 @@ bool FDetailKeyFrameCacheAndHandler::IsPropertyKeyable(const UClass* InObjectCla
 		return true;
 	}
 
-	if (InObjectClass && InObjectClass->IsChildOf(UAnimLayer::StaticClass()))
-	{
-		return true;
-	}
-
 	if ((InObjectClass && InObjectClass->IsChildOf(UAnimDetailControlsProxyTransform::StaticClass()) && InPropertyHandle.GetProperty())
 		&& (InPropertyHandle.GetProperty()->GetFName() == GET_MEMBER_NAME_CHECKED(UAnimDetailControlsProxyTransform, Location) ||
 		InPropertyHandle.GetProperty()->GetFName() == GET_MEMBER_NAME_CHECKED(UAnimDetailControlsProxyTransform, Rotation) ||
@@ -5909,14 +5904,11 @@ void FDetailKeyFrameCacheAndHandler::OnKeyPropertyClicked(const IPropertyHandle&
 	KeyedPropertyHandle.GetOuterObjects(Objects);
 	for (UObject* Object : Objects)
 	{
-		if (UControlRigControlsProxy* Proxy = Cast< UControlRigControlsProxy>(Object))
+		UControlRigControlsProxy* Proxy = Cast< UControlRigControlsProxy>(Object);
+		if (Proxy)
 		{
 			Proxy->SetKey(SequencerPtr, KeyedPropertyHandle);
 		}
-		else if (UAnimLayer* AnimLayer = (Object->GetTypedOuter<UAnimLayer>()))
-		{
-			AnimLayer->SetKey(SequencerPtr, KeyedPropertyHandle);
-		}
 	}
 }
 
@@ -5960,10 +5952,6 @@ EPropertyKeyedStatus FDetailKeyFrameCacheAndHandler::GetPropertyKeyedStatus(cons
 		{
 			KeyedStatus = Proxy->GetPropertyKeyedStatus(SequencerPtr,PropertyHandle);
 		}
-		else if (UAnimLayer* AnimLayer = (Object->GetTypedOuter<UAnimLayer>()))
-		{
-			KeyedStatus = AnimLayer->GetPropertyKeyedStatus(SequencerPtr, PropertyHandle);
-		}
 		//else check to see if it's in sequencer
 	}
 	CachedPropertyKeyedStatusMap.Add(&PropertyHandle, KeyedStatus);
diff --git a/Engine/Plugins/Animation/ControlRig/Source/ControlRigEditor/Private/EditMode/ControlRigEditModeToolkit.cpp b/Engine/Plugins/Animation/ControlRig/Source/ControlRigEditor/Private/EditMode/ControlRigEditModeToolkit.cpp
--- a/Engine/Plugins/Animation/ControlRig/Source/ControlRigEditor/Private/EditMode/ControlRigEditModeToolkit.cpp
+++ b/Engine/Plugins/Animation/ControlRig/Source/ControlRigEditor/Private/EditMode/ControlRigEditModeToolkit.cpp
@@ -24,7 +24,6 @@
 #include "EditMode/SControlRigSpacePicker.h"
 #include "LevelEditor.h"
 #include "SLevelViewport.h"
-#include "Sequencer/AnimLayers/SAnimLayers.h"
 
 #define LOCTEXT_NAMESPACE "FControlRigEditModeToolkit"
 
@@ -35,7 +34,6 @@ namespace
 }
 
 bool FControlRigEditModeToolkit::bMotionTrailsTabOpen = false;
-bool FControlRigEditModeToolkit::bAnimLayerTabOpen = false;
 bool FControlRigEditModeToolkit::bPoseTabOpen = false;
 bool FControlRigEditModeToolkit::bSnapperTabOpen = false;
 bool FControlRigEditModeToolkit::bTweenOpen = false;
@@ -43,7 +41,6 @@ bool FControlRigEditModeToolkit::bTweenOpen = false;
 const FName FControlRigEditModeToolkit::PoseTabName = FName(TEXT("PoseTab"));
 const FName FControlRigEditModeToolkit::MotionTrailTabName = FName(TEXT("MotionTrailTab"));
 const FName FControlRigEditModeToolkit::SnapperTabName = FName(TEXT("SnapperTab"));
-const FName FControlRigEditModeToolkit::AnimLayerTabName = FName(TEXT("AnimLayerTab"));
 const FName FControlRigEditModeToolkit::TweenOverlayName = FName(TEXT("TweenOverlay"));
 const FName FControlRigEditModeToolkit::OutlinerTabName = FName(TEXT("ControlRigOutlinerTab"));
 const FName FControlRigEditModeToolkit::DetailsTabName = FName(TEXT("ControlRigDetailsTab"));
@@ -118,11 +115,6 @@ void FControlRigEditModeToolkit::TryInvokeToolkitUI(const FName InName)
 		FTabId TabID(MotionTrailTabName);
 		ModeUILayerPtr->GetTabManager()->TryInvokeTab(TabID, false /*bIsActive*/);
 	}
-	else if (InName == AnimLayerTabName)
-	{
-		FTabId TabID(AnimLayerTabName);
-		ModeUILayerPtr->GetTabManager()->TryInvokeTab(TabID, false /*bIsActive*/);
-	}
 	else if (InName == PoseTabName)
 	{
 		FTabId TabID(PoseTabName);
@@ -205,14 +197,6 @@ TSharedRef<SDockTab> SpawnMotionTrailTab(const FSpawnTabArgs& Args)
 		];
 }
 
-TSharedRef<SDockTab> SpawnAnimLayerTab(const FSpawnTabArgs& Args, FControlRigEditMode* InEditorMode)
-{
-	return SNew(SDockTab)
-		[
-			SNew(SAnimLayers, *InEditorMode)
-		];
-}
-
 TSharedRef<SDockTab> SpawnOutlinerTab(const FSpawnTabArgs& Args, FControlRigEditMode* InEditorMode)
 {
 	return SNew(SDockTab)
@@ -426,16 +410,6 @@ void FControlRigEditModeToolkit::RequestModeUITabs()
 			.SetIcon(FSlateIcon(TEXT("ControlRigEditorStyle"), TEXT("ControlRig.EditableMotionTrails")));
 		ModeUILayerPtr->GetTabManager()->RegisterDefaultTabWindowSize(MotionTrailTabName, FVector2D(425, 575));
 
-		ModeUILayerPtr->GetTabManager()->UnregisterTabSpawner(AnimLayerTabName);
-		ModeUILayerPtr->GetTabManager()->RegisterTabSpawner(AnimLayerTabName, FOnSpawnTab::CreateStatic(&SpawnAnimLayerTab,&EditMode))
-			.SetDisplayName(LOCTEXT("AnimLayerTab", "Anim Layers"))
-			.SetTooltipText(LOCTEXT("AnimationLayerTabTooltip", "Animation layers"))
-			.SetGroup(MenuGroup)
-			.SetIcon(FSlateIcon(TEXT("ControlRigEditorStyle"), TEXT("ControlRig.AnimLayers")));
-		ModeUILayerPtr->GetTabManager()->RegisterDefaultTabWindowSize(AnimLayerTabName, FVector2D(425, 200));
-
-
-
 		ModeUILayer.Pin()->ToolkitHostShutdownUI().BindSP(this, &FControlRigEditModeToolkit::UnregisterAndRemoveFloatingTabs);
 
 
@@ -460,10 +434,6 @@ void FControlRigEditModeToolkit::InvokeUI()
 		{
 			TryInvokeToolkitUI(MotionTrailTabName);
 		}
-		if (bAnimLayerTabOpen)
-		{
-			TryInvokeToolkitUI(AnimLayerTabName);
-		}
 		if (bSnapperTabOpen)
 		{
 			TryInvokeToolkitUI(SnapperTabName);
@@ -504,18 +474,6 @@ void FControlRigEditModeToolkit::UnregisterAndRemoveFloatingTabs()
 			}
 			ModeUILayerPtr->GetTabManager()->UnregisterTabSpawner(MotionTrailTabName);
 
-			TSharedPtr<SDockTab> AnimLayerTab = ModeUILayerPtr->GetTabManager()->FindExistingLiveTab(FTabId(AnimLayerTabName));
-			if (AnimLayerTab)
-			{
-				bAnimLayerTabOpen = true;
-				AnimLayerTab->RequestCloseTab();
-			}
-			else
-			{
-				bAnimLayerTabOpen = false;
-			}
-			ModeUILayerPtr->GetTabManager()->UnregisterTabSpawner(AnimLayerTabName);
-
 			TSharedPtr<SDockTab> SnapperTab = ModeUILayerPtr->GetTabManager()->FindExistingLiveTab(FTabId(SnapperTabName));
 			if (SnapperTab)
 			{
diff --git a/Engine/Plugins/Animation/ControlRig/Source/ControlRigEditor/Private/EditMode/ControlRigEditModeToolkit.h b/Engine/Plugins/Animation/ControlRig/Source/ControlRigEditor/Private/EditMode/ControlRigEditModeToolkit.h
--- a/Engine/Plugins/Animation/ControlRig/Source/ControlRigEditor/Private/EditMode/ControlRigEditModeToolkit.h
+++ b/Engine/Plugins/Animation/ControlRig/Source/ControlRigEditor/Private/EditMode/ControlRigEditModeToolkit.h
@@ -65,7 +65,6 @@ public:
 public:
 	static const FName PoseTabName;
 	static const FName MotionTrailTabName;
-	static const FName AnimLayerTabName;
 	static const FName TweenOverlayName;
 	static const FName SnapperTabName;
 	static const FName DetailsTabName;
@@ -94,7 +93,6 @@ protected:
 	void UnregisterAndRemoveFloatingTabs();
 
 	static bool bMotionTrailsTabOpen;
-	static bool bAnimLayerTabOpen;
 	static bool bPoseTabOpen;
 	static bool bSnapperTabOpen;
 	static bool bTweenOpen;
diff --git a/Engine/Plugins/Animation/ControlRig/Source/ControlRigEditor/Private/EditMode/SControlRigEditModeTools.cpp b/Engine/Plugins/Animation/ControlRig/Source/ControlRigEditor/Private/EditMode/SControlRigEditModeTools.cpp
--- a/Engine/Plugins/Animation/ControlRig/Source/ControlRigEditor/Private/EditMode/SControlRigEditModeTools.cpp
+++ b/Engine/Plugins/Animation/ControlRig/Source/ControlRigEditor/Private/EditMode/SControlRigEditModeTools.cpp
@@ -1302,16 +1302,6 @@ void SControlRigEditModeTools::CustomizeToolBarPalette(FToolBarBuilder& ToolBarB
 		EUserInterfaceActionType::Button
 	);
 
-	// Anim Layer
-	ToolBarBuilder.AddToolBarButton(
-		FExecuteAction::CreateRaw(OwningToolkit.Pin().Get(), &FControlRigEditModeToolkit::TryInvokeToolkitUI, FControlRigEditModeToolkit::AnimLayerTabName),
-		NAME_None,
-		LOCTEXT("AnimLayers", "Anim Layers"),
-		LOCTEXT("AnimLayersTooltip", "Display anim layers"),
-		FSlateIcon(TEXT("ControlRigEditorStyle"), TEXT("ControlRig.AnimLayers")),
-		EUserInterfaceActionType::Button
-	);
-
 	//Pivot
 	ToolBarBuilder.AddToolBarButton(
 		FUIAction(
diff --git a/Engine/Plugins/Animation/ControlRig/Source/ControlRigEditor/Private/Sequencer/AnimLayerSequencerFilter.cpp b/Engine/Plugins/Animation/ControlRig/Source/ControlRigEditor/Private/Sequencer/AnimLayerSequencerFilter.cpp
deleted file mode 100644
diff --git a/Engine/Plugins/Animation/ControlRig/Source/ControlRigEditor/Private/Sequencer/AnimLayerSequencerFilter.h b/Engine/Plugins/Animation/ControlRig/Source/ControlRigEditor/Private/Sequencer/AnimLayerSequencerFilter.h
deleted file mode 100644
diff --git a/Engine/Plugins/Animation/ControlRig/Source/ControlRigEditor/Private/Sequencer/AnimLayers/AnimLayers.cpp b/Engine/Plugins/Animation/ControlRig/Source/ControlRigEditor/Private/Sequencer/AnimLayers/AnimLayers.cpp
deleted file mode 100644
diff --git a/Engine/Plugins/Animation/ControlRig/Source/ControlRigEditor/Private/Sequencer/AnimLayers/AnimLayers.h b/Engine/Plugins/Animation/ControlRig/Source/ControlRigEditor/Private/Sequencer/AnimLayers/AnimLayers.h
deleted file mode 100644
diff --git a/Engine/Plugins/Animation/ControlRig/Source/ControlRigEditor/Private/Sequencer/AnimLayers/SAnimLayers.cpp b/Engine/Plugins/Animation/ControlRig/Source/ControlRigEditor/Private/Sequencer/AnimLayers/SAnimLayers.cpp
deleted file mode 100644
diff --git a/Engine/Plugins/Animation/ControlRig/Source/ControlRigEditor/Private/Sequencer/AnimLayers/SAnimLayers.h b/Engine/Plugins/Animation/ControlRig/Source/ControlRigEditor/Private/Sequencer/AnimLayers/SAnimLayers.h
deleted file mode 100644
diff --git a/Engine/Plugins/Animation/ControlRig/Source/ControlRigEditor/Private/Sequencer/ControlRigParameterTrackEditor.cpp b/Engine/Plugins/Animation/ControlRig/Source/ControlRigEditor/Private/Sequencer/ControlRigParameterTrackEditor.cpp
--- a/Engine/Plugins/Animation/ControlRig/Source/ControlRigEditor/Private/Sequencer/ControlRigParameterTrackEditor.cpp
+++ b/Engine/Plugins/Animation/ControlRig/Source/ControlRigEditor/Private/Sequencer/ControlRigParameterTrackEditor.cpp
@@ -2046,7 +2046,7 @@ void FControlRigParameterTrackEditor::AddControlRigFromComponent(FGuid InGuid)
 
 bool FControlRigParameterTrackEditor::HasTransformKeyOverridePriority() const
 {
-	return false;
+	return false; 
 }
 
 bool FControlRigParameterTrackEditor::CanAddTransformKeysForSelectedObjects() const
@@ -3636,7 +3636,7 @@ void FControlRigParameterTrackEditor::HandleControlSelected(UControlRig* Subject
 			//Just select in section to key, if deselecting makes sure deselected everywhere
 			if (bSelected == true)
 			{
-				UMovieSceneSection* Section = Track->GetSectionToKey(ControlElement->GetFName());
+				UMovieSceneSection* Section = Track->GetSectionToKey();
 				UMovieSceneControlRigParameterSection* ParamSection = Cast<UMovieSceneControlRigParameterSection>(Section);
 				SelectSequencerNodeInSection(ParamSection, ControlElement->GetFName(), bSelected);
 			}
@@ -4133,18 +4133,7 @@ FKeyPropertyResult FControlRigParameterTrackEditor::AddKeysToControlRigHandle(UO
 	{
 		float Weight = 1.0f;
 
-		UMovieSceneSection* SectionToKey = Track->GetSectionToKey(RigControlName);
-		if (SectionToKey)
-		{
-			if (SectionToKey->HasEndFrame() && SectionToKey->GetExclusiveEndFrame() < KeyTime)
-			{
-				SectionToKey->SetEndFrame(KeyTime);
-			}
-			else if(SectionToKey->HasStartFrame() && SectionToKey->GetInclusiveStartFrame() > KeyTime)
-			{
-				SectionToKey->SetStartFrame(KeyTime);
-			}
-		}
+		UMovieSceneSection* SectionToKey = bCreateSection ? Track->FindOrExtendSection(KeyTime, Weight) : Track->FindSection(KeyTime);
 
 		// If there's no overlapping section to key, create one only if a track was newly created. Otherwise, skip keying altogether
 		// so that the user is forced to create a section to key on.
@@ -4176,7 +4165,7 @@ FKeyPropertyResult FControlRigParameterTrackEditor::AddKeysToControlRigHandle(UO
 		//if we create a key then compensate
 		if (KeyPropertyResult.bKeyCreated)
 		{
-			UMovieSceneControlRigParameterSection* ParamSection = Cast<UMovieSceneControlRigParameterSection>(Track->GetSectionToKey(ControlRigName));
+			UMovieSceneControlRigParameterSection* ParamSection = Cast<UMovieSceneControlRigParameterSection>(Track->GetSectionToKey());
 			if (UControlRig* SectionControlRig = ParamSection ? ParamSection->GetControlRig() : nullptr)
 			{
 				TOptional<FFrameNumber> OptionalKeyTime = KeyTime;
@@ -4247,7 +4236,7 @@ void FControlRigParameterTrackEditor::AddControlKeys(
 			BeginKeying(LocalFrameTime.RoundToFrame());
 		}
 		const FFrameNumber FrameTime = GetTimeForKey();
-		UMovieSceneSection* Section = Track->GetSectionToKey(RigControlName);
+		UMovieSceneSection* Section = Track->FindSection(FrameTime);
 		ParamSection = Cast<UMovieSceneControlRigParameterSection>(Section);
 
 		if (ParamSection && ParamSection->GetDoNotKey())
@@ -4940,12 +4929,50 @@ void FControlRigParameterTrackEditor::SelectFKBonesToAnimate(UFKControlRig* Auto
 	//reconstruct all channel proxies TODO or not to do that is the question
 }
 
+//////////////////////////////////////////////////////////////
+/// SCollapseControlsWidget
+///////////////////////////////////////////////////////////
+
+/** Widget allowing collapsing of controls */
+class SCollapseControlsWidget : public SCompoundWidget
+{
+public:
+
+	SLATE_BEGIN_ARGS(SCollapseControlsWidget)
+		: _Sequencer(nullptr), _OwnerTrack(nullptr)
+	{}
+
+	SLATE_ARGUMENT(TWeakPtr<ISequencer>, Sequencer)
+	SLATE_ARGUMENT(UMovieSceneTrack*, OwnerTrack)
+	SLATE_END_ARGS()
+
+	void Construct(const FArguments& InArgs);
+	virtual ~SCollapseControlsWidget() override {}
+
+	FReply OpenDialog(bool bModal = true);
+	void CloseDialog();
+
+
+private:
+	void Collapse();
+
+	TWeakPtr<ISequencer> Sequencer;
+	TWeakObjectPtr<UMovieSceneTrack> OwnerTrack;
+	//static to be reused
+	static TOptional<FBakingAnimationKeySettings> CollapseControlsSettings;
+	//structonscope for details panel
+	TSharedPtr < TStructOnScope<FBakingAnimationKeySettings>> Settings;
+	TWeakPtr<SWindow> DialogWindow;
+	TSharedPtr<IStructureDetailsView> DetailsView;
+};
+
 
 TOptional<FBakingAnimationKeySettings> SCollapseControlsWidget::CollapseControlsSettings;
 
 void SCollapseControlsWidget::Construct(const FArguments& InArgs)
 {
 	Sequencer = InArgs._Sequencer;
+	OwnerTrack = InArgs._OwnerTrack;
 
 	if (CollapseControlsSettings.IsSet() == false)
 	{
@@ -5038,11 +5065,13 @@ void SCollapseControlsWidget::Construct(const FArguments& InArgs)
 		];
 }
 
+
 void  SCollapseControlsWidget::Collapse()
 {
 	FBakingAnimationKeySettings* BakeSettings = Settings->Get();
 	TSharedPtr<ISequencer> SequencerPtr = Sequencer.Pin();
-	CollapseCB.ExecuteIfBound(SequencerPtr, *BakeSettings);
+	FControlRigParameterTrackEditor::CollapseAllLayers(SequencerPtr, OwnerTrack.Get(), *BakeSettings);
+
 	CollapseControlsSettings = *BakeSettings;
 }
 
@@ -5279,20 +5308,16 @@ bool FControlRigParameterTrackEditor::CollapseAllLayers(TSharedPtr<ISequencer>&S
 
 void FControlRigParameterSection::CollapseAllLayers()
 {
+	TSharedPtr<ISequencer> SequencerPtr = WeakSequencer.Pin();
 	if (UMovieSceneControlRigParameterSection* ParameterSection = CastChecked<UMovieSceneControlRigParameterSection>(WeakSection.Get()))
 	{
 		UMovieSceneTrack* OwnerTrack = ParameterSection->GetTypedOuter<UMovieSceneTrack>();
-		FCollapseControlsCB CollapseCB = FCollapseControlsCB::CreateLambda([this,OwnerTrack](TSharedPtr<ISequencer>& InSequencer, const FBakingAnimationKeySettings& InSettings)
-		{
-			FControlRigParameterTrackEditor::CollapseAllLayers(InSequencer, OwnerTrack, InSettings);
-		});
-
 		TSharedPtr<ISequencer> Sequencer = WeakSequencer.Pin();
 		TSharedRef<SCollapseControlsWidget> BakeWidget =
 			SNew(SCollapseControlsWidget)
-			.Sequencer(Sequencer);
+			.Sequencer(Sequencer)
+			.OwnerTrack(OwnerTrack);
 
-		BakeWidget->SetCollapseCB(CollapseCB);
 		BakeWidget->OpenDialog(false);
 	}
 }
diff --git a/Engine/Plugins/Animation/ControlRig/Source/ControlRigEditor/Private/Sequencer/ControlRigParameterTrackEditor.h b/Engine/Plugins/Animation/ControlRig/Source/ControlRigEditor/Private/Sequencer/ControlRigParameterTrackEditor.h
--- a/Engine/Plugins/Animation/ControlRig/Source/ControlRigEditor/Private/Sequencer/ControlRigParameterTrackEditor.h
+++ b/Engine/Plugins/Animation/ControlRig/Source/ControlRigEditor/Private/Sequencer/ControlRigParameterTrackEditor.h
@@ -36,46 +36,6 @@ struct FMovieSceneChannel;
 struct FKeyAddOrDeleteEventItem;
 struct FKeyMoveEventItem;
 struct FBakingAnimationKeySettings;
-class ISequencer;
-class UMovieSceneTrack;
-class IStructureDetailsView;
-
-//////////////////////////////////////////////////////////////
-/// SCollapseControlsWidget
-///////////////////////////////////////////////////////////
-DECLARE_DELEGATE_TwoParams(FCollapseControlsCB, TSharedPtr<ISequencer>& InSequencer, const FBakingAnimationKeySettings& InSettings);
-
-/** Widget allowing collapsing of controls */
-class SCollapseControlsWidget : public SCompoundWidget
-{
-public:
-
-	SLATE_BEGIN_ARGS(SCollapseControlsWidget)
-		: _Sequencer(nullptr)
-		{}
-		SLATE_ARGUMENT(TWeakPtr<ISequencer>, Sequencer)
-	SLATE_END_ARGS()
-
-	void Construct(const FArguments& InArgs);
-	virtual ~SCollapseControlsWidget() override {}
-
-	FReply OpenDialog(bool bModal = true);
-	void CloseDialog();
-
-	void SetCollapseCB(FCollapseControlsCB& InCB) { CollapseCB = InCB; }
-private:
-	void Collapse();
-
-	TWeakPtr<ISequencer> Sequencer;
-	//static to be reused
-	static TOptional<FBakingAnimationKeySettings> CollapseControlsSettings;
-	//structonscope for details panel
-	TSharedPtr < TStructOnScope<FBakingAnimationKeySettings>> Settings;
-	TWeakPtr<SWindow> DialogWindow;
-	TSharedPtr<IStructureDetailsView> DetailsView;
-	FCollapseControlsCB CollapseCB;
-};
-
 /**
  * Tools for animation tracks
  */
diff --git a/Engine/Plugins/Animation/ControlRig/Source/ControlRigEditor/Private/Sequencer/ControlRigSpaceChannelEditors.cpp b/Engine/Plugins/Animation/ControlRig/Source/ControlRigEditor/Private/Sequencer/ControlRigSpaceChannelEditors.cpp
--- a/Engine/Plugins/Animation/ControlRig/Source/ControlRigEditor/Private/Sequencer/ControlRigSpaceChannelEditors.cpp
+++ b/Engine/Plugins/Animation/ControlRig/Source/ControlRigEditor/Private/Sequencer/ControlRigSpaceChannelEditors.cpp
@@ -207,7 +207,7 @@ FSpaceChannelAndSection FControlRigSpaceChannelHelpers::FindSpaceChannelAndSecti
 		UMovieSceneControlRigParameterTrack* ControlRigParameterTrack = Cast<UMovieSceneControlRigParameterTrack>(MovieScene->FindTrack(UMovieSceneControlRigParameterTrack::StaticClass(), Binding.GetObjectGuid(), NAME_None));
 		if (ControlRigParameterTrack && ControlRigParameterTrack->GetControlRig() == ControlRig)
 		{
-			UMovieSceneControlRigParameterSection* ActiveSection = Cast<UMovieSceneControlRigParameterSection>(ControlRigParameterTrack->GetSectionToKey(ControlName));
+			UMovieSceneControlRigParameterSection* ActiveSection = Cast<UMovieSceneControlRigParameterSection>(ControlRigParameterTrack->GetSectionToKey());
 			if (ActiveSection)
 			{
 				ActiveSection->Modify();
diff --git a/Engine/Plugins/Animation/ControlRig/Source/ControlRigEditor/Private/Tools/ControlRigSnapper.cpp b/Engine/Plugins/Animation/ControlRig/Source/ControlRigEditor/Private/Tools/ControlRigSnapper.cpp
--- a/Engine/Plugins/Animation/ControlRig/Source/ControlRigEditor/Private/Tools/ControlRigSnapper.cpp
+++ b/Engine/Plugins/Animation/ControlRig/Source/ControlRigEditor/Private/Tools/ControlRigSnapper.cpp
@@ -414,7 +414,7 @@ static void GetControlRigParents(const FControlRigForWorldTransforms& ControlRig
 	}
 }
 
-static UMovieSceneControlRigParameterSection* GetControlRigSection(ISequencer* Sequencer, const UControlRig* ControlRig, const FName& ControlName)
+static UMovieSceneControlRigParameterSection* GetControlRigSection(ISequencer* Sequencer, const UControlRig* ControlRig)
 {
 
 	if (ControlRig == nullptr || Sequencer == nullptr)
@@ -432,7 +432,7 @@ static UMovieSceneControlRigParameterSection* GetControlRigSection(ISequencer* S
 		UMovieSceneControlRigParameterTrack* ControlRigParameterTrack = Cast<UMovieSceneControlRigParameterTrack>(MovieScene->FindTrack(UMovieSceneControlRigParameterTrack::StaticClass(), Binding.GetObjectGuid(), NAME_None));
 		if (ControlRigParameterTrack && ControlRigParameterTrack->GetControlRig() == ControlRig)
 		{
-			UMovieSceneControlRigParameterSection* ActiveSection = Cast<UMovieSceneControlRigParameterSection>(ControlRigParameterTrack->GetSectionToKey(ControlName));
+			UMovieSceneControlRigParameterSection* ActiveSection = Cast<UMovieSceneControlRigParameterSection>(ControlRigParameterTrack->GetSectionToKey());
 			if (ActiveSection)
 			{
 				return ActiveSection;
@@ -493,10 +493,10 @@ static void GetTransformFrames(TSharedPtr<ISequencer>&  Sequencer, const FContro
 	}
 	for (FControlRigForWorldTransforms& ControlRig : ParentControlRigs)
 	{
-		for (FName& ControlName : ControlRig.ControlNames)
+		UMovieSceneControlRigParameterSection* Section = GetControlRigSection(Sequencer.Get(), ControlRig.ControlRig.Get());
+		if (Section)
 		{
-			UMovieSceneControlRigParameterSection* Section = GetControlRigSection(Sequencer.Get(), ControlRig.ControlRig.Get(), ControlName);
-			if (Section)
+			for (FName& ControlName : ControlRig.ControlNames)
 			{
 				TArrayView<FMovieSceneFloatChannel*> Channels = FControlRigSequencerHelpers::GetFloatChannels(ControlRig.ControlRig.Get(),
 					ControlName, Section);
@@ -510,7 +510,6 @@ static void GetTransformFrames(TSharedPtr<ISequencer>&  Sequencer, const FContro
 		}
 	}
 }
-
 bool FControlRigSnapper::SnapIt(FFrameNumber StartFrame, FFrameNumber EndFrame,const FControlRigSnapperSelection& ActorToSnap,
 	const FControlRigSnapperSelection& ParentToSnap, const UControlRigSnapSettings* SnapSettings)
 {
@@ -618,7 +617,6 @@ bool FControlRigSnapper::SnapIt(FFrameNumber StartFrame, FFrameNumber EndFrame,c
 		FRigControlModifiedContext Context;
 		Context.SetKey = EControlRigSetKey::Always;
 
-		TSet<UMovieSceneControlRigParameterSection*> ControlRigSections;
 		for (const FControlRigForWorldTransforms& ControlRigAndSelection : ActorToSnap.ControlRigs)
 		{
 			//get actor transform...
@@ -722,23 +720,19 @@ bool FControlRigSnapper::SnapIt(FFrameNumber StartFrame, FFrameNumber EndFrame,c
 							Context.LocalTime = TickResolution.AsSeconds(FFrameTime(FrameNumber));
 							FTransform GlobalTransform = WorldTransformToSnap[Index].GetRelativeTransform(ControlRigParentWorldTransforms[Index]);
 							ControlRig->SetControlGlobalTransform(Name, GlobalTransform, true, Context, false /*undo*/, false /*bPrintPython*/, true/* bFixEulerFlips*/);
-							UMovieSceneControlRigParameterSection* ControlRigSection = GetControlRigSection(Sequencer.Get(), ControlRig,Name);
-							ControlRigSections.Add(ControlRigSection);
 						}
 					}
 
-					if (SnapSettings && SnapSettings->BakingKeySettings != EBakingKeySettings::KeysOnly
+					UMovieSceneControlRigParameterSection* ControlRigSection = GetControlRigSection(Sequencer.Get(), ControlRig);
+					if (ControlRigSection && SnapSettings && SnapSettings->BakingKeySettings != EBakingKeySettings::KeysOnly
 						&& SnapSettings->bReduceKeys == true && Frames.Num() > 2)
 					{
-						for (UMovieSceneControlRigParameterSection* ControlRigSection : ControlRigSections)
-						{
-							FKeyDataOptimizationParams Param;
-							Param.bAutoSetInterpolation = true;
-							Param.Tolerance = SnapSettings->Tolerance;
-							TRange<FFrameNumber> Range(Frames[0], Frames[Frames.Num() - 1]);
-							Param.Range = Range;
-							MovieSceneToolHelpers::OptimizeSection(Param, ControlRigSection);
-						}
+						FKeyDataOptimizationParams Param;
+						Param.bAutoSetInterpolation = true;
+						Param.Tolerance = SnapSettings->Tolerance;
+						TRange<FFrameNumber> Range(Frames[0], Frames[Frames.Num() - 1]);
+						Param.Range = Range;
+						MovieSceneToolHelpers::OptimizeSection(Param, ControlRigSection);
 					}
 				}
 			}
diff --git a/Engine/Plugins/Animation/ControlRig/Source/ControlRigEditor/Public/ControlRigSequencerEditorLibrary.h b/Engine/Plugins/Animation/ControlRig/Source/ControlRigEditor/Public/ControlRigSequencerEditorLibrary.h
--- a/Engine/Plugins/Animation/ControlRig/Source/ControlRigEditor/Public/ControlRigSequencerEditorLibrary.h
+++ b/Engine/Plugins/Animation/ControlRig/Source/ControlRigEditor/Public/ControlRigSequencerEditorLibrary.h
@@ -26,7 +26,6 @@ class UTickableTransformConstraint;
 class UTransformableHandle;
 struct FBakingAnimationKeySettings;
 class UControlRig;
-class UAnimLayer;
 
 USTRUCT(BlueprintType)
 struct FControlRigSequencerBindingProxy
@@ -1126,50 +1125,4 @@ public:
 	UFUNCTION(BlueprintCallable, Category = "Editor Scripting | Sequencer Tools | Control Rig")
 	static bool SetControlRigApplyMode(UControlRig* InControlRig, EControlRigFKRigExecuteMode InApplyMode);
 
-	/**
-	* Delete anim layer at specified index
-	* @param Index The index where the anim layer exists
-	* @return Returns true if successful, false otherwise
-	*/
-	UFUNCTION(BlueprintCallable, Category = "Editor Scripting | Sequencer Tools | Control Rig | Animation Layers")
-	bool DeleteAnimLayer(int32 Index);
-
-	/**
-	* Duplicate anim layer at specified index
-	* @param Index The index where the anim layer exists
-	* @return Returns index of new layer, -1 if none created
-	*/
-	UFUNCTION(BlueprintCallable, Category = "Editor Scripting | Sequencer Tools | Control Rig | Animation Layers")
-	int32 DuplicateAnimLayer(int32 Index);
-
-	/**
-	* Add anim layer from objects selected in Sequencer
-	* @return Returns Index of created anim layer
-	*/
-	UFUNCTION(BlueprintCallable, Category = "Editor Scripting | Sequencer Tools | Control Rig | Animation Layers")
-	int32 AddAnimLayerFromSelection();
-
-	/**
-	* Merge specified anim layers into one layer. Will merge onto the anim layer with the lowest index
-	* @param Indices The indices to merge
-	* @return Returns true if successful, false otherwise
-	*/
-	UFUNCTION(BlueprintCallable, Category = "Editor Scripting | Sequencer Tools | Control Rig | Animation Layers")
-	bool MergeAnimLayers(const TArray<int32>& Indices);
-
-	/**
-	* Get the animation layer objects
-	* @return Returns array of anim layer objects if they exist on active Sequencer
-	*/
-	UFUNCTION(BlueprintCallable, Category = "Editor Scripting | Sequencer Tools | Control Rig | Animation Layers")
-	TArray<UAnimLayer*> GetAnimLayers();
-	
-	/**
-	* Helper function to get the index in the anim layer array from the anim layer
-	* @param Anim Layer to get the index for
-	* @return Returns index for the anim layer or INDEX_NONE(-1) if it doesn't exist
-	*/
-	UFUNCTION(BlueprintCallable, Category = "Editor Scripting | Sequencer Tools | Control Rig | Animation Layers")
-	int32 GetAnimLayerIndex(UAnimLayer* AnimLayer);
-
 };
diff --git a/Engine/Source/Editor/MovieSceneTools/Private/MovieSceneToolHelpers.cpp b/Engine/Source/Editor/MovieSceneTools/Private/MovieSceneToolHelpers.cpp
--- a/Engine/Source/Editor/MovieSceneTools/Private/MovieSceneToolHelpers.cpp
+++ b/Engine/Source/Editor/MovieSceneTools/Private/MovieSceneToolHelpers.cpp
@@ -83,8 +83,6 @@
 #include "Channels/MovieSceneChannelTraits.h"
 #include "Channels/MovieSceneIntegerChannel.h"
 #include "Channels/MovieSceneByteChannel.h"
-#include "Channels/MovieSceneFloatChannel.h"
-#include "Channels/MovieSceneDoubleChannel.h"
 #include "Framework/MultiBox/MultiBoxBuilder.h"
 #include "EntitySystem/Interrogation/MovieSceneInterrogationLinker.h"
 #include "ConstraintsManager.h"
@@ -5041,6 +5039,105 @@ void MovieSceneToolHelpers::CalculateFramesBetween(
 	}
 }
 
+UENUM(BlueprintType)
+enum class FChannelMergeAlgorithm : uint8
+{
+	/**Average values together*/
+	Average,
+	/**Add values together*/
+	Add
+};
+
+//merge the channels onto the first one using the passed in algorithm, we may skip the last channel also (since it may be weight).
+template<typename ChannelType>
+static bool MergeChannels(TArray<ChannelType*>& Channels,const  TArray<UMovieSceneSection*>& Sections, const TRange<FFrameNumber>& Range,
+	FChannelMergeAlgorithm MergeAlgorithm)
+{
+	if (Channels.Num() < 2 && Channels.Num() != Sections.Num())
+	{
+		return false;
+	}
+	using ChannelValueType = typename ChannelType::ChannelValueType;
+	using CurveValueType = typename ChannelType::CurveValueType;
+
+	//base channel we set values on
+	ChannelType* BaseChannel = Channels[0];
+	TMovieSceneChannelData<ChannelValueType> BaseChannelData = BaseChannel->GetData();
+	//iterate over each key
+	TArray<FFrameNumber> KeyTimes;
+	TArray<FKeyHandle> Handles;
+	//cached set that we set at the end
+	TArray<TPair< FFrameNumber, ChannelValueType>> KeysToSet;
+	for (int32 ChannelIndex = 0; ChannelIndex < Channels.Num(); ++ChannelIndex)
+	{
+		KeyTimes.Reset();
+		Handles.Reset();
+		ChannelType* Channel = Channels[ChannelIndex];
+		Channel->GetKeys(Range, &KeyTimes, &Handles);
+		for (int32 FrameIndex = 0; FrameIndex < KeyTimes.Num(); ++FrameIndex)
+		{
+			const FFrameNumber& Frame = KeyTimes[FrameIndex];
+			if (Sections[0]->GetRange().Contains(Frame) == false)  //frame is outside base range so skip
+			{
+				continue;
+			}
+			const FFrameTime FrameTime(Frame);
+			int32 KeyIndex = Channel->GetData().GetIndex(Handles[FrameIndex]);
+			ChannelValueType Value = Channel->GetData().GetValues()[KeyIndex];
+			//got value with tangents and times, now we perform the operation
+			Value.Value = 0.0; //zero out the value we calculate it 
+			if (MergeAlgorithm == FChannelMergeAlgorithm::Average)
+			{
+				double DNumChannels = 0.0;
+				for (int32 WeightIndex = 0; WeightIndex < Sections.Num(); ++WeightIndex)
+				{
+					if (Sections[WeightIndex]->GetRange().Contains(Frame))  
+					{
+						float Weight = Sections[WeightIndex]->GetTotalWeightValue(FrameTime);
+						CurveValueType WeightedValue = 0.0;
+						ChannelType* EachChannel = Channels[WeightIndex];
+						EachChannel->Evaluate(FrameTime, WeightedValue);
+						WeightedValue *= ((double)Weight);
+						Value.Value += WeightedValue;
+						DNumChannels += 1.0;
+					}
+				}
+				if (DNumChannels > 0.0) //should always happen since base(0) at least be here
+				{
+					Value.Value /= DNumChannels;
+				}
+			}
+			else if (MergeAlgorithm == FChannelMergeAlgorithm::Add)
+			{
+				for (int32 WeightIndex = 0; WeightIndex < Sections.Num(); ++WeightIndex)
+				{
+					if (Sections[WeightIndex]->GetRange().Contains(Frame))  
+					{
+						float Weight = Sections[WeightIndex]->GetTotalWeightValue(FrameTime);
+						if (Sections[WeightIndex]->GetBlendType().IsValid() == false ||
+							Sections[WeightIndex]->GetBlendType().Get() != EMovieSceneBlendType::Additive)
+						{
+							Weight = 1.0;
+						}
+						CurveValueType WeightedValue = 0.0;
+						ChannelType* EachChannel = Channels[WeightIndex];
+						EachChannel->Evaluate(FrameTime, WeightedValue);
+						WeightedValue *= Weight;
+						Value.Value += WeightedValue;
+					}
+				}
+			}
+			KeysToSet.Add(TPair<FFrameNumber, ChannelValueType>(Frame, Value));
+		}
+	}
+	for (TPair<FFrameNumber, ChannelValueType>& KeyToSet : KeysToSet)
+	{
+		MovieSceneToolHelpers::SetOrAddKey(BaseChannelData, KeyToSet.Key, KeyToSet.Value);
+	}
+
+	return true;
+}
+
 bool MovieSceneToolHelpers::OptimizeSection(const FKeyDataOptimizationParams& InParams, UMovieSceneSection* InSection)
 {
 	TArrayView<FMovieSceneFloatChannel*> FloatChannels = InSection->GetChannelProxy().GetChannels<FMovieSceneFloatChannel>();
@@ -5056,6 +5153,87 @@ bool MovieSceneToolHelpers::OptimizeSection(const FKeyDataOptimizationParams& In
 	return true;
 }
 
+template<typename ChannelType>
+static bool MergeSections(UMovieSceneSection* BaseSection, TArrayView<ChannelType*> BaseDoubleChannels, TArray<UMovieSceneSection*>& AbsoluteSections, TArray<UMovieSceneSection*>& AdditiveSections,
+	const TRange<FFrameNumber>& Range,bool bSkipLastChannel)
+{
+	TArrayView<ChannelType*> BaseChannels = BaseSection->GetChannelProxy().GetChannels<ChannelType>();
+	if (BaseChannels.Num() > 0)
+	{
+		int32 SectionIndex = 0;
+		//sanity check to make sure channels are the same size
+		for (SectionIndex = 0; SectionIndex < AbsoluteSections.Num(); ++SectionIndex)
+		{
+			TArrayView<ChannelType*>Channels = AbsoluteSections[SectionIndex]->GetChannelProxy().GetChannels<ChannelType>();
+			if (Channels.Num() != BaseChannels.Num())
+			{
+				UE_LOG(LogMovieScene, Warning, TEXT("MergeSections:: Invalid number of channels"));
+				return false;
+			}
+		}
+		for (SectionIndex = 0; SectionIndex < AdditiveSections.Num(); ++SectionIndex)
+		{
+			TArrayView<ChannelType*>Channels = AdditiveSections[SectionIndex]->GetChannelProxy().GetChannels<ChannelType>();
+			if (Channels.Num() != BaseChannels.Num())
+			{
+				UE_LOG(LogMovieScene, Warning, TEXT("MergeSections:: Invalid number of channels"));
+				return false;
+			}
+		}
+		BaseSection->Modify();
+
+		TArray<ChannelType*> Channels;
+		int32 ChannelIndex = 0;
+		bSkipLastChannel = false;
+		int32 BaseChannelsNum = bSkipLastChannel ? BaseChannels.Num() - 1 : BaseChannels.Num();
+		for (ChannelIndex = 0; ChannelIndex < BaseChannels.Num(); ++ChannelIndex)
+		{
+			if (AbsoluteSections.Num() > 0)
+			{
+				Channels.Reset();
+				for (SectionIndex = 0; SectionIndex < AbsoluteSections.Num(); ++SectionIndex)
+				{
+					TArrayView<ChannelType*>OurChannels = AbsoluteSections[SectionIndex]->GetChannelProxy().GetChannels<ChannelType>();
+					Channels.Add(OurChannels[ChannelIndex]);
+				}
+				//now blend them
+				if (MergeChannels(Channels, AbsoluteSections, Range,
+					FChannelMergeAlgorithm::Average) == false)
+				{
+					UE_LOG(LogMovieScene, Warning, TEXT("MergeSections:: Could not merge channels"));
+					return false;
+				}
+			}
+			if (AdditiveSections.Num() > 0)
+			{
+				//now do additives
+				Channels.Reset();
+				for (SectionIndex = 0; SectionIndex < AdditiveSections.Num(); ++SectionIndex)
+				{
+					TArrayView<ChannelType*>OurChannels = AdditiveSections[SectionIndex]->GetChannelProxy().GetChannels<ChannelType>();
+					Channels.Add(OurChannels[ChannelIndex]);
+				}
+				//now blend them
+				if (MergeChannels(Channels, AdditiveSections, Range,
+					FChannelMergeAlgorithm::Add) == false)
+				{
+					UE_LOG(LogMovieScene, Warning, TEXT("MergeSections:: Could not merge channels"));
+					return false;
+				}
+			}
+		}
+		for (ChannelType* Channel : BaseChannels)
+		{
+			Channel->AutoSetTangents();
+		}
+	}
+	else
+	{
+		return false;
+	}
+	return true;
+}
+
 bool MovieSceneToolHelpers::CollapseSection(TSharedPtr<ISequencer>& SequencerPtr, UMovieSceneTrack* OwnerTrack, TArray<UMovieSceneSection*> Sections,
 	const FBakingAnimationKeySettings& InSettings)
 {
@@ -5131,7 +5309,7 @@ bool MovieSceneToolHelpers::CollapseSection(TSharedPtr<ISequencer>& SequencerPtr
 		TArrayView<FMovieSceneDoubleChannel*> BaseDoubleChannels = BaseSection->GetChannelProxy().GetChannels<FMovieSceneDoubleChannel>();
 		if (BaseDoubleChannels.Num() > 0) //transforms
 		{
-			if (MovieSceneToolHelpers::MergeSections<FMovieSceneDoubleChannel>(BaseSection, AbsoluteSections, AdditiveSections, Range, false /*bSkipLastChannel*/) == false)
+			if (MergeSections(BaseSection, BaseDoubleChannels, AbsoluteSections, AdditiveSections, Range, false /*bSkipLastChannel*/) == false)
 			{
 				Transaction.Cancel();
 				return false;
@@ -5140,7 +5318,7 @@ bool MovieSceneToolHelpers::CollapseSection(TSharedPtr<ISequencer>& SequencerPtr
 		else if(BaseFloatChannels.Num() > 0) //control rig
 		{
 			//skip weight channel for control rig
-			if (MovieSceneToolHelpers::MergeSections<FMovieSceneFloatChannel>(BaseSection,  AbsoluteSections, AdditiveSections, Range, true /*bSkipLastChannel*/) == false)
+			if (MergeSections(BaseSection, BaseFloatChannels, AbsoluteSections, AdditiveSections, Range, true /*bSkipLastChannel*/) == false)
 			{
 				Transaction.Cancel();
 				return false;
diff --git a/Engine/Source/Editor/MovieSceneTools/Public/MovieSceneToolHelpers.h b/Engine/Source/Editor/MovieSceneTools/Public/MovieSceneToolHelpers.h
--- a/Engine/Source/Editor/MovieSceneTools/Public/MovieSceneToolHelpers.h
+++ b/Engine/Source/Editor/MovieSceneTools/Public/MovieSceneToolHelpers.h
@@ -19,11 +19,10 @@
 #include "SEnumCombo.h"
 #include "Animation/AnimSequence.h"
 #include "INodeAndChannelMappings.h"
-#include "Channels/MovieSceneChannelProxy.h"
-#include "MovieSceneSection.h"
 
 class ISequencer;
 class UMovieScene;
+class UMovieSceneSection;
 class UMovieSceneSequence;
 struct FMovieSceneObjectBindingID;
 class UMovieSceneTrack;
@@ -106,14 +105,6 @@ public:
 	bool bCachedEnableUpdateRateOptimizations;
 };
 
-enum class FChannelMergeAlgorithm : uint8
-{
-	/**Average values together*/
-	Average,
-	/**Add values together*/
-	Add
-};
-
 class MOVIESCENETOOLS_API MovieSceneToolHelpers
 {
 public:
@@ -671,267 +662,6 @@ public:
 	static bool CollapseSection(TSharedPtr<ISequencer>& SequencerPtr, UMovieSceneTrack* OwnerTrack, TArray<UMovieSceneSection*> Sections,
 		const FBakingAnimationKeySettings& InSettings);
 
-	//merge the channels onto the first one using the passed in algorithm, we may skip the last channel also (since it may be weight).
-	template<typename ChannelType>
-	static bool MergeChannels(TArray<ChannelType*>& Channels, const  TArray<UMovieSceneSection*>& Sections, const TRange<FFrameNumber>& Range,
-		FChannelMergeAlgorithm MergeAlgorithm)
-	{
-		if (Channels.Num() < 2 && Channels.Num() != Sections.Num())
-		{
-			return false;
-		}
-		using ChannelValueType = typename ChannelType::ChannelValueType;
-		using CurveValueType = typename ChannelType::CurveValueType;
-
-		//base channel we set values on
-		ChannelType* BaseChannel = Channels[0];
-		TMovieSceneChannelData<ChannelValueType> BaseChannelData = BaseChannel->GetData();
-		//iterate over each key
-		TArray<FFrameNumber> KeyTimes;
-		TArray<FKeyHandle> Handles;
-		//cached set that we set at the end
-		TArray<TPair< FFrameNumber, ChannelValueType>> KeysToSet;
-		for (int32 ChannelIndex = 0; ChannelIndex < Channels.Num(); ++ChannelIndex)
-		{
-			KeyTimes.Reset();
-			Handles.Reset();
-			ChannelType* Channel = Channels[ChannelIndex];
-			Channel->GetKeys(Range, &KeyTimes, &Handles);
-			for (int32 FrameIndex = 0; FrameIndex < KeyTimes.Num(); ++FrameIndex)
-			{
-				const FFrameNumber& Frame = KeyTimes[FrameIndex];
-				if (Sections[0]->GetRange().Contains(Frame) == false)  //frame is outside base range so skip
-				{
-					continue;
-				}
-				const FFrameTime FrameTime(Frame);
-				int32 KeyIndex = Channel->GetData().GetIndex(Handles[FrameIndex]);
-				ChannelValueType Value = Channel->GetData().GetValues()[KeyIndex];
-				//got value with tangents and times, now we perform the operation
-				Value.Value = 0.0; //zero out the value we calculate it 
-				if (MergeAlgorithm == FChannelMergeAlgorithm::Average)
-				{
-					double DNumChannels = 0.0;
-					for (int32 WeightIndex = 0; WeightIndex < Sections.Num(); ++WeightIndex)
-					{
-						if (Sections[WeightIndex]->GetRange().Contains(Frame))
-						{
-							float Weight = Sections[WeightIndex]->GetTotalWeightValue(FrameTime);
-							CurveValueType WeightedValue = 0.0;
-							ChannelType* EachChannel = Channels[WeightIndex];
-							EachChannel->Evaluate(FrameTime, WeightedValue);
-							WeightedValue *= ((double)Weight);
-							Value.Value += WeightedValue;
-							DNumChannels += 1.0;
-						}
-					}
-					if (DNumChannels > 0.0) //should always happen since base(0) at least be here
-					{
-						Value.Value /= DNumChannels;
-					}
-				}
-				else if (MergeAlgorithm == FChannelMergeAlgorithm::Add)
-				{
-					for (int32 WeightIndex = 0; WeightIndex < Sections.Num(); ++WeightIndex)
-					{
-						if (Sections[WeightIndex]->GetRange().Contains(Frame))
-						{
-							float Weight = Sections[WeightIndex]->GetTotalWeightValue(FrameTime);
-							if (Sections[WeightIndex]->GetBlendType().IsValid() == false ||
-								Sections[WeightIndex]->GetBlendType().Get() != EMovieSceneBlendType::Additive)
-							{
-								Weight = 1.0;
-							}
-							CurveValueType WeightedValue = 0.0;
-							ChannelType* EachChannel = Channels[WeightIndex];
-							EachChannel->Evaluate(FrameTime, WeightedValue);
-							WeightedValue *= Weight;
-							Value.Value += WeightedValue;
-						}
-					}
-				}
-				KeysToSet.Add(TPair<FFrameNumber, ChannelValueType>(Frame, Value));
-			}
-		}
-		for (TPair<FFrameNumber, ChannelValueType>& KeyToSet : KeysToSet)
-		{
-			MovieSceneToolHelpers::SetOrAddKey(BaseChannelData, KeyToSet.Key, KeyToSet.Value);
-		}
-
-		return true;
-	}
-
-	template<typename ChannelType>
-	static bool MergeNonWeightedChannels(TArray<ChannelType*>& Channels, const  TArray<UMovieSceneSection*>& Sections, const TRange<FFrameNumber>& Range,
-		FChannelMergeAlgorithm MergeAlgorithm)
-	{
-		if (Channels.Num() < 2 && Channels.Num() != Sections.Num())
-		{
-			return false;
-		}
-		using ChannelValueType = typename ChannelType::ChannelValueType;
-		using CurveValueType = typename ChannelType::CurveValueType;
-
-		//base channel we set values on
-		ChannelType* BaseChannel = Channels[0];
-		TMovieSceneChannelData<ChannelValueType> BaseChannelData = BaseChannel->GetData();
-		//iterate over each key
-		TArray<FFrameNumber> KeyTimes;
-		TArray<FKeyHandle> Handles;
-		//cached set that we set at the end
-		TArray<TPair< FFrameNumber, ChannelValueType>> KeysToSet;
-		for (int32 ChannelIndex = 0; ChannelIndex < Channels.Num(); ++ChannelIndex)
-		{
-			KeyTimes.Reset();
-			Handles.Reset();
-			ChannelType* Channel = Channels[ChannelIndex];
-			Channel->GetKeys(Range, &KeyTimes, &Handles);
-			for (int32 FrameIndex = 0; FrameIndex < KeyTimes.Num(); ++FrameIndex)
-			{
-				const FFrameNumber& Frame = KeyTimes[FrameIndex];
-				if (Sections[0]->GetRange().Contains(Frame) == false)  //frame is outside base range so skip
-				{
-					continue;
-				}
-				const FFrameTime FrameTime(Frame);
-				int32 KeyIndex = Channel->GetData().GetIndex(Handles[FrameIndex]);
-				ChannelValueType Value = Channel->GetData().GetValues()[KeyIndex];
-				//got value with tangents and times, now we perform the operation
-				Value.Value = 0.0; //zero out the value we calculate it 
-				if (MergeAlgorithm == FChannelMergeAlgorithm::Average)
-				{
-					double DNumChannels = 0.0;
-					for (int32 WeightIndex = 0; WeightIndex < Sections.Num(); ++WeightIndex)
-					{
-						if (Sections[WeightIndex]->GetRange().Contains(Frame))
-						{
-							float Weight = Sections[WeightIndex]->GetTotalWeightValue(FrameTime);
-							CurveValueType WeightedValue = 0.0;
-							ChannelType* EachChannel = Channels[WeightIndex];
-							EachChannel->Evaluate(FrameTime, WeightedValue);
-							WeightedValue *= ((double)Weight);
-							Value.Value += WeightedValue;
-							DNumChannels += 1.0;
-						}
-					}
-					if (DNumChannels > 0.0) //should always happen since base(0) at least be here
-					{
-						Value.Value /= DNumChannels;
-					}
-				}
-				else if (MergeAlgorithm == FChannelMergeAlgorithm::Add)
-				{
-					for (int32 WeightIndex = 0; WeightIndex < Sections.Num(); ++WeightIndex)
-					{
-						if (Sections[WeightIndex]->GetRange().Contains(Frame))
-						{
-							float Weight = Sections[WeightIndex]->GetTotalWeightValue(FrameTime);
-							if (Sections[WeightIndex]->GetBlendType().IsValid() == false ||
-								Sections[WeightIndex]->GetBlendType().Get() != EMovieSceneBlendType::Additive)
-							{
-								Weight = 1.0;
-							}
-							CurveValueType WeightedValue = 0.0;
-							ChannelType* EachChannel = Channels[WeightIndex];
-							EachChannel->Evaluate(FrameTime, WeightedValue);
-							WeightedValue *= Weight;
-							Value.Value += WeightedValue;
-						}
-					}
-				}
-				KeysToSet.Add(TPair<FFrameNumber, ChannelValueType>(Frame, Value));
-			}
-		}
-		for (TPair<FFrameNumber, ChannelValueType>& KeyToSet : KeysToSet)
-		{
-			MovieSceneToolHelpers::SetOrAddKey(BaseChannelData, KeyToSet.Key, KeyToSet.Value);
-		}
-
-		return true;
-	}
-
-	template<typename ChannelType>
-	static bool MergeSections(UMovieSceneSection* BaseSection, TArray<UMovieSceneSection*>& AbsoluteSections, TArray<UMovieSceneSection*>& AdditiveSections,
-		const TRange<FFrameNumber>& Range, bool bSkipLastChannel)
-	{
-		TArrayView<ChannelType*> BaseChannels = BaseSection->GetChannelProxy().GetChannels<ChannelType>();
-		if (BaseChannels.Num() > 0)
-		{
-			int32 SectionIndex = 0;
-			//sanity check to make sure channels are the same size
-			for (SectionIndex = 0; SectionIndex < AbsoluteSections.Num(); ++SectionIndex)
-			{
-				TArrayView<ChannelType*>Channels = AbsoluteSections[SectionIndex]->GetChannelProxy().GetChannels<ChannelType>();
-				if (Channels.Num() != BaseChannels.Num())
-				{
-					UE_LOG(LogMovieScene, Warning, TEXT("MergeSections:: Invalid number of channels"));
-					return false;
-				}
-			}
-			for (SectionIndex = 0; SectionIndex < AdditiveSections.Num(); ++SectionIndex)
-			{
-				TArrayView<ChannelType*>Channels = AdditiveSections[SectionIndex]->GetChannelProxy().GetChannels<ChannelType>();
-				if (Channels.Num() != BaseChannels.Num())
-				{
-					UE_LOG(LogMovieScene, Warning, TEXT("MergeSections:: Invalid number of channels"));
-					return false;
-				}
-			}
-			BaseSection->Modify();
-
-			TArray<ChannelType*> Channels;
-			int32 ChannelIndex = 0;
-			bSkipLastChannel = false;
-			int32 BaseChannelsNum = bSkipLastChannel ? BaseChannels.Num() - 1 : BaseChannels.Num();
-			for (ChannelIndex = 0; ChannelIndex < BaseChannels.Num(); ++ChannelIndex)
-			{
-				if (AbsoluteSections.Num() > 0)
-				{
-					Channels.Reset();
-					for (SectionIndex = 0; SectionIndex < AbsoluteSections.Num(); ++SectionIndex)
-					{
-						TArrayView<ChannelType*>OurChannels = AbsoluteSections[SectionIndex]->GetChannelProxy().GetChannels<ChannelType>();
-						Channels.Add(OurChannels[ChannelIndex]);
-					}
-					//now blend them
-					if (MergeChannels(Channels, AbsoluteSections, Range,
-						FChannelMergeAlgorithm::Average) == false)
-					{
-						UE_LOG(LogMovieScene, Warning, TEXT("MergeSections:: Could not merge channels"));
-						return false;
-					}
-				}
-				if (AdditiveSections.Num() > 0)
-				{
-					//now do additives
-					Channels.Reset();
-					for (SectionIndex = 0; SectionIndex < AdditiveSections.Num(); ++SectionIndex)
-					{
-						TArrayView<ChannelType*>OurChannels = AdditiveSections[SectionIndex]->GetChannelProxy().GetChannels<ChannelType>();
-						Channels.Add(OurChannels[ChannelIndex]);
-					}
-					//now blend them
-					if (MergeChannels(Channels, AdditiveSections, Range,
-						FChannelMergeAlgorithm::Add) == false)
-					{
-						UE_LOG(LogMovieScene, Warning, TEXT("MergeSections:: Could not merge channels"));
-						return false;
-					}
-				}
-			}
-			for (ChannelType* Channel : BaseChannels)
-			{
-				Channel->AutoSetTangents();
-			}
-		}
-		else
-		{
-			return false;
-		}
-		return true;
-	}
-
-
 	static bool OptimizeSection(const FKeyDataOptimizationParams& InParams, UMovieSceneSection* InSection);
 
 	/** Returns the frame numbers between start and end. */
diff --git a/Engine/Source/Editor/Sequencer/Private/SSequencerSection.cpp b/Engine/Source/Editor/Sequencer/Private/SSequencerSection.cpp
--- a/Engine/Source/Editor/Sequencer/Private/SSequencerSection.cpp
+++ b/Engine/Source/Editor/Sequencer/Private/SSequencerSection.cpp
@@ -41,7 +41,6 @@
 #include "IKeyArea.h"
 #include "Widgets/SWeakWidget.h"
 #include "Algo/Transform.h"
-#include "Tracks/IMovieSceneSectionsToKey.h"
 
 namespace UE
 {
@@ -74,15 +73,6 @@ FTimeToPixel ConstructTimeConverterForSection(const FGeometry& InSectionGeometry
 	return FTimeToPixel(InSectionGeometry, TRange<double>(LowerTime, UpperTime), TickResolution);
 }
 
-static bool IsSectionToKey(const UMovieSceneTrack* Track, UMovieSceneSection* SectionObject)
-{
-	if (const IMovieSceneSectionsToKey* MultipleSectionsToKey = Cast<IMovieSceneSectionsToKey>(Track))
-	{
-		TArray<TWeakObjectPtr<UMovieSceneSection>> SectionsToKey = MultipleSectionsToKey->GetSectionsToKey();
-		return SectionsToKey.Contains(SectionObject);
-	}
-	return (Track->GetSectionToKey() == SectionObject);
-}
 
 struct FSequencerSectionPainterImpl : FSequencerSectionPainter
 {
@@ -798,7 +788,7 @@ struct FSequencerSectionPainterImpl : FSequencerSectionPainter
 		const bool bLocked = SectionObject->IsLocked() || SectionObject->IsReadOnly();
 
 		// Only show section to key border if we have more than one section
-		const bool bIsSectionToKey = Track && Track->GetAllSections().Num() > 1 && IsSectionToKey(Track,SectionObject);
+		const bool bIsSectionToKey = Track && Track->GetAllSections().Num() > 1 && Track->GetSectionToKey() == SectionObject;
 
 		const ESlateDrawEffect DrawEffects = bParentEnabled ? ESlateDrawEffect::None : ESlateDrawEffect::DisabledEffect;
 
@@ -1685,7 +1675,7 @@ void SSequencerSection::PaintEasingHandles( FSequencerSectionPainter& InPainter,
 		const float MinHandleSize = 8.f;
 
 		const ESlateDrawEffect DrawEffects = InPainter.bParentEnabled ? ESlateDrawEffect::None : ESlateDrawEffect::DisabledEffect;
-		const bool bIsSectionToKey = Track->GetAllSections().Num() > 1 && IsSectionToKey(Track, UnderlappingSectionObj);
+		const bool bIsSectionToKey = Track->GetAllSections().Num() > 1 && Track->GetSectionToKey() == UnderlappingSectionObj;
 
 		// If this is the section to key, we draw the easing handle in the same bright green color as the border
 		// outline. We also make the handle a bit bigger, otherwise that border being drawn on top makes it look
diff --git a/Engine/Source/Runtime/MovieScene/Private/Evaluation/Blending/MovieSceneBlendType.cpp b/Engine/Source/Runtime/MovieScene/Private/Evaluation/Blending/MovieSceneBlendType.cpp
--- a/Engine/Source/Runtime/MovieScene/Private/Evaluation/Blending/MovieSceneBlendType.cpp
+++ b/Engine/Source/Runtime/MovieScene/Private/Evaluation/Blending/MovieSceneBlendType.cpp
@@ -22,7 +22,7 @@ FMovieSceneBlendTypeField FMovieSceneBlendTypeField::All()
 {
 	FMovieSceneBlendTypeField New;
 	New.Add(EMovieSceneBlendType::Absolute, EMovieSceneBlendType::Additive, EMovieSceneBlendType::Relative,
-			EMovieSceneBlendType::AdditiveFromBase, EMovieSceneBlendType::Override);
+			EMovieSceneBlendType::AdditiveFromBase);
 	return New;
 }
 
@@ -73,11 +73,9 @@ int32 FMovieSceneBlendTypeField::Num() const
 		(Contains(EMovieSceneBlendType::Absolute) ? 1 : 0) +
 		(Contains(EMovieSceneBlendType::Relative) ? 1 : 0) +
 		(Contains(EMovieSceneBlendType::Additive) ? 1 : 0) +
-		(Contains(EMovieSceneBlendType::AdditiveFromBase) ? 1 : 0) +
-		(Contains(EMovieSceneBlendType::Override) ? 1 : 0);
+		(Contains(EMovieSceneBlendType::AdditiveFromBase) ? 1 : 0);
 }
 
-
 void FMovieSceneBlendTypeFieldIterator::IterateToNext()
 {
 	do
diff --git a/Engine/Source/Runtime/MovieScene/Public/Evaluation/Blending/MovieSceneBlendType.h b/Engine/Source/Runtime/MovieScene/Public/Evaluation/Blending/MovieSceneBlendType.h
--- a/Engine/Source/Runtime/MovieScene/Public/Evaluation/Blending/MovieSceneBlendType.h
+++ b/Engine/Source/Runtime/MovieScene/Public/Evaluation/Blending/MovieSceneBlendType.h
@@ -24,8 +24,6 @@ enum class EMovieSceneBlendType : uint8
 	Relative			= 0x4,
 	/** Applies this value as an additive equal to the difference between the current value and the first value */
 	AdditiveFromBase    = 0x8,
-	/** The value will override the current value */
-	Override			= 0x10,
 
 	// If you add new blend types, don't forget to change:
 	// - All()
diff --git a/Engine/Source/Runtime/MovieSceneTracks/Private/Tracks/IMovieSceneSectionsToKey.cpp b/Engine/Source/Runtime/MovieSceneTracks/Private/Tracks/IMovieSceneSectionsToKey.cpp
deleted file mode 100644
diff --git a/Engine/Source/Runtime/MovieSceneTracks/Public/Sections/MovieScene3DTransformSection.h b/Engine/Source/Runtime/MovieSceneTracks/Public/Sections/MovieScene3DTransformSection.h
--- a/Engine/Source/Runtime/MovieSceneTracks/Public/Sections/MovieScene3DTransformSection.h
+++ b/Engine/Source/Runtime/MovieSceneTracks/Public/Sections/MovieScene3DTransformSection.h
@@ -231,8 +231,8 @@ public:
 #endif
 
 	/* From UMovieSection*/
-
-	MOVIESCENETRACKS_API virtual bool ShowCurveForChannel(const void* Channel) const override;
+	
+	MOVIESCENETRACKS_API virtual bool ShowCurveForChannel(const void *Channel) const override;
 	MOVIESCENETRACKS_API virtual void SetBlendType(EMovieSceneBlendType InBlendType) override;
 	MOVIESCENETRACKS_API virtual void OnBindingIDsUpdated(const TMap<UE::MovieScene::FFixedObjectBindingID, UE::MovieScene::FFixedObjectBindingID>& OldFixedToNewFixedMap, FMovieSceneSequenceID LocalSequenceID, TSharedRef<UE::MovieScene::FSharedPlaybackState> SharedPlaybackState) override;
 	MOVIESCENETRACKS_API virtual void GetReferencedBindings(TArray<FGuid>& OutBindings) override;
@@ -269,10 +269,6 @@ public:
 	*/
 	MOVIESCENETRACKS_API void SetUseQuaternionInterpolation(bool bInUseQuaternionInterpolation);
 
-	/**
-	* Get Weight Channel
-	*/
-	MOVIESCENETRACKS_API FMovieSceneFloatChannel* GetWeightChannel() { return &ManualWeight; }
 protected:
 
 	MOVIESCENETRACKS_API virtual TSharedPtr<FStructOnScope> GetKeyStruct(TArrayView<const FKeyHandle> KeyHandles) override;
diff --git a/Engine/Source/Runtime/MovieSceneTracks/Public/Tracks/IMovieSceneSectionsToKey.h b/Engine/Source/Runtime/MovieSceneTracks/Public/Tracks/IMovieSceneSectionsToKey.h
deleted file mode 100644
