Adding support for transform origin overrides on shot and subsequence sections.
- Sections now have channels for loctaion and rotation keyframes, which can be used to move all objects with transform sections within that subscene.
- The transform override can be modified in the level viewport  with a gizmo tool, and editing behavior should reflect the behavior of spawnable transform sections.
- Extended the existing Transform Origin System to account for subscene origins, which are blended additively with their parent origins.

diff --git a/Engine/Source/Editor/MovieSceneTools/Private/EditModes/SubTrackEditorMode.cpp b/Engine/Source/Editor/MovieSceneTools/Private/EditModes/SubTrackEditorMode.cpp
new file mode 100644
--- /dev/null
+++ b/Engine/Source/Editor/MovieSceneTools/Private/EditModes/SubTrackEditorMode.cpp
@@ -0,0 +1,272 @@
+﻿// Copyright Epic Games, Inc. All Rights Reserved.
+
+#include "SubTrackEditorMode.h"
+
+#include "EngineUtils.h"
+#include "UnrealClient.h"
+#include "EditorViewportClient.h"
+#include "MovieScene.h"
+#include "SequencerChannelTraits.h"
+#include "Compilation/MovieSceneCompiledDataManager.h"
+#include "EntitySystem/BuiltInComponentTypes.h"
+#include "Evaluation/MovieSceneEvaluationTemplateInstance.h"
+#include "MovieSceneSection.h"
+#include "Sections/MovieSceneSubSection.h"
+#include "Systems/MovieSceneTransformOriginSystem.h"
+#include "Tracks/MovieSceneSubTrack.h"
+
+FName FSubTrackEditorMode::ModeName("EditMode.SubTrackEditMode");
+
+FSubTrackEditorMode::FSubTrackEditorMode()
+{
+	CachedLocation.Reset();
+}
+
+FSubTrackEditorMode::~FSubTrackEditorMode()
+{
+}
+
+void FSubTrackEditorMode::Initialize()
+{
+	CachedLocation.Reset();
+}
+
+bool FSubTrackEditorMode::InputDelta(FEditorViewportClient* InViewportClient, FViewport* InViewport, FVector& InDrag, FRotator& InRot, FVector& InScale)
+{
+	if(AreAnyActorsSelected())
+	{
+		return false;
+	}
+	const bool bCtrlDown = InViewport->KeyState(EKeys::LeftControl) || InViewport->KeyState(EKeys::RightControl);
+	const bool bShiftDown = InViewport->KeyState(EKeys::LeftShift) || InViewport->KeyState(EKeys::RightShift);
+	const bool bAltDown = InViewport->KeyState(EKeys::LeftAlt) || InViewport->KeyState(EKeys::RightAlt);
+	const bool bMouseButtonDown = InViewport->KeyState(EKeys::LeftMouseButton);
+	const bool bAnyModifiers = bAltDown || bCtrlDown || bShiftDown;
+	
+	const EAxisList::Type CurrentAxis = InViewportClient->GetCurrentWidgetAxis();
+
+	if(bMouseButtonDown && !bAnyModifiers && CurrentAxis != EAxisList::None)
+	{
+		OnOriginValueChanged.Broadcast( InDrag, InRot);
+		return true;
+	}
+	
+	return false;
+}
+
+bool FSubTrackEditorMode::UsesTransformWidget() const
+{
+	const UMovieSceneSubSection* SubSection = GetSelectedSection();
+	if(!AreAnyActorsSelected() && SubSection)
+	{
+		return DoesSubSectionHaveTransformOverrides(*SubSection);
+	}
+	return FEdMode::UsesTransformWidget();
+}
+
+bool FSubTrackEditorMode::UsesTransformWidget(UE::Widget::EWidgetMode CheckMode) const
+{
+	const UMovieSceneSubSection* SubSection = GetSelectedSection();
+	if(!AreAnyActorsSelected() && SubSection)
+	{
+		return DoesSubSectionHaveTransformOverrides(*SubSection);
+	}
+	return FEdMode::UsesTransformWidget(CheckMode);
+}
+
+TOptional<FMovieSceneSequenceID> FSubTrackEditorMode::GetSequenceIDForSubSeciton(const UMovieSceneSubSection* InSubSection) const
+{
+	const TSharedPtr<ISequencer> Sequencer = WeakSequencer.Pin();
+
+	if(!Sequencer)
+	{
+		return TOptional<FMovieSceneSequenceID>();
+	}
+	
+	const FMovieSceneRootEvaluationTemplateInstance& EvaluationTemplate = Sequencer->GetEvaluationTemplate();
+	
+	UMovieSceneCompiledDataManager* CompiledDataManager = EvaluationTemplate.GetCompiledDataManager();
+	UMovieSceneSequence* RootSequence = EvaluationTemplate.GetSequence(Sequencer->GetRootTemplateID());
+	const FMovieSceneCompiledDataID DataID = CompiledDataManager->Compile(RootSequence);
+
+	const FMovieSceneSequenceHierarchy& Hierarchy = CompiledDataManager->GetHierarchyChecked(DataID);
+	
+	TArray<FMovieSceneSequenceID> SubSequenceHierarchy = Sequencer->GetSubSequenceHierarchy();
+
+	UE::MovieScene::FSubSequencePath Path;
+	const FMovieSceneSequenceID ParentSequenceID = SubSequenceHierarchy.Last();
+
+	Path.Reset(ParentSequenceID, &Hierarchy);
+
+	return Path.ResolveChildSequenceID(InSubSection->GetSequenceID());
+}
+
+FTransform FSubTrackEditorMode::GetFinalTransformOriginForSubSection(const UMovieSceneSubSection* InSubSection) const
+{
+	FTransform TransformOrigin = FTransform();
+
+	const TSharedPtr<ISequencer> Sequencer = WeakSequencer.Pin();
+
+	if(!Sequencer)
+	{
+		return TransformOrigin;
+	}
+
+	const FMovieSceneRootEvaluationTemplateInstance& EvaluationTemplate = Sequencer->GetEvaluationTemplate();
+
+	const TOptional<FMovieSceneSequenceID> ChildSequenceID = GetSequenceIDForSubSeciton(InSubSection);
+	
+	const UMovieSceneEntitySystemLinker* EntityLinker = EvaluationTemplate.GetEntitySystemLinker();
+	if(!EntityLinker || !ChildSequenceID)
+	{
+		return TransformOrigin;
+	}
+
+	const UMovieSceneTransformOriginSystem* TransformOriginSystem = EntityLinker->FindSystem<UMovieSceneTransformOriginSystem>();
+
+	if(!TransformOriginSystem)
+	{
+		return TransformOrigin;
+	}
+	
+	const TSparseArray<FTransform>& TransformOrigins = TransformOriginSystem->GetTransformOriginsByInstanceID();
+	const TMap<FMovieSceneSequenceID, UE::MovieScene::FInstanceHandle> SequenceIDToInstanceHandle = TransformOriginSystem->GetSequenceIDToInstanceHandle();
+
+	if(SequenceIDToInstanceHandle.Contains(ChildSequenceID.GetValue()))
+	{
+		const UE::MovieScene::FInstanceHandle CurrentHandle = SequenceIDToInstanceHandle[ChildSequenceID.GetValue()];
+		if(TransformOrigins.IsValidIndex(CurrentHandle.InstanceID))
+		{
+			TransformOrigin = TransformOrigins[CurrentHandle.InstanceID];
+		}
+	}
+
+	return TransformOrigin;
+}
+
+bool FSubTrackEditorMode::AreAnyActorsSelected() const
+{
+	FSelectedActorIterator SelectedActorIter(GetWorld());
+	if(SelectedActorIter)
+	{
+		return true;
+	}
+	return false;
+}
+
+FVector FSubTrackEditorMode::GetWidgetLocation() const
+{
+	const UMovieSceneSubSection* SubSection = GetSelectedSection();
+	if(!AreAnyActorsSelected() && SubSection && DoesSubSectionHaveTransformOverrides(*SubSection))
+	{
+		FVector NewLocation = GetFinalTransformOriginForSubSection(SubSection).GetLocation();
+
+		if(!CachedLocation.IsSet() || !NewLocation.Equals(CachedLocation.GetValue()))
+		{
+			CachedLocation = NewLocation;
+			// Invalidate hit proxies, otherwise the hit proxy for the widget can be out of sync, and still at the old widget location
+			GEditor->RedrawLevelEditingViewports(true);
+		}
+		return CachedLocation.GetValue();
+	}
+	
+	return FEdMode::GetWidgetLocation();
+}
+
+bool FSubTrackEditorMode::ShouldDrawWidget() const
+{
+	const UMovieSceneSubSection* SubSection = GetSelectedSection();
+	if(!AreAnyActorsSelected() && SubSection && DoesSubSectionHaveTransformOverrides(*SubSection))
+	{
+		return true;
+	}
+	// If the widget is not being drawn, its hit proxies need to be invalidated the next time it is drawn.
+	// Resetting the cached location will trigger the invalidation in GetWidgetLocation
+	CachedLocation.Reset();
+	return false;
+}
+
+bool FSubTrackEditorMode::GetPivotForOrbit(FVector& OutPivot) const
+{
+	return FEdMode::GetPivotForOrbit(OutPivot);
+}
+
+bool FSubTrackEditorMode::GetCustomDrawingCoordinateSystem(FMatrix& OutMatrix, void* InData)
+{
+	// @todo support custom coordinate systems based on parent rotation.
+	return false;
+}
+
+bool FSubTrackEditorMode::GetCustomInputCoordinateSystem(FMatrix& OutMatrix, void* InData)
+{
+	return FEdMode::GetCustomInputCoordinateSystem(OutMatrix, InData);
+}
+
+bool FSubTrackEditorMode::IsCompatibleWith(FEditorModeID OtherModeID) const
+{
+	if(IncompatibleEditorModes.Contains(OtherModeID))
+	{
+		return false;
+	}
+	return true;
+}
+
+bool FSubTrackEditorMode::DoesSubSectionHaveTransformOverrides(const UMovieSceneSubSection& SubSection)
+{
+	if (SubSection.IsActive())
+	{
+		const EMovieSceneTransformChannel SectionTransformChannels = SubSection.GetMask().GetChannels();
+
+		return EnumHasAnyFlags(SectionTransformChannels, EMovieSceneTransformChannel::Translation) || EnumHasAnyFlags(SectionTransformChannels, EMovieSceneTransformChannel::Rotation);
+	}
+	
+	return false;
+}
+
+UMovieSceneSubSection* FSubTrackEditorMode::GetSelectedSection() const
+{
+	const TSharedPtr<ISequencer> PinnedSequencer = WeakSequencer.Pin();
+	UMovieSceneSubSection* SelectedSection = nullptr;
+	if(!PinnedSequencer.IsValid())
+	{
+		return SelectedSection;
+	}
+
+	TArray<UMovieSceneSection*> SelectedSections;
+	PinnedSequencer->GetSelectedSections(SelectedSections);
+	
+	for(UMovieSceneSection* Section : SelectedSections)
+	{
+		if(UMovieSceneSubSection* SubSection = Cast<UMovieSceneSubSection>(Section))
+		{
+			// Mirror behavior when multiple actors are selected in the level editor, and pick the last selected item.
+			SelectedSection = SubSection;
+		}
+	}
+
+	if(SelectedSection)
+	{
+		return SelectedSection;
+	}
+
+	TArray<UMovieSceneTrack*> SelectedTracks;
+	PinnedSequencer->GetSelectedTracks(SelectedTracks);
+	for(UMovieSceneTrack* Track : SelectedTracks)
+	{
+		// Similarly to section selection, pick the last selected track.
+		if(const UMovieSceneSubTrack* SubTrack = Cast<UMovieSceneSubTrack>(Track))
+		{
+			if(SubTrack->GetSectionToKey())
+			{
+				SelectedSection = Cast<UMovieSceneSubSection>(SubTrack->GetSectionToKey());
+			}
+			else if(SubTrack->GetAllSections().Num())
+			{
+				// Since the first section is the section that will be keyed by default, select the first section from the track.
+				SelectedSection = Cast<UMovieSceneSubSection>(SubTrack->FindSection(PinnedSequencer.Get()->GetLocalTime().Time.FrameNumber));
+			}
+		}
+			
+	}
+	return SelectedSection;
+}
diff --git a/Engine/Source/Editor/MovieSceneTools/Private/EditModes/SubTrackEditorMode.h b/Engine/Source/Editor/MovieSceneTools/Private/EditModes/SubTrackEditorMode.h
new file mode 100644
--- /dev/null
+++ b/Engine/Source/Editor/MovieSceneTools/Private/EditModes/SubTrackEditorMode.h
@@ -0,0 +1,57 @@
+﻿// Copyright Epic Games, Inc. All Rights Reserved.
+
+#pragma once
+
+#include "EdMode.h"
+#include "ISequencer.h"
+
+class FSubTrackEditorMode : public FEdMode
+{
+public:
+	static FName ModeName;
+	
+	DECLARE_MULTICAST_DELEGATE_TwoParams(FOnOriginValueChanged, FVector, FRotator)
+	
+	FSubTrackEditorMode();
+	virtual ~FSubTrackEditorMode() override;
+
+	// FEdMode interface
+	virtual void Initialize() override;
+	virtual bool InputDelta(FEditorViewportClient* InViewportClient, FViewport* InViewport, FVector& InDrag, FRotator& InRot, FVector& InScale) override;
+	virtual bool UsesTransformWidget() const override;
+	virtual bool UsesTransformWidget(UE::Widget::EWidgetMode CheckMode) const override;
+	virtual FVector GetWidgetLocation() const override;
+	virtual bool ShouldDrawWidget() const override;
+	virtual bool GetPivotForOrbit(FVector& OutPivot) const override;
+	virtual bool GetCustomDrawingCoordinateSystem(FMatrix& OutMatrix, void* InData) override;
+	virtual bool GetCustomInputCoordinateSystem(FMatrix& OutMatrix, void* InData) override;
+	virtual bool IsCompatibleWith(FEditorModeID OtherModeID) const override;
+
+	void SetSequencer(const TSharedPtr<ISequencer>& InSequencer) { WeakSequencer = InSequencer; }
+	FOnOriginValueChanged& GetOnOriginValueChanged() { return OnOriginValueChanged; };
+
+private:
+	// Returns true if there are any sub tracks with active transform origin overrides.
+	static bool DoesSubSectionHaveTransformOverrides(const UMovieSceneSubSection& SubSection);
+
+	// Gets the sequence ID from the context of the subsection in the current hierarchy.
+	TOptional<FMovieSceneSequenceID> GetSequenceIDForSubSeciton(const UMovieSceneSubSection* InSubSection) const;
+	// Gets the transform origin of the provided section, after all parent transforms have been applied.
+	FTransform GetFinalTransformOriginForSubSection(const UMovieSceneSubSection* SubSection) const;
+	// returns true if any actors are selected in the level editor. This is used to prevent this editor mode from being active when actors are selected.
+	bool AreAnyActorsSelected() const;
+
+	// Returns the currently selected 
+	UMovieSceneSubSection* GetSelectedSection() const;
+
+	/** Sequencer that owns this editor mode */
+	TWeakPtr<ISequencer> WeakSequencer;
+
+	/** Delegate called when the origin is modified from the editor gizmo */
+	FOnOriginValueChanged OnOriginValueChanged;
+
+	/** Used to tell if the gizmo as moved, and if the editor hit proxies need to be invalidated as a result */
+	mutable TOptional<FVector> CachedLocation;
+
+	const TArray<FName> IncompatibleEditorModes = TArray<FName>({ "EditMode.ControlRig" });
+};
diff --git a/Engine/Source/Editor/MovieSceneTools/Private/MovieSceneToolsModule.cpp b/Engine/Source/Editor/MovieSceneTools/Private/MovieSceneToolsModule.cpp
--- a/Engine/Source/Editor/MovieSceneTools/Private/MovieSceneToolsModule.cpp
+++ b/Engine/Source/Editor/MovieSceneTools/Private/MovieSceneToolsModule.cpp
@@ -126,6 +126,7 @@
 #include "Bindings/MovieSceneSpawnableDirectorBlueprintBinding.h"
 #include "EntitySystem/MovieSceneSharedPlaybackState.h"
 #include "MovieSceneDynamicBindingUtils.h"
+#include "EditModes/SubTrackEditorMode.h"
 #include "Kismet2/KismetEditorUtilities.h"
 
 #include "Conditions/MovieSceneDirectorBlueprintCondition.h"
@@ -302,6 +303,12 @@ void FMovieSceneToolsModule::StartupModule()
 		FSlateIcon(),
 		false);
 
+	FEditorModeRegistry::Get().RegisterMode<FSubTrackEditorMode>(
+		FSubTrackEditorMode::ModeName,
+		NSLOCTEXT("SubTrackEditorMode", "SubTrackEditMode", "Sub Track Mode"),
+		FSlateIcon(),
+		false);
+
 	// register UTransformableComponentHandle animatable interface
 	FConstraintChannelInterfaceRegistry& ConstraintChannelInterfaceRegistry = FConstraintChannelInterfaceRegistry::Get();
 	ConstraintChannelInterfaceRegistry.RegisterConstraintChannelInterface<UTransformableComponentHandle>(MakeUnique<FComponentConstraintChannelInterface>());
@@ -404,6 +411,7 @@ void FMovieSceneToolsModule::ShutdownModule()
 	}
 
 	FEditorModeRegistry::Get().UnregisterMode(FSkeletalAnimationTrackEditMode::ModeName);
+	FEditorModeRegistry::Get().UnregisterMode(FSubTrackEditorMode::ModeName);
 
 }
 
diff --git a/Engine/Source/Editor/MovieSceneTools/Private/Sections/CinematicShotSection.cpp b/Engine/Source/Editor/MovieSceneTools/Private/Sections/CinematicShotSection.cpp
--- a/Engine/Source/Editor/MovieSceneTools/Private/Sections/CinematicShotSection.cpp
+++ b/Engine/Source/Editor/MovieSceneTools/Private/Sections/CinematicShotSection.cpp
@@ -348,6 +348,88 @@ void FCinematicShotSection::AddShotMenuSection(FMenuBuilder& MenuBuilder, const
 		);
 	}
 	MenuBuilder.EndSection();
+
+	auto MakeUIAction = [this](EMovieSceneTransformChannel ChannelsToToggle, const TSharedPtr<ISequencer>& Sequencer)
+	{
+
+		UMovieSceneSubSection* SubSection = Cast<UMovieSceneSubSection>(Section);
+		if(!SubSection)
+		{
+			return FUIAction();
+		}
+		return FUIAction(
+			FExecuteAction::CreateLambda([SubSection, ChannelsToToggle, Sequencer]
+				{
+					FScopedTransaction Transaction(LOCTEXT("SetActiveChannelsTransaction", "Set Active Channels"));
+					SubSection->Modify();
+					EMovieSceneTransformChannel Channels = SubSection->GetMask().GetChannels();
+
+					if (EnumHasAllFlags(Channels, ChannelsToToggle) || (Channels & ChannelsToToggle) == EMovieSceneTransformChannel::None)
+					{
+						SubSection->SetMask(SubSection->GetMask().GetChannels() ^ ChannelsToToggle);
+					}
+					else
+					{
+						SubSection->SetMask(SubSection->GetMask().GetChannels() | ChannelsToToggle);
+					}
+				
+					Sequencer->NotifyMovieSceneDataChanged(EMovieSceneDataChangeType::MovieSceneStructureItemsChanged);
+				}
+			),
+			FCanExecuteAction(),
+			FGetActionCheckState::CreateLambda([SubSection, ChannelsToToggle]
+			{
+				EMovieSceneTransformChannel Channels = SubSection->GetMask().GetChannels();
+				if (EnumHasAllFlags(Channels, ChannelsToToggle))
+				{
+					return ECheckBoxState::Checked;
+				}
+				else if (EnumHasAnyFlags(Channels, ChannelsToToggle))
+				{
+					return ECheckBoxState::Undetermined;
+				}
+				return ECheckBoxState::Unchecked;
+			})
+		);
+	};
+
+	TSharedPtr<ISequencer> Sequencer = GetSequencer();
+	
+	MenuBuilder.BeginSection(NAME_None, LOCTEXT("SequenceMenuText", "Active Channels"));
+	MenuBuilder.AddSubMenu(
+		LOCTEXT("AllTranslation", "Translation"), LOCTEXT("AllTranslation_ToolTip", "Causes this section to affect the translation of the transform"),
+		FNewMenuDelegate::CreateLambda([Sequencer, MakeUIAction](FMenuBuilder& SubMenuBuilder){
+			SubMenuBuilder.AddMenuEntry(
+				LOCTEXT("TranslationX", "X"), LOCTEXT("TranslationX_ToolTip", "Causes this section to affect the X channel of the transform's translation"),
+				FSlateIcon(), MakeUIAction(EMovieSceneTransformChannel::TranslationX, Sequencer), NAME_None, EUserInterfaceActionType::ToggleButton);
+			SubMenuBuilder.AddMenuEntry(
+				LOCTEXT("TranslationY", "Y"), LOCTEXT("TranslationY_ToolTip", "Causes this section to affect the Y channel of the transform's translation"),
+				FSlateIcon(), MakeUIAction(EMovieSceneTransformChannel::TranslationY, Sequencer), NAME_None, EUserInterfaceActionType::ToggleButton);
+			SubMenuBuilder.AddMenuEntry(
+				LOCTEXT("TranslationZ", "Z"), LOCTEXT("TranslationZ_ToolTip", "Causes this section to affect the Z channel of the transform's translation"),
+				FSlateIcon(), MakeUIAction(EMovieSceneTransformChannel::TranslationZ, Sequencer), NAME_None, EUserInterfaceActionType::ToggleButton);
+		}),
+		MakeUIAction(EMovieSceneTransformChannel::Translation, Sequencer),
+		NAME_None,
+		EUserInterfaceActionType::ToggleButton);
+
+	MenuBuilder.AddSubMenu(
+		LOCTEXT("AllRotation", "Rotation"), LOCTEXT("AllRotation_ToolTip", "Causes this section to affect the rotation of the transform"),
+		FNewMenuDelegate::CreateLambda([Sequencer, MakeUIAction](FMenuBuilder& SubMenuBuilder){
+			SubMenuBuilder.AddMenuEntry(
+				LOCTEXT("RotationX", "Roll (X)"), LOCTEXT("RotationX_ToolTip", "Causes this section to affect the roll (X) channel the transform's rotation"),
+				FSlateIcon(), MakeUIAction(EMovieSceneTransformChannel::RotationX, Sequencer), NAME_None, EUserInterfaceActionType::ToggleButton);
+			SubMenuBuilder.AddMenuEntry(
+				LOCTEXT("RotationY", "Pitch (Y)"), LOCTEXT("RotationY_ToolTip", "Causes this section to affect the pitch (Y) channel the transform's rotation"),
+				FSlateIcon(), MakeUIAction(EMovieSceneTransformChannel::RotationY, Sequencer), NAME_None, EUserInterfaceActionType::ToggleButton);
+			SubMenuBuilder.AddMenuEntry(
+				LOCTEXT("RotationZ", "Yaw (Z)"), LOCTEXT("RotationZ_ToolTip", "Causes this section to affect the yaw (Z) channel the transform's rotation"),
+				FSlateIcon(), MakeUIAction(EMovieSceneTransformChannel::RotationZ, Sequencer), NAME_None, EUserInterfaceActionType::ToggleButton);
+		}),
+		MakeUIAction(EMovieSceneTransformChannel::Rotation, Sequencer),
+		NAME_None,
+		EUserInterfaceActionType::ToggleButton);
+	MenuBuilder.EndSection();
 }
 
 /* FCinematicShotSection callbacks
diff --git a/Engine/Source/Editor/MovieSceneTools/Private/TrackEditors/SubTrackEditor.cpp b/Engine/Source/Editor/MovieSceneTools/Private/TrackEditors/SubTrackEditor.cpp
--- a/Engine/Source/Editor/MovieSceneTools/Private/TrackEditors/SubTrackEditor.cpp
+++ b/Engine/Source/Editor/MovieSceneTools/Private/TrackEditors/SubTrackEditor.cpp
@@ -36,10 +36,19 @@
 #include "Algo/Accumulate.h"
 #include "SequencerUtilities.h"
 #include "AssetToolsModule.h"
+#include "EditorModeManager.h"
 #include "Interfaces/IMainFrameModule.h"
 #include "IDetailsView.h"
+#include "IKeyArea.h"
 #include "IStructureDetailsView.h"
 #include "PropertyEditorModule.h"
+#include "Compilation/MovieSceneCompiledDataManager.h"
+#include "EditModes/SubTrackEditorMode.h"
+#include "EntitySystem/BuiltInComponentTypes.h"
+#include "EntitySystem/MovieSceneEntitySystemTask.h"
+#include "MVVM/Selection/Selection.h"
+#include "MVVM/ViewModels/SequencerEditorViewModel.h"
+#include "Systems/MovieSceneTransformOriginSystem.h"
 
 #define LOCTEXT_NAMESPACE "FSubTrackEditor"
 
@@ -119,6 +128,82 @@ public:
 			);
 		}
 		MenuBuilder.EndSection();
+
+		auto MakeUIAction = [Section](EMovieSceneTransformChannel ChannelsToToggle, const TSharedPtr<ISequencer>& Sequencer)
+		{
+			return FUIAction(
+				FExecuteAction::CreateLambda([Section, ChannelsToToggle, Sequencer]
+					{
+						FScopedTransaction Transaction(LOCTEXT("SetActiveChannelsTransaction", "Set Active Channels"));
+						Section->Modify();
+						EMovieSceneTransformChannel Channels = Section->GetMask().GetChannels();
+
+						if (EnumHasAllFlags(Channels, ChannelsToToggle) || (Channels & ChannelsToToggle) == EMovieSceneTransformChannel::None)
+						{
+							Section->SetMask(Section->GetMask().GetChannels() ^ ChannelsToToggle);
+						}
+						else
+						{
+							Section->SetMask(Section->GetMask().GetChannels() | ChannelsToToggle);
+						}
+					
+						Sequencer->NotifyMovieSceneDataChanged(EMovieSceneDataChangeType::MovieSceneStructureItemsChanged);
+					}
+				),
+				FCanExecuteAction(),
+				FGetActionCheckState::CreateLambda([Section, ChannelsToToggle]
+				{
+					EMovieSceneTransformChannel Channels = Section->GetMask().GetChannels();
+					if (EnumHasAllFlags(Channels, ChannelsToToggle))
+					{
+						return ECheckBoxState::Checked;
+					}
+					if (EnumHasAnyFlags(Channels, ChannelsToToggle))
+					{
+						return ECheckBoxState::Undetermined;
+					}
+					return ECheckBoxState::Unchecked;
+				})
+			);
+		};
+
+		TSharedPtr<ISequencer> Sequencer = GetSequencer();
+		
+		MenuBuilder.BeginSection(NAME_None, LOCTEXT("SequenceMenuText", "Active Channels"));
+		MenuBuilder.AddSubMenu(
+			LOCTEXT("AllTranslation", "Translation"), LOCTEXT("AllTranslation_ToolTip", "Causes this section to affect the translation of the transform"),
+			FNewMenuDelegate::CreateLambda([Sequencer, MakeUIAction](FMenuBuilder& SubMenuBuilder){
+				SubMenuBuilder.AddMenuEntry(
+					LOCTEXT("TranslationX", "X"), LOCTEXT("TranslationX_ToolTip", "Causes this section to affect the X channel of the transform's translation"),
+					FSlateIcon(), MakeUIAction(EMovieSceneTransformChannel::TranslationX, Sequencer), NAME_None, EUserInterfaceActionType::ToggleButton);
+				SubMenuBuilder.AddMenuEntry(
+					LOCTEXT("TranslationY", "Y"), LOCTEXT("TranslationY_ToolTip", "Causes this section to affect the Y channel of the transform's translation"),
+					FSlateIcon(), MakeUIAction(EMovieSceneTransformChannel::TranslationY, Sequencer), NAME_None, EUserInterfaceActionType::ToggleButton);
+				SubMenuBuilder.AddMenuEntry(
+					LOCTEXT("TranslationZ", "Z"), LOCTEXT("TranslationZ_ToolTip", "Causes this section to affect the Z channel of the transform's translation"),
+					FSlateIcon(), MakeUIAction(EMovieSceneTransformChannel::TranslationZ, Sequencer), NAME_None, EUserInterfaceActionType::ToggleButton);
+			}),
+			MakeUIAction(EMovieSceneTransformChannel::Translation, Sequencer),
+			NAME_None,
+			EUserInterfaceActionType::ToggleButton);
+
+		MenuBuilder.AddSubMenu(
+			LOCTEXT("AllRotation", "Rotation"), LOCTEXT("AllRotation_ToolTip", "Causes this section to affect the rotation of the transform"),
+			FNewMenuDelegate::CreateLambda([Sequencer, MakeUIAction](FMenuBuilder& SubMenuBuilder){
+				SubMenuBuilder.AddMenuEntry(
+					LOCTEXT("RotationX", "Roll (X)"), LOCTEXT("RotationX_ToolTip", "Causes this section to affect the roll (X) channel the transform's rotation"),
+					FSlateIcon(), MakeUIAction(EMovieSceneTransformChannel::RotationX, Sequencer), NAME_None, EUserInterfaceActionType::ToggleButton);
+				SubMenuBuilder.AddMenuEntry(
+					LOCTEXT("RotationY", "Pitch (Y)"), LOCTEXT("RotationY_ToolTip", "Causes this section to affect the pitch (Y) channel the transform's rotation"),
+					FSlateIcon(), MakeUIAction(EMovieSceneTransformChannel::RotationY, Sequencer), NAME_None, EUserInterfaceActionType::ToggleButton);
+				SubMenuBuilder.AddMenuEntry(
+					LOCTEXT("RotationZ", "Yaw (Z)"), LOCTEXT("RotationZ_ToolTip", "Causes this section to affect the yaw (Z) channel the transform's rotation"),
+					FSlateIcon(), MakeUIAction(EMovieSceneTransformChannel::RotationZ, Sequencer), NAME_None, EUserInterfaceActionType::ToggleButton);
+			}),
+			MakeUIAction(EMovieSceneTransformChannel::Rotation, Sequencer),
+			NAME_None,
+			EUserInterfaceActionType::ToggleButton);
+		MenuBuilder.EndSection();
 	}
 
 	void TogglePlayableDirectly()
@@ -197,7 +282,7 @@ private:
  *****************************************************************************/
 
 FSubTrackEditor::FSubTrackEditor(TSharedRef<ISequencer> InSequencer)
-	: FMovieSceneTrackEditor(InSequencer) 
+	: FKeyframeTrackEditor<UMovieSceneSubTrack>(InSequencer)
 { }
 
 
@@ -223,11 +308,140 @@ TSharedPtr<SWidget> FSubTrackEditor::BuildOutlinerEditWidget(const FGuid& Object
 }
 
 
+void FSubTrackEditor::GetOriginKeys(const FVector& CurrentPosition, const FRotator& CurrentRotation, UMovieSceneSection* Section, FGeneratedTrackKeys& OutGeneratedKeys)
+{
+	FMovieSceneChannelProxy& SectionChannelProxy = Section->GetChannelProxy();
+	TMovieSceneChannelHandle<FMovieSceneDoubleChannel> ChannelHandles[] = {
+		SectionChannelProxy.GetChannelByName<FMovieSceneDoubleChannel>("Override.Location.X"),
+		SectionChannelProxy.GetChannelByName<FMovieSceneDoubleChannel>("Override.Location.Y"),
+		SectionChannelProxy.GetChannelByName<FMovieSceneDoubleChannel>("Override.Location.Z"),
+		SectionChannelProxy.GetChannelByName<FMovieSceneDoubleChannel>("Override.Rotation.X"),
+		SectionChannelProxy.GetChannelByName<FMovieSceneDoubleChannel>("Override.Rotation.Y"),
+		SectionChannelProxy.GetChannelByName<FMovieSceneDoubleChannel>("Override.Rotation.Z")
+	};
+	
+	OutGeneratedKeys.Add(FMovieSceneChannelValueSetter::Create<FMovieSceneDoubleChannel>(ChannelHandles[0].GetChannelIndex(), CurrentPosition.X, true));
+	OutGeneratedKeys.Add(FMovieSceneChannelValueSetter::Create<FMovieSceneDoubleChannel>(ChannelHandles[1].GetChannelIndex(), CurrentPosition.Y, true));
+	OutGeneratedKeys.Add(FMovieSceneChannelValueSetter::Create<FMovieSceneDoubleChannel>(ChannelHandles[2].GetChannelIndex(), CurrentPosition.Z, true));
+	OutGeneratedKeys.Add(FMovieSceneChannelValueSetter::Create<FMovieSceneDoubleChannel>(ChannelHandles[3].GetChannelIndex(), CurrentRotation.Roll, true));
+	OutGeneratedKeys.Add(FMovieSceneChannelValueSetter::Create<FMovieSceneDoubleChannel>(ChannelHandles[4].GetChannelIndex(), CurrentRotation.Pitch, true));
+	OutGeneratedKeys.Add(FMovieSceneChannelValueSetter::Create<FMovieSceneDoubleChannel>(ChannelHandles[5].GetChannelIndex(), CurrentRotation.Yaw, true));
+}
+
 TSharedRef<ISequencerTrackEditor> FSubTrackEditor::CreateTrackEditor(TSharedRef<ISequencer> InSequencer)
 {
 	return MakeShareable(new FSubTrackEditor(InSequencer));
 }
 
+int32 FSubTrackEditor::GetPreviousKey(FMovieSceneDoubleChannel& Channel, FFrameNumber Time)
+{
+	TArray<FFrameNumber> KeyTimes;
+	TArray<FKeyHandle> KeyHandles;
+
+	TRange<FFrameNumber> Range;
+	Range.SetLowerBound(TRangeBound<FFrameNumber>::Open());
+	Range.SetUpperBound(TRangeBound<FFrameNumber>::Exclusive(Time));
+	Channel.GetData().GetKeys(Range, &KeyTimes, &KeyHandles);
+
+	if (KeyHandles.Num() <= 0)
+	{
+		return INDEX_NONE;
+	}
+
+	int32 Index = Channel.GetData().GetIndex(KeyHandles[KeyHandles.Num() - 1]);
+	return Index;
+}
+
+double FSubTrackEditor::UnwindChannel(const double& OldValue, double NewValue)
+{
+	while( NewValue - OldValue > 180.0f )
+	{
+		NewValue -= 360.0f;
+	}
+	while( NewValue - OldValue < -180.0f )
+	{
+		NewValue += 360.0f;
+	}
+	return NewValue;
+}
+
+void FSubTrackEditor::ProcessKeyOperation(FFrameNumber InKeyTime, const UE::Sequencer::FKeyOperation& Operation, ISequencer& InSequencer)
+{
+	auto Iterator = [this, InKeyTime, &InSequencer](UMovieSceneTrack* Track, TArrayView<const UE::Sequencer::FKeySectionOperation> Operations)
+	{
+		this->ProcessKeyOperationInternal(Operations, InSequencer, InKeyTime);
+	};
+
+	Operation.IterateOperations(Iterator);
+}
+void FSubTrackEditor::ProcessKeyOperationInternal(TArrayView<const UE::Sequencer::FKeySectionOperation> SectionsToKey, ISequencer& InSequencer, FFrameNumber KeyTime)
+{
+	for (int32 Index = 0; Index < SectionsToKey.Num(); ++Index)
+	{
+		for (TSharedPtr<IKeyArea> KeyArea : SectionsToKey[Index].KeyAreas)
+		{
+			UMovieSceneSubSection* SubSection = Cast<UMovieSceneSubSection>(SectionsToKey[Index].Section->GetSectionObject());
+			FMovieSceneChannelHandle Handle = KeyArea->GetChannel();
+			if(Handle.GetChannelTypeName() == FMovieSceneDoubleChannel::StaticStruct()->GetFName() && SubSection)
+			{
+				FMovieSceneDoubleChannel* Channel = static_cast<FMovieSceneDoubleChannel*>(Handle.Get());
+
+				if (ensureAlwaysMsgf(Channel, TEXT("Channel: %s for Key Area %s does not exist. Keying may not function properly"), *Handle.GetChannelTypeName().ToString(), *KeyArea->GetName().ToString()))
+				{
+					double Value = 0.0;
+
+					FTransform RawTransformOrigin = GetTransformOriginDataForSubSection(SubSection);
+					TOptional<FVector> KeyPosition = SubSection->GetKeyPreviewPosition();
+					TOptional<FRotator> KeyRotation = SubSection->GetKeyPreviewRotation();
+;
+					switch (Handle.GetChannelIndex())
+					{
+					case 0:
+						Value = KeyPosition.IsSet() ? KeyPosition.GetValue().X : RawTransformOrigin.GetLocation().X;
+						break;
+					case 1:
+						Value = KeyPosition.IsSet() ? KeyPosition.GetValue().Y : RawTransformOrigin.GetLocation().Y;
+						break;
+					case 2:
+						Value = KeyPosition.IsSet() ? KeyPosition.GetValue().Z : RawTransformOrigin.GetLocation().Z;
+						break;
+					case 3:
+						Value = KeyRotation.IsSet() ? KeyRotation.GetValue().Roll : RawTransformOrigin.Rotator().Roll;
+						break;
+					case 4:
+						Value = KeyRotation.IsSet() ? KeyRotation.GetValue().Pitch : RawTransformOrigin.Rotator().Pitch;
+						break;
+					case 5:
+						Value = KeyRotation.IsSet() ? KeyRotation.GetValue().Yaw : RawTransformOrigin.Rotator().Yaw;
+						break;
+					default:
+						Value = 0.0;
+					}
+
+					if (KeyArea->GetName() == "Rotation.X" ||
+						KeyArea->GetName() == "Rotation.Y" ||
+						KeyArea->GetName() == "Rotation.Z")
+					{
+						int32 PreviousKey = GetPreviousKey(*Channel, KeyTime);
+						if (PreviousKey != INDEX_NONE && PreviousKey < Channel->GetData().GetValues().Num())
+						{
+							double OldValue = Channel->GetData().GetValues()[PreviousKey].Value;
+							Value = UnwindChannel(OldValue, Value);
+						}
+					}
+
+					EMovieSceneKeyInterpolation Interpolation = GetInterpolationMode(Channel, KeyTime, InSequencer.GetKeyInterpolation());
+					AddKeyToChannel(Channel, KeyTime, Value, Interpolation);
+				}
+			}
+			else
+			{
+				KeyArea->AddOrUpdateKey(KeyTime, FGuid(), InSequencer);
+			}
+		}
+	}
+}
+
 
 TSharedRef<ISequencerSection> FSubTrackEditor::MakeSectionInterface(UMovieSceneSection& SectionObject, UMovieSceneTrack& Track, FGuid ObjectBinding)
 {
@@ -432,6 +646,41 @@ bool FSubTrackEditor::IsResizable(UMovieSceneTrack* InTrack) const
 	return true;
 }
 
+void FSubTrackEditor::OnInitialize()
+{
+	GLevelEditorModeTools().ActivateDefaultMode();
+
+	GLevelEditorModeTools().ActivateMode(FSubTrackEditorMode::ModeName);
+	FSubTrackEditorMode* EditorMode = static_cast<FSubTrackEditorMode*>(GLevelEditorModeTools().GetActiveMode(FSubTrackEditorMode::ModeName));
+	if(EditorMode)
+	{
+		EditorMode->SetSequencer(GetSequencer());
+		EditorMode->GetOnOriginValueChanged().RemoveAll(this);
+		EditorMode->GetOnOriginValueChanged().AddRaw(this, &FSubTrackEditor::UpdateOrigin);
+	}
+	GetSequencer()->GetViewModel()->GetSelection()->TrackArea.OnChanged.AddSP(this, &FSubTrackEditor::UpdateActiveMode);
+
+	SectionsWithPreviews.Empty();
+	GetSequencer()->OnPlayEvent().AddSP(this, &FSubTrackEditor::ResetSectionPreviews);
+	GetSequencer()->OnBeginScrubbingEvent().AddSP(this, &FSubTrackEditor::ResetSectionPreviews);
+	GetSequencer()->OnActivateSequence().AddSP(this, &FSubTrackEditor::ResetSectionPreviews);
+}
+
+void FSubTrackEditor::OnRelease()
+{
+	FSubTrackEditorMode* EditorMode = static_cast<FSubTrackEditorMode*>(GLevelEditorModeTools().GetActiveMode(FSubTrackEditorMode::ModeName));
+	if(EditorMode)
+	{
+		EditorMode->GetOnOriginValueChanged().RemoveAll(this);	
+	}
+	GetSequencer()->GetViewModel()->GetSelection()->TrackArea.OnChanged.RemoveAll(this);
+
+	SectionsWithPreviews.Empty();
+	GetSequencer()->OnPlayEvent().RemoveAll(this);
+	GetSequencer()->OnBeginScrubbingEvent().RemoveAll(this);
+	GetSequencer()->OnActivateSequence().RemoveAll(this);
+}
+
 void FSubTrackEditor::Resize(float NewSize, UMovieSceneTrack* InTrack)
 {
 	UMovieSceneSubTrack* SubTrack = Cast<UMovieSceneSubTrack>(InTrack);
@@ -764,6 +1013,35 @@ void FSubTrackEditor::EditMetaData(UMovieSceneSubSection* Section)
 	MetaDataWindow = ExistingWindow;
 }
 
+void FSubTrackEditor::UpdateActiveMode()
+{
+	TArray<UMovieSceneSection*> Sections;
+	GetSequencer()->GetSelectedSections(Sections);
+	
+	for(UMovieSceneSection* Section : Sections)
+	{
+		if(Section->IsA(UMovieSceneSubSection::StaticClass()))
+		{
+			GLevelEditorModeTools().ActivateDefaultMode();
+			GLevelEditorModeTools().ActivateMode(FSubTrackEditorMode::ModeName);
+			return;
+		}
+	}
+	
+	TArray<UMovieSceneTrack*> Tracks;
+	GetSequencer()->GetSelectedTracks(Tracks);
+
+	for(UMovieSceneTrack* Track : Tracks)
+	{
+		if(Track->IsA(UMovieSceneSubTrack::StaticClass()))
+		{
+			GLevelEditorModeTools().ActivateDefaultMode();
+			GLevelEditorModeTools().ActivateMode(FSubTrackEditorMode::ModeName);
+			return;
+		}
+	}
+}
+
 bool FSubTrackEditor::CanAddSubSequence(const UMovieSceneSequence& Sequence) const
 {
 	UMovieSceneSequence* FocusedSequence = GetSequencer()->GetFocusedMovieSceneSequence();
@@ -807,6 +1085,225 @@ TSubclassOf<UMovieSceneSubTrack> FSubTrackEditor::GetSubTrackClass() const
 	return UMovieSceneSubTrack::StaticClass();
 }
 
+void FSubTrackEditor::UpdateOrigin(FVector InPosition, FRotator InRotation)
+{
+	if(!GetSequencer())
+	{
+		return;
+	}
+	
+	if(!GetSequencer()->IsAllowedToChange())
+	{
+		return;
+	}
+
+	UMovieSceneSequence* MovieSceneSequence = GetMovieSceneSequence();
+	if (!MovieSceneSequence)
+	{
+		return;
+	}
+	
+	// @todo Sequencer - The sequencer probably should have taken care of this
+	MovieSceneSequence->SetFlags(RF_Transactional);
+	
+	// Create a transaction record because we are about to add keys
+	const bool bShouldActuallyTransact = !GIsTransacting;		// Don't transact if we're recording in a PIE world.  That type of keyframe capture cannot be undone.
+	FScopedTransaction AutoKeyTransaction( NSLOCTEXT("AnimatablePropertyTool", "PropertyChanged", "Animatable Property Changed"), bShouldActuallyTransact);
+
+	TArray<UMovieSceneSubSection*> SectionsToKey;
+	
+	TArray<UMovieSceneSection*> SelectedSections;
+	GetSequencer()->GetSelectedSections(SelectedSections);
+
+	// Gather selected sections first, since they may be part of a track as well, and they may need to be removed if section to key is set.
+	for(UMovieSceneSection* Section : SelectedSections)
+	{
+		if(UMovieSceneSubSection* SubSection = Cast<UMovieSceneSubSection>(Section))
+		{
+			SectionsToKey.AddUnique(SubSection);
+		}
+	}
+	
+	TArray<UMovieSceneTrack*> SelectedTracks;
+	GetSequencer()->GetSelectedTracks(SelectedTracks);
+	for (UMovieSceneTrack* Track : SelectedTracks)
+	{
+		if(UMovieSceneSubTrack* SubTrack = Cast<UMovieSceneSubTrack>(Track))
+		{
+			FGeneratedTrackKeys GeneratedTrackKeys;
+			UMovieSceneSubSection* SectionToKey = Cast<UMovieSceneSubSection>(SubTrack->GetSectionToKey());
+			if(SectionToKey)
+			{
+				// Remove other sections from this track, since section to key takes precedence.
+				for(UMovieSceneSection* Section : SubTrack->GetAllSections())
+				{
+					if(UMovieSceneSubSection* SubSection = Cast<UMovieSceneSubSection>(Section))
+					{
+						SectionsToKey.Remove(SubSection);
+					}
+				}
+				SectionsToKey.AddUnique(SectionToKey);
+			}
+			else
+			{
+				for(UMovieSceneSection* Section : SubTrack->GetAllSections())
+				{
+					if(UMovieSceneSubSection* SubSection = Cast<UMovieSceneSubSection>(Section))
+					{
+						SectionsToKey.AddUnique(SubSection);
+					}
+				}
+			}
+		}
+	}
+	
+	for(UMovieSceneSubSection* SubSection : SectionsToKey)
+	{
+		EMovieSceneTransformChannel ExistingChannels = SubSection->GetMask().GetChannels();
+
+		FTransform PreviousTransform = GetTransformOriginDataForSubSection(SubSection);
+		
+		// Default to dirtying the section, since we may be adding keys 
+		bool bShouldMarkDirty = true;
+		
+		FGeneratedTrackKeys GeneratedKeys;
+		GetOriginKeys(PreviousTransform.GetLocation() + InPosition, PreviousTransform.Rotator() + InRotation, SubSection, GeneratedKeys);
+		
+		FKeyPropertyResult KeyResults = AddKeysToSection(SubSection, GetTimeForKey(), GeneratedKeys, ESequencerKeyMode::AutoKey);
+
+		// If a key wasn't created, but there is keyframe data on this section, it's preview data needs to be set to visualize the origin position for manual keyframing.
+		if(!KeyResults.bKeyCreated && SubSection->HasAnyChannelData())
+		{
+			// Preview key data is transient, and should not dirty the sequence.
+			bShouldMarkDirty = false;
+			SubSection->SetKeyPreviewPosition(PreviousTransform.GetLocation() + InPosition);
+			SubSection->SetKeyPreviewRotation(PreviousTransform.Rotator() + InRotation);
+
+			// Preview data needs to be reverted when playing the sequence, scrubbing the sequence, or navigating to a different sequence
+			// This array keeps track of sections with preview data and reverts them in FSubTrackEditor::RevertSectionPreviews
+			SectionsWithPreviews.AddUnique(SubSection);
+			
+			// Manually mark as changed since modify will not call it if not marked as dirty.
+			SubSection->MarkAsChanged();
+		}
+		SubSection->Modify(bShouldMarkDirty);
+		GetSequencer()->NotifyMovieSceneDataChanged(EMovieSceneDataChangeType::TrackValueChanged);
+	}
+	
+}
+
+void FSubTrackEditor::ResetSectionPreviews()
+{
+	bool bSectionReverted = false;
+
+	// Create a transaction record because we are about to add keys
+	const bool bShouldActuallyTransact = !GIsTransacting;		// Don't transact if we're recording in a PIE world.  That type of keyframe capture cannot be undone.
+	FScopedTransaction OriginPreviewResetTransaction( NSLOCTEXT("AnimatablePropertyTool", "PropertyChanged", "Animatable Property Changed"), bShouldActuallyTransact);
+	
+	for(UMovieSceneSubSection* SubSection : SectionsWithPreviews)
+	{
+		if(!SubSection)
+		{
+			continue;
+		}
+		SubSection->ResetKeyPreviewRotationAndLocation();
+		SubSection->Modify(false);
+		SubSection->MarkAsChanged();
+		bSectionReverted = true;
+	}
+
+	if(bSectionReverted)
+	{
+		GetSequencer()->NotifyMovieSceneDataChanged(EMovieSceneDataChangeType::TrackValueChanged);		
+	}
+
+}
+
+FTransform FSubTrackEditor::GetTransformOriginDataForSubSection(const UMovieSceneSubSection* SubSection) const
+{
+	FTransform TransformOrigin;
+
+	if(!GetSequencer())
+	{
+		return TransformOrigin;
+	}
+	const FMovieSceneRootEvaluationTemplateInstance& EvaluationTemplate = GetSequencer()->GetEvaluationTemplate();
+
+	UMovieSceneEntitySystemLinker* EntityLinker = EvaluationTemplate.GetEntitySystemLinker();
+	if(!EntityLinker)
+	{
+		return TransformOrigin;
+	}
+
+	const UMovieSceneTransformOriginSystem* TransformOriginSystem = EntityLinker->FindSystem<UMovieSceneTransformOriginSystem>();
+
+	if(TransformOriginSystem)
+	{
+		UE::MovieScene::FBuiltInComponentTypes* BuiltInComponentTypes = UE::MovieScene::FBuiltInComponentTypes::Get();
+	
+		UMovieSceneCompiledDataManager* CompiledDataManager = EvaluationTemplate.GetCompiledDataManager();
+		UMovieSceneSequence* RootSequence = EvaluationTemplate.GetSequence(GetSequencer()->GetRootTemplateID());
+		const FMovieSceneCompiledDataID DataID = CompiledDataManager->Compile(RootSequence);
+
+		const FMovieSceneSequenceHierarchy& Hierarchy = CompiledDataManager->GetHierarchyChecked(DataID);
+	
+		TArray<FMovieSceneSequenceID> SubSequenceHierarchy = GetSequencer()->GetSubSequenceHierarchy();
+
+		UE::MovieScene::FSubSequencePath Path;
+		const FMovieSceneSequenceID ParentSequenceID = SubSequenceHierarchy.Last();
+
+		Path.Reset(ParentSequenceID, &Hierarchy);
+
+		FMovieSceneSequenceID SequenceID = Path.ResolveChildSequenceID(SubSection->GetSequenceID());
+		
+		FVector OutLocation = FVector(0, 0, 0);
+		FRotator OutRotation = FRotator(0, 0, 0);
+
+		// Query the channel results directly. 
+		auto ReadSectionTransformOrigin = [&OutLocation, &OutRotation, SequenceID](const UE::MovieScene::FEntityAllocation* Allocation, UE::MovieScene::TRead<UE::MovieScene::FRootInstanceHandle> RootInstances, UE::MovieScene::TRead<FMovieSceneSequenceID> SequenceIDs,
+			UE::MovieScene::TReadOptional<double> LocationX, UE::MovieScene::TReadOptional<double> LocationY, UE::MovieScene::TReadOptional<double> LocationZ,
+			UE::MovieScene::TReadOptional<double> RotationX, UE::MovieScene::TReadOptional<double> RotationY, UE::MovieScene::TReadOptional<double> RotationZ)
+		{
+			const int32 Num = Allocation->Num();
+			
+			for(int32 Index = 0; Index < Num; ++Index)
+			{
+				if(SequenceID == SequenceIDs[Index]) 
+				{
+					OutLocation = FVector(LocationX ? LocationX[Index] : 0.0f, LocationY ? LocationY[Index] : 0.f, LocationZ ? LocationZ[Index] : 0.f);
+					OutRotation = FRotator(RotationY ? RotationY[Index] : 0.f, RotationZ ? RotationZ[Index] : 0.f, RotationX ? RotationX[Index] : 0.f);
+				}
+			}
+		};
+
+		UE::MovieScene::FEntityTaskBuilder()
+		.Read(BuiltInComponentTypes->RootInstanceHandle)
+		.Read(BuiltInComponentTypes->SequenceID)
+		.ReadOptional(BuiltInComponentTypes->DoubleResult[0])
+		.ReadOptional(BuiltInComponentTypes->DoubleResult[1])
+		.ReadOptional(BuiltInComponentTypes->DoubleResult[2])
+		.ReadOptional(BuiltInComponentTypes->DoubleResult[3])
+		.ReadOptional(BuiltInComponentTypes->DoubleResult[4])
+		.ReadOptional(BuiltInComponentTypes->DoubleResult[5])
+		.FilterAll({BuiltInComponentTypes->Tags.SubInstance})
+		.FilterNone({BuiltInComponentTypes->Tags.ImportedEntity})
+		.FilterAny(
+		{
+			BuiltInComponentTypes->DoubleResult[0],
+			BuiltInComponentTypes->DoubleResult[1],
+			BuiltInComponentTypes->DoubleResult[2],
+			BuiltInComponentTypes->DoubleResult[3],
+			BuiltInComponentTypes->DoubleResult[4],
+			BuiltInComponentTypes->DoubleResult[5]
+		})
+		.Iterate_PerAllocation(&EntityLinker->EntityManager, ReadSectionTransformOrigin);
+
+		TransformOrigin = FTransform(OutRotation, OutLocation);
+	}
+
+	return  TransformOrigin;
+}
+
 void FSubTrackEditor::GetSupportedSequenceClassPaths(TArray<FTopLevelAssetPath>& ClassPaths) const
 {
 	ClassPaths.Add(FTopLevelAssetPath(TEXT("/Script/LevelSequence"), TEXT("LevelSequence")));
diff --git a/Engine/Source/Editor/MovieSceneTools/Private/TrackEditors/TransformTrackEditor.cpp b/Engine/Source/Editor/MovieSceneTools/Private/TrackEditors/TransformTrackEditor.cpp
--- a/Engine/Source/Editor/MovieSceneTools/Private/TrackEditors/TransformTrackEditor.cpp
+++ b/Engine/Source/Editor/MovieSceneTools/Private/TrackEditors/TransformTrackEditor.cpp
@@ -52,6 +52,7 @@
 #include "IDetailKeyframeHandler.h"
 #include "ISequencerObjectChangeListener.h"
 #include "ISequencerPropertyKeyedStatus.h"
+#include "Systems/MovieSceneTransformOriginSystem.h"
 
 #define LOCTEXT_NAMESPACE "MovieScene_TransformTrack"
 
@@ -973,6 +974,40 @@ FTransform F3DTransformTrackEditor::GetTransformOrigin() const
 		// Retrieve the current origin
 		TransformOrigin = RawInterface ? RawInterface->GetTransformOrigin() : IMovieSceneTransformOrigin::Execute_BP_GetTransformOrigin(InstanceData);
 	}
+	
+	const TArray<FMovieSceneSequenceID>& Hierarchy = GetSequencer()->GetSubSequenceHierarchy();
+
+	const FMovieSceneRootEvaluationTemplateInstance& EvaluationTemplate = GetSequencer()->GetEvaluationTemplate();
+	const UMovieSceneEntitySystemLinker* EntityLinker = EvaluationTemplate.GetEntitySystemLinker();
+	if(!EntityLinker)
+	{
+		return TransformOrigin;
+	}
+
+	const UMovieSceneTransformOriginSystem* TransformOriginSystem = EntityLinker->FindSystem<UMovieSceneTransformOriginSystem>();
+
+	if(!TransformOriginSystem)
+	{
+		return TransformOrigin;
+	}
+	
+	const TSparseArray<FTransform>& TransformOrigins = TransformOriginSystem->GetTransformOriginsByInstanceID();
+	const TMap<FMovieSceneSequenceID, UE::MovieScene::FInstanceHandle> SequenceIDToInstanceHandle = TransformOriginSystem->GetSequenceIDToInstanceHandle();
+
+	// Transform Origins will be pre-multiplied at this step, so only retrieve the entry for the currently focused sub-sequence.
+	if(Hierarchy.Num())
+	{
+		const FMovieSceneSequenceID CurrentSequence = Hierarchy.Last();
+		if(SequenceIDToInstanceHandle.Contains(CurrentSequence))
+		{
+			const UE::MovieScene::FInstanceHandle CurrentHandle = SequenceIDToInstanceHandle[CurrentSequence];
+			if(TransformOrigins.IsValidIndex(CurrentHandle.InstanceID))
+			{
+				// Override the root origin 
+				TransformOrigin = TransformOrigins[CurrentHandle.InstanceID];
+			}
+		}
+	}
 
 	return TransformOrigin;
 }
diff --git a/Engine/Source/Editor/MovieSceneTools/Public/TrackEditors/SubTrackEditor.h b/Engine/Source/Editor/MovieSceneTools/Public/TrackEditors/SubTrackEditor.h
--- a/Engine/Source/Editor/MovieSceneTools/Public/TrackEditors/SubTrackEditor.h
+++ b/Engine/Source/Editor/MovieSceneTools/Public/TrackEditors/SubTrackEditor.h
@@ -8,8 +8,10 @@
 #include "Widgets/SWidget.h"
 #include "ISequencer.h"
 #include "MovieSceneTrack.h"
+#include "Tracks/MovieSceneSubTrack.h"
 #include "ISequencerSection.h"
 #include "ISequencerTrackEditor.h"
+#include "KeyframeTrackEditor.h"
 #include "MovieSceneTrackEditor.h"
 
 class AActor;
@@ -21,7 +23,7 @@ class UMovieSceneSubTrack;
  * Tools for subsequences
  */
 class MOVIESCENETOOLS_API FSubTrackEditor
-	: public FMovieSceneTrackEditor
+	: public FKeyframeTrackEditor<UMovieSceneSubTrack>
 {
 public:
 
@@ -46,7 +48,7 @@ public:
 public:
 
 	// ISequencerTrackEditor interface
-
+	virtual void ProcessKeyOperation(FFrameNumber InKeyTime, const UE::Sequencer::FKeyOperation& Operation, ISequencer& InSequencer) override;
 	virtual void BuildAddTrackMenu(FMenuBuilder& MenuBuilder) override;
 	virtual TSharedPtr<SWidget> BuildOutlinerEditWidget(const FGuid& ObjectBinding, UMovieSceneTrack* Track, const FBuildEditWidgetParams& Params) override;
 	virtual TSharedRef<ISequencerSection> MakeSectionInterface(UMovieSceneSection& SectionObject, UMovieSceneTrack& Track, FGuid ObjectBinding) override;
@@ -57,6 +59,8 @@ public:
 	virtual bool OnAllowDrop(const FDragDropEvent& DragDropEvent, FSequencerDragDropParams& DragDropParams) override;
 	virtual FReply OnDrop(const FDragDropEvent& DragDropEvent, const FSequencerDragDropParams& DragDropParams) override;
 	virtual bool IsResizable(UMovieSceneTrack* InTrack) const override;
+	virtual void OnInitialize() override;
+	virtual void OnRelease() override;
 	virtual void Resize(float NewSize, UMovieSceneTrack* InTrack) override;
 	virtual bool GetDefaultExpansionState(UMovieSceneTrack* InTrack) const override;
 
@@ -83,6 +87,9 @@ public:
 	/** Edit the section's metadata */
 	virtual void EditMetaData(UMovieSceneSubSection* Section);
 
+	/** Update the current active edit mode when a subtrack or section is selected.*/
+	void UpdateActiveMode();
+
 	/**
 	 * Check whether the given sequence can be added as a sub-sequence.
 	 *
@@ -117,6 +124,22 @@ public:
 	/** Get the UMovieSceneSubTrack class */
 	virtual TSubclassOf<UMovieSceneSubTrack> GetSubTrackClass() const;
 
+	/** Called when the editor mode has made external changes to the origin data. */
+	void UpdateOrigin(FVector InPosition, FRotator InRotation);
+
+	/** Called when sequence playback updates to revert preview data on modified secitons */
+	void ResetSectionPreviews();
+	void ResetSectionPreviews(FMovieSceneSequenceIDRef IDRef) { ResetSectionPreviews(); }
+
+	/** Query's the channel data directly (does not take parent transforms into account). Used for setting keyframes. */
+	FTransform GetTransformOriginDataForSubSection(const UMovieSceneSubSection* SubSection) const;
+
+	/** Helper function that finds the previous key. Used to "unwind" rotators */
+	int32 GetPreviousKey(FMovieSceneDoubleChannel& Channel, FFrameNumber Time);
+
+	/** Helper function to fix-up Euler rotations if they would go over 180 degrees due to interpolation. */
+	double UnwindChannel(const double& OldValue, double NewValue);
+
 protected:
 
 	/** Get the list of supported sequence class paths */
@@ -153,4 +176,13 @@ private:
 
 	/** Callback for AnimatablePropertyChanged in HandleAssetAdded. */
 	FKeyPropertyResult HandleSequenceAdded(FFrameNumber KeyTime, UMovieSceneSequence* Sequence, UMovieSceneTrack* Track, int32 RowIndex);
+
+	/** Handles adding keys to section.**/
+	void ProcessKeyOperationInternal(TArrayView<const UE::Sequencer::FKeySectionOperation> SectionsToKey, ISequencer& Sequencer, FFrameNumber KeyTime);
+
+	/** Helper for creating new keys.*/
+	void GetOriginKeys(const FVector& CurrentPosition, const FRotator& CurrentRotation, UMovieSceneSection* Section, FGeneratedTrackKeys& OutGeneratedKeys);
+
+	/** Sections this editor has added preview data to for keyframing. Used to revert data when other edits are made, before adding a key */
+	TArray<UMovieSceneSubSection*> SectionsWithPreviews;
 };
diff --git a/Engine/Source/Runtime/MovieScene/Private/EntitySystem/BuiltInComponentTypes.cpp b/Engine/Source/Runtime/MovieScene/Private/EntitySystem/BuiltInComponentTypes.cpp
--- a/Engine/Source/Runtime/MovieScene/Private/EntitySystem/BuiltInComponentTypes.cpp
+++ b/Engine/Source/Runtime/MovieScene/Private/EntitySystem/BuiltInComponentTypes.cpp
@@ -410,7 +410,7 @@ FBuiltInComponentTypes::FBuiltInComponentTypes()
 	Tags.HasUnresolvedBinding    = ComponentRegistry->NewTag(TEXT("Has Unresolved Binding"));
 	Tags.HasAssignedInitialValue = ComponentRegistry->NewTag(TEXT("Has Assigned Initial Value"));
 	Tags.Root                    = ComponentRegistry->NewTag(TEXT("Root"));
-	Tags.SubInstance             = ComponentRegistry->NewTag(TEXT("Sub Instance"));
+	Tags.SubInstance             = ComponentRegistry->NewTag(TEXT("Sub Instance"), EComponentTypeFlags::CopyToChildren);
 	Tags.ImportedEntity          = ComponentRegistry->NewTag(TEXT("Imported Entity"));
 	Tags.Finished                = ComponentRegistry->NewTag(TEXT("Finished Evaluating"));
 	Tags.Ignored                 = ComponentRegistry->NewTag(TEXT("Ignored"));
diff --git a/Engine/Source/Runtime/MovieScene/Private/Sections/MovieSceneSubSection.cpp b/Engine/Source/Runtime/MovieScene/Private/Sections/MovieSceneSubSection.cpp
--- a/Engine/Source/Runtime/MovieScene/Private/Sections/MovieSceneSubSection.cpp
+++ b/Engine/Source/Runtime/MovieScene/Private/Sections/MovieSceneSubSection.cpp
@@ -6,6 +6,8 @@
 #include "MovieSceneTrack.h"
 #include "MovieSceneSequence.h"
 #include "MovieSceneTimeHelpers.h"
+#include "MovieSceneTracksComponentTypes.h"
+#include "Channels/MovieSceneChannelProxy.h"
 #include "EntitySystem/BuiltInComponentTypes.h"
 #include "Evaluation/MovieSceneEvaluationTemplate.h"
 #include "Misc/FrameRate.h"
@@ -24,6 +26,99 @@
 
 float DeprecatedMagicNumber = TNumericLimits<float>::Lowest();
 
+#if WITH_EDITOR
+
+struct FSubSectionEditorData
+{
+
+	FText LocationGroup = NSLOCTEXT("MovieSceneSubSection", "Origin Override Location", "Origin Override Location");
+	FText RotationGroup = NSLOCTEXT("MovieSceneSubSection", "Origin Override Rotation", "Origin Override Rotation");
+	
+
+	FSubSectionEditorData(EMovieSceneTransformChannel Mask, UMovieSceneSubSection* SubSection)
+	{
+		MetaData[0].SetIdentifiers("Override.Location.X", FCommonChannelData::ChannelX, LocationGroup);
+		MetaData[0].SubPropertyPath = TEXT("Location.X");
+		MetaData[0].SortOrder = 0;
+		MetaData[0].bEnabled = EnumHasAllFlags(Mask, EMovieSceneTransformChannel::TranslationX);
+		MetaData[0].Color = FCommonChannelData::RedChannelColor;
+		MetaData[0].bCanCollapseToTrack = false;
+
+		MetaData[1].SetIdentifiers("Override.Location.Y", FCommonChannelData::ChannelY, LocationGroup);
+		MetaData[1].SubPropertyPath = TEXT("Location.Y");
+		MetaData[1].SortOrder = 1;
+		MetaData[1].bEnabled = EnumHasAllFlags(Mask, EMovieSceneTransformChannel::TranslationY);
+		MetaData[1].Color = FCommonChannelData::GreenChannelColor;
+		MetaData[1].bCanCollapseToTrack = false;
+		
+		MetaData[2].SetIdentifiers("Override.Location.Z", FCommonChannelData::ChannelZ, LocationGroup);
+		MetaData[2].SubPropertyPath = TEXT("Location.Z");
+		MetaData[2].SortOrder = 2;
+		MetaData[2].bEnabled = EnumHasAllFlags(Mask, EMovieSceneTransformChannel::TranslationZ);
+		MetaData[2].Color = FCommonChannelData::BlueChannelColor;
+		MetaData[2].bCanCollapseToTrack = false;
+
+		MetaData[3].SetIdentifiers("Override.Rotation.X", FCommonChannelData::ChannelX, RotationGroup);
+		MetaData[3].SubPropertyPath = TEXT("Rotation.X");
+		MetaData[3].SortOrder = 3;
+		MetaData[3].bEnabled = EnumHasAllFlags(Mask, EMovieSceneTransformChannel::RotationX);
+		MetaData[3].Color = FCommonChannelData::RedChannelColor;
+		MetaData[3].bCanCollapseToTrack = false;
+
+		MetaData[4].SetIdentifiers("Override.Rotation.Y", FCommonChannelData::ChannelY, RotationGroup);
+		MetaData[4].SubPropertyPath = TEXT("Rotation.Y");
+		MetaData[4].SortOrder = 4;
+		MetaData[4].bEnabled = EnumHasAllFlags(Mask, EMovieSceneTransformChannel::RotationY);
+		MetaData[4].Color = FCommonChannelData::GreenChannelColor;
+		MetaData[4].bCanCollapseToTrack = false;
+
+		MetaData[5].SetIdentifiers("Override.Rotation.Z", FCommonChannelData::ChannelZ, RotationGroup);
+		MetaData[5].SubPropertyPath = TEXT("Rotation.Z");
+		MetaData[5].SortOrder = 5;
+		MetaData[5].bEnabled = EnumHasAllFlags(Mask, EMovieSceneTransformChannel::RotationZ);
+		MetaData[5].Color = FCommonChannelData::BlueChannelColor;
+		MetaData[5].bCanCollapseToTrack = false;
+
+		ExternalValues[0].OnGetExternalValue = [SubSection](UObject& InObject, FTrackInstancePropertyBindings* Bindings) { return GetValue(SubSection, 0); };
+		ExternalValues[1].OnGetExternalValue = [SubSection](UObject& InObject, FTrackInstancePropertyBindings* Bindings) { return GetValue(SubSection, 1); };
+		ExternalValues[2].OnGetExternalValue = [SubSection](UObject& InObject, FTrackInstancePropertyBindings* Bindings) { return GetValue(SubSection, 2); };
+		ExternalValues[3].OnGetExternalValue = [SubSection](UObject& InObject, FTrackInstancePropertyBindings* Bindings) { return GetValue(SubSection, 3); };
+		ExternalValues[4].OnGetExternalValue = [SubSection](UObject& InObject, FTrackInstancePropertyBindings* Bindings) { return GetValue(SubSection, 4); };
+		ExternalValues[5].OnGetExternalValue = [SubSection](UObject& InObject, FTrackInstancePropertyBindings* Bindings) { return GetValue(SubSection, 5); };
+
+	}
+
+	static TOptional<double> GetValue(UMovieSceneSubSection* SubSection, int32 ChannelIndex)
+	{
+		if(!SubSection)
+		{
+			return TOptional<double>();
+		}
+
+		switch (ChannelIndex)
+		{
+		case 0:
+		case 1:
+		case 2:
+			return SubSection->GetKeyPreviewPosition().IsSet() ? SubSection->GetKeyPreviewPosition().GetValue()[ChannelIndex] : TOptional<double>();
+		case 3:
+			return SubSection->GetKeyPreviewRotation().IsSet() ? SubSection->GetKeyPreviewRotation().GetValue().Roll: TOptional<double>();
+		case 4:
+			return SubSection->GetKeyPreviewRotation().IsSet() ? SubSection->GetKeyPreviewRotation().GetValue().Pitch: TOptional<double>();
+		case 5:
+			return SubSection->GetKeyPreviewRotation().IsSet() ? SubSection->GetKeyPreviewRotation().GetValue().Yaw: TOptional<double>();
+		default:
+			return TOptional<double>();
+		}
+	}
+
+	FMovieSceneChannelMetaData MetaData[6];
+	TMovieSceneExternalValue<double> ExternalValues[6];
+};
+
+#endif
+
+
 /* UMovieSceneSubSection structors
  *****************************************************************************/
 
@@ -36,6 +131,12 @@ UMovieSceneSubSection::UMovieSceneSubSection(const FObjectInitializer& ObjInitia
 	NetworkMask = (uint8)(EMovieSceneServerClientMask::Server | EMovieSceneServerClientMask::Client);
 
 	SetBlendType(EMovieSceneBlendType::Absolute);
+	
+	OriginOverrideMask = EMovieSceneTransformChannel::AllTransform;
+
+#if WITH_EDITOR
+	ResetKeyPreviewRotationAndLocation();
+#endif
 }
 
 void UMovieSceneSubSection::DeleteChannels(TArrayView<const FName> ChannelNames)
@@ -71,6 +172,29 @@ EMovieSceneChannelProxyType UMovieSceneSubSection::CacheChannelProxy()
 		}
 	}
 
+#if WITH_EDITOR	
+
+	FSubSectionEditorData EditorData(OriginOverrideMask.GetChannels(), this);
+
+	Channels.Add(Translation[0], EditorData.MetaData[0], EditorData.ExternalValues[0]);
+	Channels.Add(Translation[1], EditorData.MetaData[1], EditorData.ExternalValues[1]);
+	Channels.Add(Translation[2], EditorData.MetaData[2], EditorData.ExternalValues[2]);
+	Channels.Add(Rotation[0], EditorData.MetaData[3], EditorData.ExternalValues[3]);
+	Channels.Add(Rotation[1], EditorData.MetaData[4], EditorData.ExternalValues[4]);
+	Channels.Add(Rotation[2], EditorData.MetaData[5], EditorData.ExternalValues[5]);
+
+#else
+
+	Channels.Add(Translation[0]);
+	Channels.Add(Translation[1]);
+	Channels.Add(Translation[2]);
+	Channels.Add(Rotation[0]);
+	Channels.Add(Rotation[1]);
+	Channels.Add(Rotation[2]);
+	
+#endif
+	
+
 	ChannelProxy = MakeShared<FMovieSceneChannelProxy>(MoveTemp(Channels));
 	return EMovieSceneChannelProxyType::Dynamic;
 }
@@ -168,6 +292,45 @@ bool UMovieSceneSubSection::GetValidatedInnerPlaybackRange(TRange<FFrameNumber>&
 	return false;
 }
 
+FMovieSceneSubSectionOriginOverrideMask UMovieSceneSubSection::GetMask() const
+{
+	return OriginOverrideMask;
+}
+
+void UMovieSceneSubSection::SetMask(EMovieSceneTransformChannel NewMask)
+{
+	OriginOverrideMask = NewMask;
+
+	ChannelProxy = nullptr;
+}
+
+#if WITH_EDITOR
+
+void UMovieSceneSubSection::SetKeyPreviewPosition(TOptional<FVector> InPosition)
+{
+	if(InPosition.IsSet())
+	{
+		KeyPreviewPosition = InPosition.GetValue();
+	}
+	
+}
+
+void UMovieSceneSubSection::SetKeyPreviewRotation(TOptional<FRotator> InRotation)
+{
+	if(InRotation.IsSet())
+	{
+		KeyPreviewRotation = InRotation.GetValue();
+	}
+}
+
+void UMovieSceneSubSection::ResetKeyPreviewRotationAndLocation()
+{
+	KeyPreviewPosition.Reset();
+	KeyPreviewRotation.Reset();
+}
+
+#endif
+
 TRange<FFrameNumber> UMovieSceneSubSection::GetValidatedInnerPlaybackRange(const FMovieSceneSectionParameters& SubSectionParameters, const UMovieScene& InnerMovieScene)
 {
 	const TRange<FFrameNumber> InnerPlaybackRange = InnerMovieScene.GetPlaybackRange();
@@ -527,6 +690,20 @@ FFrameNumber UMovieSceneSubSection::MapTimeToSectionFrame(FFrameTime InPosition)
 	return LocalPosition;
 }
 
+bool UMovieSceneSubSection::HasAnyChannelData() const
+{
+	bool bHasAnyData = false;
+
+	bHasAnyData |= Translation[0].HasAnyData();
+	bHasAnyData |= Translation[1].HasAnyData();
+	bHasAnyData |= Translation[2].HasAnyData();
+	bHasAnyData |= Rotation[0].HasAnyData();
+	bHasAnyData |= Rotation[1].HasAnyData();
+	bHasAnyData |= Rotation[2].HasAnyData();
+
+	return bHasAnyData;
+}
+
 void UMovieSceneSubSection::BuildDefaultSubSectionComponents(UMovieSceneEntitySystemLinker* EntityLinker, const UE::MovieScene::FEntityImportParams& Params, UE::MovieScene::FImportedEntity* OutImportedEntity) const
 {
 	using namespace UE::MovieScene;
@@ -538,12 +715,51 @@ void UMovieSceneSubSection::BuildDefaultSubSectionComponents(UMovieSceneEntitySy
 	const FSubSequencePath PathToRoot = EntityLinker->GetInstanceRegistry()->GetInstance(Params.Sequence.InstanceHandle).GetSubSequencePath();
 	FMovieSceneSequenceID ResolvedSequenceID = PathToRoot.ResolveChildSequenceID(this->GetSequenceID());
 
+	EMovieSceneTransformChannel Channels = OriginOverrideMask.GetChannels();
+
+	const bool ActiveChannelsMask[] = {
+		EnumHasAnyFlags(Channels, EMovieSceneTransformChannel::TranslationX) && Translation[0].HasAnyData(),
+		EnumHasAnyFlags(Channels, EMovieSceneTransformChannel::TranslationY) && Translation[1].HasAnyData(),
+		EnumHasAnyFlags(Channels, EMovieSceneTransformChannel::TranslationZ) && Translation[2].HasAnyData(),
+		EnumHasAnyFlags(Channels, EMovieSceneTransformChannel::RotationX) && Rotation[0].HasAnyData(),
+		EnumHasAnyFlags(Channels, EMovieSceneTransformChannel::RotationY) && Rotation[1].HasAnyData(),
+		EnumHasAnyFlags(Channels, EMovieSceneTransformChannel::RotationZ) && Rotation[2].HasAnyData(),
+	};
+
+	bool bKeyPreviewPositionIsSet = false;
+	bool bKeyPreviewRotationIsSet = false;
+
+#if WITH_EDITOR
+	bKeyPreviewPositionIsSet = KeyPreviewPosition.IsSet();
+	bKeyPreviewRotationIsSet = KeyPreviewRotation.IsSet();
+#endif
+	
 	OutImportedEntity->AddBuilder(
 		FEntityBuilder()
 		.Add(Components->SequenceID, ResolvedSequenceID)
 		.AddTag(Components->Tags.SubInstance)
 		.AddConditional(Components->HierarchicalEasingProvider, ResolvedSequenceID, bHasEasing)
+		.AddConditional(Components->DoubleChannel[0], &Translation[0], ActiveChannelsMask[0] && !bKeyPreviewPositionIsSet)
+		.AddConditional(Components->DoubleChannel[1], &Translation[1], ActiveChannelsMask[1] && !bKeyPreviewPositionIsSet)
+		.AddConditional(Components->DoubleChannel[2], &Translation[2], ActiveChannelsMask[2] && !bKeyPreviewPositionIsSet)
+		.AddConditional(Components->DoubleChannel[3], &Rotation[0], ActiveChannelsMask[3] && !bKeyPreviewRotationIsSet)
+		.AddConditional(Components->DoubleChannel[4], &Rotation[1], ActiveChannelsMask[4] && !bKeyPreviewRotationIsSet)
+		.AddConditional(Components->DoubleChannel[5], &Rotation[2], ActiveChannelsMask[5] && !bKeyPreviewRotationIsSet)
+	);
+
+	// Build Key preview entity data. Since the channel data is not written when we have preview data, this data will be used in the transform origin system.
+#if WITH_EDITOR
+	OutImportedEntity->AddBuilder(
+		FEntityBuilder()
+		.AddConditional(Components->DoubleResult[0], KeyPreviewPosition.IsSet() ? KeyPreviewPosition.GetValue().X : 0.0f, EnumHasAnyFlags(Channels, EMovieSceneTransformChannel::TranslationX) && KeyPreviewPosition.IsSet())
+		.AddConditional(Components->DoubleResult[1], KeyPreviewPosition.IsSet() ? KeyPreviewPosition.GetValue().Y : 0.0f, EnumHasAnyFlags(Channels, EMovieSceneTransformChannel::TranslationY) && KeyPreviewPosition.IsSet())
+		.AddConditional(Components->DoubleResult[2], KeyPreviewPosition.IsSet() ? KeyPreviewPosition.GetValue().Z : 0.0f, EnumHasAnyFlags(Channels, EMovieSceneTransformChannel::TranslationZ) && KeyPreviewPosition.IsSet())
+		.AddConditional(Components->DoubleResult[3], KeyPreviewPosition.IsSet() ? KeyPreviewRotation.GetValue().Roll : 0.0f, EnumHasAnyFlags(Channels, EMovieSceneTransformChannel::RotationX) && KeyPreviewRotation.IsSet())
+		.AddConditional(Components->DoubleResult[4], KeyPreviewPosition.IsSet() ? KeyPreviewRotation.GetValue().Pitch : 0.0f, EnumHasAnyFlags(Channels, EMovieSceneTransformChannel::RotationY) && KeyPreviewRotation.IsSet())
+		.AddConditional(Components->DoubleResult[5], KeyPreviewPosition.IsSet() ? KeyPreviewRotation.GetValue().Yaw : 0.0f, EnumHasAnyFlags(Channels, EMovieSceneTransformChannel::RotationX) && KeyPreviewRotation.IsSet())
 	);
+#endif
+	
 }
 
 
diff --git a/Engine/Source/Runtime/MovieScene/Private/Tracks/MovieSceneSubTrack.cpp b/Engine/Source/Runtime/MovieScene/Private/Tracks/MovieSceneSubTrack.cpp
--- a/Engine/Source/Runtime/MovieScene/Private/Tracks/MovieSceneSubTrack.cpp
+++ b/Engine/Source/Runtime/MovieScene/Private/Tracks/MovieSceneSubTrack.cpp
@@ -132,6 +132,73 @@ bool UMovieSceneSubTrack::ContainsSequence(const UMovieSceneSequence& Sequence,
 	return false;
 }
 
+TArray<UMovieSceneSection*, TInlineAllocator<4>> UMovieSceneSubTrack::FindAllSections(FFrameNumber Time) const
+{
+	TArray<UMovieSceneSection*, TInlineAllocator<4>> OverlappingSections;
+
+	for (UMovieSceneSection* Section : Sections)
+	{
+		if (MovieSceneHelpers::IsSectionKeyable(Section) && Section->GetRange().Contains(Time))
+		{
+			OverlappingSections.Add(Section);
+		}
+	}
+
+	Algo::Sort(OverlappingSections, MovieSceneHelpers::SortOverlappingSections);
+
+	return OverlappingSections;
+}
+
+#if WITH_EDITOR
+
+UMovieSceneSection* UMovieSceneSubTrack::FindSection(FFrameNumber Time) const
+{
+	TArray<UMovieSceneSection*, TInlineAllocator<4>> OverlappingSections = FindAllSections(Time);
+
+	if (OverlappingSections.Num())
+	{
+		if (SectionToKey && OverlappingSections.Contains(SectionToKey))
+		{
+			return SectionToKey;
+		}
+		else
+		{
+			return OverlappingSections[0];
+		}
+	}
+
+	return nullptr;
+}
+
+UMovieSceneSection* UMovieSceneSubTrack::FindOrExtendSection(FFrameNumber Time, float& OutWeight)
+{
+	return FindSection(Time);
+}
+
+UMovieSceneSection* UMovieSceneSubTrack::FindOrAddSection(FFrameNumber Time, bool& bSectionAdded)
+{
+	bSectionAdded = false;
+
+	UMovieSceneSection* FoundSection = FindSection(Time);
+	if (FoundSection)
+	{
+		return FoundSection;
+	}
+
+	// Add a new section that starts and ends at the same time
+	UMovieSceneSection* NewSection = CreateNewSection();
+	ensureAlwaysMsgf(NewSection->HasAnyFlags(RF_Transactional), TEXT("CreateNewSection must return an instance with RF_Transactional set! (pass RF_Transactional to NewObject)"));
+	NewSection->SetFlags(RF_Transactional);
+	NewSection->SetRange(TRange<FFrameNumber>::Inclusive(Time, Time));
+
+	Sections.Add(NewSection);
+	
+	bSectionAdded = true;
+
+	return NewSection;
+}
+
+#endif
 
 /* UMovieSceneTrack interface
  *****************************************************************************/
@@ -201,6 +268,11 @@ FText UMovieSceneSubTrack::GetDefaultDisplayName() const
 {
 	return LOCTEXT("TrackName", "Subsequences");
 }
+
+void UMovieSceneSubTrack::SetSectionToKey(UMovieSceneSection* Section)
+{
+	SectionToKey = Section;
+}
 #endif
 
 
diff --git a/Engine/Source/Runtime/MovieScene/Public/Sections/MovieSceneSubSection.h b/Engine/Source/Runtime/MovieScene/Public/Sections/MovieSceneSubSection.h
--- a/Engine/Source/Runtime/MovieScene/Public/Sections/MovieSceneSubSection.h
+++ b/Engine/Source/Runtime/MovieScene/Public/Sections/MovieSceneSubSection.h
@@ -23,6 +23,7 @@
 #include "MovieSceneFwd.h"
 #include "MovieSceneSection.h"
 #include "MovieSceneSequenceID.h"
+#include "Channels/MovieSceneDoubleChannel.h"
 #include "Templates/SubclassOf.h"
 #include "UObject/ObjectMacros.h"
 #include "UObject/ObjectPtr.h"
@@ -44,6 +45,7 @@ struct FMovieSceneSectionParameters;
 struct FMovieSceneTrackCompilerArgs;
 struct FPropertyChangedEvent;
 struct FQualifiedFrameTime;
+struct FMovieSceneTransformMask;
 
 DECLARE_DELEGATE_OneParam(FOnSequenceChanged, UMovieSceneSequence* /*Sequence*/);
 
@@ -56,6 +58,29 @@ struct FSubSequenceInstanceDataParams
 	FMovieSceneEvaluationOperand Operand;
 };
 
+USTRUCT()
+struct FMovieSceneSubSectionOriginOverrideMask
+{
+	GENERATED_BODY()
+	
+	FMovieSceneSubSectionOriginOverrideMask()
+		: Mask(0)
+	{}
+
+	FMovieSceneSubSectionOriginOverrideMask(EMovieSceneTransformChannel Channels)
+		: Mask((__underlying_type(EMovieSceneTransformChannel))Channels)
+	{}
+
+	EMovieSceneTransformChannel GetChannels() const
+	{
+		return (EMovieSceneTransformChannel) Mask;
+	}
+	
+private:
+	UPROPERTY()
+	uint32 Mask;
+};
+
 /**
  * Implements a section in sub-sequence tracks.
  */
@@ -131,6 +156,43 @@ public:
 	UFUNCTION(BlueprintCallable, Category = "Sequencer|Section")
 	MOVIESCENE_API void SetSequence(UMovieSceneSequence* Sequence);
 
+	/**
+	 * Gets the channel mask for the subsection origin overrides. 
+	 * @return The mask with bit flags for each enabled channel.
+	 */
+	MOVIESCENE_API FMovieSceneSubSectionOriginOverrideMask GetMask() const;
+
+	/**
+	 * Sets the channel mask for the subsection origin overrides.
+	 * @param MovieSceneTransformChannel the new mask.
+	 */
+	MOVIESCENE_API void SetMask(EMovieSceneTransformChannel MovieSceneTransformChannel);
+
+#if WITH_EDITOR
+	/**
+	 * Sets the external Position for controlling the transform origin in the level editor through the edtior mode's gizmo
+	 * @param InPosition An optional vector specifying the gizmo's current location
+	 * @see FSubTrackEditorMode
+	 */
+	MOVIESCENE_API void SetKeyPreviewPosition(TOptional<FVector> InPosition);
+
+	/**
+	* Sets the external rotation for controlling the transform origin in the level editor through the edtior mode's gizmo
+	* @param InRotation An optional vector specifying the gizmo's current location
+	* @see FSubTrackEditorMode
+	*/
+	MOVIESCENE_API void SetKeyPreviewRotation(TOptional<FRotator> InRotation);
+
+	/** Gets the optional value of the external position of the editor mode's gizmo */
+	TOptional<FVector> GetKeyPreviewPosition() const { return KeyPreviewPosition; }
+	
+	/** Gets the optional value of the external Rotation of the editor mode's gizmo */
+	TOptional<FRotator> GetKeyPreviewRotation() const { return KeyPreviewRotation; }
+
+	/** Resets the optional values for the gizmo's position and rotation */
+	MOVIESCENE_API void ResetKeyPreviewRotationAndLocation();
+#endif
+
 	MOVIESCENE_API virtual void PostLoad() override;
 
 #if WITH_EDITOR
@@ -142,6 +204,7 @@ public:
 #endif
 
 	MOVIESCENE_API FFrameNumber MapTimeToSectionFrame(FFrameTime InPosition) const;
+	MOVIESCENE_API bool HasAnyChannelData() const;
 
 	EMovieSceneServerClientMask GetNetworkMask() const
 	{
@@ -172,7 +235,7 @@ protected:
 	MOVIESCENE_API virtual bool PopulateEvaluationFieldImpl(const TRange<FFrameNumber>& EffectiveRange, const FMovieSceneEvaluationFieldEntityMetaData& InMetaData, FMovieSceneEntityComponentFieldBuilder* OutFieldBuilder) override;
 	MOVIESCENE_API virtual void ImportEntityImpl(UMovieSceneEntitySystemLinker* EntityLinker, const FEntityImportParams& Params, FImportedEntity* OutImportedEntity) override;
 
-	MOVIESCENE_API virtual EMovieSceneChannelProxyType CacheChannelProxy();
+	MOVIESCENE_API virtual EMovieSceneChannelProxyType CacheChannelProxy() override;
 
 public:
 
@@ -193,6 +256,25 @@ private:
 	UPROPERTY(EditAnywhere, Category="Networking", meta=(Bitmask, BitmaskEnum="/Script/MovieScene.EMovieSceneServerClientMask"))
 	uint8 NetworkMask;
 
+	UPROPERTY()
+	FMovieSceneSubSectionOriginOverrideMask OriginOverrideMask;
+
+	UPROPERTY()
+	FMovieSceneDoubleChannel Translation[3];
+	
+	UPROPERTY()
+	FMovieSceneDoubleChannel Rotation[3];
+
+#if WITH_EDITORONLY_DATA
+	/** Preview value of position used for keying. This allows for transforms without needing to commit them to the channel */
+	UPROPERTY(Transient)
+	TOptional<FVector> KeyPreviewPosition;
+
+	/** Preview value of rotation used for keying. This allows for transforms without needing to commit them to the channel */
+	UPROPERTY(Transient)
+	TOptional<FRotator> KeyPreviewRotation;
+#endif	
+
 protected:
 
 	/** Movie scene being played by this section */
diff --git a/Engine/Source/Runtime/MovieScene/Public/Tracks/MovieSceneSubTrack.h b/Engine/Source/Runtime/MovieScene/Public/Tracks/MovieSceneSubTrack.h
--- a/Engine/Source/Runtime/MovieScene/Public/Tracks/MovieSceneSubTrack.h
+++ b/Engine/Source/Runtime/MovieScene/Public/Tracks/MovieSceneSubTrack.h
@@ -65,6 +65,40 @@ public:
 	 * @return true if the sequence is in this track, false otherwise.
 	 */
 	MOVIESCENE_API bool ContainsSequence(const UMovieSceneSequence& Sequence, bool Recursively = false, const UMovieSceneSection* SectionToSkip = nullptr) const;
+	
+	/**
+	 * Finds all sections at the current time.
+	 *
+	 * @param Time The time relative to the owning movie scene where the sections should be
+	 * @return The found sections.
+	 */
+	MOVIESCENE_API TArray<UMovieSceneSection*, TInlineAllocator<4>> FindAllSections(FFrameNumber Time) const;
+
+	/**
+	 * Finds a section at the current time.
+	 *
+	 * @param Time The time relative to the owning movie scene where the section should be
+	 * @return The found section.
+	 */	
+	MOVIESCENE_API UMovieSceneSection* FindSection(FFrameNumber Time) const;
+	
+	/**
+	 * Finds a section at the current time or extends an existing one
+	 *
+	 * @param Time The time relative to the owning movie scene where the section should be
+	 * @param OutWeight The weight of the section if found
+	 * @return The found section.
+	 */
+	MOVIESCENE_API UMovieSceneSection* FindOrExtendSection(FFrameNumber Time, float& OutWeight);
+
+	/**
+	 * Finds a section at the current time, or adds one if no section is found.
+	 *
+	 * @param Time The time relative to the owning movie scene where the section should be
+	 * @param bSectionAdded Whether a section was added or not
+	 * @return The found section, or the new section.
+	 */
+	MOVIESCENE_API UMovieSceneSection* FindOrAddSection(FFrameNumber Time, bool& bSectionAdded);
 
 public:
 
@@ -83,6 +117,8 @@ public:
 
 #if WITH_EDITORONLY_DATA
 	MOVIESCENE_API virtual FText GetDefaultDisplayName() const override;
+	MOVIESCENE_API virtual UMovieSceneSection* GetSectionToKey() const override { return SectionToKey; }
+	MOVIESCENE_API virtual void SetSectionToKey(UMovieSceneSection* Section) override;
 #endif
 
 protected:
@@ -117,5 +153,9 @@ private:
 	UPROPERTY()
 	int32 RowHeight;
 
+	/** Section we should Key */
+	UPROPERTY()
+	TObjectPtr<UMovieSceneSection> SectionToKey;
+
 #endif
 };
diff --git a/Engine/Source/Runtime/MovieSceneTracks/Private/Systems/MovieSceneTransformOriginSystem.cpp b/Engine/Source/Runtime/MovieSceneTracks/Private/Systems/MovieSceneTransformOriginSystem.cpp
--- a/Engine/Source/Runtime/MovieSceneTracks/Private/Systems/MovieSceneTransformOriginSystem.cpp
+++ b/Engine/Source/Runtime/MovieSceneTracks/Private/Systems/MovieSceneTransformOriginSystem.cpp
@@ -59,6 +59,76 @@ struct FGatherTransformOrigin
 	}
 };
 
+struct FGatherTransformOriginsFromSubscenes
+{
+	// Map of child->parent instances, ordered parent-first
+	TArray<FInstanceToParentPair>* InstanceHandleToParentHandle;
+	
+	TSparseArray<FTransform>* TransformOriginsByInstanceID;
+	const FInstanceRegistry* InstanceRegistry;
+
+	// First pass on subsequence origins. Write the transform origin for relevant section to its child instance to be pre-multiplied in the post task.
+	void ForEachAllocation(const FEntityAllocation* Allocation, TRead<FRootInstanceHandle> RootInstances, TRead<FMovieSceneSequenceID> SequenceIDs,
+		TReadOptional<double> LocationX, TReadOptional<double> LocationY, TReadOptional<double> LocationZ,
+		TReadOptional<double> RotationX, TReadOptional<double> RotationY, TReadOptional<double> RotationZ) const
+	{
+		
+		const int32 Num = Allocation->Num();
+		for (int32 Index = 0; Index < Num; ++Index)
+		{
+			// The subsequence section is in the parent sequence of the instance we want to apply the transforms to
+			// Find the handle to the subinstance to write our origin data to.
+			FInstanceHandle SubInstanceHandle = InstanceRegistry->GetInstance(RootInstances[Index]).FindSubInstance(SequenceIDs[Index]);
+			
+			if (!SubInstanceHandle.IsValid())
+			{
+				continue;
+			}
+
+			const FVector    Translation(LocationX ? LocationX[Index] : 0.f, LocationY ? LocationY[Index] : 0.f, LocationZ ? LocationZ[Index] : 0.f);
+			const FRotator   Rotation(RotationY ? RotationY[Index] : 0.f, RotationZ ? RotationZ[Index] : 0.f, RotationX ? RotationX[Index] : 0.f);
+			const FTransform TransformOrigin(Rotation, Translation);
+			
+			// Set the entry for this transform origin if no prior transform is set
+			if (!TransformOriginsByInstanceID->IsValidIndex(SubInstanceHandle.InstanceID))
+			{
+				TransformOriginsByInstanceID->Insert(SubInstanceHandle.InstanceID, TransformOrigin);
+			}
+			// Transforms for the root sequence are controlled on the level sequence actor, and are gathered in FGatherTransformOrigin before this task runs
+			// If the root sequence transform origin exists, transform our offset with it, to apply additively.
+			else
+			{
+				// @todo: transform order is non-deterministic here - should we try and handle that??
+				(*TransformOriginsByInstanceID)[SubInstanceHandle.InstanceID] *= TransformOrigin;
+			}
+		}
+	}
+
+	// After all the base transforms for subsequences are gathered, multiply in their parent transform
+	void PostTask()
+	{
+		// InstanceHandleToParentHandle mapping is sorted parent first, so each parent's transform will be valid by the time its child uses it.
+		for (const FInstanceToParentPair Mapping : *InstanceHandleToParentHandle)
+		{
+			// If there's no parent transform there's nothing to do.
+			if(!TransformOriginsByInstanceID->IsValidIndex(Mapping.Parent.InstanceID))
+			{
+				continue;
+			}
+
+			// If there's a child transform it needs to be multiplied in with the parent transform.
+			if (TransformOriginsByInstanceID->IsValidIndex(Mapping.Child.InstanceID))
+			{
+				(*TransformOriginsByInstanceID)[Mapping.Child.InstanceID] *= (*TransformOriginsByInstanceID)[Mapping.Parent.InstanceID];
+			}
+			else
+			{
+				TransformOriginsByInstanceID->Insert(Mapping.Child.InstanceID, (*TransformOriginsByInstanceID)[Mapping.Parent.InstanceID]);
+			}
+		}
+	}
+};
+
 struct FAssignTransformOrigin
 {
 	const TSparseArray<FTransform>* TransformOriginsByInstanceID;
@@ -163,11 +233,24 @@ UMovieSceneTransformOriginSystem::UMovieSceneTransformOriginSystem(const FObject
 
 	Phase = ESystemPhase::Scheduling;
 
+	FBuiltInComponentTypes* BuiltInComponents = FBuiltInComponentTypes::Get();
+
+	LocationAndRotationFilterResults.Any(
+	{
+		BuiltInComponents->DoubleResult[0],
+		BuiltInComponents->DoubleResult[1],
+		BuiltInComponents->DoubleResult[2],
+		BuiltInComponents->DoubleResult[3],
+		BuiltInComponents->DoubleResult[4],
+		BuiltInComponents->DoubleResult[5]
+	});
+
 	if (HasAnyFlags(RF_ClassDefaultObject))
 	{
 		// This system relies upon anything that creates entities
 		DefineImplicitPrerequisite(GetClass(), UMovieScenePiecewiseDoubleBlenderSystem::StaticClass());
 		DefineImplicitPrerequisite(GetClass(), UMovieSceneComponentTransformSystem::StaticClass());
+		DefineImplicitPrerequisite(UDoubleChannelEvaluatorSystem::StaticClass(), GetClass());
 
 		DefineComponentConsumer(GetClass(), FBuiltInComponentTypes::Get()->DoubleResult[0]);
 		DefineComponentConsumer(GetClass(), FBuiltInComponentTypes::Get()->DoubleResult[1]);
@@ -181,6 +264,17 @@ UMovieSceneTransformOriginSystem::UMovieSceneTransformOriginSystem(const FObject
 bool UMovieSceneTransformOriginSystem::IsRelevantImpl(UMovieSceneEntitySystemLinker* InLinker) const
 {
 	using namespace UE::MovieScene;
+	
+	FBuiltInComponentTypes* BuiltInComponents = FBuiltInComponentTypes::Get();
+	
+	FEntityComponentFilter SubSequenceHasOriginFilter;
+	SubSequenceHasOriginFilter.All({BuiltInComponents->Tags.SubInstance});
+	SubSequenceHasOriginFilter.Combine(LocationAndRotationFilterResults);
+
+	if(InLinker->EntityManager.Contains(SubSequenceHasOriginFilter))
+	{
+		return true;
+	}
 
 	for (const FSequenceInstance& Instance : InLinker->GetInstanceRegistry()->GetSparseInstances())
 	{
@@ -204,6 +298,9 @@ void UMovieSceneTransformOriginSystem::OnLink()
 	UMovieSceneTransformOriginInstantiatorSystem* Instantiator = Linker->LinkSystem<UMovieSceneTransformOriginInstantiatorSystem>();
 	// This system keeps the instantiator around
 	Linker->SystemGraph.AddReference(this, Instantiator);
+
+	InstanceHandleToParentHandle.Empty();
+	
 }
 
 void UMovieSceneTransformOriginSystem::OnSchedulePersistentTasks(UE::MovieScene::IEntitySystemScheduler* TaskScheduler)
@@ -215,9 +312,36 @@ void UMovieSceneTransformOriginSystem::OnSchedulePersistentTasks(UE::MovieScene:
 	FBuiltInComponentTypes* BuiltInComponents = FBuiltInComponentTypes::Get();
 	FMovieSceneTracksComponentTypes* TracksComponents = FMovieSceneTracksComponentTypes::Get();
 
-	FEntityComponentFilter Filter;
-	Filter.All({ TracksComponents->ComponentTransform.PropertyTag, BuiltInComponents->Tags.AbsoluteBlend });
-	Filter.None({ BuiltInComponents->BlendChannelOutput });
+	FEntityComponentFilter AssignFilter;
+	AssignFilter.All({ TracksComponents->ComponentTransform.PropertyTag, BuiltInComponents->Tags.AbsoluteBlend });
+	AssignFilter.None({ BuiltInComponents->BlendChannelOutput });
+	AssignFilter.Combine(LocationAndRotationFilterResults);
+	
+	InstanceHandleToParentHandle.Empty();
+#if WITH_EDITOR
+	SequenceIDToInstanceHandle.Empty();
+#endif
+	
+	FEntityTaskBuilder()
+	.Read(BuiltInComponents->RootInstanceHandle)
+	.Read(BuiltInComponents->InstanceHandle)
+	.Read(BuiltInComponents->SequenceID)
+	.FilterAll({BuiltInComponents->Tags.SubInstance})
+	.FilterNone({BuiltInComponents->Tags.ImportedEntity}) // filter out parent entities, otherwise we'd double-up in the InstanceHandleToParentHandle mapping.
+	.Iterate_PerEntity(&Linker->EntityManager,  [this, InstanceRegistry](FRootInstanceHandle RootInstance, FInstanceHandle Instance, FMovieSceneSequenceID SequenceID)
+	{
+		FInstanceHandle SubInstanceHandle = InstanceRegistry->GetInstance(RootInstance).FindSubInstance(SequenceID);
+		InstanceHandleToParentHandle.AddUnique(FInstanceToParentPair(SubInstanceHandle, Instance));
+#if WITH_EDITOR
+		SequenceIDToInstanceHandle.Add(SequenceID, SubInstanceHandle);
+#endif
+		
+	});
+
+	// InstanceHandleToParentHandle is the iteration source for calculating transforms.
+	// This needs to be sorted parent first to ensure parent transforms are correct for when their children's transforms are calculated down-stream.
+	InstanceHandleToParentHandle.Sort();
+	
 
 	FTaskID GatherTask = TaskScheduler->AddTask<FGatherTransformOrigin>(
 		FTaskParams(TEXT("Gather Transform Origins")).ForceGameThread(),
@@ -225,6 +349,20 @@ void UMovieSceneTransformOriginSystem::OnSchedulePersistentTasks(UE::MovieScene:
 		InstanceRegistry
 	);
 
+	FTaskID GatherSubsequencesTask = FEntityTaskBuilder()
+	.Read(BuiltInComponents->RootInstanceHandle)
+	.Read(BuiltInComponents->SequenceID)
+	.ReadOptional(BuiltInComponents->DoubleResult[0])
+	.ReadOptional(BuiltInComponents->DoubleResult[1])
+	.ReadOptional(BuiltInComponents->DoubleResult[2])
+	.ReadOptional(BuiltInComponents->DoubleResult[3])
+	.ReadOptional(BuiltInComponents->DoubleResult[4])
+	.ReadOptional(BuiltInComponents->DoubleResult[5])
+	.FilterAll({BuiltInComponents->Tags.SubInstance})
+	.FilterNone({BuiltInComponents->Tags.ImportedEntity})
+	.CombineFilter(LocationAndRotationFilterResults)
+	.Fork_PerAllocation<FGatherTransformOriginsFromSubscenes>(&Linker->EntityManager, TaskScheduler, &InstanceHandleToParentHandle, &TransformOriginsByInstanceID, InstanceRegistry);
+
 	FTaskID AssignTask = FEntityTaskBuilder()
 	.Read(BuiltInComponents->InstanceHandle)
 	.Read(BuiltInComponents->BoundObject)
@@ -234,13 +372,12 @@ void UMovieSceneTransformOriginSystem::OnSchedulePersistentTasks(UE::MovieScene:
 	.WriteOptional(BuiltInComponents->DoubleResult[3])
 	.WriteOptional(BuiltInComponents->DoubleResult[4])
 	.WriteOptional(BuiltInComponents->DoubleResult[5])
-	.CombineFilter(Filter)
-	// Must contain at least one double result
-	.FilterAny({ BuiltInComponents->DoubleResult[0], BuiltInComponents->DoubleResult[1], BuiltInComponents->DoubleResult[2],
-		BuiltInComponents->DoubleResult[3], BuiltInComponents->DoubleResult[4], BuiltInComponents->DoubleResult[5] })
+	.CombineFilter(AssignFilter)
 	.Fork_PerAllocation<FAssignTransformOrigin>(&Linker->EntityManager, TaskScheduler, &TransformOriginsByInstanceID);
 
 	TaskScheduler->AddPrerequisite(GatherTask, AssignTask);
+	TaskScheduler->AddPrerequisite(GatherTask, GatherSubsequencesTask);
+	TaskScheduler->AddPrerequisite(GatherSubsequencesTask, AssignTask);
 }
 
 void UMovieSceneTransformOriginSystem::OnRun(FSystemTaskPrerequisites& InPrerequisites, FSystemSubsequentTasks& Subsequents)
diff --git a/Engine/Source/Runtime/MovieSceneTracks/Public/Systems/MovieSceneTransformOriginSystem.h b/Engine/Source/Runtime/MovieSceneTracks/Public/Systems/MovieSceneTransformOriginSystem.h
--- a/Engine/Source/Runtime/MovieSceneTracks/Public/Systems/MovieSceneTransformOriginSystem.h
+++ b/Engine/Source/Runtime/MovieSceneTracks/Public/Systems/MovieSceneTransformOriginSystem.h
@@ -4,11 +4,41 @@
 
 #include "EntitySystem/MovieSceneEntitySystem.h"
 #include "Math/Transform.h"
+#include "Misc/Build.h"
 
 #include "MovieSceneTransformOriginSystem.generated.h"
 
 struct FMovieSceneAnimTypeID;
 
+// Helper struct to define sorting behavior for parent to child mapping. 
+struct FInstanceToParentPair
+{
+	// In this struct representation the Child represents this sequence.
+	UE::MovieScene::FInstanceHandle Child;
+	UE::MovieScene::FInstanceHandle Parent;
+
+	FInstanceToParentPair(const UE::MovieScene::FInstanceHandle Child, const UE::MovieScene::FInstanceHandle Parent)
+	{
+		this->Child = Child;
+		this->Parent = Parent;
+	}
+
+	bool operator<(const FInstanceToParentPair& Other) const
+	{
+		// If this is the parent of Other, then it should be first in the mapping.
+		if(Child == Other.Parent)
+		{
+			return true;
+		}
+		return false;
+	}
+
+	bool operator==(const FInstanceToParentPair& Other) const
+	{
+		return Child == Other.Child && Parent == Other.Parent; 
+	}
+};
+
 UCLASS(MinimalAPI)
 class UMovieSceneTransformOriginInstantiatorSystem : public UMovieSceneEntitySystem
 {
@@ -43,6 +73,21 @@ private:
 private:
 
 	TSparseArray<FTransform> TransformOriginsByInstanceID;
+	TArray<FInstanceToParentPair> InstanceHandleToParentHandle;
+	UE::MovieScene::FEntityComponentFilter LocationAndRotationFilterResults;
+
+#if WITH_EDITOR
+	TMap<FMovieSceneSequenceID, UE::MovieScene::FInstanceHandle> SequenceIDToInstanceHandle;
+#endif
+	
+
+public:
+	const TSparseArray<FTransform>& GetTransformOriginsByInstanceID() const { return TransformOriginsByInstanceID; }
+
+#if WITH_EDITOR
+	const TMap<FMovieSceneSequenceID, UE::MovieScene::FInstanceHandle>& GetSequenceIDToInstanceHandle() const { return SequenceIDToInstanceHandle; }
+#endif
+	
 };
 
 
