Sequencer: Follow up fix to 36128471

Need to modify the section for the changes to propagate from the color picker

diff --git a/Engine/Source/Editor/MovieSceneTools/Private/Sections/MovieSceneSectionHelpers.cpp b/Engine/Source/Editor/MovieSceneTools/Private/Sections/MovieSceneSectionHelpers.cpp
--- a/Engine/Source/Editor/MovieSceneTools/Private/Sections/MovieSceneSectionHelpers.cpp
+++ b/Engine/Source/Editor/MovieSceneTools/Private/Sections/MovieSceneSectionHelpers.cpp
@@ -142,9 +142,9 @@ FMovieSceneKeyColorPicker::FMovieSceneKeyColorPicker(UMovieSceneSection* Section
 	PickerArgs.ParentWidget = FSlateApplication::Get().GetActiveTopLevelWindow();
 	PickerArgs.OptionalOwningDetailsView = PickerArgs.ParentWidget;
 
-	PickerArgs.OnColorCommitted = FOnLinearColorValueChanged::CreateRaw(this, &FMovieSceneKeyColorPicker::OnColorPickerPicked, RChannel, GChannel, BChannel, AChannel, InSequencer);
+	PickerArgs.OnColorCommitted = FOnLinearColorValueChanged::CreateRaw(this, &FMovieSceneKeyColorPicker::OnColorPickerPicked, Section, RChannel, GChannel, BChannel, AChannel, InSequencer);
 	PickerArgs.OnColorPickerWindowClosed = FOnWindowClosed::CreateRaw(this, &FMovieSceneKeyColorPicker::OnColorPickerClosed, Section, RChannel, GChannel, BChannel, AChannel, InSequencer);
-	PickerArgs.OnColorPickerCancelled = FOnColorPickerCancelled::CreateRaw(this, &FMovieSceneKeyColorPicker::OnColorPickerCancelled, RChannel, GChannel, BChannel, AChannel, InSequencer);
+	PickerArgs.OnColorPickerCancelled = FOnColorPickerCancelled::CreateRaw(this, &FMovieSceneKeyColorPicker::OnColorPickerCancelled, Section, RChannel, GChannel, BChannel, AChannel, InSequencer);
 
 	OpenColorPicker(PickerArgs);
 }
@@ -162,8 +162,10 @@ void UpdateOrAddKey(FMovieSceneFloatChannel* Channel, FFrameNumber KeyTime, floa
 	}
 }
 
-void FMovieSceneKeyColorPicker::OnColorPickerPicked(FLinearColor NewColor, FMovieSceneFloatChannel* RChannel, FMovieSceneFloatChannel* GChannel, FMovieSceneFloatChannel* BChannel, FMovieSceneFloatChannel* AChannel, TWeakPtr<ISequencer> WeakSequencer)
+void FMovieSceneKeyColorPicker::OnColorPickerPicked(FLinearColor NewColor, UMovieSceneSection* Section, FMovieSceneFloatChannel* RChannel, FMovieSceneFloatChannel* GChannel, FMovieSceneFloatChannel* BChannel, FMovieSceneFloatChannel* AChannel, TWeakPtr<ISequencer> WeakSequencer)
 {
+	Section->Modify();
+
 	UpdateOrAddKey(RChannel, KeyTime, NewColor.R);
 	UpdateOrAddKey(GChannel, KeyTime, NewColor.G);
 	UpdateOrAddKey(BChannel, KeyTime, NewColor.B);
@@ -211,10 +213,12 @@ void FMovieSceneKeyColorPicker::OnColorPickerClosed(const TSharedRef<SWindow>& W
 	}
 }
 
-void FMovieSceneKeyColorPicker::OnColorPickerCancelled(FLinearColor NewColor, FMovieSceneFloatChannel* RChannel, FMovieSceneFloatChannel* GChannel, FMovieSceneFloatChannel* BChannel, FMovieSceneFloatChannel* AChannel, TWeakPtr<ISequencer> WeakSequencer)
+void FMovieSceneKeyColorPicker::OnColorPickerCancelled(FLinearColor NewColor, UMovieSceneSection* Section, FMovieSceneFloatChannel* RChannel, FMovieSceneFloatChannel* GChannel, FMovieSceneFloatChannel* BChannel, FMovieSceneFloatChannel* AChannel, TWeakPtr<ISequencer> WeakSequencer)
 {
 	bColorPickerWasCancelled = true;
 	
+	Section->Modify();
+
 	// Restore the original color. No transaction will be created when the OnColorPickerClosed callback is called.
 	UpdateOrAddKey(RChannel, KeyTime, InitialColor.R);
 	UpdateOrAddKey(GChannel, KeyTime, InitialColor.G);
diff --git a/Engine/Source/Editor/MovieSceneTools/Private/Sections/MovieSceneSectionHelpers.h b/Engine/Source/Editor/MovieSceneTools/Private/Sections/MovieSceneSectionHelpers.h
--- a/Engine/Source/Editor/MovieSceneTools/Private/Sections/MovieSceneSectionHelpers.h
+++ b/Engine/Source/Editor/MovieSceneTools/Private/Sections/MovieSceneSectionHelpers.h
@@ -26,9 +26,9 @@ public:
 
 private:
 		
-	void OnColorPickerPicked(FLinearColor NewFolderColor, FMovieSceneFloatChannel* RChannel, FMovieSceneFloatChannel* GChannel, FMovieSceneFloatChannel* BChannel, FMovieSceneFloatChannel* AChannel, TWeakPtr<ISequencer> InSequencer);
+	void OnColorPickerPicked(FLinearColor NewColor, UMovieSceneSection* Section, FMovieSceneFloatChannel* RChannel, FMovieSceneFloatChannel* GChannel, FMovieSceneFloatChannel* BChannel, FMovieSceneFloatChannel* AChannel, TWeakPtr<ISequencer> InSequencer);
 	void OnColorPickerClosed(const TSharedRef<SWindow>& Window, UMovieSceneSection* Section, FMovieSceneFloatChannel* RChannel, FMovieSceneFloatChannel* GChannel, FMovieSceneFloatChannel* BChannel, FMovieSceneFloatChannel* AChannel, TWeakPtr<ISequencer> InSequencer);
-	void OnColorPickerCancelled(FLinearColor NewFolderColor, FMovieSceneFloatChannel* RChannel, FMovieSceneFloatChannel* GChannel, FMovieSceneFloatChannel* BChannel, FMovieSceneFloatChannel* AChannel, TWeakPtr<ISequencer> InSequencer);
+	void OnColorPickerCancelled(FLinearColor NewColor, UMovieSceneSection* Section, FMovieSceneFloatChannel* RChannel, FMovieSceneFloatChannel* GChannel, FMovieSceneFloatChannel* BChannel, FMovieSceneFloatChannel* AChannel, TWeakPtr<ISequencer> InSequencer);
 
 private:
 
