Fix loctext conflict in sub track editor.

diff --git a/Engine/Source/Editor/MovieSceneTools/Private/Sections/CinematicShotSection.cpp b/Engine/Source/Editor/MovieSceneTools/Private/Sections/CinematicShotSection.cpp
--- a/Engine/Source/Editor/MovieSceneTools/Private/Sections/CinematicShotSection.cpp
+++ b/Engine/Source/Editor/MovieSceneTools/Private/Sections/CinematicShotSection.cpp
@@ -395,7 +395,7 @@ void FCinematicShotSection::AddShotMenuSection(FMenuBuilder& MenuBuilder, const
 
 	TSharedPtr<ISequencer> Sequencer = GetSequencer();
 	
-	MenuBuilder.BeginSection(NAME_None, LOCTEXT("SequenceMenuText", "Active Channels"));
+	MenuBuilder.BeginSection(NAME_None, LOCTEXT("OriginChannelsText", "Active Channels"));
 	MenuBuilder.AddSubMenu(
 		LOCTEXT("AllTranslation", "Translation"), LOCTEXT("AllTranslation_ToolTip", "Causes this section to affect the translation of the transform"),
 		FNewMenuDelegate::CreateLambda([Sequencer, MakeUIAction](FMenuBuilder& SubMenuBuilder){
diff --git a/Engine/Source/Editor/MovieSceneTools/Private/TrackEditors/SubTrackEditor.cpp b/Engine/Source/Editor/MovieSceneTools/Private/TrackEditors/SubTrackEditor.cpp
--- a/Engine/Source/Editor/MovieSceneTools/Private/TrackEditors/SubTrackEditor.cpp
+++ b/Engine/Source/Editor/MovieSceneTools/Private/TrackEditors/SubTrackEditor.cpp
@@ -169,7 +169,7 @@ public:
 
 		TSharedPtr<ISequencer> Sequencer = GetSequencer();
 		
-		MenuBuilder.BeginSection(NAME_None, LOCTEXT("SequenceMenuText", "Active Channels"));
+		MenuBuilder.BeginSection(NAME_None, LOCTEXT("OriginChannelsText", "Active Channels"));
 		MenuBuilder.AddSubMenu(
 			LOCTEXT("AllTranslation", "Translation"), LOCTEXT("AllTranslation_ToolTip", "Causes this section to affect the translation of the transform"),
 			FNewMenuDelegate::CreateLambda([Sequencer, MakeUIAction](FMenuBuilder& SubMenuBuilder){
