Removing parent rotation from gizmo input when editing transform origins for subsequences.

diff --git a/Engine/Source/Editor/MovieSceneTools/Private/EditModes/SubTrackEditorMode.cpp b/Engine/Source/Editor/MovieSceneTools/Private/EditModes/SubTrackEditorMode.cpp
--- a/Engine/Source/Editor/MovieSceneTools/Private/EditModes/SubTrackEditorMode.cpp
+++ b/Engine/Source/Editor/MovieSceneTools/Private/EditModes/SubTrackEditorMode.cpp
@@ -47,7 +47,10 @@ bool FSubTrackEditorMode::InputDelta(FEditorViewportClient* InViewportClient, FV
 
 	if(bMouseButtonDown && !bAnyModifiers && CurrentAxis != EAxisList::None)
 	{
-		OnOriginValueChanged.Broadcast( InDrag, InRot);
+		// Remove parent rotation from input.
+		const FVector UnrotatedDrag = GetTransformOriginForSequence(GetFocusedSequenceID()).GetRotation().Inverse().RotateVector(InDrag);
+		
+		OnOriginValueChanged.Broadcast( UnrotatedDrag, InRot);
 		return true;
 	}
 	
@@ -74,7 +77,7 @@ bool FSubTrackEditorMode::UsesTransformWidget(UE::Widget::EWidgetMode CheckMode)
 	return FEdMode::UsesTransformWidget(CheckMode);
 }
 
-TOptional<FMovieSceneSequenceID> FSubTrackEditorMode::GetSequenceIDForSubSeciton(const UMovieSceneSubSection* InSubSection) const
+TOptional<FMovieSceneSequenceID> FSubTrackEditorMode::GetSequenceIDForSubSection(const UMovieSceneSubSection* InSubSection) const
 {
 	const TSharedPtr<ISequencer> Sequencer = WeakSequencer.Pin();
 
@@ -90,18 +93,40 @@ TOptional<FMovieSceneSequenceID> FSubTrackEditorMode::GetSequenceIDForSubSeciton
 	const FMovieSceneCompiledDataID DataID = CompiledDataManager->Compile(RootSequence);
 
 	const FMovieSceneSequenceHierarchy& Hierarchy = CompiledDataManager->GetHierarchyChecked(DataID);
-	
-	TArray<FMovieSceneSequenceID> SubSequenceHierarchy = Sequencer->GetSubSequenceHierarchy();
 
 	UE::MovieScene::FSubSequencePath Path;
-	const FMovieSceneSequenceID ParentSequenceID = SubSequenceHierarchy.Last();
+	const TOptional<FMovieSceneSequenceID> ParentSequenceID = GetFocusedSequenceID();
 
-	Path.Reset(ParentSequenceID, &Hierarchy);
+	if(!ParentSequenceID)
+	{
+		return TOptional<FMovieSceneSequenceID>();
+	}
+	
+	Path.Reset(ParentSequenceID.GetValue(), &Hierarchy);
 
 	return Path.ResolveChildSequenceID(InSubSection->GetSequenceID());
 }
 
+TOptional<FMovieSceneSequenceID> FSubTrackEditorMode::GetFocusedSequenceID() const
+{
+	const TSharedPtr<ISequencer> Sequencer = WeakSequencer.Pin();
+
+	if(!Sequencer)
+	{
+		return TOptional<FMovieSceneSequenceID>();
+	}
+
+	return Sequencer->GetFocusedTemplateID();
+}
+
 FTransform FSubTrackEditorMode::GetFinalTransformOriginForSubSection(const UMovieSceneSubSection* InSubSection) const
+{
+	const TOptional<FMovieSceneSequenceID> ChildSequenceID = GetSequenceIDForSubSection(InSubSection);
+
+	return GetTransformOriginForSequence(ChildSequenceID);
+}
+
+FTransform FSubTrackEditorMode::GetTransformOriginForSequence(const TOptional<FMovieSceneSequenceID> InSequenceID) const
 {
 	FTransform TransformOrigin = FTransform();
 
@@ -113,11 +138,9 @@ FTransform FSubTrackEditorMode::GetFinalTransformOriginForSubSection(const UMovi
 	}
 
 	const FMovieSceneRootEvaluationTemplateInstance& EvaluationTemplate = Sequencer->GetEvaluationTemplate();
-
-	const TOptional<FMovieSceneSequenceID> ChildSequenceID = GetSequenceIDForSubSeciton(InSubSection);
 	
 	const UMovieSceneEntitySystemLinker* EntityLinker = EvaluationTemplate.GetEntitySystemLinker();
-	if(!EntityLinker || !ChildSequenceID)
+	if(!EntityLinker || !InSequenceID)
 	{
 		return TransformOrigin;
 	}
@@ -132,9 +155,9 @@ FTransform FSubTrackEditorMode::GetFinalTransformOriginForSubSection(const UMovi
 	const TSparseArray<FTransform>& TransformOrigins = TransformOriginSystem->GetTransformOriginsByInstanceID();
 	const TMap<FMovieSceneSequenceID, UE::MovieScene::FInstanceHandle> SequenceIDToInstanceHandle = TransformOriginSystem->GetSequenceIDToInstanceHandle();
 
-	if(SequenceIDToInstanceHandle.Contains(ChildSequenceID.GetValue()))
+	if(SequenceIDToInstanceHandle.Contains(InSequenceID.GetValue()))
 	{
-		const UE::MovieScene::FInstanceHandle CurrentHandle = SequenceIDToInstanceHandle[ChildSequenceID.GetValue()];
+		const UE::MovieScene::FInstanceHandle CurrentHandle = SequenceIDToInstanceHandle[InSequenceID.GetValue()];
 		if(TransformOrigins.IsValidIndex(CurrentHandle.InstanceID))
 		{
 			TransformOrigin = TransformOrigins[CurrentHandle.InstanceID];
diff --git a/Engine/Source/Editor/MovieSceneTools/Private/EditModes/SubTrackEditorMode.h b/Engine/Source/Editor/MovieSceneTools/Private/EditModes/SubTrackEditorMode.h
--- a/Engine/Source/Editor/MovieSceneTools/Private/EditModes/SubTrackEditorMode.h
+++ b/Engine/Source/Editor/MovieSceneTools/Private/EditModes/SubTrackEditorMode.h
@@ -35,9 +35,13 @@ private:
 	static bool DoesSubSectionHaveTransformOverrides(const UMovieSceneSubSection& SubSection);
 
 	// Gets the sequence ID from the context of the subsection in the current hierarchy.
-	TOptional<FMovieSceneSequenceID> GetSequenceIDForSubSeciton(const UMovieSceneSubSection* InSubSection) const;
+	TOptional<FMovieSceneSequenceID> GetSequenceIDForSubSection(const UMovieSceneSubSection* InSubSection) const;
+	// Gets the sequence ID of the currently focused sequence.
+	TOptional<FMovieSceneSequenceID> GetFocusedSequenceID() const;
 	// Gets the transform origin of the provided section, after all parent transforms have been applied.
 	FTransform GetFinalTransformOriginForSubSection(const UMovieSceneSubSection* SubSection) const;
+	// Gets the transform origin corresponding to the sequence in the current hierarchy matching the provided sequence ID.
+	FTransform GetTransformOriginForSequence(TOptional<FMovieSceneSequenceID> InSequenceID) const;
 	// returns true if any actors are selected in the level editor. This is used to prevent this editor mode from being active when actors are selected.
 	bool AreAnyActorsSelected() const;
 
