Fixing issue with sub track editor delegate bindings

Changing over to shared pointer delegate to prevent broadcasting on deleted object

#rnx

diff --git a/Engine/Source/Editor/MovieSceneTools/Private/EditModes/SubTrackEditorMode.h b/Engine/Source/Editor/MovieSceneTools/Private/EditModes/SubTrackEditorMode.h
--- a/Engine/Source/Editor/MovieSceneTools/Private/EditModes/SubTrackEditorMode.h
+++ b/Engine/Source/Editor/MovieSceneTools/Private/EditModes/SubTrackEditorMode.h
@@ -57,5 +57,5 @@ private:
 	/** Used to tell if the gizmo as moved, and if the editor hit proxies need to be invalidated as a result */
 	mutable TOptional<FVector> CachedLocation;
 
-	const TArray<FName> IncompatibleEditorModes = TArray<FName>({ "EditMode.ControlRig" });
+	const TArray<FName> IncompatibleEditorModes = TArray<FName>({ "EditMode.ControlRig", "EM_Landscape" });
 };
diff --git a/Engine/Source/Editor/MovieSceneTools/Private/TrackEditors/SubTrackEditor.cpp b/Engine/Source/Editor/MovieSceneTools/Private/TrackEditors/SubTrackEditor.cpp
--- a/Engine/Source/Editor/MovieSceneTools/Private/TrackEditors/SubTrackEditor.cpp
+++ b/Engine/Source/Editor/MovieSceneTools/Private/TrackEditors/SubTrackEditor.cpp
@@ -656,7 +656,7 @@ void FSubTrackEditor::OnInitialize()
 	{
 		EditorMode->SetSequencer(GetSequencer());
 		EditorMode->GetOnOriginValueChanged().RemoveAll(this);
-		EditorMode->GetOnOriginValueChanged().AddRaw(this, &FSubTrackEditor::UpdateOrigin);
+		EditorMode->GetOnOriginValueChanged().AddSP(this, &FSubTrackEditor::UpdateOrigin);
 	}
 	GetSequencer()->GetViewModel()->GetSelection()->TrackArea.OnChanged.AddSP(this, &FSubTrackEditor::UpdateActiveMode);
 
