Sequencer: Replace PlatformCondition 'menubuilder'-based widget with a checkbox list to allow for multiple items to be checked without the widget closing or failing to refresh.

diff --git a/Engine/Source/Editor/MovieSceneTools/Private/Conditions/MovieScenePlatformConditionCustomization.cpp b/Engine/Source/Editor/MovieSceneTools/Private/Conditions/MovieScenePlatformConditionCustomization.cpp
--- a/Engine/Source/Editor/MovieSceneTools/Private/Conditions/MovieScenePlatformConditionCustomization.cpp
+++ b/Engine/Source/Editor/MovieSceneTools/Private/Conditions/MovieScenePlatformConditionCustomization.cpp
@@ -8,9 +8,11 @@
 #include "Widgets/Input/SComboButton.h"
 #include "DetailCategoryBuilder.h"
 #include "DetailWidgetRow.h"
-#include "Framework/MultiBox/MultiBoxBuilder.h"
 #include "Textures/SlateIcon.h"
-
+#include "SCheckBoxList.h"
+#include "Widgets/Images/SImage.h"
+#include "Widgets/Layout/SBox.h"
+ 
 #define LOCTEXT_NAMESPACE "MovieSceneDynamicBindingCustomization"
 
 TSharedRef<IDetailCustomization> FMovieScenePlatformConditionCustomization::MakeInstance()
@@ -26,26 +28,9 @@ void FMovieScenePlatformConditionCustomization::CustomizeDetails(IDetailLayoutBu
 
 	const TArray<const FDataDrivenPlatformInfo*>& PlatformInfos = FDataDrivenPlatformInfoRegistry::GetSortedPlatformInfos(EPlatformInfoType::TruePlatformsOnly);
 
-	const auto& GetCurrentValidPlatformNames = [](TSharedPtr<IPropertyHandle> InValidPlatformsPropertyHandle) -> TArray<FName>
-	{
-		TArray<FName> Names;
-
-		TArray<void*> RawData;
-		InValidPlatformsPropertyHandle->AccessRawData(RawData);
-
-		if (RawData.Num() > 0)
-		{
-			if (TArray<FName>* CurrentValidPlatformNamesPtr = reinterpret_cast<TArray<FName>*>(RawData[0]))
-			{
-				Names = *CurrentValidPlatformNamesPtr;
-			}
-		}
-		return Names;
-	};
-
-	auto GetComboButtonText = [SharedThis = StaticCastSharedRef<FMovieScenePlatformConditionCustomization>(AsShared()), GetCurrentValidPlatformNames]() -> FText
+	auto GetComboButtonText = [this, SharedThis = StaticCastSharedRef<FMovieScenePlatformConditionCustomization>(AsShared())]() -> FText
 	{
-		TArray<FName> CurrentValidPlatformNames = GetCurrentValidPlatformNames(SharedThis->ValidPlatformsPropertyHandle);
+		TArray<FName> CurrentValidPlatformNames = GetCurrentValidPlatformNames();
 		TArray<FText> CurrentValidPlatforms;
 		CurrentValidPlatforms.Reserve(CurrentValidPlatformNames.Num());
 
@@ -72,66 +57,134 @@ void FMovieScenePlatformConditionCustomization::CustomizeDetails(IDetailLayoutBu
 		]
 		.ValueContent()
 		[
-			
-			// Constructs the UI for bitmask property editing.
 			SNew(SComboButton)
 			.ButtonContent()
 				[
 					SNew(STextBlock)
 					.Text_Lambda(GetComboButtonText)
 				]
-			.OnGetMenuContent_Lambda([&DetailBuilder, SharedThis = StaticCastSharedRef<FMovieScenePlatformConditionCustomization>(AsShared()), PlatformInfos, GetCurrentValidPlatformNames]()
+			.OnGetMenuContent_Lambda([&DetailBuilder, this, SharedThis = StaticCastSharedRef<FMovieScenePlatformConditionCustomization>(AsShared()), PlatformInfos]()
 				{
-					TArray<FName> CurrentValidPlatformNames = GetCurrentValidPlatformNames(SharedThis->ValidPlatformsPropertyHandle);
+					TArray<FName> CurrentValidPlatformNames = GetCurrentValidPlatformNames();
 
-					FMenuBuilder MenuBuilder(false, nullptr);
+					SharedThis->CheckBoxList = SNew(SCheckBoxList)
+						.OnItemCheckStateChanged(this, &FMovieScenePlatformConditionCustomization::OnPlatformCheckChanged)
+						.IncludeGlobalCheckBoxInHeaderRow(false);
 
 					for (int32 i = 0; i < PlatformInfos.Num(); ++i)
 					{
-						MenuBuilder.AddMenuEntry(
-							FText::FromName(PlatformInfos[i]->IniPlatformName),
-							FText(),
-							FSlateIcon(FAppStyle::GetAppStyleSetName(), PlatformInfos[i]->GetIconStyleName(EPlatformIconSize::Normal)),
-							FUIAction
-							(
-								FExecuteAction::CreateLambda([SharedThis, i, PlatformInfos, CurrentValidPlatformNames, &DetailBuilder]()
-									{
-										TArray<FName> NewValidPlatformNames = CurrentValidPlatformNames;
-										if (NewValidPlatformNames.Contains(PlatformInfos[i]->IniPlatformName))
-										{
-											NewValidPlatformNames.Remove(PlatformInfos[i]->IniPlatformName);
-										}
-										else
-										{
-											NewValidPlatformNames.Add(PlatformInfos[i]->IniPlatformName);
-										}
-
-										TArray<void*> RawData;
-										SharedThis->ValidPlatformsPropertyHandle->AccessRawData(RawData);
-										if (RawData.Num() == 1)
-										{
-											if (TArray<FName>* CurrentValidPlatformNamesPtr = reinterpret_cast<TArray<FName>*>(RawData[0]))
-											{
-												*CurrentValidPlatformNamesPtr = NewValidPlatformNames;
-												SharedThis->ValidPlatformsPropertyHandle->NotifyPostChange(EPropertyChangeType::ValueSet);
-												DetailBuilder.ForceRefreshDetails();
-											}
-										}
-									}),
-								FCanExecuteAction(),
-								FIsActionChecked::CreateLambda([SharedThis, i, PlatformInfos, CurrentValidPlatformNames]() -> bool
-									{
-										return CurrentValidPlatformNames.Contains(PlatformInfos[i]->IniPlatformName);
-									})
-							),
-							NAME_None,
-							EUserInterfaceActionType::Check);
+						TSharedRef< SWidget > IconWidget = SNullWidget::NullWidget;
+						FSlateIcon SlateIcon(FAppStyle::GetAppStyleSetName(), PlatformInfos[i]->GetIconStyleName(EPlatformIconSize::Normal));
+						const FSlateBrush* IconBrush = SlateIcon.GetIcon();
+						if (IconBrush->GetResourceName() != NAME_None)
+						{
+							IconWidget = SNew(SImage)
+								.ColorAndOpacity(FSlateColor::UseSubduedForeground())
+								.Image(IconBrush);
+						}
+						const float MenuIconSize = FAppStyle::Get().GetFloat(FAppStyle::GetAppStyleSetName(), ".MenuIconSize", 16.f);
+
+						SharedThis->CheckBoxList->AddItem(
+							SNew(SHorizontalBox)
+								// Whatever we have in the icon area goes first
+								+ SHorizontalBox::Slot()
+								.AutoWidth()
+								.Padding(FMargin(2, 0, 6, 0))
+								[
+									SNew(SBox)
+										.Visibility(IconWidget != SNullWidget::NullWidget ? EVisibility::Visible : EVisibility::Collapsed)
+										.WidthOverride(MenuIconSize + 2)
+										.HeightOverride(MenuIconSize)
+										.HAlign(HAlign_Center)
+										.VAlign(VAlign_Center)
+										[
+											SNew(SBox)
+												.WidthOverride(MenuIconSize)
+												.HeightOverride(MenuIconSize)
+												[
+													IconWidget
+												]
+										]
+								]
+								+ SHorizontalBox::Slot()
+								.FillWidth(1.0f)
+								.Padding(FMargin(2, 0, 6, 0))
+								.VAlign(VAlign_Center)
+								[
+									SNew(STextBlock)
+										.TextStyle(FAppStyle::Get(), ISlateStyle::Join("Menu", ".Label"))
+										.Text(FText::FromName(PlatformInfos[i]->IniPlatformName))
+								],
+							CurrentValidPlatformNames.Contains(PlatformInfos[i]->IniPlatformName));
 					}
-
-					return MenuBuilder.MakeWidget();
+					return SharedThis->CheckBoxList.ToSharedRef();
 				})
-			
 		]; 
 }
 
+TArray<FName> FMovieScenePlatformConditionCustomization::GetCurrentValidPlatformNames()
+{
+	TArray<FName> Names;
+
+	TArray<void*> RawData;
+	ValidPlatformsPropertyHandle->AccessRawData(RawData);
+
+	if (RawData.Num() > 0)
+	{
+		if (TArray<FName>* CurrentValidPlatformNamesPtr = reinterpret_cast<TArray<FName>*>(RawData[0]))
+		{
+			Names = *CurrentValidPlatformNamesPtr;
+		}
+	}
+	return Names;
+};
+
+void FMovieScenePlatformConditionCustomization::OnPlatformCheckChanged(int32 Index)
+{
+	if (CheckBoxList.IsValid())
+	{
+		TArray<FName> CurrentValidPlatformNames = GetCurrentValidPlatformNames(); 
+		const TArray<const FDataDrivenPlatformInfo*>& PlatformInfos = FDataDrivenPlatformInfoRegistry::GetSortedPlatformInfos(EPlatformInfoType::TruePlatformsOnly);
+
+		TArray<FName> NewValidPlatformNames = CurrentValidPlatformNames;
+
+		if (Index == -1)
+		{
+			for (int32 PlatformIndex = 0; PlatformIndex < PlatformInfos.Num(); PlatformIndex++)
+			{
+				if (CheckBoxList->IsItemChecked(PlatformIndex))
+				{
+					NewValidPlatformNames.AddUnique(PlatformInfos[PlatformIndex]->IniPlatformName);
+				}
+				else
+				{
+					NewValidPlatformNames.Remove(PlatformInfos[PlatformIndex]->IniPlatformName);
+				}
+			}
+		}
+		else
+		{
+			if (CheckBoxList->IsItemChecked(Index))
+			{
+				NewValidPlatformNames.Add(PlatformInfos[Index]->IniPlatformName);
+			}
+			else
+			{
+				NewValidPlatformNames.Remove(PlatformInfos[Index]->IniPlatformName);
+			}
+		}
+
+		TArray<void*> RawData;
+		ValidPlatformsPropertyHandle->AccessRawData(RawData);
+		if (RawData.Num() == 1)
+		{
+			if (TArray<FName>* CurrentValidPlatformNamesPtr = reinterpret_cast<TArray<FName>*>(RawData[0]))
+			{
+				*CurrentValidPlatformNamesPtr = NewValidPlatformNames;
+				ValidPlatformsPropertyHandle->NotifyPostChange(EPropertyChangeType::ValueSet);
+			}
+		}
+	}
+}
+
 #undef LOCTEXT_NAMESPACE
\ No newline at end of file
diff --git a/Engine/Source/Editor/MovieSceneTools/Public/Conditions/MovieScenePlatformConditionCustomization.h b/Engine/Source/Editor/MovieSceneTools/Public/Conditions/MovieScenePlatformConditionCustomization.h
--- a/Engine/Source/Editor/MovieSceneTools/Public/Conditions/MovieScenePlatformConditionCustomization.h
+++ b/Engine/Source/Editor/MovieSceneTools/Public/Conditions/MovieScenePlatformConditionCustomization.h
@@ -7,6 +7,7 @@
 #include "PropertyHandle.h"
 
 class IDetailLayoutBuilder;
+class SCheckBoxList;
 
 class MOVIESCENETOOLS_API FMovieScenePlatformConditionCustomization : public IDetailCustomization
 {
@@ -14,5 +15,11 @@ public:
 	static TSharedRef<IDetailCustomization> MakeInstance();
 	virtual void CustomizeDetails(IDetailLayoutBuilder& DetailBuilder) override;
 
+private:
+
+	void OnPlatformCheckChanged(int32 Index);
+	TArray<FName> GetCurrentValidPlatformNames();
+
 	TSharedPtr<IPropertyHandle> ValidPlatformsPropertyHandle;
+	TSharedPtr<SCheckBoxList> CheckBoxList;
 };
