Default Transform Origin Overrides to off

Additionally check if sections are locked before moving their preview origin data

#rnx

diff --git a/Engine/Source/Editor/MovieSceneTools/Private/EditModes/SubTrackEditorMode.cpp b/Engine/Source/Editor/MovieSceneTools/Private/EditModes/SubTrackEditorMode.cpp
--- a/Engine/Source/Editor/MovieSceneTools/Private/EditModes/SubTrackEditorMode.cpp
+++ b/Engine/Source/Editor/MovieSceneTools/Private/EditModes/SubTrackEditorMode.cpp
@@ -262,8 +262,11 @@ UMovieSceneSubSection* FSubTrackEditorMode::GetSelectedSection() const
 	{
 		if(UMovieSceneSubSection* SubSection = Cast<UMovieSceneSubSection>(Section))
 		{
-			// Mirror behavior when multiple actors are selected in the level editor, and pick the last selected item.
-			SelectedSection = SubSection;
+			// Mirror behavior when multiple actors are selected in the level editor, and pick the last selected item that can still be edited.
+			if(SubSection->IsTransformOriginEditable())
+			{
+				SelectedSection = SubSection;
+			}
 		}
 	}
 
@@ -281,12 +284,26 @@ UMovieSceneSubSection* FSubTrackEditorMode::GetSelectedSection() const
 		{
 			if(SubTrack->GetSectionToKey())
 			{
-				SelectedSection = Cast<UMovieSceneSubSection>(SubTrack->GetSectionToKey());
+				UMovieSceneSubSection* SubSection = Cast<UMovieSceneSubSection>(SubTrack->GetSectionToKey());
+				if(SubSection && SubSection->IsTransformOriginEditable())
+				{
+					SelectedSection = SubSection;
+				}
 			}
 			else if(SubTrack->GetAllSections().Num())
 			{
 				// Since the first section is the section that will be keyed by default, select the first section from the track.
-				SelectedSection = Cast<UMovieSceneSubSection>(SubTrack->FindSection(PinnedSequencer.Get()->GetLocalTime().Time.FrameNumber));
+				for(UMovieSceneSection* Section : SubTrack->FindAllSections(PinnedSequencer.Get()->GetLocalTime().Time.FrameNumber))
+				{
+					if(UMovieSceneSubSection* SubSection = Cast<UMovieSceneSubSection>(Section))
+					{
+						if(SubSection->IsTransformOriginEditable())
+						{
+							SelectedSection = SubSection;
+							break;
+						}
+					}
+				}
 			}
 		}
 			
diff --git a/Engine/Source/Editor/MovieSceneTools/Private/TrackEditors/SubTrackEditor.cpp b/Engine/Source/Editor/MovieSceneTools/Private/TrackEditors/SubTrackEditor.cpp
--- a/Engine/Source/Editor/MovieSceneTools/Private/TrackEditors/SubTrackEditor.cpp
+++ b/Engine/Source/Editor/MovieSceneTools/Private/TrackEditors/SubTrackEditor.cpp
@@ -1115,7 +1115,10 @@ void FSubTrackEditor::UpdateOrigin(FVector InPosition, FRotator InRotation)
 	{
 		if(UMovieSceneSubSection* SubSection = Cast<UMovieSceneSubSection>(Section))
 		{
-			SectionsToKey.AddUnique(SubSection);
+			if(SubSection->IsTransformOriginEditable())
+			{
+				SectionsToKey.AddUnique(SubSection);
+			}
 		}
 	}
 	
@@ -1145,7 +1148,10 @@ void FSubTrackEditor::UpdateOrigin(FVector InPosition, FRotator InRotation)
 				{
 					if(UMovieSceneSubSection* SubSection = Cast<UMovieSceneSubSection>(Section))
 					{
-						SectionsToKey.AddUnique(SubSection);
+						if(SubSection->IsTransformOriginEditable())
+						{
+							SectionsToKey.AddUnique(SubSection);
+						}
 					}
 				}
 			}
diff --git a/Engine/Source/Runtime/MovieScene/Private/Sections/MovieSceneSubSection.cpp b/Engine/Source/Runtime/MovieScene/Private/Sections/MovieSceneSubSection.cpp
--- a/Engine/Source/Runtime/MovieScene/Private/Sections/MovieSceneSubSection.cpp
+++ b/Engine/Source/Runtime/MovieScene/Private/Sections/MovieSceneSubSection.cpp
@@ -132,7 +132,7 @@ UMovieSceneSubSection::UMovieSceneSubSection(const FObjectInitializer& ObjInitia
 
 	SetBlendType(EMovieSceneBlendType::Absolute);
 	
-	OriginOverrideMask = EMovieSceneTransformChannel::AllTransform;
+	OriginOverrideMask = EMovieSceneTransformChannel::None;
 
 #if WITH_EDITOR
 	ResetKeyPreviewRotationAndLocation();
@@ -684,6 +684,17 @@ FMovieSceneSubSequenceData UMovieSceneSubSection::GenerateSubSequenceData(const
 	return FMovieSceneSubSequenceData(*this);
 }
 
+#if WITH_EDITOR
+bool UMovieSceneSubSection::IsTransformOriginEditable() const
+{
+	const EMovieSceneTransformChannel SectionTransformChannels = OriginOverrideMask.GetChannels();
+
+	const bool bChannelsActive = EnumHasAnyFlags(SectionTransformChannels, EMovieSceneTransformChannel::Translation) || EnumHasAnyFlags(SectionTransformChannels, EMovieSceneTransformChannel::Rotation);
+	
+	return IsActive() && !IsLocked() && bChannelsActive;
+}
+#endif
+
 FFrameNumber UMovieSceneSubSection::MapTimeToSectionFrame(FFrameTime InPosition) const
 {
 	FFrameNumber LocalPosition = ((InPosition - Parameters.StartFrameOffset) * OuterToInnerTransform()).GetFrame();
diff --git a/Engine/Source/Runtime/MovieScene/Public/Sections/MovieSceneSubSection.h b/Engine/Source/Runtime/MovieScene/Public/Sections/MovieSceneSubSection.h
--- a/Engine/Source/Runtime/MovieScene/Public/Sections/MovieSceneSubSection.h
+++ b/Engine/Source/Runtime/MovieScene/Public/Sections/MovieSceneSubSection.h
@@ -201,6 +201,8 @@ public:
 
 	/** Delegate to fire when our sequence is changed in the property editor */
 	FOnSequenceChanged& OnSequenceChanged() { return OnSequenceChangedDelegate; }
+
+	MOVIESCENE_API bool IsTransformOriginEditable() const;
 #endif
 
 	MOVIESCENE_API FFrameNumber MapTimeToSectionFrame(FFrameTime InPosition) const;
