Fixed PVS warning

diff --git a/Engine/Source/Editor/MovieSceneTools/Private/TrackEditors/TimeWarpTrackEditor.cpp b/Engine/Source/Editor/MovieSceneTools/Private/TrackEditors/TimeWarpTrackEditor.cpp
--- a/Engine/Source/Editor/MovieSceneTools/Private/TrackEditors/TimeWarpTrackEditor.cpp
+++ b/Engine/Source/Editor/MovieSceneTools/Private/TrackEditors/TimeWarpTrackEditor.cpp
@@ -112,7 +112,7 @@ struct FScrubberHotspot : ITrackAreaHotspot
 					// @todo: Autoscroll goes wild when scrubbing warped time.
 					//        That is an intricate system that needs updating to handle warped times, but for now
 					//        we just hack it off when scrubbing.
-					if (Settings && Settings->GetAutoScrollEnabled())
+					if (Settings->GetAutoScrollEnabled())
 					{
 						Settings->SetAutoScrollEnabled(false);
 						Sequencer->OnScrubPositionChanged(ScrubTime, true, true);
