Clear origin override previews when editing from the sequencer tree view.

diff --git a/Engine/Source/Editor/MovieSceneTools/Private/TrackEditors/SubTrackEditor.cpp b/Engine/Source/Editor/MovieSceneTools/Private/TrackEditors/SubTrackEditor.cpp
--- a/Engine/Source/Editor/MovieSceneTools/Private/TrackEditors/SubTrackEditor.cpp
+++ b/Engine/Source/Editor/MovieSceneTools/Private/TrackEditors/SubTrackEditor.cpp
@@ -664,6 +664,7 @@ void FSubTrackEditor::OnInitialize()
 	GetSequencer()->OnPlayEvent().AddSP(this, &FSubTrackEditor::ResetSectionPreviews);
 	GetSequencer()->OnBeginScrubbingEvent().AddSP(this, &FSubTrackEditor::ResetSectionPreviews);
 	GetSequencer()->OnActivateSequence().AddSP(this, &FSubTrackEditor::ResetSectionPreviews);
+	GetSequencer()->OnChannelChanged().AddSP(this, &FSubTrackEditor::ResetSectionPreviews);
 }
 
 void FSubTrackEditor::OnRelease()
@@ -679,6 +680,7 @@ void FSubTrackEditor::OnRelease()
 	GetSequencer()->OnPlayEvent().RemoveAll(this);
 	GetSequencer()->OnBeginScrubbingEvent().RemoveAll(this);
 	GetSequencer()->OnActivateSequence().RemoveAll(this);
+	GetSequencer()->OnChannelChanged().RemoveAll(this);
 }
 
 void FSubTrackEditor::Resize(float NewSize, UMovieSceneTrack* InTrack)
diff --git a/Engine/Source/Editor/MovieSceneTools/Public/TrackEditors/SubTrackEditor.h b/Engine/Source/Editor/MovieSceneTools/Public/TrackEditors/SubTrackEditor.h
--- a/Engine/Source/Editor/MovieSceneTools/Public/TrackEditors/SubTrackEditor.h
+++ b/Engine/Source/Editor/MovieSceneTools/Public/TrackEditors/SubTrackEditor.h
@@ -127,6 +127,7 @@ public:
 	/** Called when sequence playback updates to revert preview data on modified secitons */
 	void ResetSectionPreviews();
 	void ResetSectionPreviews(FMovieSceneSequenceIDRef IDRef) { ResetSectionPreviews(); }
+	void ResetSectionPreviews(const FMovieSceneChannelMetaData* MetaData, UMovieSceneSection* InSection) { ResetSectionPreviews(); }
 
 	/** Query's the channel data directly (does not take parent transforms into account). Used for setting keyframes. */
 	FTransform GetTransformOriginDataForSubSection(const UMovieSceneSubSection* SubSection) const;
