Sequencer- Fix issue where platform condition changes weren't getting copied correctly back to the condition properly when using track row conditions.

diff --git a/Engine/Source/Editor/MovieSceneTools/Private/Conditions/MovieScenePlatformConditionCustomization.cpp b/Engine/Source/Editor/MovieSceneTools/Private/Conditions/MovieScenePlatformConditionCustomization.cpp
--- a/Engine/Source/Editor/MovieSceneTools/Private/Conditions/MovieScenePlatformConditionCustomization.cpp
+++ b/Engine/Source/Editor/MovieSceneTools/Private/Conditions/MovieScenePlatformConditionCustomization.cpp
@@ -119,6 +119,13 @@ void FMovieScenePlatformConditionCustomization::CustomizeDetails(IDetailLayoutBu
 					}
 					return SharedThis->CheckBoxList.ToSharedRef();
 				})
+			.OnMenuOpenChanged_Lambda([this, SharedThis = StaticCastSharedRef<FMovieScenePlatformConditionCustomization>(AsShared())](const bool IsOpen)
+				{
+					if (!IsOpen)
+					{
+						ValidPlatformsPropertyHandle->NotifyFinishedChangingProperties();
+					}
+				})
 		]; 
 }
 
@@ -180,6 +187,7 @@ void FMovieScenePlatformConditionCustomization::OnPlatformCheckChanged(int32 Ind
 		{
 			if (TArray<FName>* CurrentValidPlatformNamesPtr = reinterpret_cast<TArray<FName>*>(RawData[0]))
 			{
+				ValidPlatformsPropertyHandle->NotifyPreChange();
 				*CurrentValidPlatformNamesPtr = NewValidPlatformNames;
 				ValidPlatformsPropertyHandle->NotifyPostChange(EPropertyChangeType::ValueSet);
 			}
