Sequencer: It is no longer possible to add multiple Time Warp tracks to a sequence from the UI

diff --git a/Engine/Source/Editor/MovieSceneTools/Private/TrackEditors/TimeWarpTrackEditor.cpp b/Engine/Source/Editor/MovieSceneTools/Private/TrackEditors/TimeWarpTrackEditor.cpp
--- a/Engine/Source/Editor/MovieSceneTools/Private/TrackEditors/TimeWarpTrackEditor.cpp
+++ b/Engine/Source/Editor/MovieSceneTools/Private/TrackEditors/TimeWarpTrackEditor.cpp
@@ -343,12 +343,23 @@ TSharedPtr<UE::Sequencer::FTrackModel> FTimeWarpTrackEditor::CreateTrackModel(UM
 
 void FTimeWarpTrackEditor::ProcessKeyOperation(FFrameNumber InKeyTime, const UE::Sequencer::FKeyOperation& Operation, ISequencer& InSequencer)
 {
-	//InKeyTime = InSequencer.GetLocalTimeWarpTransform().Inverse().TryTransformTime(InKeyTime).Get(InKeyTime).RoundToFrame();
 	Operation.ApplyDefault(InKeyTime, InSequencer);
 }
 
 void FTimeWarpTrackEditor::BuildAddTrackMenu(FMenuBuilder& MenuBuilder)
 {
+	UMovieScene* FocusedMovieScene = GetFocusedMovieScene();
+	if (FocusedMovieScene == nullptr)
+	{
+		return;
+	}
+
+	if (FocusedMovieScene->IsReadOnly() || FocusedMovieScene->FindTrack<UMovieSceneTimeWarpTrack>() != nullptr)
+	{
+		return;
+	}
+
+
 	auto HandleAddTimeWarp = [this](TSubclassOf<UMovieSceneTimeWarpGetter> InClass)
 	{
 		this->HandleAddTimeWarpTrack(InClass);
@@ -382,7 +393,7 @@ void FTimeWarpTrackEditor::HandleAddTimeWarpTrack(TSubclassOf<UMovieSceneTimeWar
 		return;
 	}
 
-	if (FocusedMovieScene->IsReadOnly())
+	if (FocusedMovieScene->IsReadOnly() || FocusedMovieScene->FindTrack<UMovieSceneTimeWarpTrack>() != nullptr)
 	{
 		return;
 	}
