Sequencer: Fix camera locking mechanism. Only lock/unlock viewports that are operating with cinematic control (ie. don't do this for the little preview window for the selected camera)

This fixes a crash on/after save trying to update the locked camera in the preview window.

diff --git a/Engine/Source/Editor/MovieSceneTools/Private/TrackEditors/TransformTrackEditor.cpp b/Engine/Source/Editor/MovieSceneTools/Private/TrackEditors/TransformTrackEditor.cpp
--- a/Engine/Source/Editor/MovieSceneTools/Private/TrackEditors/TransformTrackEditor.cpp
+++ b/Engine/Source/Editor/MovieSceneTools/Private/TrackEditors/TransformTrackEditor.cpp
@@ -433,7 +433,7 @@ void F3DTransformTrackEditor::OnPreSaveWorld(UWorld* World)
 	// bound camera actors. If they are, associate them with the binding GUID.
 	for (FLevelEditorViewportClient* LevelVC : GEditor->GetLevelViewportClients())
 	{
-		if (LevelVC && LevelVC->GetViewMode() != VMI_Unknown)
+		if (LevelVC && LevelVC->GetViewMode() != VMI_Unknown && LevelVC->AllowsCinematicControl())
 		{
 			AActor* ActorLock = LevelVC->GetActiveActorLock().Get();
 			if (!ActorLock)
