Optimize FSubTrackEditorMode::AreAnyActorsSelected() that was taking 60-100ms+ per frame when 30,000 actors in a world
#rnx

#tests ran editor before/after and profiled

diff --git a/Engine/Source/Editor/MovieSceneTools/Private/EditModes/SubTrackEditorMode.cpp b/Engine/Source/Editor/MovieSceneTools/Private/EditModes/SubTrackEditorMode.cpp
--- a/Engine/Source/Editor/MovieSceneTools/Private/EditModes/SubTrackEditorMode.cpp
+++ b/Engine/Source/Editor/MovieSceneTools/Private/EditModes/SubTrackEditorMode.cpp
@@ -2,6 +2,8 @@
 
 #include "SubTrackEditorMode.h"
 
+#include "EditorModeManager.h"
+#include "Engine/Selection.h"
 #include "EngineUtils.h"
 #include "UnrealClient.h"
 #include "EditorViewportClient.h"
@@ -169,11 +171,14 @@ FTransform FSubTrackEditorMode::GetTransformOriginForSequence(const TOptional<FM
 
 bool FSubTrackEditorMode::AreAnyActorsSelected() const
 {
-	FSelectedActorIterator SelectedActorIter(GetWorld());
-	if(SelectedActorIter)
+	if (Owner)
 	{
-		return true;
+		if (USelection* SelectedActors = Owner->GetSelectedActors())
+		{
+			return SelectedActors->Num() > 0;
+		}
 	}
+
 	return false;
 }
 
