Sequencer: Typo for icon: Indicator -> Indicators

#rnx

diff --git a/Engine/Source/Editor/MovieSceneTools/Private/MVVM/ViewModels/TimeWarpChannelModel.cpp b/Engine/Source/Editor/MovieSceneTools/Private/MVVM/ViewModels/TimeWarpChannelModel.cpp
--- a/Engine/Source/Editor/MovieSceneTools/Private/MVVM/ViewModels/TimeWarpChannelModel.cpp
+++ b/Engine/Source/Editor/MovieSceneTools/Private/MVVM/ViewModels/TimeWarpChannelModel.cpp
@@ -28,7 +28,7 @@ TSharedPtr<SWidget> FTimeWarpChannelModel::CreateOutlinerViewForColumn(const FCr
 	{
 		return MakeButton(
 			LOCTEXT("ChangeTimeWarpToolTip", "Change Time Warp to utilize a different curve type"),
-			FAppStyle::GetBrush("Sequencer.Outliner.Indicator.TimeWarp"),
+			FAppStyle::GetBrush("Sequencer.Outliner.Indicators.TimeWarp"),
 			FOnGetContent::CreateSP(this, &FTimeWarpChannelModel::BuildReplaceTimeWarpSubMenu),
 			OutlinerItem.AsModel());
 	}
