Animation: Add support for recording timecode when baking an animation sequence.

This is intended to support motion capture workflows, where a livelink track is first recorded (which saves all keys at the original sample rate without a dependence on the Engine frame rate), then the driven actor can be added to the sequence (e.g. an animation blueprint), then right click to start the animation sequence baking process, make sure the "bake timecode" option is checked, then proceed to bake.

* The Bake Animation sequence now has a checkbox to select if the timecode (and slate) is to be baked to the animation (using the bone method). Defaults to false.
* The recorder can now be assigned a custom function to retrieve the current time.
* To populate the slate name, the take metadata is searched for in the level sequence. If none is found, it uses the name of the sequence.
* Timecode is obtained by converting the sequence frame number to timecode.
* The timecode rate is taken from the GenerateDefaultTimecodeFrameRate project setting, but can be overriden from the Bake dialog.
* Only non-drop is supported. If support for drop frame is desired, another project setting could be added and used.

We may default to using the source timecode track that was recently added.

#rn Baking animation sequences can now bake timecode and slate name.

diff --git a/Engine/Plugins/Animation/ControlRig/Source/ControlRigEditor/Private/Sequencer/ControlRigParameterTrackEditor.cpp b/Engine/Plugins/Animation/ControlRig/Source/ControlRigEditor/Private/Sequencer/ControlRigParameterTrackEditor.cpp
--- a/Engine/Plugins/Animation/ControlRig/Source/ControlRigEditor/Private/Sequencer/ControlRigParameterTrackEditor.cpp
+++ b/Engine/Plugins/Animation/ControlRig/Source/ControlRigEditor/Private/Sequencer/ControlRigParameterTrackEditor.cpp
@@ -1413,7 +1413,7 @@ void FControlRigParameterTrackEditor::BakeInvertedPose(UControlRig* InControlRig
 	};
 
 	FInitAnimationCB InitCallback = FInitAnimationCB::CreateLambda([]{});
-	FStartAnimationCB StartCallback = FStartAnimationCB::CreateLambda([AnimTrackAdapter, KeyFrame]
+	FStartAnimationCB StartCallback = FStartAnimationCB::CreateLambda([AnimTrackAdapter, KeyFrame](FFrameNumber FrameNumber)
 	{
 		KeyFrame(AnimTrackAdapter.GetLocalStartFrame());
 	});
diff --git a/Engine/Source/Editor/MovieSceneTools/Private/MovieSceneToolHelpers.cpp b/Engine/Source/Editor/MovieSceneTools/Private/MovieSceneToolHelpers.cpp
--- a/Engine/Source/Editor/MovieSceneTools/Private/MovieSceneToolHelpers.cpp
+++ b/Engine/Source/Editor/MovieSceneTools/Private/MovieSceneToolHelpers.cpp
@@ -3895,7 +3895,6 @@ bool MovieSceneToolHelpers::BakeToSkelMeshToCallbacks(const FAnimExportSequenceP
 	UnFbx::FLevelSequenceAnimTrackAdapter AnimTrackAdapter(AESP.Player, AESP.MovieSceneSequence, AESP.RootMovieSceneSequence, AESP.RootToLocalTransform);
 	if (ExportOptions->bUseCustomTimeRange)
 	{
-
 		const FFrameRate TickResolution = MovieScene->GetTickResolution();
 		const FFrameRate DisplayResolution = ExportOptions->CustomDisplayRate;
 		const FFrameNumber StartFrameInTick = FFrameRate::TransformTime(FFrameTime(ExportOptions->CustomStartFrame), DisplayResolution, TickResolution).FloorToFrame();
@@ -3906,13 +3905,11 @@ bool MovieSceneToolHelpers::BakeToSkelMeshToCallbacks(const FAnimExportSequenceP
 		}
 		AnimTrackAdapter.SetRange(StartFrameInTick,EndFrameInTick);
 	}
+
 	int32 LocalStartFrame = AnimTrackAdapter.GetLocalStartFrame();
-	int32 StartFrame = AnimTrackAdapter.GetStartFrame();
 	int32 AnimationLength = AnimTrackAdapter.GetLength();
 	float FrameRate = AnimTrackAdapter.GetFrameRate();
 	float DeltaTime = 1.0f / FrameRate;
-	FFrameRate SampleRate = MovieScene->GetDisplayRate();
-
 
 	//If we are running with a live link track we need to do a few things.
 	// 1. First test to see if we have one, only way to really do that is to see if we have a source that has the `Sequencer Live Link Track`.  We also evalute the first frame in case we are out of range and the sources aren't created yet.
@@ -3965,7 +3962,7 @@ bool MovieSceneToolHelpers::BakeToSkelMeshToCallbacks(const FAnimExportSequenceP
 	
 	//BeginRecording, which happens in the Start Callback below, records a frame so need to set things up first at first frame, even if we had any delay or warmup
 	TickFrame(LocalStartFrame, DeltaTime, MovieScene, AnimTrackAdapter, BakeHelpers, SkelMeshComps, LiveLinkClient, SourceAndMode);
-	StartCallback.ExecuteIfBound();
+	StartCallback.ExecuteIfBound(LocalStartFrame);
 
 	for (int32 FrameCount = 1; FrameCount <= AnimationLength; ++FrameCount)
 	{
@@ -4004,7 +4001,7 @@ bool MovieSceneToolHelpers::BakeToSkelMeshToCallbacks(const FAnimExportSequenceP
 	return true;
 }
 
-//5.5 deprecrated
+//5.5 deprecated
 bool MovieSceneToolHelpers::ExportToAnimSequence(UAnimSequence* AnimSequence, UAnimSeqExportOption* ExportOptions, UMovieScene* MovieScene, IMovieScenePlayer* Player,
 	USkeletalMeshComponent* SkelMesh, FMovieSceneSequenceIDRef& Template, FMovieSceneSequenceTransform& RootToLocalTransform)
 {
@@ -4023,6 +4020,56 @@ bool MovieSceneToolHelpers::ExportToAnimSequence(UAnimSequence* AnimSequence, UA
 		return false;
 	}
 }
+
+/**
+ * Retrieves a string from the MovieSceneSequence corresponding to the source takename.
+ * 
+ * It tries to find the UTakeMetaData slate via reflection, since UTakeMetaData is defined in a plugin.
+ * If that fails, it returns the name of the MovieSceneSequence.
+ */
+static FString GetSlateFromMovieSceneSequence(UMovieSceneSequence* MovieSceneSequence)
+{
+	if (!MovieSceneSequence)
+	{
+		return FString();
+	}
+
+	ULevelSequence* LevelSequence = Cast<ULevelSequence>(MovieSceneSequence);
+
+	if (!LevelSequence)
+	{
+		return MovieSceneSequence->GetName();
+	}
+
+	// Find the UTakeMetaData class dynamically, since it is defined in a plugin.
+	static UClass* TakeMetaDataClass = FindObject<UClass>(nullptr, TEXT("/Script/TakesCore.TakeMetaData"), false);
+	if (!TakeMetaDataClass)
+	{
+		return LevelSequence->GetName();
+	}
+
+	if (UObject* MetaDataObject = LevelSequence->FindMetaDataByClass(TakeMetaDataClass))
+	{
+		if (const FProperty* SlateProperty = TakeMetaDataClass->FindPropertyByName("Slate"))
+		{
+			if (const FStrProperty* SlateStrProperty = CastField<FStrProperty>(SlateProperty))
+			{
+				return SlateStrProperty->GetPropertyValue_InContainer(MetaDataObject);
+			}
+			else
+			{
+				UE_LOG(LogMovieScene, Error, TEXT("TakesMetaData::Slate property was found, but it unexpectedly was not a String property."));
+			}
+		}
+		else
+		{
+			UE_LOG(LogMovieScene, Error, TEXT("TakesMetaData class was found, but it unexpectedly did not have a Slate property."));
+		}
+	}
+
+	return LevelSequence->GetName();
+}
+
 bool MovieSceneToolHelpers::ExportToAnimSequence(UAnimSequence* AnimSequence, UAnimSeqExportOption* ExportOptions, const FAnimExportSequenceParameters& AESP,
 	USkeletalMeshComponent* SkelMeshComp)
 {
@@ -4031,8 +4078,23 @@ bool MovieSceneToolHelpers::ExportToAnimSequence(UAnimSequence* AnimSequence, UA
 		UE_LOG(LogMovieScene, Error, TEXT("MovieSceneToolHelpers::ExportToAnimSequence All parameters must be valid."));
 		return false;
 	}
+
 	FAnimRecorderInstance AnimationRecorder;
-	FFrameRate SampleRate = AESP.MovieSceneSequence->GetMovieScene()->GetDisplayRate();
+	const FFrameRate SampleRate = AESP.MovieSceneSequence->GetMovieScene()->GetDisplayRate();
+
+	// Used by function that retrieves the current timecode 
+	FQualifiedFrameTime CurrentTimecode;
+
+	// Prefer the rate overriden by the user, otherwise use the project setting.
+	if (ExportOptions->bTimecodeRateOverride)
+	{
+		CurrentTimecode.Rate = ExportOptions->OverrideTimecodeRate;
+	}
+	else if (GEngine)
+	{
+		CurrentTimecode.Rate = GEngine->GenerateDefaultTimecodeFrameRate;
+	}
+
 	FInitAnimationCB InitCallback = FInitAnimationCB::CreateLambda([&AnimationRecorder, SampleRate, ExportOptions, SkelMeshComp, AnimSequence]
 		{
 			FAnimationRecordingSettings RecordingSettings;
@@ -4062,28 +4124,78 @@ bool MovieSceneToolHelpers::ExportToAnimSequence(UAnimSequence* AnimSequence, UA
 			AnimationRecorder.Init(SkelMeshComp, AnimSequence, nullptr, RecordingSettings);
 		});
 
-
-	FStartAnimationCB StartCallback = FStartAnimationCB::CreateLambda([SkelMeshComp, &AnimationRecorder]
+	FStartAnimationCB StartCallback = FStartAnimationCB::CreateLambda([SkelMeshComp, &AnimationRecorder, &CurrentTimecode, &SampleRate, ExportOptions](FFrameNumber FrameNumber)
 		{
+			if (ExportOptions->bBakeTimecode)
+			{
+				ensureMsgf(AnimationRecorder.Recorder.IsValid(), TEXT("Unable to bake timecode due to Invalid Recorder."));
+
+				if (AnimationRecorder.Recorder.IsValid())
+				{
+					// Update the time because AnimationRecorder.BeginRecording() will key the first frame.
+					CurrentTimecode.Time = FQualifiedFrameTime(FFrameTime(FrameNumber), SampleRate).ConvertTo(CurrentTimecode.Rate);
+
+					AnimationRecorder.Recorder->SetCurrentFrameTimeGetter([&CurrentTimecode]() -> TOptional<FQualifiedFrameTime>
+						{
+							return CurrentTimecode;
+						});
+				}
+			}
+
 			AnimationRecorder.BeginRecording();
 			SkelMeshComp->UpdateLODStatus();
 		});
 
-	FTickAnimationCB TickCallback = FTickAnimationCB::CreateLambda([&AnimationRecorder](float DeltaTime, FFrameNumber FrameNumber)
+	FTickAnimationCB TickCallback = FTickAnimationCB::CreateLambda([&AnimationRecorder, &CurrentTimecode, SampleRate, ExportOptions](float DeltaTime, FFrameNumber FrameNumber)
 		{
-			AnimationRecorder.Update(DeltaTime);
+			// Update current timecode, which will be read during AnimationRecorder via the custom current time getter.
+			if (ExportOptions->bBakeTimecode)
+			{
+				CurrentTimecode.Time = FQualifiedFrameTime(FFrameTime(FrameNumber), SampleRate).ConvertTo(CurrentTimecode.Rate);
+			}
 
+			AnimationRecorder.Update(DeltaTime);
 		});
 
-	FEndAnimationCB EndCallback = FEndAnimationCB::CreateLambda([&AnimationRecorder]
+	FEndAnimationCB EndCallback = FEndAnimationCB::CreateLambda([&AnimationRecorder, ExportOptions, AnimSequence, SkelMeshComp, AESP]
 		{
+			if (ExportOptions->bBakeTimecode)
+			{
+				FProcessRecordedTimeParams TimeParams;
+				
+				// Use the timecode attribute names in the animation settings for consistency.
+				if (const UAnimationSettings* AnimationSettings = UAnimationSettings::Get())
+				{
+					TimeParams.HoursName = AnimationSettings->BoneTimecodeCustomAttributeNameSettings.HourAttributeName.ToString();
+					TimeParams.MinutesName = AnimationSettings->BoneTimecodeCustomAttributeNameSettings.MinuteAttributeName.ToString();
+					TimeParams.SecondsName = AnimationSettings->BoneTimecodeCustomAttributeNameSettings.SecondAttributeName.ToString();
+					TimeParams.FramesName = AnimationSettings->BoneTimecodeCustomAttributeNameSettings.FrameAttributeName.ToString();
+					TimeParams.SubFramesName = AnimationSettings->BoneTimecodeCustomAttributeNameSettings.SubframeAttributeName.ToString();
+					TimeParams.RateName = AnimationSettings->BoneTimecodeCustomAttributeNameSettings.RateAttributeName.ToString();
+					TimeParams.SlateName = AnimationSettings->BoneTimecodeCustomAttributeNameSettings.TakenameAttributeName.ToString();
+
+					TimeParams.Slate = GetSlateFromMovieSceneSequence(AESP.MovieSceneSequence);
+				}
+
+				AnimationRecorder.ProcessRecordedTimes(AnimSequence, SkelMeshComp, FTimecodeBoneMethod(), TimeParams);
+			}
+
 			const bool bShowAnimationAssetCreatedToast = false;
 			AnimationRecorder.FinishRecording(bShowAnimationAssetCreatedToast);
-		});
 
+			// Clean up the getter since it references a local variable and AnimationRecorder.Recorder is managed by a shared pointer.
+			if (ExportOptions->bBakeTimecode && AnimationRecorder.Recorder.IsValid())
+			{
+				AnimationRecorder.Recorder->SetCurrentFrameTimeGetter([]() -> TOptional<FQualifiedFrameTime>
+					{
+						return FApp::GetCurrentFrameTime();
+					});
+			}
+		});
 
 	MovieSceneToolHelpers::BakeToSkelMeshToCallbacks(AESP, SkelMeshComp, ExportOptions,
 		InitCallback, StartCallback, TickCallback, EndCallback);
+
 	return AnimSequence->GetDataModel()->HasBeenPopulated();
 }
 
diff --git a/Engine/Source/Editor/MovieSceneTools/Public/MovieSceneToolHelpers.h b/Engine/Source/Editor/MovieSceneTools/Public/MovieSceneToolHelpers.h
--- a/Engine/Source/Editor/MovieSceneTools/Public/MovieSceneToolHelpers.h
+++ b/Engine/Source/Editor/MovieSceneTools/Public/MovieSceneToolHelpers.h
@@ -77,7 +77,7 @@ struct FAnimExportSequenceParameters
 //callback's used by skel mesh recorders
 
 DECLARE_DELEGATE(FInitAnimationCB);
-DECLARE_DELEGATE(FStartAnimationCB);
+DECLARE_DELEGATE_OneParam(FStartAnimationCB, FFrameNumber);
 DECLARE_DELEGATE_TwoParams(FTickAnimationCB, float, FFrameNumber);
 DECLARE_DELEGATE(FEndAnimationCB);
 
diff --git a/Engine/Source/Editor/SequenceRecorder/Private/AnimationRecorder.cpp b/Engine/Source/Editor/SequenceRecorder/Private/AnimationRecorder.cpp
--- a/Engine/Source/Editor/SequenceRecorder/Private/AnimationRecorder.cpp
+++ b/Engine/Source/Editor/SequenceRecorder/Private/AnimationRecorder.cpp
@@ -99,6 +99,11 @@ bool FAnimationRecorder::SetAnimCompressionScheme(UAnimBoneCompressionSettings*
 	return false;
 }
 
+void FAnimationRecorder::SetCurrentFrameTimeGetter(TFunction<TOptional<FQualifiedFrameTime>()> InGetCurrentTimeFunction)
+{
+	GetCurrentFrameTimeFunction = MoveTemp(InGetCurrentTimeFunction);
+}
+
 // Internal. Pops up a dialog to get saved asset path
 static bool PromptUserForAssetDetails(FString& AssetPath, FString& AssetName, FFrameRate& OutSampleRate, float& OutMaximumDuration)
 {
@@ -1026,7 +1031,18 @@ bool FAnimationRecorder::Record(USkeletalMeshComponent* Component, FTransform co
 			}
 		}
 
-		TOptional<FQualifiedFrameTime> CurrentTime = FApp::GetCurrentFrameTime();
+		TOptional<FQualifiedFrameTime> CurrentTime;
+		
+		// When set, use the custom time getter. Otherwise fall back to FApp::GetCurrentFrameTime.
+		if (GetCurrentFrameTimeFunction)
+		{
+			CurrentTime = GetCurrentFrameTimeFunction();
+		}
+		else
+		{
+			CurrentTime = FApp::GetCurrentFrameTime();
+		}
+
 		if (CurrentTime)
 		{
 			FQualifiedFrameTime TimeToRecord = *CurrentTime;
diff --git a/Engine/Source/Editor/SequenceRecorder/Public/AnimationRecorder.h b/Engine/Source/Editor/SequenceRecorder/Public/AnimationRecorder.h
--- a/Engine/Source/Editor/SequenceRecorder/Public/AnimationRecorder.h
+++ b/Engine/Source/Editor/SequenceRecorder/Public/AnimationRecorder.h
@@ -91,6 +91,9 @@ private:
 	/** Unique notify states added to this sequence during recording */
 	TMap<UAnimNotifyState*, UAnimNotifyState*> UniqueNotifyStates;
 
+	/** This function returns the current timecode that the recorder should use to label the current frame */
+	TFunction<TOptional<FQualifiedFrameTime>()> GetCurrentFrameTimeFunction;
+
 	struct FRecordedAnimNotify
 	{
 		FRecordedAnimNotify(const FAnimNotifyEvent& InNewNotifyEvent, const FAnimNotifyEvent* InOriginalNotifyEvent, float InAnimNotifyStartTime, float InAnimNotifyEndTime)
@@ -154,6 +157,9 @@ public:
 
 	bool SetAnimCompressionScheme(UAnimBoneCompressionSettings* Settings);
 
+	/** Sets the function that should be used to get the current timecode for the current frame. */
+	void SetCurrentFrameTimeGetter(TFunction<TOptional<FQualifiedFrameTime>()> InGetCurrentTimeFunction);
+
 	const FTransform& GetInitialRootTransform() const { return InitialRootTransform; }
 
 	UE_DEPRECATED(5.5, "Use the ProcessRecordedTimes method that takes a FProcessRecordedTimeParams struct.")
diff --git a/Engine/Source/Editor/UnrealEd/Classes/Exporters/AnimSeqExportOption.h b/Engine/Source/Editor/UnrealEd/Classes/Exporters/AnimSeqExportOption.h
--- a/Engine/Source/Editor/UnrealEd/Classes/Exporters/AnimSeqExportOption.h
+++ b/Engine/Source/Editor/UnrealEd/Classes/Exporters/AnimSeqExportOption.h
@@ -70,6 +70,18 @@ public:
 	UPROPERTY(EditAnywhere, BlueprintReadWrite, AdvancedDisplay, Category = Export)
 	bool bTransactRecording = true;
 
+	/** Set to true if sequence timecode should be baked into the sequence. Timecode rate will default to the project setting "Generate Default Timecode Frame Rate" */
+	UPROPERTY(EditAnywhere, BlueprintReadWrite, AdvancedDisplay, Category = Export)
+	bool bBakeTimecode = false;
+
+	/** Set to true if the timecode rate should be overridden with the specified value. */
+	UPROPERTY(EditAnywhere, BlueprintReadWrite, AdvancedDisplay, meta = (InlineEditConditionToggle), Category = Export)
+	bool bTimecodeRateOverride = false;
+
+	/** Overriding timecode rate to be used when baking. */
+	UPROPERTY(EditAnywhere, BlueprintReadWrite, AdvancedDisplay, meta = (EditCondition = bTimecodeRateOverride), Category = Export)
+	FFrameRate OverrideTimecodeRate;
+
 	/** Whether or not to use custom time range */
 	UPROPERTY(EditAnywhere, BlueprintReadWrite, AdvancedDisplay, Category = Export)
 	bool bUseCustomTimeRange = false;
@@ -99,6 +111,9 @@ public:
 		WarmUpFrames = 0;
 		DelayBeforeStart = 0;
 		bTransactRecording = true;
+		bBakeTimecode = false;
+		bTimecodeRateOverride = false;
+		OverrideTimecodeRate = FFrameRate(30, 1);
 		bUseCustomTimeRange = false;
 		CustomStartFrame = 0;
 		CustomEndFrame = 120;
