Close SubTrackEditorMode when closing the Sequencer editor

diff --git a/Engine/Source/Editor/MovieSceneTools/Private/TrackEditors/SubTrackEditor.cpp b/Engine/Source/Editor/MovieSceneTools/Private/TrackEditors/SubTrackEditor.cpp
--- a/Engine/Source/Editor/MovieSceneTools/Private/TrackEditors/SubTrackEditor.cpp
+++ b/Engine/Source/Editor/MovieSceneTools/Private/TrackEditors/SubTrackEditor.cpp
@@ -669,10 +669,9 @@ void FSubTrackEditor::OnInitialize()
 
 void FSubTrackEditor::OnRelease()
 {
-	FSubTrackEditorMode* EditorMode = static_cast<FSubTrackEditorMode*>(GLevelEditorModeTools().GetActiveMode(FSubTrackEditorMode::ModeName));
-	if(EditorMode)
+	if(GLevelEditorModeTools().IsModeActive(FSubTrackEditorMode::ModeName))
 	{
-		EditorMode->GetOnOriginValueChanged().RemoveAll(this);	
+		GLevelEditorModeTools().DeactivateMode(FSubTrackEditorMode::ModeName);
 	}
 	GetSequencer()->GetViewModel()->GetSelection()->TrackArea.OnChanged.RemoveAll(this);
 
