Sequencer: Make sure not force update the deferred modify when setting Curve key attributes, since many may be getting set at once.

diff --git a/Engine/Source/Editor/MovieSceneTools/Private/Channels/BezierChannelCurveModel.cpp b/Engine/Source/Editor/MovieSceneTools/Private/Channels/BezierChannelCurveModel.cpp
--- a/Engine/Source/Editor/MovieSceneTools/Private/Channels/BezierChannelCurveModel.cpp
+++ b/Engine/Source/Editor/MovieSceneTools/Private/Channels/BezierChannelCurveModel.cpp
@@ -267,7 +267,7 @@ static bool IsAuto(ERichCurveTangentMode TangentMode)
 template<typename ChannelType, typename ChannelValue, typename KeyType> 
 void FBezierChannelCurveModel<ChannelType, ChannelValue, KeyType>::SetKeyAttributes(TArrayView<const FKeyHandle> InKeys, TArrayView<const FKeyAttributes> InAttributes, EPropertyChangeType::Type ChangeType)
 {
-	UE::MovieScene::FScopedSignedObjectModifyDefer Defer(true);
+	UE::MovieScene::FScopedSignedObjectModifyDefer Defer(false /*do not force upate*/);
 
 	ChannelType*             Channel     = this->GetChannelHandle().Get();
 	UMovieSceneSignedObject* SignedOwner = this->template GetOwningObjectOrOuter<UMovieSceneSignedObject>();
