Show the property name and tooltip instead of "Mode" in the EAlphaBlendMode property customization

diff --git a/Engine/Source/Editor/MovieSceneTools/Private/MovieSceneAlphaBlendOptionCustomization.cpp b/Engine/Source/Editor/MovieSceneTools/Private/MovieSceneAlphaBlendOptionCustomization.cpp
--- a/Engine/Source/Editor/MovieSceneTools/Private/MovieSceneAlphaBlendOptionCustomization.cpp
+++ b/Engine/Source/Editor/MovieSceneTools/Private/MovieSceneAlphaBlendOptionCustomization.cpp
@@ -96,7 +96,8 @@ void FAlphaBlendPropertyCustomization::CustomizeChildren(TSharedRef<IPropertyHan
 		.NameContent()
 		[
 			SNew(STextBlock)
-			.Text(FText::FromString("Mode"))
+			.Text(StructPropertyHandle->GetPropertyDisplayName())
+			.ToolTipText(StructPropertyHandle->GetToolTipText())
 			.Font(IDetailLayoutBuilder::GetDetailFont())
 		]
 		.ValueContent()
