Sequencer: Check for valid sequencer for FMovieSceneObjectBindingIDCustomization

#rnx

diff --git a/Engine/Source/Editor/MovieSceneTools/Private/MovieSceneObjectBindingIDCustomization.cpp b/Engine/Source/Editor/MovieSceneTools/Private/MovieSceneObjectBindingIDCustomization.cpp
--- a/Engine/Source/Editor/MovieSceneTools/Private/MovieSceneObjectBindingIDCustomization.cpp
+++ b/Engine/Source/Editor/MovieSceneTools/Private/MovieSceneObjectBindingIDCustomization.cpp
@@ -48,7 +48,12 @@ void FMovieSceneObjectBindingIDCustomization::BindTo(TSharedRef<ISequencer> Oute
 			FOnGetPropertyTypeCustomizationInstance BindingIDCustomizationFactory = FOnGetPropertyTypeCustomizationInstance::CreateLambda(
 				[WeakSequencer]
 				{
-					return MakeShared<FMovieSceneObjectBindingIDCustomization>(WeakSequencer.Pin()->GetFocusedTemplateID(), WeakSequencer);
+					if (TSharedPtr<ISequencer> Sequencer = WeakSequencer.Pin())
+					{
+						return MakeShared<FMovieSceneObjectBindingIDCustomization>(Sequencer->GetFocusedTemplateID(), Sequencer);
+					}
+
+					return MakeShared<FMovieSceneObjectBindingIDCustomization>();
 				}
 			);
 
