Sequencer: Enable sub track editor mode only if the selected section has overrides

#rnx

diff --git a/Engine/Source/Editor/MovieSceneTools/Private/TrackEditors/SubTrackEditor.cpp b/Engine/Source/Editor/MovieSceneTools/Private/TrackEditors/SubTrackEditor.cpp
--- a/Engine/Source/Editor/MovieSceneTools/Private/TrackEditors/SubTrackEditor.cpp
+++ b/Engine/Source/Editor/MovieSceneTools/Private/TrackEditors/SubTrackEditor.cpp
@@ -1016,11 +1016,14 @@ void FSubTrackEditor::UpdateActiveMode()
 	
 	for(UMovieSceneSection* Section : Sections)
 	{
-		if(Section->IsA(UMovieSceneSubSection::StaticClass()))
+		if(UMovieSceneSubSection* SubSection = Cast<UMovieSceneSubSection>(Section))
 		{
-			GLevelEditorModeTools().ActivateDefaultMode();
-			GLevelEditorModeTools().ActivateMode(FSubTrackEditorMode::ModeName);
-			return;
+			if (SubSection->IsTransformOriginEditable())
+			{
+				GLevelEditorModeTools().ActivateDefaultMode();
+				GLevelEditorModeTools().ActivateMode(FSubTrackEditorMode::ModeName);
+				return;
+			}
 		}
 	}
 	
@@ -1029,11 +1032,20 @@ void FSubTrackEditor::UpdateActiveMode()
 
 	for(UMovieSceneTrack* Track : Tracks)
 	{
-		if(Track->IsA(UMovieSceneSubTrack::StaticClass()))
+		if (UMovieSceneSubTrack* SubTrack = Cast<UMovieSceneSubTrack>(Track))
 		{
-			GLevelEditorModeTools().ActivateDefaultMode();
-			GLevelEditorModeTools().ActivateMode(FSubTrackEditorMode::ModeName);
-			return;
+			for (UMovieSceneSection* Section : Track->GetAllSections())
+			{
+				if (UMovieSceneSubSection* SubSection = Cast<UMovieSceneSubSection>(Section))
+				{
+					if (SubSection->IsTransformOriginEditable())
+					{
+						GLevelEditorModeTools().ActivateDefaultMode();
+						GLevelEditorModeTools().ActivateMode(FSubTrackEditorMode::ModeName);
+						return;
+					}
+				}
+			}
 		}
 	}
 }
