Sequencer camera preview tracks now render using the post process settings of the active camera.

diff --git a/Engine/Source/Editor/MovieSceneTools/Private/TrackEditorThumbnail/TrackEditorThumbnail.cpp b/Engine/Source/Editor/MovieSceneTools/Private/TrackEditorThumbnail/TrackEditorThumbnail.cpp
--- a/Engine/Source/Editor/MovieSceneTools/Private/TrackEditorThumbnail/TrackEditorThumbnail.cpp
+++ b/Engine/Source/Editor/MovieSceneTools/Private/TrackEditorThumbnail/TrackEditorThumbnail.cpp
@@ -400,7 +400,9 @@ void FTrackEditorThumbnailCache::DrawViewportThumbnail(FTrackEditorThumbnail& Tr
 	TrackEditorThumbnail.bIgnoreAlpha = true;
 	TrackEditorThumbnail.ResizeRenderTarget(RTSize);
 	UWorld* World = PreviewCameraComponent->GetWorld();
-	UE::MoveSceneTools::DrawViewportThumbnail(*TrackEditorThumbnail.GetRenderTarget(), RTSize, *World->Scene, ViewInfo, CurrentCache.Quality);
+	UE::MoveSceneTools::DrawViewportThumbnail(
+		*TrackEditorThumbnail.GetRenderTarget(), RTSize, *World->Scene, ViewInfo, CurrentCache.Quality, &PreviewCameraComponent->PostProcessSettings
+		);
 }
 
 
diff --git a/Engine/Source/Editor/MovieSceneTools/Private/TrackEditorThumbnail/TrackThumbnailUtils.cpp b/Engine/Source/Editor/MovieSceneTools/Private/TrackEditorThumbnail/TrackThumbnailUtils.cpp
--- a/Engine/Source/Editor/MovieSceneTools/Private/TrackEditorThumbnail/TrackThumbnailUtils.cpp
+++ b/Engine/Source/Editor/MovieSceneTools/Private/TrackEditorThumbnail/TrackThumbnailUtils.cpp
@@ -29,7 +29,8 @@ namespace UE::MoveSceneTools
 		const FIntPoint& RenderTargetSize,
 		FSceneInterface& Scene,
 		const FMinimalViewInfo& ViewInfo,
-		EThumbnailQuality Quality
+		EThumbnailQuality Quality,
+		const FPostProcessSettings* OverridePostProcessSettings
 		)
 	{
 		FSceneViewFamilyContext ViewFamily( FSceneViewFamily::ConstructionValues(&ThumbnailRenderTarget, &Scene, FEngineShowFlags(ESFIM_Game))
@@ -85,6 +86,19 @@ namespace UE::MoveSceneTools
 
 		FSceneView* NewView = new FSceneView(ViewInitOptions);
 		ViewFamily.Views.Add(NewView);
+		if (OverridePostProcessSettings)
+		{
+			FPostProcessSettings ProcessSettings = *OverridePostProcessSettings;
+			
+			// Temporal effects need time to warm up, which we don't do for thumbnail rendering for performance reasons.
+			// Auto-exposure is a temporal effect.
+			// Manual causes the AutoExposureBias to be used which will prevent the thumbnail from being mostly white in very bright scenes.
+			// This will not work in all cases but should in most cases (since it depends on what value is set for AutoExposureBias).
+			ProcessSettings.AutoExposureMethod = AEM_Manual;
+			ProcessSettings.bOverride_AutoExposureMethod = true;
+			
+			NewView->OverridePostProcessSettings(ProcessSettings, 1.f);
+		}
 
 		const float GlobalResolutionFraction = 1.f;
 		ViewFamily.SetScreenPercentageInterface(new FLegacyScreenPercentageDriver(ViewFamily, GlobalResolutionFraction));
diff --git a/Engine/Source/Editor/MovieSceneTools/Public/TrackEditorThumbnail/TrackThumbnailUtils.h b/Engine/Source/Editor/MovieSceneTools/Public/TrackEditorThumbnail/TrackThumbnailUtils.h
--- a/Engine/Source/Editor/MovieSceneTools/Public/TrackEditorThumbnail/TrackThumbnailUtils.h
+++ b/Engine/Source/Editor/MovieSceneTools/Public/TrackEditorThumbnail/TrackThumbnailUtils.h
@@ -10,6 +10,7 @@ class FSceneInterface;
 class FRenderTarget;
 enum class EThumbnailQuality : uint8;
 struct FMinimalViewInfo;
+struct FPostProcessSettings;
 
 namespace UE::MoveSceneTools
 {
@@ -29,6 +30,7 @@ namespace UE::MoveSceneTools
 		const FIntPoint& RenderTargetSize,
 		FSceneInterface& Scene,
 		const FMinimalViewInfo& ViewInfo,
-		EThumbnailQuality Quality
+		EThumbnailQuality Quality,
+		const FPostProcessSettings* OverridePostProcessSettings = nullptr
 		);
 }
diff --git a/Engine/Source/Editor/Sequencer/Private/Misc/Thumbnail/ThumbnailCaptureUtils_CameraCuts.cpp b/Engine/Source/Editor/Sequencer/Private/Misc/Thumbnail/ThumbnailCaptureUtils_CameraCuts.cpp
--- a/Engine/Source/Editor/Sequencer/Private/Misc/Thumbnail/ThumbnailCaptureUtils_CameraCuts.cpp
+++ b/Engine/Source/Editor/Sequencer/Private/Misc/Thumbnail/ThumbnailCaptureUtils_CameraCuts.cpp
@@ -58,7 +58,7 @@ namespace UE::Sequencer::Private
 		
 		// ... enqueues rendering commands to fill RenderTarget
 		MoveSceneTools::DrawViewportThumbnail(
-			*RenderTarget, { ThumbnailSize, ThumbnailSize }, *GWorld->Scene, ViewInfo, EThumbnailQuality::Best
+			*RenderTarget, { ThumbnailSize, ThumbnailSize }, *GWorld->Scene, ViewInfo, EThumbnailQuality::Best, &Component->PostProcessSettings
 			);
 		// ... cleans up any important state modified by PreDrawThumbnailSetupSequencer.
 
