Constraints: try to keep the current child's transform when possible
    + remove compensation keys when removing a constraint

#rnx

diff --git a/Engine/Plugins/Animation/ControlRig/Source/ControlRigEditor/Private/Sequencer/ControlConstraintChannelInterface.cpp b/Engine/Plugins/Animation/ControlRig/Source/ControlRigEditor/Private/Sequencer/ControlConstraintChannelInterface.cpp
--- a/Engine/Plugins/Animation/ControlRig/Source/ControlRigEditor/Private/Sequencer/ControlConstraintChannelInterface.cpp
+++ b/Engine/Plugins/Animation/ControlRig/Source/ControlRigEditor/Private/Sequencer/ControlConstraintChannelInterface.cpp
@@ -71,7 +71,7 @@ bool FControlConstraintChannelInterface::SmartConstraintKey(
 	UTickableTransformConstraint* InConstraint, const TOptional<bool>& InOptActive,
 	const FFrameNumber& InTime, const TSharedPtr<ISequencer>& InSequencer)
 {
-	const UTransformableControlHandle* ControlHandle = Cast<UTransformableControlHandle>(InConstraint->ChildTRSHandle);
+	const UTransformableControlHandle* ControlHandle = InConstraint ? Cast<UTransformableControlHandle>(InConstraint->ChildTRSHandle) : nullptr;
 	if (!ControlHandle)
 	{
 		return false;
@@ -98,6 +98,7 @@ bool FControlConstraintChannelInterface::SmartConstraintKey(
 	}
 
 	// add the channel
+	const bool bHadConstraintChannel = InConstraint && ConstraintSection->HasConstraintChannel(InConstraint->ConstraintID);
 	ConstraintSection->AddConstraintChannel(InConstraint);
 
 	// add key if needed
@@ -129,10 +130,24 @@ bool FControlConstraintChannelInterface::SmartConstraintKey(
 				FramesToCompensate.Add(InTime);
 			}
 
-
+			FCompensationEvaluator::FEvalParameters EvalParams(InSequencer, FramesToCompensate);
+			EvalParams.bToActive = ActiveValueToBeSet;
+			if (bHadConstraintChannel && ActiveValueToBeSet)
+			{
+				const TMovieSceneChannelData<const bool> ChannelData = Channel->ActiveChannel.GetData();
+				if (ChannelData.GetTimes().IsEmpty())
+				{
+					// this means that the channel already existed but there was no existing key
+					// e.g., when the user deletes keys directly in sequencer without deleting the constraint
+					const FFrameTime FrameTime = InSequencer->GetLocalTime().ConvertTo(TickResolution);
+					const FFrameNumber Time = FrameTime.GetFrame();
+					EvalParams.bKeepCurrent = !FramesToCompensate.IsEmpty() && FramesToCompensate[0] == Time;
+				}
+			}
+			
 			// store child and space transforms for these frames
 			FCompensationEvaluator Evaluator(InConstraint);
-			Evaluator.ComputeLocalTransforms(ControlRig->GetWorld(), InSequencer, FramesToCompensate, ActiveValueToBeSet);
+			Evaluator.ComputeLocalTransforms(ControlRig->GetWorld(), EvalParams);
 			TArray<FTransform>& ChildLocals = Evaluator.ChildLocals;
 			
 			// store tangents at this time
diff --git a/Engine/Plugins/Animation/ControlRig/Source/ControlRigEditor/Private/Tools/ConstraintBaker.cpp b/Engine/Plugins/Animation/ControlRig/Source/ControlRigEditor/Private/Tools/ConstraintBaker.cpp
--- a/Engine/Plugins/Animation/ControlRig/Source/ControlRigEditor/Private/Tools/ConstraintBaker.cpp
+++ b/Engine/Plugins/Animation/ControlRig/Source/ControlRigEditor/Private/Tools/ConstraintBaker.cpp
@@ -284,7 +284,8 @@ void FConstraintBaker::Bake(UWorld* InWorld,
 	}
 
 	FCompensationEvaluator Evaluator(InConstraint);
-	Evaluator.ComputeLocalTransformsForBaking(InWorld, InSequencer, FramesToBake);
+	const FCompensationEvaluator::FEvalParameters EvalParams(InSequencer, FramesToBake);
+	Evaluator.ComputeLocalTransformsForBaking(InWorld, EvalParams);
 	TArray<FTransform> Transforms = Evaluator.ChildLocals;
 	if (FramesToBake.Num() != Transforms.Num())
 	{
diff --git a/Engine/Source/Editor/MovieSceneTools/Private/Constraints/ComponentConstraintChannelInterface.cpp b/Engine/Source/Editor/MovieSceneTools/Private/Constraints/ComponentConstraintChannelInterface.cpp
--- a/Engine/Source/Editor/MovieSceneTools/Private/Constraints/ComponentConstraintChannelInterface.cpp
+++ b/Engine/Source/Editor/MovieSceneTools/Private/Constraints/ComponentConstraintChannelInterface.cpp
@@ -81,7 +81,8 @@ bool FComponentConstraintChannelInterface::SmartConstraintKey(
 	InConstraint->bDynamicOffset = true;
 	
 	// add the channel
-	if (ConstraintSection->HasConstraintChannel(InConstraint->ConstraintID) == false)
+	const bool bHadConstraintChannel = ConstraintSection->HasConstraintChannel(InConstraint->ConstraintID);
+	if (!bHadConstraintChannel)
 	{
 		//check if static if so we need to delete it from world, will get added later again
 		if (UConstraintsManager* Manager = InConstraint->GetTypedOuter<UConstraintsManager>())
@@ -121,16 +122,31 @@ bool FComponentConstraintChannelInterface::SmartConstraintKey(
 				FramesToCompensate.Add(InTime);
 			}
 
+			FCompensationEvaluator::FEvalParameters EvalParams(InSequencer, FramesToCompensate);
+			EvalParams.bToActive = ActiveValueToBeSet;
+			if (bHadConstraintChannel && ActiveValueToBeSet)
+			{
+				const TMovieSceneChannelData<const bool> ChannelData = Channel->ActiveChannel.GetData();
+				if (ChannelData.GetTimes().IsEmpty())
+				{
+					// this means that the channel already existed but there was no existing key
+					// e.g., when the user deletes keys directly in sequencer without deleting the constraint
+					const FFrameRate TickResolution = InSequencer->GetFocusedTickResolution();
+					const FFrameTime FrameTime = InSequencer->GetLocalTime().ConvertTo(TickResolution);
+					const FFrameNumber Time = FrameTime.GetFrame();
+					EvalParams.bKeepCurrent = !FramesToCompensate.IsEmpty() && FramesToCompensate[0] == Time;
+				}
+			}
+			
 			// store child and space transforms for these frames
 			FCompensationEvaluator Evaluator(InConstraint);
-			Evaluator.ComputeLocalTransforms(GetHandleWorld(InConstraint->ChildTRSHandle), InSequencer, FramesToCompensate, ActiveValueToBeSet);
+			Evaluator.ComputeLocalTransforms(GetHandleWorld(InConstraint->ChildTRSHandle), EvalParams);
 			TArray<FTransform>& ChildLocals = Evaluator.ChildLocals;
 
 			// store tangents at this time
 			TArray<FMovieSceneTangentData> Tangents;
-			const int32 ChannelIndex = 0;
-			const int32 NumChannels = 9;
-
+			static constexpr int32 ChannelIndex = 0, NumChannels = 9;
+			
 			if (bNeedsCompensation)
 			{
 				EvaluateTangentAtThisTime<FMovieSceneDoubleChannel>(ChannelIndex, NumChannels, TransformSection, InTime, Tangents);
diff --git a/Engine/Source/Editor/MovieSceneTools/Private/Constraints/MovieSceneConstraintChannelHelper.cpp b/Engine/Source/Editor/MovieSceneTools/Private/Constraints/MovieSceneConstraintChannelHelper.cpp
--- a/Engine/Source/Editor/MovieSceneTools/Private/Constraints/MovieSceneConstraintChannelHelper.cpp
+++ b/Engine/Source/Editor/MovieSceneTools/Private/Constraints/MovieSceneConstraintChannelHelper.cpp
@@ -34,11 +34,9 @@ FCompensationEvaluator::FCompensationEvaluator(UTickableTransformConstraint* InC
 	, Handle(InConstraint ? InConstraint->ChildTRSHandle : nullptr)
 {}
 
-void FCompensationEvaluator::ComputeLocalTransforms(
-	UWorld* InWorld, const TSharedPtr<ISequencer>& InSequencer,
-	const TArray<FFrameNumber>& InFrames, const bool bToActive)
+void FCompensationEvaluator::ComputeLocalTransforms(UWorld* InWorld, const FEvalParameters& InEvalParams)
 {
-	if (InFrames.IsEmpty())
+	if (!InEvalParams.IsValid())
 	{
 		return;
 	}
@@ -66,11 +64,16 @@ void FCompensationEvaluator::ComputeLocalTransforms(
 		return LastActiveIndex > INDEX_NONE ? Cast<UTickableTransformConstraint>(ConstraintsMinusThis[LastActiveIndex]) : nullptr;
 	};
 
-	UMovieScene* MovieScene = InSequencer->GetFocusedMovieSceneSequence()->GetMovieScene();
+	ISequencer* Sequencer = InEvalParams.Sequencer;
+	UMovieScene* MovieScene = Sequencer->GetFocusedMovieSceneSequence()->GetMovieScene();
 	const FFrameRate TickResolution = MovieScene->GetTickResolution();
-	const EMovieScenePlayerStatus::Type PlaybackStatus = InSequencer->GetPlaybackStatus();
+	const EMovieScenePlayerStatus::Type PlaybackStatus = Sequencer->GetPlaybackStatus();
+
+	const TConstArrayView Frames(InEvalParams.Frames);
+	const int32 NumFrames = Frames.Num();
 
-	const int32 NumFrames = InFrames.Num();
+	const FTransform CurrentLocal = Handle->GetLocalTransform();
+	const FTransform CurrentGlobal = Handle->GetGlobalTransform();
 
 	// resize arrays to num frames + 1 as we also evaluate at InFrames[0]-1
 	ChildLocals.SetNum(NumFrames + 1);
@@ -90,11 +93,11 @@ void FCompensationEvaluator::ComputeLocalTransforms(
 	static constexpr bool bSorted = true, bTickHandles = true;
 	const TArray<ConstraintPtr> AllConstraints = Controller.GetAllConstraints(bSorted);
 
-	FMovieSceneInverseSequenceTransform LocalToRootTransform = InSequencer->GetFocusedMovieSceneSequenceTransform().Inverse();
+	FMovieSceneInverseSequenceTransform LocalToRootTransform = Sequencer->GetFocusedMovieSceneSequenceTransform().Inverse();
 
 	for (int32 Index = 0; Index < NumFrames + 1; ++Index)
 	{
-		FFrameNumber FrameNumber = (Index == 0) ? InFrames[0] - 1 : InFrames[Index - 1];
+		FFrameNumber FrameNumber = (Index == 0) ? Frames[0] - 1 : Frames[Index - 1];
 		TOptional<FFrameTime> RootTime = LocalToRootTransform.TryTransformTime(FrameNumber);
 		if (!RootTime)
 		{
@@ -114,7 +117,7 @@ void FCompensationEvaluator::ComputeLocalTransforms(
 				BakeHelper->PreEvaluation(MovieScene, FrameNumber);
 			}
 		}
-		InSequencer->GetEvaluationTemplate().EvaluateSynchronousBlocking(Context);
+		Sequencer->GetEvaluationTemplate().EvaluateSynchronousBlocking(Context);
 
 		// evaluate constraints
 		for (const TWeakObjectPtr<UTickableConstraint>& InConstraint : AllConstraints)
@@ -137,15 +140,22 @@ void FCompensationEvaluator::ComputeLocalTransforms(
 		FTransform& ChildGlobal = ChildGlobals[Index];
 		FTransform& SpaceGlobal = SpaceGlobals[Index];
 
-		// store child transforms        	
-		ChildLocal = Handle->GetLocalTransform();
-		Handle->PreEvaluate();
-		ChildGlobal = Handle->GetGlobalTransform();
+		if (InEvalParams.bKeepCurrent && Index < 2)
+		{
+			ChildGlobal = CurrentGlobal;
+		}
+		else
+		{
+			// store child transforms        	
+			ChildLocal = Handle->GetLocalTransform();
+			Handle->PreEvaluate();
+			ChildGlobal = Handle->GetGlobalTransform();
+		}
 
 		const UTickableTransformConstraint* LastConstraint = GetLastActiveConstraint();
 
 		// store constraint/parent space global transform
-		if (bToActive)
+		if (InEvalParams.bToActive)
 		{
 			// if activating the constraint, store last constraint or parent space at T[0]-1
 			// and constraint space for all other times
@@ -161,6 +171,10 @@ void FCompensationEvaluator::ComputeLocalTransforms(
 						ChildLocal = *Relative;
 					}
 				}
+				else if (InEvalParams.bKeepCurrent)
+				{
+					ChildLocal = CurrentLocal;
+				}
 			}
 			else
 			{
@@ -202,12 +216,12 @@ void FCompensationEvaluator::ComputeLocalTransforms(
 		}
 	}
 	//get back to where we are at, should also make sure things are active
-	InSequencer->ForceEvaluate();
+	Sequencer->ForceEvaluate();
 }
 
-void FCompensationEvaluator::ComputeLocalTransformsForBaking(UWorld* InWorld, const TSharedPtr<ISequencer>& InSequencer, const TArray<FFrameNumber>& InFrames)
+void FCompensationEvaluator::ComputeLocalTransformsForBaking(UWorld* InWorld, const FEvalParameters& InEvalParams)
 {
-	if (InFrames.IsEmpty())
+	if (!InEvalParams.IsValid())
 	{
 		return;
 	}
@@ -235,11 +249,13 @@ void FCompensationEvaluator::ComputeLocalTransformsForBaking(UWorld* InWorld, co
 		return LastActiveIndex > INDEX_NONE ? Cast<UTickableTransformConstraint>(ConstraintsMinusThis[LastActiveIndex]) : nullptr;
 	};
 
-	UMovieScene* MovieScene = InSequencer->GetFocusedMovieSceneSequence()->GetMovieScene();
+	ISequencer* Sequencer = InEvalParams.Sequencer;
+	UMovieScene* MovieScene = Sequencer->GetFocusedMovieSceneSequence()->GetMovieScene();
 	const FFrameRate TickResolution = MovieScene->GetTickResolution();
-	const EMovieScenePlayerStatus::Type PlaybackStatus = InSequencer->GetPlaybackStatus();
+	const EMovieScenePlayerStatus::Type PlaybackStatus = Sequencer->GetPlaybackStatus();
 
-	const int32 NumFrames = InFrames.Num();
+	const TConstArrayView Frames(InEvalParams.Frames);
+	const int32 NumFrames = Frames.Num();
 
 	ChildLocals.SetNum(NumFrames);
 	ChildGlobals.SetNum(NumFrames);
@@ -261,11 +277,11 @@ void FCompensationEvaluator::ComputeLocalTransformsForBaking(UWorld* InWorld, co
 			BakeHelper->StartBaking(MovieScene);
 		}
 	}
-	FMovieSceneInverseSequenceTransform LocalToRootTransform = InSequencer->GetFocusedMovieSceneSequenceTransform().Inverse();
+	FMovieSceneInverseSequenceTransform LocalToRootTransform = Sequencer->GetFocusedMovieSceneSequenceTransform().Inverse();
 
 	for (int32 Index = 0; Index < NumFrames; ++Index)
 	{
-		TOptional<FFrameTime> RootTime = LocalToRootTransform.TryTransformTime(InFrames[Index]);
+		TOptional<FFrameTime> RootTime = LocalToRootTransform.TryTransformTime(Frames[Index]);
 		if (!RootTime)
 		{
 			continue;
@@ -284,7 +300,7 @@ void FCompensationEvaluator::ComputeLocalTransformsForBaking(UWorld* InWorld, co
 				BakeHelper->PreEvaluation(MovieScene, FrameNumber);
 			}
 		}
-		InSequencer->GetEvaluationTemplate().EvaluateSynchronousBlocking(Context);
+		Sequencer->GetEvaluationTemplate().EvaluateSynchronousBlocking(Context);
 
 		// evaluate constraints
 		for (const TWeakObjectPtr<UTickableConstraint>& InConstraint : AllConstraints)
@@ -339,15 +355,12 @@ void FCompensationEvaluator::ComputeLocalTransformsForBaking(UWorld* InWorld, co
 	{
 		// the handle might not be valid after baking due to spawnables or baking out of the sequence boundaries
 		// so force sequencer evaluation to make sure we're back to normal
-		InSequencer->ForceEvaluate();
+		Sequencer->ForceEvaluate();
 	}
 }
-void FCompensationEvaluator::ComputeLocalTransformsBeforeDeletion(
-	UWorld* InWorld,
-	const TSharedPtr<ISequencer>& InSequencer,
-	const TArray<FFrameNumber>& InFrames)
+void FCompensationEvaluator::ComputeLocalTransformsBeforeDeletion(UWorld* InWorld, const FEvalParameters& InEvalParams)
 {
-	if (InFrames.IsEmpty())
+	if (!InEvalParams.IsValid())
 	{
 		return;
 	}
@@ -374,12 +387,14 @@ void FCompensationEvaluator::ComputeLocalTransformsBeforeDeletion(
 	const FConstraintsManagerController& Controller = FConstraintsManagerController::Get(InWorld);
 	static constexpr bool bSorted = true;
 	const TArray<ConstraintPtr> AllConstraints = Controller.GetAllConstraints(bSorted);
-	
-	UMovieScene* MovieScene = InSequencer->GetFocusedMovieSceneSequence()->GetMovieScene();
+
+	ISequencer* Sequencer = InEvalParams.Sequencer;
+	UMovieScene* MovieScene = Sequencer->GetFocusedMovieSceneSequence()->GetMovieScene();
 	const FFrameRate TickResolution = MovieScene->GetTickResolution();
-	const EMovieScenePlayerStatus::Type PlaybackStatus = InSequencer->GetPlaybackStatus();
+	const EMovieScenePlayerStatus::Type PlaybackStatus = Sequencer->GetPlaybackStatus();
 
-	const int32 NumFrames = InFrames.Num();
+	const TConstArrayView Frames(InEvalParams.Frames);
+	const int32 NumFrames = Frames.Num();
 
 	ChildLocals.SetNum(NumFrames);
 	ChildGlobals.SetNum(NumFrames);
@@ -392,11 +407,11 @@ void FCompensationEvaluator::ComputeLocalTransformsBeforeDeletion(
 			BakeHelper->StartBaking(MovieScene);
 		}
 	}
-	FMovieSceneInverseSequenceTransform LocalToRootTransform = InSequencer->GetFocusedMovieSceneSequenceTransform().Inverse();
+	FMovieSceneInverseSequenceTransform LocalToRootTransform = Sequencer->GetFocusedMovieSceneSequenceTransform().Inverse();
 
 	for (int32 Index = 0; Index < NumFrames; ++Index)
 	{
-		TOptional<FFrameTime> RootTime = LocalToRootTransform.TryTransformTime(InFrames[Index]);
+		TOptional<FFrameTime> RootTime = LocalToRootTransform.TryTransformTime(Frames[Index]);
 		if (!RootTime)
 		{
 			continue;
@@ -415,7 +430,7 @@ void FCompensationEvaluator::ComputeLocalTransformsBeforeDeletion(
 				BakeHelper->PreEvaluation(MovieScene, FrameNumber);
 			}
 		}
-		InSequencer->GetEvaluationTemplate().EvaluateSynchronousBlocking(Context);
+		Sequencer->GetEvaluationTemplate().EvaluateSynchronousBlocking(Context);
 
 		// evaluate constraints
 		for (const TWeakObjectPtr<UTickableConstraint>& InConstraint : AllConstraints)
@@ -594,9 +609,9 @@ void FCompensationEvaluator::ComputeCompensation(UWorld* InWorld, const TSharedP
 	}
 }
 
-void FCompensationEvaluator::CacheTransforms(UWorld* InWorld, const TSharedPtr<ISequencer>& InSequencer, const TArray<FFrameNumber>& InTime)
+void FCompensationEvaluator::CacheTransforms(UWorld* InWorld, const FEvalParameters& InEvalParams)
 {
-	if (InTime.IsEmpty())
+	if (!InEvalParams.IsValid())
 	{
 		return;
 	}
@@ -608,11 +623,13 @@ void FCompensationEvaluator::CacheTransforms(UWorld* InWorld, const TSharedPtr<I
 	static constexpr bool bSorted = true;
 	const TArray<ConstraintPtr> AllConstraints = Controller.GetAllConstraints(bSorted);
 
-	UMovieScene* MovieScene = InSequencer->GetFocusedMovieSceneSequence()->GetMovieScene();
-	const FFrameRate TickResolution = InSequencer->GetRootMovieSceneSequence()->GetMovieScene()->GetTickResolution();
-	const EMovieScenePlayerStatus::Type PlaybackStatus = InSequencer->GetPlaybackStatus();
+	ISequencer* Sequencer = InEvalParams.Sequencer;
+	UMovieScene* MovieScene = Sequencer->GetFocusedMovieSceneSequence()->GetMovieScene();
+	const FFrameRate TickResolution = Sequencer->GetRootMovieSceneSequence()->GetMovieScene()->GetTickResolution();
+	const EMovieScenePlayerStatus::Type PlaybackStatus = Sequencer->GetPlaybackStatus();
 
-	const int32 NumFrames = InTime.Num();
+	const TConstArrayView Frames(InEvalParams.Frames);
+	const int32 NumFrames = Frames.Num();
 
 	ChildLocals.SetNum(NumFrames);
 	ChildGlobals.SetNum(NumFrames);
@@ -620,7 +637,7 @@ void FCompensationEvaluator::CacheTransforms(UWorld* InWorld, const TSharedPtr<I
 
 	const TArray<IMovieSceneToolsAnimationBakeHelper*>& BakeHelpers = FMovieSceneToolsModule::Get().GetAnimationBakeHelpers();
 
-	FMovieSceneInverseSequenceTransform LocalToRootTransform = InSequencer->GetFocusedMovieSceneSequenceTransform().Inverse();
+	FMovieSceneInverseSequenceTransform LocalToRootTransform = Sequencer->GetFocusedMovieSceneSequenceTransform().Inverse();
 
 	auto EvaluateAt = [&](FFrameNumber InFrame)
 	{
@@ -640,7 +657,7 @@ void FCompensationEvaluator::CacheTransforms(UWorld* InWorld, const TSharedPtr<I
 				BakeHelper->PreEvaluation(MovieScene, InFrame);
 			}
 		}
-		InSequencer->GetEvaluationTemplate().EvaluateSynchronousBlocking(Context);
+		Sequencer->GetEvaluationTemplate().EvaluateSynchronousBlocking(Context);
 
 		// evaluate constraints
 		for (const TWeakObjectPtr<UTickableConstraint>& InConstraint : AllConstraints)
@@ -671,7 +688,7 @@ void FCompensationEvaluator::CacheTransforms(UWorld* InWorld, const TSharedPtr<I
 	for (int32 Index = 0; Index < NumFrames; ++Index)
 	{
 		// evaluate animation
-		EvaluateAt(InTime[Index]);
+		EvaluateAt(Frames[Index]);
 
 		// store transforms        	
 		ChildLocals[Index] = Handle->GetLocalTransform();
@@ -789,17 +806,18 @@ void FMovieSceneConstraintChannelHelper::HandleConstraintRemoved(
 		TArray<FFrameNumber> FramesToCompensate;
 		if (FloatTransformChannels.Num() > 0)
 		{
-			FMovieSceneConstraintChannelHelper::GetFramesWithinActiveState(*InConstraintChannel, FloatTransformChannels, FramesToCompensate);
+			GetFramesWithinActiveState(*InConstraintChannel, FloatTransformChannels, FramesToCompensate);
 		}
 		else
 		{
-			FMovieSceneConstraintChannelHelper::GetFramesWithinActiveState(*InConstraintChannel, DoubleTransformChannels, FramesToCompensate);
+			GetFramesWithinActiveState(*InConstraintChannel, DoubleTransformChannels, FramesToCompensate);
 		}
 		// do the compensation
 		UWorld* World = GCurrentLevelEditingViewportClient ? GCurrentLevelEditingViewportClient->GetWorld() : nullptr;
 
 		FCompensationEvaluator Evaluator(Constraint);
-		Evaluator.ComputeLocalTransformsBeforeDeletion(World, InSequencer, FramesToCompensate);
+		const FCompensationEvaluator::FEvalParameters EvalParams(InSequencer, FramesToCompensate);
+		Evaluator.ComputeLocalTransformsBeforeDeletion(World, EvalParams);
 		const TArray<FTransform>& ChildLocals = Evaluator.ChildLocals;
 
 		const EMovieSceneTransformChannel ChannelsToKey = Constraint->GetChannelsToKey();
@@ -807,6 +825,44 @@ void FMovieSceneConstraintChannelHelper::HandleConstraintRemoved(
 
 		ControlHandle->AddTransformKeys(FramesToCompensate,
 			ChildLocals, ChannelsToKey, TickResolution, InSection);
+
+		// clean double keys
+		if (IMovieSceneConstrainedSection* Section = Cast<IMovieSceneConstrainedSection>(InSection))
+		{
+			// get constraints acting on the same child that is different that InConstraint
+			const TArray<FConstraintAndActiveChannel>& ConstraintChannels = Section->GetConstraintsChannels();
+			const TArray<FConstraintAndActiveChannel> OtherConstraints = ConstraintChannels.FilterByPredicate([Constraint](const FConstraintAndActiveChannel& Channel)
+			{
+				if (const UTickableTransformConstraint* TransformConstraint = Cast< UTickableTransformConstraint>(Channel.GetConstraint()))
+				{
+					return TransformConstraint->GetTargetHash() == Constraint->GetTargetHash() && TransformConstraint != Constraint;
+				}
+				return false;
+			});
+
+			// disable extra compensation when removing keys
+			if (OtherConstraints.IsEmpty())
+			{
+				// this was the only constraint so we can remove its double keys
+				const EMovieSceneKeyInterpolation KeyType = InSequencer->GetKeyInterpolation();
+				for (const FFrameNumber& Time: Times)
+				{
+					const FFrameNumber TimeMinusOne = Time - 1;
+					if (FloatTransformChannels.Num() > 0)
+					{
+						DeleteTransformKeys(FloatTransformChannels, TimeMinusOne);
+						//we also set the tangent at the break to the default type
+						ChangeKeyInterpolation(FloatTransformChannels, Time, KeyType);
+					}
+					else if (DoubleTransformChannels.Num() > 0)
+					{
+						DeleteTransformKeys(DoubleTransformChannels, TimeMinusOne);
+						//we also set the tangent at the break to the default type
+						ChangeKeyInterpolation(DoubleTransformChannels, Time, KeyType);
+					}
+				}
+			}
+		}
 	}
 }
 
@@ -861,9 +917,11 @@ void FMovieSceneConstraintChannelHelper::HandleConstraintKeyDeleted(
 			GetFramesAfter(*InConstraintChannel,  InTime, DoubleTransformChannels, FramesToCompensate);
 		}
 		// do the compensation
-		FCompensationEvaluator Evaluator(InConstraint);
 		UWorld* World = GCurrentLevelEditingViewportClient ? GCurrentLevelEditingViewportClient->GetWorld() : nullptr;
-		Evaluator.ComputeLocalTransforms(World, InSequencer, FramesToCompensate, PreviousValue);
+		FCompensationEvaluator Evaluator(InConstraint);
+        FCompensationEvaluator::FEvalParameters EvalParams(InSequencer, FramesToCompensate);
+        EvalParams.bToActive = PreviousValue;
+		Evaluator.ComputeLocalTransforms(World, EvalParams);
 		TArray<FTransform>& ChildLocals = Evaluator.ChildLocals;
 		//turn off constraint, if we delete the key it may not evaluate to false
 		InConstraint->SetActive(false);
@@ -1398,8 +1456,9 @@ void FMovieSceneConstraintChannelHelper::CompensateScale(
 	InSection->Modify();
 	
 	FCompensationEvaluator Evaluator(InParentConstraint);
+	const FCompensationEvaluator::FEvalParameters EvalParams(InSequencer, ActiveTimes);
 	UWorld* World = GCurrentLevelEditingViewportClient ? GCurrentLevelEditingViewportClient->GetWorld() : nullptr;
-	Evaluator.CacheTransforms(World, InSequencer, ActiveTimes);
+	Evaluator.CacheTransforms(World, EvalParams);
 
 	TArray<FTransform>& ChildLocals = Evaluator.ChildLocals;
 	const TArray<FTransform>& ChildGlobals = Evaluator.ChildGlobals;
diff --git a/Engine/Source/Editor/MovieSceneTools/Public/Constraints/MovieSceneConstraintChannelHelper.h b/Engine/Source/Editor/MovieSceneTools/Public/Constraints/MovieSceneConstraintChannelHelper.h
--- a/Engine/Source/Editor/MovieSceneTools/Public/Constraints/MovieSceneConstraintChannelHelper.h
+++ b/Engine/Source/Editor/MovieSceneTools/Public/Constraints/MovieSceneConstraintChannelHelper.h
@@ -30,11 +30,29 @@ public:
 
 	FCompensationEvaluator(UTickableTransformConstraint* InConstraint);
 
-	void ComputeLocalTransforms(UWorld* InWorld, const TSharedPtr<ISequencer>& InSequencer, const TArray<FFrameNumber>& InFrames, const bool bToActive);
-	void ComputeLocalTransformsBeforeDeletion(UWorld* InWorld, const TSharedPtr<ISequencer>& InSequencer, const TArray<FFrameNumber>& InFrames);
+	struct FEvalParameters
+	{
+		FEvalParameters(const TSharedPtr<ISequencer>& InSequencer, const TArray<FFrameNumber>& InFrames)
+			: Sequencer(InSequencer.Get())
+			, Frames(InFrames)
+		{}
+
+		bool IsValid() const
+		{
+			return Sequencer && !Frames.IsEmpty();
+		}
+		
+		ISequencer* Sequencer = nullptr;
+		const TConstArrayView<const FFrameNumber> Frames;
+		bool bToActive = false;
+		bool bKeepCurrent = false;
+	};
+	
+	void ComputeLocalTransforms(UWorld* InWorld, const FEvalParameters& InEvalParams);
+	void ComputeLocalTransformsBeforeDeletion(UWorld* InWorld, const FEvalParameters& InEvalParams);
 	void ComputeCompensation(UWorld* InWorld, const TSharedPtr<ISequencer>& InSequencer, const FFrameNumber& InTime);
-	void ComputeLocalTransformsForBaking(UWorld* InWorld, const TSharedPtr<ISequencer>& InSequencer, const TArray<FFrameNumber>& InFrames);
-	void CacheTransforms(UWorld* InWorld, const TSharedPtr<ISequencer>& InSequencer, const TArray<FFrameNumber>& InFrames);
+	void ComputeLocalTransformsForBaking(UWorld* InWorld, const FEvalParameters& InEvalParams);
+	void CacheTransforms(UWorld* InWorld, const FEvalParameters& InEvalParams);
 	void ComputeCurrentTransforms(UWorld* InWorld);
 	
 private:
