fix linux build

#rnx

diff --git a/Engine/Source/Editor/MovieSceneTools/Private/Constraints/MovieSceneConstraintChannelHelper.cpp b/Engine/Source/Editor/MovieSceneTools/Private/Constraints/MovieSceneConstraintChannelHelper.cpp
--- a/Engine/Source/Editor/MovieSceneTools/Private/Constraints/MovieSceneConstraintChannelHelper.cpp
+++ b/Engine/Source/Editor/MovieSceneTools/Private/Constraints/MovieSceneConstraintChannelHelper.cpp
@@ -69,7 +69,7 @@ void FCompensationEvaluator::ComputeLocalTransforms(UWorld* InWorld, const FEval
 	const FFrameRate TickResolution = MovieScene->GetTickResolution();
 	const EMovieScenePlayerStatus::Type PlaybackStatus = Sequencer->GetPlaybackStatus();
 
-	const TConstArrayView Frames(InEvalParams.Frames);
+	const TConstArrayView<FFrameNumber> Frames(InEvalParams.Frames);
 	const int32 NumFrames = Frames.Num();
 
 	const FTransform CurrentLocal = Handle->GetLocalTransform();
@@ -254,7 +254,7 @@ void FCompensationEvaluator::ComputeLocalTransformsForBaking(UWorld* InWorld, co
 	const FFrameRate TickResolution = MovieScene->GetTickResolution();
 	const EMovieScenePlayerStatus::Type PlaybackStatus = Sequencer->GetPlaybackStatus();
 
-	const TConstArrayView Frames(InEvalParams.Frames);
+	const TConstArrayView<FFrameNumber> Frames(InEvalParams.Frames);
 	const int32 NumFrames = Frames.Num();
 
 	ChildLocals.SetNum(NumFrames);
@@ -393,7 +393,7 @@ void FCompensationEvaluator::ComputeLocalTransformsBeforeDeletion(UWorld* InWorl
 	const FFrameRate TickResolution = MovieScene->GetTickResolution();
 	const EMovieScenePlayerStatus::Type PlaybackStatus = Sequencer->GetPlaybackStatus();
 
-	const TConstArrayView Frames(InEvalParams.Frames);
+	const TConstArrayView<FFrameNumber> Frames(InEvalParams.Frames);
 	const int32 NumFrames = Frames.Num();
 
 	ChildLocals.SetNum(NumFrames);
@@ -628,7 +628,7 @@ void FCompensationEvaluator::CacheTransforms(UWorld* InWorld, const FEvalParamet
 	const FFrameRate TickResolution = Sequencer->GetRootMovieSceneSequence()->GetMovieScene()->GetTickResolution();
 	const EMovieScenePlayerStatus::Type PlaybackStatus = Sequencer->GetPlaybackStatus();
 
-	const TConstArrayView Frames(InEvalParams.Frames);
+	const TConstArrayView<FFrameNumber> Frames(InEvalParams.Frames);
 	const int32 NumFrames = Frames.Num();
 
 	ChildLocals.SetNum(NumFrames);
