Sequencer: GetSectionTime for audio even if there is no timecode info. This will return the offset into the audio section.

diff --git a/Engine/Source/Editor/MovieSceneTools/Private/TrackEditors/AudioTrackEditor.cpp b/Engine/Source/Editor/MovieSceneTools/Private/TrackEditors/AudioTrackEditor.cpp
--- a/Engine/Source/Editor/MovieSceneTools/Private/TrackEditors/AudioTrackEditor.cpp
+++ b/Engine/Source/Editor/MovieSceneTools/Private/TrackEditors/AudioTrackEditor.cpp
@@ -950,19 +950,16 @@ TOptional<FFrameTime> FAudioSection::GetSectionTime(FSequencerSectionPainter& In
 			TimecodeInfo = SoundWave->GetTimecodeInfo();
 		}
 
-		if (TimecodeInfo.IsSet())
-		{
-			const FFrameRate SectionFrameRate = AudioSection->GetTypedOuter<UMovieScene>()->GetTickResolution();
-			const double AudioStartOffset = SectionFrameRate.AsSeconds(AudioSection->GetStartOffset());
+		const FFrameRate SectionFrameRate = AudioSection->GetTypedOuter<UMovieScene>()->GetTickResolution();
+		const double AudioStartOffset = SectionFrameRate.AsSeconds(AudioSection->GetStartOffset());
 
-			const FFrameRate TickResolution = CurrentTime.Rate;
-			const FFrameTime FramesSinceMidnight = TickResolution.AsFrameTime(TimecodeInfo->GetNumSecondsSinceMidnight() + AudioStartOffset);
+		const FFrameRate TickResolution = CurrentTime.Rate;
+		const FFrameTime FramesSinceMidnight = TimecodeInfo.IsSet() ? TickResolution.AsFrameTime(TimecodeInfo->GetNumSecondsSinceMidnight() + AudioStartOffset) : TickResolution.AsFrameTime(AudioStartOffset);
 
-			const FFrameTime SectionOffset = CurrentTime.Time - AudioSection->GetInclusiveStartFrame();
-			const FFrameTime CursorFrameTime = FramesSinceMidnight + SectionOffset;
+		const FFrameTime SectionOffset = CurrentTime.Time - AudioSection->GetInclusiveStartFrame();
+		const FFrameTime CursorFrameTime = FramesSinceMidnight + SectionOffset;
 
-			HintFrameTime = FQualifiedFrameTime(FFrameTime(CursorFrameTime.CeilToFrame()), TickResolution);
-		}
+		HintFrameTime = FQualifiedFrameTime(FFrameTime(CursorFrameTime.CeilToFrame()), TickResolution);
 	}
 
 	return HintFrameTime.Time;
