Sequencer- Fix root motion paths in the case of swap root bone actor, so that the actor walks along the path rather than the path moving along with the character.

diff --git a/Engine/Source/Editor/MovieSceneTools/Private/EditModes/SkeletalAnimationTrackEditMode.cpp b/Engine/Source/Editor/MovieSceneTools/Private/EditModes/SkeletalAnimationTrackEditMode.cpp
--- a/Engine/Source/Editor/MovieSceneTools/Private/EditModes/SkeletalAnimationTrackEditMode.cpp
+++ b/Engine/Source/Editor/MovieSceneTools/Private/EditModes/SkeletalAnimationTrackEditMode.cpp
@@ -30,6 +30,8 @@
 #include "Systems/MovieSceneComponentTransformSystem.h"
 #include "SkeletalDebugRendering.h"
 #include "AnimSequencerInstanceProxy.h"
+#include "Systems/MovieSceneSkeletalAnimationSystem.h"
+#include "Evaluation/MovieSceneEvaluationTemplateInstance.h"
 
 FName FSkeletalAnimationTrackEditMode::ModeName("EditMode.SkeletalAnimationTrackEditMode");
 
@@ -217,13 +219,37 @@ void FSkeletalAnimationTrackEditMode::AddReferencedObjects(FReferenceCollector&
 static void RenderTrail(UObject* BoundObject, UMovieSceneSkeletalAnimationTrack* SkelAnimTrack, UE::MovieScene::FSystemInterrogator& Interrogator, USkeletalMeshComponent* SkelMeshComp,
 	const TSharedPtr<ISequencer>& Sequencer, FPrimitiveDrawInterface* PDI)
 {
+	if (!Sequencer.IsValid())
+	{
+		return;
+	}
+
 	const FLinearColor RootMotionColor(0.75, 0.75, 0.75, 1.0f);
 	TOptional<FVector> OldKeyPos_G;
 
 	int32 Index = 0;
 	if (SkelAnimTrack->GetAllSections().Num() > 0)
 	{
-		const FTransform SkelMeshTransform = SkelAnimTrack->RootMotionParams.RootMotionStartOffset * SkelMeshComp->GetSocketTransform(NAME_None);
+		TOptional<FTransform> InitialActorTransform;
+
+		const FMovieSceneRootEvaluationTemplateInstance& EvaluationTemplate = Sequencer->GetEvaluationTemplate();
+		const UMovieSceneEntitySystemLinker* EntityLinker = EvaluationTemplate.GetEntitySystemLinker();
+		if (EntityLinker)
+		{
+			const UMovieSceneSkeletalAnimationSystem* SkeletalAnimationSystem = EntityLinker->FindSystem<UMovieSceneSkeletalAnimationSystem>();
+			if (SkeletalAnimationSystem)
+			{
+				// Attempt to get the pre-root motion actor transform so we can display the actor along the root motion path properly
+				InitialActorTransform = SkeletalAnimationSystem->GetInitialActorTransform(SkelMeshComp);
+				TOptional<FQuat> InverseMeshToActorRotation = SkeletalAnimationSystem->GetInverseMeshToActorRotation(SkelMeshComp);
+				if (InitialActorTransform.IsSet() && InverseMeshToActorRotation.IsSet())
+				{
+					InitialActorTransform->SetRotation(InitialActorTransform->GetRotation() * InverseMeshToActorRotation.GetValue());
+				}
+			}
+		}
+
+		const FTransform SkelMeshTransform = SkelAnimTrack->RootMotionParams.RootMotionStartOffset * ((InitialActorTransform.IsSet()) ? InitialActorTransform.GetValue() : SkelMeshComp->GetSocketTransform(NAME_None));
 		for (const FTransform& Transform : SkelAnimTrack->RootMotionParams.RootTransforms)
 		{
 			FVector NewKeyPos = Transform.GetLocation();
