Constraints & Sequencer: notify of spawnable state discrepancy to avoid future constraints' undefined states when loading the level sequence is another level

#rnx

diff --git a/Engine/Source/Editor/MovieSceneTools/Private/Constraints/MovieSceneConstraintChannelHelper.cpp b/Engine/Source/Editor/MovieSceneTools/Private/Constraints/MovieSceneConstraintChannelHelper.cpp
--- a/Engine/Source/Editor/MovieSceneTools/Private/Constraints/MovieSceneConstraintChannelHelper.cpp
+++ b/Engine/Source/Editor/MovieSceneTools/Private/Constraints/MovieSceneConstraintChannelHelper.cpp
@@ -21,6 +21,8 @@
 #include "MovieSceneToolsModule.h"
 #include "Channels/MovieSceneChannelTraits.h"
 #include "Constraints/TransformConstraintChannelInterface.h"
+#include "Framework/Notifications/NotificationManager.h"
+#include "Widgets/Notifications/SNotificationList.h"
 
 
 /*
@@ -987,6 +989,64 @@ void FMovieSceneConstraintChannelHelper::HandleConstraintKeyMoved(
 
 }
 
+namespace UE::Constraints::Private
+{
+
+void ShowSpawnableWarning(const FText& InNotification)
+{
+	// output log notification
+	UE_LOG(LogTemp, Warning, TEXT("%s"), *InNotification.ToString());
+
+	// editor notification
+	FNotificationInfo Info(InNotification);
+	Info.Image = FAppStyle::GetBrush(TEXT("MessageLog.Warning"));
+	Info.ExpireDuration = 5.0f;
+
+	FSlateNotificationManager::Get().AddNotification(Info);
+}
+	
+static void ShowSpawnableDiscrepancyWarning(
+	const TObjectPtr<UTransformableHandle>& InSpawnableHandle, const TObjectPtr<UTransformableHandle>& InNonSpawnableHandle)
+{
+	if (!IsValid(InSpawnableHandle) || !IsValid(InNonSpawnableHandle))
+	{
+		return;
+	}
+
+	FFormatNamedArguments Args;
+	Args.Add(TEXT("Spawnable"), FText::FromString(InSpawnableHandle->GetLabel()));
+	Args.Add(TEXT("NonSpawnable"), FText::FromString(InNonSpawnableHandle->GetLabel()));
+
+	const FText Notification = FText::Format(
+		NSLOCTEXT("FMovieSceneConstraintChannelHelper", "ChildParentSpawnableDiscrepancy",
+		"Object '{Spawnable}' is spawnable whereas '{NonSpawnable}' is not.\n"
+		"Consider making '{NonSpawnable}' spawnable to avoid a future undefined state of this constraint."), Args);
+
+	ShowSpawnableWarning(Notification);
+}
+
+static void ShowNonSpawnableWarning(
+	const TObjectPtr<UTransformableHandle>& InNonSpawnableHandle0, const TObjectPtr<UTransformableHandle>& InNonSpawnableHandle1)
+{
+	if (!IsValid(InNonSpawnableHandle0) || !IsValid(InNonSpawnableHandle1))
+	{
+		return;
+	}
+
+	FFormatNamedArguments Args;
+	Args.Add(TEXT("NonSpawnable0"), FText::FromString(InNonSpawnableHandle0->GetLabel()));
+	Args.Add(TEXT("NonSpawnable1"), FText::FromString(InNonSpawnableHandle1->GetLabel()));
+
+	const FText Notification = FText::Format(
+		NSLOCTEXT("FMovieSceneConstraintChannelHelper", "ChildParentNonSpawnable",
+		"Objects '{NonSpawnable0}' and '{NonSpawnable1}' are not spawnable.\n"
+		"Consider making them spawnable if you want them, as well as this constraint, to be fully integrated to your level sequence."), Args);
+
+	ShowSpawnableWarning(Notification);
+}
+	
+}
+
 bool FMovieSceneConstraintChannelHelper::AddConstraintToSequencer(
 	const TSharedPtr<ISequencer>& InSequencer,
 	UTickableTransformConstraint* InConstraint)
@@ -1002,6 +1062,9 @@ bool FMovieSceneConstraintChannelHelper::AddConstraintToSequencer(
 		return false;
 	}
 
+	const bool bIsChildSpawnable = IsHandleSpawnable(InSequencer, InConstraint->ChildTRSHandle);
+	const bool bIsParentSpawnable = IsHandleSpawnable(InSequencer, InConstraint->ParentTRSHandle);
+
 	// create bindings before smart keying so added to spawn copies
 	CreateBindingIDForHandle(InSequencer, InConstraint->ChildTRSHandle);
 	CreateBindingIDForHandle(InSequencer, InConstraint->ParentTRSHandle);
@@ -1013,6 +1076,24 @@ bool FMovieSceneConstraintChannelHelper::AddConstraintToSequencer(
 		return true;
 	}
 
+	// notify of spawnable discrepancy 
+	if (bIsChildSpawnable != bIsParentSpawnable)
+	{
+		const TObjectPtr<UTransformableHandle>& SpawnableHandle = bIsChildSpawnable ? InConstraint->ChildTRSHandle : InConstraint->ParentTRSHandle;
+		const TObjectPtr<UTransformableHandle>& NonSpawnableHandle = bIsChildSpawnable ?  InConstraint->ParentTRSHandle : InConstraint->ChildTRSHandle;
+
+		UE::Constraints::Private::ShowSpawnableDiscrepancyWarning(SpawnableHandle, NonSpawnableHandle);
+	}
+	else if (!bIsChildSpawnable && !bIsParentSpawnable)
+	{
+		const FMovieSceneObjectBindingID& ChildBindingID = InConstraint->ChildTRSHandle->ConstraintBindingID;
+		const FMovieSceneObjectBindingID& ParentBindingID = InConstraint->ParentTRSHandle->ConstraintBindingID;
+		if (ChildBindingID.IsValid() && ParentBindingID.IsValid())
+		{
+			UE::Constraints::Private::ShowNonSpawnableWarning(InConstraint->ChildTRSHandle, InConstraint->ParentTRSHandle);
+		}
+	}
+
 	const FFrameRate TickResolution = InSequencer->GetFocusedTickResolution();
 	const FFrameTime FrameTime = InSequencer->GetLocalTime().ConvertTo(TickResolution);
 	const FFrameNumber Time = FrameTime.GetFrame();
@@ -1324,6 +1405,24 @@ ITransformConstraintChannelInterface* FMovieSceneConstraintChannelHelper::GetHan
 	return InterfaceRegistry.FindConstraintChannelInterface(InHandle->GetClass());
 }
 
+bool FMovieSceneConstraintChannelHelper::IsHandleSpawnable(const TSharedPtr<ISequencer>& InSequencer, const UTransformableHandle* InHandle)
+{
+	if (!InHandle || !InSequencer.IsValid())
+	{
+		return false;
+	}
+	
+	if (USceneComponent* SceneComponent = Cast<USceneComponent>(InHandle->GetTarget().Get()))
+	{
+		if (AActor* Actor = SceneComponent->GetTypedOuter<AActor>())
+		{
+			const TOptional<FMovieSceneSpawnableAnnotation> Spawnable = FMovieSceneSpawnableAnnotation::Find(Actor);
+			return Spawnable.IsSet();
+		}
+	}
+
+	return false;
+}
 
 void FMovieSceneConstraintChannelHelper::CreateBindingIDForHandle(const TSharedPtr<ISequencer>& InSequencer, UTransformableHandle* InHandle)
 {
diff --git a/Engine/Source/Editor/MovieSceneTools/Public/Constraints/MovieSceneConstraintChannelHelper.h b/Engine/Source/Editor/MovieSceneTools/Public/Constraints/MovieSceneConstraintChannelHelper.h
--- a/Engine/Source/Editor/MovieSceneTools/Public/Constraints/MovieSceneConstraintChannelHelper.h
+++ b/Engine/Source/Editor/MovieSceneTools/Public/Constraints/MovieSceneConstraintChannelHelper.h
@@ -214,6 +214,9 @@ private:
 	/** For the given handle create any movie scene binding for it based upon the current sequencer that's open*/
 	static void CreateBindingIDForHandle(const TSharedPtr<ISequencer>& InSequencer, UTransformableHandle* InHandle);
 
+	/** For the given handle, return true if a binding exists and if that binding represents a spawnable. */
+	static bool IsHandleSpawnable(const TSharedPtr<ISequencer>& InSequencer, const UTransformableHandle* InHandle);
+
 	/** Compensate scale keys when enabling/disabling scaling for parent constraints. */
 	static void CompensateScale(
 		UTickableParentConstraint* InParentConstraint,
