Sequencer: When adding keys, extend the section by calling ExpandToFrame() which also calls FixupRelativeKeyframes.

This fixes an issue where the section keys would move when setting play rate keys because they wouldn't be fixed up.

diff --git a/Engine/Source/Editor/MovieSceneTools/Private/CurveKeyEditors/SequencerKeyEditor.h b/Engine/Source/Editor/MovieSceneTools/Private/CurveKeyEditors/SequencerKeyEditor.h
--- a/Engine/Source/Editor/MovieSceneTools/Private/CurveKeyEditors/SequencerKeyEditor.h
+++ b/Engine/Source/Editor/MovieSceneTools/Private/CurveKeyEditors/SequencerKeyEditor.h
@@ -241,13 +241,7 @@ struct TSequencerKeyEditor
 
 				if (bHasAnyKeys)
 				{
-					TRange<FFrameNumber> KeyRange     = TRange<FFrameNumber>(CurrentTime);
-					TRange<FFrameNumber> SectionRange = OwningSection->GetRange();
-
-					if (!SectionRange.Contains(KeyRange))
-					{
-						OwningSection->SetRange(TRange<FFrameNumber>::Hull(KeyRange, SectionRange));
-					}
+					OwningSection->ExpandToFrame(LocalTime.RoundToFrame());
 				}
 			}
 		}
