Sequencer: FAddKeyOperation::Commit now returns newly created keys

This is a refactor so that CreateKeysUnderMouse can be implemented in such a way to allow for the newly created keys to be selected

diff --git a/Engine/Source/Editor/MovieSceneTools/Private/TrackEditors/PropertyTrackEditors/EulerTransformPropertyTrackEditor.cpp b/Engine/Source/Editor/MovieSceneTools/Private/TrackEditors/PropertyTrackEditors/EulerTransformPropertyTrackEditor.cpp
--- a/Engine/Source/Editor/MovieSceneTools/Private/TrackEditors/PropertyTrackEditors/EulerTransformPropertyTrackEditor.cpp
+++ b/Engine/Source/Editor/MovieSceneTools/Private/TrackEditors/PropertyTrackEditors/EulerTransformPropertyTrackEditor.cpp
@@ -45,11 +45,11 @@ TSharedPtr<SWidget> FEulerTransformPropertyTrackEditor::BuildOutlinerEditWidget(
 }
 
 
-void FEulerTransformPropertyTrackEditor::ProcessKeyOperation(FFrameNumber InKeyTime, const UE::Sequencer::FKeyOperation& Operation, ISequencer& InSequencer)
+void FEulerTransformPropertyTrackEditor::ProcessKeyOperation(FFrameNumber InKeyTime, const UE::Sequencer::FKeyOperation& Operation, ISequencer& InSequencer, TArray<UE::Sequencer::FAddKeyResult>* OutResults)
 {
 	using namespace UE::Sequencer;
 
-	auto Iterator = [this, InKeyTime, &Operation, &InSequencer](UMovieSceneTrack* Track, TArrayView<const UE::Sequencer::FKeySectionOperation> Operations)
+	auto Iterator = [this, InKeyTime, &Operation, &InSequencer, OutResults](UMovieSceneTrack* Track, TArrayView<const UE::Sequencer::FKeySectionOperation> Operations)
 	{
 		FGuid ObjectBinding = Track->FindObjectBindingGuid();
 		if (ObjectBinding.IsValid())
@@ -58,25 +58,27 @@ void FEulerTransformPropertyTrackEditor::ProcessKeyOperation(FFrameNumber InKeyT
 			{
 				if (UObject* Object = WeakObject.Get())
 				{
-					this->ProcessKeyOperation(Object, Operations, InSequencer, InKeyTime);
+					this->ProcessKeyOperation(Object, Operations, InSequencer, InKeyTime, OutResults);
 					return;
 				}
 			}
 		}
 
 		// Default behavior
-		FKeyOperation::ApplyOperations(InKeyTime, Operations, ObjectBinding, InSequencer);
+		FKeyOperation::ApplyOperations(InKeyTime, Operations, ObjectBinding, InSequencer, OutResults);
 	};
 
 	Operation.IterateOperations(Iterator);
 }
 
 
-void FEulerTransformPropertyTrackEditor::ProcessKeyOperation(UObject* ObjectToKey, TArrayView<const UE::Sequencer::FKeySectionOperation> SectionsToKey, ISequencer& InSequencer, FFrameNumber KeyTime)
+void FEulerTransformPropertyTrackEditor::ProcessKeyOperation(UObject* ObjectToKey, TArrayView<const UE::Sequencer::FKeySectionOperation> SectionsToKey, ISequencer& InSequencer, FFrameNumber KeyTime, TArray<UE::Sequencer::FAddKeyResult>* OutResults)
 {
 	using namespace UE::MovieScene;
 	using namespace UE::Sequencer;
 
+	TArray<FAddKeyResult> Results;
+
 	FMovieSceneSequenceID SequenceID = GetSequencer()->GetFocusedTemplateID();
 	const FMovieSceneRootEvaluationTemplateInstance& EvaluationTemplate = GetSequencer()->GetEvaluationTemplate();
 
@@ -138,11 +140,20 @@ void FEulerTransformPropertyTrackEditor::ProcessKeyOperation(UObject* ObjectToKe
 
 					double Value = RecomposedTransform[Handle.GetChannelIndex()];
 					EMovieSceneKeyInterpolation Interpolation = GetInterpolationMode(Channel, KeyTime, InSequencer.GetKeyInterpolation());
-					AddKeyToChannel(Channel, KeyTime, Value, Interpolation);
+					FKeyHandle KeyHandle = AddKeyToChannel(Channel, KeyTime, Value, Interpolation);
+					
+					FAddKeyResult Result;
+					Result.KeyArea = KeyArea;
+					Result.KeyHandle = KeyHandle;
+					Results.Add(Result);
 				}
 				else
 				{
-					KeyArea->AddOrUpdateKey(KeyTime, FGuid(), InSequencer);
+					FKeyHandle KeyHandle = KeyArea->AddOrUpdateKey(KeyTime, FGuid(), InSequencer);
+					FAddKeyResult Result;
+					Result.KeyArea = KeyArea;
+					Result.KeyHandle = KeyHandle;
+					Results.Add(Result);
 				}
 			}
 		}
diff --git a/Engine/Source/Editor/MovieSceneTools/Private/TrackEditors/PropertyTrackEditors/EulerTransformPropertyTrackEditor.h b/Engine/Source/Editor/MovieSceneTools/Private/TrackEditors/PropertyTrackEditors/EulerTransformPropertyTrackEditor.h
--- a/Engine/Source/Editor/MovieSceneTools/Private/TrackEditors/PropertyTrackEditors/EulerTransformPropertyTrackEditor.h
+++ b/Engine/Source/Editor/MovieSceneTools/Private/TrackEditors/PropertyTrackEditors/EulerTransformPropertyTrackEditor.h
@@ -13,6 +13,8 @@
 #include "Sections/MovieScene3DTransformSection.h"
 #include "EulerTransform.h"
 
+namespace UE::Sequencer { struct FAddKeyResult; }
+
 /**
  * A property track editor for euler transforms.
  */
@@ -61,8 +63,8 @@ protected:
 
 private:
 
-	void ProcessKeyOperation(FFrameNumber InKeyTime, const UE::Sequencer::FKeyOperation& Operation, ISequencer& InSequencer);
-	void ProcessKeyOperation(UObject* ObjectToKey, TArrayView<const UE::Sequencer::FKeySectionOperation> SectionsToKey, ISequencer& InSequencer, FFrameNumber KeyTime);
+	void ProcessKeyOperation(FFrameNumber InKeyTime, const UE::Sequencer::FKeyOperation& Operation, ISequencer& InSequencer, TArray<UE::Sequencer::FAddKeyResult>* OutResults = nullptr);
+	void ProcessKeyOperation(UObject* ObjectToKey, TArrayView<const UE::Sequencer::FKeySectionOperation> SectionsToKey, ISequencer& InSequencer, FFrameNumber KeyTime, TArray<UE::Sequencer::FAddKeyResult>* OutResults = nullptr);
 
 	FEulerTransform RecomposeTransform(const FEulerTransform& InTransform, UObject* AnimatedObject, UMovieSceneSection* Section);
 };
diff --git a/Engine/Source/Editor/MovieSceneTools/Private/TrackEditors/SubTrackEditor.cpp b/Engine/Source/Editor/MovieSceneTools/Private/TrackEditors/SubTrackEditor.cpp
--- a/Engine/Source/Editor/MovieSceneTools/Private/TrackEditors/SubTrackEditor.cpp
+++ b/Engine/Source/Editor/MovieSceneTools/Private/TrackEditors/SubTrackEditor.cpp
@@ -365,17 +365,22 @@ double FSubTrackEditor::UnwindChannel(const double& OldValue, double NewValue)
 	return NewValue;
 }
 
-void FSubTrackEditor::ProcessKeyOperation(FFrameNumber InKeyTime, const UE::Sequencer::FKeyOperation& Operation, ISequencer& InSequencer)
+void FSubTrackEditor::ProcessKeyOperation(FFrameNumber InKeyTime, const UE::Sequencer::FKeyOperation& Operation, ISequencer& InSequencer, TArray<UE::Sequencer::FAddKeyResult>* OutResults)
 {
-	auto Iterator = [this, InKeyTime, &InSequencer](UMovieSceneTrack* Track, TArrayView<const UE::Sequencer::FKeySectionOperation> Operations)
+	using namespace UE::Sequencer;
+
+	auto Iterator = [this, InKeyTime, &InSequencer, OutResults](UMovieSceneTrack* Track, TArrayView<const UE::Sequencer::FKeySectionOperation> Operations)
 	{
-		this->ProcessKeyOperationInternal(Operations, InSequencer, InKeyTime);
+		this->ProcessKeyOperationInternal(Operations, InSequencer, InKeyTime, OutResults);
 	};
 
 	Operation.IterateOperations(Iterator);
 }
-void FSubTrackEditor::ProcessKeyOperationInternal(TArrayView<const UE::Sequencer::FKeySectionOperation> SectionsToKey, ISequencer& InSequencer, FFrameNumber KeyTime)
+
+void FSubTrackEditor::ProcessKeyOperationInternal(TArrayView<const UE::Sequencer::FKeySectionOperation> SectionsToKey, ISequencer& InSequencer, FFrameNumber KeyTime, TArray<UE::Sequencer::FAddKeyResult>* OutResults)
 {
+	using namespace UE::Sequencer;
+
 	for (int32 Index = 0; Index < SectionsToKey.Num(); ++Index)
 	{
 		for (TSharedPtr<IKeyArea> KeyArea : SectionsToKey[Index].KeyAreas)
@@ -431,12 +436,28 @@ void FSubTrackEditor::ProcessKeyOperationInternal(TArrayView<const UE::Sequencer
 					}
 
 					EMovieSceneKeyInterpolation Interpolation = GetInterpolationMode(Channel, KeyTime, InSequencer.GetKeyInterpolation());
-					AddKeyToChannel(Channel, KeyTime, Value, Interpolation);
+					FKeyHandle KeyHandle = AddKeyToChannel(Channel, KeyTime, Value, Interpolation);
+
+					if (OutResults)
+					{
+						FAddKeyResult Result;
+						Result.KeyArea = KeyArea;
+						Result.KeyHandle = KeyHandle;
+						OutResults->Add(Result);
+					}
 				}
 			}
 			else
 			{
-				KeyArea->AddOrUpdateKey(KeyTime, FGuid(), InSequencer);
+				FKeyHandle KeyHandle = KeyArea->AddOrUpdateKey(KeyTime, FGuid(), InSequencer);
+				
+				if (OutResults)
+				{
+					FAddKeyResult Result;
+					Result.KeyArea = KeyArea;
+					Result.KeyHandle = KeyHandle;
+					OutResults->Add(Result);
+				}
 			}
 		}
 	}
diff --git a/Engine/Source/Editor/MovieSceneTools/Private/TrackEditors/TimeWarpTrackEditor.cpp b/Engine/Source/Editor/MovieSceneTools/Private/TrackEditors/TimeWarpTrackEditor.cpp
--- a/Engine/Source/Editor/MovieSceneTools/Private/TrackEditors/TimeWarpTrackEditor.cpp
+++ b/Engine/Source/Editor/MovieSceneTools/Private/TrackEditors/TimeWarpTrackEditor.cpp
@@ -341,9 +341,9 @@ TSharedPtr<UE::Sequencer::FTrackModel> FTimeWarpTrackEditor::CreateTrackModel(UM
 	return nullptr;
 }
 
-void FTimeWarpTrackEditor::ProcessKeyOperation(FFrameNumber InKeyTime, const UE::Sequencer::FKeyOperation& Operation, ISequencer& InSequencer)
+void FTimeWarpTrackEditor::ProcessKeyOperation(FFrameNumber InKeyTime, const UE::Sequencer::FKeyOperation& Operation, ISequencer& InSequencer, TArray<UE::Sequencer::FAddKeyResult>* OutResults)
 {
-	Operation.ApplyDefault(InKeyTime, InSequencer);
+	Operation.ApplyDefault(InKeyTime, InSequencer, OutResults);
 }
 
 void FTimeWarpTrackEditor::BuildAddTrackMenu(FMenuBuilder& MenuBuilder)
diff --git a/Engine/Source/Editor/MovieSceneTools/Private/TrackEditors/TransformTrackEditor.cpp b/Engine/Source/Editor/MovieSceneTools/Private/TrackEditors/TransformTrackEditor.cpp
--- a/Engine/Source/Editor/MovieSceneTools/Private/TrackEditors/TransformTrackEditor.cpp
+++ b/Engine/Source/Editor/MovieSceneTools/Private/TrackEditors/TransformTrackEditor.cpp
@@ -1231,11 +1231,11 @@ FTransformData F3DTransformTrackEditor::RecomposeTransform(const FTransformData&
 	return FTransformData(CurrentTransform.GetLocation(), CurrentTransform.GetRotation().Rotator(), CurrentTransform.GetScale3D());
 }
 
-void F3DTransformTrackEditor::ProcessKeyOperation(FFrameNumber InKeyTime, const UE::Sequencer::FKeyOperation& Operation, ISequencer& InSequencer)
+void F3DTransformTrackEditor::ProcessKeyOperation(FFrameNumber InKeyTime, const UE::Sequencer::FKeyOperation& Operation, ISequencer& InSequencer, TArray<UE::Sequencer::FAddKeyResult>* OutResults)
 {
 	using namespace UE::Sequencer;
 
-	auto Iterator = [this, InKeyTime, &Operation, &InSequencer](UMovieSceneTrack* Track, TArrayView<const UE::Sequencer::FKeySectionOperation> Operations)
+	auto Iterator = [this, InKeyTime, &Operation, &InSequencer, OutResults](UMovieSceneTrack* Track, TArrayView<const UE::Sequencer::FKeySectionOperation> Operations)
 	{
 		FGuid ObjectBinding = Track->FindObjectBindingGuid();
 		if (ObjectBinding.IsValid())
@@ -1244,20 +1244,19 @@ void F3DTransformTrackEditor::ProcessKeyOperation(FFrameNumber InKeyTime, const
 			{
 				if (UObject* Object = WeakObject.Get())
 				{
-					this->ProcessKeyOperation(Object, Operations, InSequencer, InKeyTime);
+					this->ProcessKeyOperation(Object, Operations, InSequencer, InKeyTime, OutResults);
 					return;
 				}
 			}
 		}
 
 		// Default behavior
-		FKeyOperation::ApplyOperations(InKeyTime, Operations, ObjectBinding, InSequencer);
+		FKeyOperation::ApplyOperations(InKeyTime, Operations, ObjectBinding, InSequencer, OutResults);
 	};
 
 	Operation.IterateOperations(Iterator);
 }
 
-
 int32 GetPreviousKey(FMovieSceneDoubleChannel& Channel, FFrameNumber Time)
 {
 	TArray<FFrameNumber> KeyTimes;
@@ -1548,7 +1547,7 @@ void F3DTransformTrackEditor::OnTransformPropertyChanged(const FPropertyChangedP
 	AnimatablePropertyChanged(FOnKeyProperty::CreateLambda(KeyProperty));
 }
 
-void F3DTransformTrackEditor::ProcessKeyOperation(UObject* ObjectToKey, TArrayView<const UE::Sequencer::FKeySectionOperation> SectionsToKey, ISequencer& InSequencer, FFrameNumber KeyTime)
+void F3DTransformTrackEditor::ProcessKeyOperation(UObject* ObjectToKey, TArrayView<const UE::Sequencer::FKeySectionOperation> SectionsToKey, ISequencer& InSequencer, FFrameNumber KeyTime, TArray<UE::Sequencer::FAddKeyResult>* OutResults)
 {
 	using namespace UE::MovieScene;
 	using namespace UE::Sequencer;
@@ -1598,12 +1597,28 @@ void F3DTransformTrackEditor::ProcessKeyOperation(UObject* ObjectToKey, TArrayVi
 					}
 
 					EMovieSceneKeyInterpolation Interpolation = GetInterpolationMode(Channel, KeyTime, InSequencer.GetKeyInterpolation());
-					AddKeyToChannel(Channel, KeyTime, Value, Interpolation);
+					FKeyHandle KeyHandle = AddKeyToChannel(Channel, KeyTime, Value, Interpolation);
+
+					if (OutResults)
+					{
+						FAddKeyResult Result;
+						Result.KeyArea = KeyArea;
+						Result.KeyHandle = KeyHandle;
+						OutResults->Add(Result);
+					}
 				}
 			}
 			else
 			{
-				KeyArea->AddOrUpdateKey(KeyTime, FGuid(), InSequencer);
+				FKeyHandle KeyHandle = KeyArea->AddOrUpdateKey(KeyTime, FGuid(), InSequencer);
+
+				if (OutResults)
+				{
+					FAddKeyResult Result;
+					Result.KeyArea = KeyArea;
+					Result.KeyHandle = KeyHandle;
+					OutResults->Add(Result);
+				}
 			}
 		}
 	}
diff --git a/Engine/Source/Editor/MovieSceneTools/Private/TrackEditors/TransformTrackEditor.h b/Engine/Source/Editor/MovieSceneTools/Private/TrackEditors/TransformTrackEditor.h
--- a/Engine/Source/Editor/MovieSceneTools/Private/TrackEditors/TransformTrackEditor.h
+++ b/Engine/Source/Editor/MovieSceneTools/Private/TrackEditors/TransformTrackEditor.h
@@ -23,7 +23,7 @@ class FLevelEditorViewportClient;
 class SHorizontalBox;
 class UTickableTransformConstraint;
 
-namespace UE { namespace MovieScene { struct FIntermediate3DTransform; } }
+namespace UE::Sequencer { struct FAddKeyResult; }
 
 /**
  * Tools for animatable transforms
@@ -54,7 +54,7 @@ public:
 public:
 
 	// ISequencerTrackEditor interface
-	virtual void ProcessKeyOperation(FFrameNumber InKeyTime, const UE::Sequencer::FKeyOperation& Operation, ISequencer& InSequencer) override;
+	virtual void ProcessKeyOperation(FFrameNumber InKeyTime, const UE::Sequencer::FKeyOperation& Operation, ISequencer& InSequencer, TArray<UE::Sequencer::FAddKeyResult>* OutResults = nullptr) override;
 	virtual void BuildObjectBindingColumnWidgets(TFunctionRef<TSharedRef<SHorizontalBox>()> GetEditBox, const UE::Sequencer::TViewModelPtr<UE::Sequencer::FObjectBindingModel>& ObjectBinding, const UE::Sequencer::FCreateOutlinerViewParams& InParams, const FName& InColumnName) override;
 	virtual void BuildObjectBindingTrackMenu(FMenuBuilder& MenuBuilder, const TArray<FGuid>& ObjectBindinsg, const UClass* ObjectClass) override;
 	virtual TSharedRef<ISequencerSection> MakeSectionInterface( UMovieSceneSection& SectionObject, UMovieSceneTrack& Track, FGuid ObjectBinding ) override;
@@ -79,7 +79,7 @@ private:
 
 	void OnTransformPropertyChanged(const FPropertyChangedParams& PropertyChangedParams, EMovieSceneTransformChannel TransformChannel);
 
-	void ProcessKeyOperation(UObject* ObjectToKey, TArrayView<const UE::Sequencer::FKeySectionOperation> SectionsToKey, ISequencer& InSequencer, FFrameNumber KeyTime);
+	void ProcessKeyOperation(UObject* ObjectToKey, TArrayView<const UE::Sequencer::FKeySectionOperation> SectionsToKey, ISequencer& InSequencer, FFrameNumber KeyTime, TArray<UE::Sequencer::FAddKeyResult>* OutResults = nullptr);
 
 	/** Returns whether or not a transform track can be added for an actor with a specific handle. */
 	bool CanAddTransformTrackForActorHandle(FGuid ActorHandle) const;
diff --git a/Engine/Source/Editor/MovieSceneTools/Public/TrackEditors/SubTrackEditor.h b/Engine/Source/Editor/MovieSceneTools/Public/TrackEditors/SubTrackEditor.h
--- a/Engine/Source/Editor/MovieSceneTools/Public/TrackEditors/SubTrackEditor.h
+++ b/Engine/Source/Editor/MovieSceneTools/Public/TrackEditors/SubTrackEditor.h
@@ -19,6 +19,8 @@ struct FAssetData;
 class FMenuBuilder;
 class UMovieSceneSubTrack;
 
+namespace UE::Sequencer { struct FAddKeyResult; }
+
 /**
  * Tools for subsequences
  */
@@ -48,7 +50,7 @@ public:
 public:
 
 	// ISequencerTrackEditor interface
-	virtual void ProcessKeyOperation(FFrameNumber InKeyTime, const UE::Sequencer::FKeyOperation& Operation, ISequencer& InSequencer) override;
+	virtual void ProcessKeyOperation(FFrameNumber InKeyTime, const UE::Sequencer::FKeyOperation& Operation, ISequencer& InSequencer, TArray<UE::Sequencer::FAddKeyResult>* OutResults = nullptr) override;
 	virtual void BuildAddTrackMenu(FMenuBuilder& MenuBuilder) override;
 	virtual TSharedPtr<SWidget> BuildOutlinerEditWidget(const FGuid& ObjectBinding, UMovieSceneTrack* Track, const FBuildEditWidgetParams& Params) override;
 	virtual TSharedRef<ISequencerSection> MakeSectionInterface(UMovieSceneSection& SectionObject, UMovieSceneTrack& Track, FGuid ObjectBinding) override;
@@ -176,7 +178,7 @@ private:
 	FKeyPropertyResult HandleSequenceAdded(FFrameNumber KeyTime, UMovieSceneSequence* Sequence, UMovieSceneTrack* Track, int32 RowIndex);
 
 	/** Handles adding keys to section.**/
-	void ProcessKeyOperationInternal(TArrayView<const UE::Sequencer::FKeySectionOperation> SectionsToKey, ISequencer& Sequencer, FFrameNumber KeyTime);
+	void ProcessKeyOperationInternal(TArrayView<const UE::Sequencer::FKeySectionOperation> SectionsToKey, ISequencer& Sequencer, FFrameNumber KeyTime, TArray<UE::Sequencer::FAddKeyResult>* OutResults = nullptr);
 
 	/** Helper for creating new keys.*/
 	void GetOriginKeys(const FVector& CurrentPosition, const FRotator& CurrentRotation, UMovieSceneSection* Section, FGeneratedTrackKeys& OutGeneratedKeys);
diff --git a/Engine/Source/Editor/MovieSceneTools/Public/TrackEditors/TimeWarpTrackEditor.h b/Engine/Source/Editor/MovieSceneTools/Public/TrackEditors/TimeWarpTrackEditor.h
--- a/Engine/Source/Editor/MovieSceneTools/Public/TrackEditors/TimeWarpTrackEditor.h
+++ b/Engine/Source/Editor/MovieSceneTools/Public/TrackEditors/TimeWarpTrackEditor.h
@@ -8,6 +8,8 @@
 #include "MVVM/Extensions/DynamicExtensionContainer.h"
 #include "MVVM/Extensions/ITimeDomainExtension.h"
 
+namespace UE::Sequencer { struct FAddKeyResult; }
+
 namespace UE::Sequencer
 {
 
@@ -72,7 +74,7 @@ public:
 
 private:
 	virtual void BuildAddTrackMenu(FMenuBuilder& MenuBuilder) override;
-	virtual void ProcessKeyOperation(FFrameNumber InKeyTime, const UE::Sequencer::FKeyOperation& Operation, ISequencer& InSequencer) override;
+	virtual void ProcessKeyOperation(FFrameNumber InKeyTime, const UE::Sequencer::FKeyOperation& Operation, ISequencer& InSequencer, TArray<UE::Sequencer::FAddKeyResult>* OutResults = nullptr) override;
 	virtual TSharedRef<ISequencerSection> MakeSectionInterface(UMovieSceneSection& SectionObject, UMovieSceneTrack& Track, FGuid ObjectBinding) override;
 	virtual bool SupportsSequence(UMovieSceneSequence* InSequence) const override;
 
diff --git a/Engine/Source/Editor/Sequencer/Private/SequencerAddKeyOperation.cpp b/Engine/Source/Editor/Sequencer/Private/SequencerAddKeyOperation.cpp
--- a/Engine/Source/Editor/Sequencer/Private/SequencerAddKeyOperation.cpp
+++ b/Engine/Source/Editor/Sequencer/Private/SequencerAddKeyOperation.cpp
@@ -175,12 +175,12 @@ bool FAddKeyOperation::ProcessKeyArea(ISequencerTrackEditor* InTrackEditor, TSha
 	return false;
 }
 
-void FAddKeyOperation::Commit(FFrameNumber KeyTime, ISequencer& InSequencer)
+void FAddKeyOperation::Commit(FFrameNumber KeyTime, ISequencer& InSequencer, TArray<FAddKeyResult>* OutResults)
 {
 	for (TTuple<ISequencerTrackEditor*, FKeyOperation>& Pair : OperationsByTrackEditor)
 	{
 		Pair.Value.InitializeOperation(KeyTime);
-		Pair.Key->ProcessKeyOperation(KeyTime, Pair.Value, InSequencer);
+		Pair.Key->ProcessKeyOperation(KeyTime, Pair.Value, InSequencer, OutResults);
 	}
 
 	InSequencer.UpdatePlaybackRange();
diff --git a/Engine/Source/Editor/Sequencer/Private/SequencerKeyParams.cpp b/Engine/Source/Editor/Sequencer/Private/SequencerKeyParams.cpp
--- a/Engine/Source/Editor/Sequencer/Private/SequencerKeyParams.cpp
+++ b/Engine/Source/Editor/Sequencer/Private/SequencerKeyParams.cpp
@@ -66,24 +66,32 @@ void FKeyOperation::InitializeOperation(FFrameNumber InKeyTime)
 	}
 }
 
-void FKeyOperation::ApplyDefault(FFrameNumber InKeyTime, ISequencer& InSequencer) const
+void FKeyOperation::ApplyDefault(FFrameNumber InKeyTime, ISequencer& InSequencer, TArray<FAddKeyResult>* OutResults) const
 {
-	auto Iterator = [InKeyTime, &InSequencer](UMovieSceneTrack* Track, TArrayView<const FKeySectionOperation> Operations)
+	auto Iterator = [InKeyTime, &InSequencer, &OutResults](UMovieSceneTrack* Track, TArrayView<const FKeySectionOperation> Operations)
 	{
-		ApplyOperations(InKeyTime, Operations, Track->FindObjectBindingGuid(), InSequencer);
+		ApplyOperations(InKeyTime, Operations, Track->FindObjectBindingGuid(), InSequencer, OutResults);
 	};
 
 	IterateOperations(Iterator);
 }
 
-void FKeyOperation::ApplyOperations(FFrameNumber InKeyTime, TArrayView<const FKeySectionOperation> InOperations, const FGuid& ObjectBindingID, ISequencer& InSequencer)
+void FKeyOperation::ApplyOperations(FFrameNumber InKeyTime, TArrayView<const FKeySectionOperation> InOperations, const FGuid& ObjectBindingID, ISequencer& InSequencer, TArray<FAddKeyResult>* OutResults)
 {
 	// @todo: Do we even want to support multiple???
 	for (const FKeySectionOperation& Operation : InOperations)
 	{
 		for (TSharedPtr<IKeyArea> KeyArea : Operation.KeyAreas)
 		{
-			KeyArea->AddOrUpdateKey(InKeyTime, ObjectBindingID, InSequencer);
+			FKeyHandle KeyHandle = KeyArea->AddOrUpdateKey(InKeyTime, ObjectBindingID, InSequencer);
+
+			if (OutResults)
+			{
+				FAddKeyResult Result;
+				Result.KeyArea = KeyArea;
+				Result.KeyHandle = KeyHandle;
+				OutResults->Add(Result);
+			}
 		}
 	}
 }
diff --git a/Engine/Source/Editor/Sequencer/Public/ISequencerTrackEditor.h b/Engine/Source/Editor/Sequencer/Public/ISequencerTrackEditor.h
--- a/Engine/Source/Editor/Sequencer/Public/ISequencerTrackEditor.h
+++ b/Engine/Source/Editor/Sequencer/Public/ISequencerTrackEditor.h
@@ -10,6 +10,7 @@
 #include "ISequencerSection.h"
 #include "MovieSceneTrack.h"
 #include "SequencerKeyParams.h"
+#include "SequencerAddKeyOperation.h"
 #include "Sections/MovieScene3DTransformSection.h"
 #include "MVVM/Extensions/IOutlinerExtension.h"
 #include "MVVM/ViewModelPtr.h"
@@ -108,8 +109,9 @@ public:
 	 * @param InKeyTime   The time at which to add keys
 	 * @param Operation   Structure containing all the sections and channels to key
 	 * @param InSequencer The sequencer UI that is applying the key operation
+	 * @param OutResults  The key handles that were added with this operation
 	 */
-	virtual void ProcessKeyOperation(FFrameNumber InKeyTime, const UE::Sequencer::FKeyOperation& Operation, ISequencer& InSequencer) { Operation.ApplyDefault(InKeyTime, InSequencer); }
+	virtual void ProcessKeyOperation(FFrameNumber InKeyTime, const UE::Sequencer::FKeyOperation& Operation, ISequencer& InSequencer, TArray<UE::Sequencer::FAddKeyResult>* OutResults = nullptr) { Operation.ApplyDefault(InKeyTime, InSequencer, OutResults); }
 
 	/**
 	 * Add a new track to the sequence.
diff --git a/Engine/Source/Editor/Sequencer/Public/SequencerAddKeyOperation.h b/Engine/Source/Editor/Sequencer/Public/SequencerAddKeyOperation.h
--- a/Engine/Source/Editor/Sequencer/Public/SequencerAddKeyOperation.h
+++ b/Engine/Source/Editor/Sequencer/Public/SequencerAddKeyOperation.h
@@ -61,8 +61,9 @@ struct SEQUENCER_API FAddKeyOperation
 	 *
 	 * @param InKeyTime     The time to add keys at
 	 * @param InSequencer   The sequencer instance that is performing this operations
+	 * &param OutResults    The key handles that were added with this operation
 	 */
-	void Commit(FFrameNumber InKeyTime, ISequencer& InSequencer);
+	void Commit(FFrameNumber InKeyTime, ISequencer& InSequencer, TArray<FAddKeyResult> *OutResults = nullptr);
 
 private:
 
diff --git a/Engine/Source/Editor/Sequencer/Public/SequencerKeyParams.h b/Engine/Source/Editor/Sequencer/Public/SequencerKeyParams.h
--- a/Engine/Source/Editor/Sequencer/Public/SequencerKeyParams.h
+++ b/Engine/Source/Editor/Sequencer/Public/SequencerKeyParams.h
@@ -6,6 +6,7 @@
 #include "Containers/ArrayView.h"
 #include "Containers/ContainerAllocationPolicies.h"
 #include "Containers/Map.h"
+#include "Curves/KeyHandle.h"
 #include "HAL/PlatformCrt.h"
 #include "Misc/FrameNumber.h"
 #include "Misc/Guid.h"
@@ -25,6 +26,12 @@ namespace UE
 namespace Sequencer
 {
 
+struct SEQUENCER_API FAddKeyResult
+{
+	TSharedPtr<IKeyArea> KeyArea;
+	FKeyHandle KeyHandle;
+};
+
 /**
  * Structure defining a section that needs to be keyed by specific channels
  */
@@ -63,9 +70,9 @@ struct FKeyOperation
 	 *
 	 * @param InKeyTime      The time at which to add keys
 	 * @param InSequencer    The current sequencer instance that is adding the keys
+	 * @param OutResults     The key handles that were added with this operation
 	 */
-	SEQUENCER_API void ApplyDefault(FFrameNumber InKeyTime, ISequencer& InSequencer) const;
-
+	SEQUENCER_API void ApplyDefault(FFrameNumber InKeyTime, ISequencer& InSequencer, TArray<FAddKeyResult>* OutResults = nullptr) const;
 
 	/**
 	 * Apply a specific set of section operations
@@ -73,8 +80,9 @@ struct FKeyOperation
 	 * @param InKeyTime      The time at which to add keys
 	 * @param InOperations   The operations to apply
 	 * @param InSequencer    The current sequencer instance that is adding the keys
+	 * @param OutResults     The key handles that were added with this operation
 	 */
-	SEQUENCER_API static void ApplyOperations(FFrameNumber InKeyTime, TArrayView<const FKeySectionOperation> InOperations, const FGuid& ObjectBindingID, ISequencer& InSequencer);
+	SEQUENCER_API static void ApplyOperations(FFrameNumber InKeyTime, TArrayView<const FKeySectionOperation> InOperations, const FGuid& ObjectBindingID, ISequencer& InSequencer, TArray<FAddKeyResult>* OutResults = nullptr);
 
 public:
 
diff --git a/Engine/Source/Editor/UMGEditor/Private/Animation/Sequencer2DTransformTrackEditor.cpp b/Engine/Source/Editor/UMGEditor/Private/Animation/Sequencer2DTransformTrackEditor.cpp
--- a/Engine/Source/Editor/UMGEditor/Private/Animation/Sequencer2DTransformTrackEditor.cpp
+++ b/Engine/Source/Editor/UMGEditor/Private/Animation/Sequencer2DTransformTrackEditor.cpp
@@ -158,11 +158,11 @@ FWidgetTransform F2DTransformTrackEditor::RecomposeTransform(const FWidgetTransf
 	return InTransform;
 }
 
-void F2DTransformTrackEditor::ProcessKeyOperation(FFrameNumber InKeyTime, const UE::Sequencer::FKeyOperation& Operation, ISequencer& InSequencer)
+void F2DTransformTrackEditor::ProcessKeyOperation(FFrameNumber InKeyTime, const UE::Sequencer::FKeyOperation& Operation, ISequencer& InSequencer, TArray<UE::Sequencer::FAddKeyResult> *OutResults)
 {
 	using namespace UE::Sequencer;
 
-	auto Iterator = [this, InKeyTime, &Operation, &InSequencer](UMovieSceneTrack* Track, TArrayView<const UE::Sequencer::FKeySectionOperation> Operations)
+	auto Iterator = [this, InKeyTime, &Operation, &InSequencer, OutResults](UMovieSceneTrack* Track, TArrayView<const UE::Sequencer::FKeySectionOperation> Operations)
 	{
 		FGuid ObjectBinding = Track->FindObjectBindingGuid();
 		if (ObjectBinding.IsValid())
@@ -171,20 +171,20 @@ void F2DTransformTrackEditor::ProcessKeyOperation(FFrameNumber InKeyTime, const
 			{
 				if (UObject* Object = WeakObject.Get())
 				{
-					this->ProcessKeyOperation(Object, Operations, InSequencer, InKeyTime);
+					this->ProcessKeyOperation(Object, Operations, InSequencer, InKeyTime, OutResults);
 					return;
 				}
 			}
 		}
 
 		// Default behavior
-		FKeyOperation::ApplyOperations(InKeyTime, Operations, ObjectBinding, InSequencer);
+		FKeyOperation::ApplyOperations(InKeyTime, Operations, ObjectBinding, InSequencer, OutResults);
 	};
 
 	Operation.IterateOperations(Iterator);
 }
 
-void F2DTransformTrackEditor::ProcessKeyOperation(UObject* ObjectToKey, TArrayView<const UE::Sequencer::FKeySectionOperation> SectionsToKey, ISequencer& InSequencer, FFrameNumber KeyTime)
+void F2DTransformTrackEditor::ProcessKeyOperation(UObject* ObjectToKey, TArrayView<const UE::Sequencer::FKeySectionOperation> SectionsToKey, ISequencer& InSequencer, FFrameNumber KeyTime, TArray<UE::Sequencer::FAddKeyResult> *OutResults)
 {
 	using namespace UE::Sequencer;
 
@@ -221,7 +221,15 @@ void F2DTransformTrackEditor::ProcessKeyOperation(UObject* ObjectToKey, TArrayVi
 
 				if (Handle.GetChannelTypeName() == FMovieSceneFloatChannel::StaticStruct()->GetFName())
 				{
-					AddKeyToChannel(static_cast<FMovieSceneFloatChannel*>(Channel), KeyTime, Value, InSequencer.GetKeyInterpolation());
+					FKeyHandle KeyHandle = AddKeyToChannel(static_cast<FMovieSceneFloatChannel*>(Channel), KeyTime, Value, InSequencer.GetKeyInterpolation());
+
+					if (OutResults)
+					{
+						FAddKeyResult Result;
+						Result.KeyArea = KeyArea;
+						Result.KeyHandle = KeyHandle;
+						OutResults->Add(Result);
+					}
 				}
 			}
 		}
diff --git a/Engine/Source/Editor/UMGEditor/Private/Animation/Sequencer2DTransformTrackEditor.h b/Engine/Source/Editor/UMGEditor/Private/Animation/Sequencer2DTransformTrackEditor.h
--- a/Engine/Source/Editor/UMGEditor/Private/Animation/Sequencer2DTransformTrackEditor.h
+++ b/Engine/Source/Editor/UMGEditor/Private/Animation/Sequencer2DTransformTrackEditor.h
@@ -25,6 +25,7 @@ class UMovieSceneTrack;
 class UObject;
 namespace UE::Sequencer { struct FKeyOperation; }
 namespace UE::Sequencer { struct FKeySectionOperation; }
+namespace UE::Sequencer { struct FAddKeyResult; }
 struct FFrameNumber;
 struct FGuid;
 
@@ -64,14 +65,14 @@ protected:
 	// ISequencerTrackEditor interface
 
 	virtual TSharedRef<ISequencerSection> MakeSectionInterface( UMovieSceneSection& SectionObject, UMovieSceneTrack& Track, FGuid ObjectBinding ) override;
-	virtual void ProcessKeyOperation(FFrameNumber InKeyTime, const UE::Sequencer::FKeyOperation& Operation, ISequencer& InSequencer) override;
+	virtual void ProcessKeyOperation(FFrameNumber InKeyTime, const UE::Sequencer::FKeyOperation& Operation, ISequencer& InSequencer, TArray<UE::Sequencer::FAddKeyResult>* OutResults = nullptr) override;
 	// FPropertyTrackEditor interface
 
 	virtual void GenerateKeysFromPropertyChanged( const FPropertyChangedParams& PropertyChangedParams, UMovieSceneSection* SectionToKey, FGeneratedTrackKeys& OutGeneratedKeys ) override;
 
 	FWidgetTransform RecomposeTransform(const FWidgetTransform& InTransform, UObject* AnimatedObject, UMovieSceneSection* Section);
 
-	void ProcessKeyOperation(UObject* ObjectToKey, TArrayView<const UE::Sequencer::FKeySectionOperation> SectionsToKey, ISequencer& InSequencer, FFrameNumber KeyTime);
+	void ProcessKeyOperation(UObject* ObjectToKey, TArrayView<const UE::Sequencer::FKeySectionOperation> SectionsToKey, ISequencer& InSequencer, FFrameNumber KeyTime, TArray<UE::Sequencer::FAddKeyResult>* OutResults = nullptr);
 
 private:
 	static FName TranslationName;
