Sequencer: Clamp dimmed region of subsection paint

diff --git a/Engine/Source/Editor/MovieSceneTools/Private/TrackEditors/SubTrackEditorBase.cpp b/Engine/Source/Editor/MovieSceneTools/Private/TrackEditors/SubTrackEditorBase.cpp
--- a/Engine/Source/Editor/MovieSceneTools/Private/TrackEditors/SubTrackEditorBase.cpp
+++ b/Engine/Source/Editor/MovieSceneTools/Private/TrackEditors/SubTrackEditorBase.cpp
@@ -246,7 +246,7 @@ void FSubSectionPainterUtil::PaintSectionBounds(const UMovieSceneSubSection& Sec
 		OuterToInnerTransform.ExtractBoundariesWithinRange(MIN_int32, SectionTimeRange.GetLowerBoundValue(), SetPreceedingInvalidTime);
 		if (!PreceedingInvalidTime)
 		{
-			const float WidthFactor = static_cast<float>((FirstTime.GetValue().AsDecimal() - SectionTimeRange.GetLowerBoundValue().AsDecimal()) / SectionTimeRange.Size<FFrameTime>().AsDecimal());
+			const float WidthFactor = FMath::Min(1.f, static_cast<float>((FirstTime.GetValue().AsDecimal() - SectionTimeRange.GetLowerBoundValue().AsDecimal()) / SectionTimeRange.Size<FFrameTime>().AsDecimal()));
 
 			FSlateDrawElement::MakeBox(
 				InPainter.DrawElements,
@@ -275,8 +275,8 @@ void FSubSectionPainterUtil::PaintSectionBounds(const UMovieSceneSubSection& Sec
 		OuterToInnerTransform.ExtractBoundariesWithinRange(SectionTimeRange.GetUpperBoundValue(), MAX_int32, SetProceedingInvalidTime);
 		if (!ProceedingInvalidTime)
 		{
-			const float Offset      = InPainter.GetTimeConverter().FrameToPixel(LastTime.GetValue()) + 1;
-			const float WidthFactor = 1.f - static_cast<float>((LastTime.GetValue().AsDecimal() - SectionTimeRange.GetLowerBoundValue().AsDecimal()) / SectionTimeRange.Size<FFrameTime>().AsDecimal());
+			const float Offset      = FMath::Max(0.f, InPainter.GetTimeConverter().FrameToPixel(LastTime.GetValue()) + 1);
+			const float WidthFactor = FMath::Min(1.f, 1.f - static_cast<float>((LastTime.GetValue().AsDecimal() - SectionTimeRange.GetLowerBoundValue().AsDecimal()) / SectionTimeRange.Size<FFrameTime>().AsDecimal()));
 
 			FSlateDrawElement::MakeBox(
 				InPainter.DrawElements,
