Sequencer: Disable post process in thumbnails by default

Need to revisit 37044465 and UE-227099

diff --git a/Engine/Source/Editor/MovieSceneTools/Private/TrackEditorThumbnail/TrackThumbnailUtils.cpp b/Engine/Source/Editor/MovieSceneTools/Private/TrackEditorThumbnail/TrackThumbnailUtils.cpp
--- a/Engine/Source/Editor/MovieSceneTools/Private/TrackEditorThumbnail/TrackThumbnailUtils.cpp
+++ b/Engine/Source/Editor/MovieSceneTools/Private/TrackEditorThumbnail/TrackThumbnailUtils.cpp
@@ -1,4 +1,4 @@
-﻿// Copyright Epic Games, Inc. All Rights Reserved.
+// Copyright Epic Games, Inc. All Rights Reserved.
 
 #include "TrackEditorThumbnail/TrackThumbnailUtils.h"
 
@@ -9,6 +9,11 @@
 #include "MovieSceneToolsUserSettings.h"
 #include "SceneViewExtension.h"
 
+static TAutoConsoleVariable<int32> CVarSequencerUsePostProcessThumbnails(
+	TEXT("Sequencer.UsePostProcessThumbnails"),
+	0, TEXT("Enable post process in thumbnails."),
+	ECVF_Default);
+
 namespace UE::MoveSceneTools
 {
 	void PreDrawThumbnailSetupSequencer(ISequencer& Sequencer, FFrameTime CaptureFrame)
@@ -86,7 +91,7 @@ namespace UE::MoveSceneTools
 
 		FSceneView* NewView = new FSceneView(ViewInitOptions);
 		ViewFamily.Views.Add(NewView);
-		if (OverridePostProcessSettings)
+		if (OverridePostProcessSettings && CVarSequencerUsePostProcessThumbnails->GetInt() > 0)
 		{
 			FPostProcessSettings ProcessSettings = *OverridePostProcessSettings;
 			
