Curve Editor: Fix pasted keys getting their TangentMode overridden if they were auto or smart auto.

This splits GetKeyAttributes up into:
- GetKeyAttributesIncludingAutoComputed: Returns all attributes as they are needed to display in UI.
- GetKeyAttributesExcludingAutoComputed: Returns the attributes as they were set by the user.

diff --git a/Engine/Source/Editor/CurveEditor/Private/CurveEditor.cpp b/Engine/Source/Editor/CurveEditor/Private/CurveEditor.cpp
--- a/Engine/Source/Editor/CurveEditor/Private/CurveEditor.cpp
+++ b/Engine/Source/Editor/CurveEditor/Private/CurveEditor.cpp
@@ -1403,7 +1403,8 @@ void FCurveEditor::CopySelection() const
 					TArrayView<const FKeyHandle> KeyHandles = Pair.Value.AsArray();
 
 					Curve->GetKeyPositions(KeyHandles, CopyableCurveKeys->KeyPositions);
-					Curve->GetKeyAttributes(KeyHandles, CopyableCurveKeys->KeyAttributes);
+					// We need to get the attributes as they were specified by the user: so call the version that skips auto-computed values.
+					Curve->GetKeyAttributesExcludingAutoComputed(KeyHandles, CopyableCurveKeys->KeyAttributes);
 
 					for (int KeyIndex = 0; KeyIndex < CopyableCurveKeys->KeyPositions.Num(); ++KeyIndex)
 					{
diff --git a/Engine/Source/Editor/CurveEditor/Public/CurveModel.h b/Engine/Source/Editor/CurveEditor/Public/CurveModel.h
--- a/Engine/Source/Editor/CurveEditor/Public/CurveModel.h
+++ b/Engine/Source/Editor/CurveEditor/Public/CurveModel.h
@@ -192,7 +192,10 @@ public:
 public:
 
 	/**
-	 * Retrieve all key attributes that pertain to the specified input key handles
+	 * Retrieve ALL key attributes that pertain to the specified input key handles.
+	 *
+	 * This function is legacy. Prefer use of GetKeyAttributesIncludingAutoComputed or GetKeyAttributesExcludingAutoComputed.
+	 * @see GetKeyAttributesIncludingAutoComputed (introduced 5.6).
 	 *
 	 * @param InKeys                Array of key handles to get attributes for
 	 * @param OutAttributes         Array to receive key attributes, one per index of InKeys
@@ -200,6 +203,39 @@ public:
 	virtual void GetKeyAttributes(TArrayView<const FKeyHandle> InKeys, TArrayView<FKeyAttributes> OutAttributes) const
 	{}
 
+	/**
+	 * Gets the value of ALL attributes including those that are auto-computed.
+	 * Use this version if you want to know values e.g. the tangents or weights. Useful for UI code that visualizes tangents.
+	 * The result is not intended to be passed to SetKeyAttributes; calling SetKeyAttributes would correct the passed in attributes to reflect the user settings.
+	 *
+	 * Typically, attributes reflect the settings that the user has manually configured for the keys, so certain attributes may remain unset.
+	 * For instance, when TangentMode == RCTM_Auto, tangents and weights are automatically computed, meaning attributes like ArriveTangent are not explicitly set.
+	 * Setting ArriveTangent would imply a user-defined value, which is incompatible with TangentMode == RCTM_Auto.
+	 *
+	 * In some situations, you don't want to know the values the user has set but the values as they would be used for evaluation, e.g. tangents,
+	 * weights, etc. For these cases, you can use GetKeyAttributesIncludingAutoComputed.
+	 */
+	virtual void GetKeyAttributesIncludingAutoComputed(TArrayView<const FKeyHandle> InKeys, TArrayView<FKeyAttributes> OutAttributes) const
+	{
+		GetKeyAttributes(InKeys, OutAttributes);
+	}
+
+	/**
+	 * Gets the attributes set by the user. This specifically excludes auto-computed values.
+	 * Use this version if you want the "true" state of the keys. You can pass the result of this to SetKeyAttributes. Useful e.g. for copy-pasting keys.
+	 * 
+	 * Typically, attributes reflect the settings that the user has manually configured for the keys, so certain attributes may remain unset.
+	 * For instance, when TangentMode == RCTM_Auto, tangents and weights are automatically computed, meaning attributes like ArriveTangent are not explicitly set.
+	 * Setting ArriveTangent would imply a user-defined value, which is incompatible with TangentMode == RCTM_Auto.
+	 *
+	 * This function returns only user-set values.
+	 *  
+	 * @param InKeys                Array of key handles to get attributes for
+	 * @param OutAttributes         Array to receive key attributes, one per index of InKeys
+	 */
+	virtual void GetKeyAttributesExcludingAutoComputed(TArrayView<const FKeyHandle> InKeys, TArrayView<FKeyAttributes> OutAttributes) const
+	{}
+
 	/**
 	 * Assign key attributes for the specified key handles
 	 *
diff --git a/Engine/Source/Editor/MovieSceneTools/Private/Channels/BezierChannelCurveModel.cpp b/Engine/Source/Editor/MovieSceneTools/Private/Channels/BezierChannelCurveModel.cpp
--- a/Engine/Source/Editor/MovieSceneTools/Private/Channels/BezierChannelCurveModel.cpp
+++ b/Engine/Source/Editor/MovieSceneTools/Private/Channels/BezierChannelCurveModel.cpp
@@ -206,62 +206,104 @@ TPair<ERichCurveInterpMode, ERichCurveTangentMode> FBezierChannelCurveModel<Chan
 	return TPair<ERichCurveInterpMode, ERichCurveTangentMode>(DefaultInterpolationMode, DefaultTangentMode);
 }
 
-template<typename ChannelType, typename ChannelValue, typename KeyType> 
-void FBezierChannelCurveModel<ChannelType, ChannelValue, KeyType>::GetKeyAttributes(TArrayView<const FKeyHandle> InKeys, TArrayView<FKeyAttributes> OutAttributes) const
+static bool IsAuto(ERichCurveTangentMode TangentMode)
 {
-	ChannelType*        Channel    = this->GetChannelHandle().Get();
-	UMovieSceneSection* Section    = this->template GetOwningObjectOrOuter<UMovieSceneSection>();
-	UMovieScene*        MovieScene = Section ? Section->GetTypedOuter<UMovieScene>() : nullptr;
+	return (TangentMode == RCTM_Auto || TangentMode == RCTM_SmartAuto);
+}
 
-	if (Channel && Section && MovieScene)
+/**
+ * Shared implementation for getting keys.
+ *
+ * Typically, attributes reflect the settings that the user has manually configured for the keys, so certain attributes may remain unset.
+ * For instance, when TangentMode == RCTM_Auto, tangents and weights are automatically computed, meaning attributes like ArriveTangent are not explicitly set.
+ * Setting ArriveTangent would imply a user-defined value, which is incompatible with TangentMode == RCTM_Auto.
+ *
+ * Sometimes API users still need a way to get the auto-computed values, like those of tangents, though.
+ * For this, use bAllAttributes parameter:
+ * - true: get all attributes, even the auto-computed ones. The values are useful for e.g. UI visualization. Do not pass them to SetKeyAttributes.
+ * - false: get only the attributes that were set by the user (i.e. skip returning auto-computed values). You can pass them to SetKeyAttributes, e.g. for copy paste.
+ */
+template<bool bAllAttributes, typename ChannelType, typename ChannelValue, typename KeyType>
+static void GetKeyAttributesDetail(
+	ChannelType* InChannel, UMovieSceneSection* InSection, UMovieScene* InMovieScene,
+	TArrayView<const FKeyHandle> InKeys,
+	TArrayView<FKeyAttributes> OutAttributes
+	)
+{
+	if (!InChannel || !InSection || !InMovieScene)
 	{
-		TMovieSceneChannelData<ChannelValue> ChannelData = Channel->GetData();
-		TArrayView<const FFrameNumber>    Times  = ChannelData.GetTimes();
-		TArrayView<ChannelValue> Values = ChannelData.GetValues();
-
-		float TimeInterval = MovieScene->GetTickResolution().AsInterval();
+		return;
+	}
+	
+	TMovieSceneChannelData<ChannelValue> ChannelData = InChannel->GetData();
+	TArrayView<ChannelValue> Values = ChannelData.GetValues();
 
-		for (int32 Index = 0; Index < InKeys.Num(); ++Index)
+	const float TimeInterval = InMovieScene->GetTickResolution().AsInterval();
+	for (int32 Index = 0; Index < InKeys.Num(); ++Index)
+	{
+		const int32 KeyIndex = ChannelData.GetIndex(InKeys[Index]);
+		if (KeyIndex != INDEX_NONE)
 		{
-			const int32 KeyIndex = ChannelData.GetIndex(InKeys[Index]);
-			if (KeyIndex != INDEX_NONE)
+			const ChannelValue& KeyValue    = Values[KeyIndex];
+			FKeyAttributes&     Attributes  = OutAttributes[Index];
+
+			Attributes.SetInterpMode(KeyValue.InterpMode);
+
+			// If the previous key is cubic, show the arrive tangent handle even if this key is constant
+			const int32 PreviousKeyIndex = KeyIndex - 1;
+			const bool bGetArriveTangent = Values.IsValidIndex(PreviousKeyIndex) && Values[PreviousKeyIndex].InterpMode == RCIM_Cubic;
+			if (bAllAttributes && bGetArriveTangent)
 			{
-				const ChannelValue& KeyValue    = Values[KeyIndex];
-				FKeyAttributes&     Attributes  = OutAttributes[Index];
+				Attributes.SetArriveTangent(KeyValue.Tangent.ArriveTangent / TimeInterval);
+			}
 
-				Attributes.SetInterpMode(KeyValue.InterpMode);
+			if ((KeyValue.InterpMode != RCIM_Constant && KeyValue.InterpMode != RCIM_Linear))
+			{
+				Attributes.SetTangentMode(KeyValue.TangentMode);
 
-				// If the previous key is cubic, show the arrive tangent handle even if this key is constant
-				const int32 PreviousKeyIndex = KeyIndex - 1;
-				const bool bGetArriveTangent = Values.IsValidIndex(PreviousKeyIndex) && Values[PreviousKeyIndex].InterpMode == RCIM_Cubic;
-				if (bGetArriveTangent)
+				// The remaining settings (arrive / leave tangent, arrive / leave weight, weight mode) can only be user specified, i.e. cannot be set when in auto.
+				// Setting any of them would imply they specified by the user, which would contradict the setting that it's auto. 
+				if (!bAllAttributes && IsAuto(KeyValue.TangentMode))
 				{
-					Attributes.SetArriveTangent(KeyValue.Tangent.ArriveTangent / TimeInterval);
+					continue;
 				}
+				
+				Attributes.SetArriveTangent(KeyValue.Tangent.ArriveTangent / TimeInterval);
+				Attributes.SetLeaveTangent(KeyValue.Tangent.LeaveTangent / TimeInterval);
 
-				if ((KeyValue.InterpMode != RCIM_Constant && KeyValue.InterpMode != RCIM_Linear))
+				if (KeyValue.InterpMode == RCIM_Cubic)
 				{
-					Attributes.SetTangentMode(KeyValue.TangentMode);
-					Attributes.SetArriveTangent(KeyValue.Tangent.ArriveTangent / TimeInterval);
-					Attributes.SetLeaveTangent(KeyValue.Tangent.LeaveTangent / TimeInterval);
-
-					if (KeyValue.InterpMode == RCIM_Cubic)
+					Attributes.SetTangentWeightMode(KeyValue.Tangent.TangentWeightMode);
+					if (KeyValue.Tangent.TangentWeightMode != RCTWM_WeightedNone)
 					{
-						Attributes.SetTangentWeightMode(KeyValue.Tangent.TangentWeightMode);
-						if (KeyValue.Tangent.TangentWeightMode != RCTWM_WeightedNone)
-						{
-							Attributes.SetArriveTangentWeight(KeyValue.Tangent.ArriveTangentWeight);
-							Attributes.SetLeaveTangentWeight(KeyValue.Tangent.LeaveTangentWeight);
-						}
+						Attributes.SetArriveTangentWeight(KeyValue.Tangent.ArriveTangentWeight);
+						Attributes.SetLeaveTangentWeight(KeyValue.Tangent.LeaveTangentWeight);
 					}
 				}
 			}
 		}
 	}
 }
-static bool IsAuto(ERichCurveTangentMode TangentMode)
+
+template<typename ChannelType, typename ChannelValue, typename KeyType> 
+void FBezierChannelCurveModel<ChannelType, ChannelValue, KeyType>::GetKeyAttributes(TArrayView<const FKeyHandle> InKeys, TArrayView<FKeyAttributes> OutAttributes) const
 {
-	return (TangentMode == RCTM_Auto || TangentMode == RCTM_SmartAuto);
+	ChannelType*        Channel    = this->GetChannelHandle().Get();
+	UMovieSceneSection* Section    = this->template GetOwningObjectOrOuter<UMovieSceneSection>();
+	UMovieScene*        MovieScene = Section ? Section->GetTypedOuter<UMovieScene>() : nullptr;
+	
+	GetKeyAttributesDetail<true, ChannelType, ChannelValue, KeyType>(Channel, Section, MovieScene, InKeys, OutAttributes);
+}
+
+template <typename ChannelType, typename ChannelValue, typename KeyType>
+void FBezierChannelCurveModel<ChannelType, ChannelValue, KeyType>::GetKeyAttributesExcludingAutoComputed(TArrayView<const FKeyHandle> InKeys,
+	TArrayView<FKeyAttributes> OutAttributes) const
+{
+	ChannelType*        Channel    = this->GetChannelHandle().Get();
+	UMovieSceneSection* Section    = this->template GetOwningObjectOrOuter<UMovieSceneSection>();
+	UMovieScene*        MovieScene = Section ? Section->GetTypedOuter<UMovieScene>() : nullptr;
+
+	GetKeyAttributesDetail<false, ChannelType, ChannelValue, KeyType>(Channel, Section, MovieScene, InKeys, OutAttributes);
 }
 
 template<typename ChannelType, typename ChannelValue, typename KeyType> 
diff --git a/Engine/Source/Editor/MovieSceneTools/Public/Channels/BezierChannelCurveModel.h b/Engine/Source/Editor/MovieSceneTools/Public/Channels/BezierChannelCurveModel.h
--- a/Engine/Source/Editor/MovieSceneTools/Public/Channels/BezierChannelCurveModel.h
+++ b/Engine/Source/Editor/MovieSceneTools/Public/Channels/BezierChannelCurveModel.h
@@ -62,6 +62,7 @@ public:
 	virtual void GetKeyDrawInfo(ECurvePointType PointType, const FKeyHandle InKeyHandle, FKeyDrawInfo& OutDrawInfo) const override;
 
 	virtual void GetKeyAttributes(TArrayView<const FKeyHandle> InKeys, TArrayView<FKeyAttributes> OutAttributes) const override;
+	virtual void GetKeyAttributesExcludingAutoComputed(TArrayView<const FKeyHandle> InKeys, TArrayView<FKeyAttributes> OutAttributes) const override;
 	virtual void SetKeyAttributes(TArrayView<const FKeyHandle> InKeys, TArrayView<const FKeyAttributes> InAttributes, EPropertyChangeType::Type ChangeType = EPropertyChangeType::Unspecified) override;
 
 	virtual TPair<ERichCurveInterpMode, ERichCurveTangentMode> GetInterpolationMode(const double& InTime, ERichCurveInterpMode DefaultInterpolationMode, ERichCurveTangentMode DefaultTangentMode) const override;
