Curve editor: Improve documentation of GetKeyAttributes logic.

diff --git a/Engine/Source/Editor/MovieSceneTools/Private/Channels/BezierChannelCurveModel.cpp b/Engine/Source/Editor/MovieSceneTools/Private/Channels/BezierChannelCurveModel.cpp
--- a/Engine/Source/Editor/MovieSceneTools/Private/Channels/BezierChannelCurveModel.cpp
+++ b/Engine/Source/Editor/MovieSceneTools/Private/Channels/BezierChannelCurveModel.cpp
@@ -315,6 +315,17 @@ void FBezierChannelCurveModel<ChannelType, ChannelValue, KeyType>::SetKeyAttribu
 	UMovieSceneSignedObject* SignedOwner = this->template GetOwningObjectOrOuter<UMovieSceneSignedObject>();
 	UMovieScene*             MovieScene  = SignedOwner ? SignedOwner->GetTypedOuter<UMovieScene>() : nullptr;
 
+	// Rule: Auto means tangents and weights are auto computed.
+	// This is used below to correct invalid input InAttributes, e.g. if InAttributes says TangentMode is RCTM_Auto, then ArriveTangent cannot be set.
+	const auto ResetModeAndWeightForAuto = [](ChannelValue& KeyValue)
+	{
+		if (IsAuto(KeyValue.TangentMode))
+		{
+			KeyValue.TangentMode = RCTM_User;
+			KeyValue.Tangent.TangentWeightMode = RCTWM_WeightedNone;
+		}
+	};
+
 	if (Channel && SignedOwner && MovieScene && !this->IsReadOnly())
 	{
 		bool bAutoSetTangents = false;
@@ -373,6 +384,7 @@ void FBezierChannelCurveModel<ChannelType, ChannelValue, KeyType>::SetKeyAttribu
 					{
 						if (KeyValue.TangentMode != RCTM_User && KeyValue.TangentMode != RCTM_Break)
 						{
+							// Input attribute is invalid: Weights can only be set for user or break tangent modes. Correct it.
 							KeyValue.TangentMode = RCTM_User;
 						}
 					}
@@ -381,11 +393,8 @@ void FBezierChannelCurveModel<ChannelType, ChannelValue, KeyType>::SetKeyAttribu
 
 				if (Attributes.HasArriveTangent())
 				{
-					if (IsAuto(KeyValue.TangentMode))
-					{
-						KeyValue.TangentMode = RCTM_User;
-						KeyValue.Tangent.TangentWeightMode = RCTWM_WeightedNone;
-					}
+					// If TangentMode is auto, then the input attribute is invalid: it implies arrive tangent is user specified. Correct TangentMode.
+					ResetModeAndWeightForAuto(KeyValue);
 
 					KeyValue.Tangent.ArriveTangent = Attributes.GetArriveTangent() * TimeInterval;
 					if (KeyValue.InterpMode == RCIM_Cubic && KeyValue.TangentMode != RCTM_Break)
@@ -397,11 +406,8 @@ void FBezierChannelCurveModel<ChannelType, ChannelValue, KeyType>::SetKeyAttribu
 
 				if (Attributes.HasLeaveTangent())
 				{
-					if (IsAuto(KeyValue.TangentMode))
-					{
-						KeyValue.TangentMode = RCTM_User;
-						KeyValue.Tangent.TangentWeightMode = RCTWM_WeightedNone;
-					}
+					// If TangentMode is auto, then the input attribute is invalid: it implies leave tangent is user specified. Correct TangentMode.
+					ResetModeAndWeightForAuto(KeyValue);
 
 					KeyValue.Tangent.LeaveTangent = Attributes.GetLeaveTangent() * TimeInterval;
 					if (KeyValue.InterpMode == RCIM_Cubic && KeyValue.TangentMode != RCTM_Break)
@@ -413,11 +419,8 @@ void FBezierChannelCurveModel<ChannelType, ChannelValue, KeyType>::SetKeyAttribu
 
 				if (Attributes.HasArriveTangentWeight())
 				{
-					if (IsAuto(KeyValue.TangentMode))
-					{
-						KeyValue.TangentMode = RCTM_User;
-						KeyValue.Tangent.TangentWeightMode = RCTWM_WeightedNone;
-					}
+					// If TangentMode is auto, then the input attribute is invalid: it implies arrive weight is user specified. Correct TangentMode.
+					ResetModeAndWeightForAuto(KeyValue);
 
 					KeyValue.Tangent.ArriveTangentWeight = Attributes.GetArriveTangentWeight(); 
 					if (KeyValue.InterpMode == RCIM_Cubic && KeyValue.TangentMode != RCTM_Break)
@@ -429,11 +432,8 @@ void FBezierChannelCurveModel<ChannelType, ChannelValue, KeyType>::SetKeyAttribu
 
 				if (Attributes.HasLeaveTangentWeight())
 				{
-					if (IsAuto(KeyValue.TangentMode))
-					{
-						KeyValue.TangentMode = RCTM_User;
-						KeyValue.Tangent.TangentWeightMode = RCTWM_WeightedNone;
-					}
+					// If TangentMode is auto, then the input attribute is invalid: it implies leave weight is user specified. Correct TangentMode.
+					ResetModeAndWeightForAuto(KeyValue);
 
 					KeyValue.Tangent.LeaveTangentWeight = Attributes.GetLeaveTangentWeight();
 					if (KeyValue.InterpMode == RCIM_Cubic && KeyValue.TangentMode != RCTM_Break)
