Sequencer-  Do not optimize transform track channels to key when using custom transform origin.

Merging external partner changelist CL 38535014.

diff --git a/Engine/Source/Editor/MovieSceneTools/Private/TrackEditors/TransformTrackEditor.cpp b/Engine/Source/Editor/MovieSceneTools/Private/TrackEditors/TransformTrackEditor.cpp
--- a/Engine/Source/Editor/MovieSceneTools/Private/TrackEditors/TransformTrackEditor.cpp
+++ b/Engine/Source/Editor/MovieSceneTools/Private/TrackEditors/TransformTrackEditor.cpp
@@ -844,13 +844,17 @@ void F3DTransformTrackEditor::GetTransformKeys( const TOptional<FTransformData>&
 		SectionChannelProxy.GetChannelByName<FMovieSceneDoubleChannel>("Scale.Z")
 	};
 
+	USceneComponent* SceneComponent = MovieSceneHelpers::SceneComponentFromRuntimeObject(Object);
+	USceneComponent* AttachParent = SceneComponent ? SceneComponent->GetAttachParent() : nullptr;
+	bool bHasNeutralRotationSpace = IsValid(AttachParent) || GetTransformOrigin().GetRotation() == FQuat::Identity;
+
 	// Set translation keys/defaults
 	{
 		bool bKeyX = EnumHasAnyFlags(ChannelsToKey, EMovieSceneTransformChannel::TranslationX);
 		bool bKeyY = EnumHasAnyFlags(ChannelsToKey, EMovieSceneTransformChannel::TranslationY);
 		bool bKeyZ = EnumHasAnyFlags(ChannelsToKey, EMovieSceneTransformChannel::TranslationZ);
 
-		if (bLastVectorIsValid)
+		if (bLastVectorIsValid && bHasNeutralRotationSpace)
 		{
 			bKeyX &= !FMath::IsNearlyEqual(LastTransform->Translation.X, CurrentTransform.Translation.X);
 			bKeyY &= !FMath::IsNearlyEqual(LastTransform->Translation.Y, CurrentTransform.Translation.Y);
@@ -893,9 +897,12 @@ void F3DTransformTrackEditor::GetTransformKeys( const TOptional<FTransformData>&
 		{
 			KeyRotator = UnwindRotator(LastTransform->Rotation, CurrentTransform.Rotation);
 
-			bKeyX &= !FMath::IsNearlyEqual(LastTransform->Rotation.Roll,  KeyRotator.Roll);
-			bKeyY &= !FMath::IsNearlyEqual(LastTransform->Rotation.Pitch, KeyRotator.Pitch);
-			bKeyZ &= !FMath::IsNearlyEqual(LastTransform->Rotation.Yaw,   KeyRotator.Yaw);
+			if (bHasNeutralRotationSpace)
+			{
+				bKeyX &= !FMath::IsNearlyEqual(LastTransform->Rotation.Roll,  KeyRotator.Roll);
+				bKeyY &= !FMath::IsNearlyEqual(LastTransform->Rotation.Pitch, KeyRotator.Pitch);
+				bKeyZ &= !FMath::IsNearlyEqual(LastTransform->Rotation.Yaw,   KeyRotator.Yaw);
+			}
 		}
 
 		if (GetSequencer()->GetKeyGroupMode() == EKeyGroupMode::KeyGroup && ( bKeyX || bKeyY || bKeyZ) )
