Sequencer: Fix the actor reference key editor not resolving bound objects from the correct sequence ID

This resulted in not finding the skeletal mesh components of the bound object in order to pop up the component and socket choosers.

diff --git a/Engine/Source/Editor/MovieSceneTools/Private/Channels/BuiltInChannelEditors.cpp b/Engine/Source/Editor/MovieSceneTools/Private/Channels/BuiltInChannelEditors.cpp
--- a/Engine/Source/Editor/MovieSceneTools/Private/Channels/BuiltInChannelEditors.cpp
+++ b/Engine/Source/Editor/MovieSceneTools/Private/Channels/BuiltInChannelEditors.cpp
@@ -709,7 +709,7 @@ TSharedRef<SWidget> CreateKeyEditor(const TMovieSceneChannelHandle<FMovieSceneAc
 		ISequencer* Sequencer = KeyEditor.GetSequencer();
 		TArray<USceneComponent*> ComponentsWithSockets;
 		AActor* Actor = nullptr;
-		for (TWeakObjectPtr<> WeakObject : ActorKey.Object.ResolveBoundObjects(MovieSceneSequenceID::Root, *Sequencer))
+		for (TWeakObjectPtr<> WeakObject : ActorKey.Object.ResolveBoundObjects(Sequencer->GetFocusedTemplateID(), *Sequencer))
 		{
 			Actor = Cast<AActor>(WeakObject.Get());
 			if (Actor)
