Use the source frame rate as stored on the timecode attributes.

diff --git a/Engine/Source/Editor/MovieSceneTools/Private/TrackEditors/SkeletalAnimationTrackEditor.cpp b/Engine/Source/Editor/MovieSceneTools/Private/TrackEditors/SkeletalAnimationTrackEditor.cpp
--- a/Engine/Source/Editor/MovieSceneTools/Private/TrackEditors/SkeletalAnimationTrackEditor.cpp
+++ b/Engine/Source/Editor/MovieSceneTools/Private/TrackEditors/SkeletalAnimationTrackEditor.cpp
@@ -727,7 +727,7 @@ TOptional<FFrameTime> FSkeletalAnimationSection::GetSectionTime(FSequencerSectio
 	}
 
 	// Convert to tick resolution
-	HintFrameTime = FQualifiedFrameTime(ConvertFrameTime(HintFrameTime.Time, SamplingFrameRate, TickResolution), TickResolution);
+	HintFrameTime = FQualifiedFrameTime(ConvertFrameTime(HintFrameTime.Time, HintFrameTime.Rate, TickResolution), TickResolution);
 
 	// Get the desired frame display format and zero padding from
 	// the sequencer settings, if possible.
