Use correct origin space when editing subsequence transform origin overrides.

-Now supports world space edits.
-Rotations are now applied with FTransforms to correct for issues from adding rotators.
-Fixed issue where instance data overrides were not accounted for when editing origin overrides.

diff --git a/Engine/Source/Editor/MovieSceneTools/Private/EditModes/SubTrackEditorMode.cpp b/Engine/Source/Editor/MovieSceneTools/Private/EditModes/SubTrackEditorMode.cpp
--- a/Engine/Source/Editor/MovieSceneTools/Private/EditModes/SubTrackEditorMode.cpp
+++ b/Engine/Source/Editor/MovieSceneTools/Private/EditModes/SubTrackEditorMode.cpp
@@ -7,6 +7,7 @@
 #include "EngineUtils.h"
 #include "UnrealClient.h"
 #include "EditorViewportClient.h"
+#include "IMovieScenePlaybackClient.h"
 #include "MovieScene.h"
 #include "SequencerChannelTraits.h"
 #include "Compilation/MovieSceneCompiledDataManager.h"
@@ -15,6 +16,7 @@
 #include "MovieSceneSection.h"
 #include "Sections/MovieSceneSubSection.h"
 #include "Systems/MovieSceneTransformOriginSystem.h"
+#include "Tracks/IMovieSceneTransformOrigin.h"
 #include "Tracks/MovieSceneSubTrack.h"
 
 FName FSubTrackEditorMode::ModeName("EditMode.SubTrackEditMode");
@@ -22,6 +24,7 @@ FName FSubTrackEditorMode::ModeName("EditMode.SubTrackEditMode");
 FSubTrackEditorMode::FSubTrackEditorMode()
 {
 	CachedLocation.Reset();
+	PreviewCoordinateSpace.Reset();
 }
 
 FSubTrackEditorMode::~FSubTrackEditorMode()
@@ -31,11 +34,12 @@ FSubTrackEditorMode::~FSubTrackEditorMode()
 void FSubTrackEditorMode::Initialize()
 {
 	CachedLocation.Reset();
+	PreviewCoordinateSpace.Reset();
 }
 
 bool FSubTrackEditorMode::InputDelta(FEditorViewportClient* InViewportClient, FViewport* InViewport, FVector& InDrag, FRotator& InRot, FVector& InScale)
 {
-	if(AreAnyActorsSelected())
+	if(AreAnyActorsSelected() || (InDrag.IsNearlyZero() && InRot.IsNearlyZero()))
 	{
 		return false;
 	}
@@ -49,16 +53,41 @@ bool FSubTrackEditorMode::InputDelta(FEditorViewportClient* InViewportClient, FV
 
 	if(bMouseButtonDown && !bAnyModifiers && CurrentAxis != EAxisList::None)
 	{
-		// Remove parent rotation from input.
-		const FVector UnrotatedDrag = GetTransformOriginForSequence(GetFocusedSequenceID()).GetRotation().Inverse().RotateVector(InDrag);
+		const FTransform TransformOriginFocusedSequence = GetTransformOriginForSequence(GetFocusedSequenceID());
+
+		// Remove parent transform from inputs.
+		const FTransform LocalRotation = TransformOriginFocusedSequence * FTransform(InRot) * TransformOriginFocusedSequence.Inverse();
+		const FTransform LocalPosition = TransformOriginFocusedSequence * FTransform(InDrag) * TransformOriginFocusedSequence.Inverse();
+
+		// Keep preview space up-to-date.
+		PreviewCoordinateSpace.GetValue() *= FTransform(InRot).ToMatrixNoScale();
 		
-		OnOriginValueChanged.Broadcast( UnrotatedDrag, InRot);
+		OnOriginValueChanged.Broadcast(LocalPosition.GetLocation(), LocalRotation.Rotator());
 		return true;
 	}
 	
 	return false;
 }
 
+bool FSubTrackEditorMode::StartTracking(FEditorViewportClient* InViewportClient, FViewport* InViewport)
+{
+	const UMovieSceneSubSection* SubSection = GetSelectedSection();
+
+	if (!AreAnyActorsSelected() && SubSection && DoesSubSectionHaveTransformOverrides(*SubSection))
+	{
+		PreviewCoordinateSpace = GetFinalTransformOriginForSubSection(SubSection).ToMatrixNoScale().RemoveTranslation();
+		return true;
+	}
+	
+	return false;
+}
+
+bool FSubTrackEditorMode::EndTracking(FEditorViewportClient* InViewportClient, FViewport* InViewport)
+{
+	PreviewCoordinateSpace.Reset();
+	return true;
+}
+
 bool FSubTrackEditorMode::UsesTransformWidget() const
 {
 	const UMovieSceneSubSection* SubSection = GetSelectedSection();
@@ -130,7 +159,7 @@ FTransform FSubTrackEditorMode::GetFinalTransformOriginForSubSection(const UMovi
 
 FTransform FSubTrackEditorMode::GetTransformOriginForSequence(const TOptional<FMovieSceneSequenceID> InSequenceID) const
 {
-	FTransform TransformOrigin = FTransform();
+	FTransform TransformOrigin = FTransform::Identity;
 
 	const TSharedPtr<ISequencer> Sequencer = WeakSequencer.Pin();
 
@@ -140,6 +169,17 @@ FTransform FSubTrackEditorMode::GetTransformOriginForSequence(const TOptional<FM
 	}
 
 	const FMovieSceneRootEvaluationTemplateInstance& EvaluationTemplate = Sequencer->GetEvaluationTemplate();
+
+	const IMovieScenePlaybackClient*  Client       = Sequencer->GetPlaybackClient();
+	const UObject*                    InstanceData = Client ? Client->GetInstanceData() : nullptr;
+	const IMovieSceneTransformOrigin* RawInterface = Cast<const IMovieSceneTransformOrigin>(InstanceData);
+
+	const bool bHasInterface = RawInterface || (InstanceData && InstanceData->GetClass()->ImplementsInterface(UMovieSceneTransformOrigin::StaticClass()));
+	if (bHasInterface)
+	{
+		// Retrieve the current origin
+		TransformOrigin = RawInterface ? RawInterface->GetTransformOrigin() : IMovieSceneTransformOrigin::Execute_BP_GetTransformOrigin(InstanceData);
+	}
 	
 	const UMovieSceneEntitySystemLinker* EntityLinker = EvaluationTemplate.GetEntitySystemLinker();
 	if(!EntityLinker || !InSequenceID)
@@ -221,8 +261,28 @@ bool FSubTrackEditorMode::GetPivotForOrbit(FVector& OutPivot) const
 
 bool FSubTrackEditorMode::GetCustomDrawingCoordinateSystem(FMatrix& OutMatrix, void* InData)
 {
-	// @todo support custom coordinate systems based on parent rotation.
-	return false;
+	if (GetModeManager()->GetCoordSystem() != COORD_Local)
+	{
+		return false;
+	}
+
+	// While manipulating the gizmo, the preview coordinate space is kept up-to-date directly. The transform origin data however is set by a callback,
+	// and can be out of date, which would cause the gizmo to jitter.
+	if (PreviewCoordinateSpace.IsSet())
+	{
+		OutMatrix = PreviewCoordinateSpace.GetValue();
+		return true;
+	}
+
+	// Before interacting with the gizmo, query the origin data directly as a starting point.
+	const UMovieSceneSubSection* SubSection = GetSelectedSection();
+	if (!AreAnyActorsSelected() && SubSection && DoesSubSectionHaveTransformOverrides(*SubSection))
+	{
+		OutMatrix = GetFinalTransformOriginForSubSection(SubSection).ToMatrixNoScale().RemoveTranslation();
+		return true;
+	}
+	
+	return FEdMode::GetCustomDrawingCoordinateSystem(OutMatrix, InData);
 }
 
 bool FSubTrackEditorMode::GetCustomInputCoordinateSystem(FMatrix& OutMatrix, void* InData)
diff --git a/Engine/Source/Editor/MovieSceneTools/Private/EditModes/SubTrackEditorMode.h b/Engine/Source/Editor/MovieSceneTools/Private/EditModes/SubTrackEditorMode.h
--- a/Engine/Source/Editor/MovieSceneTools/Private/EditModes/SubTrackEditorMode.h
+++ b/Engine/Source/Editor/MovieSceneTools/Private/EditModes/SubTrackEditorMode.h
@@ -18,6 +18,8 @@ public:
 	// FEdMode interface
 	virtual void Initialize() override;
 	virtual bool InputDelta(FEditorViewportClient* InViewportClient, FViewport* InViewport, FVector& InDrag, FRotator& InRot, FVector& InScale) override;
+	virtual bool StartTracking(FEditorViewportClient* InViewportClient, FViewport* InViewport) override;
+	virtual bool EndTracking(FEditorViewportClient* InViewportClient, FViewport* InViewport) override;
 	virtual bool UsesTransformWidget() const override;
 	virtual bool UsesTransformWidget(UE::Widget::EWidgetMode CheckMode) const override;
 	virtual FVector GetWidgetLocation() const override;
@@ -54,8 +56,12 @@ private:
 	/** Delegate called when the origin is modified from the editor gizmo */
 	FOnOriginValueChanged OnOriginValueChanged;
 
-	/** Used to tell if the gizmo as moved, and if the editor hit proxies need to be invalidated as a result */
+	/** Used to tell if the gizmo has moved, and if the editor hit proxies need to be invalidated as a result */
 	mutable TOptional<FVector> CachedLocation;
 
+	/** Caches the transform space to use for editing/input.
+	 *  Updates to the channel data aren't reflected in time for the UI, so keeping the preview space up-to-date prevents the gizmo from flickering. */
+	TOptional<FMatrix> PreviewCoordinateSpace;
+
 	const TArray<FName> IncompatibleEditorModes = TArray<FName>({ "EditMode.ControlRig", "EM_Landscape" });
 };
diff --git a/Engine/Source/Editor/MovieSceneTools/Private/TrackEditors/SubTrackEditor.cpp b/Engine/Source/Editor/MovieSceneTools/Private/TrackEditors/SubTrackEditor.cpp
--- a/Engine/Source/Editor/MovieSceneTools/Private/TrackEditors/SubTrackEditor.cpp
+++ b/Engine/Source/Editor/MovieSceneTools/Private/TrackEditors/SubTrackEditor.cpp
@@ -1202,7 +1202,8 @@ void FSubTrackEditor::UpdateOrigin(FVector InPosition, FRotator InRotation)
 		bool bShouldMarkDirty = true;
 		
 		FGeneratedTrackKeys GeneratedKeys;
-		GetOriginKeys(PreviousTransform.GetLocation() + InPosition, PreviousTransform.Rotator() + InRotation, SubSection, GeneratedKeys);
+		
+		GetOriginKeys(PreviousTransform.GetLocation() + InPosition, (PreviousTransform * FTransform(InRotation)).Rotator(), SubSection, GeneratedKeys);
 		
 		FKeyPropertyResult KeyResults = AddKeysToSection(SubSection, GetTimeForKey(), GeneratedKeys, ESequencerKeyMode::AutoKey);
 
@@ -1212,7 +1213,7 @@ void FSubTrackEditor::UpdateOrigin(FVector InPosition, FRotator InRotation)
 			// Preview key data is transient, and should not dirty the sequence.
 			bShouldMarkDirty = false;
 			SubSection->SetKeyPreviewPosition(PreviousTransform.GetLocation() + InPosition);
-			SubSection->SetKeyPreviewRotation(PreviousTransform.Rotator() + InRotation);
+			SubSection->SetKeyPreviewRotation((PreviousTransform * FTransform(InRotation)).Rotator());
 
 			// Preview data needs to be reverted when playing the sequence, scrubbing the sequence, or navigating to a different sequence
 			// This array keeps track of sections with preview data and reverts them in FSubTrackEditor::RevertSectionPreviews
