Option to remove transform origins from FBX imports into sequences.

Added option to UMovieSceneUserImportFBXSettings to correct for the movie scene's transform origins
Importing FBX can now interrogate transform origins from the actor's instance data and focused subsequence's overrides based on that setting

diff --git a/Engine/Plugins/MovieScene/SequencerScripting/Source/SequencerScriptingEditor/Private/SequencerTools.cpp b/Engine/Plugins/MovieScene/SequencerScripting/Source/SequencerScriptingEditor/Private/SequencerTools.cpp
--- a/Engine/Plugins/MovieScene/SequencerScripting/Source/SequencerScriptingEditor/Private/SequencerTools.cpp
+++ b/Engine/Plugins/MovieScene/SequencerScripting/Source/SequencerScriptingEditor/Private/SequencerTools.cpp
@@ -728,15 +728,18 @@ bool ImportFBXInternal(UWorld* World, UMovieSceneSequence* Sequence, const TArra
 	return bResult;
 }
 
-bool USequencerToolsFunctionLibrary::ImportLevelSequenceFBX(UWorld* World, ULevelSequence* Sequence, const TArray<FMovieSceneBindingProxy>& InBindings, UMovieSceneUserImportFBXSettings* ImportFBXSettings, const FString&  ImportFilename)
+bool USequencerToolsFunctionLibrary::ImportLevelSequenceFBX(UWorld* World, ULevelSequence* Sequence, const TArray<FMovieSceneBindingProxy>& InBindings, UMovieSceneUserImportFBXSettings* ImportFBXSettings, const FString&  ImportFilename, ALevelSequenceActor* ActorContext)
 {
-	ALevelSequenceActor* OutActor;
+	ALevelSequenceActor* OutActor = nullptr;
 	FMovieSceneSequencePlaybackSettings Settings;
 	FLevelSequenceCameraSettings CameraSettings;
-	ULevelSequencePlayer* Player = ULevelSequencePlayer::CreateLevelSequencePlayer(World, Sequence, Settings, OutActor);
+	ULevelSequencePlayer* Player = ActorContext ? ActorContext->GetSequencePlayer() : ULevelSequencePlayer::CreateLevelSequencePlayer(World, Sequence, Settings, OutActor);
 
 	bool bSuccess = ImportFBXInternal(World, Sequence, InBindings, ImportFBXSettings, ImportFilename, Player);
-	World->DestroyActor(OutActor);
+	if(OutActor)
+	{
+		World->DestroyActor(OutActor);
+	}
 
 	return bSuccess;
 }
diff --git a/Engine/Plugins/MovieScene/SequencerScripting/Source/SequencerScriptingEditor/Public/SequencerTools.h b/Engine/Plugins/MovieScene/SequencerScripting/Source/SequencerScriptingEditor/Public/SequencerTools.h
--- a/Engine/Plugins/MovieScene/SequencerScripting/Source/SequencerScriptingEditor/Public/SequencerTools.h
+++ b/Engine/Plugins/MovieScene/SequencerScripting/Source/SequencerScriptingEditor/Public/SequencerTools.h
@@ -234,7 +234,7 @@ public:
 	 * @InPlayer Player to bind to
 	 */
 	UFUNCTION(BlueprintCallable, Category = "Editor Scripting | Sequencer Tools | FBX")
-	static bool ImportLevelSequenceFBX(UWorld* InWorld, ULevelSequence* InSequence, const TArray<FMovieSceneBindingProxy>& InBindings, UMovieSceneUserImportFBXSettings* InImportFBXSettings, const FString& InImportFilename);
+	static bool ImportLevelSequenceFBX(UWorld* InWorld, ULevelSequence* InSequence, const TArray<FMovieSceneBindingProxy>& InBindings, UMovieSceneUserImportFBXSettings* InImportFBXSettings, const FString& InImportFilename, ALevelSequenceActor* ActorContext = nullptr);
 
 public:
 	/**
diff --git a/Engine/Source/Editor/MovieSceneTools/Private/MovieSceneToolHelpers.cpp b/Engine/Source/Editor/MovieSceneTools/Private/MovieSceneToolHelpers.cpp
--- a/Engine/Source/Editor/MovieSceneTools/Private/MovieSceneToolHelpers.cpp
+++ b/Engine/Source/Editor/MovieSceneTools/Private/MovieSceneToolHelpers.cpp
@@ -88,6 +88,8 @@
 #include "Framework/MultiBox/MultiBoxBuilder.h"
 #include "EntitySystem/Interrogation/MovieSceneInterrogationLinker.h"
 #include "ConstraintsManager.h"
+#include "IMovieScenePlaybackClient.h"
+#include "SequencerChannelInterface.h"
 #include "Constraints/MovieSceneConstraintChannelHelper.inl"
 #include "Exporters/FbxExportOption.h"
 #include "Bindings/MovieSceneSpawnableBinding.h"
@@ -2743,7 +2745,255 @@ bool MovieSceneToolHelpers::ExportFBXFromControlRigChannelsWithDialog(const TSha
 	return true;
 }
 
-bool ImportFBXTransform(FString NodeName, FGuid ObjectBinding, UnFbx::FFbxCurvesAPI& CurveAPI, UMovieSceneSequence* InSequence)
+double UnwindKeyValue(const double OldValue, double NewValue)
+{
+	while (NewValue - OldValue > 180.0f)
+	{
+		NewValue -= 360.0f;
+	}
+	while (NewValue - OldValue < -180.0f)
+	{
+		NewValue += 360.0f;
+	}
+	return NewValue;
+}
+
+void AddCorrectedKey(const FMovieSceneDoubleValue& CorrectedValue, const FFrameTime Time, TMovieSceneChannelData<FMovieSceneDoubleValue> ChannelData)
+{
+	const int32 KeyIndex = ChannelData.FindKey(Time.RoundToFrame());
+
+	if (KeyIndex != INDEX_NONE)
+	{
+		FMovieSceneDoubleValue& KeyValue = ChannelData.GetValues()[KeyIndex];
+
+		KeyValue.Value = CorrectedValue.Value;
+		// Tangents need to be set to auto, so they can be corrected later, as their imported values may no longer be valid after correcting the transform
+		if (KeyValue.TangentMode != RCTM_Auto || KeyValue.TangentMode != RCTM_SmartAuto)
+		{
+			KeyValue.TangentMode = RCTM_Auto;
+		}
+	}
+	else
+	{
+		MovieSceneToolHelpers::SetOrAddKey(ChannelData, Time.RoundToFrame(), CorrectedValue);
+	}
+}
+
+void CorrectForTransformOrigin(IMovieScenePlayer* Player, UMovieScene3DTransformSection* TransformSection, ISequencer* Sequencer)
+{
+	FMovieSceneChannelProxy& SectionChannelProxy = TransformSection->GetChannelProxy();
+	TMovieSceneChannelHandle<FMovieSceneDoubleChannel> DoubleChannels[] = {
+		SectionChannelProxy.GetChannelByName<FMovieSceneDoubleChannel>("Location.X"),
+		SectionChannelProxy.GetChannelByName<FMovieSceneDoubleChannel>("Location.Y"),
+		SectionChannelProxy.GetChannelByName<FMovieSceneDoubleChannel>("Location.Z"),
+		SectionChannelProxy.GetChannelByName<FMovieSceneDoubleChannel>("Rotation.X"),
+		SectionChannelProxy.GetChannelByName<FMovieSceneDoubleChannel>("Rotation.Y"),
+		SectionChannelProxy.GetChannelByName<FMovieSceneDoubleChannel>("Rotation.Z")
+	};
+
+	FFrameRate FrameRate = TransformSection->GetTypedOuter<UMovieScene>()->GetTickResolution();
+
+	// If importing from the Sequencer editor, retrieve the focused subsequence hierarchy, in order to get the subsequence transform origin overrides.
+	// If not importing from the Sequencer editor, assume the transform is being imported onto the root sequence.
+	TArray<FMovieSceneSequenceID> SubsequenceHierarchy = Sequencer ? Sequencer->GetSubSequenceHierarchy() : TArray { MovieSceneSequenceID::Root };
+
+	UE::MovieScene::FSystemInterrogator Interrogator;
+	
+	const FMovieSceneRootEvaluationTemplateInstance& RootInstance = Player->GetEvaluationTemplate();
+	const FMovieSceneSequenceHierarchy* Hierarchy = RootInstance.GetCompiledDataManager()->FindHierarchy(RootInstance.GetCompiledDataID());
+	
+	// Sub tracks can only be interrogated if the Sequencer and Hierarchy are valid (i.e. the import is triggered from the editor),
+	// but the instance data can still be interrogated if they aren't, so this step can be safely skipped when importing from scripting.
+	if (Sequencer && Hierarchy)
+	{
+		Interrogator.SetHierarchy(const_cast<FMovieSceneSequenceHierarchy*>(Hierarchy));
+
+		FMovieSceneSequenceID OwningSequenceID = MovieSceneSequenceID::Root;
+		for (const FMovieSceneSequenceID SubSequenceID : SubsequenceHierarchy)
+		{
+			if (const UMovieSceneSubSection* SubSection = Sequencer->FindSubSection(SubSequenceID))
+			{
+				if (UMovieSceneTrack* SubTrack = Cast<UMovieSceneTrack>(SubSection->GetOuter()))
+				{
+					Interrogator.ImportTrack(SubTrack, UE::MovieScene::FInterrogationChannel::FromIndex(SubSection->GetRowIndex()), OwningSequenceID);
+				}
+			}
+			OwningSequenceID = SubSequenceID;
+		}
+	}
+	
+	TSet<FFrameNumber> Times;
+
+	for (int32 ChannelIndex = 0; ChannelIndex < 6; ++ChannelIndex)
+	{
+		TMovieSceneChannelHandle<FMovieSceneDoubleChannel> DoubleChannel = DoubleChannels[ChannelIndex];
+		
+		for (FFrameNumber Time : DoubleChannel.Get()->GetTimes())
+		{
+			Times.Add(Time);
+		}
+	}
+
+	TArray<FFrameNumber> TimesElements = Times.Array();
+	TimesElements.Sort();
+
+	for (FFrameNumber Time : TimesElements)
+	{
+		Interrogator.AddInterrogation(Time);
+	}
+	
+	const IMovieScenePlaybackClient*  Client = Player->GetPlaybackClient();
+	const UObject* InstanceData = Client ? Client->GetInstanceData() : nullptr;
+	
+	TArray<FTransform> TransformOrigins;
+	TArray<FTransform> CorrectedTransforms;
+
+	Interrogator.Update();
+	Interrogator.QueryTransformOrigins(TransformOrigins, SubsequenceHierarchy, InstanceData);
+	
+	// Gather the corrected transforms
+	for (int32 TransformOriginIndex = 0; TransformOriginIndex < TransformOrigins.Num(); ++TransformOriginIndex)
+	{
+		FFrameTime Time = TimesElements[TransformOriginIndex];
+		FTransform TransformOrigin = TransformOrigins[TransformOriginIndex];
+		
+		double LocationX = 0.0;
+		double LocationY = 0.0;
+		double LocationZ = 0.0;
+
+		double RotationX = 0.0;
+		double RotationY = 0.0;
+		double RotationZ = 0.0;
+
+		if (DoubleChannels[0].Get())
+		{
+			DoubleChannels[0].Get()->Evaluate(Time, LocationX);
+		}
+
+		if (DoubleChannels[1].Get())
+		{
+			DoubleChannels[1].Get()->Evaluate(Time, LocationY);
+		}
+
+		if (DoubleChannels[2].Get())
+		{
+			DoubleChannels[2].Get()->Evaluate(Time, LocationZ);
+		}
+
+		if (DoubleChannels[3].Get())
+		{
+			DoubleChannels[3].Get()->Evaluate(Time, RotationX);
+		}
+
+		if (DoubleChannels[4].Get())
+		{
+			DoubleChannels[4].Get()->Evaluate(Time, RotationY);
+		}
+
+		if (DoubleChannels[5].Get())
+		{
+			DoubleChannels[5].Get()->Evaluate(Time, RotationZ);
+		}
+
+		FVector EvaluatedLocation = FVector(LocationX, LocationY, LocationZ);
+		FRotator EvaluatedRotation = FRotator(RotationY, RotationZ, RotationX);
+
+		FTransform EvaluatedTransform = FTransform(EvaluatedRotation, EvaluatedLocation);
+
+		EvaluatedTransform *= TransformOrigin.Inverse();
+
+		CorrectedTransforms.Add(EvaluatedTransform);
+	}
+
+	TOptional<double> PreviousRotationX;
+	TOptional<double> PreviousRotationY;
+	TOptional<double> PreviousRotationZ;
+
+	// Apply the correct transforms. This has to be done after gathering, so the corrected transforms don't interfere with evaluation.
+	for (int32 CorrectedTransformIndex = 0; CorrectedTransformIndex < CorrectedTransforms.Num(); ++CorrectedTransformIndex)
+	{
+		FTransform CorrectedTransform = CorrectedTransforms[CorrectedTransformIndex];
+		FFrameTime Time = TimesElements[CorrectedTransformIndex];
+		
+		if (DoubleChannels[0].Get())
+		{
+			TMovieSceneChannelData<FMovieSceneDoubleValue> ChannelData = DoubleChannels[0].Get()->GetData();
+			AddCorrectedKey(FMovieSceneDoubleValue(CorrectedTransform.GetTranslation().X), Time, ChannelData);
+		}
+
+		if (DoubleChannels[1].Get())
+		{
+			TMovieSceneChannelData<FMovieSceneDoubleValue> ChannelData = DoubleChannels[1].Get()->GetData();
+			AddCorrectedKey(FMovieSceneDoubleValue(CorrectedTransform.GetTranslation().Y), Time, ChannelData);
+		}
+
+		if (DoubleChannels[2].Get())
+		{
+			TMovieSceneChannelData<FMovieSceneDoubleValue> ChannelData = DoubleChannels[2].Get()->GetData();
+			AddCorrectedKey(FMovieSceneDoubleValue(CorrectedTransform.GetTranslation().Z), Time, ChannelData);
+		}
+
+		if (DoubleChannels[3].Get())
+		{
+			TMovieSceneChannelData<FMovieSceneDoubleValue> ChannelData = DoubleChannels[3].Get()->GetData();
+			FMovieSceneDoubleValue CorrectedValue = FMovieSceneDoubleValue(CorrectedTransform.GetRotation().Rotator().Roll);
+
+			if (PreviousRotationX.IsSet())
+			{
+				CorrectedValue.Value = UnwindKeyValue(PreviousRotationX.GetValue(), CorrectedValue.Value);
+			}
+
+			PreviousRotationX = CorrectedValue.Value;
+
+			AddCorrectedKey(CorrectedValue, Time, ChannelData);
+		}
+
+		if (DoubleChannels[4].Get())
+		{
+			TMovieSceneChannelData<FMovieSceneDoubleValue> ChannelData = DoubleChannels[4].Get()->GetData();
+			FMovieSceneDoubleValue CorrectedValue = FMovieSceneDoubleValue(CorrectedTransform.GetRotation().Rotator().Pitch);
+
+			if (PreviousRotationY.IsSet())
+			{
+				CorrectedValue.Value = UnwindKeyValue(PreviousRotationY.GetValue(), CorrectedValue.Value);
+			}
+
+			PreviousRotationY = CorrectedValue.Value;
+
+			AddCorrectedKey(CorrectedValue, Time, ChannelData);
+		}
+
+		if (DoubleChannels[5].Get())
+		{
+			TMovieSceneChannelData<FMovieSceneDoubleValue> ChannelData = DoubleChannels[5].Get()->GetData();
+			FMovieSceneDoubleValue CorrectedValue = FMovieSceneDoubleValue(CorrectedTransform.GetRotation().Rotator().Yaw);
+
+			if (PreviousRotationZ.IsSet())
+			{
+				CorrectedValue.Value = UnwindKeyValue(PreviousRotationZ.GetValue(), CorrectedValue.Value);
+			}
+
+			PreviousRotationZ = CorrectedValue.Value;
+
+			AddCorrectedKey(CorrectedValue, Time, ChannelData);
+		}
+	}
+
+	for (TMovieSceneChannelHandle<FMovieSceneDoubleChannel> DoubleChannel : DoubleChannels)
+	{
+		const UMovieSceneUserImportFBXSettings* ImportFBXSettings = GetDefault<UMovieSceneUserImportFBXSettings>();
+		if (ImportFBXSettings->bReduceKeys)
+		{
+			FKeyDataOptimizationParams Params;
+			Params.Tolerance = ImportFBXSettings->ReduceKeysTolerance;
+			Params.DisplayRate = FrameRate;
+			DoubleChannel.Get()->Optimize(Params);
+		}
+		DoubleChannel.Get()->AutoSetTangents();
+	}
+}
+
+bool ImportFBXTransform(FString NodeName, FGuid ObjectBinding, UnFbx::FFbxCurvesAPI& CurveAPI, UMovieSceneSequence* InSequence, IMovieScenePlayer* Player, ISequencer* Sequencer)
 {
 	UMovieScene* MovieScene = InSequence->GetMovieScene();
 
@@ -2857,10 +3107,15 @@ bool ImportFBXTransform(FString NodeName, FGuid ObjectBinding, UnFbx::FFbxCurves
 		ImportTransformChannelToDouble(Scale[2], DoubleChannels[8].Get(), FrameRate, false, true);
 	}
 
+	if(ImportFBXSettings->bCorrectForTransformOrigin)
+	{
+		CorrectForTransformOrigin(Player, TransformSection, Sequencer);
+	}
+
 	return true;
 }
 
-bool MovieSceneToolHelpers::ImportFBXNode(FString NodeName, UnFbx::FFbxCurvesAPI& CurveAPI, UMovieSceneSequence* InSequence, IMovieScenePlayer* Player, FMovieSceneSequenceIDRef TemplateID, FGuid ObjectBinding)
+bool MovieSceneToolHelpers::ImportFBXNode(FString NodeName, UnFbx::FFbxCurvesAPI& CurveAPI, UMovieSceneSequence* InSequence, IMovieScenePlayer* Player, FMovieSceneSequenceIDRef TemplateID, FGuid ObjectBinding, ISequencer* Sequencer)
 {
 	// Look for animated float properties
 	TArray<FString> AnimatedPropertyNames;
@@ -2871,7 +3126,7 @@ bool MovieSceneToolHelpers::ImportFBXNode(FString NodeName, UnFbx::FFbxCurvesAPI
 		ImportFBXProperty(NodeName, AnimatedPropertyName, ObjectBinding, CurveAPI, InSequence, Player, TemplateID);
 	}
 	
-	ImportFBXTransform(NodeName, ObjectBinding, CurveAPI, InSequence);
+	ImportFBXTransform(NodeName, ObjectBinding, CurveAPI, InSequence, Player, Sequencer);
 
 	// Custom static string properties
 	TArray<TPair<FString, FString> > CustomPropertyPairs;
@@ -3429,7 +3684,7 @@ private:
 		ImportFBXCamera(FbxImporter, Sequence, *Sequencer, ObjectBindingMap, bMatchByNameOnly, bCreateCameras.IsSet() ? bCreateCameras.GetValue() : ImportFBXSettings->bCreateCameras);
 
 		UWorld* World = Sequencer->GetPlaybackContext()->GetWorld();
-		bool bValid = MovieSceneToolHelpers::ImportFBXIfReady(World, Sequence, Sequencer, Sequencer->GetFocusedTemplateID(), ObjectBindingMap, ImportFBXSettings, InOutParams);
+		bool bValid = MovieSceneToolHelpers::ImportFBXIfReady(World, Sequence, Sequencer, Sequencer->GetFocusedTemplateID(), ObjectBindingMap, ImportFBXSettings, InOutParams, Sequencer);
 	
 		Sequencer->NotifyMovieSceneDataChanged(EMovieSceneDataChangeType::MovieSceneStructureItemAdded);
 
@@ -3505,7 +3760,7 @@ bool ImportFBXOntoControlRigs(UWorld* World, UMovieScene* MovieScene, IMovieScen
 }
 
 bool MovieSceneToolHelpers::ImportFBXIfReady(UWorld* World, UMovieSceneSequence* Sequence, IMovieScenePlayer* Player, FMovieSceneSequenceIDRef TemplateID, TMap<FGuid, FString>& ObjectBindingMap, UMovieSceneUserImportFBXSettings* ImportFBXSettings,
-	const FFBXInOutParameters& InParams)
+	const FFBXInOutParameters& InParams, ISequencer* Sequencer)
 {
 	UMovieScene* MovieScene = Sequence->GetMovieScene();
 
@@ -3521,6 +3776,7 @@ bool MovieSceneToolHelpers::ImportFBXIfReady(UWorld* World, UMovieSceneSequence*
 	CurrentImportFBXSettings->bConvertSceneUnit = ImportFBXSettings->bConvertSceneUnit;
 	CurrentImportFBXSettings->ImportUniformScale = ImportFBXSettings->ImportUniformScale;
 	CurrentImportFBXSettings->bReplaceTransformTrack = ImportFBXSettings->bReplaceTransformTrack;
+	CurrentImportFBXSettings->bCorrectForTransformOrigin = ImportFBXSettings->bCorrectForTransformOrigin;
 	UnFbx::FFbxImporter* FbxImporter = UnFbx::FFbxImporter::GetInstance();
 
 	UnFbx::FFbxCurvesAPI CurveAPI;
@@ -3548,7 +3804,7 @@ bool MovieSceneToolHelpers::ImportFBXIfReady(UWorld* World, UMovieSceneSequence*
 		{
 			if (FCString::Strcmp(*It.Value().ToUpper(), *NodeName.ToUpper()) == 0)
 			{
-				MovieSceneToolHelpers::ImportFBXNode(NodeName, CurveAPI, Sequence, Player, TemplateID, It.Key());
+				MovieSceneToolHelpers::ImportFBXNode(NodeName, CurveAPI, Sequence, Player, TemplateID, It.Key(), Sequencer);
 
 				ObjectBindingMap.Remove(It.Key());
 				AllNodeNames.RemoveAt(NodeIndex);
@@ -3581,7 +3837,7 @@ bool MovieSceneToolHelpers::ImportFBXIfReady(UWorld* World, UMovieSceneSequence*
 			auto It = ObjectBindingMap.CreateConstIterator();
 			if (It)
 			{
-				MovieSceneToolHelpers::ImportFBXNode(NodeName, CurveAPI, Sequence, Player, TemplateID, It.Key());
+				MovieSceneToolHelpers::ImportFBXNode(NodeName, CurveAPI, Sequence, Player, TemplateID, It.Key(), Sequencer);
 
 				UE_LOG(LogMovieScene, Warning, TEXT("Fbx Import: Failed to find any matching node for (%s). Defaulting to first available (%s)."), *NodeName, *It.Value());
 				ObjectBindingMap.Remove(It.Key());
diff --git a/Engine/Source/Editor/MovieSceneTools/Private/MovieSceneToolsUserSettings.cpp b/Engine/Source/Editor/MovieSceneTools/Private/MovieSceneToolsUserSettings.cpp
--- a/Engine/Source/Editor/MovieSceneTools/Private/MovieSceneToolsUserSettings.cpp
+++ b/Engine/Source/Editor/MovieSceneTools/Private/MovieSceneToolsUserSettings.cpp
@@ -46,6 +46,7 @@ UMovieSceneUserImportFBXSettings::UMovieSceneUserImportFBXSettings(const FObject
 	bForceFrontXAxis = false;
 	bCreateCameras = true;
 	bReplaceTransformTrack = true;
+	bCorrectForTransformOrigin = false;
 	bReduceKeys = true;
 	ReduceKeysTolerance = 0.001f;
 	bConvertSceneUnit = true;
diff --git a/Engine/Source/Editor/MovieSceneTools/Public/MovieSceneToolHelpers.h b/Engine/Source/Editor/MovieSceneTools/Public/MovieSceneToolHelpers.h
--- a/Engine/Source/Editor/MovieSceneTools/Public/MovieSceneToolHelpers.h
+++ b/Engine/Source/Editor/MovieSceneTools/Public/MovieSceneToolHelpers.h
@@ -385,7 +385,7 @@ public:
 	*/
 
 	static bool ImportFBXIfReady(UWorld* World, UMovieSceneSequence* InSequence, IMovieScenePlayer* Player, FMovieSceneSequenceIDRef TemplateID, TMap<FGuid, FString>& ObjectBindingMap, UMovieSceneUserImportFBXSettings* ImportFBXSettings,
-		const FFBXInOutParameters& InFBXParams);
+		const FFBXInOutParameters& InFBXParams, ISequencer* Sequencer = nullptr);
 
 	/**
 	* Import FBX Camera to existing camera's
@@ -413,7 +413,7 @@ public:
 	* @param ObjectBinding Guid of the object we are importing onto.
 	* @return Whether the import was successful
 	*/
-	static bool ImportFBXNode(FString NodeName, UnFbx::FFbxCurvesAPI& CurveAPI, UMovieSceneSequence* InSequence, IMovieScenePlayer* Player, FMovieSceneSequenceIDRef TemplateID, FGuid ObjectBinding);
+	static bool ImportFBXNode(FString NodeName, UnFbx::FFbxCurvesAPI& CurveAPI, UMovieSceneSequence* InSequence, IMovieScenePlayer* Player, FMovieSceneSequenceIDRef TemplateID, FGuid ObjectBinding, ISequencer* Sequencer = nullptr);
 
 	/**
 	 * Lock the given camera actor to the viewport
diff --git a/Engine/Source/Editor/MovieSceneTools/Public/MovieSceneToolsUserSettings.h b/Engine/Source/Editor/MovieSceneTools/Public/MovieSceneToolsUserSettings.h
--- a/Engine/Source/Editor/MovieSceneTools/Public/MovieSceneToolsUserSettings.h
+++ b/Engine/Source/Editor/MovieSceneTools/Public/MovieSceneToolsUserSettings.h
@@ -91,6 +91,12 @@ public:
 	UPROPERTY(EditAnywhere, config, Category=Import)
 	bool bReplaceTransformTrack;
 
+	/** When enabled, the current transform origin of the focused sequence will be removed from any imported transforms.
+	 *  If the FBX was exported with a transform origin, this will prevent it from being applied twice.
+	 */
+	UPROPERTY(EditAnywhere, config, Category=Import)
+	bool bCorrectForTransformOrigin;
+	
 	/** Whether to remove keyframes within a tolerance from the imported tracks */
 	UPROPERTY(EditAnywhere, config, Category=Import)
 	bool bReduceKeys;
