Fix static analysis warning about expression always evaluating to true.

Fixed logic when checking if tangent mode on imported keys were auto or smart auto.

#rnx

diff --git a/Engine/Source/Editor/MovieSceneTools/Private/MovieSceneToolHelpers.cpp b/Engine/Source/Editor/MovieSceneTools/Private/MovieSceneToolHelpers.cpp
--- a/Engine/Source/Editor/MovieSceneTools/Private/MovieSceneToolHelpers.cpp
+++ b/Engine/Source/Editor/MovieSceneTools/Private/MovieSceneToolHelpers.cpp
@@ -2768,7 +2768,7 @@ void AddCorrectedKey(const FMovieSceneDoubleValue& CorrectedValue, const FFrameT
 
 		KeyValue.Value = CorrectedValue.Value;
 		// Tangents need to be set to auto, so they can be corrected later, as their imported values may no longer be valid after correcting the transform
-		if (KeyValue.TangentMode != RCTM_Auto || KeyValue.TangentMode != RCTM_SmartAuto)
+		if (KeyValue.TangentMode != RCTM_Auto && KeyValue.TangentMode != RCTM_SmartAuto)
 		{
 			KeyValue.TangentMode = RCTM_Auto;
 		}
