Constraints: avoid registering transactions when evaluating sequencer

in the context of spawnables, we don't need to register transaction, like selection change when destructing the actor for instance, otherwise it results in keeping useless references to that actor and wrong undo/redo stack

#rnx

diff --git a/Engine/Source/Editor/MovieSceneTools/Private/Constraints/MovieSceneConstraintChannelHelper.cpp b/Engine/Source/Editor/MovieSceneTools/Private/Constraints/MovieSceneConstraintChannelHelper.cpp
--- a/Engine/Source/Editor/MovieSceneTools/Private/Constraints/MovieSceneConstraintChannelHelper.cpp
+++ b/Engine/Source/Editor/MovieSceneTools/Private/Constraints/MovieSceneConstraintChannelHelper.cpp
@@ -81,6 +81,10 @@ void FCompensationEvaluator::ComputeLocalTransforms(UWorld* InWorld, const FEval
 	ChildLocals.SetNum(NumFrames + 1);
 	ChildGlobals.SetNum(NumFrames + 1);
 	SpaceGlobals.SetNum(NumFrames + 1);
+
+	// avoid transacting when evaluating sequencer
+	TGuardValue<ITransaction*> TransactionGuard(GUndo, nullptr);
+	
 	const TArray<IMovieSceneToolsAnimationBakeHelper*>& BakeHelpers = FMovieSceneToolsModule::Get().GetAnimationBakeHelpers();
 	for (IMovieSceneToolsAnimationBakeHelper* BakeHelper : BakeHelpers)
 	{
@@ -401,6 +405,10 @@ void FCompensationEvaluator::ComputeLocalTransformsBeforeDeletion(UWorld* InWorl
 	ChildLocals.SetNum(NumFrames);
 	ChildGlobals.SetNum(NumFrames);
 	SpaceGlobals.SetNum(NumFrames);
+
+	// avoid transacting when evaluating sequencer
+	TGuardValue<ITransaction*> TransactionGuard(GUndo, nullptr);
+	
 	const TArray<IMovieSceneToolsAnimationBakeHelper*>& BakeHelpers = FMovieSceneToolsModule::Get().GetAnimationBakeHelpers();
 	for (IMovieSceneToolsAnimationBakeHelper* BakeHelper : BakeHelpers)
 	{
@@ -505,6 +513,9 @@ void FCompensationEvaluator::ComputeCompensation(UWorld* InWorld, const TSharedP
 	const FConstraintsManagerController& Controller = FConstraintsManagerController::Get(InWorld);
 	static constexpr bool bSorted = true;
 	const TArray<ConstraintPtr> AllConstraints = Controller.GetAllConstraints(bSorted);
+
+	// avoid transacting when evaluating sequencer
+	TGuardValue<ITransaction*> TransactionGuard(GUndo, nullptr);
 	
 	UMovieScene* MovieScene = InSequencer->GetFocusedMovieSceneSequence()->GetMovieScene();
 	const TArray<IMovieSceneToolsAnimationBakeHelper*>& BakeHelpers = FMovieSceneToolsModule::Get().GetAnimationBakeHelpers();
@@ -679,6 +690,9 @@ void FCompensationEvaluator::CacheTransforms(UWorld* InWorld, const FEvalParamet
 		}
 	};
 
+	// avoid transacting when evaluating sequencer
+	TGuardValue<ITransaction*> TransactionGuard(GUndo, nullptr);
+	
 	for (IMovieSceneToolsAnimationBakeHelper* BakeHelper : BakeHelpers)
 	{
 		if (BakeHelper)
