Fix for PostEditMove no longer being called.

StartTracking and EndTracking now call their parent functions if the SubTrackEditorMode does not handle tracking.

#rnx

diff --git a/Engine/Source/Editor/MovieSceneTools/Private/EditModes/SubTrackEditorMode.cpp b/Engine/Source/Editor/MovieSceneTools/Private/EditModes/SubTrackEditorMode.cpp
--- a/Engine/Source/Editor/MovieSceneTools/Private/EditModes/SubTrackEditorMode.cpp
+++ b/Engine/Source/Editor/MovieSceneTools/Private/EditModes/SubTrackEditorMode.cpp
@@ -75,17 +75,24 @@ bool FSubTrackEditorMode::StartTracking(FEditorViewportClient* InViewportClient,
 
 	if (!AreAnyActorsSelected() && SubSection && DoesSubSectionHaveTransformOverrides(*SubSection))
 	{
+		bIsTracking = true;
 		PreviewCoordinateSpace = GetFinalTransformOriginForSubSection(SubSection).ToMatrixNoScale().RemoveTranslation();
 		return true;
 	}
 	
-	return false;
+	return FEdMode::StartTracking(InViewportClient, InViewport);
 }
 
 bool FSubTrackEditorMode::EndTracking(FEditorViewportClient* InViewportClient, FViewport* InViewport)
 {
-	PreviewCoordinateSpace.Reset();
-	return true;
+	if(bIsTracking)
+	{
+		bIsTracking = false;
+		PreviewCoordinateSpace.Reset();
+		return true;
+	}
+	
+	return FEdMode::EndTracking(InViewportClient, InViewport);
 }
 
 bool FSubTrackEditorMode::UsesTransformWidget() const
diff --git a/Engine/Source/Editor/MovieSceneTools/Private/EditModes/SubTrackEditorMode.h b/Engine/Source/Editor/MovieSceneTools/Private/EditModes/SubTrackEditorMode.h
--- a/Engine/Source/Editor/MovieSceneTools/Private/EditModes/SubTrackEditorMode.h
+++ b/Engine/Source/Editor/MovieSceneTools/Private/EditModes/SubTrackEditorMode.h
@@ -59,6 +59,9 @@ private:
 	/** Used to tell if the gizmo has moved, and if the editor hit proxies need to be invalidated as a result */
 	mutable TOptional<FVector> CachedLocation;
 
+	/** Used to ensure that we mirror the behavior in StartTracking if the selection were to change mid-drag. */
+	bool bIsTracking = false;
+
 	/** Caches the transform space to use for editing/input.
 	 *  Updates to the channel data aren't reflected in time for the UI, so keeping the preview space up-to-date prevents the gizmo from flickering. */
 	TOptional<FMatrix> PreviewCoordinateSpace;
