Sequencer: Fix crash when upgrading Sequence in level with World Partitions due to blueprint compilation during level sequence postload.
See https://epicgames.lightning.force.com/lightning/r/Case/500QP00000SLzErYAL/view for further details.

diff --git a/Engine/Source/Editor/MovieSceneTools/Private/MovieSceneToolsModule.cpp b/Engine/Source/Editor/MovieSceneTools/Private/MovieSceneToolsModule.cpp
--- a/Engine/Source/Editor/MovieSceneTools/Private/MovieSceneToolsModule.cpp
+++ b/Engine/Source/Editor/MovieSceneTools/Private/MovieSceneToolsModule.cpp
@@ -494,7 +494,7 @@ void FMovieSceneToolsModule::FixupDynamicBindingsEvent(ULevelSequence* LevelSequ
 		if (UBlueprint* DirectorBlueprint = LevelSequence->GetDirectorBlueprint())
 		{
 			FMovieSceneDynamicBindingUtils::EnsureBlueprintExtensionCreated(LevelSequence, DirectorBlueprint);
-			FKismetEditorUtilities::CompileBlueprint(DirectorBlueprint);
+			FKismetEditorUtilities::CompileBlueprint(DirectorBlueprint, EBlueprintCompileOptions::SkipGarbageCollection);
 		}
 	}
 }
