Sequencer: Fix crash on reloading sequence assets under some circumstances.

diff --git a/Engine/Source/Editor/MovieSceneTools/Private/Conditions/MovieSceneConditionCustomization.cpp b/Engine/Source/Editor/MovieSceneTools/Private/Conditions/MovieSceneConditionCustomization.cpp
--- a/Engine/Source/Editor/MovieSceneTools/Private/Conditions/MovieSceneConditionCustomization.cpp
+++ b/Engine/Source/Editor/MovieSceneTools/Private/Conditions/MovieSceneConditionCustomization.cpp
@@ -72,7 +72,7 @@ TSharedRef<IPropertyTypeCustomization> FMovieSceneConditionCustomization::MakeIn
 	return Instance;
 }
 
-TSharedRef<IPropertyTypeCustomization> FMovieSceneConditionCustomization::MakeInstance(UMovieSceneSequence* InMovieSceneSequence, const TWeakPtr<ISequencer> Sequencer)
+TSharedRef<IPropertyTypeCustomization> FMovieSceneConditionCustomization::MakeInstance(TWeakObjectPtr<UMovieSceneSequence> InMovieSceneSequence, const TWeakPtr<ISequencer> Sequencer)
 {
 	TSharedRef<FMovieSceneConditionCustomization> Instance = MakeShared<FMovieSceneConditionCustomization>();
 	Instance->Sequence = InMovieSceneSequence;
diff --git a/Engine/Source/Editor/MovieSceneTools/Public/Conditions/MovieSceneConditionCustomization.h b/Engine/Source/Editor/MovieSceneTools/Public/Conditions/MovieSceneConditionCustomization.h
--- a/Engine/Source/Editor/MovieSceneTools/Public/Conditions/MovieSceneConditionCustomization.h
+++ b/Engine/Source/Editor/MovieSceneTools/Public/Conditions/MovieSceneConditionCustomization.h
@@ -42,7 +42,7 @@ class FMovieSceneConditionCustomization : public IPropertyTypeCustomization
 public:
 
 	static MOVIESCENETOOLS_API TSharedRef<IPropertyTypeCustomization> MakeInstance();
-	static MOVIESCENETOOLS_API TSharedRef<IPropertyTypeCustomization> MakeInstance(UMovieSceneSequence* Sequence, const TWeakPtr<ISequencer> Sequencer);
+	static MOVIESCENETOOLS_API TSharedRef<IPropertyTypeCustomization> MakeInstance(TWeakObjectPtr<UMovieSceneSequence> Sequence, const TWeakPtr<ISequencer> Sequencer);
 
 	// IPropertyTypeCustomization interface
 	virtual void CustomizeHeader(TSharedRef<IPropertyHandle> InPropertyHandle, FDetailWidgetRow& HeaderRow, IPropertyTypeCustomizationUtils& CustomizationUtils) override;
