[Sequencer] Revert F3DTransformTrackEditor::ImportAnimSequenceTransforms to static, move to MovieSceneToolHelpers, and make public for use by Skydance and to be made scriptable in the future

diff --git a/Engine/Source/Editor/MovieSceneTools/Private/MovieSceneToolHelpers.cpp b/Engine/Source/Editor/MovieSceneTools/Private/MovieSceneToolHelpers.cpp
--- a/Engine/Source/Editor/MovieSceneTools/Private/MovieSceneToolHelpers.cpp
+++ b/Engine/Source/Editor/MovieSceneTools/Private/MovieSceneToolHelpers.cpp
@@ -95,6 +95,7 @@
 #include "Bindings/MovieSceneSpawnableBinding.h"
 #include "Tracks/MovieSceneBindingLifetimeTrack.h"
 #include "Sections/MovieSceneBindingLifetimeSection.h"
+#include "GameFramework/Character.h"
 
 /* FSkelMeshRecorder
  ***********/
@@ -5785,3 +5786,196 @@ bool MovieSceneToolHelpers::AddTransformKeys(
 
 	return true;
 }
+
+void AddUnwoundKey(FMovieSceneDoubleChannel& Channel, FFrameNumber Time, double Value)
+{
+	int32 Index = Channel.AddLinearKey(Time, Value);
+
+	TArrayView<FMovieSceneDoubleValue> Values = Channel.GetData().GetValues();
+	if (Index >= 1)
+	{
+		const double PreviousValue = Values[Index - 1].Value;
+		double NewValue = Value;
+
+		while (NewValue - PreviousValue > 180.0f)
+		{
+			NewValue -= 360.f;
+		}
+		while (NewValue - PreviousValue < -180.0f)
+		{
+			NewValue += 360.f;
+		}
+
+		Values[Index].Value = NewValue;
+	}
+}
+
+void MovieSceneToolHelpers::ImportAnimSequenceTransforms(const TSharedRef<ISequencer>& Sequencer, const FAssetData& Asset, UMovieScene3DTransformTrack* TransformTrack)
+{
+	if (!IsValid(TransformTrack))
+	{
+		return;
+	}
+
+	FSlateApplication::Get().DismissAllMenus();
+
+	FQualifiedFrameTime CurrentTime = Sequencer->GetLocalTime();
+
+	UAnimSequence* AnimSequence = Cast<UAnimSequence>(Asset.GetAsset());
+
+	// find object binding to recover any component transforms we need to incorporate (for characters)
+	FTransform InvComponentTransform;
+	UMovieSceneSequence* MovieSceneSequence = Sequencer->GetFocusedMovieSceneSequence();
+	if(MovieSceneSequence)
+	{
+		UMovieScene* MovieScene = MovieSceneSequence->GetMovieScene();
+		if(MovieScene)
+		{
+			FGuid ObjectBinding;
+			if(MovieScene->FindTrackBinding(*TransformTrack, ObjectBinding))
+			{
+				const UClass* ObjectClass = nullptr;
+				if(FMovieSceneSpawnable* Spawnable = MovieScene->FindSpawnable(ObjectBinding))
+				{
+					ObjectClass = Spawnable->GetObjectTemplate()->GetClass();
+				}
+				else if(FMovieScenePossessable* Possessable = MovieScene->FindPossessable(ObjectBinding))
+				{
+					ObjectClass = Possessable->GetPossessedObjectClass();
+				}
+
+				if(ObjectClass)
+				{
+					const ACharacter* Character = Cast<const ACharacter>(ObjectClass->ClassDefaultObject);
+					if(Character)
+					{
+						const USkeletalMeshComponent* SkeletalMeshComponent = Character->GetMesh();
+						FTransform MeshRelativeTransform = SkeletalMeshComponent->GetRelativeTransform();
+						InvComponentTransform = MeshRelativeTransform.GetRelativeTransform(SkeletalMeshComponent->GetOwner()->GetTransform()).Inverse();
+					}
+				}
+			}
+		}
+	}
+
+	if (AnimSequence && AnimSequence->GetDataModel()->GetNumBoneTracks() > 0)
+	{
+		const FScopedTransaction Transaction( NSLOCTEXT( "Sequencer", "ImportAnimSequenceTransforms", "Import Anim Sequence Transforms" ) );
+
+		TransformTrack->Modify();
+
+		UMovieScene3DTransformSection* Section = Cast<UMovieScene3DTransformSection>(TransformTrack->CreateNewSection());
+
+		Section->SetBlendType(EMovieSceneBlendType::Additive);
+		Section->SetMask(EMovieSceneTransformChannel::Translation | EMovieSceneTransformChannel::Rotation);
+		
+		FFrameRate TickResolution = Section->GetTypedOuter<UMovieScene>()->GetTickResolution();
+
+		// We know that there aren't any channel overrides (all 9 channels are double channels) because we're working
+		// on a brand new transform section.
+		TArrayView<FMovieSceneDoubleChannel*> DoubleChannels = Section->GetChannelProxy().GetChannels<FMovieSceneDoubleChannel>();
+
+		// Set default translation and rotation
+		for (int32 Index = 0; Index < 6; ++Index)
+		{
+			DoubleChannels[Index]->SetDefault(0.f);
+		}
+		// Set default scale
+		for (int32 Index = 6; Index < 9; ++Index)
+		{
+			DoubleChannels[Index]->SetDefault(1.f);
+		}
+
+		TransformTrack->AddSection(*Section);
+
+		if (Section->TryModify())
+		{
+			struct FTempTransformKey
+			{
+				FTransform Transform;
+				FRotator WoundRotation;
+				float Time;
+			};
+
+			float MinTime = FLT_MAX;
+
+			TArray<FTempTransformKey> TempKeys;
+
+			TArray<FName> BoneTrackNames;
+			AnimSequence->GetDataModelInterface()->GetBoneTrackNames(BoneTrackNames);
+
+			TArray<FTransform> BoneTransforms;
+			AnimSequence->GetDataModelInterface()->GetBoneTrackTransforms(BoneTrackNames[0], BoneTransforms);
+
+			for(int32 KeyIndex = 0; KeyIndex < BoneTransforms.Num(); KeyIndex++)
+			{
+				FTempTransformKey TempKey;
+				TempKey.Time = AnimSequence->GetTimeAtFrame(KeyIndex);
+				TempKey.Transform = BoneTransforms[KeyIndex];
+
+				// apply component transform if any
+				TempKey.Transform = InvComponentTransform * TempKey.Transform;
+
+				TempKey.WoundRotation = TempKey.Transform.GetRotation().Rotator();
+
+				TempKeys.Add(TempKey);
+
+				MinTime = FMath::Min(MinTime, TempKey.Time);
+			}
+
+			int32 TransformCount = TempKeys.Num();
+			for(int32 TransformIndex = 0; TransformIndex < TransformCount - 1; TransformIndex++)
+			{
+				FRotator& Rotator = TempKeys[TransformIndex].WoundRotation;
+				FRotator& NextRotator = TempKeys[TransformIndex + 1].WoundRotation;
+
+				FMath::WindRelativeAnglesDegrees(Rotator.Pitch, NextRotator.Pitch);
+				FMath::WindRelativeAnglesDegrees(Rotator.Yaw, NextRotator.Yaw);
+				FMath::WindRelativeAnglesDegrees(Rotator.Roll, NextRotator.Roll);
+			}
+
+			FFrameNumber MinKeyTime = (MinTime * TickResolution).RoundToFrame();
+
+			TRange<FFrameNumber> Range = Section->GetRange();
+			for(const FTempTransformKey& TempKey : TempKeys)
+			{
+				FFrameNumber KeyTime = (TempKey.Time * TickResolution).RoundToFrame();
+
+				KeyTime = CurrentTime.Time.FrameNumber + (KeyTime - MinKeyTime);
+
+				Range = TRange<FFrameNumber>::Hull(Range, TRange<FFrameNumber>(KeyTime));
+
+				const FVector3f Translation = (FVector3f)TempKey.Transform.GetTranslation();
+				const FVector3f Rotation = (FVector3f)TempKey.WoundRotation.Euler();
+				const FVector3f Scale = (FVector3f)TempKey.Transform.GetScale3D();
+
+				TArrayView<FMovieSceneDoubleChannel*> Channels = Section->GetChannelProxy().GetChannels<FMovieSceneDoubleChannel>();
+
+				Channels[0]->AddLinearKey(KeyTime, Translation.X);
+				Channels[1]->AddLinearKey(KeyTime, Translation.Y);
+				Channels[2]->AddLinearKey(KeyTime, Translation.Z);
+
+				AddUnwoundKey(*Channels[3], KeyTime, Rotation.X);
+				AddUnwoundKey(*Channels[4], KeyTime, Rotation.Y);
+				AddUnwoundKey(*Channels[5], KeyTime, Rotation.Z);
+
+				Channels[6]->AddLinearKey(KeyTime, Scale.X);
+				Channels[7]->AddLinearKey(KeyTime, Scale.Y);
+				Channels[8]->AddLinearKey(KeyTime, Scale.Z);
+			}
+
+			Section->SetRange(Range);
+			Section->SetRowIndex(FindAvailableRowIndex(TransformTrack, Section));
+
+			Sequencer->NotifyMovieSceneDataChanged( EMovieSceneDataChangeType::MovieSceneStructureItemAdded );
+		}
+	}
+}
+
+void MovieSceneToolHelpers::ImportAnimSequenceTransformsEnterPressed(const TSharedRef<ISequencer>& Sequencer, const TArray<FAssetData>& Asset, UMovieScene3DTransformTrack* TransformTrack)
+{
+	if (Asset.Num() > 0)
+	{
+		ImportAnimSequenceTransforms(Sequencer, Asset[0].GetAsset(), TransformTrack);
+	}
+}
diff --git a/Engine/Source/Editor/MovieSceneTools/Private/TrackEditors/TransformTrackEditor.cpp b/Engine/Source/Editor/MovieSceneTools/Private/TrackEditors/TransformTrackEditor.cpp
--- a/Engine/Source/Editor/MovieSceneTools/Private/TrackEditors/TransformTrackEditor.cpp
+++ b/Engine/Source/Editor/MovieSceneTools/Private/TrackEditors/TransformTrackEditor.cpp
@@ -208,12 +208,18 @@ void F3DTransformTrackEditor::BuildAssetPickerSubMenu(FMenuBuilder& InMenuBuilde
 	AssetPickerConfig.OnAssetSelected = FOnAssetSelected::CreateLambda(
 		[this, TransformTrackWeak](const FAssetData& InAssetData)
 		{
-			ImportAnimSequenceTransforms(InAssetData, TransformTrackWeak.Get());
+			if (const TSharedPtr<ISequencer> SequencerPtr = GetSequencer())
+			{
+				MovieSceneToolHelpers::ImportAnimSequenceTransforms(SequencerPtr.ToSharedRef(), InAssetData, TransformTrackWeak.Get());
+			}
 		});
 	AssetPickerConfig.OnAssetEnterPressed = FOnAssetEnterPressed::CreateLambda(
 		[this, TransformTrackWeak](const TArray<FAssetData>& InSelectedAssets)
 		{
-			ImportAnimSequenceTransformsEnterPressed(InSelectedAssets, TransformTrackWeak.Get());
+			if (const TSharedPtr<ISequencer> SequencerPtr = GetSequencer())
+			{
+				MovieSceneToolHelpers::ImportAnimSequenceTransformsEnterPressed(SequencerPtr.ToSharedRef(), InSelectedAssets, TransformTrackWeak.Get());
+			}
 		});
 	AssetPickerConfig.SaveSettingsName = TEXT("SequencerAssetPicker");
 	AssetPickerConfig.AdditionalReferencingAssets.Add(FAssetData(Sequence));
@@ -1631,206 +1637,6 @@ void F3DTransformTrackEditor::ProcessKeyOperation(UObject* ObjectToKey, TArrayVi
 	}
 }
 
-void AddUnwoundKey(FMovieSceneDoubleChannel& Channel, FFrameNumber Time, double Value)
-{
-	int32 Index = Channel.AddLinearKey(Time, Value);
-
-	TArrayView<FMovieSceneDoubleValue> Values = Channel.GetData().GetValues();
-	if (Index >= 1)
-	{
-		const double PreviousValue = Values[Index - 1].Value;
-		double NewValue = Value;
-
-		while (NewValue - PreviousValue > 180.0f)
-		{
-			NewValue -= 360.f;
-		}
-		while (NewValue - PreviousValue < -180.0f)
-		{
-			NewValue += 360.f;
-		}
-
-		Values[Index].Value = NewValue;
-	}
-}
-
-
-void F3DTransformTrackEditor::ImportAnimSequenceTransforms(const FAssetData& Asset, UMovieScene3DTransformTrack* TransformTrack)
-{
-	if (!IsValid(TransformTrack))
-	{
-		return;
-	}
-
-	const TSharedPtr<ISequencer> SequencerPtr = GetSequencer();
-	if (!SequencerPtr.IsValid())
-	{
-		return;
-	}
-
-	FSlateApplication::Get().DismissAllMenus();
-
-	FQualifiedFrameTime CurrentTime = SequencerPtr->GetLocalTime();
-
-	UAnimSequence* AnimSequence = Cast<UAnimSequence>(Asset.GetAsset());
-
-	// find object binding to recover any component transforms we need to incorporate (for characters)
-	FTransform InvComponentTransform;
-	UMovieSceneSequence* MovieSceneSequence = SequencerPtr->GetFocusedMovieSceneSequence();
-	if(MovieSceneSequence)
-	{
-		UMovieScene* MovieScene = MovieSceneSequence->GetMovieScene();
-		if(MovieScene)
-		{
-			FGuid ObjectBinding;
-			if(MovieScene->FindTrackBinding(*TransformTrack, ObjectBinding))
-			{
-				const UClass* ObjectClass = nullptr;
-				if(FMovieSceneSpawnable* Spawnable = MovieScene->FindSpawnable(ObjectBinding))
-				{
-					ObjectClass = Spawnable->GetObjectTemplate()->GetClass();
-				}
-				else if(FMovieScenePossessable* Possessable = MovieScene->FindPossessable(ObjectBinding))
-				{
-					ObjectClass = Possessable->GetPossessedObjectClass();
-				}
-
-				if(ObjectClass)
-				{
-					const ACharacter* Character = Cast<const ACharacter>(ObjectClass->ClassDefaultObject);
-					if(Character)
-					{
-						const USkeletalMeshComponent* SkeletalMeshComponent = Character->GetMesh();
-						FTransform MeshRelativeTransform = SkeletalMeshComponent->GetRelativeTransform();
-						InvComponentTransform = MeshRelativeTransform.GetRelativeTransform(SkeletalMeshComponent->GetOwner()->GetTransform()).Inverse();
-					}
-				}
-			}
-		}
-	}
-
-	if (AnimSequence && AnimSequence->GetDataModel()->GetNumBoneTracks() > 0)
-	{
-		const FScopedTransaction Transaction( NSLOCTEXT( "Sequencer", "ImportAnimSequenceTransforms", "Import Anim Sequence Transforms" ) );
-
-		TransformTrack->Modify();
-
-		UMovieScene3DTransformSection* Section = Cast<UMovieScene3DTransformSection>(TransformTrack->CreateNewSection());
-
-		Section->SetBlendType(EMovieSceneBlendType::Additive);
-		Section->SetMask(EMovieSceneTransformChannel::Translation | EMovieSceneTransformChannel::Rotation);
-		
-		FFrameRate TickResolution = Section->GetTypedOuter<UMovieScene>()->GetTickResolution();
-
-		// We know that there aren't any channel overrides (all 9 channels are double channels) because we're working
-		// on a brand new transform section.
-		TArrayView<FMovieSceneDoubleChannel*> DoubleChannels = Section->GetChannelProxy().GetChannels<FMovieSceneDoubleChannel>();
-
-		// Set default translation and rotation
-		for (int32 Index = 0; Index < 6; ++Index)
-		{
-			DoubleChannels[Index]->SetDefault(0.f);
-		}
-		// Set default scale
-		for (int32 Index = 6; Index < 9; ++Index)
-		{
-			DoubleChannels[Index]->SetDefault(1.f);
-		}
-
-		TransformTrack->AddSection(*Section);
-
-		if (Section->TryModify())
-		{
-			struct FTempTransformKey
-			{
-				FTransform Transform;
-				FRotator WoundRotation;
-				float Time;
-			};
-
-			float MinTime = FLT_MAX;
-
-			TArray<FTempTransformKey> TempKeys;
-
-			TArray<FName> BoneTrackNames;
-			AnimSequence->GetDataModelInterface()->GetBoneTrackNames(BoneTrackNames);
-
-			TArray<FTransform> BoneTransforms;
-			AnimSequence->GetDataModelInterface()->GetBoneTrackTransforms(BoneTrackNames[0], BoneTransforms);
-
-			for(int32 KeyIndex = 0; KeyIndex < BoneTransforms.Num(); KeyIndex++)
-			{
-				FTempTransformKey TempKey;
-				TempKey.Time = AnimSequence->GetTimeAtFrame(KeyIndex);
-				TempKey.Transform = BoneTransforms[KeyIndex];
-
-				// apply component transform if any
-				TempKey.Transform = InvComponentTransform * TempKey.Transform;
-
-				TempKey.WoundRotation = TempKey.Transform.GetRotation().Rotator();
-
-				TempKeys.Add(TempKey);
-
-				MinTime = FMath::Min(MinTime, TempKey.Time);
-			}
-
-			int32 TransformCount = TempKeys.Num();
-			for(int32 TransformIndex = 0; TransformIndex < TransformCount - 1; TransformIndex++)
-			{
-				FRotator& Rotator = TempKeys[TransformIndex].WoundRotation;
-				FRotator& NextRotator = TempKeys[TransformIndex + 1].WoundRotation;
-
-				FMath::WindRelativeAnglesDegrees(Rotator.Pitch, NextRotator.Pitch);
-				FMath::WindRelativeAnglesDegrees(Rotator.Yaw, NextRotator.Yaw);
-				FMath::WindRelativeAnglesDegrees(Rotator.Roll, NextRotator.Roll);
-			}
-
-			FFrameNumber MinKeyTime = (MinTime * TickResolution).RoundToFrame();
-
-			TRange<FFrameNumber> Range = Section->GetRange();
-			for(const FTempTransformKey& TempKey : TempKeys)
-			{
-				FFrameNumber KeyTime = (TempKey.Time * TickResolution).RoundToFrame();
-
-				KeyTime = CurrentTime.Time.FrameNumber + (KeyTime - MinKeyTime);
-
-				Range = TRange<FFrameNumber>::Hull(Range, TRange<FFrameNumber>(KeyTime));
-
-				const FVector3f Translation = (FVector3f)TempKey.Transform.GetTranslation();
-				const FVector3f Rotation = (FVector3f)TempKey.WoundRotation.Euler();
-				const FVector3f Scale = (FVector3f)TempKey.Transform.GetScale3D();
-
-				TArrayView<FMovieSceneDoubleChannel*> Channels = Section->GetChannelProxy().GetChannels<FMovieSceneDoubleChannel>();
-
-				Channels[0]->AddLinearKey(KeyTime, Translation.X);
-				Channels[1]->AddLinearKey(KeyTime, Translation.Y);
-				Channels[2]->AddLinearKey(KeyTime, Translation.Z);
-
-				AddUnwoundKey(*Channels[3], KeyTime, Rotation.X);
-				AddUnwoundKey(*Channels[4], KeyTime, Rotation.Y);
-				AddUnwoundKey(*Channels[5], KeyTime, Rotation.Z);
-
-				Channels[6]->AddLinearKey(KeyTime, Scale.X);
-				Channels[7]->AddLinearKey(KeyTime, Scale.Y);
-				Channels[8]->AddLinearKey(KeyTime, Scale.Z);
-			}
-
-			Section->SetRange(Range);
-			Section->SetRowIndex(MovieSceneToolHelpers::FindAvailableRowIndex(TransformTrack, Section));
-
-			SequencerPtr->NotifyMovieSceneDataChanged( EMovieSceneDataChangeType::MovieSceneStructureItemAdded );
-		}
-	}
-}
-
-void F3DTransformTrackEditor::ImportAnimSequenceTransformsEnterPressed(const TArray<FAssetData>& Asset, UMovieScene3DTransformTrack* TransformTrack)
-{
-	if (Asset.Num() > 0)
-	{
-		ImportAnimSequenceTransforms(Asset[0].GetAsset(), TransformTrack);
-	}
-}
-
 bool F3DTransformTrackEditor::MatchesContext(const FTransactionContext& InContext, const TArray<TPair<UObject*, FTransactionObjectEvent>>& TransactionObjects) const
 {
 	SectionsGettingUndone.SetNum(0);
diff --git a/Engine/Source/Editor/MovieSceneTools/Private/TrackEditors/TransformTrackEditor.h b/Engine/Source/Editor/MovieSceneTools/Private/TrackEditors/TransformTrackEditor.h
--- a/Engine/Source/Editor/MovieSceneTools/Private/TrackEditors/TransformTrackEditor.h
+++ b/Engine/Source/Editor/MovieSceneTools/Private/TrackEditors/TransformTrackEditor.h
@@ -183,12 +183,6 @@ private:
 
 private:
 
-	/** Import an animation sequence's root transforms into a transform section */
-	void ImportAnimSequenceTransforms(const FAssetData& Asset, UMovieScene3DTransformTrack* TransformTrack);
-
-	/** Import an animation sequence's root transforms into a transform section */
-	void ImportAnimSequenceTransformsEnterPressed(const TArray<FAssetData>& Asset, UMovieScene3DTransformTrack* TransformTrack);
-
 	/** ConstraintChannel Delegates*/
 	FDelegateHandle OnSceneComponentConstrainedHandle;
 	void HandleOnConstraintAdded(IMovieSceneConstrainedSection* InSection, FMovieSceneConstraintChannel* InConstraintChannel);
diff --git a/Engine/Source/Editor/MovieSceneTools/Public/MovieSceneToolHelpers.h b/Engine/Source/Editor/MovieSceneTools/Public/MovieSceneToolHelpers.h
--- a/Engine/Source/Editor/MovieSceneTools/Public/MovieSceneToolHelpers.h
+++ b/Engine/Source/Editor/MovieSceneTools/Public/MovieSceneToolHelpers.h
@@ -45,6 +45,7 @@ struct FBakingAnimationKeySettings;
 struct FKeyDataOptimizationParams;
 class UMovieSceneSubSection;
 enum class EMovieSceneTransformChannel : uint32;
+class UMovieScene3DTransformTrack;
 
 namespace fbxsdk
 {
@@ -1094,6 +1095,12 @@ public:
 		const TArray<FFrameNumber>& Frames,
 		const TArray<FTransform>& InLocalTransforms,
 		const EMovieSceneTransformChannel& InChannels);
+
+	/** Import an animation sequence's root transforms into a transform section */
+	static void ImportAnimSequenceTransforms(const TSharedRef<ISequencer>& Sequencer, const FAssetData& Asset, UMovieScene3DTransformTrack* TransformTrack);
+
+	/** Import an animation sequence's root transforms into a transform section */
+	static void ImportAnimSequenceTransformsEnterPressed(const TSharedRef<ISequencer>& Sequencer, const TArray<FAssetData>& Asset, UMovieScene3DTransformTrack* TransformTrack);
 };
 
 // Helper to make spawnables persist throughout the export process and then restore properly afterwards
