Improve gizmo positioning for subsequence transform origin editing.

The gizmo now uses the average location of the actors in the sub-hierarchy of the subsequence being edited.

diff --git a/Engine/Source/Editor/MovieSceneTools/Private/EditModes/SubTrackEditorMode.cpp b/Engine/Source/Editor/MovieSceneTools/Private/EditModes/SubTrackEditorMode.cpp
--- a/Engine/Source/Editor/MovieSceneTools/Private/EditModes/SubTrackEditorMode.cpp
+++ b/Engine/Source/Editor/MovieSceneTools/Private/EditModes/SubTrackEditorMode.cpp
@@ -8,6 +8,7 @@
 #include "UnrealClient.h"
 #include "EditorViewportClient.h"
 #include "IMovieScenePlaybackClient.h"
+#include "LevelSequence.h"
 #include "MovieScene.h"
 #include "SequencerChannelTraits.h"
 #include "Compilation/MovieSceneCompiledDataManager.h"
@@ -24,7 +25,8 @@ FName FSubTrackEditorMode::ModeName("EditMode.SubTrackEditMode");
 FSubTrackEditorMode::FSubTrackEditorMode()
 {
 	CachedLocation.Reset();
-	PreviewCoordinateSpace.Reset();
+	PreviewCoordinateSpaceRotation.Reset();
+	PreviewLocation.Reset();
 }
 
 FSubTrackEditorMode::~FSubTrackEditorMode()
@@ -34,7 +36,8 @@ FSubTrackEditorMode::~FSubTrackEditorMode()
 void FSubTrackEditorMode::Initialize()
 {
 	CachedLocation.Reset();
-	PreviewCoordinateSpace.Reset();
+	PreviewCoordinateSpaceRotation.Reset();
+	PreviewLocation.Reset();
 }
 
 bool FSubTrackEditorMode::InputDelta(FEditorViewportClient* InViewportClient, FViewport* InViewport, FVector& InDrag, FRotator& InRot, FVector& InScale)
@@ -60,7 +63,8 @@ bool FSubTrackEditorMode::InputDelta(FEditorViewportClient* InViewportClient, FV
 		const FTransform LocalPosition = TransformOriginFocusedSequence * FTransform(InDrag) * TransformOriginFocusedSequence.Inverse();
 
 		// Keep preview space up-to-date.
-		PreviewCoordinateSpace.GetValue() *= FTransform(InRot).ToMatrixNoScale();
+		PreviewCoordinateSpaceRotation.GetValue() *= FTransform(InRot).ToMatrixNoScale().RemoveTranslation();
+		PreviewLocation.GetValue() += InDrag;
 		
 		OnOriginValueChanged.Broadcast(LocalPosition.GetLocation(), LocalRotation.Rotator());
 		return true;
@@ -71,12 +75,14 @@ bool FSubTrackEditorMode::InputDelta(FEditorViewportClient* InViewportClient, FV
 
 bool FSubTrackEditorMode::StartTracking(FEditorViewportClient* InViewportClient, FViewport* InViewport)
 {
-	const UMovieSceneSubSection* SubSection = GetSelectedSection();
-
-	if (!AreAnyActorsSelected() && SubSection && DoesSubSectionHaveTransformOverrides(*SubSection))
+	if (const UMovieSceneSubSection* SubSection = GetSectionToEdit())
 	{
 		bIsTracking = true;
-		PreviewCoordinateSpace = GetFinalTransformOriginForSubSection(SubSection).ToMatrixNoScale().RemoveTranslation();
+		PreviewCoordinateSpaceRotation = GetFinalTransformOriginForSubSection(SubSection).ToMatrixNoScale().RemoveTranslation();
+		if (!PreviewLocation.IsSet())
+		{
+			PreviewLocation = GetAverageLocationOfBindingsInSubSection(SubSection);
+		}
 		return true;
 	}
 	
@@ -88,7 +94,9 @@ bool FSubTrackEditorMode::EndTracking(FEditorViewportClient* InViewportClient, F
 	if(bIsTracking)
 	{
 		bIsTracking = false;
-		PreviewCoordinateSpace.Reset();
+		// Only reset the preview rotation. Resetting the preview location here could interfere with multiple rotation edits in a row.
+		// The location would otherwise change after the user let go of the mouse, and would have to move it to the new location to continue rotating.
+		PreviewCoordinateSpaceRotation.Reset();
 		return true;
 	}
 	
@@ -231,11 +239,10 @@ bool FSubTrackEditorMode::AreAnyActorsSelected() const
 
 FVector FSubTrackEditorMode::GetWidgetLocation() const
 {
-	const UMovieSceneSubSection* SubSection = GetSelectedSection();
-	if(!AreAnyActorsSelected() && SubSection && DoesSubSectionHaveTransformOverrides(*SubSection))
+	if(const UMovieSceneSubSection* SubSection = GetSectionToEdit())
 	{
-		FVector NewLocation = GetFinalTransformOriginForSubSection(SubSection).GetLocation();
-
+		const FVector NewLocation = PreviewLocation.IsSet() ? PreviewLocation.GetValue() : GetAverageLocationOfBindingsInSubSection(SubSection);
+		
 		if(!CachedLocation.IsSet() || !NewLocation.Equals(CachedLocation.GetValue()))
 		{
 			CachedLocation = NewLocation;
@@ -248,10 +255,114 @@ FVector FSubTrackEditorMode::GetWidgetLocation() const
 	return FEdMode::GetWidgetLocation();
 }
 
+FVector FSubTrackEditorMode::GetAverageLocationOfBindingsInSubSection(const UMovieSceneSubSection* SubSection) const
+{
+	FVector TotalPosition = FVector::ZeroVector;
+	int32 ActorCount = 0;
+
+	const UMovieSceneSequence* CurrentSequence = SubSection->GetSequence();
+	if (!CurrentSequence)
+	{
+		return TotalPosition;
+	}
+
+	const UMovieScene* CurrentMovieScene = CurrentSequence->GetMovieScene();
+	if (!CurrentMovieScene)
+	{
+		return TotalPosition;
+	}
+
+	TArray<FMovieSceneBinding> Bindings = CurrentMovieScene->GetBindings();
+
+	const TSharedPtr<ISequencer> Sequencer = WeakSequencer.Pin();
+
+	if (!Sequencer)
+	{
+		return FVector::ZeroVector;
+	}
+
+	const FMovieSceneSequenceID FocusedSequenceID = GetFocusedSequenceID().GetValue();
+	
+	const FMovieSceneRootEvaluationTemplateInstance& EvaluationTemplate = Sequencer->GetEvaluationTemplate();
+	UMovieSceneCompiledDataManager* CompiledDataManager = EvaluationTemplate.GetCompiledDataManager();
+	UMovieSceneSequence* RootSequence = Sequencer->GetRootMovieSceneSequence();
+	
+
+	const FMovieSceneCompiledDataID DataID = CompiledDataManager->Compile(RootSequence);
+	const FMovieSceneSequenceHierarchy Hierarchy = CompiledDataManager->GetHierarchyChecked(DataID);
+
+	RecursiveAccumulateBindingPositions(SubSection, TotalPosition, ActorCount, &Hierarchy, FocusedSequenceID, FocusedSequenceID, Sequencer);
+
+	if (ActorCount > 0)
+	{
+		return TotalPosition / ActorCount;
+	}
+	
+	return FVector::ZeroVector;
+}
+
+void FSubTrackEditorMode::RecursiveAccumulateBindingPositions(const UMovieSceneSubSection* SubSection, FVector& AccumulatedLocation,
+	int32& ActorCount, const FMovieSceneSequenceHierarchy* Hierarchy, const FMovieSceneSequenceID FocusedSequenceID, const FMovieSceneSequenceID ParentSequenceID, const TSharedPtr<ISequencer> Sequencer) const
+{
+	const UMovieSceneSequence* CurrentSequence = SubSection->GetSequence();
+	if (!CurrentSequence)
+	{
+		return;
+	}
+
+	const UMovieScene* CurrentMovieScene = CurrentSequence->GetMovieScene();
+	if (!CurrentMovieScene)
+	{
+		return;
+	}
+	
+	TArray<FMovieSceneBinding> Bindings = CurrentMovieScene->GetBindings();
+
+	UE::MovieScene::FSubSequencePath Path;
+	
+	Path.Reset(ParentSequenceID, Hierarchy);
+
+	const FMovieSceneSequenceID ResolvedSequenceID = Path.ResolveChildSequenceID(SubSection->GetSequenceID());
+	
+	for (FMovieSceneBinding& Binding : Bindings)
+	{
+		FMovieSceneObjectBindingID RelativeBinding = UE::MovieScene::FRelativeObjectBindingID(FocusedSequenceID, ResolvedSequenceID, Binding.GetObjectGuid(), Hierarchy);
+	
+		for (TWeakObjectPtr<> Object : RelativeBinding.ResolveBoundObjects(FocusedSequenceID, *Sequencer.Get()))
+		{
+			if (TStrongObjectPtr<UObject> StrongObject = Object.Pin())
+			{
+				if (const AActor* BoundActor = Cast<AActor>(StrongObject.Get()))
+				{
+					AccumulatedLocation += BoundActor->GetActorLocation();
+					ActorCount++;
+				}
+			}
+		}
+	}
+
+	const TArray<UMovieSceneTrack*> Tracks = CurrentMovieScene->GetTracks();
+
+	for (UMovieSceneTrack* Track : Tracks)
+	{
+		if (Track)
+		{
+			if (const UMovieSceneSubTrack* SubTrack = Cast<UMovieSceneSubTrack>(Track))
+			{
+				for (UMovieSceneSection* Section : SubTrack->GetAllSections())
+				{
+					const UMovieSceneSubSection* ChildSubSection = Cast<UMovieSceneSubSection>(Section);
+
+					RecursiveAccumulateBindingPositions(ChildSubSection, AccumulatedLocation, ActorCount, Hierarchy, FocusedSequenceID, ResolvedSequenceID, Sequencer);
+				}
+			}
+		}
+	}
+}
+
 bool FSubTrackEditorMode::ShouldDrawWidget() const
 {
-	const UMovieSceneSubSection* SubSection = GetSelectedSection();
-	if(!AreAnyActorsSelected() && SubSection && DoesSubSectionHaveTransformOverrides(*SubSection))
+	if (GetSectionToEdit())
 	{
 		return true;
 	}
@@ -273,18 +384,16 @@ bool FSubTrackEditorMode::GetCustomDrawingCoordinateSystem(FMatrix& OutMatrix, v
 		return false;
 	}
 
-	// While manipulating the gizmo, the preview coordinate space is kept up-to-date directly. The transform origin data however is set by a callback,
-	// and can be out of date, which would cause the gizmo to jitter.
-	if (PreviewCoordinateSpace.IsSet())
+	if (const UMovieSceneSubSection* SubSection = GetSectionToEdit())
 	{
-		OutMatrix = PreviewCoordinateSpace.GetValue();
-		return true;
-	}
+		// While manipulating the gizmo, the preview coordinate space is kept up-to-date directly, since the transform origin data is set by a callback,
+		// and can be out of date, which would cause the gizmo to jitter.
+		if (PreviewCoordinateSpaceRotation.IsSet())
+		{
+			OutMatrix = PreviewCoordinateSpaceRotation.GetValue().RemoveTranslation();
+			return true;
+		}
 
-	// Before interacting with the gizmo, query the origin data directly as a starting point.
-	const UMovieSceneSubSection* SubSection = GetSelectedSection();
-	if (!AreAnyActorsSelected() && SubSection && DoesSubSectionHaveTransformOverrides(*SubSection))
-	{
 		OutMatrix = GetFinalTransformOriginForSubSection(SubSection).ToMatrixNoScale().RemoveTranslation();
 		return true;
 	}
@@ -306,6 +415,13 @@ bool FSubTrackEditorMode::IsCompatibleWith(FEditorModeID OtherModeID) const
 	return true;
 }
 
+void FSubTrackEditorMode::ClearCachedCoordinates()
+{
+	PreviewLocation.Reset();
+	CachedLocation.Reset();
+	PreviewCoordinateSpaceRotation.Reset();
+}
+
 bool FSubTrackEditorMode::DoesSubSectionHaveTransformOverrides(const UMovieSceneSubSection& SubSection)
 {
 	if (SubSection.IsActive())
@@ -382,3 +498,14 @@ UMovieSceneSubSection* FSubTrackEditorMode::GetSelectedSection() const
 	}
 	return SelectedSection;
 }
+
+UMovieSceneSubSection* FSubTrackEditorMode::GetSectionToEdit() const
+{
+	UMovieSceneSubSection* SubSectionToReturn = GetSelectedSection();
+	if (!AreAnyActorsSelected() && SubSectionToReturn && DoesSubSectionHaveTransformOverrides(*SubSectionToReturn))
+	{
+		return SubSectionToReturn;
+	}
+
+	return nullptr;
+}
\ No newline at end of file
diff --git a/Engine/Source/Editor/MovieSceneTools/Private/EditModes/SubTrackEditorMode.h b/Engine/Source/Editor/MovieSceneTools/Private/EditModes/SubTrackEditorMode.h
--- a/Engine/Source/Editor/MovieSceneTools/Private/EditModes/SubTrackEditorMode.h
+++ b/Engine/Source/Editor/MovieSceneTools/Private/EditModes/SubTrackEditorMode.h
@@ -32,6 +32,8 @@ public:
 	void SetSequencer(const TSharedPtr<ISequencer>& InSequencer) { WeakSequencer = InSequencer; }
 	FOnOriginValueChanged& GetOnOriginValueChanged() { return OnOriginValueChanged; };
 
+	void ClearCachedCoordinates();
+
 private:
 	// Returns true if there are any sub tracks with active transform origin overrides.
 	static bool DoesSubSectionHaveTransformOverrides(const UMovieSceneSubSection& SubSection);
@@ -46,9 +48,16 @@ private:
 	FTransform GetTransformOriginForSequence(TOptional<FMovieSceneSequenceID> InSequenceID) const;
 	// returns true if any actors are selected in the level editor. This is used to prevent this editor mode from being active when actors are selected.
 	bool AreAnyActorsSelected() const;
+	/** Gets the average location of the objects in a subsection. */
+	FVector GetAverageLocationOfBindingsInSubSection(const UMovieSceneSubSection* SubSection) const;
+	/** Accumulate the positions and counts of all bound objects in and descended from the given SubSection */
+	void RecursiveAccumulateBindingPositions(const UMovieSceneSubSection* SubSection, FVector& AccumulatedLocation,
+	int32& ActorCount, const FMovieSceneSequenceHierarchy* Hierarchy, const FMovieSceneSequenceID FocusedSequenceID, const FMovieSceneSequenceID ParentSequenceID, const TSharedPtr<ISequencer> Sequencer) const;
 
-	// Returns the currently selected 
+	/** Returns the currently selected section. */
 	UMovieSceneSubSection* GetSelectedSection() const;
+	/** Returns the selected section, if its origin overrides can be edited, and there are no selected actors. */
+	UMovieSceneSubSection* GetSectionToEdit() const;
 
 	/** Sequencer that owns this editor mode */
 	TWeakPtr<ISequencer> WeakSequencer;
@@ -62,9 +71,13 @@ private:
 	/** Used to ensure that we mirror the behavior in StartTracking if the selection were to change mid-drag. */
 	bool bIsTracking = false;
 
-	/** Caches the transform space to use for editing/input.
+	/** Caches the transform space to use for editing with the gizmo.
 	 *  Updates to the channel data aren't reflected in time for the UI, so keeping the preview space up-to-date prevents the gizmo from flickering. */
-	TOptional<FMatrix> PreviewCoordinateSpace;
+	TOptional<FMatrix> PreviewCoordinateSpaceRotation;
+	/** Caches the location to use for editing with the gizmo.
+	 * When rotating, the average location of the actors in a subsequence can change, so instead of querying it every frame,
+	 * the value is cached when an edit begins, and is only updated if the gizmo is dragged. */
+	TOptional<FVector> PreviewLocation;
 
 	const TArray<FName> IncompatibleEditorModes = TArray<FName>({ "EditMode.ControlRig", "EM_Landscape" });
 };
diff --git a/Engine/Source/Editor/MovieSceneTools/Private/TrackEditors/SubTrackEditor.cpp b/Engine/Source/Editor/MovieSceneTools/Private/TrackEditors/SubTrackEditor.cpp
--- a/Engine/Source/Editor/MovieSceneTools/Private/TrackEditors/SubTrackEditor.cpp
+++ b/Engine/Source/Editor/MovieSceneTools/Private/TrackEditors/SubTrackEditor.cpp
@@ -1248,6 +1248,13 @@ void FSubTrackEditor::ResetSectionPreviews()
 		bSectionReverted = true;
 	}
 
+	FSubTrackEditorMode* EditorMode = static_cast<FSubTrackEditorMode*>(GLevelEditorModeTools().GetActiveMode(FSubTrackEditorMode::ModeName));
+
+	if (EditorMode)
+	{
+		EditorMode->ClearCachedCoordinates();
+	}
+	
 	if(bSectionReverted)
 	{
 		GetSequencer()->NotifyMovieSceneDataChanged(EMovieSceneDataChangeType::TrackValueChanged);		
