Support "S" Key, shift + "W" and shift + "E" for keyframing transform origins for subsequences.

SInce there are no scale channels for transform overrides shift + "R" will not key anything.

diff --git a/Engine/Source/Editor/MovieSceneTools/Private/TrackEditors/SubTrackEditor.cpp b/Engine/Source/Editor/MovieSceneTools/Private/TrackEditors/SubTrackEditor.cpp
--- a/Engine/Source/Editor/MovieSceneTools/Private/TrackEditors/SubTrackEditor.cpp
+++ b/Engine/Source/Editor/MovieSceneTools/Private/TrackEditors/SubTrackEditor.cpp
@@ -41,7 +41,9 @@
 #include "IDetailsView.h"
 #include "IKeyArea.h"
 #include "IStructureDetailsView.h"
+#include "LevelEditorViewport.h"
 #include "PropertyEditorModule.h"
+#include "Selection.h"
 #include "Compilation/MovieSceneCompiledDataManager.h"
 #include "EditModes/SubTrackEditorMode.h"
 #include "EntitySystem/BuiltInComponentTypes.h"
@@ -308,7 +310,7 @@ TSharedPtr<SWidget> FSubTrackEditor::BuildOutlinerEditWidget(const FGuid& Object
 }
 
 
-void FSubTrackEditor::GetOriginKeys(const FVector& CurrentPosition, const FRotator& CurrentRotation, UMovieSceneSection* Section, FGeneratedTrackKeys& OutGeneratedKeys)
+void FSubTrackEditor::GetOriginKeys(const FVector& CurrentPosition, const FRotator& CurrentRotation, const EMovieSceneTransformChannel ChannelsToKey, UMovieSceneSection* Section, FGeneratedTrackKeys& OutGeneratedKeys)
 {
 	FMovieSceneChannelProxy& SectionChannelProxy = Section->GetChannelProxy();
 	TMovieSceneChannelHandle<FMovieSceneDoubleChannel> ChannelHandles[] = {
@@ -320,12 +322,12 @@ void FSubTrackEditor::GetOriginKeys(const FVector& CurrentPosition, const FRotat
 		SectionChannelProxy.GetChannelByName<FMovieSceneDoubleChannel>("Override.Rotation.Z")
 	};
 	
-	OutGeneratedKeys.Add(FMovieSceneChannelValueSetter::Create<FMovieSceneDoubleChannel>(ChannelHandles[0].GetChannelIndex(), CurrentPosition.X, true));
-	OutGeneratedKeys.Add(FMovieSceneChannelValueSetter::Create<FMovieSceneDoubleChannel>(ChannelHandles[1].GetChannelIndex(), CurrentPosition.Y, true));
-	OutGeneratedKeys.Add(FMovieSceneChannelValueSetter::Create<FMovieSceneDoubleChannel>(ChannelHandles[2].GetChannelIndex(), CurrentPosition.Z, true));
-	OutGeneratedKeys.Add(FMovieSceneChannelValueSetter::Create<FMovieSceneDoubleChannel>(ChannelHandles[3].GetChannelIndex(), CurrentRotation.Roll, true));
-	OutGeneratedKeys.Add(FMovieSceneChannelValueSetter::Create<FMovieSceneDoubleChannel>(ChannelHandles[4].GetChannelIndex(), CurrentRotation.Pitch, true));
-	OutGeneratedKeys.Add(FMovieSceneChannelValueSetter::Create<FMovieSceneDoubleChannel>(ChannelHandles[5].GetChannelIndex(), CurrentRotation.Yaw, true));
+	OutGeneratedKeys.Add(FMovieSceneChannelValueSetter::Create<FMovieSceneDoubleChannel>(ChannelHandles[0].GetChannelIndex(), CurrentPosition.X,     EnumHasAnyFlags(ChannelsToKey, EMovieSceneTransformChannel::TranslationX)));
+	OutGeneratedKeys.Add(FMovieSceneChannelValueSetter::Create<FMovieSceneDoubleChannel>(ChannelHandles[1].GetChannelIndex(), CurrentPosition.Y,     EnumHasAnyFlags(ChannelsToKey, EMovieSceneTransformChannel::TranslationY)));
+	OutGeneratedKeys.Add(FMovieSceneChannelValueSetter::Create<FMovieSceneDoubleChannel>(ChannelHandles[2].GetChannelIndex(), CurrentPosition.Z,     EnumHasAnyFlags(ChannelsToKey, EMovieSceneTransformChannel::TranslationZ)));
+	OutGeneratedKeys.Add(FMovieSceneChannelValueSetter::Create<FMovieSceneDoubleChannel>(ChannelHandles[3].GetChannelIndex(), CurrentRotation.Roll,  EnumHasAnyFlags(ChannelsToKey, EMovieSceneTransformChannel::RotationX)));
+	OutGeneratedKeys.Add(FMovieSceneChannelValueSetter::Create<FMovieSceneDoubleChannel>(ChannelHandles[4].GetChannelIndex(), CurrentRotation.Pitch, EnumHasAnyFlags(ChannelsToKey, EMovieSceneTransformChannel::RotationY)));
+	OutGeneratedKeys.Add(FMovieSceneChannelValueSetter::Create<FMovieSceneDoubleChannel>(ChannelHandles[5].GetChannelIndex(), CurrentRotation.Yaw,   EnumHasAnyFlags(ChannelsToKey, EMovieSceneTransformChannel::RotationZ)));
 }
 
 TSharedRef<ISequencerTrackEditor> FSubTrackEditor::CreateTrackEditor(TSharedRef<ISequencer> InSequencer)
@@ -722,6 +724,190 @@ bool FSubTrackEditor::GetDefaultExpansionState(UMovieSceneTrack* InTrack) const
 }
 
 
+bool FSubTrackEditor::CanAddTransformKeysForSelectedObjects() const
+{
+	// WASD hotkeys to fly the viewport can conflict with hotkeys for setting keyframes (ie. s). 
+	// If the viewport is moving, disregard setting keyframes.
+	for (FLevelEditorViewportClient* LevelVC : GEditor->GetLevelViewportClients())
+	{
+		if (LevelVC && LevelVC->IsMovingCamera())
+		{
+			return false;
+		}
+	}
+
+	if (USelection* SelectedActors = GEditor->GetSelectedActors())
+	{
+		if(SelectedActors->Num() > 0)
+		{
+			return false;
+		}
+	}
+
+	const TSharedPtr<ISequencer> PinnedSequencer = GetSequencer();
+
+	if (!PinnedSequencer.IsValid())
+	{
+		return false;
+	}
+
+	TArray<UMovieSceneSection*> OutSections;
+	PinnedSequencer->GetSelectedSections(OutSections);
+
+	for (UMovieSceneSection* Section : OutSections)
+	{
+		if (const UMovieSceneSubSection* SubSection = Cast<UMovieSceneSubSection>(Section))
+		{
+			UMovieSceneTrack* OuterTrack = Cast<UMovieSceneTrack>(SubSection->GetOuter());
+			if (OuterTrack && OuterTrack->GetSectionToKey())
+			{
+				if (OuterTrack->GetSectionToKey() == SubSection && SubSection->IsTransformOriginEditable())
+				{
+					return true;
+				}
+			}
+			else if (SubSection->IsTransformOriginEditable())
+			{
+				return true;
+			}
+		}
+	}
+
+	TArray<UMovieSceneTrack*> OutTracks;
+	PinnedSequencer->GetSelectedTracks(OutTracks);
+
+	for (UMovieSceneTrack* Track : OutTracks)
+	{
+		if (const UMovieSceneSubTrack* SubTrack = Cast<UMovieSceneSubTrack>(Track))
+		{
+			if (const UMovieSceneSection* SectionToKey = SubTrack->GetSectionToKey())
+			{
+				const UMovieSceneSubSection* SubSectionToKey = Cast<UMovieSceneSubSection>(SectionToKey);
+				if (SubSectionToKey && SubSectionToKey->IsTransformOriginEditable())
+				{
+					return true;
+				}
+			}
+			else if (SubTrack->GetAllSections().Num())
+			{
+				for (UMovieSceneSection* Section : SubTrack->FindAllSections(PinnedSequencer.Get()->GetLocalTime().Time.FrameNumber))
+				{
+					UMovieSceneSubSection* SubSection = Cast<UMovieSceneSubSection>(Section);
+					if (SubSection && SubSection->IsTransformOriginEditable())
+					{
+						return true;
+					}
+				}
+			}
+		}
+	}
+
+	return false;
+}
+
+void FSubTrackEditor::GetSectionsToKey(TArray<UMovieSceneSubSection*>& OutSectionsToKey) const
+{
+	const TSharedPtr<ISequencer> PinnedSequencer = GetSequencer();
+
+	if (!PinnedSequencer.IsValid())
+	{
+		return;
+	}
+
+	TArray<UMovieSceneSection*> OutSections;
+	PinnedSequencer->GetSelectedSections(OutSections);
+
+	for (UMovieSceneSection* Section : OutSections)
+	{
+		if (UMovieSceneSubSection* SubSection = Cast<UMovieSceneSubSection>(Section))
+		{
+			const UMovieSceneTrack* OwningTrack = Cast<UMovieSceneTrack>(SubSection->GetOuter());
+
+			if(OwningTrack && OwningTrack->GetSectionToKey())
+			{
+				if (SubSection == OwningTrack->GetSectionToKey() && SubSection->IsTransformOriginEditable())
+				{
+					OutSectionsToKey.Add(SubSection);
+				}
+			}
+			else if (SubSection->IsTransformOriginEditable())
+			{
+				OutSectionsToKey.Add(SubSection);
+			}
+		}
+	}
+	
+	TArray<UMovieSceneTrack*> OutTracks;
+	PinnedSequencer->GetSelectedTracks(OutTracks);
+
+	for (const UMovieSceneTrack* Track : OutTracks)
+	{
+		if (const UMovieSceneSubTrack* SubTrack = Cast<UMovieSceneSubTrack>(Track))
+		{
+			UMovieSceneSubSection* SubSectionToKey = Cast<UMovieSceneSubSection>(SubTrack->GetSectionToKey());
+			if (SubSectionToKey && SubSectionToKey->IsTransformOriginEditable())
+			{
+				OutSectionsToKey.AddUnique(SubSectionToKey);
+			}
+			else if (SubTrack->GetAllSections().Num())
+			{
+				for (UMovieSceneSection* Section : SubTrack->FindAllSections(PinnedSequencer.Get()->GetLocalTime().Time.FrameNumber))
+				{
+					UMovieSceneSubSection* SubSection = Cast<UMovieSceneSubSection>(Section);
+					if (SubSection && SubSection->IsTransformOriginEditable())
+					{
+						OutSectionsToKey.AddUnique(SubSection);
+					}
+				}
+			}
+		}
+	}
+}
+
+void FSubTrackEditor::OnAddTransformKeysForSelectedObjects(EMovieSceneTransformChannel Channel)
+{
+	// WASD hotkeys to fly the viewport can conflict with hotkeys for setting keyframes (ie. s). 
+	// If the viewport is moving, disregard setting keyframes.
+	for (FLevelEditorViewportClient* LevelVC : GEditor->GetLevelViewportClients())
+	{
+		if (LevelVC && LevelVC->IsMovingCamera())
+		{
+			return;
+		}
+	}
+
+	if (USelection* SelectedActors = GEditor->GetSelectedActors())
+	{
+		if(SelectedActors->Num() > 0)
+		{
+			return;
+		}
+	}
+	
+	const TSharedPtr<ISequencer> PinnedSequencer = GetSequencer();
+
+	// Create a transaction record because we are about to add keys
+	const bool bShouldActuallyTransact = !GIsTransacting;		// Don't transact if we're recording in a PIE world.  That type of keyframe capture cannot be undone.
+	FScopedTransaction AutoKeyTransaction( NSLOCTEXT("AnimatablePropertyTool", "PropertyChanged", "Animatable Property Changed"), bShouldActuallyTransact);
+
+	TArray<UMovieSceneSubSection*> SectionsToKey;
+
+	GetSectionsToKey(SectionsToKey);
+
+	for(UMovieSceneSubSection* SubSection : SectionsToKey)
+	{
+		FTransform CurrentTransform = GetTransformOriginDataForSubSection(SubSection);
+		
+		FGeneratedTrackKeys GeneratedKeys;
+		
+		GetOriginKeys(CurrentTransform.GetLocation(), CurrentTransform.Rotator(), Channel, SubSection, GeneratedKeys);
+		
+		FKeyPropertyResult KeyResults = AddKeysToSection(SubSection, GetTimeForKey(), GeneratedKeys, ESequencerKeyMode::ManualKeyForced);
+		
+		SubSection->Modify(true);
+	}
+}
+
 void FSubTrackEditor::InsertSection(UMovieSceneTrack* Track)
 {
 	FFrameTime NewSectionStartTime = GetSequencer()->GetLocalTime().Time;
@@ -1140,62 +1326,11 @@ void FSubTrackEditor::UpdateOrigin(FVector InPosition, FRotator InRotation)
 	FScopedTransaction AutoKeyTransaction( NSLOCTEXT("AnimatablePropertyTool", "PropertyChanged", "Animatable Property Changed"), bShouldActuallyTransact);
 
 	TArray<UMovieSceneSubSection*> SectionsToKey;
-	
-	TArray<UMovieSceneSection*> SelectedSections;
-	GetSequencer()->GetSelectedSections(SelectedSections);
 
-	// Gather selected sections first, since they may be part of a track as well, and they may need to be removed if section to key is set.
-	for(UMovieSceneSection* Section : SelectedSections)
-	{
-		if(UMovieSceneSubSection* SubSection = Cast<UMovieSceneSubSection>(Section))
-		{
-			if(SubSection->IsTransformOriginEditable())
-			{
-				SectionsToKey.AddUnique(SubSection);
-			}
-		}
-	}
-	
-	TArray<UMovieSceneTrack*> SelectedTracks;
-	GetSequencer()->GetSelectedTracks(SelectedTracks);
-	for (UMovieSceneTrack* Track : SelectedTracks)
-	{
-		if(UMovieSceneSubTrack* SubTrack = Cast<UMovieSceneSubTrack>(Track))
-		{
-			FGeneratedTrackKeys GeneratedTrackKeys;
-			UMovieSceneSubSection* SectionToKey = Cast<UMovieSceneSubSection>(SubTrack->GetSectionToKey());
-			if(SectionToKey)
-			{
-				// Remove other sections from this track, since section to key takes precedence.
-				for(UMovieSceneSection* Section : SubTrack->GetAllSections())
-				{
-					if(UMovieSceneSubSection* SubSection = Cast<UMovieSceneSubSection>(Section))
-					{
-						SectionsToKey.Remove(SubSection);
-					}
-				}
-				SectionsToKey.AddUnique(SectionToKey);
-			}
-			else
-			{
-				for(UMovieSceneSection* Section : SubTrack->GetAllSections())
-				{
-					if(UMovieSceneSubSection* SubSection = Cast<UMovieSceneSubSection>(Section))
-					{
-						if(SubSection->IsTransformOriginEditable())
-						{
-							SectionsToKey.AddUnique(SubSection);
-						}
-					}
-				}
-			}
-		}
-	}
+	GetSectionsToKey(SectionsToKey);
 	
 	for(UMovieSceneSubSection* SubSection : SectionsToKey)
 	{
-		EMovieSceneTransformChannel ExistingChannels = SubSection->GetMask().GetChannels();
-
 		FTransform PreviousTransform = GetTransformOriginDataForSubSection(SubSection);
 		
 		// Default to dirtying the section, since we may be adding keys 
@@ -1203,7 +1338,7 @@ void FSubTrackEditor::UpdateOrigin(FVector InPosition, FRotator InRotation)
 		
 		FGeneratedTrackKeys GeneratedKeys;
 		
-		GetOriginKeys(PreviousTransform.GetLocation() + InPosition, (PreviousTransform * FTransform(InRotation)).Rotator(), SubSection, GeneratedKeys);
+		GetOriginKeys(PreviousTransform.GetLocation() + InPosition, (PreviousTransform * FTransform(InRotation)).Rotator(), EMovieSceneTransformChannel::Translation | EMovieSceneTransformChannel::Rotation, SubSection, GeneratedKeys);
 		
 		FKeyPropertyResult KeyResults = AddKeysToSection(SubSection, GetTimeForKey(), GeneratedKeys, ESequencerKeyMode::AutoKey);
 
diff --git a/Engine/Source/Editor/MovieSceneTools/Public/TrackEditors/SubTrackEditor.h b/Engine/Source/Editor/MovieSceneTools/Public/TrackEditors/SubTrackEditor.h
--- a/Engine/Source/Editor/MovieSceneTools/Public/TrackEditors/SubTrackEditor.h
+++ b/Engine/Source/Editor/MovieSceneTools/Public/TrackEditors/SubTrackEditor.h
@@ -65,6 +65,9 @@ public:
 	virtual void OnRelease() override;
 	virtual void Resize(float NewSize, UMovieSceneTrack* InTrack) override;
 	virtual bool GetDefaultExpansionState(UMovieSceneTrack* InTrack) const override;
+	virtual bool HasTransformKeyBindings() const override { return true; }
+	virtual bool CanAddTransformKeysForSelectedObjects() const override;
+	virtual void OnAddTransformKeysForSelectedObjects(EMovieSceneTransformChannel Channel) override;
 
 public:
 	
@@ -181,7 +184,10 @@ private:
 	void ProcessKeyOperationInternal(TArrayView<const UE::Sequencer::FKeySectionOperation> SectionsToKey, ISequencer& Sequencer, FFrameNumber KeyTime, TArray<UE::Sequencer::FAddKeyResult>* OutResults = nullptr);
 
 	/** Helper for creating new keys.*/
-	void GetOriginKeys(const FVector& CurrentPosition, const FRotator& CurrentRotation, UMovieSceneSection* Section, FGeneratedTrackKeys& OutGeneratedKeys);
+	void GetOriginKeys(const FVector& CurrentPosition, const FRotator& CurrentRotation, EMovieSceneTransformChannel ChannelsToKey, UMovieSceneSection* Section, FGeneratedTrackKeys& OutGeneratedKeys);
+
+	/** Helper for getting the sections that should be keyed. **/
+	void GetSectionsToKey(TArray<UMovieSceneSubSection*>& OutSectionsToKey) const;
 
 	/** Sections this editor has added preview data to for keyframing. Used to revert data when other edits are made, before adding a key */
 	TArray<UMovieSceneSubSection*> SectionsWithPreviews;
