Constraints: unregister constraints when closing sequencer

#rnx

diff --git a/Engine/Plugins/Animation/ControlRig/Source/ControlRigEditor/Private/Sequencer/ControlConstraintChannelInterface.cpp b/Engine/Plugins/Animation/ControlRig/Source/ControlRigEditor/Private/Sequencer/ControlConstraintChannelInterface.cpp
--- a/Engine/Plugins/Animation/ControlRig/Source/ControlRigEditor/Private/Sequencer/ControlConstraintChannelInterface.cpp
+++ b/Engine/Plugins/Animation/ControlRig/Source/ControlRigEditor/Private/Sequencer/ControlConstraintChannelInterface.cpp
@@ -300,4 +300,26 @@ UMovieSceneControlRigParameterSection* FControlConstraintChannelInterface::GetCo
 	return Cast<UMovieSceneControlRigParameterSection>(Section);
 }
 
+void FControlConstraintChannelInterface::UnregisterTrack(UMovieSceneTrack* InTrack, UWorld* InWorld)
+{
+	ITransformConstraintChannelInterface::UnregisterTrack(InTrack, InWorld);
+	
+	UMovieSceneControlRigParameterTrack* ControlRigTrack = Cast<UMovieSceneControlRigParameterTrack>(InTrack);
+	if (!ControlRigTrack)
+	{
+		return;
+	}
+
+	const TArray<UMovieSceneSection*>& AllSections = ControlRigTrack->GetAllSections();
+	UMovieSceneControlRigParameterSection* Section =
+		AllSections.IsEmpty() ? nullptr : Cast<UMovieSceneControlRigParameterSection>(AllSections[0]);
+	
+	if (!Section)
+	{
+		return;
+	}
+
+	UnregisterConstraints(Section, InWorld);
+}
+
 #undef LOCTEXT_NAMESPACE
diff --git a/Engine/Plugins/Animation/ControlRig/Source/ControlRigEditor/Public/ControlConstraintChannelInterface.h b/Engine/Plugins/Animation/ControlRig/Source/ControlRigEditor/Public/ControlConstraintChannelInterface.h
--- a/Engine/Plugins/Animation/ControlRig/Source/ControlRigEditor/Public/ControlConstraintChannelInterface.h
+++ b/Engine/Plugins/Animation/ControlRig/Source/ControlRigEditor/Public/ControlConstraintChannelInterface.h
@@ -23,7 +23,7 @@ class UTransformableControlHandle;
 
 struct FControlConstraintChannelInterface : ITransformConstraintChannelInterface
 {
-	virtual ~FControlConstraintChannelInterface() = default;
+	virtual ~FControlConstraintChannelInterface() override = default;
 
 	/** Get the UMovieSceneControlRigParameterSection from the ControlRig. */
 	
@@ -46,6 +46,9 @@ struct FControlConstraintChannelInterface : ITransformConstraintChannelInterface
 		const TArray<FTransform>& InLocalTransforms,
 		const EMovieSceneTransformChannel& InChannels) override;
 
+protected:
+		virtual void UnregisterTrack(UMovieSceneTrack* InTrack, UWorld* InWorld) override;
+	
 private:
 	static UMovieSceneControlRigParameterSection* GetControlSection(const UTransformableControlHandle* InHandle, const TSharedPtr<ISequencer>& InSequencer, const bool bIsConstraint);
 };
diff --git a/Engine/Source/Editor/MovieSceneTools/Private/Constraints/ComponentConstraintChannelInterface.cpp b/Engine/Source/Editor/MovieSceneTools/Private/Constraints/ComponentConstraintChannelInterface.cpp
--- a/Engine/Source/Editor/MovieSceneTools/Private/Constraints/ComponentConstraintChannelInterface.cpp
+++ b/Engine/Source/Editor/MovieSceneTools/Private/Constraints/ComponentConstraintChannelInterface.cpp
@@ -470,3 +470,24 @@ void FComponentConstraintChannelInterface::RecomposeTransforms(
 		}
 	}
 }
+
+void FComponentConstraintChannelInterface::UnregisterTrack(UMovieSceneTrack* InTrack, UWorld* InWorld)
+{
+	ITransformConstraintChannelInterface::UnregisterTrack(InTrack, InWorld);
+	
+	UMovieScene3DTransformTrack* TransformTrack = Cast<UMovieScene3DTransformTrack>(InTrack);
+	if (!TransformTrack)
+	{
+		return;
+	}
+
+	const TArray<UMovieSceneSection*>& AllSections = TransformTrack->GetAllSections();
+	UMovieScene3DTransformSection* Section =
+		AllSections.IsEmpty() ? nullptr : Cast<UMovieScene3DTransformSection>(AllSections[0]);
+	if (!Section)
+	{
+		return;
+	}
+	
+	UnregisterConstraints(Section, InWorld);
+}
\ No newline at end of file
diff --git a/Engine/Source/Editor/MovieSceneTools/Private/Constraints/ComponentConstraintChannelInterface.h b/Engine/Source/Editor/MovieSceneTools/Private/Constraints/ComponentConstraintChannelInterface.h
--- a/Engine/Source/Editor/MovieSceneTools/Private/Constraints/ComponentConstraintChannelInterface.h
+++ b/Engine/Source/Editor/MovieSceneTools/Private/Constraints/ComponentConstraintChannelInterface.h
@@ -36,6 +36,9 @@ struct FComponentConstraintChannelInterface : ITransformConstraintChannelInterfa
 		const TArray<FTransform>& InLocalTransforms,
 		const EMovieSceneTransformChannel& InChannels) override;
 
+protected:
+	virtual void UnregisterTrack(UMovieSceneTrack* InTrack, UWorld* InWorld) override;
+	
 private:
 		static UMovieScene3DTransformSection* GetComponentSection(const UTransformableComponentHandle* InHandle, const TSharedPtr<ISequencer>& InSequencer, const bool bInConstraint);
 	
diff --git a/Engine/Source/Editor/MovieSceneTools/Private/Constraints/TransformConstraintChannelInterface.cpp b/Engine/Source/Editor/MovieSceneTools/Private/Constraints/TransformConstraintChannelInterface.cpp
--- a/Engine/Source/Editor/MovieSceneTools/Private/Constraints/TransformConstraintChannelInterface.cpp
+++ b/Engine/Source/Editor/MovieSceneTools/Private/Constraints/TransformConstraintChannelInterface.cpp
@@ -1,13 +1,59 @@
 ﻿// Copyright Epic Games, Inc. All Rights Reserved.
 
 #include "Constraints/TransformConstraintChannelInterface.h"
-#include "Sections/MovieSceneConstrainedSection.h"
+
 #include "ConstraintChannel.h"
+#include "ISequencerModule.h"
+#include "Modules/ModuleManager.h"
+#include "MovieScene.h"
+#include "Sections/MovieSceneConstrainedSection.h"
+#include "Sections/MovieSceneSubSection.h"
+#include "Tracks/MovieSceneSubTrack.h"
+
+ITransformConstraintChannelInterface::ITransformConstraintChannelInterface()
+{
+	Initialize();
+}
+
+ITransformConstraintChannelInterface::~ITransformConstraintChannelInterface()
+{
+	Shutdown();
+}
+
+void ITransformConstraintChannelInterface::Initialize()
+{
+	ISequencerModule& SequencerModule = FModuleManager::Get().LoadModuleChecked<ISequencerModule>("Sequencer");
+	SequencerCreatedHandle = SequencerModule.RegisterOnSequencerCreated(
+		FOnSequencerCreated::FDelegate::CreateRaw(this, &ITransformConstraintChannelInterface::OnSequencerCreated));
+}
+
+void ITransformConstraintChannelInterface::Shutdown()
+{
+	for (const TWeakPtr<ISequencer>& Sequencer : Sequencers)
+	{
+		if (Sequencer.IsValid())
+		{
+			Sequencer.Pin()->OnCloseEvent().RemoveAll(this);
+		}
+	}
+	Sequencers.Reset();
+	
+	if (ISequencerModule* SequencerModulePtr = FModuleManager::Get().GetModulePtr<ISequencerModule>("Sequencer"))
+	{
+		if (SequencerCreatedHandle.IsValid())
+		{
+			SequencerModulePtr->UnregisterOnSequencerCreated(SequencerCreatedHandle);
+		}
+	}
+	SequencerCreatedHandle.Reset();
+}
+
+void ITransformConstraintChannelInterface::OnSequencerCreated(TSharedRef<ISequencer> InSequencer)
+{
+	Sequencers.Add(TWeakPtr<ISequencer>(InSequencer));
+	InSequencer->OnCloseEvent().AddRaw(this, &ITransformConstraintChannelInterface::OnSequencerClosed);
+}
 
-// IMovieSceneConstrainedSection* ITransformConstraintChannelInterface::GetHandleConstraintSection(const UTransformableHandle* InHandle, const TSharedPtr<ISequencer>& InSequencer)
-// {
-// 	return Cast<IMovieSceneConstrainedSection>(GetHandleSection(InHandle, InSequencer));
-// }
 //we changed this so that it will always return true
 bool ITransformConstraintChannelInterface::CanAddKey(FMovieSceneConstraintChannel& ActiveChannel, const FFrameNumber& InTime, bool& ActiveValue)
 {
@@ -58,6 +104,95 @@ bool ITransformConstraintChannelInterface::CanAddKey(FMovieSceneConstraintChanne
 	return true;
 }
 
+void ITransformConstraintChannelInterface::OnSequencerClosed(TSharedRef<ISequencer> InSequencer)
+{
+	UWorld* World = [InSequencer]() -> UWorld*
+	{
+		if (const TSharedPtr<const UE::MovieScene::FSharedPlaybackState> SharedPlaybackState = InSequencer->FindSharedPlaybackState())
+		{
+			if (UObject* PlaybackContext = SharedPlaybackState->GetPlaybackContext())
+			{
+				return PlaybackContext->GetWorld();	
+			}
+		}
+		return nullptr;
+	}();
+
+	if (World)
+	{
+		UMovieSceneSequence* RootSceneSequence = InSequencer->GetRootMovieSceneSequence();
+		if (const UMovieScene* RootMovieScene = RootSceneSequence ? RootSceneSequence->GetMovieScene() : nullptr)
+		{
+			UnregisterMovieScene(RootMovieScene, World);
+		}
+	}
+	
+	InSequencer->OnCloseEvent().RemoveAll(this);
+	Sequencers.Remove(InSequencer);
+}
+
+void ITransformConstraintChannelInterface::UnregisterMovieScene(const UMovieScene* InMovieScene, UWorld* InWorld)
+{
+	if (!InMovieScene)
+	{
+		return;
+	}
+	
+	for (UMovieSceneTrack* Track: InMovieScene->GetTracks())
+	{
+		UnregisterTrack(Track, InWorld);
+	}
+	
+	for (const FMovieSceneBinding& Binding: InMovieScene->GetBindings())
+	{
+		for (UMovieSceneTrack* Track: Binding.GetTracks())
+		{
+			UnregisterTrack(Track, InWorld);	
+		}
+	}
+}
+
+void ITransformConstraintChannelInterface::UnregisterTrack(UMovieSceneTrack* InTrack, UWorld* InWorld)
+{
+	if (UMovieSceneSubTrack* SubTrack = Cast<UMovieSceneSubTrack>(InTrack))
+	{
+		for (const UMovieSceneSection* Section: SubTrack->GetAllSections())
+		{
+			if (const UMovieSceneSubSection* SubSection = Cast<UMovieSceneSubSection>(Section))
+			{
+				const UMovieSceneSequence* SubSequence = SubSection->GetSequence();
+				if (const UMovieScene* SubMovieScene = SubSequence ? SubSequence->GetMovieScene() : nullptr)
+				{
+					UnregisterMovieScene(SubMovieScene, InWorld);
+				}
+			}
+		}
+	}
+}
+
+void ITransformConstraintChannelInterface::UnregisterConstraints(IMovieSceneConstrainedSection* InSection, UWorld* InWorld)
+{
+	if (!InSection || !InWorld)
+	{
+		return;
+	}
+
+	const TArray<FConstraintAndActiveChannel> Constraints = InSection->GetConstraintsChannels();
+	if (Constraints.IsEmpty())
+	{
+		return;
+	}
+
+	FConstraintsManagerController& Controller = FConstraintsManagerController::Get(InWorld);
+	for (const FConstraintAndActiveChannel& ConstraintAndActiveChannel : Constraints)
+	{
+		if (UTickableConstraint* Constraint = ConstraintAndActiveChannel.GetConstraint())
+		{
+			Controller.UnregisterConstraint(Constraint);
+		}
+	}		
+}
+
 FConstraintChannelInterfaceRegistry& FConstraintChannelInterfaceRegistry::Get()
 {
 	static FConstraintChannelInterfaceRegistry Singleton;
diff --git a/Engine/Source/Editor/MovieSceneTools/Public/Constraints/TransformConstraintChannelInterface.h b/Engine/Source/Editor/MovieSceneTools/Public/Constraints/TransformConstraintChannelInterface.h
--- a/Engine/Source/Editor/MovieSceneTools/Public/Constraints/TransformConstraintChannelInterface.h
+++ b/Engine/Source/Editor/MovieSceneTools/Public/Constraints/TransformConstraintChannelInterface.h
@@ -2,10 +2,15 @@
 
 #pragma once
 
+#include "Containers/Map.h"
+#include "Delegates/IDelegateInstance.h"
 #include "Templates/SharedPointer.h"
 #include "Templates/UniquePtr.h"
-#include "Containers/Map.h"
 
+class UMovieScene;
+class UMovieSceneSequence;
+class UMovieSceneSubTrack;
+class UMovieSceneTrack;
 struct FMovieSceneConstraintChannel;
 class IMovieSceneConstrainedSection;
 class UWorld;
@@ -22,12 +27,11 @@ enum class EMovieSceneTransformChannel : uint32;
 
 struct ITransformConstraintChannelInterface
 {
-	virtual ~ITransformConstraintChannelInterface() = default;
+	virtual MOVIESCENETOOLS_API ~ITransformConstraintChannelInterface();
 
 	/** Get the section where the channels live from the object that the handle wraps. */
 	virtual UMovieSceneSection* GetHandleSection(const UTransformableHandle* InHandle, const TSharedPtr<ISequencer>& InSequencer) = 0;
 	virtual UMovieSceneSection* GetHandleConstraintSection(const UTransformableHandle* InHandle, const TSharedPtr<ISequencer>& InSequencer) = 0;
-	// MOVIESCENETOOLS_API IMovieSceneConstrainedSection* GetHandleConstraintSection(const UTransformableHandle* InHandle, const TSharedPtr<ISequencer>& InSequencer);
 
 	/** Get the world from the object that the handle wraps. */
 	virtual UWorld* GetHandleWorld(UTransformableHandle* InHandle) = 0;
@@ -37,7 +41,7 @@ struct ITransformConstraintChannelInterface
 		UTickableTransformConstraint* InConstraint, const TOptional<bool>& InOptActive,
 		const FFrameNumber& InTime, const TSharedPtr<ISequencer>& InSequencer) = 0;
 
-	/** Add keys on the transform channels of the the object that the handle wraps. */
+	/** Add keys on the transform channels of the object that the handle wraps. */
 	virtual void AddHandleTransformKeys(
 		const TSharedPtr<ISequencer>& InSequencer,
 		const UTransformableHandle* InHandle,
@@ -46,7 +50,23 @@ struct ITransformConstraintChannelInterface
 		const EMovieSceneTransformChannel& InChannels) = 0;
 
 protected:
+
+	MOVIESCENETOOLS_API ITransformConstraintChannelInterface();
+
 	static MOVIESCENETOOLS_API bool CanAddKey(FMovieSceneConstraintChannel& ActiveChannel, const FFrameNumber& InTime, bool& ActiveValue);
+
+	void Initialize();
+	void Shutdown();
+	
+	void OnSequencerCreated(TSharedRef<ISequencer> InSequencer);
+	void OnSequencerClosed(TSharedRef<ISequencer> InSequencer);
+
+	virtual void MOVIESCENETOOLS_API UnregisterTrack(UMovieSceneTrack* InTrack, UWorld* InWorld);
+	static MOVIESCENETOOLS_API void UnregisterConstraints(IMovieSceneConstrainedSection* InSection, UWorld* InWorld);
+	void UnregisterMovieScene(const UMovieScene* InMovieScene, UWorld* InWorld);
+	
+	FDelegateHandle SequencerCreatedHandle;
+	TArray<TWeakPtr<ISequencer>> Sequencers;
 };
 
 /**
diff --git a/Engine/Source/Runtime/Experimental/Animation/Constraints/Private/ConstraintsManager.cpp b/Engine/Source/Runtime/Experimental/Animation/Constraints/Private/ConstraintsManager.cpp
--- a/Engine/Source/Runtime/Experimental/Animation/Constraints/Private/ConstraintsManager.cpp
+++ b/Engine/Source/Runtime/Experimental/Animation/Constraints/Private/ConstraintsManager.cpp
@@ -556,6 +556,28 @@ bool FConstraintsManagerController::RemoveConstraint(UTickableConstraint* InCons
 	return RemoveConstraint(Index, bDoNotCompensate);
 }
 
+bool FConstraintsManagerController::UnregisterConstraint(UTickableConstraint* InConstraint)
+{
+	if (!InConstraint)
+	{
+		return false;
+	}
+	
+	UConstraintSubsystem* Subsystem = UConstraintSubsystem::Get();
+	if (!Subsystem)
+	{
+		return false;
+	}
+
+	static constexpr bool bDoNotCompensate = true;
+	if (Subsystem->HasConstraint(World, InConstraint))
+	{
+		Subsystem->RemoveConstraint(World, InConstraint, bDoNotCompensate);
+		return true;
+	}
+	return false;
+}
+
 bool FConstraintsManagerController::RemoveConstraint(const int32 InConstraintIndex,bool bDoNotCompensate) 
 {
 	UConstraintSubsystem* Subsystem = UConstraintSubsystem::Get();
diff --git a/Engine/Source/Runtime/Experimental/Animation/Constraints/Public/ConstraintsManager.h b/Engine/Source/Runtime/Experimental/Animation/Constraints/Public/ConstraintsManager.h
--- a/Engine/Source/Runtime/Experimental/Animation/Constraints/Public/ConstraintsManager.h
+++ b/Engine/Source/Runtime/Experimental/Animation/Constraints/Public/ConstraintsManager.h
@@ -261,6 +261,9 @@ public:
 
 	/** Remove constraint by ptr*/
 	CONSTRAINTS_API bool RemoveConstraint(UTickableConstraint* InConstraint, bool bDoNotCompensate = false);
+	
+	/** Disable a constraint and remove it from the subsystem. */
+	CONSTRAINTS_API bool UnregisterConstraint(UTickableConstraint* InConstraint);
 
 	/** Remove All Constraints*/
 	CONSTRAINTS_API bool RemoveAllConstraints(bool bDoNotCompensate = false);
