MovieScene: Prevent crash in condition customization header.

diff --git a/Engine/Source/Editor/MovieSceneTools/Private/Conditions/MovieSceneConditionCustomization.cpp b/Engine/Source/Editor/MovieSceneTools/Private/Conditions/MovieSceneConditionCustomization.cpp
--- a/Engine/Source/Editor/MovieSceneTools/Private/Conditions/MovieSceneConditionCustomization.cpp
+++ b/Engine/Source/Editor/MovieSceneTools/Private/Conditions/MovieSceneConditionCustomization.cpp
@@ -95,10 +95,14 @@ void FMovieSceneConditionCustomization::CustomizeHeader(TSharedRef<IPropertyHand
 		Track = GetCommonTrack();
 	}
 
+	TStrongObjectPtr<UMovieSceneSequence> SequencePtr = Sequence.Pin();
+	TStrongObjectPtr<UMovieSceneTrack> TrackPtr = Track.Pin();
+	TSharedPtr<ISequencer> SequencerPtr = Sequencer.Pin();
+
 	// If conditions not allowed, hide condition property functionality
-	if (!Sequence.IsValid() 
-	|| !Sequence->GetMovieScene()->IsConditionClassAllowed(UMovieSceneCondition::StaticClass()) 
-	|| (Track.IsValid() && Sequencer.IsValid() && !Sequencer.Pin()->TrackSupportsConditions(Track.Get())))
+	if (!SequencePtr.IsValid()
+	|| !SequencePtr->GetMovieScene()->IsConditionClassAllowed(UMovieSceneCondition::StaticClass())
+	|| (TrackPtr.IsValid() && SequencerPtr.IsValid() && !SequencerPtr->TrackSupportsConditions(TrackPtr.Get())))
 	{
 		ConditionContainerPropertyHandle->MarkHiddenByCustomization();
 		return;
