Fix red MegaLights warning text on level sequence thumbnails created by [REDACTED_PERSON].

diff --git a/Engine/Source/Editor/MovieSceneTools/Private/TrackEditorThumbnail/TrackThumbnailUtils.cpp b/Engine/Source/Editor/MovieSceneTools/Private/TrackEditorThumbnail/TrackThumbnailUtils.cpp
--- a/Engine/Source/Editor/MovieSceneTools/Private/TrackEditorThumbnail/TrackThumbnailUtils.cpp
+++ b/Engine/Source/Editor/MovieSceneTools/Private/TrackEditorThumbnail/TrackThumbnailUtils.cpp
@@ -41,6 +41,10 @@ namespace UE::MoveSceneTools
 		FSceneViewFamilyContext ViewFamily( FSceneViewFamily::ConstructionValues(&ThumbnailRenderTarget, &Scene, FEngineShowFlags(ESFIM_Game))
 		.SetTime(FGameTime::GetTimeSinceAppStart())
 		.SetResolveScene(true));
+		
+		// When GRayTracingMode is disabled(default: disabled) and MegaLights is enabled(default: enabled), we'll get a red warning text on the thumbnail.
+		// Even if the project was set up to use ray tracing, we don't need thumbnails to use raytracing anyways (performance), so just disable. 
+		ViewFamily.EngineShowFlags.SetMegaLights(false);
 
 		FSceneViewStateInterface* ViewStateInterface = nullptr;
 
