Sequencer: Follow up fix to 39154700. Return true when the global class filter isn't valid

diff --git a/Engine/Source/Editor/MovieSceneTools/Private/TrackEditors/SkeletalAnimationTrackEditor.cpp b/Engine/Source/Editor/MovieSceneTools/Private/TrackEditors/SkeletalAnimationTrackEditor.cpp
--- a/Engine/Source/Editor/MovieSceneTools/Private/TrackEditors/SkeletalAnimationTrackEditor.cpp
+++ b/Engine/Source/Editor/MovieSceneTools/Private/TrackEditors/SkeletalAnimationTrackEditor.cpp
@@ -1308,7 +1308,7 @@ bool FSkeletalAnimationTrackEditor::CanCreatePoseAsset(FGuid InObjectBinding) co
 		return GlobalClassFilter->IsClassAllowed(ClassViewerOptions, UPoseAsset::StaticClass(), ClassFilterFuncs);
 	}
 
-	return false;
+	return true;
 }
 
 FSkeletalAnimationTrackEditor::FSkeletalAnimationTrackEditor( TSharedRef<ISequencer> InSequencer )
