Sequencer: Use Shot Track and Subsequence Track for GetSubTrackName and don't split it up into "Shot" and "Subsequence" and "Track" because localization prefers a single string

diff --git a/Engine/Source/Editor/MovieSceneTools/Private/TrackEditors/CinematicShotTrackEditor.cpp b/Engine/Source/Editor/MovieSceneTools/Private/TrackEditors/CinematicShotTrackEditor.cpp
--- a/Engine/Source/Editor/MovieSceneTools/Private/TrackEditors/CinematicShotTrackEditor.cpp
+++ b/Engine/Source/Editor/MovieSceneTools/Private/TrackEditors/CinematicShotTrackEditor.cpp
@@ -215,7 +215,7 @@ void FCinematicShotTrackEditor::RenderShots(const TArray<UMovieSceneCinematicSho
 
 FText FCinematicShotTrackEditor::GetSubTrackName() const
 {
-	return LOCTEXT("ShotTrackName", "Shot");
+	return LOCTEXT("ShotTrackName", "Shot Track");
 }
 
 FText FCinematicShotTrackEditor::GetSubTrackToolTip() const
diff --git a/Engine/Source/Editor/MovieSceneTools/Private/TrackEditors/SubTrackEditor.cpp b/Engine/Source/Editor/MovieSceneTools/Private/TrackEditors/SubTrackEditor.cpp
--- a/Engine/Source/Editor/MovieSceneTools/Private/TrackEditors/SubTrackEditor.cpp
+++ b/Engine/Source/Editor/MovieSceneTools/Private/TrackEditors/SubTrackEditor.cpp
@@ -294,7 +294,7 @@ FSubTrackEditor::FSubTrackEditor(TSharedRef<ISequencer> InSequencer)
 void FSubTrackEditor::BuildAddTrackMenu(FMenuBuilder& MenuBuilder)
 {
 	MenuBuilder.AddMenuEntry(
-		FText::Join(FText::FromString(" "), GetSubTrackName(), LOCTEXT("TrackText", "Track")),
+		GetSubTrackName(),
 		GetSubTrackToolTip(),
 		FSlateIcon(FAppStyle::GetAppStyleSetName(), GetSubTrackBrushName()),
 		FUIAction(
@@ -507,7 +507,7 @@ bool FSubTrackEditor::HandleAssetAdded(UObject* Asset, const FGuid& TargetObject
 
 	if (CanAddSubSequence(*Sequence))
 	{
-		const FScopedTransaction Transaction(FText::Join(FText::FromString(" "), LOCTEXT("AddText", "Add"), GetSubTrackName(), LOCTEXT("TrackText", "Track")));
+		const FScopedTransaction Transaction(FText::Join(FText::FromString(" "), LOCTEXT("AddText", "Add"), GetSubTrackName()));
 
 		int32 RowIndex = INDEX_NONE;
 		UMovieSceneTrack* Track = nullptr;
@@ -1265,7 +1265,7 @@ bool FSubTrackEditor::CanAddSubSequence(const UMovieSceneSequence& Sequence) con
 
 FText FSubTrackEditor::GetSubTrackName() const
 {
-	return LOCTEXT("SubTrackName", "Subsequence");
+	return LOCTEXT("SubTrackName", "Subsequence Track");
 }
 
 FText FSubTrackEditor::GetSubTrackToolTip() const
@@ -1509,7 +1509,7 @@ void FSubTrackEditor::HandleAddSubTrackMenuEntryExecute()
 		return;
 	}
 
-	const FScopedTransaction Transaction(FText::Join(FText::FromString(" "), LOCTEXT("AddText", "Add"), GetSubTrackName(), LOCTEXT("TrackText", "Track")));
+	const FScopedTransaction Transaction(FText::Join(FText::FromString(" "), LOCTEXT("AddText", "Add"), GetSubTrackName()));
 	FocusedMovieScene->Modify();
 
 	UMovieSceneSubTrack* NewTrack = FindOrCreateSubTrack(FocusedMovieScene, nullptr);
