Sequencer: Add attach socket tooltip for audio sections that are attached. Attach sections already have this.

diff --git a/Engine/Source/Editor/MovieSceneTools/Private/TrackEditors/AudioTrackEditor.cpp b/Engine/Source/Editor/MovieSceneTools/Private/TrackEditors/AudioTrackEditor.cpp
--- a/Engine/Source/Editor/MovieSceneTools/Private/TrackEditors/AudioTrackEditor.cpp
+++ b/Engine/Source/Editor/MovieSceneTools/Private/TrackEditors/AudioTrackEditor.cpp
@@ -691,10 +691,57 @@ FText FAudioSection::GetSectionTitle() const
 
 FText FAudioSection::GetSectionToolTip() const
 {
+	FText ToolTip;
+
+	TSharedPtr<ISequencer> SequencerPtr = Sequencer.Pin();
+	if (!SequencerPtr.IsValid())
+	{
+		return ToolTip;
+	}
+
 	UMovieSceneAudioSection* AudioSection = Cast<UMovieSceneAudioSection>(&Section);
-	const USoundWave* Sound = AudioSection ? DeriveSoundWave(AudioSection) : nullptr;
+	if (!AudioSection)
+	{
+		return ToolTip;
+	}
+
+	FText AttachToolTip;
 
-	if (AudioSection && Sound && AudioSection->HasStartFrame() && AudioSection->HasEndFrame())
+	const FMovieSceneActorReferenceData& AttachActorData = AudioSection->GetAttachActorData();
+
+	TMovieSceneChannelData<const FMovieSceneActorReferenceKey> ChannelData = AttachActorData.GetData();
+	TArrayView<const FMovieSceneActorReferenceKey> Values = ChannelData.GetValues().Num() != 0
+		? ChannelData.GetValues()
+		: MakeArrayView(&AttachActorData.GetDefault(), 1);
+
+	FMovieSceneActorReferenceKey CurrentValue;
+	AttachActorData.Evaluate(SequencerPtr->GetLocalTime().Time, CurrentValue);
+
+	for (int32 Index = 0; Index < Values.Num(); ++Index)
+	{
+		FMovieSceneObjectBindingID AttachBindingID = Values[Index].Object;
+		FName AttachSocketName = Values[Index].SocketName;
+
+		for (TWeakObjectPtr<> WeakObject : AttachBindingID.ResolveBoundObjects(SequencerPtr->GetFocusedTemplateID(), *SequencerPtr))
+		{
+			AActor* AttachActor = Cast<AActor>(WeakObject.Get());
+			if (AttachActor)
+			{
+				if (AttachSocketName.IsNone())
+				{
+					AttachToolTip = FText::FromString(AttachActor->GetActorLabel());
+				}
+				else
+				{
+					AttachToolTip = FText::Format(LOCTEXT("ToolTipAttachSocket", "{0} ({1})"), FText::FromString(AttachActor->GetActorLabel()), FText::FromName(AttachSocketName));
+				}
+			}
+		}
+	}
+
+	const USoundWave* Sound = DeriveSoundWave(AudioSection);
+
+	if (Sound && AudioSection->HasStartFrame() && AudioSection->HasEndFrame())
 	{
 		UMovieScene* MovieScene = AudioSection->GetTypedOuter<UMovieScene>();
 		FFrameRate TickResolution = MovieScene->GetTickResolution();
@@ -704,7 +751,7 @@ FText FAudioSection::GetSectionToolTip() const
 
 		if (AudioSection->GetLooping())
 		{
-			return FText::Format(LOCTEXT("ToolTipContentFormatLooping", "Start: {0}s\nDuration: {1}s\nLooping"),
+			ToolTip = FText::Format(LOCTEXT("ToolTipContentFormatLooping", "Start: {0}s\nDuration: {1}s\nLooping"),
 				AudioStartTime,
 				SectionDuration);
 		}
@@ -715,7 +762,7 @@ FText FAudioSection::GetSectionToolTip() const
 
 			if (Duration > 0.0f)
 			{
-				return FText::Format(LOCTEXT("ToolTipContentFormat", "{0}s - {1}s ({2} seconds)"),
+				ToolTip = FText::Format(LOCTEXT("ToolTipContentFormat", "{0}s - {1}s ({2} seconds)"),
 					AudioStartTime,
 					AudioStartTime + Duration,
 					Duration);
@@ -723,7 +770,9 @@ FText FAudioSection::GetSectionToolTip() const
 		}
 	}
 
-	return FText::GetEmpty();
+	ToolTip = FText::Join(FText::GetEmpty(), ToolTip, AttachToolTip);
+
+	return ToolTip;
 }
 
 float FAudioSection::GetSectionHeight(const UE::Sequencer::FViewDensityInfo& ViewDensity) const
