Constraints: fixed wrong compensation range calculation when adding duplicate

#rnx

diff --git a/Engine/Plugins/Animation/ControlRig/Source/ControlRigEditor/Private/Sequencer/ControlConstraintChannelInterface.cpp b/Engine/Plugins/Animation/ControlRig/Source/ControlRigEditor/Private/Sequencer/ControlConstraintChannelInterface.cpp
--- a/Engine/Plugins/Animation/ControlRig/Source/ControlRigEditor/Private/Sequencer/ControlConstraintChannelInterface.cpp
+++ b/Engine/Plugins/Animation/ControlRig/Source/ControlRigEditor/Private/Sequencer/ControlConstraintChannelInterface.cpp
@@ -212,6 +212,21 @@ bool FControlConstraintChannelInterface::SmartConstraintKey(
 					SetTangentsAtThisTime<FMovieSceneFloatChannel>(ChannelIndex, NumChannels, TransformSection, InTime, Tangents);
 				}
 			}
+
+			// cleanup next keys that have the same value as they are useless 
+			{
+				TArray<FFrameNumber> TimesRemoved;
+				CleanDuplicates(Channel->ActiveChannel, InTime, ActiveValueToBeSet, TimesRemoved);
+
+				// as compensation has been disabled earlier (FMovieSceneConstraintChannelHelper::bDoNotCompensate set to true),
+				// previous transform compensation keys must be explicitly removed as HandleConstraintKeyDeleted won't do it
+				for (const FFrameNumber& TimeToRemove: TimesRemoved)
+				{
+					const FFrameNumber TimeMinusOne(TimeToRemove - 1);
+					FMovieSceneConstraintChannelHelper::DeleteTransformKeys(Channels, TimeMinusOne);
+				}
+			}
+			
 			return true;
 		}
 	}
diff --git a/Engine/Source/Editor/MovieSceneTools/Private/Constraints/ComponentConstraintChannelInterface.cpp b/Engine/Source/Editor/MovieSceneTools/Private/Constraints/ComponentConstraintChannelInterface.cpp
--- a/Engine/Source/Editor/MovieSceneTools/Private/Constraints/ComponentConstraintChannelInterface.cpp
+++ b/Engine/Source/Editor/MovieSceneTools/Private/Constraints/ComponentConstraintChannelInterface.cpp
@@ -194,6 +194,21 @@ bool FComponentConstraintChannelInterface::SmartConstraintKey(
 					SetTangentsAtThisTime<FMovieSceneDoubleChannel>(ChannelIndex,NumChannels, TransformSection, InTime, Tangents);
 				}
 			}
+
+			// cleanup next keys that have the same value as they are useless 
+			{
+				TArray<FFrameNumber> TimesRemoved;
+				CleanDuplicates(Channel->ActiveChannel, InTime, ActiveValueToBeSet, TimesRemoved);
+
+				// as compensation has been disabled earlier (FMovieSceneConstraintChannelHelper::bDoNotCompensate set to true),
+				// previous transform compensation keys must be explicitly removed as HandleConstraintKeyDeleted won't do it
+				for (const FFrameNumber& TimeToRemove: TimesRemoved)
+				{
+					const FFrameNumber TimeMinusOne(TimeToRemove - 1);
+					FMovieSceneConstraintChannelHelper::DeleteTransformKeys(Channels, TimeMinusOne);
+				}
+			}
+			
 			// evaluate the constraint, this is needed so the global transform will be set up on the component 
 			//Todo do we need to evaluate all constraints?
 			InConstraint->SetActive(true); //will be false
diff --git a/Engine/Source/Editor/MovieSceneTools/Private/Constraints/TransformConstraintChannelInterface.cpp b/Engine/Source/Editor/MovieSceneTools/Private/Constraints/TransformConstraintChannelInterface.cpp
--- a/Engine/Source/Editor/MovieSceneTools/Private/Constraints/TransformConstraintChannelInterface.cpp
+++ b/Engine/Source/Editor/MovieSceneTools/Private/Constraints/TransformConstraintChannelInterface.cpp
@@ -55,53 +55,47 @@ void ITransformConstraintChannelInterface::OnSequencerCreated(TSharedRef<ISequen
 }
 
 //we changed this so that it will always return true
-bool ITransformConstraintChannelInterface::CanAddKey(FMovieSceneConstraintChannel& ActiveChannel, const FFrameNumber& InTime, bool& ActiveValue)
+bool ITransformConstraintChannelInterface::CanAddKey(const FMovieSceneConstraintChannel& InActiveChannel, const FFrameNumber& InTime, bool& ActiveValue)
 {
-	TMovieSceneChannelData<bool> ChannelData = ActiveChannel.GetData();
+	const TMovieSceneChannelData<const bool> ChannelData = InActiveChannel.GetData();
 	const TArrayView<const FFrameNumber> Times = ChannelData.GetTimes();
 	if (Times.IsEmpty())
 	{
 		ActiveValue = true;
 		return true;
 	}
-	ActiveChannel.Evaluate(InTime, ActiveValue);
+	
+	InActiveChannel.Evaluate(InTime, ActiveValue);
 	ActiveValue = !ActiveValue;
+	
+	return true;
+}
+
+void ITransformConstraintChannelInterface::CleanDuplicates(FMovieSceneConstraintChannel& InOutActiveChannel,
+	const FFrameNumber& InTime, const bool InValueToClean, TArray<FFrameNumber>& OutTimesRemoved)
+{
+	TMovieSceneChannelData<bool> ChannelData = InOutActiveChannel.GetData();
+	const TArrayView<const FFrameNumber> Times = ChannelData.GetTimes();
+	if (Times.IsEmpty())
+	{
+		return;
+	}
+	
 	const int32 Index = Times.Find(InTime);
 	const int32 NextTimeIndex = Algo::UpperBound(Times, InTime);
-	//have key at that time so need to delete it
-	auto DeleteNextKey = [&ActiveChannel, ActiveValue, &Times, Index,NextTimeIndex]()
-	{
-		TMovieSceneChannelData<bool> ChannelData = ActiveChannel.GetData();
-		//if has same value then delete the next
-		if (NextTimeIndex != Index && Times.IsValidIndex(NextTimeIndex))
-		{
-			bool NextValue = false;
-			ActiveChannel.Evaluate(Times[NextTimeIndex], NextValue);
-			if (NextValue == ActiveValue)
-			{
-				//same value so delete
-				ChannelData.RemoveKey(NextTimeIndex);
-			}
-		}
-	};
-	if (Index != INDEX_NONE)
+	
+	if (NextTimeIndex != Index && Times.IsValidIndex(NextTimeIndex))
 	{
-		DeleteNextKey();
-		//delete key will add it later
-		ChannelData.RemoveKey(Index);
-		//if last key then we won't add any more keys
-		if (Index == (Times.Num() - 1))
+		bool NextValue = false;
+		InOutActiveChannel.Evaluate(Times[NextTimeIndex], NextValue);
+		if (NextValue == InValueToClean)
 		{
-			return false;
+			OutTimesRemoved.Add(Times[NextTimeIndex]);
+
+			//same value so delete
+			ChannelData.RemoveKey(NextTimeIndex);
 		}
 	}
-	//if has same value then delete the next
-	else if (Times.IsValidIndex(NextTimeIndex))
-	{
-		DeleteNextKey();
-	}
-	
-	return true;
 }
 
 void ITransformConstraintChannelInterface::OnSequencerClosed(TSharedRef<ISequencer> InSequencer)
diff --git a/Engine/Source/Editor/MovieSceneTools/Public/Constraints/TransformConstraintChannelInterface.h b/Engine/Source/Editor/MovieSceneTools/Public/Constraints/TransformConstraintChannelInterface.h
--- a/Engine/Source/Editor/MovieSceneTools/Public/Constraints/TransformConstraintChannelInterface.h
+++ b/Engine/Source/Editor/MovieSceneTools/Public/Constraints/TransformConstraintChannelInterface.h
@@ -53,7 +53,9 @@ protected:
 
 	MOVIESCENETOOLS_API ITransformConstraintChannelInterface();
 
-	static MOVIESCENETOOLS_API bool CanAddKey(FMovieSceneConstraintChannel& ActiveChannel, const FFrameNumber& InTime, bool& ActiveValue);
+	static MOVIESCENETOOLS_API bool CanAddKey(const FMovieSceneConstraintChannel& InActiveChannel, const FFrameNumber& InTime, bool& ActiveValue);
+	static MOVIESCENETOOLS_API void CleanDuplicates(FMovieSceneConstraintChannel& InOutActiveChannel,
+		const FFrameNumber& InTime, const bool InValueToClean, TArray<FFrameNumber>& OutTimesRemoved);
 
 	void Initialize();
 	void Shutdown();
