Stop crash in transform origin editor mode when using TRS gizmo.

The TRS Gizmo still will not be able to edit transform origins since it does not trigger FSubTrackEditorMode::StartTracking

diff --git a/Engine/Source/Editor/MovieSceneTools/Private/EditModes/SubTrackEditorMode.cpp b/Engine/Source/Editor/MovieSceneTools/Private/EditModes/SubTrackEditorMode.cpp
--- a/Engine/Source/Editor/MovieSceneTools/Private/EditModes/SubTrackEditorMode.cpp
+++ b/Engine/Source/Editor/MovieSceneTools/Private/EditModes/SubTrackEditorMode.cpp
@@ -42,7 +42,7 @@ void FSubTrackEditorMode::Initialize()
 
 bool FSubTrackEditorMode::InputDelta(FEditorViewportClient* InViewportClient, FViewport* InViewport, FVector& InDrag, FRotator& InRot, FVector& InScale)
 {
-	if(AreAnyActorsSelected() || (InDrag.IsNearlyZero() && InRot.IsNearlyZero()))
+	if(!bIsTracking || AreAnyActorsSelected() || (InDrag.IsNearlyZero() && InRot.IsNearlyZero()))
 	{
 		return false;
 	}
