Add category support to animation params in section properties

diff --git a/Engine/Source/Editor/MovieSceneTools/Private/TrackEditors/SkeletalAnimationTrackEditor.cpp b/Engine/Source/Editor/MovieSceneTools/Private/TrackEditors/SkeletalAnimationTrackEditor.cpp
--- a/Engine/Source/Editor/MovieSceneTools/Private/TrackEditors/SkeletalAnimationTrackEditor.cpp
+++ b/Engine/Source/Editor/MovieSceneTools/Private/TrackEditors/SkeletalAnimationTrackEditor.cpp
@@ -90,6 +90,7 @@
 #include "AnimSequenceLevelSequenceLink.h"
 #include "UObject/SavePackage.h"
 #include "AnimSequencerInstanceProxy.h"
+#include "IDetailGroup.h"
 #include "TimeToPixel.h"
 #include "SequencerUtilities.h"
 #include "SequencerAnimationOverride.h"
@@ -474,20 +475,44 @@ public:
 		SlotNameProperty = PropertyHandle->GetChildHandle(GET_MEMBER_NAME_CHECKED(FMovieSceneSkeletalAnimationParams, SlotName));
 	}
 
+	IDetailPropertyRow& AddPropertyRow(IDetailChildrenBuilder& ChildBuilder, const FString& GroupName, TMap<FString, IDetailGroup*>& NameToGroupMap, const TSharedPtr<IPropertyHandle>& ChildPropertyHandle)
+	{
+		if (GroupName.IsEmpty())
+		{
+			return ChildBuilder.AddProperty(ChildPropertyHandle.ToSharedRef());
+		}
+		
+		IDetailGroup* const* DetailGroup = NameToGroupMap.Find(GroupName);
+		if (!DetailGroup)
+		{
+			IDetailGroup& NewGroup = ChildBuilder.AddGroup(*GroupName, FText::FromString(GroupName));
+			DetailGroup = &NameToGroupMap.Add(GroupName, &NewGroup);
+		}
+		return (*DetailGroup)->AddPropertyRow(ChildPropertyHandle.ToSharedRef());
+	
+	}
+
 	virtual void CustomizeChildren(TSharedRef<IPropertyHandle> PropertyHandle, IDetailChildrenBuilder& ChildBuilder, IPropertyTypeCustomizationUtils& CustomizationUtils) override
 	{
 		const FName AnimationPropertyName = GET_MEMBER_NAME_CHECKED(FMovieSceneSkeletalAnimationParams, Animation);
 		const FName MirrorDataTableName = GET_MEMBER_NAME_CHECKED(FMovieSceneSkeletalAnimationParams, MirrorDataTable);
 		const FName SlotNamePropertyName = GET_MEMBER_NAME_CHECKED(FMovieSceneSkeletalAnimationParams, SlotName);
 
+		TMap<FString, IDetailGroup*> NameToGroupMap;
 
 		uint32 NumChildren;
 		PropertyHandle->GetNumChildren(NumChildren);
 		for (uint32 i = 0; i < NumChildren; ++i)
 		{
 			TSharedPtr<IPropertyHandle> ChildPropertyHandle = PropertyHandle->GetChildHandle(i);
-			IDetailPropertyRow& ChildPropertyRow = ChildBuilder.AddProperty(ChildPropertyHandle.ToSharedRef());
 			FName ChildPropertyName = ChildPropertyHandle->GetProperty()->GetFName();
+			
+			FString GroupName;
+			FString CategoryName;
+			ChildPropertyHandle->GetDefaultCategoryName().ToString().Split("|", &CategoryName, &GroupName);
+
+			IDetailPropertyRow& ChildPropertyRow = AddPropertyRow(ChildBuilder, GroupName, NameToGroupMap, ChildPropertyHandle);
+			
 			// Let most properties be whatever they want to be... we just want to customize the `Animation` and `MirrorDataTable` properties
 			// by making it look like a normal asset reference property, but with some custom filtering.
 			if (ChildPropertyName == AnimationPropertyName || ChildPropertyName == MirrorDataTableName)
@@ -526,6 +551,8 @@ public:
 			}
 			else if (ChildPropertyName == SlotNamePropertyName)
 			{
+				// If the anim instance implements the ISequencerAnimationOverride interface, and has defined slots to used, override this row with a drop-down menu
+				// Otherwise the default row will be created, which uses a text input field.
 				if (bShowSlotNameOptions)
 				{
 					ChildPropertyRow.IsEnabled(TAttribute<bool>::CreateSP(this, &FMovieSceneSkeletalAnimationParamsDetailCustomization::GetCanEditSlotName));
@@ -533,17 +560,17 @@ public:
 					Row.NameContent()[ChildPropertyHandle->CreatePropertyNameWidget()];
 
 					Row.ValueContent()
-						[
-							SNew(SComboBox<FName>)
-							.OptionsSource(&SlotNameOptions)
+					[
+						SNew(SComboBox<FName>)
+						.OptionsSource(&SlotNameOptions)
 						.OnSelectionChanged(this, &FMovieSceneSkeletalAnimationParamsDetailCustomization::OnSlotNameChanged)
 						.OnGenerateWidget_Lambda([](FName InSlotName) { return SNew(STextBlock).Text(FText::FromName(InSlotName)); })
 						[
 							SNew(STextBlock)
 							.Font(IPropertyTypeCustomizationUtils::GetRegularFont())
-						.Text(this, &FMovieSceneSkeletalAnimationParamsDetailCustomization::GetSlotNameDesc)
+							.Text(this, &FMovieSceneSkeletalAnimationParamsDetailCustomization::GetSlotNameDesc)
 						]
-						];
+					];
 				}
 			}
 		}
diff --git a/Engine/Source/Runtime/MovieSceneTracks/Public/Sections/MovieSceneSkeletalAnimationSection.h b/Engine/Source/Runtime/MovieSceneTracks/Public/Sections/MovieSceneSkeletalAnimationSection.h
--- a/Engine/Source/Runtime/MovieSceneTracks/Public/Sections/MovieSceneSkeletalAnimationSection.h
+++ b/Engine/Source/Runtime/MovieSceneTracks/Public/Sections/MovieSceneSkeletalAnimationSection.h
@@ -52,30 +52,30 @@ struct FMovieSceneSkeletalAnimationParams
 	TObjectPtr<UAnimSequenceBase> Animation;
 
 	/** The offset into the beginning of the animation clip for the first loop of play. */
-	UPROPERTY(BlueprintReadWrite, EditAnywhere, Category="Animation")
+	UPROPERTY(BlueprintReadWrite, EditAnywhere, Category="Animation|Offset")
 	FFrameNumber FirstLoopStartFrameOffset;
 
 	/** The offset into the beginning of the animation clip */
-	UPROPERTY(BlueprintReadWrite, EditAnywhere, Category="Animation")
+	UPROPERTY(BlueprintReadWrite, EditAnywhere, Category="Animation|Offset")
 	FFrameNumber StartFrameOffset;
 
 	/** The offset into the end of the animation clip */
-	UPROPERTY(BlueprintReadWrite, EditAnywhere, Category="Animation")
+	UPROPERTY(BlueprintReadWrite, EditAnywhere, Category="Animation|Offset")
 	FFrameNumber EndFrameOffset;
 
 	/** The playback rate of the animation clip */
-	UPROPERTY(BlueprintReadWrite, EditAnywhere, Category="Animation")
+	UPROPERTY(BlueprintReadWrite, EditAnywhere, Category="Animation|Playback")
 	FMovieSceneTimeWarpVariant PlayRate;
 
 	/** Reverse the playback of the animation clip */
-	UPROPERTY(BlueprintReadWrite, EditAnywhere, Category="Animation")
+	UPROPERTY(BlueprintReadWrite, EditAnywhere, Category="Animation|Playback")
 	uint32 bReverse:1;
 
 	/** The slot name to use for the animation */
-	UPROPERTY(BlueprintReadWrite, EditAnywhere, Category = "Animation" )
+	UPROPERTY(BlueprintReadWrite, EditAnywhere, Category = "Animation|Playback" )
 	FName SlotName;
 
-	UPROPERTY(BlueprintReadWrite, EditAnywhere, Category = "Animation")
+	UPROPERTY(BlueprintReadWrite, EditAnywhere, Category = "Animation|Playback")
 	TObjectPtr<class UMirrorDataTable> MirrorDataTable;
 
 	/** The weight curve for this animation section */
@@ -83,15 +83,15 @@ struct FMovieSceneSkeletalAnimationParams
 	FMovieSceneFloatChannel Weight;
 
 	/** If on will skip sending animation notifies */
-	UPROPERTY(BlueprintReadWrite, EditAnywhere, Category = "Animation")
+	UPROPERTY(BlueprintReadWrite, EditAnywhere, Category = "Animation|Playback")
 	bool bSkipAnimNotifiers;
 
 	/** If on animation sequence will always play when active even if the animation is controlled by a Blueprint or Anim Instance Class*/
-	UPROPERTY(BlueprintReadWrite, EditAnywhere, Category = "Animation")
+	UPROPERTY(BlueprintReadWrite, EditAnywhere, Category = "Animation|Playback")
 	bool bForceCustomMode;
 
 	/** If on the root bone transform will be swapped to the specified root*/
-	UPROPERTY(BlueprintReadWrite, EditAnywhere, Category = "Animation")
+	UPROPERTY(BlueprintReadWrite, EditAnywhere, Category = "Animation|Playback")
 	ESwapRootBone SwapRootBone;
 
 	UPROPERTY()
