SubTrackEditorMode: enable new TRS gizmo transform

#rnx

diff --git a/Engine/Source/Editor/MovieSceneTools/Private/EditModes/SubTrackEditorMode.cpp b/Engine/Source/Editor/MovieSceneTools/Private/EditModes/SubTrackEditorMode.cpp
--- a/Engine/Source/Editor/MovieSceneTools/Private/EditModes/SubTrackEditorMode.cpp
+++ b/Engine/Source/Editor/MovieSceneTools/Private/EditModes/SubTrackEditorMode.cpp
@@ -74,6 +74,36 @@ bool FSubTrackEditorMode::InputDelta(FEditorViewportClient* InViewportClient, FV
 }
 
 bool FSubTrackEditorMode::StartTracking(FEditorViewportClient* InViewportClient, FViewport* InViewport)
+{
+	if (HandleBeginTransform())
+	{
+		return true;
+	}
+	
+	return FEdMode::StartTracking(InViewportClient, InViewport);
+}
+
+bool FSubTrackEditorMode::EndTracking(FEditorViewportClient* InViewportClient, FViewport* InViewport)
+{
+	if (HandleEndTransform())
+	{
+		return true;
+	}
+	
+	return FEdMode::EndTracking(InViewportClient, InViewport);
+}
+
+bool FSubTrackEditorMode::BeginTransform(const FGizmoState& InState)
+{
+	return HandleBeginTransform();
+}
+
+bool FSubTrackEditorMode::EndTransform(const FGizmoState& InState)
+{
+	return HandleEndTransform();
+}
+
+bool FSubTrackEditorMode::HandleBeginTransform()
 {
 	if (const UMovieSceneSubSection* SubSection = GetSectionToEdit())
 	{
@@ -85,13 +115,13 @@ bool FSubTrackEditorMode::StartTracking(FEditorViewportClient* InViewportClient,
 		}
 		return true;
 	}
-	
-	return FEdMode::StartTracking(InViewportClient, InViewport);
+
+	return false;
 }
 
-bool FSubTrackEditorMode::EndTracking(FEditorViewportClient* InViewportClient, FViewport* InViewport)
+bool FSubTrackEditorMode::HandleEndTransform()
 {
-	if(bIsTracking)
+	if (bIsTracking)
 	{
 		bIsTracking = false;
 		// Only reset the preview rotation. Resetting the preview location here could interfere with multiple rotation edits in a row.
@@ -99,8 +129,8 @@ bool FSubTrackEditorMode::EndTracking(FEditorViewportClient* InViewportClient, F
 		PreviewCoordinateSpaceRotation.Reset();
 		return true;
 	}
-	
-	return FEdMode::EndTracking(InViewportClient, InViewport);
+
+	return false;
 }
 
 bool FSubTrackEditorMode::UsesTransformWidget() const
diff --git a/Engine/Source/Editor/MovieSceneTools/Private/EditModes/SubTrackEditorMode.h b/Engine/Source/Editor/MovieSceneTools/Private/EditModes/SubTrackEditorMode.h
--- a/Engine/Source/Editor/MovieSceneTools/Private/EditModes/SubTrackEditorMode.h
+++ b/Engine/Source/Editor/MovieSceneTools/Private/EditModes/SubTrackEditorMode.h
@@ -34,6 +34,11 @@ public:
 
 	void ClearCachedCoordinates();
 
+	// Start IGizmoEdModeInterface overrides
+	virtual bool BeginTransform(const FGizmoState& InState) override;
+	virtual bool EndTransform(const FGizmoState& InState) override;
+	// End IGizmoEdModeInterface overrides
+
 private:
 	// Returns true if there are any sub tracks with active transform origin overrides.
 	static bool DoesSubSectionHaveTransformOverrides(const UMovieSceneSubSection& SubSection);
@@ -80,4 +85,8 @@ private:
 	TOptional<FVector> PreviewLocation;
 
 	const TArray<FName> IncompatibleEditorModes = TArray<FName>({ "EditMode.ControlRig", "EM_Landscape" });
+
+	/** Manages the start and end of a transform action in the viewport */
+	bool HandleBeginTransform();
+	bool HandleEndTransform();
 };
