Curve Editor: Rework caching of draw params for better editor performance.

* Cached Draw Params are moved to their own class
* Curve models can now implement a curve cache pool to cache their data efficiently.
* Bezier curves now implement a curve cache pool to optimize their editor performance
* The optimizations are disabled by default, use CVar CurveEditor.UseCurveCachePool 1 to enable them

diff --git a/Engine/Source/Editor/CurveEditor/Private/CurveEditor.cpp b/Engine/Source/Editor/CurveEditor/Private/CurveEditor.cpp
--- a/Engine/Source/Editor/CurveEditor/Private/CurveEditor.cpp
+++ b/Engine/Source/Editor/CurveEditor/Private/CurveEditor.cpp
@@ -7,6 +7,7 @@
 #include "CoreGlobals.h"
 #include "CurveEditorCommands.h"
 #include "CurveEditorCopyBuffer.h"
+#include "CurveEditorCurveDrawParamsCache.h"
 #include "CurveEditorSettings.h"
 #include "CurveEditorSnapMetrics.h"
 #include "CurveEditorAxis.h"
@@ -132,6 +133,8 @@ void FCurveEditor::InitCurveEditor(const FCurveEditorInitParams& InInitParams)
 	SuspendBroadcastCount = 0;
 	// Listen to global undo so we can fix up our selection state for keys that no longer exist.
 	GEditor->RegisterForUndo(this);
+
+	CurveDrawParamsCache = MakeShared<UE::CurveEditor::FCurveDrawParamsCache>(AsShared());
 }
 
 int32 FCurveEditor::GetSupportedTangentTypes()
diff --git a/Engine/Source/Editor/CurveEditor/Private/CurveEditorCurveDrawParamsCache.cpp b/Engine/Source/Editor/CurveEditor/Private/CurveEditorCurveDrawParamsCache.cpp
new file mode 100644
--- /dev/null
+++ b/Engine/Source/Editor/CurveEditor/Private/CurveEditorCurveDrawParamsCache.cpp
@@ -0,0 +1,358 @@
+// Copyright Epic Games, Inc. All Rights Reserved.
+
+#include "CurveEditorCurveDrawParamsCache.h"
+
+#include "CurveDrawInfo.h"
+#include "CurveEditor.h"
+#include "CurveEditorCurveDrawParamsHandle.h"
+#include "CurveEditorHelpers.h"
+#include "CurveModel.h"
+#include "ICurveEditorCurveCachePool.h"
+#include "SCurveEditorView.h"
+
+namespace UE::CurveEditor
+{
+	int32 GCurveEditorUseCurveCachePool = 1;
+
+	static FAutoConsoleVariableRef CCurveEditorUseCurveCachePool(
+		TEXT("CurveEditor.UseCurveCachePool"),
+		GCurveEditorUseCurveCachePool,
+		TEXT("Enables support for curve editor cache Pool"),
+		ECVF_Default);
+
+	void FCurveDrawParamsCache::Invalidate(const TSharedRef<SCurveEditorView>& CurveEditorView, const TArray<FCurveModelID>& ModelIDs)
+	{
+		TRACE_CPUPROFILER_EVENT_SCOPE(FCurveDrawParamsCache::Invalidate);
+
+		UpdateCurveCacheFlags(CurveEditorView);
+		DrawCurves(CurveEditorView, ModelIDs);
+	}
+
+	void FCurveDrawParamsCache::UpdateAllCurveDrawParamSynchonous(const TSharedRef<SCurveEditorView>& CurveEditorView, const TArray<FCurveModelID>& CurveModelIDs, TArray<FCurveDrawParams>& OutParams)
+	{
+		if (!WeakCurveEditor.IsValid())
+		{
+			return;
+		}
+		const TSharedRef<FCurveEditor> CurveEditor = WeakCurveEditor.Pin().ToSharedRef();
+
+		CachedDrawParams.Reset(CurveModelIDs.Num());
+		for (const FCurveModelID& CurveModelID : CurveModelIDs)
+		{
+			if (FCurveModel* CurveModel = CurveEditor->FindCurve(CurveModelID))
+			{
+				const FCurveEditorScreenSpace ScreenSpace = CurveEditorView->GetCurveSpace(CurveModelID);
+
+				FCurveDrawParams NewDrawParams(CurveModelID);
+				UpdateCurveDrawParamsSynchonous(ScreenSpace, CurveModel, CurveModelID, NewDrawParams);
+				OutParams.Add(NewDrawParams);
+			}
+		}
+	}
+	
+	void FCurveDrawParamsCache::UpdateCurveDrawParamsSynchonous(const FCurveEditorScreenSpace& CurveSpace, FCurveModel* CurveModel, const FCurveModelID& ModelID, FCurveDrawParams& OutParams)
+	{
+		if (!WeakCurveEditor.IsValid())
+		{
+			return;
+		}
+		const TSharedRef<FCurveEditor> CurveEditor = WeakCurveEditor.Pin().ToSharedRef();
+
+		const double InputMin = CurveSpace.GetInputMin();
+		const double InputMax = CurveSpace.GetInputMax();
+
+		const double DisplayRatio = (CurveSpace.PixelsPerOutput() / CurveSpace.PixelsPerInput());
+
+		const FKeyHandleSet* SelectedKeys = CurveEditor->GetSelection().GetAll().Find(ModelID);
+
+		// Create a new set of Curve Drawing Parameters to represent this particular Curve
+		OutParams.Color = CurveModel->GetColor();
+		OutParams.Thickness = CurveModel->GetThickness();
+		OutParams.DashLengthPx = CurveModel->GetDashLength();
+		OutParams.bKeyDrawEnabled = CurveModel->IsKeyDrawEnabled();
+
+		// Gather the display metrics to use for each key type. This allows a Curve Model to override
+		// whether or not the curve supports Keys, Arrive/Leave Tangents, etc. If the Curve Model doesn't
+		// support a particular capability we can skip drawing them.
+		CurveModel->GetKeyDrawInfo(ECurvePointType::ArriveTangent, FKeyHandle::Invalid(), OutParams.ArriveTangentDrawInfo);
+		CurveModel->GetKeyDrawInfo(ECurvePointType::LeaveTangent, FKeyHandle::Invalid(), OutParams.LeaveTangentDrawInfo);
+
+		// Gather the interpolating points in input/output space
+		TArray<TTuple<double, double>> InterpolatingPoints;
+
+		CurveModel->DrawCurve(*CurveEditor, CurveSpace, InterpolatingPoints);
+		OutParams.InterpolatingPoints.Reset(InterpolatingPoints.Num());
+
+		// An Input Offset allows for a fixed offset to all keys, such as displaying them in the middle of a frame instead of at the start.
+		double InputOffset = CurveModel->GetInputDisplayOffset();
+
+		// Convert the interpolating points to screen space
+		for (const TTuple<double, double>& Point : InterpolatingPoints)
+		{
+			OutParams.InterpolatingPoints.Add(
+				FVector2D(
+					CurveSpace.SecondsToScreen(Point.Get<0>() + InputOffset),
+					CurveSpace.ValueToScreen(Point.Get<1>())
+				)
+			);
+		}
+
+		TArray<FKeyHandle> VisibleKeys;
+		CurveModel->GetKeys(InputMin, InputMax, TNumericLimits<double>::Lowest(), TNumericLimits<double>::Max(), VisibleKeys);
+
+		// Always reset the points to cover case going from 1 to 0 keys
+		OutParams.Points.Reset(VisibleKeys.Num());
+
+		if (VisibleKeys.Num())
+		{
+			ECurveEditorTangentVisibility TangentVisibility = CurveEditor->GetSettings()->GetTangentVisibility();
+
+			TArray<FKeyPosition> AllKeyPositions;
+			TArray<FKeyAttributes> AllKeyAttributes;
+
+			AllKeyPositions.SetNum(VisibleKeys.Num());
+			AllKeyAttributes.SetNum(VisibleKeys.Num());
+
+			CurveModel->GetKeyPositions(VisibleKeys, AllKeyPositions);
+			CurveModel->GetKeyAttributes(VisibleKeys, AllKeyAttributes);
+			for (int32 Index = 0; Index < VisibleKeys.Num(); ++Index)
+			{
+				const FKeyHandle      KeyHandle = VisibleKeys[Index];
+				const FKeyPosition& KeyPosition = AllKeyPositions[Index];
+				const FKeyAttributes& Attributes = AllKeyAttributes[Index];
+
+				bool bShowTangents = TangentVisibility == ECurveEditorTangentVisibility::AllTangents ||
+					(TangentVisibility == ECurveEditorTangentVisibility::SelectedKeys && SelectedKeys &&
+						(SelectedKeys->Contains(VisibleKeys[Index], ECurvePointType::Any)));
+
+				double TimeScreenPos = CurveSpace.SecondsToScreen(KeyPosition.InputValue + InputOffset);
+				double ValueScreenPos = CurveSpace.ValueToScreen(KeyPosition.OutputValue);
+
+				// Add this key
+				FCurvePointInfo Key(KeyHandle);
+				Key.ScreenPosition = FVector2D(TimeScreenPos, ValueScreenPos);
+				Key.LayerBias = 2;
+
+				// Add draw info for the specific key
+				CurveModel->GetKeyDrawInfo(ECurvePointType::Key, KeyHandle, /*Out*/ Key.DrawInfo);
+				OutParams.Points.Add(Key);
+
+				if (bShowTangents && Attributes.HasArriveTangent())
+				{
+					float ArriveTangent = Attributes.GetArriveTangent();
+
+					FCurvePointInfo ArriveTangentPoint(KeyHandle);
+					ArriveTangentPoint.Type = ECurvePointType::ArriveTangent;
+
+
+					if (Attributes.HasTangentWeightMode() && Attributes.HasArriveTangentWeight() &&
+						(Attributes.GetTangentWeightMode() == RCTWM_WeightedBoth || Attributes.GetTangentWeightMode() == RCTWM_WeightedArrive))
+					{
+						FVector2D TangentOffset = ::CurveEditor::ComputeScreenSpaceTangentOffset(CurveSpace, ArriveTangent, -Attributes.GetArriveTangentWeight());
+						ArriveTangentPoint.ScreenPosition = Key.ScreenPosition + TangentOffset;
+					}
+					else
+					{
+						float PixelLength = 60.0f;
+						ArriveTangentPoint.ScreenPosition = Key.ScreenPosition + ::CurveEditor::GetVectorFromSlopeAndLength(ArriveTangent * -DisplayRatio, -PixelLength);
+					}
+					ArriveTangentPoint.LineDelta = Key.ScreenPosition - ArriveTangentPoint.ScreenPosition;
+					ArriveTangentPoint.LayerBias = 1;
+
+					// Add draw info for the specific tangent
+					FKeyDrawInfo TangentDrawInfo;
+					CurveModel->GetKeyDrawInfo(ECurvePointType::ArriveTangent, KeyHandle, /*Out*/ ArriveTangentPoint.DrawInfo);
+
+					OutParams.Points.Add(ArriveTangentPoint);
+				}
+
+				if (bShowTangents && Attributes.HasLeaveTangent())
+				{
+					float LeaveTangent = Attributes.GetLeaveTangent();
+
+					FCurvePointInfo LeaveTangentPoint(KeyHandle);
+					LeaveTangentPoint.Type = ECurvePointType::LeaveTangent;
+
+					if (Attributes.HasTangentWeightMode() && Attributes.HasLeaveTangentWeight() &&
+						(Attributes.GetTangentWeightMode() == RCTWM_WeightedBoth || Attributes.GetTangentWeightMode() == RCTWM_WeightedLeave))
+					{
+						FVector2D TangentOffset = ::CurveEditor::ComputeScreenSpaceTangentOffset(CurveSpace, LeaveTangent, Attributes.GetLeaveTangentWeight());
+
+						LeaveTangentPoint.ScreenPosition = Key.ScreenPosition + TangentOffset;
+					}
+					else
+					{
+						float PixelLength = 60.0f;
+						LeaveTangentPoint.ScreenPosition = Key.ScreenPosition + ::CurveEditor::GetVectorFromSlopeAndLength(LeaveTangent * -DisplayRatio, PixelLength);
+					}
+
+					LeaveTangentPoint.LineDelta = Key.ScreenPosition - LeaveTangentPoint.ScreenPosition;
+					LeaveTangentPoint.LayerBias = 1;
+
+					// Add draw info for the specific tangent
+					FKeyDrawInfo TangentDrawInfo;
+					CurveModel->GetKeyDrawInfo(ECurvePointType::LeaveTangent, KeyHandle, /*Out*/ LeaveTangentPoint.DrawInfo);
+
+					OutParams.Points.Add(LeaveTangentPoint);
+				}
+			}
+		}
+	}
+
+	FCurveDrawParamsCache::FCurveDrawParamsCache(const TSharedRef<FCurveEditor>& InCurveEditor)
+		: WeakCurveEditor(InCurveEditor)
+	{
+		CurveCacheFlags = ECurveCacheFlags::UpdateAll;
+		CachedCurveEditorData.ActiveCurvesSerialNumber = 0xFFFFFFFF;
+		CachedCurveEditorData.SelectionSerialNumber = 0xFFFFFFFF;
+		CachedCurveEditorData.CachedGeometrySize.X = -1;
+		CachedCurveEditorData.CachedGeometrySize.Y = -1;
+	}
+
+	void FCurveDrawParamsCache::UpdateCurveDrawParams(const TSharedRef<SCurveEditorView>& CurveEditorView, const TArray<FCurveDrawParamsHandle>& CurveDrawParamsHandles)
+	{
+		if (!WeakCurveEditor.IsValid())
+		{
+			return;
+		}
+		const TSharedRef<FCurveEditor> CurveEditor = WeakCurveEditor.Pin().ToSharedRef();
+
+		const bool bUseCurveCachePool = GCurveEditorUseCurveCachePool > 0;
+
+		TSet<ICurveEditorCurveCachePool*> CachePools;
+		for (const FCurveDrawParamsHandle& CurveDrawParamsHandle : CurveDrawParamsHandles)
+		{
+			if (FCurveModel* CurveModel = CurveEditor->FindCurve(CurveDrawParamsHandle.GetID()))
+			{
+				const FCurveEditorScreenSpace ScreenSpace = CurveEditorView->GetCurveSpace(CurveDrawParamsHandle.GetID());
+
+				ICurveEditorCurveCachePool* CachePool = bUseCurveCachePool ? CurveModel->DrawCurveToCachePool(CurveEditor, CurveDrawParamsHandle, ScreenSpace) : nullptr;
+				if (CachePool)
+				{
+					CachePools.Add(CachePool);
+				}
+				else
+				{
+					if (FCurveDrawParams* CurveDrawParams = CurveDrawParamsHandle.Get())
+					{
+						UpdateCurveDrawParamsSynchonous(ScreenSpace, CurveModel, CurveDrawParamsHandle.GetID(), *CurveDrawParams);
+					}
+				}
+			}
+		}
+
+		for (ICurveEditorCurveCachePool* CachePool : CachePools)
+		{
+			CachePool->DrawCachedCurves(WeakCurveEditor);
+		}
+	}
+
+	void FCurveDrawParamsCache::UpdateCurveCacheFlags(const TSharedRef<SCurveEditorView>& CurveEditorView)
+	{
+		if (!WeakCurveEditor.IsValid())
+		{
+			return;
+		}
+		const TSharedRef<FCurveEditor> CurveEditor = WeakCurveEditor.Pin().ToSharedRef();
+
+		if (CurveEditor->GetActiveCurvesSerialNumber() != CachedCurveEditorData.ActiveCurvesSerialNumber)
+		{
+			CachedCurveEditorData.ActiveCurvesSerialNumber = CurveEditor->GetActiveCurvesSerialNumber();
+			CurveCacheFlags = ECurveCacheFlags::UpdateAll;
+		}
+
+		if (CachedCurveEditorData.TangentVisibility != CurveEditor->GetSettings()->GetTangentVisibility())
+		{
+			CachedCurveEditorData.TangentVisibility = CurveEditor->GetSettings()->GetTangentVisibility();
+			CurveCacheFlags = ECurveCacheFlags::UpdateAll;
+		}
+
+		if (CachedCurveEditorData.SelectionSerialNumber != CurveEditor->GetSelection().GetSerialNumber())
+		{
+			CachedCurveEditorData.SelectionSerialNumber = CurveEditor->GetSelection().GetSerialNumber();
+			CurveCacheFlags = ECurveCacheFlags::UpdateAll;
+		}
+
+		// Only get view values if we need to since we will reset them every time we get all
+		if (CurveCacheFlags != ECurveCacheFlags::UpdateAll)
+		{
+			const double OutputMin = CurveEditorView->GetOutputMin();
+			const double OutputMax = CurveEditorView->GetOutputMax();
+
+			if (OutputMin != CachedCurveEditorData.OutputMin || OutputMax != CachedCurveEditorData.OutputMax)
+			{
+				CachedCurveEditorData.OutputMin = OutputMin;
+				CachedCurveEditorData.OutputMax = OutputMax;
+				CurveCacheFlags = ECurveCacheFlags::UpdateAll;
+			}
+			else if (CachedCurveEditorData.CachedGeometrySize != CurveEditorView->GetCachedGeometry().GetLocalSize())
+			{
+				CurveCacheFlags = ECurveCacheFlags::UpdateAll;
+			}
+			else
+			{
+				double InputMin = 0, InputMax = 1;
+				CurveEditorView->GetInputBounds(InputMin, InputMax);
+				if (InputMin != CachedCurveEditorData.InputMin || InputMax != CachedCurveEditorData.InputMax)
+				{
+					CurveCacheFlags = ECurveCacheFlags::UpdateAll;
+				}
+			}
+		}
+	}
+
+	void FCurveDrawParamsCache::DrawCurves(const TSharedRef<SCurveEditorView>& CurveEditorView, const TArray<FCurveModelID>& ModelIDs)
+	{
+		if (!WeakCurveEditor.IsValid())
+		{
+			return;
+		}
+		const TSharedRef<FCurveEditor> CurveEditor = WeakCurveEditor.Pin().ToSharedRef();
+
+		if (CurveCacheFlags == ECurveCacheFlags::UpdateAll)
+		{
+			CurveEditorView->GetInputBounds(CachedCurveEditorData.InputMin, CachedCurveEditorData.InputMax);
+			CachedCurveEditorData.CachedGeometrySize = CurveEditorView->GetCachedGeometry().GetLocalSize();
+
+			CachedDrawParams.Reset(ModelIDs.Num());
+
+			TArray<FCurveDrawParamsHandle> DrawParamHandles;
+			for (const FCurveModelID& ModelID : ModelIDs)
+			{
+				const int32 DrawParamsIndex = CachedDrawParams.Emplace(ModelID);
+				DrawParamHandles.Emplace(AsShared(), DrawParamsIndex);
+			}
+			UpdateCurveDrawParams(CurveEditorView, DrawParamHandles);
+
+			CurveCacheFlags = ECurveCacheFlags::CheckCurves;
+
+			CurveEditorView->RefreshRetainer();
+		}
+		else if (CurveCacheFlags == ECurveCacheFlags::CheckCurves)
+		{
+			TArray<FCurveDrawParamsHandle> ChangedDrawParamHandles;
+			for (int32 DrawParamsIndex = 0; DrawParamsIndex < CachedDrawParams.Num(); DrawParamsIndex++)
+			{
+				const FCurveModelID& CurveModelID = CachedDrawParams[DrawParamsIndex].GetID();
+
+				FCurveModel* CurveModel = CurveEditor->FindCurve(CurveModelID);
+				if (CurveModel && CurveModel->HasChangedAndResetTest())
+				{
+					ChangedDrawParamHandles.Emplace(AsShared(), DrawParamsIndex);
+					if (CurveEditorView->bAllowModelViewTransforms)
+					{
+						CurveEditorView->bUpdateModelViewTransforms = true;
+						CurveEditorView->UpdateCustomAxes();
+					}
+				}
+			}
+
+			if (!ChangedDrawParamHandles.IsEmpty())
+			{
+				UpdateCurveDrawParams(CurveEditorView, ChangedDrawParamHandles);
+
+				CurveEditorView->RefreshRetainer();
+			}
+		}
+	}
+}
diff --git a/Engine/Source/Editor/CurveEditor/Private/CurveEditorCurveDrawParamsHandle.cpp b/Engine/Source/Editor/CurveEditor/Private/CurveEditorCurveDrawParamsHandle.cpp
new file mode 100644
--- /dev/null
+++ b/Engine/Source/Editor/CurveEditor/Private/CurveEditorCurveDrawParamsHandle.cpp
@@ -0,0 +1,43 @@
+// Copyright Epic Games, Inc. All Rights Reserved.
+
+#include "CurveEditorCurveDrawParamsHandle.h"
+
+#include "CurveDrawInfo.h"
+#include "CurveEditorCurveDrawParamsCache.h"
+
+namespace UE::CurveEditor
+{
+	/** Constructs the handle from the index in curve draw params. */
+	FCurveDrawParamsHandle::FCurveDrawParamsHandle(const TSharedRef<FCurveDrawParamsCache>& InDrawParamsCache, const int32 InIndex)
+		: Index(InIndex)
+		, WeakDrawParamsCache(InDrawParamsCache)
+	{
+		CurveModelID = InDrawParamsCache->CachedDrawParams[InIndex].GetID();
+	}
+
+	FCurveDrawParams* FCurveDrawParamsHandle::Get() const
+	{
+		if (!WeakDrawParamsCache.IsValid())
+		{
+			return nullptr;
+		}
+		FCurveDrawParamsCache& Cache = *WeakDrawParamsCache.Pin();
+
+		if (Cache.CachedDrawParams.IsValidIndex(Index) &&
+			Cache.CachedDrawParams[Index].GetID() == CurveModelID)
+		{
+			return &Cache.CachedDrawParams[Index];
+		}
+		else 
+		{
+			const TArray<FCurveDrawParams>& CurveDrawParams = Cache.CachedDrawParams;
+			Index = WeakDrawParamsCache.Pin()->CachedDrawParams.IndexOfByPredicate(
+				[this](const FCurveDrawParams& CurveDrawParams)
+				{
+					return CurveDrawParams.GetID() == CurveModelID;
+				});
+
+			return Index == INDEX_NONE ? nullptr : &Cache.CachedDrawParams[Index];
+		}
+	}
+}
diff --git a/Engine/Source/Editor/CurveEditor/Private/SCurveEditorView.cpp b/Engine/Source/Editor/CurveEditor/Private/SCurveEditorView.cpp
--- a/Engine/Source/Editor/CurveEditor/Private/SCurveEditorView.cpp
+++ b/Engine/Source/Editor/CurveEditor/Private/SCurveEditorView.cpp
@@ -7,6 +7,7 @@
 #include "CurveDrawInfo.h"
 #include "CurveEditor.h"
 #include "CurveEditorAxis.h"
+#include "CurveEditorCurveDrawParamsCache.h"
 #include "CurveEditorHelpers.h"
 #include "CurveEditorScreenSpace.h"
 #include "CurveEditorSelection.h"
@@ -24,7 +25,11 @@
 #include "Templates/Tuple.h"
 #include "Templates/UnrealTemplate.h"
 
-TAutoConsoleVariable<bool> CVarUseCurveCache(TEXT("CurveEditor.UseCurveCache"), true, TEXT("When true we cache curve values, when false we always regenerate."));
+int32 GCurveEditorUseCurveCache = 1;
+FAutoConsoleVariableRef CCurveEditorUseCurveCache(
+	TEXT("CurveEditor.UseCurveCache"), 
+	GCurveEditorUseCurveCache, 
+	TEXT("When true we cache curve values, when false we always regenerate"));
 
 SCurveEditorView::SCurveEditorView()
 	: bPinned(0)
@@ -37,16 +42,20 @@ SCurveEditorView::SCurveEditorView()
 	, bNeedsDefaultGridLinesH(1)
 	, bNeedsDefaultGridLinesV(1)
 {
+	PRAGMA_DISABLE_DEPRECATION_WARNINGS
 	CurveCacheFlags = ECurveCacheFlags::All;
 	CachedValues.CachedActiveCurvesSerialNumber = 0xFFFFFFFF;
 	CachedValues.CachedSelectionSerialNumber = 0xFFFFFFFF;
 	CachedValues.CachedGeometrySize.X = -1;
 	CachedValues.CachedGeometrySize.Y = -1;
+	PRAGMA_DISABLE_DEPRECATION_WARNINGS
 }
 
+PRAGMA_DISABLE_DEPRECATION_WARNINGS // Allow deprecated CachedCurveParams to destruct
 SCurveEditorView::~SCurveEditorView()
 {
 }
+PRAGMA_ENABLE_DEPRECATION_WARNINGS
 
 FVector2D SCurveEditorView::ComputeDesiredSize(float LayoutScaleMultiplier) const
 {
@@ -551,8 +560,10 @@ void SCurveEditorView::ZoomAround(const FVector2D& Amount, double InputOrigin, d
 	}
 }
 
-void SCurveEditorView::GetCurveDrawParams(TArray<FCurveDrawParams>& OutDrawParams) 
+void SCurveEditorView::GetCurveDrawParams(TArray<FCurveDrawParams>& OutDrawParams)
 {
+	// DEPRECATED 5.6, instead use FCurveDrawParamsCache::GetCurvedDrawParamsSynchronous
+
 	TSharedPtr<FCurveEditor> CurveEditor = WeakCurveEditor.Pin();
 	if (!CurveEditor)
 	{
@@ -578,13 +589,18 @@ void SCurveEditorView::GetCurveDrawParams(TArray<FCurveDrawParams>& OutDrawParam
 			continue;
 		}
 
+		PRAGMA_DISABLE_DEPRECATION_WARNINGS
 		GetCurveDrawParam(CurveEditor, Pair.Key, CurveModel, Params);
+		PRAGMA_DISABLE_DEPRECATION_WARNINGS
+		
 		OutDrawParams.Add(MoveTemp(Params));
 	}
 }
 
-void SCurveEditorView::GetCurveDrawParam(TSharedPtr<FCurveEditor>& CurveEditor,const FCurveModelID& ModelID, FCurveModel* CurveModel, FCurveDrawParams& Params) const
-{
+void SCurveEditorView::GetCurveDrawParam(TSharedPtr<FCurveEditor>& CurveEditor, const FCurveModelID& ModelID, FCurveModel* CurveModel, FCurveDrawParams& Params) const
+{	
+	// DEPRECATED 5.6, instead use FCurveDrawParamsCache::GetCurveDrawParamSynchronous
+
 	FCurveEditorScreenSpace CurveSpace = GetCurveSpace(ModelID);
 
 	const double InputMin = CurveSpace.GetInputMin();
@@ -604,7 +620,7 @@ void SCurveEditorView::GetCurveDrawParam(TSharedPtr<FCurveEditor>& CurveEditor,c
 	// whether or not the curve supports Keys, Arrive/Leave Tangents, etc. If the Curve Model doesn't
 	// support a particular capability we can skip drawing them.
 	CurveModel->GetKeyDrawInfo(ECurvePointType::ArriveTangent, FKeyHandle::Invalid(), Params.ArriveTangentDrawInfo);
-	CurveModel->GetKeyDrawInfo(ECurvePointType::LeaveTangent,  FKeyHandle::Invalid(), Params.LeaveTangentDrawInfo);
+	CurveModel->GetKeyDrawInfo(ECurvePointType::LeaveTangent, FKeyHandle::Invalid(), Params.LeaveTangentDrawInfo);
 
 	// Gather the interpolating points in input/output space
 	TArray<TTuple<double, double>> InterpolatingPoints;
@@ -743,82 +759,14 @@ void SCurveEditorView::CheckCacheAndInvalidateIfNeeded()
 	{
 		return;
 	}
-	const bool bUseCurveCache = CVarUseCurveCache.GetValueOnGameThread();
-
+	
+	const bool bUseCurveCache = GCurveEditorUseCurveCache > 0;
 	if (bUseCurveCache)
 	{
-		//if number of curves have changed, just redo all
-		if (CurveEditor->GetActiveCurvesSerialNumber() != CachedValues.CachedActiveCurvesSerialNumber)
-		{
-			CachedValues.CachedActiveCurvesSerialNumber = CurveEditor->GetActiveCurvesSerialNumber();
-			CurveCacheFlags = ECurveCacheFlags::All;
-		}
-		if (CachedValues.CachedTangentVisibility != CurveEditor->GetSettings()->GetTangentVisibility())
-		{
-			CachedValues.CachedTangentVisibility = CurveEditor->GetSettings()->GetTangentVisibility();
-			CurveCacheFlags = ECurveCacheFlags::All;
-		}
-		if (CachedValues.CachedSelectionSerialNumber != CurveEditor->GetSelection().GetSerialNumber())
-		{
-			CachedValues.CachedSelectionSerialNumber = CurveEditor->GetSelection().GetSerialNumber();
-			CurveCacheFlags = ECurveCacheFlags::All;
-		}
+		TArray<FCurveModelID> ModelIDs;
+		CurveInfoByID.GetKeys(ModelIDs);
 
-		//Only get view values if we need to since we will reset them every time we get all
-		if (CurveCacheFlags != ECurveCacheFlags::All)
-		{
-			if (OutputMin != CachedValues.CachedOutputMin || OutputMax != CachedValues.CachedOutputMax)
-			{
-				CurveCacheFlags = ECurveCacheFlags::All;
-			}
-			else if (CachedValues.CachedGeometrySize != GetCachedGeometry().GetLocalSize())
-			{
-				CurveCacheFlags = ECurveCacheFlags::All;
-			}
-			else
-			{
-				double InputMin = 0, InputMax = 1;
-				GetInputBounds(InputMin, InputMax);
-				if (InputMin != CachedValues.CachedInputMin || InputMax != CachedValues.CachedInputMax)
-				{
-					CurveCacheFlags = ECurveCacheFlags::All;
-				}
-			}
-		}
-		if (CurveCacheFlags == ECurveCacheFlags::All)
-		{
-			CachedValues.CachedOutputMin = OutputMin;
-			CachedValues.CachedOutputMax = OutputMax;
-			GetInputBounds(CachedValues.CachedInputMin, CachedValues.CachedInputMax);
-			CachedValues.CachedGeometrySize = GetCachedGeometry().GetLocalSize();
-
-			CachedDrawParams.Reset();
-			GetCurveDrawParams(CachedDrawParams);
-			CurveCacheFlags = ECurveCacheFlags::CheckCurves;
-			RefreshRetainer();
-		}
-		else if (CurveCacheFlags == ECurveCacheFlags::CheckCurves)
-		{
-			bool bSomethingChanged = false;
-			for (FCurveDrawParams& Params : CachedDrawParams)
-			{
-				FCurveModel* CurveModel = CurveEditor->FindCurve(Params.GetID());
-				if (CurveModel->HasChangedAndResetTest())
-				{
-					bSomethingChanged = true;
-					if (bAllowModelViewTransforms)
-					{
-						bUpdateModelViewTransforms = true;
-						UpdateCustomAxes();
-					}
-					GetCurveDrawParam(CurveEditor, Params.GetID(), CurveModel, Params);
-				}
-			}
-			if (bSomethingChanged)
-			{
-				RefreshRetainer();
-			}
-		}
+		CurveEditor->GetCurveDrawParamsCache().Invalidate(SharedThis(this), ModelIDs);
 	}
 	else
 	{
@@ -826,10 +774,15 @@ void SCurveEditorView::CheckCacheAndInvalidateIfNeeded()
 		CachedValues.CachedOutputMax = OutputMax;
 		GetInputBounds(CachedValues.CachedInputMin, CachedValues.CachedInputMax);
 		CachedValues.CachedGeometrySize = GetCachedGeometry().GetLocalSize();
+	
+		CurveEditor->GetCurveDrawParamsCache().CachedDrawParams.Reset();
 
-		CachedDrawParams.Reset();
-		GetCurveDrawParams(CachedDrawParams);
+		PRAGMA_DISABLE_DEPRECATION_WARNINGS
+		// Using deprecated methods to allow testing if the deprecated API remains functional for the time being.
+		// Can be updated to DrawParamsCache::GetDrawParamSynchronous when GetCurveDrawParams is removed.
+		GetCurveDrawParams(CurveEditor->GetCurveDrawParamsCache().CachedDrawParams);
+		PRAGMA_ENABLE_DEPRECATION_WARNINGS
 
 		RefreshRetainer();
 	}
-}
\ No newline at end of file
+}
diff --git a/Engine/Source/Editor/CurveEditor/Private/SCurveViewerPanel.cpp b/Engine/Source/Editor/CurveEditor/Private/SCurveViewerPanel.cpp
--- a/Engine/Source/Editor/CurveEditor/Private/SCurveViewerPanel.cpp
+++ b/Engine/Source/Editor/CurveEditor/Private/SCurveViewerPanel.cpp
@@ -6,6 +6,7 @@
 #include "Containers/SortedMap.h"
 #include "CurveDrawInfo.h"
 #include "CurveEditor.h"
+#include "CurveEditorCurveDrawParamsCache.h"
 #include "CurveEditorTypes.h"
 #include "Layout/Clipping.h"
 #include "Layout/Geometry.h"
@@ -38,7 +39,8 @@ void SCurveViewerPanel::Construct(const FArguments& InArgs, TSharedRef<FCurveEdi
 	CurveThickness = InArgs._CurveThickness;
 
 	InCurveEditor->SetView(SharedThis(this));
-	CachedValues.CachedTangentVisibility = InCurveEditor->GetSettings()->GetTangentVisibility();
+
+	InCurveEditor->GetCurveDrawParamsCache().Invalidate(SharedThis(this));
 
 	for (const TPair<FCurveModelID, TUniquePtr<FCurveModel>>& CurvePair : InCurveEditor->GetCurves())
 	{
@@ -50,8 +52,7 @@ void SCurveViewerPanel::Construct(const FArguments& InArgs, TSharedRef<FCurveEdi
 
 void SCurveViewerPanel::Tick(const FGeometry& AllottedGeometry, const double InCurrentTime, const float InDeltaTime)
 {
-	CachedDrawParams.Reset();
-	GetCurveDrawParams(CachedDrawParams);
+	CheckCacheAndInvalidateIfNeeded();
 }
 
 int32 SCurveViewerPanel::OnPaint(const FPaintArgs& Args, const FGeometry& AllottedGeometry, const FSlateRect& MyCullingRect, FSlateWindowElementList& OutDrawElements, int32 LayerId, const FWidgetStyle& InWidgetStyle, bool bParentEnabled) const
diff --git a/Engine/Source/Editor/CurveEditor/Private/Views/SInteractiveCurveEditorView.cpp b/Engine/Source/Editor/CurveEditor/Private/Views/SInteractiveCurveEditorView.cpp
--- a/Engine/Source/Editor/CurveEditor/Private/Views/SInteractiveCurveEditorView.cpp
+++ b/Engine/Source/Editor/CurveEditor/Private/Views/SInteractiveCurveEditorView.cpp
@@ -16,6 +16,7 @@
 #include "CurveEditorAxis.h"
 #include "CurveEditorCommands.h"
 #include "CurveEditorContextMenu.h"
+#include "CurveEditorCurveDrawParamsCache.h"
 #include "CurveEditorHelpers.h"
 #include "CurveEditorScreenSpace.h"
 #include "CurveEditorSelection.h"
@@ -121,7 +122,9 @@ void SInteractiveCurveEditorView::Construct(const FArguments& InArgs, TWeakPtr<F
 	{
 		TSharedPtr<FCurveEditor> CurveEditorPin = InCurveEditor.Pin();
 		check(CurveEditorPin);
-		CachedValues.CachedTangentVisibility = CurveEditorPin->GetSettings()->GetTangentVisibility();
+
+		CurveEditorPin->GetCurveDrawParamsCache().Invalidate(SharedThis(this));
+		
 		CurveEditorPin->OnActiveToolChangedDelegate.AddSP(this, &SInteractiveCurveEditorView::OnCurveEditorToolChanged);
 		CurveEditorPin->GetSettings()->OnShowValueIndicatorsChanged().AddSP(this, &SInteractiveCurveEditorView::OnShowValueIndicatorsChanged);
 	}
@@ -796,22 +799,26 @@ void SInteractiveCurveEditorView::DrawCurves(TSharedRef<FCurveEditor> CurveEdito
 	const bool bDrawLines = CVarDrawCurveLines.GetValueOnGameThread();
 	const bool bDrawKeys = CVarDrawCurveKeys.GetValueOnGameThread();
 
+	const TArray<FCurveDrawParams>& DrawParamsArray = CurveEditor->GetCurveDrawParamsCache().GetCurveDrawParams();
+
 	TOptional<FCurveModelID> HoveredCurve = GetHoveredCurve();
-	for (const FCurveDrawParams& Params : CachedDrawParams)
+	for (const FCurveDrawParams& DrawParams : DrawParamsArray)
 	{
-		const bool bIsCurveHovered = HoveredCurve.IsSet() && HoveredCurve.GetValue() == Params.GetID();
-		const float Thickness = bIsCurveHovered ? Params.Thickness+HoverThicknessOffset : Params.Thickness;
+		const FCurveModelID& ModelID = DrawParams.GetID();
+
+		const bool bIsCurveHovered = HoveredCurve.IsSet() && HoveredCurve.GetValue() == ModelID;
+		const float Thickness = bIsCurveHovered ? DrawParams.Thickness+HoverThicknessOffset : DrawParams.Thickness;
 		const int32 CurveLayerId = bIsCurveHovered ? BaseLayerId + CurveViewConstants::ELayerOffset::Curves : BaseLayerId + CurveViewConstants::ELayerOffset::HoveredCurves;
 
-		if (bDrawLines)
+		if (bDrawLines && DrawParams.bDrawInterpolatingPoints)
 		{
-			if (Params.DashLengthPx > 0.f)
+			if (DrawParams.DashLengthPx > 0.f)
 			{
 				float DashOffset = static_cast<float>(GetViewSpace().PixelsPerInput() * GetViewSpace().GetInputMin());
 
 				TArray<FVector2f> NewVector;
-				NewVector.Reserve(Params.InterpolatingPoints.Num());
-				for (FVector2d Vect : Params.InterpolatingPoints)
+				NewVector.Reserve(DrawParams.InterpolatingPoints.Num());
+				for (const FVector2d& Vect : DrawParams.InterpolatingPoints)
 				{
 					NewVector.Add(UE::Slate::CastToVector2f(Vect));
 				}
@@ -822,9 +829,9 @@ void SInteractiveCurveEditorView::DrawCurves(TSharedRef<FCurveEditor> CurveEdito
 					PaintGeometry,
 					MoveTemp(NewVector),
 					DrawEffects,
-					Params.Color,
+					DrawParams.Color,
 					Thickness,
-					Params.DashLengthPx,
+					DrawParams.DashLengthPx,
 					DashOffset
 				);
 			}
@@ -834,9 +841,9 @@ void SInteractiveCurveEditorView::DrawCurves(TSharedRef<FCurveEditor> CurveEdito
 					OutDrawElements,
 					CurveLayerId,
 					PaintGeometry,
-					Params.InterpolatingPoints,
+					DrawParams.InterpolatingPoints,
 					DrawEffects,
-					Params.Color,
+					DrawParams.Color,
 					bAntiAliasCurves,
 					Thickness
 				);
@@ -847,14 +854,19 @@ void SInteractiveCurveEditorView::DrawCurves(TSharedRef<FCurveEditor> CurveEdito
 		LinePoints.SetNum(2);
 
 		// Draw tangents
-		if (bDrawKeys && Params.bKeyDrawEnabled)
+		if (bDrawKeys && DrawParams.bKeyDrawEnabled)
 		{
-			for (int32 PointIndex = 0; PointIndex < Params.Points.Num(); PointIndex++)
+			for (int32 PointIndex = 0; PointIndex < DrawParams.Points.Num(); PointIndex++)
 			{
-				const FCurvePointInfo& Point = Params.Points[PointIndex];
-				const FKeyDrawInfo& PointDrawInfo = Params.GetKeyDrawInfo(Point.Type, PointIndex);
-				const bool          bSelected = CurveEditor->GetSelection().IsSelected(FCurvePointHandle(Params.GetID(), Point.Type, Point.KeyHandle));
-				FLinearColor  PointTint = PointDrawInfo.Tint.IsSet() ? PointDrawInfo.Tint.GetValue() : Params.Color;
+				const FCurvePointInfo& Point = DrawParams.Points[PointIndex];
+				if (!Point.bDraw)
+				{
+					continue;
+				}
+
+				const FKeyDrawInfo& PointDrawInfo = DrawParams.GetKeyDrawInfo(Point.Type, PointIndex);
+				const bool          bSelected = CurveEditor->GetSelection().IsSelected(FCurvePointHandle(DrawParams.GetID(), Point.Type, Point.KeyHandle));
+				FLinearColor  PointTint = PointDrawInfo.Tint.IsSet() ? PointDrawInfo.Tint.GetValue() : DrawParams.Color;
 
 				if (bSelected)
 				{
@@ -1085,11 +1097,13 @@ bool SInteractiveCurveEditorView::GetPointsWithinWidgetRange(const FSlateRect& W
 	FBox WidgetRectangleBox(FVector(WidgetRectangle.Left, WidgetRectangle.Top, 0), FVector(WidgetRectangle.Right, WidgetRectangle.Bottom, 0));
 
 	const float PointOverlapSensitivity = CurveEditor->GetSettings()->GetMarqueePointSensitivity();
-	
+
+	const TArray<FCurveDrawParams>& DrawParamsArray = CurveEditor->GetCurveDrawParamsCache().GetCurveDrawParams();
+
 	// Iterate through all of our points and see which points the marquee overlaps. Both of these coordinate systems
 	// are in screen space pixels.  Also check tangent lines
 	bool bFound = false;
-	for (const FCurveDrawParams& DrawParams : CachedDrawParams)
+	for (const FCurveDrawParams& DrawParams : DrawParamsArray)
 	{
 		for (int32 PointIndex = 0; PointIndex < DrawParams.Points.Num(); PointIndex++)
 		{
@@ -1135,10 +1149,12 @@ bool SInteractiveCurveEditorView::GetCurveWithinWidgetRange(const FSlateRect& Wi
 
 	FBox WidgetRectangleBox(FVector(WidgetRectangle.Left, WidgetRectangle.Top, 0), FVector(WidgetRectangle.Right, WidgetRectangle.Bottom, 0));
 
+	const TArray<FCurveDrawParams>& DrawParamsArray =CurveEditor->GetCurveDrawParamsCache().GetCurveDrawParams();
+
 	// Iterate through all of our interpolating points and terminates if one overlaps the marquee. Both of these coordinate systems
 	// are in screen space pixels.
 	TSet<FCurveModelID> CurveIDs;
-	for (const FCurveDrawParams& DrawParams : CachedDrawParams)
+	for (const FCurveDrawParams& DrawParams : DrawParamsArray)
 	{
 		for (int32 InterpolatingPointIndex = 1; InterpolatingPointIndex < DrawParams.InterpolatingPoints.Num(); InterpolatingPointIndex++)
 		{
@@ -1256,7 +1272,7 @@ void SInteractiveCurveEditorView::UpdateCurveProximities(FVector2D MousePixel)
 
 	if (PreviouslyHovered != NewHovered)
 	{
-		RefreshRetainer();
+		bNeedsRefresh = true;
 	}
 }
 
@@ -1303,12 +1319,8 @@ FReply SInteractiveCurveEditorView::OnMouseMove(const FGeometry& MyGeometry, con
 		return FReply::Unhandled();
 	}
 
-	// Update our Curve Proximities for hover states and context actions. This also updates our cached hovered curve.
-	FVector2D MousePixel = MyGeometry.AbsoluteToLocal(MouseEvent.GetScreenSpacePosition());
-	UpdateCurveProximities(MousePixel);
-
 	// Cache the mouse position so that commands such as add key can work from command bindings 
-	CachedMousePosition = MousePixel;
+	CachedMousePosition = MyGeometry.AbsoluteToLocal(MouseEvent.GetScreenSpacePosition());
 
 	if (DragOperation.IsSet())
 	{
@@ -1316,13 +1328,13 @@ FReply SInteractiveCurveEditorView::OnMouseMove(const FGeometry& MyGeometry, con
 
 		if (!DragOperation->IsDragging() && DragOperation->AttemptDragStart(MouseEvent))
 		{
-			DragOperation->DragImpl->BeginDrag(InitialPosition, MousePixel, MouseEvent);
+			DragOperation->DragImpl->BeginDrag(InitialPosition, CachedMousePosition, MouseEvent);
 			return FReply::Handled().CaptureMouse(AsShared());
 		}
 		else if (DragOperation->IsDragging())
 		{
 			bHadMouseMovesThisTick = true;
-			DragOperation->DragImpl->Drag(InitialPosition, MousePixel, MouseEvent);
+			DragOperation->DragImpl->Drag(InitialPosition, CachedMousePosition, MouseEvent);
 		}
 		return FReply::Handled();
 	}
@@ -1333,6 +1345,9 @@ FReply SInteractiveCurveEditorView::OnMouseMove(const FGeometry& MyGeometry, con
 
 void SInteractiveCurveEditorView::OnFinishedPointerInput()
 {
+	// Update our Curve Proximities for hover states and context actions. This also updates our cached hovered curve.
+	UpdateCurveProximities(CachedMousePosition);
+
 	// Some operations defer processing Drag calls for performance reasons. Give them a chance to process the accumulated input.
 	if (bHadMouseMovesThisTick && DragOperation && DragOperation->IsDragging())
 	{
@@ -1888,24 +1903,35 @@ void SInteractiveCurveEditorView::BuildContextMenu(FMenuBuilder& MenuBuilder, TO
 
 TOptional<FCurvePointHandle> SInteractiveCurveEditorView::HitPoint(FVector2D MousePixel) const
 {
+	if (!WeakCurveEditor.IsValid())
+	{
+		return TOptional<FCurvePointHandle>();
+	}
+	const TSharedRef<FCurveEditor> CurveEditor = WeakCurveEditor.Pin().ToSharedRef();
+
+
 	TOptional<FCurvePointHandle> HitPoint;
 	TOptional<float> ClosestDistance;
 
 	TOptional<FCurveModelID> HoveredCurve = GetHoveredCurve();
 
+	const TArray<FCurveDrawParams>& DrawParamsArray = CurveEditor->GetCurveDrawParamsCache().GetCurveDrawParams();
+
 	// Find all keys within the current hit test time
-	for (const FCurveDrawParams& Params : CachedDrawParams)
+	for (const FCurveDrawParams& DrawParams : DrawParamsArray)
 	{
+		const FCurveModelID& ModelID = DrawParams.GetID();
+
 		// If we have a hovered curve, only hit a point within that curve
-		if (HoveredCurve.IsSet() && Params.GetID() != HoveredCurve.GetValue())
+		if (HoveredCurve.IsSet() && ModelID != HoveredCurve.GetValue())
 		{
 			continue;
 		}
 	
-		for (int32 PointIndex = 0; PointIndex < Params.Points.Num(); PointIndex++)
+		for (int32 PointIndex = 0; PointIndex < DrawParams.Points.Num(); PointIndex++)
 		{
-			const FCurvePointInfo& Point = Params.Points[PointIndex];
-			const FKeyDrawInfo& PointDrawInfo = Params.GetKeyDrawInfo(Point.Type, PointIndex);
+			const FCurvePointInfo& Point = DrawParams.Points[PointIndex];
+			const FKeyDrawInfo& PointDrawInfo = DrawParams.GetKeyDrawInfo(Point.Type, PointIndex);
 
 			// We artificially inflate the hit testing region for keys by a few pixels to make them easier to hit. The PointDrawInfo.ScreenSize specifies their drawn size,
 			// so we need to inflate here when doing the actual hit testing. We subtract by half the extent to center it on the drawing.
@@ -1919,7 +1945,7 @@ TOptional<FCurvePointHandle> SInteractiveCurveEditorView::HitPoint(FVector2D Mou
 				if (DistanceSquared <= ClosestDistance.Get(DistanceSquared))
 				{
 					ClosestDistance = DistanceSquared;
-					HitPoint = FCurvePointHandle(Params.GetID(), Point.Type, Point.KeyHandle);
+					HitPoint = FCurvePointHandle(DrawParams.GetID(), Point.Type, Point.KeyHandle);
 				}
 			}
 		}
diff --git a/Engine/Source/Editor/CurveEditor/Public/CurveDataAbstraction.h b/Engine/Source/Editor/CurveEditor/Public/CurveDataAbstraction.h
--- a/Engine/Source/Editor/CurveEditor/Public/CurveDataAbstraction.h
+++ b/Engine/Source/Editor/CurveEditor/Public/CurveDataAbstraction.h
@@ -155,6 +155,9 @@ struct FKeyAttributes
 		return NewAttributes;
 	}
 
+	bool operator==(const FKeyAttributes& Other) const = default;
+	bool operator!=(const FKeyAttributes& Other) const = default;
+
 private:
 
 	/** True if this key supports entry tangents */
diff --git a/Engine/Source/Editor/CurveEditor/Public/CurveDrawInfo.h b/Engine/Source/Editor/CurveEditor/Public/CurveDrawInfo.h
--- a/Engine/Source/Editor/CurveEditor/Public/CurveDrawInfo.h
+++ b/Engine/Source/Editor/CurveEditor/Public/CurveDrawInfo.h
@@ -14,7 +14,7 @@ struct FSlateBrush;
 /**
  * Structure that defines how to draw a particular key
  */
-struct FKeyDrawInfo
+struct CURVEEDITOR_API FKeyDrawInfo
 {
 	FKeyDrawInfo()
 		: ScreenSize(0.f, 0.f)
@@ -36,7 +36,7 @@ struct FKeyDrawInfo
 /**
  * Structure that defines the necessary data for painting a given curve point
  */
-struct FCurvePointInfo
+struct CURVEEDITOR_API FCurvePointInfo
 {
 	FCurvePointInfo(FKeyHandle InHandle)
 		: KeyHandle(InHandle)
@@ -58,13 +58,15 @@ struct FCurvePointInfo
 	int32 LayerBias;
 	/** DrawInfo to use when rendering this point. Allows customization of points on a per-point basis instead of per-curve. */
 	FKeyDrawInfo DrawInfo;
+	/** If false, the key will not be painted. */
+	bool bDraw = true;
 };
 
 
 /**
  * Structure that defines the necessary data for painting a whole curve
  */
-struct FCurveDrawParams
+struct CURVEEDITOR_API FCurveDrawParams
 {
 	/** The color to draw this curve */
 	FLinearColor Color;
@@ -90,6 +92,9 @@ struct FCurveDrawParams
 	/** Whether or not to draw keys on the curve. */
 	uint8 bKeyDrawEnabled : 1;
 
+	/** If false, the interpolating points will not be painted. */
+	uint8 bDrawInterpolatingPoints : 1 = true;
+
 	/**
 	 * Construct new draw parameters for the specified curve ID
 	 */
diff --git a/Engine/Source/Editor/CurveEditor/Public/CurveEditor.h b/Engine/Source/Editor/CurveEditor/Public/CurveEditor.h
--- a/Engine/Source/Editor/CurveEditor/Public/CurveEditor.h
+++ b/Engine/Source/Editor/CurveEditor/Public/CurveEditor.h
@@ -53,6 +53,7 @@ struct FCurveDrawParams;
 struct FCurveEditorInitParams;
 struct FCurveEditorSnapMetrics;
 struct FGeometry;
+namespace UE::CurveEditor { class FCurveDrawParamsCache; }
 
 DECLARE_DELEGATE_OneParam(FOnSetBoolean, bool)
 DECLARE_MULTICAST_DELEGATE_OneParam(FOnActiveToolChanged, FCurveEditorToolID)
@@ -358,6 +359,9 @@ public:
 	/** Returns whether the buffered curve is to be acted on, ie. selected, in the tree view or with selected keys */
 	bool IsActiveBufferedCurve(const TUniquePtr<IBufferedCurveModel>& BufferedCurve) const;
 
+	/** Returns a cache for curve draw params */
+	UE::CurveEditor::FCurveDrawParamsCache& GetCurveDrawParamsCache() { return *CurveDrawParamsCache; }
+
 	// ~FCurveEditor
 
 	// FEditorUndoClient
@@ -686,7 +690,7 @@ protected:
 	* which allows Normalized views to push past their boundaries without the normalization ratio changing per-frame as you drag.
 	*/
 	bool bBoundTransformUpdatesSuppressed;
-	
+
 	/** Track which axis UI movements should be snapped to (where applicable) based on limitations imposed by the UI. */
 	FCurveEditorAxisSnap AxisSnapMetrics;
 
@@ -700,6 +704,8 @@ protected:
 	int32 SuspendBroadcastCount;
 
 private:
+	/** Cache for curve draw params */
+	TSharedPtr<UE::CurveEditor::FCurveDrawParamsCache> CurveDrawParamsCache;
 
 	/** Cached physical size of the panel representing this editor */
 	FVector2D CachedPhysicalSize;
diff --git a/Engine/Source/Editor/CurveEditor/Public/CurveEditorCurveDrawParamsCache.h b/Engine/Source/Editor/CurveEditor/Public/CurveEditorCurveDrawParamsCache.h
new file mode 100644
--- /dev/null
+++ b/Engine/Source/Editor/CurveEditor/Public/CurveEditorCurveDrawParamsCache.h
@@ -0,0 +1,115 @@
+// Copyright Epic Games, Inc. All Rights Reserved.
+
+#pragma once
+
+#include "CurveDrawInfo.h"
+#include "CurveEditorSettings.h"
+#include "Templates/SharedPointer.h"
+
+class FCurveEditor;
+class FCurveModel;
+class SCurveEditorView;
+struct FCurveDrawParams;
+struct FCurveEditorScreenSpace;
+struct FCurveInfo;
+struct FCurveModelID;
+
+namespace UE::CurveEditor
+{		
+	struct FCurveDrawParamsHandle;
+
+	/** Caches curve draw params for curves of a curve editor. */
+	class CURVEEDITOR_API FCurveDrawParamsCache
+		: public TSharedFromThis<FCurveDrawParamsCache>
+	{
+		// Let SCurveEditorView write directly to the array of cached curve params when CVar CurveEditor.UseCurveCache is disabled
+		friend class ::SCurveEditorView;
+
+		// Allow draw params handle to reference draw params
+		friend struct FCurveDrawParamsHandle;
+
+	public:
+		/** Flag enum signifying how the curve cache has changed since it was last generated	 */
+		enum class ECurveCacheFlags : uint8
+		{
+			CheckCurves = 0,       // The cache may be valid need to check each curve to see if they are still valid
+			UpdateAll = 1 << 0,	   // Update all curves
+		};
+
+		/** Set of cached values we need to check each tick to see if we need to redo cache */
+		struct FCachedCurveEditorData
+		{
+			/** Serial number cached from FCurveEditor::GetActiveCurvesSerialNumber() on tick */
+			uint32 ActiveCurvesSerialNumber = 0;
+
+			/** Serial number bached from CurveEditorSelecstion::GetSerialNumber */
+			uint32 SelectionSerialNumber = 0;
+
+			/** Cached Tangent Visibility*/
+			ECurveEditorTangentVisibility TangentVisibility = ECurveEditorTangentVisibility::NoTangents;
+
+			/** Cached input min value */
+			double InputMin = 0.0;
+			
+			/** Cached input max value */
+			double InputMax = 0.0;
+			
+			/** Cached output min value */
+			double OutputMin = 0.0;
+			
+			/** Cached output max value */
+			double OutputMax = 0.0;
+
+			/** Cached Geometry Size */
+			FVector2D CachedGeometrySize = FVector2D::ZeroVector;
+		};
+
+		FCurveDrawParamsCache(const TSharedRef<FCurveEditor>& InCurveEditor);
+
+		/** 
+		 * Invalidates the cache for specified view and curve model IDs and causes an update on cached data.
+		 * The cache may perform updates that are not instantly relevant asynchronously. To update data synchronous use UpdateCurveDrawParamSynchonous instead.
+		 */
+		void Invalidate(const TSharedRef<SCurveEditorView>& CurveEditorView, const TArray<FCurveModelID>& ModelIDs = TArray<FCurveModelID>());
+
+		/**
+		 * Returns how the curve cache has changed since it was last generated.
+		 * Note for a data change it may only effect certain data (curves) not every drawn curve
+		 */
+		ECurveCacheFlags GetCurveCacheFlags() const { return CurveCacheFlags; }
+
+		/** Returns cached data common to the curve editor */
+		const FCachedCurveEditorData& GetCurveEditorData() const { return CachedCurveEditorData; }
+
+		/** Returns the current draw params */
+		const TArray<FCurveDrawParams>& GetCurveDrawParams() const { return CachedDrawParams; }
+
+		/** Updates curve draw params for all specified curve models, synchronous. */
+		void UpdateAllCurveDrawParamSynchonous(const TSharedRef<SCurveEditorView>& CurveEditorView, const TArray<FCurveModelID>& CurveModelIDs, TArray<FCurveDrawParams>& OutParams);
+
+		/** Gets curve draw params for a curve model, synchronous. */
+		void UpdateCurveDrawParamsSynchonous(const FCurveEditorScreenSpace& CurveSpace, FCurveModel* CurveModel, const FCurveModelID& ModelID, FCurveDrawParams& Params);
+
+	private:
+		/** Updates the curve cache flags */
+		void UpdateCurveCacheFlags(const TSharedRef<SCurveEditorView>& CurveEditorView);
+
+		/** Draws the curves, possibly async */
+		void DrawCurves(const TSharedRef<SCurveEditorView>& CurveEditorView, const TArray<FCurveModelID>& ModelIDs = TArray<FCurveModelID>());
+
+		/** Updates curve draw params, possibly async */
+		void UpdateCurveDrawParams(const TSharedRef<SCurveEditorView>& CurveEditorView, const TArray<FCurveDrawParamsHandle>& CurveDrawParamsHandles);
+
+		/** Curve cache flags that change based upon data or view getting modified*/
+		ECurveCacheFlags CurveCacheFlags = ECurveCacheFlags::UpdateAll;
+
+		/** The actual cached draw params */
+		TArray<FCurveDrawParams> CachedDrawParams;
+
+		/** Cached data common to the curve editor */
+		FCachedCurveEditorData CachedCurveEditorData;
+
+		/** The curve editor that owns this cache */
+		const TWeakPtr<FCurveEditor> WeakCurveEditor;
+	};
+}
diff --git a/Engine/Source/Editor/CurveEditor/Public/CurveEditorCurveDrawParamsHandle.h b/Engine/Source/Editor/CurveEditor/Public/CurveEditorCurveDrawParamsHandle.h
new file mode 100644
--- /dev/null
+++ b/Engine/Source/Editor/CurveEditor/Public/CurveEditorCurveDrawParamsHandle.h
@@ -0,0 +1,41 @@
+// Copyright Epic Games, Inc. All Rights Reserved.
+
+#pragma once
+
+#include "CurveEditorTypes.h"
+
+class FCurveDrawParamsCache;
+struct FCurveDrawParams;
+
+namespace UE::CurveEditor
+{
+	class FCurveDrawParamsCache;
+
+	/** Handle for an entry in the draw params array */
+	struct CURVEEDITOR_API FCurveDrawParamsHandle
+	{
+		FCurveDrawParamsHandle() = default;
+
+		/** Constructs the handle from the index in curve draw params. */
+		FCurveDrawParamsHandle(const TSharedRef<FCurveDrawParamsCache>& InDrawParamsCache, const int32 InIndex);
+
+		/** 
+		 * Returns the draw params or nullptr if the draw params no longer exists. 
+		 * Does a slow lookup if the array changed.
+		 */
+		FCurveDrawParams* Get() const;
+
+		/** Returns the curve model ID that corresponds to these curve draw params. */
+		const FCurveModelID& GetID() const { return CurveModelID; }
+
+	private:
+		/** The curve model ID that corresponds to the curve draw params */
+		FCurveModelID CurveModelID;
+
+		/** Index of the draw params in the array of the draw params cache */
+		mutable int32 Index = INDEX_NONE;
+
+		/** Weak ptr to the draw params cache that owns the draw params */
+		TWeakPtr<class FCurveDrawParamsCache> WeakDrawParamsCache;
+	};
+}
diff --git a/Engine/Source/Editor/CurveEditor/Public/CurveEditorHelpers.h b/Engine/Source/Editor/CurveEditor/Public/CurveEditorHelpers.h
--- a/Engine/Source/Editor/CurveEditor/Public/CurveEditorHelpers.h
+++ b/Engine/Source/Editor/CurveEditor/Public/CurveEditorHelpers.h
@@ -14,11 +14,11 @@ struct FCurveEditorScreenSpace;
 namespace CurveEditor
 {
 
-	FVector2D ComputeScreenSpaceTangentOffset(const FCurveEditorScreenSpace& CurveSpace, float Tangent, float Weight);
+	CURVEEDITOR_API FVector2D ComputeScreenSpaceTangentOffset(const FCurveEditorScreenSpace& CurveSpace, float Tangent, float Weight);
 
 	void TangentAndWeightFromOffset(const FCurveEditorScreenSpace& CurveSpace, const FVector2D& TangentOffset, float& OutTangent, float& OutWeight);
 
-	FVector2D GetVectorFromSlopeAndLength(float Slope, float Length);
+	CURVEEDITOR_API FVector2D GetVectorFromSlopeAndLength(float Slope, float Length);
 
 	void PopulateGridLineValues(float PhysicalSize, double ViewMin, double ViewMax, uint8 InMinorDivisions, TArray<double>& OutMajorGridLines, TArray<double>& OutMinorGridLines);
 
diff --git a/Engine/Source/Editor/CurveEditor/Public/CurveEditorScreenSpace.h b/Engine/Source/Editor/CurveEditor/Public/CurveEditorScreenSpace.h
--- a/Engine/Source/Editor/CurveEditor/Public/CurveEditorScreenSpace.h
+++ b/Engine/Source/Editor/CurveEditor/Public/CurveEditorScreenSpace.h
@@ -194,4 +194,20 @@ public:
 
 		return New;
 	}
-};
\ No newline at end of file
+
+	bool operator==(const FCurveEditorScreenSpace& Other) const
+	{
+		return
+			GetInputMax() == Other.GetInputMin() &&
+			GetInputMax() == Other.GetInputMax() &&
+			GetPhysicalHeight() == Other.GetPhysicalHeight() &&
+			GetOutputMin() == Other.GetOutputMin() &&
+			GetOutputMax() == Other.GetOutputMax() &&
+			GetPhysicalWidth() == Other.GetPhysicalWidth();
+	}
+
+	bool operator !=(const FCurveEditorScreenSpace& Other)
+	{
+		return !(*this == Other);
+	}
+};
diff --git a/Engine/Source/Editor/CurveEditor/Public/CurveModel.h b/Engine/Source/Editor/CurveEditor/Public/CurveModel.h
--- a/Engine/Source/Editor/CurveEditor/Public/CurveModel.h
+++ b/Engine/Source/Editor/CurveEditor/Public/CurveModel.h
@@ -34,11 +34,18 @@ struct FCurveModelID;
 struct FKeyAttributes;
 struct FKeyDrawInfo;
 struct FKeyHandle;
+struct FKeyHandleSet;
 struct FKeyPosition;
 
 enum class ECurveEditorAxisOrientation : uint8;
 enum class ECurvePointType : uint8;
 
+namespace UE::CurveEditor 
+{ 
+	class ICurveEditorCurveCachePool;
+	struct FCurveDrawParamsHandle;
+}
+
 /**
  * Class that models an underlying curve data structure through a generic abstraction that the curve editor understands.
  */
@@ -74,6 +81,23 @@ public:
 	 */
 	virtual void DrawCurve(const FCurveEditor& CurveEditor, const FCurveEditorScreenSpace& ScreenSpace, TArray<TTuple<double, double>>& InterpolatingPoints) const = 0;
 
+	/**
+	 * Draws curve draw params to a cache pool. Curves can implement this to improve performance over DrawCurve by caching their data.
+	 *
+	 * @param WeakCurveEditor			Weak reference to the curve editor that is drawing the curve. 
+	 * @param CurveDrawParamsHandle		A handle to the actual curve draw params.
+	 * @param ScreenSpace				The screenspace of this curve
+	 * 
+	 * @return							Returns the cache pool or nullptr if the curve model does not implement a cache pool.
+	 */
+	virtual UE::CurveEditor::ICurveEditorCurveCachePool* DrawCurveToCachePool(
+		const TSharedRef<FCurveEditor>& CurveEditor,
+		const UE::CurveEditor::FCurveDrawParamsHandle& CurveDrawParamsHandle,
+		const FCurveEditorScreenSpace& ScreenSpace) 
+	{ 
+		return nullptr; 
+	}
+
 	UE_DEPRECATED(5.5, "Implement the version that doesn't take the Curve Editor parameter.")
 	virtual void GetKeys(const FCurveEditor& CurveEditor, double MinTime, double MaxTime, double MinValue, double MaxValue, TArray<FKeyHandle>& OutKeyHandles) const
 	{
diff --git a/Engine/Source/Editor/CurveEditor/Public/ICurveEditorCurveCachePool.h b/Engine/Source/Editor/CurveEditor/Public/ICurveEditorCurveCachePool.h
new file mode 100644
--- /dev/null
+++ b/Engine/Source/Editor/CurveEditor/Public/ICurveEditorCurveCachePool.h
@@ -0,0 +1,18 @@
+// Copyright Epic Games, Inc. All Rights Reserved.
+
+#pragma once
+
+#include "CurveEditorTypes.h"
+
+class FCurveEditor;
+struct FCurveDrawParams;
+
+namespace UE::CurveEditor
+{
+	class CURVEEDITOR_API ICurveEditorCurveCachePool
+	{
+	public:
+		/** Draws the cached curves. Called when all curves were drawn to the curve cache pool. */
+		virtual void DrawCachedCurves(TWeakPtr<const FCurveEditor> WeakCurveEditor) = 0;
+	};
+}
diff --git a/Engine/Source/Editor/CurveEditor/Public/SCurveEditorView.h b/Engine/Source/Editor/CurveEditor/Public/SCurveEditorView.h
--- a/Engine/Source/Editor/CurveEditor/Public/SCurveEditorView.h
+++ b/Engine/Source/Editor/CurveEditor/Public/SCurveEditorView.h
@@ -31,7 +31,6 @@ class SCurveEditorView;
 
 enum class ECurveEditorAxisOrientation : uint8;
 
-
 /**
  * Identifier for a specific axis on a view. These identifiers are transient and should not be stored externally.
  */
@@ -327,17 +326,21 @@ public:
 	virtual void GetGridLinesX(TSharedRef<const FCurveEditor> CurveEditor, TArray<float>& MajorGridLines, TArray<float>& MinorGridLines, TArray<FText>* MajorGridLabels = nullptr) const {}
 	virtual void GetGridLinesY(TSharedRef<const FCurveEditor> CurveEditor, TArray<float>& MajorGridLines, TArray<float>& MinorGridLines, TArray<FText>* MajorGridLabels = nullptr) const {}
 
-protected:
+	/** Request a new render from the retainer widget */
+	void RefreshRetainer();
 
+	/** Update the custom axes if necessary */
+	void UpdateCustomAxes();
+
+protected:
 	/** Gets info about the curves being drawn. Converts actual curves into an abstract series of lines/points/handles/etc. */
+	UE_DEPRECATED(5.6, "Instead  get the cache by calling FCurveEditor::GetCurveDrawParamsCache, use FCurveDrawParamsCache::GetCurveDrawParamsSynchronous to get Curve Draw Parms as with this function.")
 	void GetCurveDrawParams(TArray<FCurveDrawParams>& OutDrawParams);
 
-	/** Get it for just one curve*/
+	/** Gets curve draw params  get the cache by calling a single curve */
+	UE_DEPRECATED(5.6, "Instead use FCurveEditor::GetCurveDrawParamsCache, use FCurveDrawParamsCache::GetCurveDrawParamSynchronous to get Curve Draw Parm as with this function.")
 	void GetCurveDrawParam(TSharedPtr<FCurveEditor>& CurveEditor, const FCurveModelID& ModelID, FCurveModel* CurveModel, FCurveDrawParams& OutDrawParam) const;
 
-	/** Request a new render from the retainer widget */
-	void RefreshRetainer();
-
 	/** Update all the curve view transforms from curve models */
 	void UpdateCurveViewTransformsFromModels();
 
@@ -356,9 +359,6 @@ private:
 	/** Add a curve from this view */
 	void RemoveCurve(FCurveModelID CurveID);
 
-	/** Update the custom axes if necessary */
-	void UpdateCustomAxes();
-
 public:
 
 	/** (Default: false) When true, this view has been pinned to the top of the view container. Only takes affect from the next call to SCurveEditorPanel::RebuildCurveViews. */
@@ -443,20 +443,25 @@ protected:
 	}
 
 	/** Flag enum signifying how the curve cache has changed since it was last generated
-	Note for a data change it may only effect certain data(curves) not every drawn curve*/
-	enum class ECurveCacheFlags : uint8
+	Note for a data change it may only effect certain data(curves) not every drawn curve*/	
+	enum class UE_DEPRECATED(5.6, "Instead  get the cache by calling FCurveEditor::GetCurveDrawParamsCache, call FCurveDrawParamsCache::ECurveCacheFlags to get a corresponding type.") ECurveCacheFlags : uint8
 	{
 		CheckCurves = 0,       // The cache may be valid need to check each curve to see if they are still valid
 		All = 1 << 0,		   // Get all
 	};
+
 	/** Curve cache flags that change based upon data or view getting modified*/
+	UE_DEPRECATED(5.6, "Instead get the cache by calling FCurveEditor::GetCurveDrawParamsCache, call FCurveDrawParamsCache::GetCurveCacheFlags to get current flags. Use FCurveDrawParamsCache::Update to update them.")
+	PRAGMA_DISABLE_DEPRECATION_WARNINGS
 	ECurveCacheFlags CurveCacheFlags;
+	PRAGMA_ENABLE_DEPRECATION_WARNINGS
 
 	/** Curve draw parameters that are re-generated on tick if the cache has changed. We generate them once and then they're used in multiple places per frame. */
+	UE_DEPRECATED(5.5, "Instead  get the cache by calling FCurveEditor::GetCurveDrawParamsCache, call GetCurveDrawParamsSynchronous.")
 	TArray<FCurveDrawParams> CachedDrawParams;
 
 	/** Set of Cached values we need to check each tick to see if we need to redo cache*/
-	struct FCachedValuesToCheck
+	struct UE_DEPRECATED(5.6, "Instead  get the cache by calling FCurveEditor::GetCurveDrawParamsCache, call FCurveDrawParamsCache::FCachedCurveEditorData to get a corresponding type.") FCachedValuesToCheck
 	{
 		/** Serial number cached from FCurveEditor::GetActiveCurvesSerialNumber() on tick */
 		uint32 CachedActiveCurvesSerialNumber;
@@ -474,8 +479,11 @@ protected:
 		FVector2D CachedGeometrySize;
 	};
 
+	UE_DEPRECATED(5.6, "Instead  get the cache by calling FCurveEditor::GetCurveDrawParamsCache, use FCurveDrawParamsCache::GetCurveEditorData to get cached values. Use FCurveDrawParamsCache::Update to update them.")
+	PRAGMA_DISABLE_DEPRECATION_WARNINGS
 	FCachedValuesToCheck CachedValues;
-
+	PRAGMA_ENABLE_DEPRECATION_WARNINGS
+	
 	/** Possible pointer to a retainer widget that we may need to force update*/
 	TSharedPtr<SRetainerWidget> RetainerWidget;
 
@@ -484,4 +492,4 @@ public:
 	{
 		RetainerWidget = InWidget;
 	}
-};
\ No newline at end of file
+};
diff --git a/Engine/Source/Editor/MovieSceneTools/Private/Cache/MovieSceneAsyncCurvePainter.cpp b/Engine/Source/Editor/MovieSceneTools/Private/Cache/MovieSceneAsyncCurvePainter.cpp
new file mode 100644
--- /dev/null
+++ b/Engine/Source/Editor/MovieSceneTools/Private/Cache/MovieSceneAsyncCurvePainter.cpp
@@ -0,0 +1,215 @@
+// Copyright Epic Games, Inc. All Rights Reserved.
+
+#include "MovieSceneAsyncCurvePainter.h"
+
+#include "Async/Async.h"
+#include "Async/ParallelFor.h"
+#include "Cache/MovieSceneCachedCurve.h"
+#include "Channels/MovieSceneDoubleChannel.h"
+#include "Channels/MovieSceneFloatChannel.h"
+#include "Channels/MovieScenePiecewiseCurve.h"
+#include "CurveEditorScreenSpace.h"
+
+namespace UE::MovieSceneTools
+{
+	template class FMovieSceneAsyncCurvePainter<FMovieSceneFloatChannel>;
+	template class FMovieSceneAsyncCurvePainter<FMovieSceneDoubleChannel>;
+
+	template <typename ChannelType>
+	FMovieSceneAsyncCurvePainter<ChannelType>::FCurvePainterTask::FCurvePainterTask(
+		const TSharedRef<FMovieSceneCachedCurve<ChannelType>>& CachedCurve,
+		const TFunction<void(TArray<FVector2D>, TArray<int32>)>& InCallback)
+		: Callback(InCallback)
+		, ScreenSpace(CachedCurve->GetScreenSpace())
+		, TickResolution(CachedCurve->GetTickResolution())
+		, TimeThreshold(CachedCurve->GetTimeThreshold())	
+		, ValueThreshold(CachedCurve->GetValueThreshold())
+		, PiecewiseCurve(CachedCurve->GetPiecewiseCurve())
+	{
+		KeyPointsSource = CachedCurve->GetKeyPoints();
+		InterpolatingPointsTarget = KeyPointsSource;
+	}
+
+	template <typename ChannelType>
+	void FMovieSceneAsyncCurvePainter<ChannelType>::FCurvePainterTask::SetFlags(ECurvePainterTaskStateFlags NewFlags)
+	{
+		StateFlags.store(NewFlags);
+	}
+
+	template <typename ChannelType>
+	bool FMovieSceneAsyncCurvePainter<ChannelType>::FCurvePainterTask::HasAnyFlags(ECurvePainterTaskStateFlags Flags) const
+	{
+		return EnumHasAnyFlags(StateFlags.load(), Flags);
+	}
+
+	template <typename ChannelType>
+	void FMovieSceneAsyncCurvePainter<ChannelType>::PaintCurve(
+		const TSharedRef<FMovieSceneCachedCurve<ChannelType>>& CachedCurve, 
+		const bool bSynchronous,
+		const TFunction<void(TArray<FVector2D> /** OutInterpolatedPoints */, TArray<int32> /** OutKeyOffsets */)>& InCallback)
+	{
+		static TSharedRef<FMovieSceneAsyncCurvePainter<ChannelType>> Instance = []() 
+			{
+				const TSharedRef<FMovieSceneAsyncCurvePainter<ChannelType>> NewInstance = MakeShared<FMovieSceneAsyncCurvePainter<ChannelType>>();
+				
+				// Core delegates have to be bound on the game thread to avoid any concurrency
+				AsyncTask(ENamedThreads::GameThread,
+					[NewInstance]()
+					{
+						FCoreDelegates::OnEndFrame.AddSP(NewInstance, &FMovieSceneAsyncCurvePainter<ChannelType>::OnEndFrame);
+					});
+
+				return NewInstance;
+			}(); 
+
+		Instance->AddTask(CachedCurve, bSynchronous, InCallback);
+	}
+
+
+	template <typename ChannelType>
+	void FMovieSceneAsyncCurvePainter<ChannelType>::AddTask(
+		const TSharedRef<FMovieSceneCachedCurve<ChannelType>>& CachedCurve,
+		const bool bSynchronous,
+		const TFunction<void(TArray<FVector2D> /** OutInterpolatedPoints */, TArray<int32> /** OutKeyOffsets */)>& InCallback)
+	{
+		const FScopeLock Lock(&AccessTasksCritSec);
+
+		const TSharedRef<FCurvePainterTask> Task = MakeShared<FCurvePainterTask>(CachedCurve, InCallback);
+		if (bSynchronous)
+		{
+			TRACE_CPUPROFILER_EVENT_SCOPE(FMovieSceneAsyncCurvePainter DrawInterpolatiedPointsSynchronous);
+			int32 OldSize;
+			do
+			{
+				OldSize = Task->InterpolatingPointsTarget.Num();
+				RefineCurvePoints(Task);
+			} while (OldSize != Task->InterpolatingPointsTarget.Num());
+		}
+		else
+		{
+			TSharedPtr<FCurvePainterTask>* RunningTaskPtr = CachedCurveToTaskMap.Find(CachedCurve);
+			if (RunningTaskPtr && (*RunningTaskPtr).IsValid())
+			{
+				// Let a currently running task know he's now void
+				(*RunningTaskPtr)->SetFlags(ECurvePainterTaskStateFlags::Void);
+				// Set the new task
+				(*RunningTaskPtr) = Task;
+			}
+			else
+			{
+				// Add the new task
+				CachedCurveToTaskMap.Add(CachedCurve, Task);
+			}
+		}
+	}
+
+	template <typename ChannelType>
+	void FMovieSceneAsyncCurvePainter<ChannelType>::OnEndFrame()
+	{
+		const FScopeLock Lock(&AccessTasksCritSec);
+
+		// Avoid any concurrency while work is ongoing
+		if (bWorking)
+		{
+			return;
+		}
+
+		// Remove completed and void tasks
+		for (typename TMap<TSharedPtr<FMovieSceneCachedCurve<ChannelType>>, TSharedPtr<FCurvePainterTask>>::TIterator It(CachedCurveToTaskMap); It; ++It)
+		{
+			if ((*It).Value->HasAnyFlags(ECurvePainterTaskStateFlags::Completed | ECurvePainterTaskStateFlags::Void))
+			{
+				It.RemoveCurrent();
+			}
+		}
+		if (CachedCurveToTaskMap.IsEmpty())
+		{
+			return;
+		}
+		TArray<TSharedPtr<FCurvePainterTask>> Tasks;
+		CachedCurveToTaskMap.GenerateValueArray(Tasks);
+
+		UE::Tasks::Launch(UE_SOURCE_LOCATION,
+			[Tasks, this]()
+			{
+				TRACE_CPUPROFILER_EVENT_SCOPE(FMovieSceneAsyncCurvePainter DrawInterpolatiedPointsAsync);
+				
+				bWorking = true;
+				
+				ParallelFor(Tasks.Num(),
+					[&Tasks, this](int32 Index)
+					{
+						RefineCurvePoints(Tasks[Index].ToSharedRef());
+					});
+
+				bWorking = false;
+			}, LowLevelTasks::ETaskPriority::BackgroundNormal);
+	}
+
+	template <typename ChannelType>
+	void FMovieSceneAsyncCurvePainter<ChannelType>::RefineCurvePoints(const TSharedRef<FCurvePainterTask>& Task)
+	{
+		TArray<FVector2D>& InterpolatingPoints = Task->InterpolatingPointsTarget;
+
+		const int32 OldSize = InterpolatingPoints.Num();
+
+		constexpr float InterpTimes[] = { 0.25f, 0.5f, 0.6f };
+		for (int32 Index = 0; Index < InterpolatingPoints.Num() - 1; Index++)
+		{
+			const FVector2D& Lower = InterpolatingPoints[Index];
+			const FVector2D& Upper = InterpolatingPoints[Index + 1];
+
+			if ((Upper.X - Lower.X) >= Task->TimeThreshold)
+			{
+				bool bSegmentIsLinear = true;
+
+				FVector2D Evaluated[UE_ARRAY_COUNT(InterpTimes)] = { FVector2D::ZeroVector };
+				for (int32 InterpIndex = 0; InterpIndex < UE_ARRAY_COUNT(InterpTimes); ++InterpIndex)
+				{
+					double& EvalTime = Evaluated[InterpIndex].X;
+					EvalTime = FMath::Lerp(Lower.X, Upper.X, InterpTimes[InterpIndex]);
+					double Value = 0.0;
+
+					Task->PiecewiseCurve->Evaluate(EvalTime * Task->TickResolution, Value);
+					const double LinearValue = FMath::Lerp(Lower.Y, Upper.Y, InterpTimes[InterpIndex]);
+					if (bSegmentIsLinear)
+					{
+						bSegmentIsLinear = FMath::IsNearlyEqual(Value, LinearValue, Task->ValueThreshold);
+					}
+
+					Evaluated[InterpIndex].Y = Value;
+				}
+
+				if (!bSegmentIsLinear)
+				{
+					// Add the point
+					InterpolatingPoints.Insert(Evaluated, UE_ARRAY_COUNT(Evaluated), Index + 1);
+					--Index;
+				}
+			}
+		}
+
+		const int32 NewSize = InterpolatingPoints.Num();
+		if (OldSize == NewSize)
+		{
+			TArray<int32> KeyOffsets;
+			int32 Offset = 0;
+			for (int32 InterpPointIndex = 0; InterpPointIndex < InterpolatingPoints.Num(); InterpPointIndex++)
+			{
+				if (Task->KeyPointsSource.Num() <= Offset)
+				{
+					KeyOffsets.Add(InterpPointIndex);
+					break;
+				}
+				else if (InterpolatingPoints[InterpPointIndex].X >= Task->KeyPointsSource[Offset].X)
+				{
+					KeyOffsets.Add(InterpPointIndex);
+					Offset++;
+				}
+			}
+
+			Task->Callback(InterpolatingPoints, KeyOffsets);
+			Task->SetFlags(ECurvePainterTaskStateFlags::Completed);
+		}
+	}
+}
diff --git a/Engine/Source/Editor/MovieSceneTools/Private/Cache/MovieSceneAsyncCurvePainter.h b/Engine/Source/Editor/MovieSceneTools/Private/Cache/MovieSceneAsyncCurvePainter.h
new file mode 100644
--- /dev/null
+++ b/Engine/Source/Editor/MovieSceneTools/Private/Cache/MovieSceneAsyncCurvePainter.h
@@ -0,0 +1,113 @@
+// Copyright Epic Games, Inc. All Rights Reserved.
+
+#pragma once
+
+#include "CurveEditorScreenSpace.h"
+#include "Misc/EnumClassFlags.h"
+#include "Misc/FrameRate.h"
+#include "Templates/SharedPointer.h"
+#include <atomic>
+
+namespace UE::MovieScene { struct FPiecewiseCurve; }
+namespace UE::MovieSceneTools
+{
+	template <typename ChannelType> class FMovieSceneCachedCurve;
+
+	/** States of a task in the worker */
+	enum class ECurvePainterTaskStateFlags : uint8
+	{
+		None = 0x00,
+		/** The task is waiting to be processed */
+		Waiting = 0x01,
+		/** The task successfully finished */
+		Completed = 0x02,
+		/** The task is running but no longer useful */
+		Void = 0x04,
+	};
+	ENUM_CLASS_FLAGS(ECurvePainterTaskStateFlags);
+
+	/** Paints a curve to its draw params possibly asynchornously */
+	template <typename ChannelType>
+	class FMovieSceneAsyncCurvePainter
+		: public TSharedFromThis<FMovieSceneAsyncCurvePainter<ChannelType>>
+	{
+	public:
+		/** Paints a curve, possibly asynchronously */
+		static void PaintCurve(
+			const TSharedRef<FMovieSceneCachedCurve<ChannelType>>& CachedCurve, 
+			const bool bSynchronous,
+			const TFunction<void(TArray<FVector2D> /** OutInterpolatedPoints */, TArray<int32> /** OutKeyOffsets */)>& InCallback);
+
+	private:
+		/** A task for the the worker */
+		struct FCurvePainterTask
+			: public TSharedFromThis<FCurvePainterTask>
+		{
+		public:
+			FCurvePainterTask(
+				const TSharedRef<FMovieSceneCachedCurve<ChannelType>>& CachedCurve, 
+				const TFunction<void(TArray<FVector2D> /** OutInterpolatedPoints */, TArray<int32> /** OutKeyOffsets */)>& InCallback);
+
+			/** The callback on completion */
+			const TFunction<void(TArray<FVector2D> /** OutInterpolatedPoints */, TArray<int32> /** OutKeyOffsets */)> Callback;
+
+		private:
+			// Allow the worker to work access data		
+			friend class FMovieSceneAsyncCurvePainter<ChannelType>;
+
+			/** Sets flags for this task */
+			void SetFlags(ECurvePainterTaskStateFlags NewFlags);
+
+			/** Returns true if the task has the specified flag set */
+			bool HasAnyFlags(ECurvePainterTaskStateFlags Flags) const;
+
+			/** The screenspace in which we paint */
+			FCurveEditorScreenSpace ScreenSpace;
+
+			/** The current tick resolution */
+			FFrameRate TickResolution;
+
+			/** Threshold of visible times */
+			double TimeThreshold = 0.0;
+
+			/** Threshold of visible values */
+			double ValueThreshold = 0.0;
+
+			/** The curve as piecewise curve */
+			TSharedPtr<UE::MovieScene::FPiecewiseCurve> PiecewiseCurve;
+
+			/** The initial points to start interpolating from */
+			TArray<FVector2D> KeyPointsSource;
+
+			/** The points to interpolate */
+			TArray<FVector2D> InterpolatingPointsTarget;
+
+			/** True if the task needs to iterate again */
+			std::atomic<ECurvePainterTaskStateFlags> StateFlags = ECurvePainterTaskStateFlags::Waiting;
+		};
+
+		/** Adds a new task to the worker */
+		void AddTask(
+			const TSharedRef<FMovieSceneCachedCurve<ChannelType>>& CachedCurve,
+			const bool bSynchronous,
+			const TFunction<void(TArray<FVector2D> /** OutInterpolatedPoints */, TArray<int32> /** OutKeyOffsets */)>& InCallback);
+
+		/** Called on end frame. Updates work on tick */
+		void OnEndFrame();
+
+		/** Refines curve points once */
+		void RefineCurvePoints(const TSharedRef<FCurvePainterTask>& Task);
+		
+		/** A map of cached curves with their related tasks that need to be performed asynchronous */
+		TMap<TSharedPtr<FMovieSceneCachedCurve<ChannelType>>, TSharedPtr<FCurvePainterTask>> CachedCurveToTaskMap;
+
+		/** True while work is ongoing */
+		std::atomic<bool> bWorking = false;
+
+		/** Critical section to access tasks */
+		FCriticalSection AccessTasksCritSec;
+
+		/** Critical section to enter when binding OnEndFrame from different threads */
+		FCriticalSection BindOnEndFrameCritSec;
+	};
+}
diff --git a/Engine/Source/Editor/MovieSceneTools/Private/Cache/MovieSceneCachedCurve.h b/Engine/Source/Editor/MovieSceneTools/Private/Cache/MovieSceneCachedCurve.h
new file mode 100644
--- /dev/null
+++ b/Engine/Source/Editor/MovieSceneTools/Private/Cache/MovieSceneCachedCurve.h
@@ -0,0 +1,971 @@
+// Copyright Epic Games, Inc. All Rights Reserved.
+
+#pragma once
+
+#include "Cache/MovieSceneAsyncCurvePainter.h"
+#include "Cache/MovieSceneCurveCachePool.h"
+#include "Cache/MovieSceneUpdateCachedCurveData.h"
+#include "Channels/MovieSceneInterpolation.h"
+#include "CurveDataAbstraction.h"
+#include "CurveEditor.h"
+#include "CurveEditorCurveDrawParamsHandle.h"
+#include "CurveEditorScreenSpace.h"
+#include "CurveEditorSettings.h"
+
+namespace UE::MovieSceneTools
+{
+	template <typename ChannelType> struct FMovieSceneUpdateCachedCurveData;
+
+	/** Flags defining how the cache changed when it was last updated */
+	enum class EMovieSceneCurveCacheChangeFlags : uint32
+	{
+		None = 0,
+		ChangedPosition = 1 << 0,
+		ChangedSize = 1 << 1,
+		ChangedKeyIndices = 1 << 3,
+		ChangedTangentVisibility = 1 << 4,
+		ChangedSelection = 1 << 5,
+		ChangedCurveData = 1 << 6,
+		ChangedInterpolatingPoints = 1 << 7,
+	};
+	ENUM_CLASS_FLAGS(EMovieSceneCurveCacheChangeFlags);
+
+	/** Class holding cached curve data to speed up drawing the curve */
+	template <typename ChannelType>
+	class FMovieSceneCachedCurve
+		: public IMovieSceneCachedCurve
+	{
+	private:
+		/** The type of channel value structs */
+		using ChannelValueType = typename ChannelType::ChannelValueType;
+
+	public:
+		FMovieSceneCachedCurve(const FCurveModelID& InCurveModelID);
+		virtual ~FMovieSceneCachedCurve();
+
+		/** Initializes the cached curve */
+		void Initialize(TWeakPtr<const FCurveEditor> WeakCurveEditor);
+
+		/** Returns true when the curve changed */
+		bool HasChanged() const { return Flags != EMovieSceneCurveCacheChangeFlags::None; }
+
+		/** Updates the cached curve data. Doesn't draw the curve */
+		void UpdateCachedCurve(const FMovieSceneUpdateCachedCurveData<ChannelType>& UpdateData, const UE::CurveEditor::FCurveDrawParamsHandle& CurveDrawParamsHandle);
+
+		/** Draws the cached curve to the actual draw params */
+		virtual void DrawCachedCurve() override;
+
+		/** Returns the curve model ID that corresponds to the cached curve */
+		virtual const FCurveModelID& GetID() const override { return CurveModelID; }
+
+		/** Returns the cached screen space */
+		const FCurveEditorScreenSpace& GetScreenSpace() const { return ScreenSpace; }
+
+		/** Returns the cached tick resolution */
+		const FFrameRate& GetTickResolution() const { return TickResolution; }
+
+		/** Returns the threshold of visible pixels per time */
+		double GetTimeThreshold() const { return FMath::Max(0.0001, 1.0 / ScreenSpace.PixelsPerInput()); }
+
+		/** Returns the threshold of visible pixels per value */
+		double GetValueThreshold() const { return FMath::Max(0.0001, 1.0 / ScreenSpace.PixelsPerInput()); }
+
+		/** Returns the cached piecewise curve */
+		const TSharedPtr<UE::MovieScene::FPiecewiseCurve>& GetPiecewiseCurve() const { return PiecewiseCurve; }
+
+		/** Returns the points where keys reside */
+		TArray<FVector2D> GetKeyPoints();
+
+		// IMovieSceneCachedCurve interface
+		virtual uint32 GetInterpolatingPointsHash() const override;
+		virtual const UE::CurveEditor::FCurveDrawParamsHandle& GetDrawParamsHandle() const override;
+		// ~IMovieSceneCachedCurve interface
+
+		/** The curve model ID that corresponds to the cached curve */
+		const FCurveModelID CurveModelID;
+
+	private:
+		/** Updates the cached screen space */
+		void UpdateScreenSpace(const FCurveEditorScreenSpace& NewScreenSpace);
+
+		/** Updates the cached tick resolution */
+		void UpdateTickResolution(const FFrameRate& NewTickResolution);
+
+		/** Updates the cached input display offset */
+		void UpdateInputDisplayOffset(const double NewInputDisplayOffset);
+
+		/** Updates cached indices of key in the times and values array */
+		void UpdateIndices(const FMovieSceneUpdateCachedCurveData<ChannelType>& UpdateData, const FFrameNumber& StartFrame, const FFrameNumber& EndFrame);
+
+		/** Updates the cached selection */
+		void UpdateSelection(const FMovieSceneUpdateCachedCurveData<ChannelType>& UpdateData);
+
+		/** Updates the tangent visiblity */
+		void UpdateTangentVisibility(const ECurveEditorTangentVisibility NewTangentVisibility);
+
+		/** Updates cached curve data such as keys, values, key handles and key attributes */
+		void UpdateCurveData(const FMovieSceneUpdateCachedCurveData<ChannelType>& UpdateData);
+
+		/** Updates the cached pre and post inifinity extrapolation */
+		void UpdatePrePostInfinityExtrapolation(const FMovieSceneUpdateCachedCurveData<ChannelType>& UpdateData, const FFrameNumber& StartFrame, const FFrameNumber& EndFrame);
+
+		/** Updates key draw info depending on flags */
+		void ConditionallyUpdateKeyDrawInfos(const FMovieSceneUpdateCachedCurveData<ChannelType>& UpdateData);
+
+		/** Updates the piecewise cuve depending on flags */
+		void ConditionallyUpdatePiecewiseCurve(const FMovieSceneUpdateCachedCurveData<ChannelType>& Data);
+
+		/** Draws the full range of interpolating points, possibly async. */
+		void DrawFullRangeInterpolatingPoints();
+
+		/** Draws the interpolating points from the currently cached full range interpolating points */
+		void DrawInterpolatingPointsFromFullRange();
+
+		/** Draws the pre-infinity interpolating points */
+		void DrawPreInfinityInterpolatingPoints();
+
+		/** Draws the post-infinity interpolating points */
+		void DrawPostInfinityInterpolatingPoints();
+
+		/** Tries to draw a single pre-infinity interpolating point fast */
+		[[nodiscard]] bool TryDrawPreInfinityExtentFast();
+
+		/** Tries to draw a single post-infinity interpolating point fast */
+		[[nodiscard]] bool TryDrawPostInfinityExtentFast();
+
+		/** Draws the keys of the curve */
+		void DrawKeys();
+
+		/** Applies the cached draw params to the actual draw params */
+		void ApplyDrawParams();
+
+		/** Cached screen space */
+		FCurveEditorScreenSpace ScreenSpace;
+
+		/** Cached tick resolution */
+		FFrameRate TickResolution;
+
+		/** Cached input display offset */
+		double InputDisplayOffset = TNumericLimits<double>::Max();
+
+		/** Cached first index in keys / values */
+		int32 StartingIndex = INDEX_NONE;
+
+		/** Cached last index in keys / values */
+		int32 EndingIndex = INDEX_NONE;
+
+		/** Cached selection */
+		FKeyHandleSet Selection;
+
+		/** Cached tangent visibility. Optional soley to detect the initial change */
+		TOptional<ECurveEditorTangentVisibility> TangentVisibility;
+
+		/** Cached default value */
+		double DefaultValue = 0.0;
+
+		/** Cached curve times */
+		TArray<const FFrameNumber> Times;
+
+		/** Cached curve values */
+		TArray<const ChannelValueType> Values;
+
+		/** Cached key handles */
+		TArray<FKeyHandle> KeyHandles;
+
+		/** Cached key attributes */
+		TArray<FKeyAttributes> KeyAttributes;
+
+		/** Key draw info as a tuple of the key, an optional arrive tangent and an optional leave tangent */
+		TArray<TTuple<FKeyDrawInfo, TOptional<FKeyDrawInfo>, TOptional<FKeyDrawInfo>>> KeyDrawInfos;
+
+		/** The curve as piecewise curve. */
+		TSharedPtr<UE::MovieScene::FPiecewiseCurve> PiecewiseCurve;
+
+		/** Struct holding the whole range of interpolating points. */
+		struct FFullRangeinterpolatingPoints
+		{
+			/** The interpolating points in the finite curve range */
+			TArray<FVector2D> Points;
+
+			/** Offsets of keys in the points array */
+			TArray<int32> KeyOffsets;
+
+			/** Hash of the full range interpolating points */
+			std::atomic<uint32> Hash = 0;
+
+			/** How pre-infinity should be extrapolated */
+			ERichCurveExtrapolation PreInfinityExtrapolation = ERichCurveExtrapolation::RCCE_None;
+
+			/** How post-infinity should be extrapolated */
+			ERichCurveExtrapolation PostInfinityExtrapolation = ERichCurveExtrapolation::RCCE_None;
+
+		};
+
+		/** The whole range of interpolating points. Useful to avoid drawing the visible range when curve and zoom don't change. */
+		FFullRangeinterpolatingPoints FullRangeInterpolation;
+
+		/** Critical section to enter when accessing the whole range or the cached draw params interpolating points */
+		mutable FCriticalSection LockInterpolatingPoints;
+
+		/** Cached curve draw params */
+		FCurveDrawParams CachedDrawParams;
+
+		/** Handle to the actual draw params we're drawing to */
+		UE::CurveEditor::FCurveDrawParamsHandle DrawParamsHandle;
+
+		/** Flags defining how the cache changed since it last was updated */
+		std::atomic<EMovieSceneCurveCacheChangeFlags> Flags = EMovieSceneCurveCacheChangeFlags::None;
+	};
+
+	template <typename ChannelType>
+	FMovieSceneCachedCurve<ChannelType>::FMovieSceneCachedCurve(const FCurveModelID& InCurveModelID)
+		: CurveModelID(InCurveModelID)
+		, ScreenSpace(FVector2D::ZeroVector, 0.0, 1.0, 0.0, 1.0)
+		, CachedDrawParams(InCurveModelID)
+	{}
+
+	template <typename ChannelType>
+	FMovieSceneCachedCurve<ChannelType>::~FMovieSceneCachedCurve()
+	{
+		FMovieSceneCurveCachePool::Get().Leave(*this);
+	}
+
+	template <typename ChannelType>
+	void FMovieSceneCachedCurve<ChannelType>::Initialize(TWeakPtr<const FCurveEditor> WeakCurveEditor)
+	{
+		FMovieSceneCurveCachePool::Get().Join(WeakCurveEditor, AsShared());
+	}
+
+	template <typename ChannelType>
+	TArray<FVector2D> FMovieSceneCachedCurve<ChannelType>::GetKeyPoints()
+	{
+		return CachedDrawParams.InterpolatingPoints;
+	}
+
+	template <typename ChannelType>
+	void FMovieSceneCachedCurve<ChannelType>::UpdateCachedCurve(const FMovieSceneUpdateCachedCurveData<ChannelType>& UpdateData, const UE::CurveEditor::FCurveDrawParamsHandle& CurveDrawParamsHandle)
+	{
+		DrawParamsHandle = CurveDrawParamsHandle;
+
+		CachedDrawParams.Color = UpdateData.CurveModel.GetColor();
+		CachedDrawParams.Thickness = UpdateData.CurveModel.GetThickness();
+		CachedDrawParams.DashLengthPx = UpdateData.CurveModel.GetDashLength();
+		CachedDrawParams.bKeyDrawEnabled = UpdateData.CurveModel.IsKeyDrawEnabled();
+
+		/** Cached first fame which needs to be updated  */
+		const FFrameNumber StartFrame = (UpdateData.ScreenSpace.GetInputMin() * UpdateData.TickResolution).FloorToFrame();
+		const FFrameNumber EndFrame = (UpdateData.ScreenSpace.GetInputMax() * UpdateData.TickResolution).CeilToFrame();
+
+		UpdateScreenSpace(UpdateData.ScreenSpace);
+		UpdateTickResolution(UpdateData.TickResolution);
+		UpdateTangentVisibility(UpdateData.CurveEditor.GetSettings()->GetTangentVisibility());
+		UpdateInputDisplayOffset(UpdateData.CurveModel.GetInputDisplayOffset());
+		UpdateIndices(UpdateData, StartFrame, EndFrame);
+		UpdateSelection(UpdateData);
+		UpdateCurveData(UpdateData);
+		UpdatePrePostInfinityExtrapolation(UpdateData, StartFrame, EndFrame);
+
+		// Update depending on above data
+		ConditionallyUpdateKeyDrawInfos(UpdateData);
+		ConditionallyUpdatePiecewiseCurve(UpdateData);
+	}
+
+	template <typename ChannelType>
+	void FMovieSceneCachedCurve<ChannelType>::DrawCachedCurve()
+	{
+		if (Flags == EMovieSceneCurveCacheChangeFlags::None)
+		{
+			// Only apply cached data if nothing changed.
+			ApplyDrawParams();
+		}
+		else
+		{
+			DrawKeys();
+			DrawFullRangeInterpolatingPoints();
+			DrawInterpolatingPointsFromFullRange();
+
+			ApplyDrawParams();
+
+			Flags = EMovieSceneCurveCacheChangeFlags::None;
+		}
+	}
+
+	template <typename ChannelType>
+	uint32 FMovieSceneCachedCurve<ChannelType>::GetInterpolatingPointsHash() const
+	{
+		return FullRangeInterpolation.Hash.load();
+	}
+
+	template <typename ChannelType>
+	const UE::CurveEditor::FCurveDrawParamsHandle& FMovieSceneCachedCurve<ChannelType>::GetDrawParamsHandle() const
+	{
+		check(IsInGameThread());
+		return DrawParamsHandle;
+	}
+
+	template <typename ChannelType>
+	void FMovieSceneCachedCurve<ChannelType>::UpdateScreenSpace(const FCurveEditorScreenSpace& NewScreenSpace)
+	{
+		const bool bScreenPositionChanged =
+			!FMath::IsNearlyEqual(ScreenSpace.GetInputMin(), NewScreenSpace.GetInputMin()) ||
+			!FMath::IsNearlyEqual(ScreenSpace.GetInputMax(), NewScreenSpace.GetInputMax()) ||
+			!FMath::IsNearlyEqual(ScreenSpace.GetOutputMin(), NewScreenSpace.GetOutputMin()) ||
+			!FMath::IsNearlyEqual(ScreenSpace.GetOutputMax(), NewScreenSpace.GetOutputMax());
+
+		const EMovieSceneCurveCacheChangeFlags ScreenPositionChangedFlag = bScreenPositionChanged ? EMovieSceneCurveCacheChangeFlags::ChangedPosition : EMovieSceneCurveCacheChangeFlags::None;
+
+		const bool bScreenSizeChanged =
+			!FMath::IsNearlyEqual(ScreenSpace.PixelsPerInput(), NewScreenSpace.PixelsPerInput()) ||
+			!FMath::IsNearlyEqual(ScreenSpace.PixelsPerOutput(), NewScreenSpace.PixelsPerOutput());
+
+		const EMovieSceneCurveCacheChangeFlags ScreenSizeChangedFlag = bScreenSizeChanged ? EMovieSceneCurveCacheChangeFlags::ChangedSize : EMovieSceneCurveCacheChangeFlags::None;
+
+		if (bScreenPositionChanged || bScreenSizeChanged)
+		{
+			ScreenSpace = NewScreenSpace;
+			Flags.store(Flags.load() | ScreenPositionChangedFlag | ScreenSizeChangedFlag);
+		}
+	}
+
+	template <typename ChannelType>
+	void FMovieSceneCachedCurve<ChannelType>::UpdateTickResolution(const FFrameRate& NewTickResolution)
+	{
+		if (TickResolution != NewTickResolution)
+		{
+			TickResolution = NewTickResolution;
+			Flags.store(Flags.load() | EMovieSceneCurveCacheChangeFlags::ChangedPosition | EMovieSceneCurveCacheChangeFlags::ChangedSize);
+		}
+	}
+
+	template <typename ChannelType>
+	void FMovieSceneCachedCurve<ChannelType>::UpdateInputDisplayOffset(const double NewInputDisplayOffset)
+	{
+		if (InputDisplayOffset != NewInputDisplayOffset)
+		{
+			InputDisplayOffset = NewInputDisplayOffset;
+			Flags.store(Flags.load() | EMovieSceneCurveCacheChangeFlags::ChangedPosition);
+		}
+	}
+
+	template <typename ChannelType>
+	void FMovieSceneCachedCurve<ChannelType>::UpdateIndices(const FMovieSceneUpdateCachedCurveData<ChannelType>& UpdateData, const FFrameNumber& StartFrame, const FFrameNumber& EndFrame)
+	{
+		const int32 NewStartingIndex = Algo::LowerBound(UpdateData.Times, StartFrame);
+		const int32 NewEndingIndex = Algo::UpperBound(UpdateData.Times, EndFrame) - 1;
+
+		if (StartingIndex != NewStartingIndex ||
+			EndingIndex != NewEndingIndex)
+		{
+			StartingIndex = NewStartingIndex;
+			EndingIndex = NewEndingIndex;
+
+			Flags.store(Flags.load() | EMovieSceneCurveCacheChangeFlags::ChangedKeyIndices);
+		}
+	}
+
+	template <typename ChannelType>
+	void FMovieSceneCachedCurve<ChannelType>::UpdateSelection(const FMovieSceneUpdateCachedCurveData<ChannelType>& UpdateData)
+	{
+		const FKeyHandleSet* SelectionPtr = UpdateData.CurveEditor.GetSelection().FindForCurve(CurveModelID);
+		const FKeyHandleSet	NewSelection = SelectionPtr ? *SelectionPtr : FKeyHandleSet();
+
+		const TArrayView<const FKeyHandle> SelectedKeyHandles = Selection.AsArray();
+		const TArrayView<const FKeyHandle> NewSelectedKeyHandles = NewSelection.AsArray();
+		if (SelectedKeyHandles.Num() != NewSelectedKeyHandles.Num() ||
+			FMemory::Memcmp(SelectedKeyHandles.GetData(), NewSelectedKeyHandles.GetData(), SelectedKeyHandles.Num()) != 0)
+		{
+			Selection = NewSelection;
+			Flags.store(Flags.load() | EMovieSceneCurveCacheChangeFlags::ChangedSelection);
+		}
+	}
+
+	template <typename ChannelType>
+	void FMovieSceneCachedCurve<ChannelType>::UpdateTangentVisibility(const ECurveEditorTangentVisibility NewTangentVisibility)
+	{
+		if (!TangentVisibility.IsSet() || TangentVisibility.GetValue() != NewTangentVisibility)
+		{
+			TangentVisibility = NewTangentVisibility;
+
+			Flags.store(Flags.load() | EMovieSceneCurveCacheChangeFlags::ChangedTangentVisibility);
+		}
+	}
+
+	template <typename ChannelType>
+	void FMovieSceneCachedCurve<ChannelType>::UpdateCurveData(const FMovieSceneUpdateCachedCurveData<ChannelType>& UpdateData)
+	{
+		// DefaultValue
+		const double NewDefaultValue = UpdateData.Channel.GetDefault().IsSet() ? UpdateData.Channel.GetDefault().GetValue() : 0.0;
+		if (DefaultValue != NewDefaultValue)
+		{
+			DefaultValue = NewDefaultValue;
+
+			Flags.store(Flags.load() | EMovieSceneCurveCacheChangeFlags::ChangedCurveData);
+		}
+
+		// Times
+		if (Times != UpdateData.Times)
+		{
+			Times = UpdateData.Times;
+
+			Flags.store(Flags.load() | EMovieSceneCurveCacheChangeFlags::ChangedCurveData);
+		}
+
+		// Values
+		if (Values != UpdateData.Values)
+		{
+			Values = TArray<const ChannelValueType>(UpdateData.Values.GetData(), UpdateData.Values.Num());
+
+			Flags.store(Flags.load() | EMovieSceneCurveCacheChangeFlags::ChangedCurveData);
+		}
+
+		// Key Handles
+		TArray<FKeyHandle> NewKeyHandles;
+		NewKeyHandles.Reserve(Times.Num());
+
+		UpdateData.CurveModel.GetKeys(TNumericLimits<double>::Lowest(), TNumericLimits<double>::Max(), TNumericLimits<double>::Lowest(), TNumericLimits<double>::Max(), NewKeyHandles);
+		if (KeyHandles != NewKeyHandles)
+		{
+			KeyHandles = NewKeyHandles;
+
+			Flags.store(Flags.load() | EMovieSceneCurveCacheChangeFlags::ChangedCurveData);
+		}
+
+		// Key Attributes
+		TArray<FKeyAttributes> NewKeyAttributes;
+		NewKeyAttributes.SetNum(KeyHandles.Num());
+
+		UpdateData.CurveModel.GetKeyAttributes(KeyHandles, NewKeyAttributes);
+
+		if (KeyAttributes != NewKeyAttributes)
+		{
+			KeyAttributes = NewKeyAttributes;
+
+			Flags.store(Flags.load() | EMovieSceneCurveCacheChangeFlags::ChangedCurveData);
+		}
+	}
+
+	template <typename ChannelType>
+	void FMovieSceneCachedCurve<ChannelType>::UpdatePrePostInfinityExtrapolation(const FMovieSceneUpdateCachedCurveData<ChannelType>& UpdateData, const FFrameNumber& StartFrame, const FFrameNumber& EndFrame)
+	{
+		const ERichCurveExtrapolation NewPreInfinityExtrapolation = UpdateData.Channel.PreInfinityExtrap;
+		const ERichCurveExtrapolation NewPostInfinityExtrapolation = UpdateData.Channel.PostInfinityExtrap;
+
+		if (FullRangeInterpolation.PreInfinityExtrapolation != NewPreInfinityExtrapolation ||
+			FullRangeInterpolation.PostInfinityExtrapolation != NewPostInfinityExtrapolation)
+		{
+			FullRangeInterpolation.PreInfinityExtrapolation = NewPreInfinityExtrapolation;
+			FullRangeInterpolation.PostInfinityExtrapolation = NewPostInfinityExtrapolation;
+		}
+	}
+
+	template <typename ChannelType>
+	void FMovieSceneCachedCurve<ChannelType>::ConditionallyUpdateKeyDrawInfos(const FMovieSceneUpdateCachedCurveData<ChannelType>& UpdateData)
+	{
+		if (EnumHasAnyFlags(Flags.load(), EMovieSceneCurveCacheChangeFlags::ChangedCurveData | EMovieSceneCurveCacheChangeFlags::ChangedSelection))
+		{
+			UpdateData.CurveModel.GetKeyDrawInfo(ECurvePointType::ArriveTangent, FKeyHandle::Invalid(), CachedDrawParams.ArriveTangentDrawInfo);
+			UpdateData.CurveModel.GetKeyDrawInfo(ECurvePointType::LeaveTangent, FKeyHandle::Invalid(), CachedDrawParams.LeaveTangentDrawInfo);
+
+			KeyDrawInfos.Reset(KeyHandles.Num());
+			for (int32 DataIndex = 0; DataIndex < KeyHandles.Num(); DataIndex++)
+			{
+				const FKeyHandle& KeyHandle = KeyHandles[DataIndex];
+
+				// Key draw info
+				FKeyDrawInfo KeyDrawInfo;
+				UpdateData.CurveModel.GetKeyDrawInfo(ECurvePointType::Key, KeyHandle, KeyDrawInfo);
+
+				// Tangent draw info
+				const FKeyAttributes* KeyAttributesPtr = nullptr;
+				if (TangentVisibility == ECurveEditorTangentVisibility::SelectedKeys)
+				{
+					if (Selection.Contains(KeyHandle, ECurvePointType::Any))
+					{
+						KeyAttributesPtr = &KeyAttributes[DataIndex];
+					}
+				}
+				else if (TangentVisibility == ECurveEditorTangentVisibility::AllTangents)
+				{
+					KeyAttributesPtr = &KeyAttributes[DataIndex];
+				}
+
+				TOptional<FKeyDrawInfo> ArriveTangentDrawInfo;
+				if (KeyAttributesPtr && KeyAttributesPtr->HasArriveTangent())
+				{
+					FKeyDrawInfo NewArriveTangentDrawInfo;
+					UpdateData.CurveModel.GetKeyDrawInfo(ECurvePointType::ArriveTangent, KeyHandle, NewArriveTangentDrawInfo);
+					ArriveTangentDrawInfo = NewArriveTangentDrawInfo;
+				}
+
+				TOptional<FKeyDrawInfo> LeaveTangentDrawInfo;
+				if (KeyAttributesPtr && KeyAttributesPtr->HasLeaveTangent())
+				{
+					FKeyDrawInfo NewLeaveTangentDrawInfo;
+					UpdateData.CurveModel.GetKeyDrawInfo(ECurvePointType::ArriveTangent, KeyHandle, NewLeaveTangentDrawInfo);
+					LeaveTangentDrawInfo = NewLeaveTangentDrawInfo;
+				}
+
+				KeyDrawInfos.Emplace(KeyDrawInfo, ArriveTangentDrawInfo, LeaveTangentDrawInfo);
+			}
+		}
+	}
+
+	template <typename ChannelType>
+	void FMovieSceneCachedCurve<ChannelType>::ConditionallyUpdatePiecewiseCurve(const FMovieSceneUpdateCachedCurveData<ChannelType>& Data)
+	{
+		if (EnumHasAnyFlags(Flags.load(), EMovieSceneCurveCacheChangeFlags::ChangedCurveData))
+		{
+			constexpr bool bWithPreAndPostInfinityExtrap = false;
+			PiecewiseCurve = MakeShared<UE::MovieScene::FPiecewiseCurve>(Data.Channel.AsPiecewiseCurve(bWithPreAndPostInfinityExtrap));
+		}
+	}
+
+	template <typename ChannelType>
+	void FMovieSceneCachedCurve<ChannelType>::DrawFullRangeInterpolatingPoints()
+	{
+		// Only draw the full range interpolating points when the curve was edited or when zooming
+
+		const bool bCurveIsBeingEdited = EnumHasAnyFlags(Flags.load(), EMovieSceneCurveCacheChangeFlags::ChangedCurveData);
+		const bool bZooming = EnumHasAnyFlags(Flags.load(), EMovieSceneCurveCacheChangeFlags::ChangedSize);
+		if (bCurveIsBeingEdited || bZooming)
+		{
+			const FScopeLock Lock(&LockInterpolatingPoints);
+
+			CachedDrawParams.InterpolatingPoints.Reset();
+
+			CachedDrawParams.InterpolatingPoints.Reserve(Times.Num());
+			for (int32 DataIndex = 0; DataIndex < Times.Num(); DataIndex++)
+			{
+				CachedDrawParams.InterpolatingPoints.Emplace(Times[DataIndex] / TickResolution, double(Values[DataIndex].Value));
+			}
+
+			// Paint the curve synchronous when the it is being edited
+			const bool bSynchronous = bCurveIsBeingEdited;
+
+			FMovieSceneAsyncCurvePainter<ChannelType>::PaintCurve(
+				SharedThis(this),
+				bSynchronous,
+				[this](TArray<FVector2D> NewInterpolatingPoints, TArray<int32> NewKeyOffsets)
+				{
+					const FScopeLock Lock(&LockInterpolatingPoints);
+
+					FullRangeInterpolation.Points = NewInterpolatingPoints;
+					FullRangeInterpolation.KeyOffsets = NewKeyOffsets;
+
+					uint32 Hash = 0;
+					for (FVector2D& InterpolatingPoint : FullRangeInterpolation.Points)
+					{
+						const int32 X = FMath::RoundToInt32(InterpolatingPoint.X);
+						const int32 Y = FMath::RoundToInt32(InterpolatingPoint.Y);
+						Hash = HashCombine(FullRangeInterpolation.Hash, HashCombine(X, Y));
+					}
+					FullRangeInterpolation.Hash.store(Hash);
+
+					Flags.store(Flags.load() | EMovieSceneCurveCacheChangeFlags::ChangedInterpolatingPoints);
+				});
+		}
+	}
+
+	template <typename ChannelType>
+	void FMovieSceneCachedCurve<ChannelType>::DrawInterpolatingPointsFromFullRange()
+	{
+		const FScopeLock Lock(&LockInterpolatingPoints);
+
+		CachedDrawParams.InterpolatingPoints.Reset(FullRangeInterpolation.Points.Num() + 2);
+
+		if (FullRangeInterpolation.Points.IsEmpty())
+		{
+			DrawPreInfinityInterpolatingPoints();
+			DrawPostInfinityInterpolatingPoints();
+		}
+		else
+		{
+			const bool bHasKeyBefore = FullRangeInterpolation.KeyOffsets.IsValidIndex(StartingIndex - 1);
+			const bool bHasKeyAfter = FullRangeInterpolation.KeyOffsets.IsValidIndex(EndingIndex + 1);
+
+			// Add the lower bound of the visible space if there is no key before
+			if (!bHasKeyBefore)
+			{
+				DrawPreInfinityInterpolatingPoints();
+			}
+
+			// Overdraw to the previous key if possible
+			const int32 DataIndex = bHasKeyBefore ?
+				FullRangeInterpolation.KeyOffsets[StartingIndex - 1] :
+				FullRangeInterpolation.KeyOffsets[StartingIndex];
+
+			// Overdraw to the next key if possible
+			const int32 InterpolatingPointsDataSize = bHasKeyAfter ?
+				FullRangeInterpolation.KeyOffsets[EndingIndex + 1] - DataIndex + 1 :
+				FullRangeInterpolation.KeyOffsets[EndingIndex] - DataIndex + 1;
+
+			// Add points in between
+			check(FullRangeInterpolation.Points.IsValidIndex(DataIndex) && FullRangeInterpolation.Points.Num() >= DataIndex + InterpolatingPointsDataSize);
+
+			const TArrayView<FVector2D> InterpolatingPointsView(FullRangeInterpolation.Points.GetData() + DataIndex, InterpolatingPointsDataSize);
+			for (const FVector2D& InterpolatingPoint : InterpolatingPointsView)
+			{
+				CachedDrawParams.InterpolatingPoints.Emplace(
+					ScreenSpace.SecondsToScreen(InterpolatingPoint.X + InputDisplayOffset),
+					ScreenSpace.ValueToScreen(InterpolatingPoint.Y));
+			}
+
+			// Add the upper bound of the visible space if there is no key after
+			if (!bHasKeyAfter)
+			{
+				DrawPostInfinityInterpolatingPoints();
+			}
+		}
+	}
+
+	template <typename ChannelType>
+	void FMovieSceneCachedCurve<ChannelType>::DrawPreInfinityInterpolatingPoints()
+	{
+		if (TryDrawPreInfinityExtentFast())
+		{
+			return;
+		}
+
+		const FScopeLock Lock(&LockInterpolatingPoints);
+
+		if (FullRangeInterpolation.PreInfinityExtrapolation == ERichCurveExtrapolation::RCCE_Cycle ||
+			FullRangeInterpolation.PreInfinityExtrapolation == ERichCurveExtrapolation::RCCE_CycleWithOffset ||
+			FullRangeInterpolation.PreInfinityExtrapolation == ERichCurveExtrapolation::RCCE_Oscillate)
+		{
+			const double InfinityOffset = FullRangeInterpolation.Points[0].X - ScreenSpace.GetInputMin();
+			if (InfinityOffset < 0.0)
+			{
+				// Don't draw if pre-infinity is not visible
+				return;
+			}
+
+			const double Duration = FullRangeInterpolation.Points.Last().X - FullRangeInterpolation.Points[0].X;
+			if (FMath::IsNearlyEqual(Duration, 0.0))
+			{
+				// Draw a single line in the odd case where there are many keys with a nearly zero duration.
+				CachedDrawParams.InterpolatingPoints.Emplace(
+					ScreenSpace.SecondsToScreen(ScreenSpace.GetInputMin()),
+					ScreenSpace.ValueToScreen(FullRangeInterpolation.Points[0].Y));
+
+				CachedDrawParams.InterpolatingPoints.Emplace(
+					ScreenSpace.SecondsToScreen(ScreenSpace.GetInputMax()),
+					ScreenSpace.ValueToScreen(FullRangeInterpolation.Points[0].Y));
+			}
+			else
+			{
+				// Cycle or oscillate
+				const double StartTime = FullRangeInterpolation.Points[0].X;
+				const double ValueOffset = FullRangeInterpolation.PreInfinityExtrapolation == ERichCurveExtrapolation::RCCE_CycleWithOffset ?
+					FullRangeInterpolation.Points.Last().Y - FullRangeInterpolation.Points[0].Y :
+					0.0;
+
+				const int32 NumIterations = static_cast<int32>(InfinityOffset / Duration) + 1;
+				for (int32 Iteration = NumIterations; Iteration > 0; Iteration--)
+				{
+					const bool bReverse = FullRangeInterpolation.PostInfinityExtrapolation == ERichCurveExtrapolation::RCCE_Oscillate && Iteration % 2 != 0;
+
+					if (bReverse)
+					{
+						for (int32 PointIndex = FullRangeInterpolation.Points.Num() - 1; PointIndex >= 0; PointIndex--)
+						{
+							const FVector2D& Point = FullRangeInterpolation.Points[PointIndex];
+
+							// Mirror around start time
+							const double Time = 2 * StartTime + Duration - Point.X - Duration * Iteration + InputDisplayOffset;
+							const double Value = Point.Y - ValueOffset * Iteration;
+
+							CachedDrawParams.InterpolatingPoints.Emplace(ScreenSpace.SecondsToScreen(Time), ScreenSpace.ValueToScreen(Value));
+						}
+					}
+					else
+					{
+						for (const FVector2D& Point : FullRangeInterpolation.Points)
+						{
+							const double Time = Point.X - Duration * Iteration + InputDisplayOffset;
+							const double Value = Point.Y - ValueOffset * Iteration;
+
+							CachedDrawParams.InterpolatingPoints.Emplace(ScreenSpace.SecondsToScreen(Time), ScreenSpace.ValueToScreen(Value));
+						}
+					}
+				}
+			}
+		}
+	}
+
+	template <typename ChannelType>
+	void FMovieSceneCachedCurve<ChannelType>::DrawPostInfinityInterpolatingPoints()
+	{
+		if (TryDrawPostInfinityExtentFast())
+		{
+			return;
+		}
+
+		const FScopeLock Lock(&LockInterpolatingPoints);
+
+		if (FullRangeInterpolation.Points.Num() > 1 &&
+			(FullRangeInterpolation.PostInfinityExtrapolation == ERichCurveExtrapolation::RCCE_Cycle ||
+			FullRangeInterpolation.PostInfinityExtrapolation == ERichCurveExtrapolation::RCCE_CycleWithOffset ||
+			FullRangeInterpolation.PostInfinityExtrapolation == ERichCurveExtrapolation::RCCE_Oscillate))
+		{
+			const double InfinityOffset = ScreenSpace.GetInputMax() - FullRangeInterpolation.Points.Last().X;
+			if (InfinityOffset < 0.0)
+			{
+				// Don't draw if post-infinity is not visible
+				return;
+			}
+
+			const double Duration = FullRangeInterpolation.Points.Last().X - FullRangeInterpolation.Points[0].X;
+			if (FMath::IsNearlyEqual(Duration, 0.0))
+			{
+				// Draw a single line in the odd case where there are many keys with a nearly zero duration.
+				CachedDrawParams.InterpolatingPoints.Emplace(
+					ScreenSpace.SecondsToScreen(ScreenSpace.GetInputMin()),
+					ScreenSpace.ValueToScreen(FullRangeInterpolation.Points.Last().Y));
+
+				CachedDrawParams.InterpolatingPoints.Emplace(
+					ScreenSpace.SecondsToScreen(ScreenSpace.GetInputMax()),
+					ScreenSpace.ValueToScreen(FullRangeInterpolation.Points.Last().Y));
+			}
+			else
+			{
+				// Cycle or oscillate
+				const double StartTime = FullRangeInterpolation.Points[0].X;
+				const double ValueOffset = FullRangeInterpolation.PreInfinityExtrapolation == ERichCurveExtrapolation::RCCE_CycleWithOffset ?
+					FullRangeInterpolation.Points.Last().Y - FullRangeInterpolation.Points[0].Y :
+					0.0;
+
+				const int32 NumIterations = InfinityOffset / Duration + 1;
+				for (int32 Iteration = 1; Iteration <= NumIterations; Iteration++)
+				{
+					const bool bReverse = FullRangeInterpolation.PostInfinityExtrapolation == ERichCurveExtrapolation::RCCE_Oscillate && Iteration % 2 != 0;
+
+					if (bReverse)
+					{
+						for (int32 PointIndex = FullRangeInterpolation.Points.Num() - 1; PointIndex >= 0; PointIndex--)
+						{
+							const FVector2D& Point = FullRangeInterpolation.Points[PointIndex];
+
+							// Mirror around start time
+							const double Time = 2 * StartTime + Duration - Point.X + Duration * Iteration + InputDisplayOffset;
+							const double Value = Point.Y - ValueOffset * Iteration;
+
+							CachedDrawParams.InterpolatingPoints.Emplace(ScreenSpace.SecondsToScreen(Time), ScreenSpace.ValueToScreen(Value));
+						}
+					}
+					else
+					{
+						for (const FVector2D& Point : FullRangeInterpolation.Points)
+						{
+							const double Time = Point.X + Duration * Iteration + InputDisplayOffset;
+							const double Value = Point.Y + ValueOffset * Iteration;
+
+							CachedDrawParams.InterpolatingPoints.Emplace(ScreenSpace.SecondsToScreen(Time), ScreenSpace.ValueToScreen(Value));
+						}
+					}
+				}
+			}
+		}
+	}
+
+	template <typename ChannelType>
+	bool FMovieSceneCachedCurve<ChannelType>::TryDrawPreInfinityExtentFast()
+	{
+		if (Values.IsEmpty())
+		{
+			CachedDrawParams.InterpolatingPoints.Emplace(
+				ScreenSpace.SecondsToScreen(ScreenSpace.GetInputMin()), 
+				ScreenSpace.ValueToScreen(DefaultValue));
+
+			return true;
+		}
+		else if (Values.Num() == 1)
+		{
+			const double SingleValue = Values[0].Value;
+			CachedDrawParams.InterpolatingPoints.Emplace(
+				ScreenSpace.SecondsToScreen(ScreenSpace.GetInputMin()), 
+				ScreenSpace.ValueToScreen(SingleValue));
+
+			return true;
+		}
+		else
+		{
+			const bool bPreinfinityVisible = Times[0].Value > ScreenSpace.GetInputMin() * TickResolution;
+			if (bPreinfinityVisible &&
+				(FullRangeInterpolation.PreInfinityExtrapolation == ERichCurveExtrapolation::RCCE_None ||
+				FullRangeInterpolation.PreInfinityExtrapolation == ERichCurveExtrapolation::RCCE_Constant))
+			{
+				const double ExtentValue = Values[0].Value;
+				CachedDrawParams.InterpolatingPoints.Emplace(
+					ScreenSpace.SecondsToScreen(ScreenSpace.GetInputMin()), 
+					ScreenSpace.ValueToScreen(ExtentValue));
+
+				return true;
+			}
+			else if (bPreinfinityVisible &&
+				FullRangeInterpolation.PreInfinityExtrapolation == ERichCurveExtrapolation::RCCE_Linear)
+			{
+				const double ExtentValue = Values[0].Value + KeyAttributes[0].GetArriveTangent() * (ScreenSpace.GetInputMin() - Times[0].Value);
+				CachedDrawParams.InterpolatingPoints.Emplace(
+					ScreenSpace.SecondsToScreen(ScreenSpace.GetInputMin()), 
+					ScreenSpace.ValueToScreen(ExtentValue));
+
+				return true;
+			}
+		}
+		
+		return false;
+	}
+
+	template <typename ChannelType>
+	bool FMovieSceneCachedCurve<ChannelType>::TryDrawPostInfinityExtentFast()
+	{
+		if (Values.IsEmpty())
+		{
+			CachedDrawParams.InterpolatingPoints.Emplace(
+				ScreenSpace.SecondsToScreen(ScreenSpace.GetInputMax()), 
+				ScreenSpace.ValueToScreen(DefaultValue));
+
+			return true;
+		}
+		else if (Times.Num() == 1)
+		{
+			const double SingleValue = Values[0].Value;
+			CachedDrawParams.InterpolatingPoints.Emplace(
+				ScreenSpace.SecondsToScreen(ScreenSpace.GetInputMax()),
+				ScreenSpace.ValueToScreen(SingleValue));
+
+			return true;
+		}
+		else
+		{
+			const bool bPostinfinityVisible = Times.Last().Value < ScreenSpace.GetInputMax() * TickResolution;
+			if (bPostinfinityVisible &&
+				(FullRangeInterpolation.PostInfinityExtrapolation == ERichCurveExtrapolation::RCCE_None ||
+				FullRangeInterpolation.PostInfinityExtrapolation == ERichCurveExtrapolation::RCCE_Constant))
+			{
+				const double ExtentValue = Values.Last().Value;
+				CachedDrawParams.InterpolatingPoints.Emplace(
+					ScreenSpace.SecondsToScreen(ScreenSpace.GetInputMax()), 
+					ScreenSpace.ValueToScreen(ExtentValue));
+
+				return true;
+			}
+			else if (bPostinfinityVisible &&
+				FullRangeInterpolation.PostInfinityExtrapolation == ERichCurveExtrapolation::RCCE_Linear)
+			{
+				const double ExtentValue = Values.Last().Value + KeyAttributes.Last().GetArriveTangent() * (ScreenSpace.GetInputMax() - Times.Last().Value);
+				CachedDrawParams.InterpolatingPoints.Emplace(
+					ScreenSpace.SecondsToScreen(ScreenSpace.GetInputMax()),
+					ScreenSpace.ValueToScreen(ExtentValue));
+
+				return true;
+			}
+		}
+
+		return false;
+	}
+
+	template <typename ChannelType>
+	void FMovieSceneCachedCurve<ChannelType>::DrawKeys()
+	{
+		const int32 DataSize = EndingIndex - StartingIndex + 1;
+		CachedDrawParams.Points.Reset(DataSize);
+
+		for (int32 DataIndex = StartingIndex; DataIndex <= EndingIndex; DataIndex++)
+		{
+			const FKeyHandle& KeyHandle = KeyHandles[DataIndex];
+			const FKeyAttributes& KeyAttribute = KeyAttributes[DataIndex];
+
+			// Key
+			FCurvePointInfo CurvePointInfo(KeyHandles[DataIndex]);
+
+			CurvePointInfo.Type = ECurvePointType::Key;
+			CurvePointInfo.LayerBias = 2;
+
+			CurvePointInfo.ScreenPosition.X = ScreenSpace.SecondsToScreen(Times[DataIndex] / TickResolution + InputDisplayOffset);
+			CurvePointInfo.ScreenPosition.Y = ScreenSpace.ValueToScreen(Values[DataIndex].Value);
+			CurvePointInfo.DrawInfo = KeyDrawInfos[DataIndex].template Get<0>();
+			CurvePointInfo.LineDelta = FVector2D::ZeroVector;
+
+			CachedDrawParams.Points.Add(CurvePointInfo);
+
+			// Arrive Tangent
+			{
+				if (KeyDrawInfos[DataIndex].template Get<1>().IsSet())
+				{
+					FCurvePointInfo ArriveTangentInfo(KeyHandle);
+
+					ArriveTangentInfo.Type = ECurvePointType::ArriveTangent;
+					ArriveTangentInfo.LayerBias = 1;
+
+					const float ArriveTangent = KeyAttribute.GetArriveTangent();
+
+					if (KeyAttribute.HasTangentWeightMode() && KeyAttribute.HasArriveTangentWeight() &&
+						(KeyAttribute.GetTangentWeightMode() == RCTWM_WeightedBoth || KeyAttribute.GetTangentWeightMode() == RCTWM_WeightedArrive))
+					{
+						FVector2D TangentOffset = ::CurveEditor::ComputeScreenSpaceTangentOffset(ScreenSpace, ArriveTangent, -KeyAttribute.GetArriveTangentWeight());
+						ArriveTangentInfo.ScreenPosition = CurvePointInfo.ScreenPosition + TangentOffset;
+					}
+					else
+					{
+						const double DisplayRatio = (ScreenSpace.PixelsPerOutput() / ScreenSpace.PixelsPerInput());
+
+						float PixelLength = 60.0f;
+						ArriveTangentInfo.ScreenPosition = CurvePointInfo.ScreenPosition + ::CurveEditor::GetVectorFromSlopeAndLength(ArriveTangent * -DisplayRatio, -PixelLength);
+					}
+
+					ArriveTangentInfo.DrawInfo = KeyDrawInfos[DataIndex].template Get<1>().GetValue();
+					ArriveTangentInfo.LineDelta = CurvePointInfo.ScreenPosition - ArriveTangentInfo.ScreenPosition;
+
+					CachedDrawParams.Points.Add(ArriveTangentInfo);
+				}
+			}
+
+			// Leave Tangent
+			{
+				if (KeyDrawInfos[DataIndex].template Get<2>().IsSet())
+				{
+					FCurvePointInfo LeaveTangentInfo(KeyHandle);
+
+					LeaveTangentInfo.Type = ECurvePointType::LeaveTangent;
+					LeaveTangentInfo.LayerBias = 1;
+
+					const float LeaveTangent = KeyAttribute.GetLeaveTangent();
+
+					if (KeyAttribute.HasTangentWeightMode() && KeyAttribute.HasLeaveTangentWeight() &&
+						(KeyAttribute.GetTangentWeightMode() == RCTWM_WeightedBoth || KeyAttribute.GetTangentWeightMode() == RCTWM_WeightedLeave))
+					{
+						FVector2D TangentOffset = ::CurveEditor::ComputeScreenSpaceTangentOffset(ScreenSpace, LeaveTangent, KeyAttribute.GetLeaveTangentWeight());
+						LeaveTangentInfo.ScreenPosition = CurvePointInfo.ScreenPosition + TangentOffset;
+					}
+					else
+					{
+						const double DisplayRatio = (ScreenSpace.PixelsPerOutput() / ScreenSpace.PixelsPerInput());
+
+						float PixelLength = 60.0f;
+						LeaveTangentInfo.ScreenPosition = CurvePointInfo.ScreenPosition + ::CurveEditor::GetVectorFromSlopeAndLength(LeaveTangent * -DisplayRatio, PixelLength);
+					}
+
+					LeaveTangentInfo.DrawInfo = KeyDrawInfos[DataIndex].template Get<2>().GetValue();
+					LeaveTangentInfo.LineDelta = CurvePointInfo.ScreenPosition - LeaveTangentInfo.ScreenPosition;
+
+					CachedDrawParams.Points.Add(LeaveTangentInfo);
+				}
+			}
+
+			if (DataIndex == EndingIndex)
+			{
+				break;
+			}
+		}
+	}
+
+	template <typename ChannelType>
+	void FMovieSceneCachedCurve<ChannelType>::ApplyDrawParams()
+	{
+		if (FCurveDrawParams* CurveDrawParamsPtr = DrawParamsHandle.Get())
+		{
+			*CurveDrawParamsPtr = CachedDrawParams;
+		}
+	}
+}
diff --git a/Engine/Source/Editor/MovieSceneTools/Private/Cache/MovieSceneCurveCachePool.cpp b/Engine/Source/Editor/MovieSceneTools/Private/Cache/MovieSceneCurveCachePool.cpp
new file mode 100644
--- /dev/null
+++ b/Engine/Source/Editor/MovieSceneTools/Private/Cache/MovieSceneCurveCachePool.cpp
@@ -0,0 +1,165 @@
+// Copyright Epic Games, Inc. All Rights Reserved.
+
+#include "MovieSceneCurveCachePool.h"
+
+#include "Algo/Find.h"
+#include "Algo/RemoveIf.h"
+#include "Async/ParallelFor.h"
+#include "Cache/MovieSceneCachedCurve.h"
+
+namespace UE::MovieSceneTools 
+{ 	
+	int32 GCullCachedCurves = 1; 
+
+	static FAutoConsoleVariableRef CCullCachedCurves(
+		TEXT("MovieSceneTools.CullCachedCurves"),
+		UE::MovieSceneTools::GCullCachedCurves,
+		TEXT("When set to true, movie scene cached curves cull draw params"),
+		ECVF_Default);
+
+	FMovieSceneCurveCachePool& FMovieSceneCurveCachePool::Get()
+	{
+		static TSharedRef<FMovieSceneCurveCachePool> Instance = MakeShared<FMovieSceneCurveCachePool>();
+		return *Instance;
+	}
+
+	void FMovieSceneCurveCachePool::DrawCachedCurves(TWeakPtr<const FCurveEditor> WeakCurveEditor)
+	{
+		const TArray<TSharedPtr<IMovieSceneCachedCurve>>* CachedCurvesPtr = CurveEditorToCachedCurvesMap.Find(WeakCurveEditor);
+		if (CachedCurvesPtr)
+		{
+			DrawCachedCurves(*CachedCurvesPtr);
+			CullCachedCurves(*CachedCurvesPtr);
+		}
+	}
+
+	void FMovieSceneCurveCachePool::Join(TWeakPtr<const FCurveEditor> WeakCurveEditor, const TSharedRef<IMovieSceneCachedCurve>& CachedCurve)
+	{
+		CurveEditorToCachedCurvesMap.FindOrAdd(WeakCurveEditor).AddUnique(CachedCurve);
+	}
+
+	void FMovieSceneCurveCachePool::Leave(const IMovieSceneCachedCurve& CachedCurve)
+	{
+		// Find where the curve was added
+		TTuple<TWeakPtr<const FCurveEditor>, TArray<TSharedPtr<IMovieSceneCachedCurve>>>* CurveEditorToCachedCurvesPairPtr =
+			Algo::FindByPredicate(CurveEditorToCachedCurvesMap, 
+				[&CachedCurve](const TTuple<TWeakPtr<const FCurveEditor>, TArray<TSharedPtr<IMovieSceneCachedCurve>>>& CurveEditorToCachedCurvesPair)
+				{
+					return Algo::FindByPredicate(CurveEditorToCachedCurvesPair.Value,
+						[&CachedCurve](const TSharedPtr<IMovieSceneCachedCurve>& OtherCachedCurve)
+						{
+							return OtherCachedCurve.Get() == &CachedCurve;
+						}) != nullptr;
+				});
+	
+		if (CurveEditorToCachedCurvesPairPtr)
+		{
+			for (auto MapIt = CurveEditorToCachedCurvesMap.CreateIterator(); MapIt; ++MapIt)
+			{
+				// Remove the curve
+				const int32 OldSize = (*CurveEditorToCachedCurvesPairPtr).Value.Num();
+				const int32 NewSize = Algo::RemoveIf((*CurveEditorToCachedCurvesPairPtr).Value,
+					[&CachedCurve](const TSharedPtr<IMovieSceneCachedCurve>& OtherCachedCurve)
+					{
+						return OtherCachedCurve.Get() == &CachedCurve;
+					});
+
+				if (OldSize != NewSize)
+				{
+					(*CurveEditorToCachedCurvesPairPtr).Value.SetNum(NewSize);
+
+					// Remove the curve editor if its there's no related curves anymore
+					if ((*CurveEditorToCachedCurvesPairPtr).Value.IsEmpty())
+					{
+						MapIt.RemoveCurrent();
+					}
+
+					break;
+				}
+			}
+		}
+	}
+
+	void FMovieSceneCurveCachePool::DrawCachedCurves(const TArray<TSharedPtr<IMovieSceneCachedCurve>>& CachedCurves) const
+	{
+		ParallelFor(CachedCurves.Num(),
+			[CachedCurves](int32 Index)
+			{
+				CachedCurves[Index]->DrawCachedCurve();
+
+			}, EParallelForFlags::Unbalanced);
+	}
+
+	void FMovieSceneCurveCachePool::CullCachedCurves(const TArray<TSharedPtr<IMovieSceneCachedCurve>>& CachedCurves) const
+	{
+		if (GCullCachedCurves <= 0)
+		{
+			return;
+		}
+
+		TMap<uint32, const FCurvePointInfo*> HashToKeyMap;
+		TMap<uint32, const FCurveDrawParams*> InterpolatingPointsHashToDrawParamsMap;
+
+		// Reverse to avoid culling those curves that are drawn last resp. top-most
+		TArray<TSharedPtr<IMovieSceneCachedCurve>> InverseCachedCurves = CachedCurves;
+		Algo::Reverse(InverseCachedCurves);
+
+		for (const TSharedPtr<IMovieSceneCachedCurve>& CachedCurve : InverseCachedCurves)
+		{
+			FCurveDrawParams* DrawParamsPtr = CachedCurve->GetDrawParamsHandle().Get();
+			if (!DrawParamsPtr)
+			{
+				continue;
+			}
+
+			for (FCurvePointInfo& Point : (*DrawParamsPtr).Points)
+			{
+				const uint32 PointHash = HashCombine(FMath::RoundToInt32(Point.ScreenPosition.X), FMath::RoundToInt32(Point.ScreenPosition.Y));
+
+				// Cull keys but not tangents
+				if (Point.Type == ECurvePointType::Key)
+				{
+					const FCurvePointInfo* const* OtherKeyPtr = HashToKeyMap.Find(PointHash);
+					if (OtherKeyPtr && (*OtherKeyPtr)->LayerBias <= Point.LayerBias)
+					{
+						Point.bDraw = false;
+					}
+					else
+					{
+						HashToKeyMap.Add(PointHash, &Point);
+					}
+				}
+			}
+
+			const uint32 InterpolatingPointsHash = CachedCurve->GetInterpolatingPointsHash();
+			const FCurveDrawParams** OtherDrawParamsPtrPtr = InterpolatingPointsHashToDrawParamsMap.Find(InterpolatingPointsHash);
+
+			bool bCull = false;
+			if (OtherDrawParamsPtrPtr && (*OtherDrawParamsPtrPtr)->InterpolatingPoints.Num() == (*DrawParamsPtr).InterpolatingPoints.Num())
+			{
+				bCull = true;
+
+				// Test for hash collisions
+				for (int32 PointIndex = 0; PointIndex < (*DrawParamsPtr).InterpolatingPoints.Num(); PointIndex++)
+				{
+					if (FMath::RoundToInt((*OtherDrawParamsPtrPtr)->InterpolatingPoints[PointIndex].X) != FMath::RoundToInt((*DrawParamsPtr).InterpolatingPoints[PointIndex].X) ||
+						FMath::RoundToInt((*OtherDrawParamsPtrPtr)->InterpolatingPoints[PointIndex].Y) != FMath::RoundToInt((*DrawParamsPtr).InterpolatingPoints[PointIndex].Y))
+					{
+						bCull = false;
+						break;
+					}
+				}
+			}
+
+			if (bCull)
+			{
+				(*DrawParamsPtr).bDrawInterpolatingPoints = false;
+			}
+			else
+			{
+				InterpolatingPointsHashToDrawParamsMap.Add(InterpolatingPointsHash, DrawParamsPtr);
+				(*DrawParamsPtr).bDrawInterpolatingPoints = true;
+			}
+		}
+	}
+}
diff --git a/Engine/Source/Editor/MovieSceneTools/Private/Cache/MovieSceneCurveCachePool.h b/Engine/Source/Editor/MovieSceneTools/Private/Cache/MovieSceneCurveCachePool.h
new file mode 100644
--- /dev/null
+++ b/Engine/Source/Editor/MovieSceneTools/Private/Cache/MovieSceneCurveCachePool.h
@@ -0,0 +1,63 @@
+// Copyright Epic Games, Inc. All Rights Reserved.
+
+#pragma once
+
+#include "ICurveEditorCurveCachePool.h"
+
+class FCurveEditor;
+namespace UE::CurveEditor { struct FCurveDrawParamsHandle; }
+
+namespace UE::MovieSceneTools
+{
+	/** Interface for a cached curve in the curve cache pool */
+	class IMovieSceneCachedCurve
+		: public TSharedFromThis<IMovieSceneCachedCurve>
+	{
+	public:
+		IMovieSceneCachedCurve() = default;
+		virtual ~IMovieSceneCachedCurve() = default;
+
+		/** Returns the curve model ID that corresponds to this curve */
+		virtual const FCurveModelID& GetID() const = 0;
+
+		/** Draws the cached curve */
+		virtual void DrawCachedCurve() = 0;
+
+		/** Returns a hash of the currently visible interpolated points */
+		virtual uint32 GetInterpolatingPointsHash() const = 0;
+
+		/** Returns a handle to the actual draw params */
+		virtual const UE::CurveEditor::FCurveDrawParamsHandle& GetDrawParamsHandle() const = 0;
+	};
+
+	/** A pool of cached curves in a curve editor */
+	class FMovieSceneCurveCachePool
+		: public UE::CurveEditor::ICurveEditorCurveCachePool
+		, public TSharedFromThis<FMovieSceneCurveCachePool>
+	{
+	public:
+		virtual ~FMovieSceneCurveCachePool() = default;
+
+		static FMovieSceneCurveCachePool& Get();
+
+		// ICurveEditorCurveCachePool Interface
+		virtual void DrawCachedCurves(TWeakPtr<const FCurveEditor> WeakCurveEditor) override;
+		//~ICurveEditorCurveCachePool Interface
+
+		/** Lets a cached curve join the cache pool. Curves that joined should also leave */
+		void Join(TWeakPtr<const FCurveEditor> WeakCurveEditor, const TSharedRef<IMovieSceneCachedCurve>& CachedCurve);
+
+		/** Lets a cached curve leave the cache pool */
+		void Leave(const IMovieSceneCachedCurve& CachedCurve);
+
+	private:
+		/** Lets cached curves draw to their draw params */
+		void DrawCachedCurves(const TArray<TSharedPtr<IMovieSceneCachedCurve>>& CachedCurves) const;
+
+		/** Culls cached curves */
+		void CullCachedCurves(const TArray<TSharedPtr<IMovieSceneCachedCurve>>& CachedCurves) const;
+
+		/** Cached curves per curve editor */
+		TMap<TWeakPtr<const FCurveEditor>, TArray<TSharedPtr<IMovieSceneCachedCurve>>> CurveEditorToCachedCurvesMap;
+	};
+}
diff --git a/Engine/Source/Editor/MovieSceneTools/Private/Cache/MovieSceneUpdateCachedCurveData.h b/Engine/Source/Editor/MovieSceneTools/Private/Cache/MovieSceneUpdateCachedCurveData.h
new file mode 100644
--- /dev/null
+++ b/Engine/Source/Editor/MovieSceneTools/Private/Cache/MovieSceneUpdateCachedCurveData.h
@@ -0,0 +1,48 @@
+// Copyright Epic Games, Inc. All Rights Reserved.
+
+#pragma once
+
+#include "CurveDrawInfo.h"
+#include "Templates/UnrealTemplate.h"
+
+class FCurveEditor;
+class FCurveModel;
+struct FCurveEditorScreenSpace;
+
+namespace UE::MovieSceneTools
+{
+	/** A struct that holds relevant parameters to cache a curve */
+	template <typename ChannelType>
+	struct FMovieSceneUpdateCachedCurveData
+		: public FNoncopyable
+	{
+	private:
+		/** The type of channel value structs */
+		using ChannelValueType = typename ChannelType::ChannelValueType;
+
+	public:
+		FMovieSceneUpdateCachedCurveData(
+			const FCurveEditor& InCurveEditor,
+			const FCurveModel& InCurveModel,
+			const ChannelType& InChannel,
+			const FCurveEditorScreenSpace& InScreenSpace,
+			const FFrameRate& InTickResolution)
+			: CurveEditor(InCurveEditor)
+			, CurveModel(InCurveModel)
+			, Channel(InChannel)
+			, ScreenSpace(InScreenSpace)
+			, TickResolution(InTickResolution)
+			, Times(InChannel.GetTimes())
+			, Values(InChannel.GetValues())
+		{};
+
+		const FCurveEditor& CurveEditor;
+		const FCurveModel& CurveModel;
+		const ChannelType& Channel;
+		const FCurveEditorScreenSpace& ScreenSpace;
+		const FFrameRate& TickResolution;
+
+		const TArrayView<const FFrameNumber> Times;
+		const TArrayView<const ChannelValueType> Values;
+	};
+}
diff --git a/Engine/Source/Editor/MovieSceneTools/Private/Channels/BezierChannelCurveModel.cpp b/Engine/Source/Editor/MovieSceneTools/Private/Channels/BezierChannelCurveModel.cpp
--- a/Engine/Source/Editor/MovieSceneTools/Private/Channels/BezierChannelCurveModel.cpp
+++ b/Engine/Source/Editor/MovieSceneTools/Private/Channels/BezierChannelCurveModel.cpp
@@ -2,6 +2,8 @@
 
 #include "Channels/BezierChannelCurveModel.h"
 
+#include "Cache/MovieSceneCachedCurve.h"
+#include "Cache/MovieSceneUpdateCachedCurveData.h"
 #include "Channels/MovieSceneChannelData.h"
 #include "Channels/MovieSceneCurveChannelCommon.h"
 #include "Channels/MovieSceneDoubleChannel.h"
@@ -31,13 +33,12 @@ class FCurveEditor;
 
 template<typename ChannelType>
 FBezierChannelBufferedCurveModel<ChannelType>::FBezierChannelBufferedCurveModel(
-		const ChannelType* InChannel, TWeakObjectPtr<UMovieSceneSection> InWeakSection,
-		TArray<FKeyPosition>&& InKeyPositions, TArray<FKeyAttributes>&& InKeyAttributes, const FString& InLongDisplayName, const double InValueMin, const double InValueMax)
+	const ChannelType* InChannel, TWeakObjectPtr<UMovieSceneSection> InWeakSection,
+	TArray<FKeyPosition>&& InKeyPositions, TArray<FKeyAttributes>&& InKeyAttributes, const FString& InLongDisplayName, const double InValueMin, const double InValueMax)
 	: IBufferedCurveModel(MoveTemp(InKeyPositions), MoveTemp(InKeyAttributes), InLongDisplayName, InValueMin, InValueMax)
 	, Channel(*InChannel)
 	, WeakSection(InWeakSection)
-{
-}
+{}
 
 template<typename ChannelType>
 void FBezierChannelBufferedCurveModel<ChannelType>::DrawCurve(const FCurveEditor& InCurveEditor, const FCurveEditorScreenSpace& InScreenSpace, TArray<TTuple<double, double>>& OutInterpolatingPoints) const
@@ -84,26 +85,62 @@ FBezierChannelCurveModel<ChannelType, ChannelValue, KeyType>::FBezierChannelCurv
 	}
 }
 
-template<typename ChannelType, typename ChannelValue, typename KeyType> 
-void FBezierChannelCurveModel<ChannelType, ChannelValue, KeyType>::DrawCurve(const FCurveEditor& CurveEditor, const FCurveEditorScreenSpace& ScreenSpace, TArray<TTuple<double, double>>& InterpolatingPoints) const
+template<typename ChannelType, typename ChannelValue, typename KeyType>
+void FBezierChannelCurveModel<ChannelType, ChannelValue, KeyType>::DrawCurve(const FCurveEditor& InCurveEditor, const FCurveEditorScreenSpace& InScreenSpace, TArray<TTuple<double, double>>& OutInterpolatingPoints) const
 {
-	ChannelType* Channel		= this->GetChannelHandle().Get();
+	ChannelType* Channel = this->GetChannelHandle().Get();
 	UMovieSceneSection* Section = this->template GetOwningObjectOrOuter<UMovieSceneSection>();
 
 	if (Channel && Section && Section->GetTypedOuter<UMovieScene>())
 	{
-		FFrameRate   TickResolution   = Section->GetTypedOuter<UMovieScene>()->GetTickResolution();
+		FFrameRate   TickResolution = Section->GetTypedOuter<UMovieScene>()->GetTickResolution();
 
-		const double DisplayOffset    = this->GetInputDisplayOffset();
-		const double StartTimeSeconds = ScreenSpace.GetInputMin() - DisplayOffset;
-		const double EndTimeSeconds   = ScreenSpace.GetInputMax() - DisplayOffset;
-		const double TimeThreshold    = FMath::Max(0.0001, 1.0 / ScreenSpace.PixelsPerInput());
-		const double ValueThreshold   = FMath::Max(0.0001, 1.0 / ScreenSpace.PixelsPerOutput());
+		const double StartTimeSeconds = InScreenSpace.GetInputMin();
+		const double EndTimeSeconds = InScreenSpace.GetInputMax();
+		const double TimeThreshold = FMath::Max(0.0001, 1.0 / InScreenSpace.PixelsPerInput());
+		const double ValueThreshold = FMath::Max(0.0001, 1.0 / InScreenSpace.PixelsPerOutput());
 
-		Channel->PopulateCurvePoints(StartTimeSeconds, EndTimeSeconds, TimeThreshold, ValueThreshold, TickResolution, InterpolatingPoints);
+		Channel->PopulateCurvePoints(StartTimeSeconds, EndTimeSeconds, TimeThreshold, ValueThreshold, TickResolution, OutInterpolatingPoints);
 	}
 }
 
+template<typename ChannelType, typename ChannelValue, typename KeyType>
+UE::CurveEditor::ICurveEditorCurveCachePool* FBezierChannelCurveModel<ChannelType, ChannelValue, KeyType>::DrawCurveToCachePool(
+	const TSharedRef<FCurveEditor>& CurveEditor,
+	const UE::CurveEditor::FCurveDrawParamsHandle& CurveDrawParamsHandle,
+	const FCurveEditorScreenSpace& ScreenSpace)
+{
+	ChannelType* Channel = this->GetChannelHandle().Get();
+	UMovieSceneSection* Section = this->template GetOwningObjectOrOuter<UMovieSceneSection>();
+
+	if (!CachedCurve.IsValid())
+	{
+		CachedCurve = MakeShared<UE::MovieSceneTools::FMovieSceneCachedCurve<ChannelType>>(CurveDrawParamsHandle.GetID());
+		CachedCurve->Initialize(CurveEditor);
+	}
+
+	if (Channel && Section && Section->GetTypedOuter<UMovieScene>())
+	{
+		const FFrameRate TickResolution = Section->GetTypedOuter<UMovieScene>()->GetTickResolution();
+
+		UE::MovieSceneTools::FMovieSceneUpdateCachedCurveData<ChannelType> UpdateData(*CurveEditor, *this, *Channel, ScreenSpace, TickResolution);
+		CachedCurve->UpdateCachedCurve(UpdateData, CurveDrawParamsHandle);
+	}
+
+	return &UE::MovieSceneTools::FMovieSceneCurveCachePool::Get();
+}
+
+template<typename ChannelType, typename ChannelValue, typename KeyType>
+bool FBezierChannelCurveModel<ChannelType, ChannelValue, KeyType>::HasChangedAndResetTest()
+{
+	if (CachedCurve.IsValid())
+	{
+		return CachedCurve->HasChanged() || Super::HasChangedAndResetTest();
+	}
+	
+	return Super::HasChangedAndResetTest();
+}
+
 template<typename ChannelType, typename ChannelValue, typename KeyType> 
 void FBezierChannelCurveModel<ChannelType, ChannelValue, KeyType>::GetKeyDrawInfo(ECurvePointType PointType, const FKeyHandle InKeyHandle, FKeyDrawInfo& OutDrawInfo) const
 {
diff --git a/Engine/Source/Editor/MovieSceneTools/Public/Channels/BezierChannelCurveModel.h b/Engine/Source/Editor/MovieSceneTools/Public/Channels/BezierChannelCurveModel.h
--- a/Engine/Source/Editor/MovieSceneTools/Public/Channels/BezierChannelCurveModel.h
+++ b/Engine/Source/Editor/MovieSceneTools/Public/Channels/BezierChannelCurveModel.h
@@ -29,6 +29,8 @@ struct FKeyAttributes;
 struct FKeyDrawInfo;
 struct FKeyPosition;
 
+namespace UE::MovieSceneTools { template<typename ChannelType> class FMovieSceneCachedCurve; }
+
 /**
  * Buffered curve implementation for a bezier curve model, stores a copy of the bezier curve channel in order to draw itself.
  */
@@ -53,12 +55,16 @@ private:
 template<typename ChannelType, typename ChannelValue, typename KeyType> 
 class FBezierChannelCurveModel : public FChannelCurveModel<ChannelType, ChannelValue, KeyType>
 {
+	using Super = FChannelCurveModel<ChannelType, ChannelValue, KeyType>;
+
 public:
 	FBezierChannelCurveModel(TMovieSceneChannelHandle<ChannelType> InChannel, UMovieSceneSection* InOwningSection, TWeakPtr<ISequencer> InWeakSequencer);
 	FBezierChannelCurveModel(TMovieSceneChannelHandle<ChannelType> InChannel, UMovieSceneSection* InOwningSection, UObject* InOwningObject, TWeakPtr<ISequencer> InWeakSequencer);
 
 	// FCurveModel
 	virtual void DrawCurve(const FCurveEditor& CurveEditor, const FCurveEditorScreenSpace& ScreenSpace, TArray<TTuple<double, double>>& InterpolatingPoints) const override;
+	virtual UE::CurveEditor::ICurveEditorCurveCachePool* DrawCurveToCachePool(const TSharedRef<FCurveEditor>& CurveEditor, const UE::CurveEditor::FCurveDrawParamsHandle& CurveDrawParamsHandle, const FCurveEditorScreenSpace& ScreenSpace) override;
+	virtual bool HasChangedAndResetTest() override;
 	virtual void GetKeyDrawInfo(ECurvePointType PointType, const FKeyHandle InKeyHandle, FKeyDrawInfo& OutDrawInfo) const override;
 
 	virtual void GetKeyAttributes(TArrayView<const FKeyHandle> InKeys, TArrayView<FKeyAttributes> OutAttributes) const override;
@@ -78,5 +84,7 @@ protected:
 
 private:
 	void FeaturePointMethod(double StartTime, double EndTime, double StartValue, double Mu, int Depth, int MaxDepth, double& MaxV, double& MinVal) const;
+
+	TSharedPtr<UE::MovieSceneTools::FMovieSceneCachedCurve<ChannelType>> CachedCurve;
 };
 
diff --git a/Engine/Source/Runtime/MovieScene/Private/Channels/MovieSceneDoubleChannel.cpp b/Engine/Source/Runtime/MovieScene/Private/Channels/MovieSceneDoubleChannel.cpp
--- a/Engine/Source/Runtime/MovieScene/Private/Channels/MovieSceneDoubleChannel.cpp
+++ b/Engine/Source/Runtime/MovieScene/Private/Channels/MovieSceneDoubleChannel.cpp
@@ -588,7 +588,7 @@ UE::MovieScene::Interpolation::FInterpolationExtents FMovieSceneDoubleChannel::C
 	return ComputePiecewiseExtents(FDoubleChannelPiecewiseData{ this }, StartTime, EndTime);
 }
 
-UE::MovieScene::FPiecewiseCurve FMovieSceneDoubleChannel::AsPiecewiseCurve() const
+UE::MovieScene::FPiecewiseCurve FMovieSceneDoubleChannel::AsPiecewiseCurve(bool bWithPreAndPostInfinityExtrap) const
 {
 	using namespace UE::MovieScene;
 	using namespace UE::MovieScene::Interpolation;
@@ -604,7 +604,7 @@ UE::MovieScene::FPiecewiseCurve FMovieSceneDoubleChannel::AsPiecewiseCurve() con
 		return Curve;
 	}
 
-	if (PreInfinityExtrap != RCCE_None)
+	if (bWithPreAndPostInfinityExtrap && PreInfinityExtrap != RCCE_None)
 	{
 		Interpolation::FCachedInterpolation PreExtrap;
 		if (FMovieSceneDoubleChannelImpl::CacheExtrapolation(this, Times[0] - 1, PreExtrap))
@@ -629,7 +629,7 @@ UE::MovieScene::FPiecewiseCurve FMovieSceneDoubleChannel::AsPiecewiseCurve() con
 		}
 	}
 
-	if (PostInfinityExtrap != RCCE_None)
+	if (bWithPreAndPostInfinityExtrap && PostInfinityExtrap != RCCE_None)
 	{
 		Interpolation::FCachedInterpolation PostExtrap;
 		if (FMovieSceneDoubleChannelImpl::CacheExtrapolation(this, Times.Last() + 1, PostExtrap))
diff --git a/Engine/Source/Runtime/MovieScene/Private/Channels/MovieSceneFloatChannel.cpp b/Engine/Source/Runtime/MovieScene/Private/Channels/MovieSceneFloatChannel.cpp
--- a/Engine/Source/Runtime/MovieScene/Private/Channels/MovieSceneFloatChannel.cpp
+++ b/Engine/Source/Runtime/MovieScene/Private/Channels/MovieSceneFloatChannel.cpp
@@ -67,6 +67,63 @@ bool FMovieSceneFloatValue::operator!=(const FMovieSceneFloatValue& Other) const
 	return !(*this == Other);
 }
 
+UE::MovieScene::FPiecewiseCurve FMovieSceneFloatChannel::AsPiecewiseCurve(bool bWithPreAndPostInfinityExtrap) const
+{
+	using namespace UE::MovieScene;
+	using namespace UE::MovieScene::Interpolation;
+
+	FPiecewiseCurve Curve;
+
+	if (Times.IsEmpty())
+	{
+		if (bHasDefaultValue)
+		{
+			Curve.Values.Add(FCachedInterpolation(FCachedInterpolationRange::Infinite(), FConstantValue(0, DefaultValue)));
+		}
+		return Curve;
+	}
+
+	if (bWithPreAndPostInfinityExtrap && PreInfinityExtrap != RCCE_None)
+	{
+		Interpolation::FCachedInterpolation PreExtrap;
+		if (FMovieSceneFloatChannelImpl::CacheExtrapolation(this, Times[0] - 1, PreExtrap))
+		{
+			Curve.Values.Emplace(PreExtrap);
+		}
+		else
+		{
+			ensureMsgf(false, TEXT("Unrepresentable extrapolation mode encountered for piecewise curve"));
+		}
+	}
+
+	if (Times.Num() > 1)
+	{
+		for (int32 Index = 0; Index < Times.Num() - 1; ++Index)
+		{
+			Interpolation::FCachedInterpolation Interp = FMovieSceneFloatChannelImpl::GetInterpolationForKey(this, Index);
+			if (ensure(Interp.IsValid()))
+			{
+				Curve.Values.Emplace(Interp);
+			}
+		}
+	}
+	
+	if (bWithPreAndPostInfinityExtrap && PostInfinityExtrap != RCCE_None)
+	{
+		Interpolation::FCachedInterpolation PostExtrap;
+		if (FMovieSceneFloatChannelImpl::CacheExtrapolation(this, Times.Last() + 1, PostExtrap))
+		{
+			Curve.Values.Emplace(PostExtrap);
+		}
+		else
+		{
+			ensureMsgf(false, TEXT("Unrepresentable extrapolation mode encountered for piecewise curve"));
+		}
+	}
+
+	return Curve;
+}
+
 int32 FMovieSceneFloatChannel::AddConstantKey(FFrameNumber InTime, float InValue)
 {
 	return FMovieSceneFloatChannelImpl::AddConstantKey(this, InTime, InValue);
diff --git a/Engine/Source/Runtime/MovieScene/Public/Channels/MovieSceneDoubleChannel.h b/Engine/Source/Runtime/MovieScene/Public/Channels/MovieSceneDoubleChannel.h
--- a/Engine/Source/Runtime/MovieScene/Public/Channels/MovieSceneDoubleChannel.h
+++ b/Engine/Source/Runtime/MovieScene/Public/Channels/MovieSceneDoubleChannel.h
@@ -62,8 +62,8 @@ struct FMovieSceneDoubleValue
 	{}
 
 	bool Serialize(FArchive& Ar);
-	bool operator==(const FMovieSceneDoubleValue& Other) const;
-	bool operator!=(const FMovieSceneDoubleValue& Other) const;
+	MOVIESCENE_API bool operator==(const FMovieSceneDoubleValue& Other) const;
+	MOVIESCENE_API bool operator!=(const FMovieSceneDoubleValue& Other) const;
 	friend FArchive& operator<<(FArchive& Ar, FMovieSceneDoubleValue& P)
 	{
 		P.Serialize(Ar);
@@ -138,7 +138,7 @@ struct FMovieSceneDoubleChannel : public FMovieSceneChannel
 #endif
 	{}
 
-	UE::MovieScene::FPiecewiseCurve AsPiecewiseCurve() const;
+	MOVIESCENE_API UE::MovieScene::FPiecewiseCurve AsPiecewiseCurve(bool bWithPreAndPostInfinityExtrap = true) const;
 
 	/**
 	 * Access a mutable interface for this channel's data
diff --git a/Engine/Source/Runtime/MovieScene/Public/Channels/MovieSceneFloatChannel.h b/Engine/Source/Runtime/MovieScene/Public/Channels/MovieSceneFloatChannel.h
--- a/Engine/Source/Runtime/MovieScene/Public/Channels/MovieSceneFloatChannel.h
+++ b/Engine/Source/Runtime/MovieScene/Public/Channels/MovieSceneFloatChannel.h
@@ -4,6 +4,7 @@
 
 #include "Channels/MovieSceneChannelEditorData.h"
 #include "Channels/MovieSceneCurveChannelCommon.h"
+#include "Channels/MovieScenePiecewiseCurve.h"
 #include "Containers/Array.h"
 #include "Containers/ArrayView.h"
 #include "Containers/EnumAsByte.h"
@@ -58,8 +59,8 @@ struct FMovieSceneFloatValue
 	{}
 
 	bool Serialize(FArchive& Ar);
-	bool operator==(const FMovieSceneFloatValue& Other) const;
-	bool operator!=(const FMovieSceneFloatValue& Other) const;
+	MOVIESCENE_API bool operator==(const FMovieSceneFloatValue& Other) const;
+	MOVIESCENE_API bool operator!=(const FMovieSceneFloatValue& Other) const;
 	friend FArchive& operator<<(FArchive& Ar, FMovieSceneFloatValue& P)
 	{
 		P.Serialize(Ar);
@@ -134,6 +135,8 @@ struct FMovieSceneFloatChannel : public FMovieSceneChannel
 #endif
 	{}
 
+	MOVIESCENE_API UE::MovieScene::FPiecewiseCurve AsPiecewiseCurve(bool bWithPreAndPostInfinityExtrap = true) const;
+
 	/**
 	 * Access a mutable interface for this channel's data
 	 *
