Sequencer: add property metadata for better naming tracks that animate struct-wrapped-values

Some systems use structures to wrap values, similar to : struct FMyValue { float Value; int32 MetaData; }

Another case is for property bags, where there's a "Value" between the property bag and the properties.

When animating values like this, we want to ignore these intermediate property names to name the new property track. This metadata allows this.

diff --git a/Engine/Source/Editor/MovieSceneTools/Public/PropertyTrackEditor.h b/Engine/Source/Editor/MovieSceneTools/Public/PropertyTrackEditor.h
--- a/Engine/Source/Editor/MovieSceneTools/Public/PropertyTrackEditor.h
+++ b/Engine/Source/Editor/MovieSceneTools/Public/PropertyTrackEditor.h
@@ -106,22 +106,25 @@ protected:
 		const FProperty* ChangedProperty = PropertyChangedParams.PropertyPath.GetLeafMostProperty().Property.Get();
 		if (ChangedProperty)
 		{
-			NewTrack->SetPropertyNameAndPath( ChangedProperty->GetFName(), PropertyChangedParams.GetPropertyPathString() );
-#if WITH_EDITORONLY_DATA
+			NewTrack->SetPropertyNameAndPath(ChangedProperty->GetFName(), PropertyChangedParams.GetPropertyPathString());
 
+#if WITH_EDITORONLY_DATA
 			FText DisplayText;
+			const FPropertyPath NamingPropertyPath = BuildNamingPropertyPath(PropertyChangedParams.PropertyPath);
+			const FProperty* LeafNameProperty = NamingPropertyPath.GetLeafMostProperty().Property.Get();
+			const int32 NumProperties = NamingPropertyPath.GetNumProperties();
 
 			// Set up the appropriate name for the track from an array/nested struct index if necessary
-			for (int32 PropertyIndex = PropertyChangedParams.PropertyPath.GetNumProperties() - 1; PropertyIndex >= 0; --PropertyIndex)
+			for (int32 PropertyIndex = NumProperties - 1; PropertyIndex >= 0; --PropertyIndex)
 			{
-				const FPropertyInfo& Info = PropertyChangedParams.PropertyPath.GetPropertyInfo(PropertyIndex);
-				const FArrayProperty* ParentArrayProperty = PropertyIndex > 0 ? CastField<FArrayProperty>(PropertyChangedParams.PropertyPath.GetPropertyInfo(PropertyIndex - 1).Property.Get()) : nullptr;
+				const FPropertyInfo& Info = NamingPropertyPath.GetPropertyInfo(PropertyIndex);
+				const FArrayProperty* ParentArrayProperty = PropertyIndex > 0 ? CastField<FArrayProperty>(NamingPropertyPath.GetPropertyInfo(PropertyIndex - 1).Property.Get()) : nullptr;
 
 				const FProperty* ArrayInnerProperty = Info.Property.Get();
 				if (ArrayInnerProperty && Info.ArrayIndex != INDEX_NONE)
 				{
 					DisplayText = FText::Format(LOCTEXT("DisplayTextArrayFormat", "{0} ({1}[{2}])"),
-						ChangedProperty->GetDisplayNameText(),
+						LeafNameProperty->GetDisplayNameText(),
 						(ParentArrayProperty ? ParentArrayProperty : ArrayInnerProperty)->GetDisplayNameText(),
 						FText::AsNumber(Info.ArrayIndex)
 					);
@@ -131,13 +134,13 @@ protected:
 
 			if (DisplayText.IsEmpty())
 			{
-				for (int32 PropertyIndex = PropertyChangedParams.PropertyPath.GetNumProperties() - 1; PropertyIndex >= 0; --PropertyIndex)
+				for (int32 PropertyIndex = NumProperties - 1; PropertyIndex >= 0; --PropertyIndex)
 				{
 					const FStructProperty* ParentStructProperty = PropertyIndex > 0 ? CastField<FStructProperty>(PropertyChangedParams.PropertyPath.GetPropertyInfo(PropertyIndex - 1).Property.Get()) : nullptr;
 					if (ParentStructProperty)
 					{
 						DisplayText = FText::Format(LOCTEXT("DisplayTextStructFormat", "{0} ({1})"),
-							ChangedProperty->GetDisplayNameText(),
+							LeafNameProperty->GetDisplayNameText(),
 							ParentStructProperty->GetDisplayNameText()
 						);
 						break;
@@ -147,7 +150,7 @@ protected:
 
 			if (DisplayText.IsEmpty())
 			{
-				DisplayText = ChangedProperty->GetDisplayNameText();
+				DisplayText = LeafNameProperty->GetDisplayNameText();
 			}
 
 			NewTrack->SetDisplayName(DisplayText);
@@ -189,6 +192,28 @@ private:
 		AddWatchedProperty(FAnimatedPropertyKey::FromPropertyTypeName(WatchedPropertyTypeName));
 	}
 
+	/** 
+	 * Builds a property path that omits properties that are "insignificant" for computing track names. 
+	 * IMPORTANT: the result path is only for looking up names, and is non functional!
+	 */
+	FPropertyPath BuildNamingPropertyPath(const FPropertyPath& InPropertyPath)
+	{
+		FPropertyPath Result;
+		for (int32 Index = 0; Index < InPropertyPath.GetNumProperties(); ++Index)
+		{
+			const FPropertyInfo& Info = InPropertyPath.GetPropertyInfo(Index);
+			if (const FProperty* Property = Info.Property.Get())
+			{
+				if (Info.ArrayIndex == INDEX_NONE && Property->GetBoolMetaData(TEXT("SequencerUseParentPropertyName")))
+				{
+					continue;
+				}
+			}
+			Result.AddProperty(Info);
+		}
+		return Result;
+	}
+
 	/** Get a customized track class from the property if there is one, otherwise return nullptr. */
 	TSubclassOf<UMovieSceneTrack> GetCustomizedTrackClass( const FProperty* Property )
 	{
diff --git a/Engine/Source/Runtime/CoreUObject/Public/StructUtils/PropertyBag.h b/Engine/Source/Runtime/CoreUObject/Public/StructUtils/PropertyBag.h
--- a/Engine/Source/Runtime/CoreUObject/Public/StructUtils/PropertyBag.h
+++ b/Engine/Source/Runtime/CoreUObject/Public/StructUtils/PropertyBag.h
@@ -774,7 +774,7 @@ protected:
 	const void* GetValueAddress(const FPropertyBagPropertyDesc* Desc) const;
 	void* GetMutableValueAddress(const FPropertyBagPropertyDesc* Desc);
 	
-	UPROPERTY(EditAnywhere, Category="")
+	UPROPERTY(EditAnywhere, Category="", meta=(SequencerUseParentPropertyName=true))
 	FInstancedStruct Value;
 };
 
@@ -1199,4 +1199,4 @@ protected:
 	std::atomic<int32> RefCount = 0;
 	
 	friend struct FInstancedPropertyBag;
-};
\ No newline at end of file
+};
