Sequencer: Send a track value changed if any property changed for skeletal animations. This fixes a bug where if you change any property on a skeletal mesh, the pose will revert to a tpose until you scrub Sequencer again. This is because the change causes the skeletal mesh to update/tick, but since Sequencer data isn't actually changed (unless you autokeky), Sequencer doesn't evaluate again.

diff --git a/Engine/Source/Editor/MovieSceneTools/Private/TrackEditors/SkeletalAnimationTrackEditor.cpp b/Engine/Source/Editor/MovieSceneTools/Private/TrackEditors/SkeletalAnimationTrackEditor.cpp
--- a/Engine/Source/Editor/MovieSceneTools/Private/TrackEditors/SkeletalAnimationTrackEditor.cpp
+++ b/Engine/Source/Editor/MovieSceneTools/Private/TrackEditors/SkeletalAnimationTrackEditor.cpp
@@ -1627,51 +1627,26 @@ void FSkeletalAnimationTrackEditor::OnSequencerDataChanged(EMovieSceneDataChange
 
 void FSkeletalAnimationTrackEditor::OnPostPropertyChanged(UObject* InObject, FPropertyChangedEvent& InPropertyChangedEvent)
 {
-	// This attempts to fix tposes when changing properties without evaluating Sequencer UE-101261, 
-	// but unfortunately causes other problems like the temporary unkeyed value getting lost UE-136405
-/*
-	if (InPropertyChangedEvent.ChangeType != EPropertyChangeType::ValueSet)
+	const FName PropertyName = (InPropertyChangedEvent.Property != nullptr) ? InPropertyChangedEvent.Property->GetFName() : NAME_None;
+	if (PropertyName == USceneComponent::GetRelativeLocationPropertyName() ||
+		PropertyName == USceneComponent::GetRelativeRotationPropertyName() ||
+		PropertyName == USceneComponent::GetRelativeScale3DPropertyName())
 	{
 		return;
 	}
 
-	// If the object changed has any animation track:
-	// 1. Store the current transform (which may be an unkeyed value),
-	// 2. Evaluate Sequencer so that the skeletal animation track will be evaluated, and then the skeletal mesh with tick
-	// 3. Restore the current transform
-	// Without this, changing a value on a skeletal mesh will tick but not necessarily evaluate Sequencer, resulting in a tpose.
+	// If the object changed has any animation track, notify sequencer to update because animations tick on their own and sequencer needs to evaluate again
 	const bool bCreateIfMissing = false;
-	FFindOrCreateHandleResult HandleResult = FindOrCreateHandleToObject(InObject, bCreateIfMissing );
+	FFindOrCreateHandleResult HandleResult = FindOrCreateHandleToObject(InObject, bCreateIfMissing);
 	FGuid ObjectHandle = HandleResult.Handle;
 	if (ObjectHandle.IsValid())
 	{
 		FFindOrCreateTrackResult TrackResult = FindOrCreateTrackForObject(ObjectHandle, UMovieSceneSkeletalAnimationTrack::StaticClass(), NAME_None, bCreateIfMissing);
 		if (TrackResult.Track)
 		{
-			USceneComponent* SceneComponent = Cast<USceneComponent>(InObject);
-			if (!SceneComponent)
-			{
-				if (AActor* Actor = Cast<AActor>(InObject))
-				{
-					SceneComponent = Actor->GetRootComponent();
-				}
-			}
-
-			FTransform RelativeTransform;
-			if (SceneComponent)
-			{
-				RelativeTransform = SceneComponent->GetRelativeTransform();
-			}
-
-			GetSequencer()->ForceEvaluate();
-
-			if (SceneComponent)
-			{
-				SceneComponent->SetRelativeTransform(RelativeTransform);
-			}
+			GetSequencer()->NotifyMovieSceneDataChanged(EMovieSceneDataChangeType::TrackValueChanged);
 		}
 	}
-*/
 }
 
 bool FSkeletalAnimationTrackEditor::CreateAnimationSequence(const TArray<UObject*> NewAssets, USkeletalMeshComponent* SkelMeshComp, FGuid Binding, bool bCreateSoftLink)
