Curve Editor: Fix a crash when drawing pre- or post-infinity points of a linear curve that has no arrive resp. leave tangent

diff --git a/Engine/Source/Editor/MovieSceneTools/Private/Cache/MovieSceneCachedCurve.h b/Engine/Source/Editor/MovieSceneTools/Private/Cache/MovieSceneCachedCurve.h
--- a/Engine/Source/Editor/MovieSceneTools/Private/Cache/MovieSceneCachedCurve.h
+++ b/Engine/Source/Editor/MovieSceneTools/Private/Cache/MovieSceneCachedCurve.h
@@ -792,7 +792,8 @@ namespace UE::MovieSceneTools
 			const bool bPreinfinityVisible = Times[0].Value > ScreenSpace.GetInputMin() * TickResolution;
 			if (bPreinfinityVisible &&
 				(FullRangeInterpolation.PreInfinityExtrapolation == ERichCurveExtrapolation::RCCE_None ||
-				FullRangeInterpolation.PreInfinityExtrapolation == ERichCurveExtrapolation::RCCE_Constant))
+				FullRangeInterpolation.PreInfinityExtrapolation == ERichCurveExtrapolation::RCCE_Constant ||
+				!KeyAttributes[0].HasArriveTangent()))
 			{
 				const double ExtentValue = Values[0].Value;
 				CachedDrawParams.InterpolatingPoints.Emplace(
@@ -827,7 +828,7 @@ namespace UE::MovieSceneTools
 
 			return true;
 		}
-		else if (Times.Num() == 1)
+		else if (Values.Num() == 1)
 		{
 			const double SingleValue = Values[0].Value;
 			CachedDrawParams.InterpolatingPoints.Emplace(
@@ -841,7 +842,8 @@ namespace UE::MovieSceneTools
 			const bool bPostinfinityVisible = Times.Last().Value < ScreenSpace.GetInputMax() * TickResolution;
 			if (bPostinfinityVisible &&
 				(FullRangeInterpolation.PostInfinityExtrapolation == ERichCurveExtrapolation::RCCE_None ||
-				FullRangeInterpolation.PostInfinityExtrapolation == ERichCurveExtrapolation::RCCE_Constant))
+				FullRangeInterpolation.PostInfinityExtrapolation == ERichCurveExtrapolation::RCCE_Constant ||
+				!KeyAttributes.Last().HasLeaveTangent()))
 			{
 				const double ExtentValue = Values.Last().Value;
 				CachedDrawParams.InterpolatingPoints.Emplace(
@@ -853,7 +855,7 @@ namespace UE::MovieSceneTools
 			else if (bPostinfinityVisible &&
 				FullRangeInterpolation.PostInfinityExtrapolation == ERichCurveExtrapolation::RCCE_Linear)
 			{
-				const double ExtentValue = Values.Last().Value + KeyAttributes.Last().GetArriveTangent() * (ScreenSpace.GetInputMax() - Times.Last().Value);
+				const double ExtentValue = Values.Last().Value + KeyAttributes.Last().GetLeaveTangent() * (ScreenSpace.GetInputMax() - Times.Last().Value);
 				CachedDrawParams.InterpolatingPoints.Emplace(
 					ScreenSpace.SecondsToScreen(ScreenSpace.GetInputMax()),
 					ScreenSpace.ValueToScreen(ExtentValue));
