Movie Scene Tools: Fix a potential issue where async curve painter can hinder cached curves from being destroyed

#rnx

diff --git a/Engine/Source/Editor/MovieSceneTools/Private/Cache/MovieSceneAsyncCurvePainter.cpp b/Engine/Source/Editor/MovieSceneTools/Private/Cache/MovieSceneAsyncCurvePainter.cpp
--- a/Engine/Source/Editor/MovieSceneTools/Private/Cache/MovieSceneAsyncCurvePainter.cpp
+++ b/Engine/Source/Editor/MovieSceneTools/Private/Cache/MovieSceneAsyncCurvePainter.cpp
@@ -115,7 +115,7 @@ namespace UE::MovieSceneTools
 		}
 
 		// Remove completed and void tasks
-		for (typename TMap<TSharedPtr<FMovieSceneCachedCurve<ChannelType>>, TSharedPtr<FCurvePainterTask>>::TIterator It(CachedCurveToTaskMap); It; ++It)
+		for (typename TMap<TWeakPtr<FMovieSceneCachedCurve<ChannelType>>, TSharedPtr<FCurvePainterTask>>::TIterator It(CachedCurveToTaskMap); It; ++It)
 		{
 			if ((*It).Value->HasAnyFlags(ECurvePainterTaskStateFlags::Completed | ECurvePainterTaskStateFlags::Void))
 			{
diff --git a/Engine/Source/Editor/MovieSceneTools/Private/Cache/MovieSceneAsyncCurvePainter.h b/Engine/Source/Editor/MovieSceneTools/Private/Cache/MovieSceneAsyncCurvePainter.h
--- a/Engine/Source/Editor/MovieSceneTools/Private/Cache/MovieSceneAsyncCurvePainter.h
+++ b/Engine/Source/Editor/MovieSceneTools/Private/Cache/MovieSceneAsyncCurvePainter.h
@@ -99,7 +99,7 @@ namespace UE::MovieSceneTools
 		void RefineCurvePoints(const TSharedRef<FCurvePainterTask>& Task);
 		
 		/** A map of cached curves with their related tasks that need to be performed asynchronous */
-		TMap<TSharedPtr<FMovieSceneCachedCurve<ChannelType>>, TSharedPtr<FCurvePainterTask>> CachedCurveToTaskMap;
+		TMap<TWeakPtr<FMovieSceneCachedCurve<ChannelType>>, TSharedPtr<FCurvePainterTask>> CachedCurveToTaskMap;
 
 		/** True while work is ongoing */
 		std::atomic<bool> bWorking = false;
diff --git a/Engine/Source/Editor/MovieSceneTools/Private/Cache/MovieSceneCurveCachePool.cpp b/Engine/Source/Editor/MovieSceneTools/Private/Cache/MovieSceneCurveCachePool.cpp
--- a/Engine/Source/Editor/MovieSceneTools/Private/Cache/MovieSceneCurveCachePool.cpp
+++ b/Engine/Source/Editor/MovieSceneTools/Private/Cache/MovieSceneCurveCachePool.cpp
@@ -25,7 +25,7 @@ namespace UE::MovieSceneTools
 
 	void FMovieSceneCurveCachePool::DrawCachedCurves(TWeakPtr<const FCurveEditor> WeakCurveEditor)
 	{
-		const TArray<TSharedPtr<IMovieSceneCachedCurve>>* CachedCurvesPtr = CurveEditorToCachedCurvesMap.Find(WeakCurveEditor);
+		const TArray<TWeakPtr<IMovieSceneCachedCurve>>* CachedCurvesPtr = CurveEditorToCachedCurvesMap.Find(WeakCurveEditor);
 		if (CachedCurvesPtr)
 		{
 			DrawCachedCurves(*CachedCurvesPtr);
@@ -41,14 +41,16 @@ namespace UE::MovieSceneTools
 	void FMovieSceneCurveCachePool::Leave(const IMovieSceneCachedCurve& CachedCurve)
 	{
 		// Find where the curve was added
-		TTuple<TWeakPtr<const FCurveEditor>, TArray<TSharedPtr<IMovieSceneCachedCurve>>>* CurveEditorToCachedCurvesPairPtr =
+		TTuple<TWeakPtr<const FCurveEditor>, TArray<TWeakPtr<IMovieSceneCachedCurve>>>* CurveEditorToCachedCurvesPairPtr =
 			Algo::FindByPredicate(CurveEditorToCachedCurvesMap, 
-				[&CachedCurve](const TTuple<TWeakPtr<const FCurveEditor>, TArray<TSharedPtr<IMovieSceneCachedCurve>>>& CurveEditorToCachedCurvesPair)
+				[&CachedCurve](const TTuple<TWeakPtr<const FCurveEditor>, TArray<TWeakPtr<IMovieSceneCachedCurve>>>& CurveEditorToCachedCurvesPair)
 				{
 					return Algo::FindByPredicate(CurveEditorToCachedCurvesPair.Value,
-						[&CachedCurve](const TSharedPtr<IMovieSceneCachedCurve>& OtherCachedCurve)
+						[&CachedCurve](const TWeakPtr<IMovieSceneCachedCurve>& OtherCachedCurve)
 						{
-							return OtherCachedCurve.Get() == &CachedCurve;
+							return 
+								!OtherCachedCurve.IsValid() ||
+								OtherCachedCurve.Pin().Get() == &CachedCurve;
 						}) != nullptr;
 				});
 	
@@ -59,9 +61,11 @@ namespace UE::MovieSceneTools
 				// Remove the curve
 				const int32 OldSize = (*CurveEditorToCachedCurvesPairPtr).Value.Num();
 				const int32 NewSize = Algo::RemoveIf((*CurveEditorToCachedCurvesPairPtr).Value,
-					[&CachedCurve](const TSharedPtr<IMovieSceneCachedCurve>& OtherCachedCurve)
+					[&CachedCurve](const TWeakPtr<IMovieSceneCachedCurve>& OtherCachedCurve)
 					{
-						return OtherCachedCurve.Get() == &CachedCurve;
+						return 
+							!OtherCachedCurve.IsValid() || 
+							OtherCachedCurve.Pin().Get() == &CachedCurve;
 					});
 
 				if (OldSize != NewSize)
@@ -80,17 +84,20 @@ namespace UE::MovieSceneTools
 		}
 	}
 
-	void FMovieSceneCurveCachePool::DrawCachedCurves(const TArray<TSharedPtr<IMovieSceneCachedCurve>>& CachedCurves) const
+	void FMovieSceneCurveCachePool::DrawCachedCurves(const TArray<TWeakPtr<IMovieSceneCachedCurve>>& CachedCurves) const
 	{
 		ParallelFor(CachedCurves.Num(),
 			[CachedCurves](int32 Index)
 			{
-				CachedCurves[Index]->DrawCachedCurve();
+				if (const TSharedPtr<IMovieSceneCachedCurve>& CachedCurve = CachedCurves[Index].Pin())
+				{
+					CachedCurve->DrawCachedCurve();
+				}
 
 			}, EParallelForFlags::Unbalanced);
 	}
 
-	void FMovieSceneCurveCachePool::CullCachedCurves(const TArray<TSharedPtr<IMovieSceneCachedCurve>>& CachedCurves) const
+	void FMovieSceneCurveCachePool::CullCachedCurves(const TArray<TWeakPtr<IMovieSceneCachedCurve>>& CachedCurves) const
 	{
 		if (GCullCachedCurves <= 0)
 		{
@@ -101,11 +108,17 @@ namespace UE::MovieSceneTools
 		TMap<uint32, const FCurveDrawParams*> InterpolatingPointsHashToDrawParamsMap;
 
 		// Reverse to avoid culling those curves that are drawn last resp. top-most
-		TArray<TSharedPtr<IMovieSceneCachedCurve>> InverseCachedCurves = CachedCurves;
+		TArray<TWeakPtr<IMovieSceneCachedCurve>> InverseCachedCurves = CachedCurves;
 		Algo::Reverse(InverseCachedCurves);
 
-		for (const TSharedPtr<IMovieSceneCachedCurve>& CachedCurve : InverseCachedCurves)
+		for (const TWeakPtr<IMovieSceneCachedCurve>& WeakCachedCurve : InverseCachedCurves)
 		{
+			if (!WeakCachedCurve.IsValid())
+			{
+				continue;
+			}
+			const TSharedRef<IMovieSceneCachedCurve> CachedCurve = WeakCachedCurve.Pin().ToSharedRef();
+
 			FCurveDrawParams* DrawParamsPtr = CachedCurve->GetDrawParamsHandle().Get();
 			if (!DrawParamsPtr)
 			{
diff --git a/Engine/Source/Editor/MovieSceneTools/Private/Cache/MovieSceneCurveCachePool.h b/Engine/Source/Editor/MovieSceneTools/Private/Cache/MovieSceneCurveCachePool.h
--- a/Engine/Source/Editor/MovieSceneTools/Private/Cache/MovieSceneCurveCachePool.h
+++ b/Engine/Source/Editor/MovieSceneTools/Private/Cache/MovieSceneCurveCachePool.h
@@ -52,12 +52,12 @@ namespace UE::MovieSceneTools
 
 	private:
 		/** Lets cached curves draw to their draw params */
-		void DrawCachedCurves(const TArray<TSharedPtr<IMovieSceneCachedCurve>>& CachedCurves) const;
+		void DrawCachedCurves(const TArray<TWeakPtr<IMovieSceneCachedCurve>>& CachedCurves) const;
 
 		/** Culls cached curves */
-		void CullCachedCurves(const TArray<TSharedPtr<IMovieSceneCachedCurve>>& CachedCurves) const;
+		void CullCachedCurves(const TArray<TWeakPtr<IMovieSceneCachedCurve>>& CachedCurves) const;
 
 		/** Cached curves per curve editor */
-		TMap<TWeakPtr<const FCurveEditor>, TArray<TSharedPtr<IMovieSceneCachedCurve>>> CurveEditorToCachedCurvesMap;
+		TMap<TWeakPtr<const FCurveEditor>, TArray<TWeakPtr<IMovieSceneCachedCurve>>> CurveEditorToCachedCurvesMap;
 	};
 }
