Anim In Engine: Added Sequencer.SubsequenceAnimBakeInIsolation cvar to turn on Evaluate SubSequence In Isolation temporary if baking out an anim sequence from a subsequence. Will default to false but will turn it in after more live testing from the Dev-Cine team.

diff --git a/Engine/Source/Editor/MovieSceneTools/Private/TrackEditors/SkeletalAnimationTrackEditor.cpp b/Engine/Source/Editor/MovieSceneTools/Private/TrackEditors/SkeletalAnimationTrackEditor.cpp
--- a/Engine/Source/Editor/MovieSceneTools/Private/TrackEditors/SkeletalAnimationTrackEditor.cpp
+++ b/Engine/Source/Editor/MovieSceneTools/Private/TrackEditors/SkeletalAnimationTrackEditor.cpp
@@ -101,6 +101,15 @@ namespace SkeletalAnimationEditorConstants
 {
 	// @todo Sequencer Allow this to be customizable
 	const uint32 AnimationTrackHeight = 28;
+
+	const TCHAR* CVarName = TEXT("Sequencer.SubsequenceAnimBakeInIsolation");
+
+	static bool bSubsequenceAnimBakeInIsolation = false; //default to false for now will turn it on after more testing
+	static FAutoConsoleVariableRef CVarSetSubsequenceAnimBakeInIsolation(
+		CVarName,
+		bSubsequenceAnimBakeInIsolation,
+		TEXT("Automatically Force Evaluate Subsequence In Isolation when Baking Anim Sequences")
+	);
 }
 
 #define LOCTEXT_NAMESPACE "FSkeletalAnimationTrackEditor"
@@ -1482,12 +1491,23 @@ void FSkeletalAnimationTrackEditor::OnSequencerSaved(ISequencer& )
 		if (IInterface_AssetUserData* AssetUserDataInterface = Cast< IInterface_AssetUserData >(LevelSequence))
 		{
 			ULevelSequenceAnimSequenceLink* LevelAnimLink = AssetUserDataInterface->GetAssetUserData< ULevelSequenceAnimSequenceLink >();
-			if (LevelAnimLink)
+			if (LevelAnimLink && LevelAnimLink->AnimSequenceLinks.Num() > 0)
 			{
+				
+
 				UMovieSceneSequence* MovieSceneSequence = SequencerPtr->GetFocusedMovieSceneSequence();
 				UMovieSceneSequence* RootMovieSceneSequence = SequencerPtr->GetRootMovieSceneSequence();
 				FMovieSceneSequenceIDRef Template = SequencerPtr->GetFocusedTemplateID();
 				FMovieSceneSequenceTransform RootToLocalTransform = SequencerPtr->GetFocusedMovieSceneSequenceTransform();
+				//if in subsquence and we want to, turn on ShouldEvaluateSubSequencesInIsolation
+				TOptional<bool> bRestoreShouldEvaluateSubSequencesInIsolation;
+				if (SkeletalAnimationEditorConstants::bSubsequenceAnimBakeInIsolation &&
+					MovieSceneSequence != RootMovieSceneSequence &&
+					SequencerPtr->GetSequencerSettings()->ShouldEvaluateSubSequencesInIsolation() == false)
+				{
+					bRestoreShouldEvaluateSubSequencesInIsolation = false;
+					SequencerPtr->GetSequencerSettings()->SetEvaluateSubSequencesInIsolation(true);
+				}
 				for (int32 Index = LevelAnimLink->AnimSequenceLinks.Num() -1; Index >=0 ; --Index)
 				{
 					FLevelSequenceAnimSequenceLinkItem& Item = LevelAnimLink->AnimSequenceLinks[Index];
@@ -1552,6 +1572,8 @@ void FSkeletalAnimationTrackEditor::OnSequencerSaved(ISequencer& )
 						AESP.MovieSceneSequence = MovieSceneSequence;
 						AESP.RootMovieSceneSequence = RootMovieSceneSequence;
 						AESP.bForceUseOfMovieScenePlaybackRange = SequencerPtr.Get()->GetSequencerSettings()->ShouldEvaluateSubSequencesInIsolation();
+						//make sure all spanwables are present
+						SequencerPtr->ForceEvaluate();
 						bool bResult = MovieSceneToolHelpers::ExportToAnimSequence(AnimSequence, AnimSeqExportOption, AESP, SkelMeshComp);
 
 						AnimSeqExportOption->bExportMorphTargets = bSavedExportMorphTargets;
@@ -1583,6 +1605,12 @@ void FSkeletalAnimationTrackEditor::OnSequencerSaved(ISequencer& )
 						UPackage::SavePackage(Package, NULL, *PackageFileName, SaveArgs);
 					}
 				}
+				if (bRestoreShouldEvaluateSubSequencesInIsolation.IsSet())
+				{
+					SequencerPtr->GetSequencerSettings()->SetEvaluateSubSequencesInIsolation(false);
+				}
+				//re-evaluate at current frame
+				SequencerPtr->ForceEvaluate();
 			}
 		}
 	}
@@ -1688,6 +1716,14 @@ bool FSkeletalAnimationTrackEditor::CreateAnimationSequence(const TArray<UObject
 				UMovieSceneSequence* MovieSceneSequence = ParentSequencer->GetFocusedMovieSceneSequence();
 				UMovieSceneSequence* RootMovieSceneSequence = ParentSequencer->GetRootMovieSceneSequence();
 				FMovieSceneSequenceTransform RootToLocalTransform = ParentSequencer->GetFocusedMovieSceneSequenceTransform();
+				TOptional<bool> bRestoreShouldEvaluateSubSequencesInIsolation;
+				if (SkeletalAnimationEditorConstants::bSubsequenceAnimBakeInIsolation &&
+					MovieSceneSequence != RootMovieSceneSequence &&
+					ParentSequencer->GetSequencerSettings()->ShouldEvaluateSubSequencesInIsolation() == false)
+				{
+					bRestoreShouldEvaluateSubSequencesInIsolation = false;
+					ParentSequencer->GetSequencerSettings()->SetEvaluateSubSequencesInIsolation(true);
+				}
 				FAnimExportSequenceParameters AESP;
 				AESP.Player = ParentSequencer.Get();
 				AESP.RootToLocalTransform = RootToLocalTransform;
@@ -1696,6 +1732,12 @@ bool FSkeletalAnimationTrackEditor::CreateAnimationSequence(const TArray<UObject
 				AESP.bForceUseOfMovieScenePlaybackRange = ParentSequencer->GetSequencerSettings()->ShouldEvaluateSubSequencesInIsolation();
 				AnimSeqExportOption->CustomDisplayRate = ParentSequencer->GetFocusedDisplayRate();
 				bResult  = MovieSceneToolHelpers::ExportToAnimSequence(AnimSequence, AnimSeqExportOption, AESP, SkelMeshComp);
+				if (bRestoreShouldEvaluateSubSequencesInIsolation.IsSet())
+				{
+					ParentSequencer->GetSequencerSettings()->SetEvaluateSubSequencesInIsolation(false);
+				}
+				//re-evaluate at current frame
+				ParentSequencer->ForceEvaluate();
 			}
 		}
 
diff --git a/Engine/Source/Editor/Sequencer/Private/Sequencer.cpp b/Engine/Source/Editor/Sequencer/Private/Sequencer.cpp
--- a/Engine/Source/Editor/Sequencer/Private/Sequencer.cpp
+++ b/Engine/Source/Editor/Sequencer/Private/Sequencer.cpp
@@ -380,14 +380,11 @@ void FSequencer::InitSequencer(const FSequencerInitParams& InitParams
 
 	Settings = USequencerSettingsContainer::GetOrCreate<USequencerSettings>(*InitParams.ViewParams.UniqueName);
 
-	Settings->GetOnEvaluateSubSequencesInIsolationChanged().AddSP(this, &FSequencer::RestorePreAnimatedState);
+	Settings->GetOnEvaluateSubSequencesInIsolationChanged().AddSP(this, &FSequencer::OnEvaluateSubSequencesInIsolationChanged);
 	Settings->GetOnShowSelectedNodesOnlyChanged().AddSP(this, &FSequencer::OnSelectedNodesOnlyChanged);
 	Settings->GetOnTimeDisplayFormatChanged().AddSP(this, &FSequencer::OnTimeDisplayFormatChanged);
 
 
-
-
-
 	// Initialize numeric type interfaces
 	{
 
@@ -801,6 +798,15 @@ void FSequencer::OnPlaybackContextChanged()
 	InitRootSequenceInstance();
 }
 
+void FSequencer::OnEvaluateSubSequencesInIsolationChanged()
+{
+	RestorePreAnimatedState();
+	FMovieSceneSequenceID NewOverrideRoot = Settings->ShouldEvaluateSubSequencesInIsolation() ? ActiveTemplateIDs.Top() : MovieSceneSequenceID::Root;
+	UMovieSceneEntitySystemLinker* Linker = RootTemplateInstance.GetEntitySystemLinker();
+	RootTemplateInstance.FindInstance(MovieSceneSequenceID::Root)->OverrideRootSequence(NewOverrideRoot);
+	ForceEvaluate();
+}
+
 void FSequencer::InitRootSequenceInstance()
 {
 	using namespace UE::MovieScene;
@@ -11637,12 +11643,6 @@ void FSequencer::BindCommands()
 		FExecuteAction::CreateLambda( [this]{
 			const bool bNewValue = !Settings->ShouldEvaluateSubSequencesInIsolation();
 			Settings->SetEvaluateSubSequencesInIsolation( bNewValue );
-
-			FMovieSceneSequenceID NewOverrideRoot = bNewValue ? ActiveTemplateIDs.Top() : MovieSceneSequenceID::Root;
-			UMovieSceneEntitySystemLinker* Linker = RootTemplateInstance.GetEntitySystemLinker();
-			RootTemplateInstance.FindInstance(MovieSceneSequenceID::Root)->OverrideRootSequence(NewOverrideRoot);
-
-			ForceEvaluate();
 		} ),
 		FCanExecuteAction::CreateLambda( [this]{ return true; } ),
 		FIsActionChecked::CreateLambda( [this]{ return Settings->ShouldEvaluateSubSequencesInIsolation(); } ) );
diff --git a/Engine/Source/Editor/Sequencer/Private/Sequencer.h b/Engine/Source/Editor/Sequencer/Private/Sequencer.h
--- a/Engine/Source/Editor/Sequencer/Private/Sequencer.h
+++ b/Engine/Source/Editor/Sequencer/Private/Sequencer.h
@@ -982,6 +982,8 @@ protected:
 	/** Called when the scrub position parent sequence is changed */
 	void OnScrubPositionParentChanged(FMovieSceneSequenceID InScrubPositionParent);
 
+	void OnEvaluateSubSequencesInIsolationChanged();
+
 protected:
 
 	/**
