Sequencer: Fixed crash in spawnable actor binding customization when an assigned Level Name did not exist

diff --git a/Engine/Source/Editor/MovieSceneTools/Private/Bindings/MovieSceneSpawnableActorBindingCustomization.cpp b/Engine/Source/Editor/MovieSceneTools/Private/Bindings/MovieSceneSpawnableActorBindingCustomization.cpp
--- a/Engine/Source/Editor/MovieSceneTools/Private/Bindings/MovieSceneSpawnableActorBindingCustomization.cpp
+++ b/Engine/Source/Editor/MovieSceneTools/Private/Bindings/MovieSceneSpawnableActorBindingCustomization.cpp
@@ -30,7 +30,9 @@ void FMovieSceneSpawnableActorBindingBaseCustomization::CustomizeDetails(IDetail
 	RefreshComboList();
 
 	const int32 FoundIndex = LevelNameList.Find(LevelNameComboSelectedName);
-	const TSharedPtr<FString> ComboBoxSelectedItem = LevelNameComboListItems[FoundIndex];
+	const TSharedPtr<FString> ComboBoxSelectedItem = LevelNameComboListItems.IsValidIndex(FoundIndex)
+		? LevelNameComboListItems[FoundIndex]
+		: nullptr;
 	IDetailCategoryBuilder& SectionCategory = DetailBuilder.EditCategory("Actor");
 	SectionCategory.AddCustomRow(FText())
 		.NameContent()
@@ -46,7 +48,6 @@ void FMovieSceneSpawnableActorBindingBaseCustomization::CustomizeDetails(IDetail
 				.OnSelectionChanged(this, &FMovieSceneSpawnableActorBindingBaseCustomization::OnLevelNameChanged)
 				.InitiallySelectedItem(ComboBoxSelectedItem)
 				.ContentPadding(2.f)
-				.ToolTipText(FText::FromString(*ComboBoxSelectedItem))
 			]
 		];
 }
