CurveEditor: Fix anim curves are mirrored over bars

#rnx

diff --git a/Engine/Plugins/Animation/ControlRig/Source/ControlRigEditor/Private/ControlRigEditorModule.cpp b/Engine/Plugins/Animation/ControlRig/Source/ControlRigEditor/Private/ControlRigEditorModule.cpp
--- a/Engine/Plugins/Animation/ControlRig/Source/ControlRigEditor/Private/ControlRigEditorModule.cpp
+++ b/Engine/Plugins/Animation/ControlRig/Source/ControlRigEditor/Private/ControlRigEditorModule.cpp
@@ -330,7 +330,8 @@ void FControlRigEditorModule::StartupModule()
 	FControlRigSpaceChannelCurveModel::ViewID = CurveEditorModule.RegisterView(FOnCreateCurveEditorView::CreateStatic(
 		[](TWeakPtr<FCurveEditor> WeakCurveEditor) -> TSharedRef<SCurveEditorView>
 		{
-			return SNew(SCurveEditorKeyBarView, WeakCurveEditor);
+			return SNew(SCurveEditorKeyBarView, WeakCurveEditor)
+				.DrawCurves(false);
 		}
 	));
 
diff --git a/Engine/Source/Editor/MovieSceneTools/Private/Channels/SCurveEditorKeyBarView.cpp b/Engine/Source/Editor/MovieSceneTools/Private/Channels/SCurveEditorKeyBarView.cpp
--- a/Engine/Source/Editor/MovieSceneTools/Private/Channels/SCurveEditorKeyBarView.cpp
+++ b/Engine/Source/Editor/MovieSceneTools/Private/Channels/SCurveEditorKeyBarView.cpp
@@ -39,6 +39,8 @@ float SCurveEditorKeyBarView::TrackHeight = 24.f;
 
 void SCurveEditorKeyBarView::Construct(const FArguments& InArgs, TWeakPtr<FCurveEditor> InCurveEditor)
 {
+	bDrawCurves = InArgs._DrawCurves;
+
 	bFixedOutputBounds = true;
 	bAllowModelViewTransforms = false;
 
@@ -47,7 +49,7 @@ void SCurveEditorKeyBarView::Construct(const FArguments& InArgs, TWeakPtr<FCurve
 	WeakCurveEditor = InCurveEditor;
 	SortBias = 25;
 
-	SInteractiveCurveEditorView::Construct(InArgs, InCurveEditor);
+	SInteractiveCurveEditorView::Construct(SInteractiveCurveEditorView::FArguments(), InCurveEditor);
 }
 
 FVector2D SCurveEditorKeyBarView::ComputeDesiredSize(float LayoutScaleMultiplier) const
@@ -92,7 +94,11 @@ void SCurveEditorKeyBarView::PaintView(const FPaintArgs& Args, const FGeometry&
 		DrawBackground(AllottedGeometry, OutDrawElements, BaseLayerId, DrawEffects);
 		DrawLabels(AllottedGeometry, OutDrawElements, BaseLayerId, DrawEffects);
 		DrawGridLines(CurveEditor.ToSharedRef(), AllottedGeometry, OutDrawElements, BaseLayerId, DrawEffects);
-		DrawCurves(CurveEditor.ToSharedRef(), AllottedGeometry, MyCullingRect, OutDrawElements, BaseLayerId, InWidgetStyle, DrawEffects);
+		
+		if (bDrawCurves)
+		{
+			DrawCurves(CurveEditor.ToSharedRef(), AllottedGeometry, MyCullingRect, OutDrawElements, BaseLayerId, InWidgetStyle, DrawEffects);
+		}
 	}
 }
 
diff --git a/Engine/Source/Editor/MovieSceneTools/Private/MovieSceneToolsModule.cpp b/Engine/Source/Editor/MovieSceneTools/Private/MovieSceneToolsModule.cpp
--- a/Engine/Source/Editor/MovieSceneTools/Private/MovieSceneToolsModule.cpp
+++ b/Engine/Source/Editor/MovieSceneTools/Private/MovieSceneToolsModule.cpp
@@ -263,7 +263,8 @@ void FMovieSceneToolsModule::StartupModule()
 		FConstraintChannelCurveModel::ViewID = CurveEditorModule.RegisterView(FOnCreateCurveEditorView::CreateStatic(
 			[](TWeakPtr<FCurveEditor> WeakCurveEditor) -> TSharedRef<SCurveEditorView>
 			{
-				return SNew(SCurveEditorKeyBarView, WeakCurveEditor);
+				return SNew(SCurveEditorKeyBarView, WeakCurveEditor)
+					.DrawCurves(false);
 			}
 		));
 
diff --git a/Engine/Source/Editor/MovieSceneTools/Public/Channels/SCurveEditorKeyBarView.h b/Engine/Source/Editor/MovieSceneTools/Public/Channels/SCurveEditorKeyBarView.h
--- a/Engine/Source/Editor/MovieSceneTools/Public/Channels/SCurveEditorKeyBarView.h
+++ b/Engine/Source/Editor/MovieSceneTools/Public/Channels/SCurveEditorKeyBarView.h
@@ -24,6 +24,14 @@ struct FGeometry;
 class MOVIESCENETOOLS_API SCurveEditorKeyBarView : public SInteractiveCurveEditorView
 {
 public:
+	SLATE_BEGIN_ARGS(SCurveEditorKeyBarView) 
+		: _DrawCurves(true)
+		{}
+	
+		/** When false the widget will not draw curves or keys */
+		SLATE_ARGUMENT(bool, DrawCurves)
+	
+	SLATE_END_ARGS()
 
 	void Construct(const FArguments& InArgs, TWeakPtr<FCurveEditor> InCurveEditor);
 
@@ -42,5 +50,8 @@ private:
 
 	void DrawLabels(const FGeometry& AllottedGeometry, FSlateWindowElementList& OutDrawElements, int32 BaseLayerId, ESlateDrawEffect DrawEffects) const;
 
+	/** When false the widget will not draw curves or keys */
+	bool bDrawCurves;
+
 	static float TrackHeight;
 };
\ No newline at end of file
