Sequencer: Fix crash when ctrl-double click on a subsequence, which effectively opens the sequence (closing the existing one)

diff --git a/Engine/Source/Editor/MovieSceneTools/Public/TrackEditors/SubTrackEditorBase.h b/Engine/Source/Editor/MovieSceneTools/Public/TrackEditors/SubTrackEditorBase.h
--- a/Engine/Source/Editor/MovieSceneTools/Public/TrackEditors/SubTrackEditorBase.h
+++ b/Engine/Source/Editor/MovieSceneTools/Public/TrackEditors/SubTrackEditorBase.h
@@ -20,6 +20,7 @@
 #include "SequencerSectionPainter.h"
 #include "Subsystems/AssetEditorSubsystem.h"
 #include "Tracks/MovieSceneSubTrack.h"
+#include "TimerManager.h"
 
 struct FMovieSceneTimeWarpChannel;
 struct FMovieSceneSequenceTransform;
@@ -320,16 +321,20 @@ FReply TSubSectionMixin<ParentSectionClass>::OnSectionDoubleClicked(const FGeome
 {
     if( MouseEvent.GetEffectingButton() == EKeys::LeftMouseButton )
     {
-        if (SubSectionObject.GetSequence())
+		UMovieSceneSequence* Sequence = SubSectionObject.GetSequence();
+        if (Sequence)
         {
-	    if (MouseEvent.IsControlDown())
-	    {
-		GEditor->GetEditorSubsystem<UAssetEditorSubsystem>()->OpenEditorForAsset(SubSectionObject.GetSequence());
-	    }
-	    else
-	    {
-                this->GetSequencer()->FocusSequenceInstance(SubSectionObject);
-	    }
+			if (MouseEvent.IsControlDown())
+			{
+				GEditor->GetTimerManager()->SetTimerForNextTick([Sequence]()
+				{
+					GEditor->GetEditorSubsystem<UAssetEditorSubsystem>()->OpenEditorForAsset(Sequence);
+				});
+			}
+			else
+			{
+				this->GetSequencer()->FocusSequenceInstance(SubSectionObject);
+			}
         }
     }
 
