Sequencer: Enable recursive classes for the subsequence asset picker

diff --git a/Engine/Source/Editor/MovieSceneTools/Private/TrackEditors/SubTrackEditor.cpp b/Engine/Source/Editor/MovieSceneTools/Private/TrackEditors/SubTrackEditor.cpp
--- a/Engine/Source/Editor/MovieSceneTools/Private/TrackEditors/SubTrackEditor.cpp
+++ b/Engine/Source/Editor/MovieSceneTools/Private/TrackEditors/SubTrackEditor.cpp
@@ -1578,6 +1578,7 @@ TSharedRef<SWidget> FSubTrackEditor::HandleAddSubSequenceComboButtonGetMenuConte
 			AssetPickerConfig.bAddFilterUI = true;
 			AssetPickerConfig.bShowTypeInColumnView = false;
 			AssetPickerConfig.InitialAssetViewType = EAssetViewType::List;
+			AssetPickerConfig.Filter.bRecursiveClasses = true;
 			GetSupportedSequenceClassPaths(AssetPickerConfig.Filter.ClassPaths);
 			AssetPickerConfig.SaveSettingsName = TEXT("SequencerAssetPicker");
 			AssetPickerConfig.AdditionalReferencingAssets.Add(FAssetData(Sequence));
