PR #12989: Fix Vector4f properties crashing sequencer

diff --git a/Engine/Source/Editor/MovieSceneTools/Private/TrackEditors/PropertyTrackEditors/VectorPropertyTrackEditor.cpp b/Engine/Source/Editor/MovieSceneTools/Private/TrackEditors/PropertyTrackEditors/VectorPropertyTrackEditor.cpp
--- a/Engine/Source/Editor/MovieSceneTools/Private/TrackEditors/PropertyTrackEditors/VectorPropertyTrackEditor.cpp
+++ b/Engine/Source/Editor/MovieSceneTools/Private/TrackEditors/PropertyTrackEditors/VectorPropertyTrackEditor.cpp
@@ -96,7 +96,7 @@ void FFloatVectorPropertyTrackEditor::InitializeNewTrack( UMovieSceneFloatVector
 	{
 		NewTrack->SetNumChannelsUsed( 3 );
 	}
-	if ( StructName == NAME_Vector4 )
+	if ( StructName == NAME_Vector4f )
 	{
 		NewTrack->SetNumChannelsUsed( 4 );
 	}
