Sequencer: Add CVarEvaluateSkeletalMeshOnPropertyChange to enable/disable re-evaluation on property changed

Follow up to 40256104

#rnx

diff --git a/Engine/Source/Editor/MovieSceneTools/Private/TrackEditors/SkeletalAnimationTrackEditor.cpp b/Engine/Source/Editor/MovieSceneTools/Private/TrackEditors/SkeletalAnimationTrackEditor.cpp
--- a/Engine/Source/Editor/MovieSceneTools/Private/TrackEditors/SkeletalAnimationTrackEditor.cpp
+++ b/Engine/Source/Editor/MovieSceneTools/Private/TrackEditors/SkeletalAnimationTrackEditor.cpp
@@ -110,6 +110,11 @@ namespace SkeletalAnimationEditorConstants
 		bSubsequenceAnimBakeInIsolation,
 		TEXT("Automatically Force Evaluate Subsequence In Isolation when Baking Anim Sequences")
 	);
+
+	static TAutoConsoleVariable<bool> CVarEvaluateSkeletalMeshOnPropertyChange(
+		TEXT("Sequencer.EvaluateSkeletalMeshOnPropertyChange"),
+		true,
+		TEXT("Enable/disable sending a track value changed when properties change so that the skeletal mesh can be re-evaluated in Sequencer"));
 }
 
 #define LOCTEXT_NAMESPACE "FSkeletalAnimationTrackEditor"
@@ -1660,24 +1665,27 @@ void FSkeletalAnimationTrackEditor::OnSequencerDataChanged(EMovieSceneDataChange
 
 void FSkeletalAnimationTrackEditor::OnPostPropertyChanged(UObject* InObject, FPropertyChangedEvent& InPropertyChangedEvent)
 {
-	const FName PropertyName = (InPropertyChangedEvent.Property != nullptr) ? InPropertyChangedEvent.Property->GetFName() : NAME_None;
-	if (PropertyName == USceneComponent::GetRelativeLocationPropertyName() ||
-		PropertyName == USceneComponent::GetRelativeRotationPropertyName() ||
-		PropertyName == USceneComponent::GetRelativeScale3DPropertyName())
+	if (SkeletalAnimationEditorConstants::CVarEvaluateSkeletalMeshOnPropertyChange->GetBool())
 	{
-		return;
-	}
+		const FName PropertyName = (InPropertyChangedEvent.Property != nullptr) ? InPropertyChangedEvent.Property->GetFName() : NAME_None;
+		if (PropertyName == USceneComponent::GetRelativeLocationPropertyName() ||
+			PropertyName == USceneComponent::GetRelativeRotationPropertyName() ||
+			PropertyName == USceneComponent::GetRelativeScale3DPropertyName())
+		{
+			return;
+		}
 
-	// If the object changed has any animation track, notify sequencer to update because animations tick on their own and sequencer needs to evaluate again
-	const bool bCreateIfMissing = false;
-	FFindOrCreateHandleResult HandleResult = FindOrCreateHandleToObject(InObject, bCreateIfMissing);
-	FGuid ObjectHandle = HandleResult.Handle;
-	if (ObjectHandle.IsValid())
-	{
-		FFindOrCreateTrackResult TrackResult = FindOrCreateTrackForObject(ObjectHandle, UMovieSceneSkeletalAnimationTrack::StaticClass(), NAME_None, bCreateIfMissing);
-		if (TrackResult.Track)
+		// If the object changed has any animation track, notify sequencer to update because animations tick on their own and sequencer needs to evaluate again
+		const bool bCreateIfMissing = false;
+		FFindOrCreateHandleResult HandleResult = FindOrCreateHandleToObject(InObject, bCreateIfMissing);
+		FGuid ObjectHandle = HandleResult.Handle;
+		if (ObjectHandle.IsValid())
 		{
-			GetSequencer()->NotifyMovieSceneDataChanged(EMovieSceneDataChangeType::TrackValueChanged);
+			FFindOrCreateTrackResult TrackResult = FindOrCreateTrackForObject(ObjectHandle, UMovieSceneSkeletalAnimationTrack::StaticClass(), NAME_None, bCreateIfMissing);
+			if (TrackResult.Track)
+			{
+				GetSequencer()->NotifyMovieSceneDataChanged(EMovieSceneDataChangeType::TrackValueChanged);
+			}
 		}
 	}
 }
