Curve Editor: Add API for evaluating IBufferedCurveModel.
This change is needed for implementing a time offsetting tween function.

diff --git a/Engine/Source/Editor/CurveEditor/Private/RichCurveEditorModel.cpp b/Engine/Source/Editor/CurveEditor/Private/RichCurveEditorModel.cpp
--- a/Engine/Source/Editor/CurveEditor/Private/RichCurveEditorModel.cpp
+++ b/Engine/Source/Editor/CurveEditor/Private/RichCurveEditorModel.cpp
@@ -193,6 +193,12 @@ public:
 		} while (OldSize != InterpolatingPoints.Num());
 	}
 
+	virtual bool Evaluate(double InTime, double& OutValue) const override
+	{
+		OutValue = RichCurve.Eval(InTime);
+		return true;
+	}
+
 private:
 	FRichCurve RichCurve;
 };
diff --git a/Engine/Source/Editor/CurveEditor/Public/IBufferedCurveModel.h b/Engine/Source/Editor/CurveEditor/Public/IBufferedCurveModel.h
--- a/Engine/Source/Editor/CurveEditor/Public/IBufferedCurveModel.h
+++ b/Engine/Source/Editor/CurveEditor/Public/IBufferedCurveModel.h
@@ -28,6 +28,13 @@ public:
 
 	/** draws the curve into an array of (input, output) pairs with a given screen space */
 	virtual void DrawCurve(const FCurveEditor& CurveEditor, const FCurveEditorScreenSpace& ScreenSpace, TArray<TTuple<double, double>>& InterpolatingPoints) const = 0;
+
+	/**
+	 * @param InTime The time at which to evaluate the curve.
+	 * @param OutValue The evaluated curve value at InTime, if this function returned true.
+	 * @return Whether the curve was evaluated.
+	 */
+	virtual bool Evaluate(double InTime, double& OutValue) const = 0;
 	
 	/**
 	 * Retrieve all key positions stored in this buffered curve
diff --git a/Engine/Source/Editor/MovieSceneTools/Private/Channels/BezierChannelCurveModel.cpp b/Engine/Source/Editor/MovieSceneTools/Private/Channels/BezierChannelCurveModel.cpp
--- a/Engine/Source/Editor/MovieSceneTools/Private/Channels/BezierChannelCurveModel.cpp
+++ b/Engine/Source/Editor/MovieSceneTools/Private/Channels/BezierChannelCurveModel.cpp
@@ -4,6 +4,7 @@
 
 #include "Cache/MovieSceneCachedCurve.h"
 #include "Cache/MovieSceneUpdateCachedCurveData.h"
+#include "Channels/CurveModelHelpers.h"
 #include "Channels/MovieSceneChannelData.h"
 #include "Channels/MovieSceneCurveChannelCommon.h"
 #include "Channels/MovieSceneDoubleChannel.h"
@@ -58,6 +59,12 @@ void FBezierChannelBufferedCurveModel<ChannelType>::DrawCurve(const FCurveEditor
 	}
 }
 
+template <typename ChannelType>
+bool FBezierChannelBufferedCurveModel<ChannelType>::Evaluate(double InTime, double& OutValue) const
+{
+	return UE::MovieSceneTools::CurveHelpers::Evaluate(InTime, OutValue, Channel, WeakSection);
+}
+
 template class FBezierChannelBufferedCurveModel<FMovieSceneFloatChannel>;
 template class FBezierChannelBufferedCurveModel<FMovieSceneDoubleChannel>;
 
diff --git a/Engine/Source/Editor/MovieSceneTools/Private/Channels/BoolChannelCurveModel.cpp b/Engine/Source/Editor/MovieSceneTools/Private/Channels/BoolChannelCurveModel.cpp
--- a/Engine/Source/Editor/MovieSceneTools/Private/Channels/BoolChannelCurveModel.cpp
+++ b/Engine/Source/Editor/MovieSceneTools/Private/Channels/BoolChannelCurveModel.cpp
@@ -4,6 +4,7 @@
 
 #include "Algo/BinarySearch.h"
 #include "Channels/BoolChannelKeyProxy.h"
+#include "Channels/CurveModelHelpers.h"
 #include "Channels/MovieSceneBoolChannel.h"
 #include "Channels/MovieSceneChannelData.h"
 #include "Channels/MovieSceneChannelHandle.h"
@@ -75,6 +76,11 @@ public:
 		}
 	}
 
+	virtual bool Evaluate(double InTime, double& OutValue) const override
+	{
+		return UE::MovieSceneTools::CurveHelpers::Evaluate(InTime, OutValue, Channel, WeakSection);
+	}
+
 private:
 	FMovieSceneBoolChannel Channel;
 	TWeakObjectPtr<UMovieSceneSection> WeakSection;
diff --git a/Engine/Source/Editor/MovieSceneTools/Private/Channels/ByteChannelCurveModel.cpp b/Engine/Source/Editor/MovieSceneTools/Private/Channels/ByteChannelCurveModel.cpp
--- a/Engine/Source/Editor/MovieSceneTools/Private/Channels/ByteChannelCurveModel.cpp
+++ b/Engine/Source/Editor/MovieSceneTools/Private/Channels/ByteChannelCurveModel.cpp
@@ -3,6 +3,7 @@
 #include "Channels/ByteChannelCurveModel.h"
 
 #include "Algo/BinarySearch.h"
+#include "Channels/CurveModelHelpers.h"
 #include "Channels/ByteChannelKeyProxy.h"
 #include "Channels/MovieSceneChannelData.h"
 #include "Channels/MovieSceneChannelHandle.h"
@@ -74,6 +75,11 @@ public:
 			}
 		}
 	}
+	
+	virtual bool Evaluate(double InTime, double& OutValue) const override
+	{
+		return UE::MovieSceneTools::CurveHelpers::Evaluate(InTime, OutValue, Channel, WeakSection);
+	}
 
 private:
 	FMovieSceneByteChannel Channel;
diff --git a/Engine/Source/Editor/MovieSceneTools/Private/Channels/CurveModelHelpers.h b/Engine/Source/Editor/MovieSceneTools/Private/Channels/CurveModelHelpers.h
new file mode 100644
--- /dev/null
+++ b/Engine/Source/Editor/MovieSceneTools/Private/Channels/CurveModelHelpers.h
@@ -0,0 +1,34 @@
+﻿// Copyright Epic Games, Inc. All Rights Reserved.
+
+#pragma once
+
+#include "MovieScene.h"
+#include "MovieSceneSection.h"
+
+namespace UE::MovieSceneTools::CurveHelpers
+{
+template <typename T>
+concept CCurveValueTypeable = requires { typename T::CurveValueType; };
+
+/** Shared implementation for evaluating a ChannelType. Useful helper for evaluating IBufferedCurveModel. */
+template <CCurveValueTypeable ChannelType>
+bool Evaluate(
+	double InTime, double& OutValue, const ChannelType& Channel, const TWeakObjectPtr<UMovieSceneSection> WeakSection
+	)
+{
+	if (UMovieSceneSection* Section = WeakSection.Get())
+	{
+		FFrameRate TickResolution = Section->GetTypedOuter<UMovieScene>()->GetTickResolution();
+
+		using KeyType = typename ChannelType::CurveValueType;
+		KeyType ThisValue = static_cast<KeyType>(0.0);
+		if (Channel.Evaluate(InTime * TickResolution, ThisValue))
+		{
+			OutValue = ThisValue;
+			return true;
+		}
+	}
+
+	return false;
+}
+}
diff --git a/Engine/Source/Editor/MovieSceneTools/Private/Channels/DoubleChannelCurveModel.cpp b/Engine/Source/Editor/MovieSceneTools/Private/Channels/DoubleChannelCurveModel.cpp
--- a/Engine/Source/Editor/MovieSceneTools/Private/Channels/DoubleChannelCurveModel.cpp
+++ b/Engine/Source/Editor/MovieSceneTools/Private/Channels/DoubleChannelCurveModel.cpp
@@ -2,6 +2,7 @@
 
 #include "Channels/DoubleChannelCurveModel.h"
 
+#include "Channels/CurveModelHelpers.h"
 #include "Channels/DoubleChannelKeyProxy.h"
 #include "Channels/MovieSceneChannelData.h"
 #include "Channels/MovieSceneChannelHandle.h"
@@ -64,6 +65,11 @@ public:
 		}
 	}
 
+	virtual bool Evaluate(double InTime, double& OutValue) const override
+	{
+		return UE::MovieSceneTools::CurveHelpers::Evaluate(InTime, OutValue, Channel, WeakSection);
+	}
+
 private:
 	FMovieSceneDoubleChannel Channel;
 	TWeakObjectPtr<UMovieSceneSection> WeakSection;
diff --git a/Engine/Source/Editor/MovieSceneTools/Private/Channels/FloatChannelCurveModel.cpp b/Engine/Source/Editor/MovieSceneTools/Private/Channels/FloatChannelCurveModel.cpp
--- a/Engine/Source/Editor/MovieSceneTools/Private/Channels/FloatChannelCurveModel.cpp
+++ b/Engine/Source/Editor/MovieSceneTools/Private/Channels/FloatChannelCurveModel.cpp
@@ -2,6 +2,7 @@
 
 #include "Channels/FloatChannelCurveModel.h"
 
+#include "Channels/CurveModelHelpers.h"
 #include "Channels/FloatChannelKeyProxy.h"
 #include "Channels/MovieSceneChannelData.h"
 #include "Channels/MovieSceneChannelHandle.h"
@@ -63,6 +64,11 @@ public:
 		}
 	}
 
+	virtual bool Evaluate(double InTime, double& OutValue) const override
+	{
+		return UE::MovieSceneTools::CurveHelpers::Evaluate(InTime, OutValue, Channel, WeakSection);
+	}
+
 private:
 	FMovieSceneFloatChannel Channel;
 	TWeakObjectPtr<UMovieSceneSection> WeakSection;
diff --git a/Engine/Source/Editor/MovieSceneTools/Private/Channels/IntegerChannelCurveModel.cpp b/Engine/Source/Editor/MovieSceneTools/Private/Channels/IntegerChannelCurveModel.cpp
--- a/Engine/Source/Editor/MovieSceneTools/Private/Channels/IntegerChannelCurveModel.cpp
+++ b/Engine/Source/Editor/MovieSceneTools/Private/Channels/IntegerChannelCurveModel.cpp
@@ -3,6 +3,7 @@
 #include "Channels/IntegerChannelCurveModel.h"
 
 #include "Algo/BinarySearch.h"
+#include "Channels/CurveModelHelpers.h"
 #include "Channels/IntegerChannelKeyProxy.h"
 #include "Channels/MovieSceneChannelData.h"
 #include "Channels/MovieSceneChannelHandle.h"
@@ -128,6 +129,11 @@ public:
 		}
 	}
 
+	virtual bool Evaluate(double InTime, double& OutValue) const override
+	{
+		return UE::MovieSceneTools::CurveHelpers::Evaluate(InTime, OutValue, Channel, WeakSection);
+	}
+
 private:
 	FMovieSceneIntegerChannel Channel;
 	TWeakObjectPtr<UMovieSceneSection> WeakSection;
diff --git a/Engine/Source/Editor/MovieSceneTools/Public/Channels/BezierChannelCurveModel.h b/Engine/Source/Editor/MovieSceneTools/Public/Channels/BezierChannelCurveModel.h
--- a/Engine/Source/Editor/MovieSceneTools/Public/Channels/BezierChannelCurveModel.h
+++ b/Engine/Source/Editor/MovieSceneTools/Public/Channels/BezierChannelCurveModel.h
@@ -43,6 +43,7 @@ public:
 		TArray<FKeyPosition>&& InKeyPositions, TArray<FKeyAttributes>&& InKeyAttributes, const FString& InLongDisplayName, const double InValueMin, const double InValueMax);
 
 	virtual void DrawCurve(const FCurveEditor& InCurveEditor, const FCurveEditorScreenSpace& InScreenSpace, TArray<TTuple<double, double>>& OutInterpolatingPoints) const override;
+	virtual bool Evaluate(double InTime, double& OutValue) const override;
 
 private:
 	ChannelType Channel;
