Sequencer: Add object property button for MPC tracks which allows you to browse to the asset or select another asset.

Added BuildOutlinerEditColumnWidget which allows Track Editors to add a widget to the Edit Column. Note, there is an existing BuildOutlinerEditWidget which really should be named BuildOutlinerAddWidget but looks like it wasn't renamed for backwards compatibility

Removed Select and Set menus for the MPC track since they're now redundant.

diff --git a/Engine/Source/Editor/MovieSceneTools/Private/TrackEditors/MaterialParameterCollectionTrackEditor.cpp b/Engine/Source/Editor/MovieSceneTools/Private/TrackEditors/MaterialParameterCollectionTrackEditor.cpp
--- a/Engine/Source/Editor/MovieSceneTools/Private/TrackEditors/MaterialParameterCollectionTrackEditor.cpp
+++ b/Engine/Source/Editor/MovieSceneTools/Private/TrackEditors/MaterialParameterCollectionTrackEditor.cpp
@@ -33,6 +33,7 @@
 #include "MovieSceneSection.h"
 #include "MovieSceneSequence.h"
 #include "MovieSceneTrack.h"
+#include "PropertyCustomizationHelpers.h"
 #include "ScopedTransaction.h"
 #include "Sections/MovieSceneParameterSection.h"
 #include "Sections/ParameterSection.h"
@@ -110,69 +111,6 @@ TSharedRef<SWidget> CreateAssetPicker(FOnAssetSelected OnAssetSelected, FOnAsset
 		];
 }
 
-void FMaterialParameterCollectionTrackEditor::BuildTrackContextMenu(FMenuBuilder& MenuBuilder, UMovieSceneTrack* Track)
-{
-	UMovieSceneMaterialParameterCollectionTrack* MPCTrack = Cast<UMovieSceneMaterialParameterCollectionTrack>(Track);
-
-	auto AssignAsset = [MPCTrack](const FAssetData& InAssetData)
-	{
-		UMaterialParameterCollection* MPC = Cast<UMaterialParameterCollection>(InAssetData.GetAsset());
-
-		if (MPC)
-		{
-			FScopedTransaction Transaction(LOCTEXT("SetAssetTransaction", "Assign Material Parameter Collection"));
-			MPCTrack->Modify();
-			MPCTrack->SetDisplayName(FText::FromString(MPC->GetName()));
-			MPCTrack->MPC = MPC;
-		}
-
-		FSlateApplication::Get().DismissAllMenus();
-	};
-
-	auto AssignAssetEnterPressed = [AssignAsset](const TArray<FAssetData>& InAssetData)
-	{
-		if (InAssetData.Num() > 0)
-		{
-			AssignAsset(InAssetData[0].GetAsset());
-		}
-	};
-
-	auto SubMenuCallback = [this, AssignAsset, AssignAssetEnterPressed](FMenuBuilder& SubMenuBuilder)
-	{
-		SubMenuBuilder.AddWidget(CreateAssetPicker(FOnAssetSelected::CreateLambda(AssignAsset), FOnAssetEnterPressed::CreateLambda(AssignAssetEnterPressed), GetSequencer()), FText::GetEmpty(), true);
-	};
-
-	UMaterialParameterCollection* MPC = MPCTrack ? MPCTrack->MPC : nullptr;
-
-	MenuBuilder.AddMenuEntry(
-		FText::Format(LOCTEXT("SelectAssetFormat", "Select {0}"), MPC ? FText::FromString(MPC->GetName()) : FText::GetEmpty()),
-		FText::Format(LOCTEXT("SelectAssetTooltipFormat", "Select {0}"), MPC ? FText::FromString(MPC->GetName()) : FText::GetEmpty()),
-		FSlateIcon(),
-		FUIAction(FExecuteAction::CreateRaw(this, &FMaterialParameterCollectionTrackEditor::OnSelectMPC, MPC),
-				  FCanExecuteAction::CreateLambda([MPC] { return MPC != nullptr; }))
-	);
-
-	MenuBuilder.AddSubMenu(
-		LOCTEXT("SetAsset", "Set Asset"),
-		LOCTEXT("SetAsset_ToolTip", "Sets the Material Parameter Collection that this track animates."),
-		FNewMenuDelegate::CreateLambda(SubMenuCallback)
-	);
-}
-
-void FMaterialParameterCollectionTrackEditor::OnSelectMPC(UMaterialParameterCollection* MPC)
-{
-	if (!MPC)
-	{
-		return;
-	}
-
-	TArray<UObject*> ObjectsToFocus;
-	ObjectsToFocus.Add(MPC);
-
-	GEditor->SyncBrowserToObjects(ObjectsToFocus);
-}
-
-
 void FMaterialParameterCollectionTrackEditor::BuildAddTrackMenu(FMenuBuilder& MenuBuilder)
 {
 	auto SubMenuCallback = [this](FMenuBuilder& SubMenuBuilder)
@@ -273,6 +211,51 @@ TSharedPtr<SWidget> FMaterialParameterCollectionTrackEditor::BuildOutlinerEditWi
 	return UE::Sequencer::MakeAddButton(LOCTEXT("AddParameterButton", "Parameter"), MenuContent, Params.ViewModel);
 }
 
+TSharedPtr<SWidget> FMaterialParameterCollectionTrackEditor::BuildOutlinerEditColumnWidget(const FGuid& ObjectBinding, UMovieSceneTrack* Track, const FBuildEditWidgetParams& Params)
+{
+	UMovieSceneMaterialParameterCollectionTrack* MPCTrack = Cast<UMovieSceneMaterialParameterCollectionTrack>(Track);
+	FOnGetContent MenuContent = FOnGetContent::CreateSP(this, &FMaterialParameterCollectionTrackEditor::OnGetAddParameterMenuContent, MPCTrack, Params.RowIndex, Params.TrackInsertRowIndex);
+
+	auto OnSetObjectLambda = [MPCTrack](const FAssetData& Asset) mutable
+	{
+		UMaterialParameterCollection* MPC = Cast<UMaterialParameterCollection>(Asset.GetAsset());
+
+		FScopedTransaction Transaction(LOCTEXT("SetAssetTransaction", "Assign Material Parameter Collection"));
+		MPCTrack->Modify();
+		MPCTrack->SetDisplayName(MPC ? FText::FromString(MPC->GetName()) : FText::GetEmpty());
+		MPCTrack->MPC = MPC;
+	};
+
+	auto GetObjectPathLambda = [MPCTrack]() -> FString
+	{
+		UMaterialParameterCollection* MPC = MPCTrack->MPC;
+		return MPC ? MPC->GetPathName() : FString();
+	};
+
+	TArray<FAssetData> AssetDataArray;
+
+	UMovieSceneSequence* Sequence = GetSequencer()->GetFocusedMovieSceneSequence();
+	AssetDataArray.Add((FAssetData)Sequence);
+
+	return SNew(SHorizontalBox)
+
+		+ SHorizontalBox::Slot()
+		.AutoWidth()
+		[
+			SNew(SBox)
+			.WidthOverride(100.f)
+			[
+				SNew(SObjectPropertyEntryBox)
+				.DisplayBrowse(true)
+				.DisplayUseSelected(false)
+				.ObjectPath_Lambda(GetObjectPathLambda)
+				.AllowedClass(UMaterialParameterCollection::StaticClass())
+				.OnObjectChanged_Lambda(OnSetObjectLambda)
+				.OwnerAssetDataArray(AssetDataArray)
+			]
+		];
+}
+
 TSharedRef<SWidget> FMaterialParameterCollectionTrackEditor::OnGetAddParameterMenuContent(UMovieSceneMaterialParameterCollectionTrack* MPCTrack, int32 RowIndex, int32 TrackInsertRowIndex)
 {
 	FMenuBuilder MenuBuilder(true, nullptr);
diff --git a/Engine/Source/Editor/MovieSceneTools/Private/TrackEditors/MaterialParameterCollectionTrackEditor.h b/Engine/Source/Editor/MovieSceneTools/Private/TrackEditors/MaterialParameterCollectionTrackEditor.h
--- a/Engine/Source/Editor/MovieSceneTools/Private/TrackEditors/MaterialParameterCollectionTrackEditor.h
+++ b/Engine/Source/Editor/MovieSceneTools/Private/TrackEditors/MaterialParameterCollectionTrackEditor.h
@@ -44,21 +44,19 @@ public:
 	//~ ISequencerTrackEditor interface
 	virtual FText GetDisplayName() const override;
 	virtual TSharedPtr<SWidget> BuildOutlinerEditWidget(const FGuid& ObjectBinding, UMovieSceneTrack* Track, const FBuildEditWidgetParams& Params) override;
+	virtual TSharedPtr<SWidget> BuildOutlinerEditColumnWidget(const FGuid& ObjectBinding, UMovieSceneTrack* Track, const FBuildEditWidgetParams& Params) override;
 	virtual TSharedRef<ISequencerSection> MakeSectionInterface(UMovieSceneSection& SectionObject, UMovieSceneTrack& Track, FGuid ObjectBinding) override;
 	virtual bool SupportsType(TSubclassOf<UMovieSceneTrack> Type) const override;
 	virtual bool SupportsSequence(UMovieSceneSequence* InSequence) const override;
 	virtual void BuildAddTrackMenu(FMenuBuilder& MenuBuilder) override;
 	virtual const FSlateBrush* GetIconBrush() const override;
 	virtual bool HandleAssetAdded(UObject* Asset, const FGuid& TargetObjectGuid) override;
-	virtual void BuildTrackContextMenu(FMenuBuilder& MenuBuilder, UMovieSceneTrack* Track) override;
 
 private:
 
 	/** Provides the contents of the add parameter menu. */
 	TSharedRef<SWidget> OnGetAddParameterMenuContent(UMovieSceneMaterialParameterCollectionTrack* MaterialTrack, int32 RowIndex, int32 TrackInsertRowIndex);
 
-	void OnSelectMPC(UMaterialParameterCollection* MPC);
-
 	void AddScalarParameter(UMovieSceneMaterialParameterCollectionTrack* Track, int32 RowIndex, FCollectionScalarParameter Parameter);
 	void AddVectorParameter(UMovieSceneMaterialParameterCollectionTrack* Track, int32 RowIndex, FCollectionVectorParameter Parameter);
 
diff --git a/Engine/Source/Editor/Sequencer/Private/MovieSceneTrackEditor.cpp b/Engine/Source/Editor/Sequencer/Private/MovieSceneTrackEditor.cpp
--- a/Engine/Source/Editor/Sequencer/Private/MovieSceneTrackEditor.cpp
+++ b/Engine/Source/Editor/Sequencer/Private/MovieSceneTrackEditor.cpp
@@ -259,6 +259,22 @@ TSharedPtr<SWidget> FMovieSceneTrackEditor::BuildOutlinerColumnWidget(const FBui
 		{
 			return SNew(SKeyAreaEditorSwitcher, TopLevelChannel, Editor);
 		}
+
+		TAttribute<bool> IsHovered = MakeAttributeSP(Params.TreeViewRow, &ISequencerTreeViewRow::IsHovered);
+
+		FBuildEditWidgetParams EditWidgetParams(Params);
+		EditWidgetParams.NodeIsHovered = IsHovered;
+		EditWidgetParams.RowIndex = Params.TrackModel->GetRowIndex();
+		if (Params.ViewModel->FindAncestorOfType<ITrackExtension>())
+		{
+			EditWidgetParams.TrackInsertRowIndex = Params.TrackModel->GetRowIndex();
+		}
+		else if (Track && Track->SupportsMultipleRows())
+		{
+			EditWidgetParams.TrackInsertRowIndex = Track->GetMaxRowIndex() + 1;
+		}
+		TViewModelPtr<IObjectBindingExtension> ObjectBinding = Params.ViewModel->FindAncestorOfType<IObjectBindingExtension>();
+		return BuildOutlinerEditColumnWidget(ObjectBinding ? ObjectBinding->GetObjectGuid() : FGuid(), Track, EditWidgetParams);
 	}
 
 	if (ColumnName == FCommonOutlinerNames::Add)
@@ -341,7 +357,12 @@ TSharedPtr<SWidget> FMovieSceneTrackEditor::BuildOutlinerEditWidget(const FGuid&
 	}
 }
 
-void FMovieSceneTrackEditor::BuildTrackContextMenu(FMenuBuilder& MenuBuilder, UMovieSceneTrack* Track) 
+TSharedPtr<SWidget> FMovieSceneTrackEditor::BuildOutlinerEditColumnWidget(const FGuid& ObjectBinding, UMovieSceneTrack* Track, const FBuildEditWidgetParams& Params)
+{
+	return TSharedPtr<SWidget>();
+}
+
+void FMovieSceneTrackEditor::BuildTrackContextMenu(FMenuBuilder& MenuBuilder, UMovieSceneTrack* Track)
 {
 }
 
diff --git a/Engine/Source/Editor/Sequencer/Public/ISequencerTrackEditor.h b/Engine/Source/Editor/Sequencer/Public/ISequencerTrackEditor.h
--- a/Engine/Source/Editor/Sequencer/Public/ISequencerTrackEditor.h
+++ b/Engine/Source/Editor/Sequencer/Public/ISequencerTrackEditor.h
@@ -184,9 +184,8 @@ public:
 	 */
 	virtual void BuildObjectBindingContextMenu(FMenuBuilder& MenuBuilder, const TArray<FGuid>& ObjectBindings, const UClass* ObjectClass) {}
 
-
 	/**
-	 * Builds an edit widget for the outliner nodes which represent tracks which are edited by this editor.
+	 * Builds a widget for the add column for outliner nodes which represent tracks which are edited by this editor.
 	 * @param ObjectBinding The object binding associated with the track being edited by this editor.
 	 * @param Track The track being edited by this editor.
 	 * @param Params Parameter struct containing data relevant to the edit widget
@@ -194,6 +193,14 @@ public:
 	 */
 	virtual TSharedPtr<SWidget> BuildOutlinerEditWidget(const FGuid& ObjectBinding, UMovieSceneTrack* Track, const FBuildEditWidgetParams& Params) = 0;
 
+	/**
+	 * Builds a widget for the edit column for outliner nodes which represent tracks which are edited by this editor.
+	 * @param ObjectBinding The object binding associated with the track being edited by this editor.
+	 * @param Track The track being edited by this editor.
+	 * @param Params Parameter struct containing data relevant to the edit widget
+	 * @returns The the widget to display in the outliner, or an empty shared ptr if not widget is to be displayed.
+	 */
+	virtual TSharedPtr<SWidget> BuildOutlinerEditColumnWidget(const FGuid& ObjectBinding, UMovieSceneTrack* Track, const FBuildEditWidgetParams& Params) = 0;
 
 	/**
 	 * Builds an outliner column widget for the outliner nodes which represent tracks which are edited by this editor.
diff --git a/Engine/Source/Editor/Sequencer/Public/MovieSceneTrackEditor.h b/Engine/Source/Editor/Sequencer/Public/MovieSceneTrackEditor.h
--- a/Engine/Source/Editor/Sequencer/Public/MovieSceneTrackEditor.h
+++ b/Engine/Source/Editor/Sequencer/Public/MovieSceneTrackEditor.h
@@ -183,6 +183,7 @@ public:
 	virtual void BuildObjectBindingColumnWidgets(TFunctionRef<TSharedRef<SHorizontalBox>()> GetEditBox, const UE::Sequencer::TViewModelPtr<UE::Sequencer::FObjectBindingModel>& ObjectBinding, const UE::Sequencer::FCreateOutlinerViewParams& InParams, const FName& InColumnName) override;
 	virtual void BuildObjectBindingTrackMenu(FMenuBuilder& MenuBuilder, const TArray<FGuid>& ObjectBindings, const UClass* ObjectClass) override;
 	virtual TSharedPtr<SWidget> BuildOutlinerEditWidget(const FGuid& ObjectBinding, UMovieSceneTrack* Track, const FBuildEditWidgetParams& Params) override;
+	virtual TSharedPtr<SWidget> BuildOutlinerEditColumnWidget(const FGuid& ObjectBinding, UMovieSceneTrack* Track, const FBuildEditWidgetParams& Params) override;
 	virtual TSharedPtr<SWidget> BuildOutlinerColumnWidget(const FBuildColumnWidgetParams& Params, const FName& ColumnName) override;
 	virtual void BuildTrackContextMenu(FMenuBuilder& MenuBuilder, UMovieSceneTrack* Track) override;
 	virtual void BuildTrackSidebarMenu(FMenuBuilder& MenuBuilder, UMovieSceneTrack* Track) override;
