Sequencer: Initial Animation Mixer Track prototype

diff --git a/Engine/Source/Editor/MovieSceneTools/Private/TrackEditors/SkeletalAnimationTrackEditor.cpp b/Engine/Source/Editor/MovieSceneTools/Private/TrackEditors/SkeletalAnimationTrackEditor.cpp
--- a/Engine/Source/Editor/MovieSceneTools/Private/TrackEditors/SkeletalAnimationTrackEditor.cpp
+++ b/Engine/Source/Editor/MovieSceneTools/Private/TrackEditors/SkeletalAnimationTrackEditor.cpp
@@ -1,965 +1,266 @@
 // Copyright Epic Games, Inc. All Rights Reserved.
 
 #include "TrackEditors/SkeletalAnimationTrackEditor.h"
-#include "Rendering/DrawElements.h"
-#include "Widgets/SBoxPanel.h"
+
+#include "AnimSequencerInstanceProxy.h"
 #include "Framework/MultiBox/MultiBoxBuilder.h"
-#include "GameFramework/Actor.h"
-#include "AssetRegistry/AssetData.h"
-#include "Animation/AnimSequenceBase.h"
-#include "Modules/ModuleManager.h"
-#include "Layout/WidgetPath.h"
-#include "Framework/Application/MenuStack.h"
-#include "Framework/Application/SlateApplication.h"
-#include "Widgets/Layout/SBox.h"
-#include "SequencerSectionPainter.h"
-#include "Channels/MovieSceneTimeWarpChannel.h"
-#include "Variants/MovieSceneTimeWarpGetter.h"
-#include "Components/SceneComponent.h"
-#include "Components/SkeletalMeshComponent.h"
-#include "Editor/UnrealEdEngine.h"
-#include "UnrealEdGlobals.h"
 #include "Tracks/MovieSceneSkeletalAnimationTrack.h"
-#include "Sections/MovieSceneSkeletalAnimationSection.h"
-#include "AssetRegistry/AssetRegistryModule.h"
-#include "IContentBrowserSingleton.h"
-#include "ContentBrowserModule.h"
-#include "SequencerSettings.h"
-#include "MVVM/Views/ViewUtilities.h"
-#include "MVVM/ViewModels/ViewDensity.h"
-#include "MVVM/Extensions/ITrackExtension.h"
-#include "ISectionLayoutBuilder.h"
-#include "ISequencer.h"
-#include "Animation/AnimMontage.h"
-#include "Animation/AnimSequence.h"
-#include "Animation/PoseAsset.h"
-#include "Animation/MirrorDataTable.h"
-#include "Styling/AppStyle.h"
-#include "DragAndDrop/AssetDragDropOp.h"
-#include "MovieSceneTimeHelpers.h"
-#include "SequencerTimeSliderController.h"
-#include "FrameNumberDisplayFormat.h"
-#include "FrameNumberNumericInterface.h"
-#include "AnimationBlueprintLibrary.h"
-#include "MovieSceneTransformTypes.h"
-#include "AnimationEditorUtils.h"
-#include "Factories/PoseAssetFactory.h"
-#include "Misc/MessageDialog.h"
-#include "Framework/Notifications/NotificationManager.h"
-#include "SequencerSettings.h"
-#include "Widgets/Notifications/SNotificationList.h"
-
-#include "Engine/SCS_Node.h"
-#include "Engine/SimpleConstructionScript.h"
-#include "Engine/Blueprint.h"
-
-#include "Subsystems/AssetEditorSubsystem.h"
-#include "Editor.h"
-
-#include "ClassViewerFilter.h"
-#include "ClassViewerModule.h"
-
-#include "IDetailCustomization.h"
-#include "DetailLayoutBuilder.h"
-#include "PropertyEditorModule.h"
-#include "IDetailChildrenBuilder.h"
-#include "DetailWidgetRow.h"
-#include "PropertyCustomizationHelpers.h"
-#include "IPropertyUtilities.h"
-#include "Factories/AnimSequenceFactory.h"
-#include "MovieSceneToolHelpers.h"
-
-#include "Widgets/SWindow.h"
-#include "Framework/Application/SlateApplication.h"
-#include "Interfaces/IMainFrameModule.h"
-#include "Widgets/SCompoundWidget.h"
-#include "Widgets/SWindow.h"
-#include "Widgets/Layout/SBorder.h"
-#include "Widgets/Text/STextBlock.h"
-#include "Widgets/Layout/SBox.h"
-#include "Widgets/Layout/SUniformGridPanel.h"
-#include "Widgets/Input/SButton.h"
-#include "Widgets/Input/SComboBox.h"
-#include "Exporters/AnimSeqExportOption.h"
-
-#include "EditModes/SkeletalAnimationTrackEditMode.h"
-#include "EditorModeManager.h"
-
-#include "LevelSequence.h"
-#include "LevelSequenceAnimSequenceLink.h"
-#include "AnimSequenceLevelSequenceLink.h"
-#include "UObject/SavePackage.h"
-#include "AnimSequencerInstanceProxy.h"
-#include "IDetailGroup.h"
-#include "TimeToPixel.h"
-#include "SequencerUtilities.h"
-#include "SequencerAnimationOverride.h"
-
-int32 FSkeletalAnimationTrackEditor::NumberActive = 0;
 
-namespace SkeletalAnimationEditorConstants
-{
-	// @todo Sequencer Allow this to be customizable
-	const uint32 AnimationTrackHeight = 28;
-
-	const TCHAR* CVarName = TEXT("Sequencer.SubsequenceAnimBakeInIsolation");
+#define LOCTEXT_NAMESPACE "FSkeletalAnimationTrackEditor"
 
-	static bool bSubsequenceAnimBakeInIsolation = false; //default to false for now will turn it on after more testing
-	static FAutoConsoleVariableRef CVarSetSubsequenceAnimBakeInIsolation(
-		CVarName,
-		bSubsequenceAnimBakeInIsolation,
-		TEXT("Automatically Force Evaluate Subsequence In Isolation when Baking Anim Sequences")
-	);
 
+namespace UE::Sequencer::SkeletalAnimationEditorConstants
+{
 	static TAutoConsoleVariable<bool> CVarEvaluateSkeletalMeshOnPropertyChange(
 		TEXT("Sequencer.EvaluateSkeletalMeshOnPropertyChange"),
 		true,
 		TEXT("Enable/disable sending a track value changed when properties change so that the skeletal mesh can be re-evaluated in Sequencer"));
-}
-
-#define LOCTEXT_NAMESPACE "FSkeletalAnimationTrackEditor"
+} // namespace UE::Sequencer::SkeletalAnimationEditorConstants
 
 
-class SAnimSequenceOptionsWindow : public SCompoundWidget
+FSkeletalAnimationTrackEditor::FSkeletalAnimationTrackEditor( TSharedRef<ISequencer> InSequencer )
+	: UE::Sequencer::FCommonAnimationTrackEditor(InSequencer)
 {
-public:
-	SLATE_BEGIN_ARGS(SAnimSequenceOptionsWindow)
-		: _ExportOptions(nullptr)
-		, _WidgetWindow()
-		, _FullPath()
-	{}
-
-	SLATE_ARGUMENT(UAnimSeqExportOption*, ExportOptions)
-		SLATE_ARGUMENT(TSharedPtr<SWindow>, WidgetWindow)
-		SLATE_ARGUMENT(FText, FullPath)
-		SLATE_END_ARGS()
-
-public:
-	void Construct(const FArguments& InArgs);
-	virtual bool SupportsKeyboardFocus() const override { return true; }
-
-	FReply OnExport()
-	{
-		bShouldExport = true;
-		if (WidgetWindow.IsValid())
-		{
-			WidgetWindow.Pin()->RequestDestroyWindow();
-		}
-		return FReply::Handled();
-	}
-
-
-	FReply OnCancel()
-	{
-		bShouldExport = false;
-		if (WidgetWindow.IsValid())
-		{
-			WidgetWindow.Pin()->RequestDestroyWindow();
-		}
-		return FReply::Handled();
-	}
-
-	virtual FReply OnKeyDown(const FGeometry& MyGeometry, const FKeyEvent& InKeyEvent) override
-	{
-		if (InKeyEvent.GetKey() == EKeys::Escape)
-		{
-			return OnCancel();
-		}
-
-		return FReply::Unhandled();
-	}
-
-	bool ShouldExport() const
-	{
-		return bShouldExport;
-	}
-
-	
-	SAnimSequenceOptionsWindow()
-		: ExportOptions(nullptr)
-		, bShouldExport(false)
-	{}
-
-private:
-
-	FReply OnResetToDefaultClick() const;
-
-private:
-	UAnimSeqExportOption* ExportOptions;
-	TSharedPtr<class IDetailsView> DetailsView;
-	TWeakPtr< SWindow > WidgetWindow;
-	bool			bShouldExport;
-};
+}
 
+TSharedRef<ISequencerTrackEditor> FSkeletalAnimationTrackEditor::CreateTrackEditor( TSharedRef<ISequencer> InSequencer )
+{
+	return MakeShared<FSkeletalAnimationTrackEditor>(InSequencer);
+}
 
-void SAnimSequenceOptionsWindow::Construct(const FArguments& InArgs)
+void FSkeletalAnimationTrackEditor::OnInitialize()
 {
-	ExportOptions = InArgs._ExportOptions;
-	WidgetWindow = InArgs._WidgetWindow;
-
-	check(ExportOptions);
-
-	FText CancelText =  LOCTEXT("AnimSequenceOptions_Cancel", "Cancel");
-	FText CancelTooltipText = LOCTEXT("AnimSequenceOptions_Cancel_ToolTip", "Cancel the current Anim Sequence Creation.");
-
-	TSharedPtr<SBox> HeaderToolBox;
-	TSharedPtr<SHorizontalBox> AnimHeaderButtons;
-	TSharedPtr<SBox> InspectorBox;
-	this->ChildSlot
-		[
-			SNew(SBox)
-			[
-				SNew(SVerticalBox)
-				+ SVerticalBox::Slot()
-		.AutoHeight()
-		.Padding(2)
-		[
-			SAssignNew(HeaderToolBox, SBox)
-		]
-	+ SVerticalBox::Slot()
-		.AutoHeight()
-		.Padding(2)
-		[
-			SNew(SBorder)
-			.Padding(FMargin(3))
-		.BorderImage(FAppStyle::GetBrush("ToolPanel.GroupBorder"))
-		[
-			SNew(SHorizontalBox)
-			+ SHorizontalBox::Slot()
-		.AutoWidth()
-		[
-			SNew(STextBlock)
-			.Font(FAppStyle::GetFontStyle("CurveEd.LabelFont"))
-		.Text(LOCTEXT("Export_CurrentFileTitle", "Current File: "))
-		]
-	+ SHorizontalBox::Slot()
-		.Padding(5, 0, 0, 0)
-		.AutoWidth()
-		.VAlign(VAlign_Center)
-		[
-			SNew(STextBlock)
-			.Font(FAppStyle::GetFontStyle("CurveEd.InfoFont"))
-		.Text(InArgs._FullPath)
-		]
-		]
-		]
-	+ SVerticalBox::Slot()
-		.FillHeight(1.0f)
-		.Padding(2)
-		[
-			SAssignNew(InspectorBox, SBox)
-		]
-	+ SVerticalBox::Slot()
-		.AutoHeight()
-		.HAlign(HAlign_Right)
-		.Padding(2)
-		[
-			SNew(SUniformGridPanel)
-			.SlotPadding(2)
-	   + SUniformGridPanel::Slot(1, 0)
-		[
-			SNew(SButton)
-			.HAlign(HAlign_Center)
-			.Text(LOCTEXT("AnimExportOptionsWindow_Export", "Export To Animation Sequence"))
-			.OnClicked(this, &SAnimSequenceOptionsWindow::OnExport)
-		]
-	   + SUniformGridPanel::Slot(2, 0)
-		[
-			SNew(SButton)
-			.HAlign(HAlign_Center)
-		.Text(CancelText)
-		.ToolTipText(CancelTooltipText)
-		.OnClicked(this, &SAnimSequenceOptionsWindow::OnCancel)
-		]
-		]
-			]
-		];
-
-	FPropertyEditorModule& PropertyEditorModule = FModuleManager::GetModuleChecked<FPropertyEditorModule>("PropertyEditor");
-	FDetailsViewArgs DetailsViewArgs;
-	DetailsViewArgs.bAllowSearch = false;
-	DetailsViewArgs.NameAreaSettings = FDetailsViewArgs::HideNameArea;
-	DetailsView = PropertyEditorModule.CreateDetailView(DetailsViewArgs);
-
-	InspectorBox->SetContent(DetailsView->AsShared());
-
-	HeaderToolBox->SetContent(
-		SNew(SBorder)
-		.Padding(FMargin(3))
-		.BorderImage(FAppStyle::GetBrush("ToolPanel.GroupBorder"))
-		[
-			SNew(SHorizontalBox)
-			+ SHorizontalBox::Slot()
-		[
-			SNew(SBox)
-			.HAlign(HAlign_Right)
-		[
-			SAssignNew(AnimHeaderButtons, SHorizontalBox)
-			+ SHorizontalBox::Slot()
-		.AutoWidth()
-		.Padding(FMargin(2.0f, 0.0f))
-		[
-			SNew(SButton)
-			.Text(LOCTEXT("AnimSequenceOptions_ResetOptions", "Reset to Default"))
-		.OnClicked(this, &SAnimSequenceOptionsWindow::OnResetToDefaultClick)
-		]
-		]
-		]
-		]
-	);
-
-	DetailsView->SetObject(ExportOptions);
+	FCoreUObjectDelegates::OnObjectPropertyChanged.AddRaw(this, &FSkeletalAnimationTrackEditor::OnPostPropertyChanged);
+	FCommonAnimationTrackEditor::OnInitialize();
 }
 
-FReply SAnimSequenceOptionsWindow::OnResetToDefaultClick() const
+void FSkeletalAnimationTrackEditor::OnRelease()
 {
-	ExportOptions->ResetToDefault();
-	//Refresh the view to make sure the custom UI are updating correctly
-	DetailsView->SetObject(ExportOptions, true);
-	return FReply::Handled();
+	FCoreUObjectDelegates::OnObjectPropertyChanged.RemoveAll(this);
+	FCommonAnimationTrackEditor::OnRelease();
 }
 
+bool FSkeletalAnimationTrackEditor::SupportsSequence(UMovieSceneSequence* InSequence) const
+{
+	ETrackSupport TrackSupported = InSequence ? InSequence->IsTrackSupported(UMovieSceneSkeletalAnimationTrack::StaticClass()) : ETrackSupport::NotSupported;
+	return TrackSupported == ETrackSupport::Supported;
+}
 
-USkeletalMeshComponent* AcquireSkeletalMeshFromObjectGuid(const FGuid& Guid, TSharedPtr<ISequencer> SequencerPtr)
+bool FSkeletalAnimationTrackEditor::SupportsType( TSubclassOf<UMovieSceneTrack> Type ) const
 {
-	UObject* BoundObject = SequencerPtr.IsValid() ? SequencerPtr->FindSpawnedObjectOrTemplate(Guid) : nullptr;
+	return Type == UMovieSceneSkeletalAnimationTrack::StaticClass();
+}
 
-	if (AActor* Actor = Cast<AActor>(BoundObject))
+void FSkeletalAnimationTrackEditor::OnPostPropertyChanged(UObject* InObject, FPropertyChangedEvent& InPropertyChangedEvent)
+{
+	if (SkeletalAnimationEditorConstants::CVarEvaluateSkeletalMeshOnPropertyChange->GetBool())
 	{
-		if (USkeletalMeshComponent* SkeletalMeshComponent = Cast<USkeletalMeshComponent>(Actor->GetRootComponent()))
+		const FName PropertyName = (InPropertyChangedEvent.Property != nullptr) ? InPropertyChangedEvent.Property->GetFName() : NAME_None;
+		if (PropertyName == USceneComponent::GetRelativeLocationPropertyName() ||
+			PropertyName == USceneComponent::GetRelativeRotationPropertyName() ||
+			PropertyName == USceneComponent::GetRelativeScale3DPropertyName())
 		{
-			return SkeletalMeshComponent;
+			return;
 		}
 
-		TArray<USkeletalMeshComponent*> SkeletalMeshComponents;
-		Actor->GetComponents(SkeletalMeshComponents);
-		
-		if (SkeletalMeshComponents.Num() == 1)
-		{
-			return SkeletalMeshComponents[0];
-		}
-	}
-	else if(USkeletalMeshComponent* SkeletalMeshComponent = Cast<USkeletalMeshComponent>(BoundObject))
-	{
-		if (SkeletalMeshComponent->GetSkeletalMeshAsset())
+		// If the object changed has any animation track, notify sequencer to update because animations tick on their own and sequencer needs to evaluate again
+		const bool bCreateIfMissing = false;
+		FFindOrCreateHandleResult HandleResult = FindOrCreateHandleToObject(InObject, bCreateIfMissing);
+		FGuid ObjectHandle = HandleResult.Handle;
+		if (ObjectHandle.IsValid())
 		{
-			return SkeletalMeshComponent;
+			FFindOrCreateTrackResult TrackResult = FindOrCreateTrackForObject(ObjectHandle, UMovieSceneSkeletalAnimationTrack::StaticClass(), NAME_None, bCreateIfMissing);
+			if (TrackResult.Track)
+			{
+				GetSequencer()->NotifyMovieSceneDataChanged(EMovieSceneDataChangeType::TrackValueChanged);
+			}
 		}
 	}
-
-	return nullptr;
 }
 
-USkeleton* GetSkeletonFromComponent(UActorComponent* InComponent)
+TSubclassOf<UMovieSceneCommonAnimationTrack> FSkeletalAnimationTrackEditor::GetTrackClass() const
 {
-	USkeletalMeshComponent* SkeletalMeshComp = Cast<USkeletalMeshComponent>(InComponent);
-	if (SkeletalMeshComp && SkeletalMeshComp->GetSkeletalMeshAsset() && SkeletalMeshComp->GetSkeletalMeshAsset()->GetSkeleton())
-	{
-		// @todo Multiple actors, multiple components
-		return SkeletalMeshComp->GetSkeletalMeshAsset()->GetSkeleton();
-	}
-
-	return nullptr;
+	return UMovieSceneSkeletalAnimationTrack::StaticClass();
 }
 
-// Get the skeletal mesh components from the guid
-// If bGetSingleRootComponent - return only the root component if it is a skeletal mesh component. 
-// This allows the root object binding to have an animation track without needing a skeletal mesh component binding
-//
-TArray<USkeletalMeshComponent*> AcquireSkeletalMeshComponentsFromObjectGuid(const FGuid& Guid, TSharedPtr<ISequencer> SequencerPtr, const bool bGetSingleRootComponent = true)
+void FSkeletalAnimationTrackEditor::BuildTrackContextMenu_Internal(FMenuBuilder& MenuBuilder, UMovieSceneTrack* Track, const bool bAddSeparatorAtEnd)
 {
-	TArray<USkeletalMeshComponent*> SkeletalMeshComponents;
-
-	UObject* BoundObject = SequencerPtr.IsValid() ? SequencerPtr->FindSpawnedObjectOrTemplate(Guid) : nullptr;
-
-	AActor* Actor = Cast<AActor>(BoundObject);
+	//there's a bug with a section being open already, so we end it.
+	UMovieSceneSkeletalAnimationTrack* SkeletalAnimationTrack = Cast<UMovieSceneSkeletalAnimationTrack>(Track);
 
-	if (!Actor)
+	/** Put this back when and if it works
+	MenuBuilder.BeginSection(NAME_None, LOCTEXT("MotionBlendingOptions", "Motion Blending Options"));
 	{
-		if (UChildActorComponent* ChildActorComponent = Cast<UChildActorComponent>(BoundObject))
-		{
-			Actor = ChildActorComponent->GetChildActor();
-		}
+		MenuBuilder.AddMenuEntry(
+			NSLOCTEXT("Sequencer", "AutoMatchClipsRootMotions", "Auto Match Clips Root Motions"),
+			NSLOCTEXT("Sequencer", "AutoMatchClipsRootMotionsTooltip", "Preceeding clips will auto match to the preceding clips root bones position. You can override this behavior per clip in it's section options."),
+			FSlateIcon(),
+			FUIAction(
+				FExecuteAction::CreateLambda([=]()->void {
+					SkeletalAnimationTrack->ToggleAutoMatchClipsRootMotions(); 
+					SequencerPtr->NotifyMovieSceneDataChanged(EMovieSceneDataChangeType::TrackValueChanged);
+
+					}),
+				FCanExecuteAction::CreateLambda([=]()->bool { return SequencerPtr && SkeletalAnimationTrack != nullptr; }),
+				FIsActionChecked::CreateLambda([=]()->bool { return SkeletalAnimationTrack != nullptr && SkeletalAnimationTrack->bAutoMatchClipsRootMotions; })),
+			NAME_None, 
+			EUserInterfaceActionType::ToggleButton
+		);
+		MenuBuilder.EndSection();
 	}
+	*/
 
-	if (Actor)
+	MenuBuilder.BeginSection(NAME_None, LOCTEXT("SkelAnimRootMOtion", "Root Motion"));
 	{
-		if (bGetSingleRootComponent)
-		{
-			if (USkeletalMeshComponent* SkeletalMeshComponent = Cast<USkeletalMeshComponent>(Actor->GetRootComponent()))
-			{
-				SkeletalMeshComponents.Add(SkeletalMeshComponent);
-				return SkeletalMeshComponents;
-			}
-		}
-
-		Actor->GetComponents(SkeletalMeshComponents);
-		if (SkeletalMeshComponents.Num())
-		{
-			return SkeletalMeshComponents;
-		}
-
-		AActor* ActorCDO = Cast<AActor>(Actor->GetClass()->GetDefaultObject());
-		if (ActorCDO)
-		{
-			if (bGetSingleRootComponent)
-			{
-				if (USkeletalMeshComponent* SkeletalMeshComponent = Cast<USkeletalMeshComponent>(ActorCDO->GetRootComponent()))
+		MenuBuilder.AddMenuEntry(
+			NSLOCTEXT("Sequencer", "BlendFirstChildOfRoot", "Blend First Child Of Root"),
+			NSLOCTEXT("Sequencer", "BlendFirstChildOfRootTooltip", "If True, do not blend and match the root bones but instead the first child bone of the root. Toggle this on when the matched sequences in the track have no motion on the root."),
+			FSlateIcon(),
+			FUIAction(
+				FExecuteAction::CreateLambda([this, SkeletalAnimationTrack]() -> void
 				{
-					SkeletalMeshComponents.Add(SkeletalMeshComponent);
-					return SkeletalMeshComponents;
-				}
-			}
+					SkeletalAnimationTrack->bBlendFirstChildOfRoot = SkeletalAnimationTrack->bBlendFirstChildOfRoot ? false : true;
+					SkeletalAnimationTrack->SetRootMotionsDirty();
 
-			ActorCDO->GetComponents(SkeletalMeshComponents);
-			if (SkeletalMeshComponents.Num())
-			{
-				return SkeletalMeshComponents;
-			}
-		}
+					if (const TSharedPtr<ISequencer> SequencerPtr = GetSequencer())
+					{
+						SequencerPtr->NotifyMovieSceneDataChanged(EMovieSceneDataChangeType::TrackValueChanged);
+					}
+				}),
+				FCanExecuteAction::CreateLambda([this, SkeletalAnimationTrack]() -> bool
+				{
+					return SkeletalAnimationTrack != nullptr;
+				}),
+				FIsActionChecked::CreateLambda([this, SkeletalAnimationTrack]() -> bool
+				{
+					return SkeletalAnimationTrack != nullptr && SkeletalAnimationTrack->bBlendFirstChildOfRoot;
+				})),
+				NAME_None,
+				EUserInterfaceActionType::ToggleButton
+			);
 
-		UBlueprintGeneratedClass* ActorBlueprintGeneratedClass = Cast<UBlueprintGeneratedClass>(Actor->GetClass());
-		if (ActorBlueprintGeneratedClass && ActorBlueprintGeneratedClass->SimpleConstructionScript)
-		{
-			const TArray<USCS_Node*>& ActorBlueprintNodes = ActorBlueprintGeneratedClass->SimpleConstructionScript->GetAllNodes();
+		MenuBuilder.AddMenuEntry(
+			NSLOCTEXT("Sequencer", "ShowRootMotionTrails", "Show Root Motion Trail"),
+			NSLOCTEXT("Sequencer", "ShowRootMotionTrailsTooltip", "Show the Root Motion Trail for all Animation Clips."),
+			FSlateIcon(),
+			FUIAction(
+				FExecuteAction::CreateLambda([this, SkeletalAnimationTrack]() -> void
+				{
+					SkeletalAnimationTrack->ToggleShowRootMotionTrail();
 
-			for (USCS_Node* Node : ActorBlueprintNodes)
-			{
-				if (Node->ComponentClass->IsChildOf(USkeletalMeshComponent::StaticClass()))
+					if (const TSharedPtr<ISequencer> SequencerPtr = GetSequencer())
+					{
+						SequencerPtr->NotifyMovieSceneDataChanged(EMovieSceneDataChangeType::TrackValueChanged);
+					}
+				}),
+				FCanExecuteAction::CreateLambda([this, SkeletalAnimationTrack]()->bool
 				{
-					if (USkeletalMeshComponent* SkeletalMeshComponent = Cast<USkeletalMeshComponent>(Node->GetActualComponentTemplate(ActorBlueprintGeneratedClass)))
+					if (const TSharedPtr<ISequencer> SequencerPtr = GetSequencer())
 					{
-						SkeletalMeshComponents.Add(SkeletalMeshComponent);
+						return SkeletalAnimationTrack != nullptr;
 					}
-				}
-			}
+					return false;
+				}),
+				FIsActionChecked::CreateLambda([this, SkeletalAnimationTrack]()->bool
+				{
+					return SkeletalAnimationTrack != nullptr && SkeletalAnimationTrack->bShowRootMotionTrail;
+				})),
+			NAME_None,
+				EUserInterfaceActionType::ToggleButton
+			);
 
-			if (SkeletalMeshComponents.Num())
-			{
-				return SkeletalMeshComponents;
-			}
-		}
+		MenuBuilder.AddMenuEntry(
+			NSLOCTEXT("Sequencer", "SwapRootBoneNone", "Swap Root Bone None"),
+			NSLOCTEXT("Sequencer", "SwapRootBoneNoneTooltip", "Do not swap root bone for all sections."),
+			FSlateIcon(),
+			FUIAction(
+				FExecuteAction::CreateLambda([SkeletalAnimationTrack]()->void {
+					SkeletalAnimationTrack->SetSwapRootBone(ESwapRootBone::SwapRootBone_None);
+					}),
+				FCanExecuteAction::CreateLambda([SkeletalAnimationTrack]()->bool { return  SkeletalAnimationTrack != nullptr; }),
+				FIsActionChecked::CreateLambda([SkeletalAnimationTrack]()->bool { return SkeletalAnimationTrack != nullptr && SkeletalAnimationTrack->SwapRootBone == ESwapRootBone::SwapRootBone_None; })),
+				NAME_None,
+				EUserInterfaceActionType::RadioButton
+		);
+
+		MenuBuilder.AddMenuEntry(
+			NSLOCTEXT("Sequencer", "SwapRootBoneActor", "Swap Root Bone Actor"),
+			NSLOCTEXT("Sequencer", "SwapRootBoneActorTooltip", "Swap root bone on root actor component for all sections."),
+			FSlateIcon(),
+			FUIAction(
+				FExecuteAction::CreateLambda([SkeletalAnimationTrack]()->void {
+					SkeletalAnimationTrack->SetSwapRootBone(ESwapRootBone::SwapRootBone_Actor);
+					}),
+				FCanExecuteAction::CreateLambda([SkeletalAnimationTrack]()->bool { return  SkeletalAnimationTrack != nullptr; }),
+				FIsActionChecked::CreateLambda([SkeletalAnimationTrack]()->bool { return SkeletalAnimationTrack != nullptr && SkeletalAnimationTrack->SwapRootBone == ESwapRootBone::SwapRootBone_Actor; })),
+				NAME_None,
+				EUserInterfaceActionType::RadioButton
+			);
+		MenuBuilder.AddMenuEntry(
+			NSLOCTEXT("Sequencer", "SwapRootBoneComponent", "Swap Root Bone Component"),
+			NSLOCTEXT("Sequencer", "SwapRootBoneComponentTooltip", "Swap root bone on current component for all sections."),
+			FSlateIcon(),
+			FUIAction(
+				FExecuteAction::CreateLambda([SkeletalAnimationTrack]()->void {
+					SkeletalAnimationTrack->SetSwapRootBone(ESwapRootBone::SwapRootBone_Component);
+					}),
+				FCanExecuteAction::CreateLambda([SkeletalAnimationTrack]()->bool { return  SkeletalAnimationTrack != nullptr; }),
+				FIsActionChecked::CreateLambda([SkeletalAnimationTrack]()->bool { return SkeletalAnimationTrack != nullptr && SkeletalAnimationTrack->SwapRootBone == ESwapRootBone::SwapRootBone_Component; })),
+				NAME_None,
+				EUserInterfaceActionType::RadioButton
+			);
 	}
-	else if (USkeletalMeshComponent* SkeletalMeshComponent = Cast<USkeletalMeshComponent>(BoundObject))
+	MenuBuilder.EndSection();
+
+	if (bAddSeparatorAtEnd)
 	{
-		SkeletalMeshComponents.Add(SkeletalMeshComponent);
-		return SkeletalMeshComponents;
+		MenuBuilder.AddSeparator();
 	}
-	
-	return SkeletalMeshComponents;
 }
 
-USkeleton* AcquireSkeletonFromObjectGuid(const FGuid& Guid, TSharedPtr<ISequencer> SequencerPtr)
+TSharedRef<ISequencerSection> FSkeletalAnimationTrackEditor::MakeSectionInterface(UMovieSceneSection& SectionObject, UMovieSceneTrack& Track, FGuid ObjectBinding)
 {
-	TArray<USkeletalMeshComponent*> SkeletalMeshComponents = AcquireSkeletalMeshComponentsFromObjectGuid(Guid, SequencerPtr);
-
-	if (SkeletalMeshComponents.Num() == 1)
-	{
-		return GetSkeletonFromComponent(SkeletalMeshComponents[0]);
-	}
-
-	return nullptr;
+	check( SupportsType( SectionObject.GetOuter()->GetClass() ) );
+	return MakeShared<FSkeletalAnimationSection>(SectionObject, GetSequencer());
 }
 
+FSkeletalAnimationSection::FSkeletalAnimationSection( UMovieSceneSection& InSection, TWeakPtr<ISequencer> InSequencer)
+	: UE::Sequencer::FCommonAnimationSection( InSection, InSequencer)
+{
+}
 
-class FMovieSceneSkeletalAnimationParamsDetailCustomization : public IPropertyTypeCustomization
+void FSkeletalAnimationSection::BuildSectionContextMenu(FMenuBuilder& MenuBuilder, const FGuid& ObjectBinding)
 {
-public:
-	FMovieSceneSkeletalAnimationParamsDetailCustomization(const FSequencerSectionPropertyDetailsViewCustomizationParams& InParams)
-		: Params(InParams)
+	const TSharedPtr<ISequencer> SequencerPtr = Sequencer.Pin();
+	if (!SequencerPtr.IsValid())
 	{
-		if (Params.ParentObjectBindingGuid.IsValid())
-		{
-			if (USkeletalMeshComponent* SkelMeshComp = AcquireSkeletalMeshFromObjectGuid(Params.ParentObjectBindingGuid, Params.SequencerWeak.Pin()))
-			{
-				TScriptInterface<ISequencerAnimationOverride> SequencerAnimOverride = ISequencerAnimationOverride::GetSequencerAnimOverride(SkelMeshComp);
-				if (SequencerAnimOverride.GetObject())
-				{
-					bAllowsCinematicOverride = ISequencerAnimationOverride::Execute_AllowsCinematicOverride(SequencerAnimOverride.GetObject());
-					SlotNameOptions = ISequencerAnimationOverride::Execute_GetSequencerAnimSlotNames(SequencerAnimOverride.GetObject());
-					bShowSlotNameOptions = SlotNameOptions.Num() > 0 && !bAllowsCinematicOverride;
-				}
-			}
-		}
+		return;
 	}
 
-	// IDetailCustomization interface
-	virtual void CustomizeHeader(TSharedRef<IPropertyHandle> PropertyHandle, FDetailWidgetRow& HeaderRow, IPropertyTypeCustomizationUtils& CustomizationUtils) override
+	UMovieSceneSkeletalAnimationTrack* const Track = Section.GetTypedOuter<UMovieSceneSkeletalAnimationTrack>();
+	if (!IsValid(Track))
 	{
-		SlotNameProperty = PropertyHandle->GetChildHandle(GET_MEMBER_NAME_CHECKED(FMovieSceneSkeletalAnimationParams, SlotName));
+		return;
 	}
 
-	IDetailPropertyRow& AddPropertyRow(IDetailChildrenBuilder& ChildBuilder, const FString& GroupName, TMap<FString, IDetailGroup*>& NameToGroupMap, const TSharedPtr<IPropertyHandle>& ChildPropertyHandle)
+	// Can't pick the object that this track binds
+	USkeleton* const Skeleton = UE::Sequencer::FCommonAnimationTrackEditor::AcquireSkeletonFromObjectGuid(ObjectBinding, SequencerPtr);
+	if (!IsValid(Skeleton))
 	{
-		if (GroupName.IsEmpty())
-		{
-			return ChildBuilder.AddProperty(ChildPropertyHandle.ToSharedRef());
-		}
-		
-		IDetailGroup* const* DetailGroup = NameToGroupMap.Find(GroupName);
-		if (!DetailGroup)
-		{
-			IDetailGroup& NewGroup = ChildBuilder.AddGroup(*GroupName, FText::FromString(GroupName));
-			DetailGroup = &NameToGroupMap.Add(GroupName, &NewGroup);
-		}
-		return (*DetailGroup)->AddPropertyRow(ChildPropertyHandle.ToSharedRef());
-	
+		return;
 	}
 
-	virtual void CustomizeChildren(TSharedRef<IPropertyHandle> PropertyHandle, IDetailChildrenBuilder& ChildBuilder, IPropertyTypeCustomizationUtils& CustomizationUtils) override
+	const int32 NumBones = Skeleton->GetReferenceSkeleton().GetNum();
+	TArray<FName> BoneNames;
+	for (int32 BoneIndex = 0; BoneIndex < NumBones; ++BoneIndex)
 	{
-		const FName AnimationPropertyName = GET_MEMBER_NAME_CHECKED(FMovieSceneSkeletalAnimationParams, Animation);
-		const FName MirrorDataTableName = GET_MEMBER_NAME_CHECKED(FMovieSceneSkeletalAnimationParams, MirrorDataTable);
-		const FName SlotNamePropertyName = GET_MEMBER_NAME_CHECKED(FMovieSceneSkeletalAnimationParams, SlotName);
-
-		TMap<FString, IDetailGroup*> NameToGroupMap;
-
-		uint32 NumChildren;
-		PropertyHandle->GetNumChildren(NumChildren);
-		for (uint32 i = 0; i < NumChildren; ++i)
-		{
-			TSharedPtr<IPropertyHandle> ChildPropertyHandle = PropertyHandle->GetChildHandle(i);
-			FName ChildPropertyName = ChildPropertyHandle->GetProperty()->GetFName();
-			
-			FString GroupName;
-			FString CategoryName;
-			ChildPropertyHandle->GetDefaultCategoryName().ToString().Split("|", &CategoryName, &GroupName);
-
-			IDetailPropertyRow& ChildPropertyRow = AddPropertyRow(ChildBuilder, GroupName, NameToGroupMap, ChildPropertyHandle);
-			
-			// Let most properties be whatever they want to be... we just want to customize the `Animation` and `MirrorDataTable` properties
-			// by making it look like a normal asset reference property, but with some custom filtering.
-			if (ChildPropertyName == AnimationPropertyName || ChildPropertyName == MirrorDataTableName)
-			{
-				FDetailWidgetRow& Row = ChildPropertyRow.CustomWidget();
+		BoneNames.Add(Skeleton->GetReferenceSkeleton().GetBoneName(BoneIndex));
+	}
 
-				if (Params.ParentObjectBindingGuid.IsValid())
-				{
-					// Store the compatible skeleton's name, and create a property widget with a filter that will check
-					// for animations that match that skeleton.
-					Skeleton = AcquireSkeletonFromObjectGuid(Params.ParentObjectBindingGuid, Params.SequencerWeak.Pin());
-					SkeletonName = FAssetData(Skeleton).GetExportTextName();
-
-					TSharedPtr<IPropertyUtilities> PropertyUtilities = CustomizationUtils.GetPropertyUtilities();
-					UClass* AllowedStaticClass = ChildPropertyName == AnimationPropertyName ? UAnimSequenceBase::StaticClass() : UMirrorDataTable::StaticClass(); 
-
-					TSharedRef<SObjectPropertyEntryBox> ContentWidget = SNew(SObjectPropertyEntryBox)
-						.PropertyHandle(ChildPropertyHandle)
-						.AllowedClass(AllowedStaticClass)
-						.DisplayThumbnail(true)
-						.ThumbnailPool(PropertyUtilities.IsValid() ? PropertyUtilities->GetThumbnailPool() : nullptr)
-						.OnShouldFilterAsset(FOnShouldFilterAsset::CreateRaw(this, &FMovieSceneSkeletalAnimationParamsDetailCustomization::ShouldFilterAsset));
-
-					Row.NameContent()[ChildPropertyHandle->CreatePropertyNameWidget()];
-					Row.ValueContent()[ContentWidget];
-
-					float MinDesiredWidth, MaxDesiredWidth;
-					ContentWidget->GetDesiredWidth(MinDesiredWidth, MaxDesiredWidth);
-					Row.ValueContent().MinWidth = MinDesiredWidth;
-					Row.ValueContent().MaxWidth = MaxDesiredWidth;
-
-					// The content widget already contains a "reset to default" button, so we don't want the details view row
-					// to make another one. We add this metadata on the property handle instance to suppress it.
-					ChildPropertyHandle->SetInstanceMetaData(TEXT("NoResetToDefault"), TEXT("true"));
-				}
-			}
-			else if (ChildPropertyName == SlotNamePropertyName)
-			{
-				// If the anim instance implements the ISequencerAnimationOverride interface, and has defined slots to used, override this row with a drop-down menu
-				// Otherwise the default row will be created, which uses a text input field.
-				if (bShowSlotNameOptions)
-				{
-					ChildPropertyRow.IsEnabled(TAttribute<bool>::CreateSP(this, &FMovieSceneSkeletalAnimationParamsDetailCustomization::GetCanEditSlotName));
-					FDetailWidgetRow& Row = ChildPropertyRow.CustomWidget();
-					Row.NameContent()[ChildPropertyHandle->CreatePropertyNameWidget()];
-
-					Row.ValueContent()
-					[
-						SNew(SComboBox<FName>)
-						.OptionsSource(&SlotNameOptions)
-						.OnSelectionChanged(this, &FMovieSceneSkeletalAnimationParamsDetailCustomization::OnSlotNameChanged)
-						.OnGenerateWidget_Lambda([](FName InSlotName) { return SNew(STextBlock).Text(FText::FromName(InSlotName)); })
-						[
-							SNew(STextBlock)
-							.Font(IPropertyTypeCustomizationUtils::GetRegularFont())
-							.Text(this, &FMovieSceneSkeletalAnimationParamsDetailCustomization::GetSlotNameDesc)
-						]
-					];
-				}
-			}
-		}
-	}
-
-	bool ShouldFilterAsset(const FAssetData& AssetData)
-	{
-		// Since the `SObjectPropertyEntryBox` doesn't support passing some `Filter` properties for the asset picker, 
-		// we just combine the tag value filtering we want (i.e. checking the skeleton compatibility) along with the
-		// other filtering we already get from the track editor's filter callback.
-		FSkeletalAnimationTrackEditor& TrackEditor = static_cast<FSkeletalAnimationTrackEditor&>(Params.TrackEditor);
-		if (TrackEditor.ShouldFilterAsset(AssetData))
-		{
-			return true;
-		}
-
-		return !(Skeleton && Skeleton->IsCompatibleForEditor(AssetData));
-	}
-
-
-	FText GetSlotNameDesc() const
-	{
-		FName NameValue;
-		SlotNameProperty->GetValue(NameValue);
-
-		return FText::FromString(NameValue.ToString());
-	}
-
-	bool GetCanEditSlotName() const
-	{
-		if (bShowSlotNameOptions)
-		{
-			FName NameValue;
-			SlotNameProperty->GetValue(NameValue);
-			// If we're allowing cinematic override, then the slot names are irrelevant, don't allow edit.
-			// If we have less than 2 slot name options, then changing them is irrelevant, don't allow edit.
-			// Always allow an edit if the current slot name isn't currently set to one of the provided ones.
-			if (bAllowsCinematicOverride || (SlotNameOptions.Num() < 2 && SlotNameOptions.Contains(NameValue)))
-			{
-				return false;
-			}
-		}
-		return true;
-	}
-
-	void OnSlotNameChanged(FName InSlotName, ESelectInfo::Type InInfo)
-	{
-		SlotNameProperty->SetValue(InSlotName);
-	}
-
-private:
-	FSequencerSectionPropertyDetailsViewCustomizationParams Params;
-	FString SkeletonName;
-	USkeleton* Skeleton = nullptr;
-	TSharedPtr<IPropertyHandle> SlotNameProperty;
-	TArray<FName> SlotNameOptions;
-	bool bShowSlotName = true;
-	bool bShowSlotNameOptions = false;
-	bool bAllowsCinematicOverride = false;
-};
-
-
-FSkeletalAnimationSection::FSkeletalAnimationSection( UMovieSceneSection& InSection, TWeakPtr<ISequencer> InSequencer)
-	: Section(*CastChecked<UMovieSceneSkeletalAnimationSection>(&InSection))
-	, Sequencer(InSequencer)
-	, PreDilatePlayRate(1.0)
-{
-}
-
-FSkeletalAnimationSection::~FSkeletalAnimationSection()
-{
-}
-
-void FSkeletalAnimationSection::BeginDilateSection()
-{
-	if (Section.Params.PlayRate.GetType() == EMovieSceneTimeWarpType::FixedPlayRate)
-	{
-		PreDilatePlayRate = Section.Params.PlayRate.AsFixedPlayRate(); //make sure to cache the play rate
-	}
-	else if (Section.Params.PlayRate.GetType() == EMovieSceneTimeWarpType::Custom)
-	{
-		FMovieSceneTimeWarpChannel* Channel = Section.GetChannelProxy().GetChannel<FMovieSceneTimeWarpChannel>(0);
-		if (Channel)
-		{
-			Section.Params.PlayRate.AsCustom()->Modify();
-			PreDilateChannel = MakeUnique<FMovieSceneTimeWarpChannel>(*Channel);
-		}
-	}
-}
-
-void FSkeletalAnimationSection::DilateSection(const TRange<FFrameNumber>& NewRange, float DilationFactor)
-{
-	if (Section.Params.PlayRate.GetType() == EMovieSceneTimeWarpType::FixedPlayRate)
-	{
-		Section.Params.PlayRate.Set(PreDilatePlayRate / DilationFactor);
-	}
-	else if (Section.Params.PlayRate.GetType() == EMovieSceneTimeWarpType::Custom)
-	{
-		FMovieSceneTimeWarpChannel* Channel = Section.GetChannelProxy().GetChannel<FMovieSceneTimeWarpChannel>(0);
-		if (Channel)
-		{
-			*Channel = *PreDilateChannel;
-
-			// Dilate the times
-			Dilate(Channel, FFrameNumber(0), DilationFactor);
-
-			Section.Params.PlayRate.AsCustom()->MarkAsChanged();
-		}
-	}
-	Section.SetRange(NewRange);
-}
-
-UMovieSceneSection* FSkeletalAnimationSection::GetSectionObject()
-{ 
-	return &Section;
-}
-
-FText FSkeletalAnimationSection::GetSectionTitle() const
-{
-	if (Section.Params.Animation != nullptr)
-	{
-		if (!Section.Params.MirrorDataTable)
-		{
-			return FText::FromString( Section.Params.Animation->GetName() );
-		}
-		else
-		{
-			return FText::Format(LOCTEXT("SectionTitleContentFormat", "{0} mirrored with {1}"), FText::FromString(Section.Params.Animation->GetName()), FText::FromString(Section.Params.MirrorDataTable->GetName()));
-		}
-	}
-	return LOCTEXT("NoAnimationSection", "No Animation");
-}
-
-FText FSkeletalAnimationSection::GetSectionToolTip() const
-{
-	if (Section.Params.Animation != nullptr && Section.HasStartFrame() && Section.HasEndFrame())
-	{
-		UMovieScene* MovieScene = Section.GetTypedOuter<UMovieScene>();
-		FFrameRate TickResolution = MovieScene->GetTickResolution();
-
-		FMovieSceneSequenceTransform Transform = Section.Params.MakeTransform(TickResolution, Section.GetRange());
-
-		const double StartTime     = Transform.TransformTime(0).AsDecimal();
-		const float  SectionLength = Section.GetRange().Size<FFrameTime>() / TickResolution;
-
-		if (Section.Params.PlayRate.GetType() == EMovieSceneTimeWarpType::FixedPlayRate)
-		{
-			const double PlayRate = Section.Params.PlayRate.AsFixedPlayRate();
-			return FText::Format(LOCTEXT("ToolTipContentFormat_FixedPlayRate", "Start: {0}s\nDuration: {1}s\nPlay Rate: {2}x"), StartTime, SectionLength, PlayRate);
-		}
-		else if (Section.Params.PlayRate.GetType() == EMovieSceneTimeWarpType::Custom)
-		{
-			return FText::Format(LOCTEXT("ToolTipContentFormat_TimwWarp", "Start: {0}s\nDuration: {1}s\nPlay Rate: Variable"), StartTime, SectionLength);
-		}
-	}
-	return FText::GetEmpty();
-}
-
-TOptional<FFrameTime> FSkeletalAnimationSection::GetSectionTime(FSequencerSectionPainter& InPainter) const
-{
-	if (!InPainter.bIsSelected || !Sequencer.Pin() || Section.Params.Animation == nullptr)
-	{
-		return TOptional<FFrameTime>();
-	}
-
-	FFrameTime CurrentTime = Sequencer.Pin()->GetLocalTime().Time;
-	if (!Section.GetRange().Contains(CurrentTime.FrameNumber))
-	{
-		return TOptional<FFrameTime>();
-	}
-
-	const FTimeToPixel& TimeToPixelConverter = InPainter.GetTimeConverter();
-	FFrameRate TickResolution = TimeToPixelConverter.GetTickResolution();
-
-	// Draw the current time next to the scrub handle
-	const double AnimTime = Section.MapTimeToAnimation(CurrentTime, TickResolution);
-	const FFrameRate SamplingFrameRate = Section.Params.Animation->GetSamplingFrameRate();
-
-	FQualifiedFrameTime HintFrameTime;
-	if (!UAnimationBlueprintLibrary::EvaluateRootBoneTimecodeAttributesAtTime(Section.Params.Animation, static_cast<float>(AnimTime), HintFrameTime))
-	{
-		const FFrameTime FrameTime = SamplingFrameRate.AsFrameTime(AnimTime);
-		HintFrameTime = FQualifiedFrameTime(FrameTime, SamplingFrameRate);
-	}
-
-	// Convert to tick resolution
-	HintFrameTime = FQualifiedFrameTime(ConvertFrameTime(HintFrameTime.Time, HintFrameTime.Rate, TickResolution), TickResolution);
-
-	// Get the desired frame display format and zero padding from
-	// the sequencer settings, if possible.
-	TAttribute<EFrameNumberDisplayFormats> DisplayFormatAttr(EFrameNumberDisplayFormats::Frames);
-	TAttribute<uint8> ZeroPadFrameNumbersAttr(0u);
-	if (const USequencerSettings* SequencerSettings = Sequencer.Pin()->GetSequencerSettings())
-	{
-		DisplayFormatAttr.Set(SequencerSettings->GetTimeDisplayFormat());
-		ZeroPadFrameNumbersAttr.Set(SequencerSettings->GetZeroPadFrames());
-	}
-
-	// No frame rate conversion necessary since we're displaying
-	// the source frame time/rate.
-	const TAttribute<FFrameRate> TickResolutionAttr(HintFrameTime.Rate);
-	const TAttribute<FFrameRate> DisplayRateAttr(HintFrameTime.Rate);
-
-	FFrameNumberInterface FrameNumberInterface(DisplayFormatAttr, ZeroPadFrameNumbersAttr, TickResolutionAttr, DisplayRateAttr);
-
-	float Subframe = 0.0f;
-	if (UAnimationBlueprintLibrary::EvaluateRootBoneTimecodeSubframeAttributeAtTime(Section.Params.Animation, static_cast<float>(AnimTime), Subframe))
-	{
-		if (FMath::IsNearlyEqual(Subframe, FMath::RoundToFloat(Subframe)))
-		{
-			FrameNumberInterface.SetSubframeIndicator(FString::Printf(TEXT(" (%d)"), FMath::RoundToInt(Subframe)));
-		}
-		else
-		{
-			FrameNumberInterface.SetSubframeIndicator(FString::Printf(TEXT(" (%s)"), *LexToSanitizedString(Subframe)));
-		}
-	}
-
-	return HintFrameTime.Time;
-}
-
-
-float FSkeletalAnimationSection::GetSectionHeight(const UE::Sequencer::FViewDensityInfo& ViewDensity) const
-{
-	return ViewDensity.UniformHeight.Get(SkeletalAnimationEditorConstants::AnimationTrackHeight);
-}
-
-
-FMargin FSkeletalAnimationSection::GetContentPadding() const
-{
-	return FMargin(8.0f, 8.0f);
-}
-
-
-int32 FSkeletalAnimationSection::OnPaintSection( FSequencerSectionPainter& Painter ) const
-{
-	using namespace UE::MovieScene;
-
-	const ESlateDrawEffect DrawEffects = Painter.bParentEnabled ? ESlateDrawEffect::None : ESlateDrawEffect::DisabledEffect;
-	
-	const FTimeToPixel& TimeToPixelConverter = Painter.GetTimeConverter();
-
-	int32 LayerId = Painter.PaintSectionBackground();
-
-	if (!Section.HasStartFrame() || !Section.HasEndFrame())
-	{
-		return LayerId;
-	}
-
-	FFrameRate TickResolution = TimeToPixelConverter.GetTickResolution();
-
-	const FFrameNumber StartFrame = Section.GetInclusiveStartFrame();
-	const FFrameNumber EndFrame   = Section.GetExclusiveEndFrame();
-
-	if (UMovieScene* MovieScene = Section.GetTypedOuter<UMovieScene>())
-	{
-		FFrameRate FrameRate = MovieScene->GetTickResolution();
-		FMovieSceneSequenceTransform OuterToInnerTransform = Section.Params.MakeTransform(FrameRate, Section.GetRange());
-
-		// As seconds represented as a FFrameTime
-		FFrameTime LoopStart = FFrameTime::FromDecimal(Section.Params.StartFrameOffset / FrameRate);
-
-		FLinearColor SectionTint = Painter.GetSectionColor().LinearRGBToHSV();
-		SectionTint.B *= 0.1f;
-		SectionTint = SectionTint.HSVToLinearRGB();
-
-		auto PaintTime = [&](FFrameTime Time)
-		{
-			float OffsetPixel = TimeToPixelConverter.FrameToPixel(Time);
-
-			TArray<FVector2f> NewVector;
-			NewVector.Reserve(2);
-
-			NewVector.Add(FVector2f(OffsetPixel, 1.f));
-			NewVector.Add(FVector2f(OffsetPixel, Painter.SectionGeometry.Size.Y-2.f));
-
-			constexpr float Thickness = 1.f;
-			constexpr float DashLengthPx = 3.f;
-			FSlateDrawElement::MakeDashedLines(
-				Painter.DrawElements,
-				Painter.LayerId++,
-				Painter.SectionGeometry.ToPaintGeometry(),
-				MoveTemp(NewVector),
-				DrawEffects,
-				SectionTint,
-				Thickness,
-				DashLengthPx
-			);
-			return true;
-		};
-
-		OuterToInnerTransform.ExtractBoundariesWithinRange(StartFrame, EndFrame, [&PaintTime](FFrameTime StartTime){
-			PaintTime(StartTime);
-			return true;
-		});
-	}
-
-	return LayerId;
-}
-
-void FSkeletalAnimationSection::BeginResizeSection()
-{
-	const FFrameRate FrameRate = Section.GetTypedOuter<UMovieScene>()->GetTickResolution();
-	InitialDragTransform = MakeUnique<FMovieSceneSequenceTransform>(Section.Params.MakeTransform(FrameRate, Section.GetRange()));
-}
-
-void FSkeletalAnimationSection::ResizeSection(ESequencerSectionResizeMode ResizeMode, FFrameNumber ResizeTime)
-{
-	if (ResizeMode == SSRM_LeadingEdge && Section.Params.PlayRate.GetType() != EMovieSceneTimeWarpType::Custom)
-	{
-		const FFrameRate FrameRate = Section.GetTypedOuter<UMovieScene>()->GetTickResolution();
-		Section.Params.FirstLoopStartFrameOffset = (InitialDragTransform->TransformTime(ResizeTime).AsDecimal() * FrameRate).RoundToFrame();
-	}
-
-	ISequencerSection::ResizeSection(ResizeMode, ResizeTime);
-}
-
-void FSkeletalAnimationSection::BeginSlipSection()
-{
-	BeginResizeSection();
-}
-
-void FSkeletalAnimationSection::SlipSection(FFrameNumber SlipTime)
-{
-	const FFrameRate FrameRate = Section.GetTypedOuter<UMovieScene>()->GetTickResolution();
-	Section.Params.FirstLoopStartFrameOffset = (InitialDragTransform->TransformTime(SlipTime).AsDecimal() * FrameRate).RoundToFrame();
-	ISequencerSection::SlipSection(SlipTime);
-}
-
-bool FSkeletalAnimationSection::RequestDeleteKeyArea(const TArray<FName>& KeyAreaNamePath)
-{
-	Section.Modify();
-	Section.DeleteChannels(KeyAreaNamePath);
-	return true;
-}
-
-void FSkeletalAnimationSection::CustomizePropertiesDetailsView(TSharedRef<IDetailsView> DetailsView, const FSequencerSectionPropertyDetailsViewCustomizationParams& InParams) const
-{
-	DetailsView->RegisterInstancedCustomPropertyTypeLayout(
-		TEXT("MovieSceneSkeletalAnimationParams"),
-		FOnGetPropertyTypeCustomizationInstance::CreateLambda([=]() { return MakeShared<FMovieSceneSkeletalAnimationParamsDetailCustomization>(InParams); }));
-}
-
-void FSkeletalAnimationSection::BuildSectionContextMenu(FMenuBuilder& MenuBuilder, const FGuid& ObjectBinding)
-{
-	const TSharedPtr<ISequencer> SequencerPtr = Sequencer.Pin();
-	if (!SequencerPtr.IsValid())
-	{
-		return;
-	}
-
-	UMovieSceneSkeletalAnimationTrack* const Track = Section.GetTypedOuter<UMovieSceneSkeletalAnimationTrack>();
-	if (!IsValid(Track))
-	{
-		return;
-	}
-
-	// Can't pick the object that this track binds
-	USkeleton* const Skeleton = AcquireSkeletonFromObjectGuid(ObjectBinding, SequencerPtr);
-	if (!IsValid(Skeleton))
-	{
-		return;
-	}
-
-	const int32 NumBones = Skeleton->GetReferenceSkeleton().GetNum();
-	TArray<FName> BoneNames;
-	for (int32 BoneIndex = 0; BoneIndex < NumBones; ++BoneIndex)
-	{
-		BoneNames.Add(Skeleton->GetReferenceSkeleton().GetBoneName(BoneIndex));
-	}
-
-	auto MatchToBone = [this, ObjectBinding](bool bMatchPrevious, int32 Index)
+	auto MatchToBone = [this, ObjectBinding, Skeleton](bool bMatchPrevious, int32 Index)
 	{
 		return FUIAction(
-			FExecuteAction::CreateLambda([this, ObjectBinding, bMatchPrevious, Index]
+			FExecuteAction::CreateLambda([this, ObjectBinding, Skeleton, bMatchPrevious, Index]
 				{
 					const TSharedPtr<ISequencer> SequencerPtr = Sequencer.Pin();
 					if (!SequencerPtr.IsValid())
@@ -973,19 +274,12 @@ void FSkeletalAnimationSection::BuildSectionContextMenu(FMenuBuilder& MenuBuilde
 						return;
 					}
 
-					// Can't pick the object that this track binds
-					USkeleton* const Skeleton = AcquireSkeletonFromObjectGuid(ObjectBinding, SequencerPtr);
-					if (!IsValid(Skeleton))
-					{
-						return;
-					}
-
 					FScopedTransaction MatchSection(LOCTEXT("MatchSectionByBone_Transaction", "Match Section By Bone"));
 					Section.Modify();	
 					Section.bMatchWithPrevious = bMatchPrevious;
 					if (Index >= 0)
 					{
-						USkeletalMeshComponent* const SkelMeshComp = AcquireSkeletalMeshFromObjectGuid(ObjectBinding, SequencerPtr);
+						USkeletalMeshComponent* const SkelMeshComp = UE::Sequencer::FCommonAnimationTrackEditor::AcquireSkeletalMeshFromObjectGuid(ObjectBinding, SequencerPtr);
 
 						FName Name = Skeleton->GetReferenceSkeleton().GetBoneName(Index);
 						Section.MatchSectionByBoneTransform(SkelMeshComp, SequencerPtr->GetLocalTime().Time, SequencerPtr->GetLocalTime().Rate, Name);
@@ -1010,7 +304,7 @@ void FSkeletalAnimationSection::BuildSectionContextMenu(FMenuBuilder& MenuBuilde
 						return false;
 					}
 
-					USkeleton* const Skeleton = AcquireSkeletonFromObjectGuid(ObjectBinding, SequencerPtr);
+					USkeleton* const Skeleton = UE::Sequencer::FCommonAnimationTrackEditor::AcquireSkeletonFromObjectGuid(ObjectBinding, SequencerPtr);
 					if (!IsValid(Skeleton))
 					{
 						return false;
@@ -1087,7 +381,7 @@ void FSkeletalAnimationSection::BuildSectionContextMenu(FMenuBuilder& MenuBuilde
 							return;
 						}
 
-						USkeletalMeshComponent* const SkelMeshComp = AcquireSkeletalMeshFromObjectGuid(ObjectBinding, SequencerPtr);
+						USkeletalMeshComponent* const SkelMeshComp = UE::Sequencer::FCommonAnimationTrackEditor::AcquireSkeletalMeshFromObjectGuid(ObjectBinding, SequencerPtr);
 
 						FScopedTransaction MatchTransaction(LOCTEXT("MatchTranslation_Transaction", "Match Translation"));
 						Section.Modify();	
@@ -1119,7 +413,7 @@ void FSkeletalAnimationSection::BuildSectionContextMenu(FMenuBuilder& MenuBuilde
 							return;
 						}
 
-						USkeletalMeshComponent* const SkelMeshComp = AcquireSkeletalMeshFromObjectGuid(ObjectBinding, SequencerPtr);
+						USkeletalMeshComponent* const SkelMeshComp = UE::Sequencer::FCommonAnimationTrackEditor::AcquireSkeletalMeshFromObjectGuid(ObjectBinding, SequencerPtr);
 
 						FScopedTransaction MatchTransaction(LOCTEXT("MatchZHeight_Transaction", "Match Z Height"));
 						Section.Modify();
@@ -1152,7 +446,7 @@ void FSkeletalAnimationSection::BuildSectionContextMenu(FMenuBuilder& MenuBuilde
 							return;
 						}
 
-						USkeletalMeshComponent* const SkelMeshComp = AcquireSkeletalMeshFromObjectGuid(ObjectBinding, SequencerPtr);
+						USkeletalMeshComponent* const SkelMeshComp = UE::Sequencer::FCommonAnimationTrackEditor::AcquireSkeletalMeshFromObjectGuid(ObjectBinding, SequencerPtr);
 
 						FScopedTransaction MatchTransaction(LOCTEXT("MatchYawRotation_Transaction", "Match Yaw Rotation"));
 						Section.Modify();
@@ -1183,7 +477,7 @@ void FSkeletalAnimationSection::BuildSectionContextMenu(FMenuBuilder& MenuBuilde
 							return;
 						}
 
-						USkeletalMeshComponent* const SkelMeshComp = AcquireSkeletalMeshFromObjectGuid(ObjectBinding, SequencerPtr);
+						USkeletalMeshComponent* const SkelMeshComp = UE::Sequencer::FCommonAnimationTrackEditor::AcquireSkeletalMeshFromObjectGuid(ObjectBinding, SequencerPtr);
 
 						FScopedTransaction MatchTransaction(LOCTEXT("MatchPitchRotation_Transaction", "Match Pitch Rotation"));
 						Section.Modify();
@@ -1215,7 +509,7 @@ void FSkeletalAnimationSection::BuildSectionContextMenu(FMenuBuilder& MenuBuilde
 							return;
 						}
 
-						USkeletalMeshComponent* const SkelMeshComp = AcquireSkeletalMeshFromObjectGuid(ObjectBinding, SequencerPtr);
+						USkeletalMeshComponent* const SkelMeshComp = UE::Sequencer::FCommonAnimationTrackEditor::AcquireSkeletalMeshFromObjectGuid(ObjectBinding, SequencerPtr);
 
 						FScopedTransaction MatchTransaction(LOCTEXT("MatchRollRotation_Transaction", "Match Roll Rotation"));
 						Section.Modify();
@@ -1268,1306 +562,15 @@ void FSkeletalAnimationSection::BuildSectionContextMenu(FMenuBuilder& MenuBuilde
 	MenuBuilder.EndSection();
 }
 
-
-////FSkeletalAnimationTrackEditor
-
-
-bool FSkeletalAnimationTrackEditor::CreatePoseAsset(const TArray<UObject*> NewAssets, FGuid InObjectBinding)
-{
-	USkeletalMeshComponent* SkeletalMeshComponent = AcquireSkeletalMeshFromObjectGuid(InObjectBinding, GetSequencer());
-
-	bool bResult = false;
-	if (NewAssets.Num() > 0)
-	{
-		for (auto NewAsset : NewAssets)
-		{
-			UPoseAsset* NewPoseAsset = Cast<UPoseAsset>(NewAsset);
-			if (NewPoseAsset)
-			{
-				NewPoseAsset->AddPoseWithUniqueName(SkeletalMeshComponent);
-				bResult = true;
-			}
-		}
-
-		// if it contains error, warn them
-		if (bResult)
-		{				
-			FText NotificationText;
-			if (NewAssets.Num() == 1)
-			{
-				NotificationText = FText::Format(LOCTEXT("NumPoseAssetsCreated", "{0} Pose assets created."), NewAssets.Num());
-			}
-			else
-			{
-				NotificationText = FText::Format(LOCTEXT("PoseAssetsCreated", "Pose asset created: '{0}'."), FText::FromString(NewAssets[0]->GetName()));
-			}
-						
-			FNotificationInfo Info(NotificationText);	
-			Info.ExpireDuration = 8.0f;	
-			Info.bUseLargeFont = false;
-			Info.Hyperlink = FSimpleDelegate::CreateLambda([NewAssets]()
-			{
-				GEditor->GetEditorSubsystem<UAssetEditorSubsystem>()->OpenEditorForAssets(NewAssets);
-			});
-			Info.HyperlinkText = FText::Format(LOCTEXT("OpenNewPoseAssetHyperlink", "Open {0}"), FText::FromString(NewAssets[0]->GetName()));
-				
-			TSharedPtr<SNotificationItem> Notification = FSlateNotificationManager::Get().AddNotification(Info);
-			if ( Notification.IsValid() )
-			{
-				Notification->SetCompletionState( SNotificationItem::CS_Success );
-			}
-		}
-		else
-		{
-			FMessageDialog::Open(EAppMsgType::Ok, LOCTEXT("FailedToCreateAsset", "Failed to create asset"));
-		}
-	}
-	return bResult;
-}
-
-
-void FSkeletalAnimationTrackEditor::HandleCreatePoseAsset(FGuid InObjectBinding)
-{
-	USkeleton* Skeleton = AcquireSkeletonFromObjectGuid(InObjectBinding, GetSequencer());
-	if (Skeleton)
-	{
-		TArray<TSoftObjectPtr<UObject>> Skeletons;
-		Skeletons.Add(Skeleton);
-		AnimationEditorUtils::ExecuteNewAnimAsset<UPoseAssetFactory, UPoseAsset>(Skeletons, FString("_PoseAsset"), FAnimAssetCreated::CreateSP(this, &FSkeletalAnimationTrackEditor::CreatePoseAsset, InObjectBinding), false, false);
-	}
-}
-
-bool FSkeletalAnimationTrackEditor::CanCreatePoseAsset(FGuid InObjectBinding) const
+void FSkeletalAnimationTrackEditor::BuildTrackContextMenu(FMenuBuilder& MenuBuilder, UMovieSceneTrack* Track)
 {
-	FClassViewerModule& ClassViewerModule = FModuleManager::LoadModuleChecked<FClassViewerModule>("ClassViewer");
-	const TSharedPtr<IClassViewerFilter>& GlobalClassFilter = ClassViewerModule.GetGlobalClassViewerFilter();
-	TSharedRef<FClassViewerFilterFuncs> ClassFilterFuncs = ClassViewerModule.CreateFilterFuncs();
-	FClassViewerInitializationOptions ClassViewerOptions = {};
-
-	if (GlobalClassFilter.IsValid())
-	{
-		return GlobalClassFilter->IsClassAllowed(ClassViewerOptions, UPoseAsset::StaticClass(), ClassFilterFuncs);
-	}
-
-	return true;
-}
-
-FSkeletalAnimationTrackEditor::FSkeletalAnimationTrackEditor( TSharedRef<ISequencer> InSequencer )
-	: FMovieSceneTrackEditor( InSequencer ) 
-{ 
-	//We use the FGCObject pattern to keep the anim export option alive during the editor session
-
-	AnimSeqExportOption = NewObject<UAnimSeqExportOption>();
-
+	BuildTrackContextMenu_Internal(MenuBuilder, Track, true);
 }
 
-void FSkeletalAnimationTrackEditor::OnInitialize()
-{
-	SequencerSavedHandle = GetSequencer()->OnPostSave().AddRaw(this, &FSkeletalAnimationTrackEditor::OnSequencerSaved);
-	SequencerChangedHandle = GetSequencer()->OnMovieSceneDataChanged().AddRaw(this, &FSkeletalAnimationTrackEditor::OnSequencerDataChanged);
-	FCoreUObjectDelegates::OnObjectPropertyChanged.AddRaw(this, &FSkeletalAnimationTrackEditor::OnPostPropertyChanged);
-
-	++FSkeletalAnimationTrackEditor::NumberActive;
-
-	// Activate the default mode in case FEditorModeTools::Tick isn't run before here. 
-	// This can be removed once a general fix for UE-143791 has been implemented.
-	GLevelEditorModeTools().ActivateDefaultMode();
-
-	GLevelEditorModeTools().ActivateMode(FSkeletalAnimationTrackEditMode::ModeName);
-	FSkeletalAnimationTrackEditMode* EditMode = static_cast<FSkeletalAnimationTrackEditMode*>(GLevelEditorModeTools().GetActiveMode(FSkeletalAnimationTrackEditMode::ModeName));
-	if (EditMode)
-	{
-		EditMode->SetSequencer(GetSequencer());
-	}
-
-}
-void FSkeletalAnimationTrackEditor::OnRelease()
-{
-	FCoreUObjectDelegates::OnObjectPropertyChanged.RemoveAll(this);
-
-	--FSkeletalAnimationTrackEditor::NumberActive;
-
-	if (GetSequencer().IsValid())
-	{
-		if (SequencerSavedHandle.IsValid())
-		{
-			GetSequencer()->OnPostSave().Remove(SequencerSavedHandle);
-			SequencerSavedHandle.Reset();
-		}
-		if (SequencerChangedHandle.IsValid())
-		{
-			GetSequencer()->OnMovieSceneDataChanged().Remove(SequencerChangedHandle);
-			SequencerChangedHandle.Reset();
-		}
-	}
-	if (FSkeletalAnimationTrackEditor::NumberActive == 0)
-	{
-		GLevelEditorModeTools().DeactivateMode(FSkeletalAnimationTrackEditMode::ModeName);
-	}
-
-}
-
-void FSkeletalAnimationTrackEditor::AddReferencedObjects(FReferenceCollector& Collector)
-{
-	if (AnimSeqExportOption != nullptr)
-	{
-		Collector.AddReferencedObject(AnimSeqExportOption);
-	}
-}
-
-
-TSharedRef<ISequencerTrackEditor> FSkeletalAnimationTrackEditor::CreateTrackEditor( TSharedRef<ISequencer> InSequencer )
-{
-	return MakeShareable( new FSkeletalAnimationTrackEditor( InSequencer ) );
-}
-
-bool FSkeletalAnimationTrackEditor::SupportsSequence(UMovieSceneSequence* InSequence) const
-{
-	ETrackSupport TrackSupported = InSequence ? InSequence->IsTrackSupported(UMovieSceneSkeletalAnimationTrack::StaticClass()) : ETrackSupport::NotSupported;
-	return TrackSupported == ETrackSupport::Supported;
-}
-
-bool FSkeletalAnimationTrackEditor::SupportsType( TSubclassOf<UMovieSceneTrack> Type ) const
-{
-	return Type == UMovieSceneSkeletalAnimationTrack::StaticClass();
-}
-
-
-TSharedRef<ISequencerSection> FSkeletalAnimationTrackEditor::MakeSectionInterface( UMovieSceneSection& SectionObject, UMovieSceneTrack& Track, FGuid ObjectBinding )
-{
-	check( SupportsType( SectionObject.GetOuter()->GetClass() ) );
-	
-	return MakeShareable( new FSkeletalAnimationSection(SectionObject, GetSequencer()) );
-}
-
-
-bool FSkeletalAnimationTrackEditor::HandleAssetAdded(UObject* Asset, const FGuid& TargetObjectGuid)
-{
-	TSharedPtr<ISequencer> SequencerPtr = GetSequencer();
-
-	if (Asset->IsA<UAnimSequenceBase>() && SequencerPtr.IsValid())
-	{
-		UAnimSequenceBase* AnimSequence = Cast<UAnimSequenceBase>(Asset);
-		
-		if (TargetObjectGuid.IsValid() && AnimSequence->CanBeUsedInComposition())
-		{
-			USkeleton* Skeleton = AcquireSkeletonFromObjectGuid(TargetObjectGuid, GetSequencer());
-
-			if (Skeleton && Skeleton->IsCompatibleForEditor(AnimSequence->GetSkeleton()))
-			{
-				UObject* Object = SequencerPtr->FindSpawnedObjectOrTemplate(TargetObjectGuid);
-				
-				UMovieSceneTrack* Track = nullptr;
-
-				const FScopedTransaction Transaction(LOCTEXT("AddAnimation_Transaction", "Add Animation"));
-
-				int32 RowIndex = INDEX_NONE;
-				AnimatablePropertyChanged(FOnKeyProperty::CreateRaw(this, &FSkeletalAnimationTrackEditor::AddKeyInternal, Object, AnimSequence, Track, RowIndex));
-
-				return true;
-			}
-		}
-	}
-	return false;
-}
-
-FText FSkeletalAnimationTrackEditor::GetDisplayName() const
-{
-	return LOCTEXT("SkeletalAnimationTrackEditor_DisplayName", "Skeletal Animation");
-}
-
-void FSkeletalAnimationTrackEditor::BuildObjectBindingContextMenu(FMenuBuilder& MenuBuilder, const TArray<FGuid>& ObjectBindings, const UClass* ObjectClass)
-{
-	if (ObjectClass->IsChildOf(USkeletalMeshComponent::StaticClass()) || ObjectClass->IsChildOf(AActor::StaticClass()) || ObjectClass->IsChildOf(UChildActorComponent::StaticClass()))
-	{
-		ConstructObjectBindingTrackMenu(MenuBuilder, ObjectBindings);
-	}
-}
-void FSkeletalAnimationTrackEditor::OnSequencerSaved(ISequencer& )
-{
-	TSharedPtr<ISequencer> SequencerPtr = GetSequencer();
-	if (!SequencerPtr.IsValid())
-	{
-		return;
-	}
-	ULevelSequence* LevelSequence = Cast<ULevelSequence>(SequencerPtr->GetFocusedMovieSceneSequence());
-	if (LevelSequence && LevelSequence->GetClass()->ImplementsInterface(UInterface_AssetUserData::StaticClass()))
-	{
-		if (IInterface_AssetUserData* AssetUserDataInterface = Cast< IInterface_AssetUserData >(LevelSequence))
-		{
-			ULevelSequenceAnimSequenceLink* LevelAnimLink = AssetUserDataInterface->GetAssetUserData< ULevelSequenceAnimSequenceLink >();
-			if (LevelAnimLink && LevelAnimLink->AnimSequenceLinks.Num() > 0)
-			{
-				
-
-				UMovieSceneSequence* MovieSceneSequence = SequencerPtr->GetFocusedMovieSceneSequence();
-				UMovieSceneSequence* RootMovieSceneSequence = SequencerPtr->GetRootMovieSceneSequence();
-				FMovieSceneSequenceIDRef Template = SequencerPtr->GetFocusedTemplateID();
-				FMovieSceneSequenceTransform RootToLocalTransform = SequencerPtr->GetFocusedMovieSceneSequenceTransform();
-				//if in subsquence and we want to, turn on ShouldEvaluateSubSequencesInIsolation
-				TOptional<bool> bRestoreShouldEvaluateSubSequencesInIsolation;
-				if (SkeletalAnimationEditorConstants::bSubsequenceAnimBakeInIsolation &&
-					MovieSceneSequence != RootMovieSceneSequence &&
-					SequencerPtr->GetSequencerSettings()->ShouldEvaluateSubSequencesInIsolation() == false)
-				{
-					bRestoreShouldEvaluateSubSequencesInIsolation = false;
-					SequencerPtr->GetSequencerSettings()->SetEvaluateSubSequencesInIsolation(true);
-				}
-				for (int32 Index = LevelAnimLink->AnimSequenceLinks.Num() -1; Index >=0 ; --Index)
-				{
-					FLevelSequenceAnimSequenceLinkItem& Item = LevelAnimLink->AnimSequenceLinks[Index];
-					UAnimSequence* AnimSequence = Item.ResolveAnimSequence();
-					if (AnimSequence == nullptr)
-					{
-						LevelAnimLink->AnimSequenceLinks.RemoveAt(Index);
-						continue;
-					}
-					if (IInterface_AssetUserData* AnimAssetUserData = Cast< IInterface_AssetUserData >(AnimSequence))
-					{
-						UAnimSequenceLevelSequenceLink* AnimLevelLink = AnimAssetUserData->GetAssetUserData< UAnimSequenceLevelSequenceLink >();
-						if (!AnimLevelLink)
-						{
-							AnimLevelLink = NewObject<UAnimSequenceLevelSequenceLink>(AnimSequence, NAME_None, RF_Public | RF_Transactional);
-							AnimAssetUserData->AddAssetUserData(AnimLevelLink);
-						}
-						AnimLevelLink->SetLevelSequence(LevelSequence);
-						AnimLevelLink->SkelTrackGuid = Item.SkelTrackGuid;
-					}
-					USkeletalMeshComponent* SkelMeshComp = AcquireSkeletalMeshFromObjectGuid(Item.SkelTrackGuid, GetSequencer());
-					if (AnimSequence && SkelMeshComp)
-					{
-						const bool bSavedExportMorphTargets = AnimSeqExportOption->bExportMorphTargets;
-						const bool bSavedExportAttributeCurves = AnimSeqExportOption->bExportAttributeCurves;
-						const bool bSavedExportMaterialCurves = AnimSeqExportOption->bExportMaterialCurves;
-						const bool bSavedExportTransforms = AnimSeqExportOption->bExportTransforms;
-						const bool bSavedIncludeComponentTransform = AnimSeqExportOption->bRecordInWorldSpace;
-						const bool bSavedEvaluateAllSkeletalMeshComponents = AnimSeqExportOption->bEvaluateAllSkeletalMeshComponents;
-						const EAnimInterpolationType SavedInterpolationType = AnimSeqExportOption->Interpolation;
-						const ERichCurveInterpMode SavedCurveInterpolationType = AnimSeqExportOption->CurveInterpolation;
-						const TArray<FString> SavedIncludeAnimationNames = AnimSeqExportOption->IncludeAnimationNames;
-						const TArray<FString> SavedExcludeAnimationNames = AnimSeqExportOption->ExcludeAnimationNames;
-						const FFrameNumber SavedWarmUpFrames = AnimSeqExportOption->WarmUpFrames;
-						const FFrameNumber SavedDelayBeforeStart = AnimSeqExportOption->DelayBeforeStart;
-						const bool bSavedUseCustomTimeRange = AnimSeqExportOption->bUseCustomTimeRange;
-						const FFrameNumber SavedCustomStartFrame = AnimSeqExportOption->CustomStartFrame;
-						const FFrameNumber SavedCustomEndFrame = AnimSeqExportOption->CustomEndFrame;
-						const FFrameRate SavedCustomDisplayRate = AnimSeqExportOption->CustomDisplayRate;
-						const bool bSavedUseCustomFrameRate = AnimSeqExportOption->bUseCustomFrameRate;
-						const FFrameRate SavedCustomFrameRate = AnimSeqExportOption->CustomFrameRate;
-
-						AnimSeqExportOption->bExportMorphTargets = Item.bExportMorphTargets;
-						AnimSeqExportOption->bExportAttributeCurves = Item.bExportAttributeCurves;
-						AnimSeqExportOption->bExportMaterialCurves = Item.bExportMaterialCurves;
-						AnimSeqExportOption->bExportTransforms = Item.bExportTransforms;
-						AnimSeqExportOption->bRecordInWorldSpace = Item.bRecordInWorldSpace;
-						AnimSeqExportOption->bEvaluateAllSkeletalMeshComponents = Item.bEvaluateAllSkeletalMeshComponents;
-						AnimSeqExportOption->Interpolation = Item.Interpolation;
-						AnimSeqExportOption->CurveInterpolation = Item.CurveInterpolation;
-
-						AnimSeqExportOption->IncludeAnimationNames = Item.IncludeAnimationNames;
-						AnimSeqExportOption->ExcludeAnimationNames = Item.ExcludeAnimationNames;
-						AnimSeqExportOption->WarmUpFrames = Item.WarmUpFrames;
-						AnimSeqExportOption->DelayBeforeStart = Item.DelayBeforeStart;
-						AnimSeqExportOption->bUseCustomTimeRange = Item.bUseCustomTimeRange;
-						AnimSeqExportOption->CustomStartFrame = Item.CustomStartFrame;
-						AnimSeqExportOption->CustomEndFrame = Item.CustomEndFrame;
-						AnimSeqExportOption->CustomDisplayRate = Item.CustomDisplayRate;
-						AnimSeqExportOption->bUseCustomFrameRate = Item.bUseCustomFrameRate;
-						AnimSeqExportOption->CustomFrameRate = Item.CustomFrameRate;
-
-						FAnimExportSequenceParameters AESP;
-						AESP.Player = SequencerPtr.Get();
-						AESP.RootToLocalTransform = RootToLocalTransform;
-						AESP.MovieSceneSequence = MovieSceneSequence;
-						AESP.RootMovieSceneSequence = RootMovieSceneSequence;
-						AESP.bForceUseOfMovieScenePlaybackRange = SequencerPtr.Get()->GetSequencerSettings()->ShouldEvaluateSubSequencesInIsolation();
-						//make sure all spanwables are present
-						SequencerPtr->ForceEvaluate();
-						bool bResult = MovieSceneToolHelpers::ExportToAnimSequence(AnimSequence, AnimSeqExportOption, AESP, SkelMeshComp);
-
-						AnimSeqExportOption->bExportMorphTargets = bSavedExportMorphTargets;
-						AnimSeqExportOption->bExportAttributeCurves = bSavedExportAttributeCurves;
-						AnimSeqExportOption->bExportMaterialCurves = bSavedExportMaterialCurves;
-						AnimSeqExportOption->bExportTransforms = bSavedExportTransforms;
-						AnimSeqExportOption->bRecordInWorldSpace = bSavedIncludeComponentTransform;
-						AnimSeqExportOption->bEvaluateAllSkeletalMeshComponents = bSavedEvaluateAllSkeletalMeshComponents;
-						AnimSeqExportOption->Interpolation = SavedInterpolationType;
-						AnimSeqExportOption->CurveInterpolation = SavedCurveInterpolationType;
-
-						AnimSeqExportOption->IncludeAnimationNames = SavedIncludeAnimationNames;
-						AnimSeqExportOption->ExcludeAnimationNames = SavedExcludeAnimationNames;
-						AnimSeqExportOption->WarmUpFrames = SavedWarmUpFrames;
-						AnimSeqExportOption->DelayBeforeStart = SavedDelayBeforeStart;
-						AnimSeqExportOption->bUseCustomTimeRange = bSavedUseCustomTimeRange;
-						AnimSeqExportOption->CustomStartFrame = SavedCustomStartFrame;
-						AnimSeqExportOption->CustomEndFrame = SavedCustomEndFrame;
-						AnimSeqExportOption->CustomDisplayRate = SavedCustomDisplayRate;
-						AnimSeqExportOption->bUseCustomFrameRate = bSavedUseCustomFrameRate;
-						AnimSeqExportOption->CustomFrameRate = SavedCustomFrameRate;
-						//save the anim sequence to disk to make sure they are in sync
-						UPackage* const Package = AnimSequence->GetOutermost();
-						FString const PackageName = Package->GetName();
-						FString const PackageFileName = FPackageName::LongPackageNameToFilename(PackageName, FPackageName::GetAssetPackageExtension());
-
-						FSavePackageArgs SaveArgs;
-						SaveArgs.TopLevelFlags = RF_Standalone;
-						SaveArgs.SaveFlags = SAVE_NoError;
-						UPackage::SavePackage(Package, NULL, *PackageFileName, SaveArgs);
-					}
-				}
-				if (bRestoreShouldEvaluateSubSequencesInIsolation.IsSet())
-				{
-					SequencerPtr->GetSequencerSettings()->SetEvaluateSubSequencesInIsolation(false);
-				}
-				//re-evaluate at current frame
-				SequencerPtr->ForceEvaluate();
-			}
-		}
-	}
-}
-
-//dirty anim sequence when the sequencer changes, to make sure it get's checked out etc..
-void FSkeletalAnimationTrackEditor::OnSequencerDataChanged(EMovieSceneDataChangeType DataChangeType)
-{
-	//only return if data really changed
-	if(DataChangeType ==  EMovieSceneDataChangeType::RefreshTree ||
-		DataChangeType == EMovieSceneDataChangeType::ActiveMovieSceneChanged ||
-		DataChangeType == EMovieSceneDataChangeType::RefreshAllImmediately)
-	{
-		return;
-	}
-	TSharedPtr<ISequencer> SequencerPtr = GetSequencer();
-	if (!SequencerPtr.IsValid())
-	{
-		return;
-	}
-	ULevelSequence* LevelSequence = Cast<ULevelSequence>(SequencerPtr->GetFocusedMovieSceneSequence());
-	if (LevelSequence && LevelSequence->GetClass()->ImplementsInterface(UInterface_AssetUserData::StaticClass()))
-	{
-		if (IInterface_AssetUserData* AssetUserDataInterface = Cast< IInterface_AssetUserData >(LevelSequence))
-		{
-			ULevelSequenceAnimSequenceLink* LevelAnimLink = AssetUserDataInterface->GetAssetUserData< ULevelSequenceAnimSequenceLink >();
-			if (LevelAnimLink)
-			{
-				for (int32 Index = LevelAnimLink->AnimSequenceLinks.Num() - 1; Index >= 0; --Index)
-				{
-					FLevelSequenceAnimSequenceLinkItem& Item = LevelAnimLink->AnimSequenceLinks[Index];
-					UAnimSequence* AnimSequence = Item.ResolveAnimSequence();
-					if (AnimSequence)
-					{
-						AnimSequence->Modify();
-					}
-				}
-			}
-		}
-	}
-}
-
-void FSkeletalAnimationTrackEditor::OnPostPropertyChanged(UObject* InObject, FPropertyChangedEvent& InPropertyChangedEvent)
-{
-	if (SkeletalAnimationEditorConstants::CVarEvaluateSkeletalMeshOnPropertyChange->GetBool())
-	{
-		const FName PropertyName = (InPropertyChangedEvent.Property != nullptr) ? InPropertyChangedEvent.Property->GetFName() : NAME_None;
-		if (PropertyName == USceneComponent::GetRelativeLocationPropertyName() ||
-			PropertyName == USceneComponent::GetRelativeRotationPropertyName() ||
-			PropertyName == USceneComponent::GetRelativeScale3DPropertyName())
-		{
-			return;
-		}
-
-		// If the object changed has any animation track, notify sequencer to update because animations tick on their own and sequencer needs to evaluate again
-		const bool bCreateIfMissing = false;
-		FFindOrCreateHandleResult HandleResult = FindOrCreateHandleToObject(InObject, bCreateIfMissing);
-		FGuid ObjectHandle = HandleResult.Handle;
-		if (ObjectHandle.IsValid())
-		{
-			FFindOrCreateTrackResult TrackResult = FindOrCreateTrackForObject(ObjectHandle, UMovieSceneSkeletalAnimationTrack::StaticClass(), NAME_None, bCreateIfMissing);
-			if (TrackResult.Track)
-			{
-				GetSequencer()->NotifyMovieSceneDataChanged(EMovieSceneDataChangeType::TrackValueChanged);
-			}
-		}
-	}
-}
-
-bool FSkeletalAnimationTrackEditor::CreateAnimationSequence(const TArray<UObject*> NewAssets, USkeletalMeshComponent* SkelMeshComp, FGuid Binding, bool bCreateSoftLink)
-{
-	bool bResult = false;
-	if (NewAssets.Num() > 0)
-	{
-		UAnimSequence* AnimSequence = Cast<UAnimSequence>(NewAssets[0]);
-		if (AnimSequence)
-		{
-			UObject* NewAsset = NewAssets[0];
-			TSharedPtr<SWindow> ParentWindow;
-			if (FModuleManager::Get().IsModuleLoaded("MainFrame"))
-			{
-				IMainFrameModule& MainFrame = FModuleManager::LoadModuleChecked<IMainFrameModule>("MainFrame");
-				ParentWindow = MainFrame.GetParentWindow();
-			}
-
-			TSharedRef<SWindow> Window = SNew(SWindow)
-				.Title(NSLOCTEXT("UnrealEd", "AnimSeqOpionsTitle", "Animation Sequence Options"))
-				.SizingRule(ESizingRule::UserSized)
-				.AutoCenter(EAutoCenter::PrimaryWorkArea)
-				.ClientSize(FVector2D(500, 445));
-
-			TSharedPtr<SAnimSequenceOptionsWindow> OptionWindow;
-			Window->SetContent
-			(
-				SAssignNew(OptionWindow, SAnimSequenceOptionsWindow)
-				.ExportOptions(AnimSeqExportOption)
-				.WidgetWindow(Window)
-				.FullPath(FText::FromString(NewAssets[0]->GetName()))
-			);
-
-			FSlateApplication::Get().AddModalWindow(Window, ParentWindow, false);
-
-			if (OptionWindow->ShouldExport())
-			{
-				const TSharedPtr<ISequencer> ParentSequencer = GetSequencer();
-				UMovieSceneSequence* MovieSceneSequence = ParentSequencer->GetFocusedMovieSceneSequence();
-				UMovieSceneSequence* RootMovieSceneSequence = ParentSequencer->GetRootMovieSceneSequence();
-				FMovieSceneSequenceTransform RootToLocalTransform = ParentSequencer->GetFocusedMovieSceneSequenceTransform();
-				TOptional<bool> bRestoreShouldEvaluateSubSequencesInIsolation;
-				if (SkeletalAnimationEditorConstants::bSubsequenceAnimBakeInIsolation &&
-					MovieSceneSequence != RootMovieSceneSequence &&
-					ParentSequencer->GetSequencerSettings()->ShouldEvaluateSubSequencesInIsolation() == false)
-				{
-					bRestoreShouldEvaluateSubSequencesInIsolation = false;
-					ParentSequencer->GetSequencerSettings()->SetEvaluateSubSequencesInIsolation(true);
-				}
-				FAnimExportSequenceParameters AESP;
-				AESP.Player = ParentSequencer.Get();
-				AESP.RootToLocalTransform = RootToLocalTransform;
-				AESP.MovieSceneSequence = MovieSceneSequence;
-				AESP.RootMovieSceneSequence = RootMovieSceneSequence;
-				AESP.bForceUseOfMovieScenePlaybackRange = ParentSequencer->GetSequencerSettings()->ShouldEvaluateSubSequencesInIsolation();
-				AnimSeqExportOption->CustomDisplayRate = ParentSequencer->GetFocusedDisplayRate();
-				bResult  = MovieSceneToolHelpers::ExportToAnimSequence(AnimSequence, AnimSeqExportOption, AESP, SkelMeshComp);
-				if (bRestoreShouldEvaluateSubSequencesInIsolation.IsSet())
-				{
-					ParentSequencer->GetSequencerSettings()->SetEvaluateSubSequencesInIsolation(false);
-				}
-				//re-evaluate at current frame
-				ParentSequencer->ForceEvaluate();
-			}
-		}
-
-		if (bResult && bCreateSoftLink)
-		{
-			FScopedTransaction Transaction(LOCTEXT("SaveLinkedAnimation_Transaction", "Save Link Animation"));
-			TSharedPtr<ISequencer> SequencerPtr = GetSequencer();
-			ULevelSequence* LevelSequence = Cast<ULevelSequence>(SequencerPtr->GetFocusedMovieSceneSequence());
-			if (LevelSequence && LevelSequence->GetClass()->ImplementsInterface(UInterface_AssetUserData::StaticClass())
-				&& AnimSequence->GetClass()->ImplementsInterface(UInterface_AssetUserData::StaticClass()))
-			{
-				LevelSequence->Modify();
-				if (IInterface_AssetUserData* AnimAssetUserData = Cast< IInterface_AssetUserData >(AnimSequence))
-				{
-					UAnimSequenceLevelSequenceLink* AnimLevelLink = AnimAssetUserData->GetAssetUserData< UAnimSequenceLevelSequenceLink >();
-					if (!AnimLevelLink)
-					{
-						AnimLevelLink = NewObject<UAnimSequenceLevelSequenceLink>(AnimSequence, NAME_None, RF_Public | RF_Transactional);
-						AnimAssetUserData->AddAssetUserData(AnimLevelLink);
-					}
-					
-					AnimLevelLink->SetLevelSequence(LevelSequence);
-					AnimLevelLink->SkelTrackGuid = Binding;
-				}
-				if (IInterface_AssetUserData* AssetUserDataInterface = Cast< IInterface_AssetUserData >(LevelSequence))
-				{
-					bool bAddItem = true;
-					ULevelSequenceAnimSequenceLink* LevelAnimLink = AssetUserDataInterface->GetAssetUserData< ULevelSequenceAnimSequenceLink >();
-					if (LevelAnimLink)
-					{
-						for (FLevelSequenceAnimSequenceLinkItem& LevelAnimLinkItem : LevelAnimLink->AnimSequenceLinks)
-						{
-							if (LevelAnimLinkItem.IsEqual(Binding, AnimSeqExportOption->bUseCustomTimeRange,
-								AnimSeqExportOption->CustomStartFrame, AnimSeqExportOption->CustomEndFrame, AnimSeqExportOption->CustomDisplayRate,
-								AnimSeqExportOption->bUseCustomFrameRate, AnimSeqExportOption->CustomFrameRate
-								))
-							{
-								bAddItem = false;
-								UAnimSequence* OtherAnimSequence = LevelAnimLinkItem.ResolveAnimSequence();
-								
-								if (OtherAnimSequence != AnimSequence)
-								{
-									if (IInterface_AssetUserData* OtherAnimAssetUserData = Cast< IInterface_AssetUserData >(OtherAnimSequence))
-									{
-										UAnimSequenceLevelSequenceLink* OtherAnimLevelLink = OtherAnimAssetUserData->GetAssetUserData< UAnimSequenceLevelSequenceLink >();
-										if (OtherAnimLevelLink)
-										{
-											OtherAnimAssetUserData->RemoveUserDataOfClass(UAnimSequenceLevelSequenceLink::StaticClass());
-										}
-									}
-								}
-								LevelAnimLinkItem.PathToAnimSequence = FSoftObjectPath(AnimSequence);
-								LevelAnimLinkItem.bExportMorphTargets = AnimSeqExportOption->bExportMorphTargets;
-								LevelAnimLinkItem.bExportAttributeCurves = AnimSeqExportOption->bExportAttributeCurves;
-								LevelAnimLinkItem.bExportMaterialCurves = AnimSeqExportOption->bExportMaterialCurves;
-								LevelAnimLinkItem.bExportTransforms = AnimSeqExportOption->bExportTransforms;
-								LevelAnimLinkItem.bRecordInWorldSpace = AnimSeqExportOption->bRecordInWorldSpace;
-								LevelAnimLinkItem.bEvaluateAllSkeletalMeshComponents = AnimSeqExportOption->bEvaluateAllSkeletalMeshComponents;
-								LevelAnimLinkItem.Interpolation = AnimSeqExportOption->Interpolation;
-								LevelAnimLinkItem.CurveInterpolation = AnimSeqExportOption->CurveInterpolation;
-								LevelAnimLinkItem.IncludeAnimationNames = AnimSeqExportOption->IncludeAnimationNames;
-								LevelAnimLinkItem.ExcludeAnimationNames = AnimSeqExportOption->ExcludeAnimationNames;
-								LevelAnimLinkItem.WarmUpFrames = AnimSeqExportOption->WarmUpFrames;
-								LevelAnimLinkItem.DelayBeforeStart = AnimSeqExportOption->DelayBeforeStart;
-								LevelAnimLinkItem.bUseCustomTimeRange = AnimSeqExportOption->bUseCustomTimeRange;
-								LevelAnimLinkItem.CustomStartFrame = AnimSeqExportOption->CustomStartFrame;
-								LevelAnimLinkItem.CustomEndFrame = AnimSeqExportOption->CustomEndFrame;
-								LevelAnimLinkItem.CustomDisplayRate = AnimSeqExportOption->CustomDisplayRate;
-								LevelAnimLinkItem.bUseCustomFrameRate = AnimSeqExportOption->bUseCustomFrameRate;
-								LevelAnimLinkItem.CustomFrameRate = AnimSeqExportOption->CustomFrameRate;
-								break;
-							}
-						}
-					}
-					else
-					{
-						LevelAnimLink = NewObject<ULevelSequenceAnimSequenceLink>(LevelSequence, NAME_None, RF_Public | RF_Transactional);
-						
-					}
-					if (bAddItem == true)
-					{
-						FLevelSequenceAnimSequenceLinkItem LevelAnimLinkItem;
-						LevelAnimLinkItem.SkelTrackGuid = Binding;
-						LevelAnimLinkItem.PathToAnimSequence = FSoftObjectPath(AnimSequence);
-						LevelAnimLinkItem.bExportMorphTargets = AnimSeqExportOption->bExportMorphTargets;
-						LevelAnimLinkItem.bExportAttributeCurves = AnimSeqExportOption->bExportAttributeCurves;
-						LevelAnimLinkItem.bExportMaterialCurves = AnimSeqExportOption->bExportMaterialCurves;
-						LevelAnimLinkItem.bExportTransforms = AnimSeqExportOption->bExportTransforms;
-						LevelAnimLinkItem.bRecordInWorldSpace = AnimSeqExportOption->bRecordInWorldSpace;
-						LevelAnimLinkItem.bEvaluateAllSkeletalMeshComponents = AnimSeqExportOption->bEvaluateAllSkeletalMeshComponents;
-						LevelAnimLinkItem.Interpolation = AnimSeqExportOption->Interpolation;
-						LevelAnimLinkItem.CurveInterpolation = AnimSeqExportOption->CurveInterpolation;
-						LevelAnimLinkItem.IncludeAnimationNames = AnimSeqExportOption->IncludeAnimationNames;
-						LevelAnimLinkItem.ExcludeAnimationNames = AnimSeqExportOption->ExcludeAnimationNames;
-						LevelAnimLinkItem.WarmUpFrames = AnimSeqExportOption->WarmUpFrames;
-						LevelAnimLinkItem.DelayBeforeStart = AnimSeqExportOption->DelayBeforeStart;
-						LevelAnimLinkItem.bUseCustomTimeRange = AnimSeqExportOption->bUseCustomTimeRange;
-						LevelAnimLinkItem.CustomStartFrame = AnimSeqExportOption->CustomStartFrame;
-						LevelAnimLinkItem.CustomEndFrame = AnimSeqExportOption->CustomEndFrame;
-						LevelAnimLinkItem.CustomDisplayRate = AnimSeqExportOption->CustomDisplayRate;
-						LevelAnimLinkItem.bUseCustomFrameRate = AnimSeqExportOption->bUseCustomFrameRate;
-						LevelAnimLinkItem.CustomFrameRate = AnimSeqExportOption->CustomFrameRate;
-
-						LevelAnimLink->AnimSequenceLinks.Add(LevelAnimLinkItem);
-						AssetUserDataInterface->AddAssetUserData(LevelAnimLink);
-					}
-
-					/*
-					FText Name = SequencerPtr->GetDisplayName(Binding);
-					FString StringName = Name.ToString();
-					if (bAddItem == false)
-					{
-						//ok already had a name added so need to remove the old one..
-						TArray<FString> Strings;
-						StringName.ParseIntoArray(Strings,TEXT(" --> "));
-						if (Strings.Num() > 0)
-						{
-							StringName = Strings[0];
-						}
-					}
-	
-					FString AnimName = AnimSequence->GetName();
-					StringName = StringName + FString(TEXT(" --> ")) + AnimName;
-					SequencerPtr->SetDisplayName(Binding, FText::FromString(StringName));
-					*/
-
-				}
-			}
-		}
-		// if it contains error, warn them
-		if (bResult)
-		{
-			FText NotificationText;
-			if (NewAssets.Num() == 1)
-			{
-				NotificationText = FText::Format(LOCTEXT("NumAnimSequenceAssetsCreated", "{0} Anim Sequence  assets created."), NewAssets.Num());
-			}
-			else
-			{
-				NotificationText = FText::Format(LOCTEXT("AnimSequenceAssetsCreated", "Anim Sequence asset created: '{0}'."), FText::FromString(NewAssets[0]->GetName()));
-			}
-
-			FNotificationInfo Info(NotificationText);
-			Info.ExpireDuration = 8.0f;
-			Info.bUseLargeFont = false;
-			Info.Hyperlink = FSimpleDelegate::CreateLambda([NewAssets]()
-			{
-				GEditor->GetEditorSubsystem<UAssetEditorSubsystem>()->OpenEditorForAssets(NewAssets);
-			});
-			Info.HyperlinkText = FText::Format(LOCTEXT("OpenNewPoseAssetHyperlink", "Open {0}"), FText::FromString(NewAssets[0]->GetName()));
-
-			TSharedPtr<SNotificationItem> Notification = FSlateNotificationManager::Get().AddNotification(Info);
-			if (Notification.IsValid())
-			{
-				Notification->SetCompletionState(SNotificationItem::CS_Success);
-			}
-			if (const TSharedPtr<ISequencer> ParentSequencer = GetSequencer())
-			{
-				ParentSequencer->RequestEvaluate(); 
-			}
-		}
-		else
-		{
-			FMessageDialog::Open(EAppMsgType::Ok, LOCTEXT("FailedToCreateAsset", "Failed to create asset"));
-		}
-	}
-	return bResult;
-}
-
-void FSkeletalAnimationTrackEditor::HandleCreateAnimationSequence(USkeletalMeshComponent* SkelMeshComp, USkeleton* Skeleton, FGuid Binding, bool bCreateSoftLink)
-{
-	if (SkelMeshComp)
-	{
-		TArray<TSoftObjectPtr<UObject>> Skels;
-		if (SkelMeshComp->GetSkeletalMeshAsset())
-		{
-			Skels.Add(SkelMeshComp->GetSkeletalMeshAsset());
-		}
-		else
-		{
-			Skels.Add(Skeleton);
-		}
-	
-		const bool bDoNotShowNameDialog = false;
-		const bool bAllowReplaceExisting = true;
-		AnimationEditorUtils::ExecuteNewAnimAsset<UAnimSequenceFactory, UAnimSequence>(Skels, FString("_Sequence"), FAnimAssetCreated::CreateSP(this, &FSkeletalAnimationTrackEditor::CreateAnimationSequence, SkelMeshComp, Binding, bCreateSoftLink), bDoNotShowNameDialog, bAllowReplaceExisting);
-	}
-}
-
-void FSkeletalAnimationTrackEditor::OpenLinkedAnimSequence(FGuid Binding)
-{
-	TSharedPtr<ISequencer> SequencerPtr = GetSequencer();
-	if (!SequencerPtr.IsValid())
-	{
-		return;
-	}
-	ULevelSequence* LevelSequence = Cast<ULevelSequence>(SequencerPtr->GetFocusedMovieSceneSequence());
-	if (LevelSequence && LevelSequence->GetClass()->ImplementsInterface(UInterface_AssetUserData::StaticClass()))
-	{
-		if (IInterface_AssetUserData* AssetUserDataInterface = Cast< IInterface_AssetUserData >(LevelSequence))
-		{
-			ULevelSequenceAnimSequenceLink* LevelAnimLink = AssetUserDataInterface->GetAssetUserData< ULevelSequenceAnimSequenceLink >();
-			if (LevelAnimLink)
-			{
-				
-				for (FLevelSequenceAnimSequenceLinkItem& Item : LevelAnimLink->AnimSequenceLinks)
-				{
-					if (Item.SkelTrackGuid == Binding)
-					{
-						UAnimSequence* AnimSequence = Item.ResolveAnimSequence();
-						if (AnimSequence)
-						{
-							GEditor->GetEditorSubsystem<UAssetEditorSubsystem>()->OpenEditorForAsset(AnimSequence);
-						}
-					}
-				}
-			}
-		}
-	}
-}
-
-bool FSkeletalAnimationTrackEditor::CanOpenLinkedAnimSequence(FGuid Binding)
-{
-
-	TSharedPtr<ISequencer> SequencerPtr = GetSequencer();
-	if (!SequencerPtr.IsValid())
-	{
-		return false;
-	}
-	ULevelSequence* LevelSequence = Cast<ULevelSequence>(SequencerPtr->GetFocusedMovieSceneSequence());
-	if (LevelSequence && LevelSequence->GetClass()->ImplementsInterface(UInterface_AssetUserData::StaticClass()))
-	{
-		if (IInterface_AssetUserData* AssetUserDataInterface = Cast< IInterface_AssetUserData >(LevelSequence))
-		{
-			ULevelSequenceAnimSequenceLink* LevelAnimLink = AssetUserDataInterface->GetAssetUserData< ULevelSequenceAnimSequenceLink >();
-			if (LevelAnimLink)
-			{
-
-				for (FLevelSequenceAnimSequenceLinkItem& Item : LevelAnimLink->AnimSequenceLinks)
-				{
-					if (Item.SkelTrackGuid == Binding)
-					{
-						UAnimSequence* AnimSequence = Item.ResolveAnimSequence();
-						if (AnimSequence)
-						{
-							return true;
-						}
-					}
-				}
-			}
-		}
-	}
-	return false;
-}
-
-void FSkeletalAnimationTrackEditor::ConstructObjectBindingTrackMenu(FMenuBuilder& MenuBuilder, TArray<FGuid> ObjectBindings)
-{
-	if(ObjectBindings.Num() > 0)
-	{
-		USkeletalMeshComponent* SkelMeshComp =  AcquireSkeletalMeshFromObjectGuid(ObjectBindings[0], GetSequencer());
-
-		if (SkelMeshComp)
-		{
-
-			MenuBuilder.BeginSection("Create Animation Assets", LOCTEXT("CreateAnimationAssetsName", "Create Animation Assets"));
-			USkeleton* Skeleton = GetSkeletonFromComponent(SkelMeshComp);
-			//todo do we not link if alreadhy linked???
-
-			MenuBuilder.AddMenuEntry(
-				LOCTEXT("CreateLinkAnimSequence", "Create Linked Animation Sequence"),
-				LOCTEXT("CreateLinkAnimSequenceTooltip", "Create Animation Sequence for this Skeletal Mesh and have this Track Own that Anim Sequence. Note it will create it based upon the Sequencer Display Range and Display Frame Rate"),
-				FSlateIcon(),
-				FUIAction(FExecuteAction::CreateRaw(this, &FSkeletalAnimationTrackEditor::HandleCreateAnimationSequence, SkelMeshComp, Skeleton, ObjectBindings[0], true)),
-				NAME_None,
-				EUserInterfaceActionType::Button);
-
-			MenuBuilder.AddMenuEntry(
-				LOCTEXT("OpenAnimSequence", "Open Linked Animation Sequence"),
-				LOCTEXT("OpenAnimSequenceTooltip", "Open Animation Sequence that this Animation Track is Driving."),
-				FSlateIcon(),
-				FUIAction(
-					FExecuteAction::CreateRaw(this, &FSkeletalAnimationTrackEditor::OpenLinkedAnimSequence, ObjectBindings[0]),
-					FCanExecuteAction::CreateRaw(this, &FSkeletalAnimationTrackEditor::CanOpenLinkedAnimSequence, ObjectBindings[0])
-				),
-				NAME_None,
-				EUserInterfaceActionType::Button);
-
-			MenuBuilder.AddMenuEntry(
-				LOCTEXT("CreateAnimSequence", "Bake Animation Sequence"),
-				LOCTEXT("PasteCreateAnimSequenceTooltip", "Bake an Animation Sequence for this Skeletal Mesh. Note it will create it based upon the Sequencer Display Range and Display Frame Rate"),
-				FSlateIcon(),
-				FUIAction(FExecuteAction::CreateRaw(this, &FSkeletalAnimationTrackEditor::HandleCreateAnimationSequence, SkelMeshComp,Skeleton, ObjectBindings[0], false)),
-				NAME_None,
-				EUserInterfaceActionType::Button);
-
-			if (CanCreatePoseAsset(ObjectBindings[0]))
-			{
-				MenuBuilder.AddMenuEntry(
-					LOCTEXT("CreatePoseAsset", "Bake Pose Asset"),
-					LOCTEXT("CreatePoseAsset_ToolTip", "Bake Animation from current Pose"),
-					FSlateIcon(),
-					FUIAction(FExecuteAction::CreateRaw(this, &FSkeletalAnimationTrackEditor::HandleCreatePoseAsset, ObjectBindings[0])),
-					NAME_None,
-					EUserInterfaceActionType::Button);
-			}
-
-			MenuBuilder.EndSection();
-		}
-	}
-}
-
-void FSkeletalAnimationTrackEditor::BuildObjectBindingTrackMenu(FMenuBuilder& MenuBuilder, const TArray<FGuid>& ObjectBindings, const UClass* ObjectClass)
-{
-	if (ObjectClass->IsChildOf(USkeletalMeshComponent::StaticClass()) || ObjectClass->IsChildOf(AActor::StaticClass()) || ObjectClass->IsChildOf(UChildActorComponent::StaticClass()))
-	{
-		const TSharedPtr<ISequencer> ParentSequencer = GetSequencer();
-
-		USkeleton* Skeleton = AcquireSkeletonFromObjectGuid(ObjectBindings[0], GetSequencer());
-
-		if (Skeleton)
-		{
-			// Load the asset registry module
-			FAssetRegistryModule& AssetRegistryModule = FModuleManager::LoadModuleChecked<FAssetRegistryModule>(TEXT("AssetRegistry"));
-
-			// Collect a full list of assets with the specified class
-			TArray<FAssetData> AssetDataList;
-			AssetRegistryModule.Get().GetAssetsByClass(UAnimSequenceBase::StaticClass()->GetClassPathName(), AssetDataList, true);
-
-			if (AssetDataList.Num())
-			{
-				UMovieSceneTrack* Track = nullptr;
-
-				MenuBuilder.AddSubMenu(
-					LOCTEXT("AddAnimation", "Animation"), NSLOCTEXT("Sequencer", "AddAnimationTooltip", "Adds an animation track."),
-					FNewMenuDelegate::CreateRaw(this, &FSkeletalAnimationTrackEditor::AddAnimationSubMenu, ObjectBindings, Skeleton, Track)
-				);
-			}
-		}
-	}
-}
-
-TSharedRef<SWidget> FSkeletalAnimationTrackEditor::BuildAddAnimationSubMenu(FGuid ObjectBinding, USkeleton* Skeleton, UE::Sequencer::TWeakViewModelPtr<UE::Sequencer::ITrackExtension> WeakTrackModel)
-{
-	FMenuBuilder MenuBuilder(true, nullptr);
-
-	TArray<FGuid> ObjectBindings;
-	ObjectBindings.Add(ObjectBinding);
-
-
-	MenuBuilder.BeginSection(NAME_None, LOCTEXT("TimeWarp_Label", "Time Warp"));
-	{
-		FSequencerUtilities::MakeTimeWarpMenuEntry(MenuBuilder, WeakTrackModel);
-	}
-	MenuBuilder.EndSection();
-
-
-	MenuBuilder.BeginSection(NAME_None, LOCTEXT("AddAnimation_Label", "Add Animation"));
-	{
-		AddAnimationSubMenu(MenuBuilder, ObjectBindings, Skeleton, WeakTrackModel.Pin()->GetTrack());
-	}
-	MenuBuilder.EndSection();
-
-	return MenuBuilder.MakeWidget();
-}
-
-TSharedRef<SWidget> FSkeletalAnimationTrackEditor::BuildAnimationSubMenu(FGuid ObjectBinding, USkeleton* Skeleton, UMovieSceneTrack* Track)
-{
-	FMenuBuilder MenuBuilder(true, nullptr);
-
-	TArray<FGuid> ObjectBindings;
-	ObjectBindings.Add(ObjectBinding);
-
-	AddAnimationSubMenu(MenuBuilder, ObjectBindings, Skeleton, Track);
-
-	return MenuBuilder.MakeWidget();
-}
-
-bool FSkeletalAnimationTrackEditor::ShouldFilterAsset(const FAssetData& AssetData)
-{
-	// we don't want montage
-	if (AssetData.AssetClassPath == UAnimMontage::StaticClass()->GetClassPathName())
-	{
-		return true;
-	}
-
-	const FString EnumString = AssetData.GetTagValueRef<FString>(GET_MEMBER_NAME_CHECKED(UAnimSequence, AdditiveAnimType));
-	if (EnumString.IsEmpty())
-	{
-		return false;
-	}
-
-	UEnum* AdditiveTypeEnum = StaticEnum<EAdditiveAnimationType>();
-	return ((EAdditiveAnimationType)AdditiveTypeEnum->GetValueByName(*EnumString) == AAT_RotationOffsetMeshSpace);
-}
-
-void FSkeletalAnimationTrackEditor::AddAnimationSubMenu(FMenuBuilder& MenuBuilder, TArray<FGuid> ObjectBindings, USkeleton* Skeleton, UMovieSceneTrack* Track)
-{
-	TSharedPtr<ISequencer> SequencerPtr = GetSequencer();
-	UMovieSceneSequence* Sequence = SequencerPtr.IsValid() ? SequencerPtr->GetFocusedMovieSceneSequence() : nullptr;
-
-	FAssetPickerConfig AssetPickerConfig;
-	{
-		AssetPickerConfig.OnAssetSelected = FOnAssetSelected::CreateRaw( this, &FSkeletalAnimationTrackEditor::OnAnimationAssetSelected, ObjectBindings, Track);
-		AssetPickerConfig.OnAssetEnterPressed = FOnAssetEnterPressed::CreateRaw( this, &FSkeletalAnimationTrackEditor::OnAnimationAssetEnterPressed, ObjectBindings, Track);
-		AssetPickerConfig.bAllowNullSelection = false;
-		AssetPickerConfig.bAddFilterUI = true;
-		AssetPickerConfig.bShowTypeInColumnView = false;
-		AssetPickerConfig.InitialAssetViewType = EAssetViewType::List;
-		AssetPickerConfig.Filter.bRecursiveClasses = true;
-		AssetPickerConfig.Filter.ClassPaths.Add(UAnimSequenceBase::StaticClass()->GetClassPathName());
-		AssetPickerConfig.OnShouldFilterAsset.BindRaw(this, &FSkeletalAnimationTrackEditor::FilterAnimSequences, Skeleton);
-		AssetPickerConfig.SaveSettingsName = TEXT("SequencerAssetPicker");
-		AssetPickerConfig.AdditionalReferencingAssets.Add(FAssetData(Sequence));
-	}
-
-	FContentBrowserModule& ContentBrowserModule = FModuleManager::Get().LoadModuleChecked<FContentBrowserModule>(TEXT("ContentBrowser"));
-
-	const float WidthOverride = SequencerPtr.IsValid() ? SequencerPtr->GetSequencerSettings()->GetAssetBrowserWidth() : 500.f;
-	const float HeightOverride = SequencerPtr.IsValid() ? SequencerPtr->GetSequencerSettings()->GetAssetBrowserHeight() : 400.f;
-
-	TSharedPtr<SBox> MenuEntry = SNew(SBox)
-		.WidthOverride(WidthOverride)
-		.HeightOverride(HeightOverride)
-		[
-			ContentBrowserModule.Get().CreateAssetPicker(AssetPickerConfig)
-		];
-
-	MenuBuilder.AddWidget(MenuEntry.ToSharedRef(), FText::GetEmpty(), true);
-}
-
-bool FSkeletalAnimationTrackEditor::FilterAnimSequences(const FAssetData& AssetData, USkeleton* Skeleton)
-{
-	if (ShouldFilterAsset(AssetData))
-	{
-		return true;
-	}
-
-	if (Skeleton && Skeleton->IsCompatibleForEditor(AssetData) == false)
-	{
-		return true;
-	}
-
-	return false;
-}
-
-void FSkeletalAnimationTrackEditor::OnAnimationAssetSelected(const FAssetData& AssetData, TArray<FGuid> ObjectBindings, UMovieSceneTrack* Track)
-{
-	FSlateApplication::Get().DismissAllMenus();
-
-	UObject* SelectedObject = AssetData.GetAsset();
-	TSharedPtr<ISequencer> SequencerPtr = GetSequencer();
-
-	if (SelectedObject && SelectedObject->IsA(UAnimSequenceBase::StaticClass()) && SequencerPtr.IsValid())
-	{
-		UAnimSequenceBase* AnimSequence = CastChecked<UAnimSequenceBase>(AssetData.GetAsset());
-
-		const FScopedTransaction Transaction(LOCTEXT("AddAnimation_Transaction", "Add Animation"));
-
-		for (FGuid ObjectBinding : ObjectBindings)
-		{
-			UObject* Object = SequencerPtr->FindSpawnedObjectOrTemplate(ObjectBinding);
-			int32 RowIndex = INDEX_NONE;
-			AnimatablePropertyChanged(FOnKeyProperty::CreateRaw(this, &FSkeletalAnimationTrackEditor::AddKeyInternal, Object, AnimSequence, Track, RowIndex));
-		}
-	}
-}
-
-void FSkeletalAnimationTrackEditor::OnAnimationAssetEnterPressed(const TArray<FAssetData>& AssetData, TArray<FGuid> ObjectBindings, UMovieSceneTrack* Track)
-{
-	if (AssetData.Num() > 0)
-	{
-		OnAnimationAssetSelected(AssetData[0].GetAsset(), ObjectBindings, Track);
-	}
-}
-
-
-FKeyPropertyResult FSkeletalAnimationTrackEditor::AddKeyInternal( FFrameNumber KeyTime, UObject* Object, class UAnimSequenceBase* AnimSequence, UMovieSceneTrack* Track, int32 RowIndex )
-{
-	FKeyPropertyResult KeyPropertyResult;
-
-	FFindOrCreateHandleResult HandleResult = FindOrCreateHandleToObject( Object );
-	FGuid ObjectHandle = HandleResult.Handle;
-	KeyPropertyResult.bHandleCreated |= HandleResult.bWasCreated;
-	if (ObjectHandle.IsValid())
-	{
-		UMovieScene* MovieScene = GetSequencer()->GetFocusedMovieSceneSequence()->GetMovieScene();
-		UMovieSceneSkeletalAnimationTrack* SkelAnimTrack = Cast<UMovieSceneSkeletalAnimationTrack>(Track);
-		FMovieSceneBinding* Binding = MovieScene->FindBinding(ObjectHandle);
-
-		// Add a track if no track was specified or if the track specified doesn't belong to the tracks of the targeted guid
-		if (!SkelAnimTrack || (Binding && !Binding->GetTracks().Contains(SkelAnimTrack)))
-		{
-			SkelAnimTrack = CastChecked<UMovieSceneSkeletalAnimationTrack>(AddTrack(MovieScene, ObjectHandle, UMovieSceneSkeletalAnimationTrack::StaticClass(), NAME_None), ECastCheckedType::NullAllowed);
-			KeyPropertyResult.bTrackCreated = true;
-		}
-
-		if (ensure(SkelAnimTrack))
-		{
-			SkelAnimTrack->Modify();
-
-			UMovieSceneSkeletalAnimationSection* NewSection = Cast<UMovieSceneSkeletalAnimationSection>(SkelAnimTrack->AddNewAnimationOnRow(KeyTime, AnimSequence, RowIndex));
-			KeyPropertyResult.bTrackModified = true;
-			KeyPropertyResult.SectionsCreated.Add(NewSection);
-
-			// Init the slot name on the new section if necessary
-			if (USkeletalMeshComponent* SkeletalMeshComponent = AcquireSkeletalMeshFromObjectGuid(ObjectHandle, GetSequencer()))
-			{
-				if (TSubclassOf<UAnimInstance> AnimInstanceClass = SkeletalMeshComponent->GetAnimClass())
-				{
-					if (UAnimInstance* AnimInstance = AnimInstanceClass->GetDefaultObject<UAnimInstance>())
-					{
-						if (AnimInstance->Implements<USequencerAnimationOverride>())
-						{
-							TScriptInterface<ISequencerAnimationOverride> SequencerAnimOverride = AnimInstance;
-							if (SequencerAnimOverride.GetObject())
-							{
-								TArray<FName> SlotNameOptions = ISequencerAnimationOverride::Execute_GetSequencerAnimSlotNames(SequencerAnimOverride.GetObject());
-								if (SlotNameOptions.Num() > 0)
-								{
-									NewSection->Params.SlotName = SlotNameOptions[0];
-								}
-							}
-						}
-					}
-				}
-			}
-
-			GetSequencer()->EmptySelection();
-			GetSequencer()->SelectSection(NewSection);
-			GetSequencer()->ThrobSectionSelection();
-		}
-	}
-
-	return KeyPropertyResult;
-}
-
-void FSkeletalAnimationTrackEditor::BuildTrackContextMenu_Internal(FMenuBuilder& MenuBuilder, UMovieSceneTrack* Track, const bool bAddSeparatorAtEnd)
-{
-	//there's a bug with a section being open already, so we end it.
-	UMovieSceneSkeletalAnimationTrack* SkeletalAnimationTrack = Cast<UMovieSceneSkeletalAnimationTrack>(Track);
-	/** Put this back when and if it works
-	MenuBuilder.BeginSection(NAME_None, LOCTEXT("MotionBlendingOptions", "Motion Blending Options"));
-	{
-		MenuBuilder.AddMenuEntry(
-			NSLOCTEXT("Sequencer", "AutoMatchClipsRootMotions", "Auto Match Clips Root Motions"),
-			NSLOCTEXT("Sequencer", "AutoMatchClipsRootMotionsTooltip", "Preceeding clips will auto match to the preceding clips root bones position. You can override this behavior per clip in it's section options."),
-			FSlateIcon(),
-			FUIAction(
-				FExecuteAction::CreateLambda([=]()->void {
-					SkeletalAnimationTrack->ToggleAutoMatchClipsRootMotions(); 
-					SequencerPtr->NotifyMovieSceneDataChanged(EMovieSceneDataChangeType::TrackValueChanged);
-
-					}),
-				FCanExecuteAction::CreateLambda([=]()->bool { return SequencerPtr && SkeletalAnimationTrack != nullptr; }),
-				FIsActionChecked::CreateLambda([=]()->bool { return SkeletalAnimationTrack != nullptr && SkeletalAnimationTrack->bAutoMatchClipsRootMotions; })),
-			NAME_None, 
-			EUserInterfaceActionType::ToggleButton
-		);
-		MenuBuilder.EndSection();
-	}
-	*/
-
-	MenuBuilder.BeginSection(NAME_None, LOCTEXT("SkelAnimRootMOtion", "Root Motion"));
-	{
-		MenuBuilder.AddMenuEntry(
-			NSLOCTEXT("Sequencer", "BlendFirstChildOfRoot", "Blend First Child Of Root"),
-			NSLOCTEXT("Sequencer", "BlendFirstChildOfRootTooltip", "If True, do not blend and match the root bones but instead the first child bone of the root. Toggle this on when the matched sequences in the track have no motion on the root."),
-			FSlateIcon(),
-			FUIAction(
-				FExecuteAction::CreateLambda([this, SkeletalAnimationTrack]() -> void
-				{
-					SkeletalAnimationTrack->bBlendFirstChildOfRoot = SkeletalAnimationTrack->bBlendFirstChildOfRoot ? false : true;
-					SkeletalAnimationTrack->SetRootMotionsDirty();
-
-					if (const TSharedPtr<ISequencer> SequencerPtr = GetSequencer())
-					{
-						SequencerPtr->NotifyMovieSceneDataChanged(EMovieSceneDataChangeType::TrackValueChanged);
-					}
-				}),
-				FCanExecuteAction::CreateLambda([this, SkeletalAnimationTrack]() -> bool
-				{
-					return SkeletalAnimationTrack != nullptr;
-				}),
-				FIsActionChecked::CreateLambda([this, SkeletalAnimationTrack]() -> bool
-				{
-					return SkeletalAnimationTrack != nullptr && SkeletalAnimationTrack->bBlendFirstChildOfRoot;
-				})),
-				NAME_None,
-				EUserInterfaceActionType::ToggleButton
-			);
-
-		MenuBuilder.AddMenuEntry(
-			NSLOCTEXT("Sequencer", "ShowRootMotionTrails", "Show Root Motion Trail"),
-			NSLOCTEXT("Sequencer", "ShowRootMotionTrailsTooltip", "Show the Root Motion Trail for all Animation Clips."),
-			FSlateIcon(),
-			FUIAction(
-				FExecuteAction::CreateLambda([this, SkeletalAnimationTrack]() -> void
-				{
-					SkeletalAnimationTrack->ToggleShowRootMotionTrail();
-
-					if (const TSharedPtr<ISequencer> SequencerPtr = GetSequencer())
-					{
-						SequencerPtr->NotifyMovieSceneDataChanged(EMovieSceneDataChangeType::TrackValueChanged);
-					}
-				}),
-				FCanExecuteAction::CreateLambda([this, SkeletalAnimationTrack]()->bool
-				{
-					if (const TSharedPtr<ISequencer> SequencerPtr = GetSequencer())
-					{
-						return SkeletalAnimationTrack != nullptr;
-					}
-					return false;
-				}),
-				FIsActionChecked::CreateLambda([this, SkeletalAnimationTrack]()->bool
-				{
-					return SkeletalAnimationTrack != nullptr && SkeletalAnimationTrack->bShowRootMotionTrail;
-				})),
-			NAME_None,
-				EUserInterfaceActionType::ToggleButton
-			);
-
-		MenuBuilder.AddMenuEntry(
-			NSLOCTEXT("Sequencer", "SwapRootBoneNone", "Swap Root Bone None"),
-			NSLOCTEXT("Sequencer", "SwapRootBoneNoneTooltip", "Do not swap root bone for all sections."),
-			FSlateIcon(),
-			FUIAction(
-				FExecuteAction::CreateLambda([SkeletalAnimationTrack]()->void {
-					SkeletalAnimationTrack->SetSwapRootBone(ESwapRootBone::SwapRootBone_None);
-					}),
-				FCanExecuteAction::CreateLambda([SkeletalAnimationTrack]()->bool { return  SkeletalAnimationTrack != nullptr; }),
-				FIsActionChecked::CreateLambda([SkeletalAnimationTrack]()->bool { return SkeletalAnimationTrack != nullptr && SkeletalAnimationTrack->SwapRootBone == ESwapRootBone::SwapRootBone_None; })),
-				NAME_None,
-				EUserInterfaceActionType::RadioButton
-		);
-
-		MenuBuilder.AddMenuEntry(
-			NSLOCTEXT("Sequencer", "SwapRootBoneActor", "Swap Root Bone Actor"),
-			NSLOCTEXT("Sequencer", "SwapRootBoneActorTooltip", "Swap root bone on root actor component for all sections."),
-			FSlateIcon(),
-			FUIAction(
-				FExecuteAction::CreateLambda([SkeletalAnimationTrack]()->void {
-					SkeletalAnimationTrack->SetSwapRootBone(ESwapRootBone::SwapRootBone_Actor);
-					}),
-				FCanExecuteAction::CreateLambda([SkeletalAnimationTrack]()->bool { return  SkeletalAnimationTrack != nullptr; }),
-				FIsActionChecked::CreateLambda([SkeletalAnimationTrack]()->bool { return SkeletalAnimationTrack != nullptr && SkeletalAnimationTrack->SwapRootBone == ESwapRootBone::SwapRootBone_Actor; })),
-				NAME_None,
-				EUserInterfaceActionType::RadioButton
-			);
-		MenuBuilder.AddMenuEntry(
-			NSLOCTEXT("Sequencer", "SwapRootBoneComponent", "Swap Root Bone Component"),
-			NSLOCTEXT("Sequencer", "SwapRootBoneComponentTooltip", "Swap root bone on current component for all sections."),
-			FSlateIcon(),
-			FUIAction(
-				FExecuteAction::CreateLambda([SkeletalAnimationTrack]()->void {
-					SkeletalAnimationTrack->SetSwapRootBone(ESwapRootBone::SwapRootBone_Component);
-					}),
-				FCanExecuteAction::CreateLambda([SkeletalAnimationTrack]()->bool { return  SkeletalAnimationTrack != nullptr; }),
-				FIsActionChecked::CreateLambda([SkeletalAnimationTrack]()->bool { return SkeletalAnimationTrack != nullptr && SkeletalAnimationTrack->SwapRootBone == ESwapRootBone::SwapRootBone_Component; })),
-				NAME_None,
-				EUserInterfaceActionType::RadioButton
-			);
-	}
-	MenuBuilder.EndSection();
-
-	if (bAddSeparatorAtEnd)
-	{
-		MenuBuilder.AddSeparator();
-	}
-}
-
-void FSkeletalAnimationTrackEditor::BuildTrackContextMenu(FMenuBuilder& MenuBuilder, UMovieSceneTrack* Track)
-{
-	BuildTrackContextMenu_Internal(MenuBuilder, Track, true);
-}
-
-void FSkeletalAnimationTrackEditor::BuildTrackSidebarMenu(FMenuBuilder& MenuBuilder, UMovieSceneTrack* Track)
+void FSkeletalAnimationTrackEditor::BuildTrackSidebarMenu(FMenuBuilder& MenuBuilder, UMovieSceneTrack* Track)
 {
 	BuildTrackContextMenu_Internal(MenuBuilder, Track, false);
 }
 
-TSharedPtr<SWidget> FSkeletalAnimationTrackEditor::BuildOutlinerEditWidget(const FGuid& ObjectBinding, UMovieSceneTrack* Track, const FBuildEditWidgetParams& Params)
-{
-	USkeleton* Skeleton = AcquireSkeletonFromObjectGuid(ObjectBinding, GetSequencer());
-
-	if (Skeleton)
-	{
-		FOnGetContent HandleGetAddButtonContent = FOnGetContent::CreateSP(this, &FSkeletalAnimationTrackEditor::BuildAddAnimationSubMenu, ObjectBinding, Skeleton, Params.TrackModel.AsWeak());
-		return UE::Sequencer::MakeAddButton(LOCTEXT("AnimationText", "Animation"), HandleGetAddButtonContent, Params.ViewModel);
-	}
-	else
-	{
-		return TSharedPtr<SWidget>();
-	}
-}
-
-bool FSkeletalAnimationTrackEditor::OnAllowDrop(const FDragDropEvent& DragDropEvent, FSequencerDragDropParams& DragDropParams)
-{
-	TSharedPtr<FDragDropOperation> Operation = DragDropEvent.GetOperation();
-	  
-	if (!Operation.IsValid() || !Operation->IsOfType<FAssetDragDropOp>() )
-	{
-		return false;
-	}
-	
-	if (!DragDropParams.TargetObjectGuid.IsValid())
-	{
-		return false;
-	}
-
-	TSharedPtr<ISequencer> SequencerPtr = GetSequencer();
-	if (!SequencerPtr)
-	{
-		return false;
-	}
-
-	UMovieSceneSequence* FocusedSequence = SequencerPtr->GetFocusedMovieSceneSequence();
-	if (!FocusedSequence)
-	{
-		return false;
-	}
-
-	TArray<USkeletalMeshComponent*> SkeletalMeshComponents = AcquireSkeletalMeshComponentsFromObjectGuid(DragDropParams.TargetObjectGuid, SequencerPtr, false);
-
-	TSharedPtr<FAssetDragDropOp> DragDropOp = StaticCastSharedPtr<FAssetDragDropOp>( Operation );
-
-	for (const FAssetData& AssetData : DragDropOp->GetAssets())
-	{
-		if (!MovieSceneToolHelpers::IsValidAsset(FocusedSequence, AssetData))
-		{
-			continue;
-		}
-
-		UAnimSequenceBase* AnimSequence = Cast<UAnimSequenceBase>(AssetData.GetAsset());
-
-		const bool bValidAnimSequence = AnimSequence && AnimSequence->CanBeUsedInComposition();
-
-		for (USkeletalMeshComponent* SkeletalMeshComponent : SkeletalMeshComponents)
-		{
-			USkeleton* Skeleton = GetSkeletonFromComponent(SkeletalMeshComponent);
-			if (bValidAnimSequence && Skeleton && Skeleton->IsCompatibleForEditor(AnimSequence->GetSkeleton()))
-			{
-				FFrameRate TickResolution = SequencerPtr->GetFocusedTickResolution();
-				FFrameNumber LengthInFrames = TickResolution.AsFrameNumber(AnimSequence->GetPlayLength());
-				DragDropParams.FrameRange = TRange<FFrameNumber>(DragDropParams.FrameNumber, DragDropParams.FrameNumber + LengthInFrames);
-				return true;
-			}
-		}
-	}
-
-	return false;
-}
-
-
-FReply FSkeletalAnimationTrackEditor::OnDrop(const FDragDropEvent& DragDropEvent, const FSequencerDragDropParams& DragDropParams)
-{
-	TSharedPtr<FDragDropOperation> Operation = DragDropEvent.GetOperation();
-
-	if (!Operation.IsValid() || !Operation->IsOfType<FAssetDragDropOp>() )
-	{
-		return FReply::Unhandled();
-	}
-	
-	if (!DragDropParams.TargetObjectGuid.IsValid())
-	{
-		return FReply::Unhandled();
-	}
-
-	TSharedPtr<ISequencer> SequencerPtr = GetSequencer();
-	if (!SequencerPtr)
-	{
-		return FReply::Unhandled();
-	}
-
-	UMovieSceneSequence* FocusedSequence = SequencerPtr->GetFocusedMovieSceneSequence();
-	if (!FocusedSequence)
-	{
-		return FReply::Unhandled();
-	}
-
-	TArray<USkeletalMeshComponent*> SkeletalMeshComponents = AcquireSkeletalMeshComponentsFromObjectGuid(DragDropParams.TargetObjectGuid, SequencerPtr, false);
-
-	const FScopedTransaction Transaction(LOCTEXT("DropAssets", "Drop Assets"));
-
-	TSharedPtr<FAssetDragDropOp> DragDropOp = StaticCastSharedPtr<FAssetDragDropOp>( Operation );
-
-	FMovieSceneTrackEditor::BeginKeying(DragDropParams.FrameNumber);
-
-	bool bAnyDropped = false;
-	for (const FAssetData& AssetData : DragDropOp->GetAssets())
-	{
-		if (!MovieSceneToolHelpers::IsValidAsset(FocusedSequence, AssetData))
-		{
-			continue;
-		}
-
-		UAnimSequenceBase* AnimSequence = Cast<UAnimSequenceBase>(AssetData.GetAsset());
-		const bool bValidAnimSequence = AnimSequence && AnimSequence->CanBeUsedInComposition();
-
-		for (USkeletalMeshComponent* SkeletalMeshComponent : SkeletalMeshComponents)
-		{
-			USkeleton* Skeleton = GetSkeletonFromComponent(SkeletalMeshComponent);
-
-			if (bValidAnimSequence && Skeleton && Skeleton->IsCompatibleForEditor(AnimSequence->GetSkeleton()))
-			{
-				UObject* BoundObject = SequencerPtr.IsValid() ? SequencerPtr->FindSpawnedObjectOrTemplate(DragDropParams.TargetObjectGuid) : nullptr;
-
-				AnimatablePropertyChanged( FOnKeyProperty::CreateRaw(this, &FSkeletalAnimationTrackEditor::AddKeyInternal, BoundObject, AnimSequence, DragDropParams.Track.Get(), DragDropParams.RowIndex));
-
-				bAnyDropped = true;
-			}
-		}
-	}
-
-	FMovieSceneTrackEditor::EndKeying();
-
-	return bAnyDropped ? FReply::Handled() : FReply::Unhandled();
-}
-
 
 #undef LOCTEXT_NAMESPACE
