Sequencer: Initial Animation Mixer Track prototype

diff --git a/Engine/Source/Editor/MovieSceneTools/Private/TrackEditors/SkeletalAnimationTrackEditor.h b/Engine/Source/Editor/MovieSceneTools/Private/TrackEditors/SkeletalAnimationTrackEditor.h
--- a/Engine/Source/Editor/MovieSceneTools/Private/TrackEditors/SkeletalAnimationTrackEditor.h
+++ b/Engine/Source/Editor/MovieSceneTools/Private/TrackEditors/SkeletalAnimationTrackEditor.h
@@ -2,40 +2,15 @@
 
 #pragma once
 
-#include "CoreMinimal.h"
-#include "Misc/Guid.h"
-#include "Templates/SubclassOf.h"
-#include "Templates/UniquePtr.h"
-#include "Widgets/SWidget.h"
-#include "ISequencer.h"
-#include "MovieSceneTrack.h"
-#include "ISequencerSection.h"
-#include "ISequencerTrackEditor.h"
-#include "MovieSceneTrackEditor.h"
-#include "EditModes/SkeletalAnimationTrackEditMode.h"
-#include "SequencerCoreFwd.h"
-
-namespace UE::Sequencer
-{
-	class ITrackExtension;
-}
-
-struct FAssetData;
-struct FMovieSceneTimeWarpChannel;
-struct FMovieSceneSequenceTransform;
-class FMenuBuilder;
-class FSequencerSectionPainter;
-class UMovieSceneSkeletalAnimationSection;
-class USkeleton;
-class USkeletalMeshComponent;
-class UAnimSeqExportOption;
+#include "TrackEditors/CommonAnimationTrackEditor.h"
 
 /**
  * Tools for animation tracks
  */
-class FSkeletalAnimationTrackEditor : public FMovieSceneTrackEditor, public FGCObject
+class FSkeletalAnimationTrackEditor : public UE::Sequencer::FCommonAnimationTrackEditor
 {
 public:
+
 	/**
 	 * Constructor
 	 *
@@ -43,156 +18,35 @@ public:
 	 */
 	FSkeletalAnimationTrackEditor( TSharedRef<ISequencer> InSequencer );
 
-	/** Virtual destructor. */
-	virtual ~FSkeletalAnimationTrackEditor() { }
+	static TSharedRef<ISequencerTrackEditor> CreateTrackEditor(TSharedRef<ISequencer> InSequencer);
 
-	//~ FGCObject
-	virtual void AddReferencedObjects(FReferenceCollector& Collector) override;
-	virtual FString GetReferencerName() const override
-	{
-		return TEXT("FSkeletalAnimationTrackEditor");
-	}
-
-	/**
-	 * Creates an instance of this class.  Called by a sequencer 
-	 *
-	 * @param OwningSequencer The sequencer instance to be used by this tool
-	 * @return The new instance of this class
-	 */
-	static TSharedRef<ISequencerTrackEditor> CreateTrackEditor( TSharedRef<ISequencer> OwningSequencer );
-
-	/**
-	* Keeps track of how many skeletal animation track editors we have*
-	*/
-	static int32 NumberActive;
-public:
-
-	// ISequencerTrackEditor interface
+	virtual void OnInitialize() override;
+	virtual void OnRelease() override;
 
-	virtual FText GetDisplayName() const override;
-	virtual void BuildObjectBindingContextMenu(FMenuBuilder& MenuBuilder, const TArray<FGuid>& ObjectBindings, const UClass* ObjectClass) override;
-	virtual void BuildObjectBindingTrackMenu(FMenuBuilder& MenuBuilder, const TArray<FGuid>& ObjectBindings, const UClass* ObjectClass) override;
-	virtual bool HandleAssetAdded(UObject* Asset, const FGuid& TargetObjectGuid) override;
-	virtual TSharedRef<ISequencerSection> MakeSectionInterface( UMovieSceneSection& SectionObject, UMovieSceneTrack& Track, FGuid ObjectBinding ) override;
-	virtual bool SupportsSequence(UMovieSceneSequence* InSequence) const override;
-	virtual bool SupportsType( TSubclassOf<UMovieSceneTrack> Type ) const override;
 	virtual void BuildTrackContextMenu(FMenuBuilder& MenuBuilder, UMovieSceneTrack* Track) override;
 	virtual void BuildTrackSidebarMenu(FMenuBuilder& MenuBuilder, UMovieSceneTrack* Track) override;
-	virtual TSharedPtr<SWidget> BuildOutlinerEditWidget(const FGuid& ObjectBinding, UMovieSceneTrack* Track, const FBuildEditWidgetParams& Params) override;
-	virtual bool OnAllowDrop(const FDragDropEvent& DragDropEvent, FSequencerDragDropParams& DragDropParams) override;
-	virtual FReply OnDrop(const FDragDropEvent& DragDropEvent, const FSequencerDragDropParams& DragDropParams) override;
-	virtual void OnInitialize() override;
-	virtual void OnRelease() override;
+	virtual bool SupportsSequence(UMovieSceneSequence* InSequence) const override;
+	virtual bool SupportsType( TSubclassOf<UMovieSceneTrack> Type ) const override;
+	virtual TSharedRef<ISequencerSection> MakeSectionInterface(UMovieSceneSection& SectionObject, UMovieSceneTrack& Track, FGuid ObjectBinding) override;
+	virtual TSubclassOf<UMovieSceneCommonAnimationTrack> GetTrackClass() const override;
 
 private:
 
-	/** Animation sub menu */
-	TSharedRef<SWidget> BuildAddAnimationSubMenu(FGuid ObjectBinding, USkeleton* Skeleton, UE::Sequencer::TWeakViewModelPtr<UE::Sequencer::ITrackExtension> TrackModel);
-	TSharedRef<SWidget> BuildAnimationSubMenu(FGuid ObjectBinding, USkeleton* Skeleton, UMovieSceneTrack* Track);
-	void AddAnimationSubMenu(FMenuBuilder& MenuBuilder, TArray<FGuid> ObjectBindings, USkeleton* Skeleton, UMovieSceneTrack* Track);
-
-	/** Filter only compatible skeletons */
-	bool FilterAnimSequences(const FAssetData& AssetData, USkeleton* Skeleton);
-
-	/** Animation sub menu filter function */
-	bool ShouldFilterAsset(const FAssetData& AssetData);
-
-	/** Animation asset selected */
-	void OnAnimationAssetSelected(const FAssetData& AssetData, TArray<FGuid> ObjectBindings, UMovieSceneTrack* Track);
-
-	/** Animation asset enter pressed */
-	void OnAnimationAssetEnterPressed(const TArray<FAssetData>& AssetData, TArray<FGuid> ObjectBindings, UMovieSceneTrack* Track);
-
-	/** Delegate for AnimatablePropertyChanged in AddKey */
-	FKeyPropertyResult AddKeyInternal(FFrameNumber KeyTime, UObject* Object, class UAnimSequenceBase* AnimSequence, UMovieSceneTrack* Track, int32 RowIndex);
-	
-	/** Construct the binding menu*/
-	void ConstructObjectBindingTrackMenu(FMenuBuilder& MenuBuilder, TArray<FGuid> ObjectBindings);
-
-	/** Callback to Create the Animation Asset, pop open the dialog */
-	void HandleCreateAnimationSequence(USkeletalMeshComponent* SkelMeshComp, USkeleton* Skeleton, FGuid Binding, bool bCeateSoftLink);
-
-	/** Callback to Creae the Animation Asset after getting the name*/
-	bool CreateAnimationSequence(const TArray<UObject*> NewAssets,USkeletalMeshComponent* SkelMeshComp, FGuid Binding, bool bCreateSoftLink);
-
-	/** Open the linked Anim Sequence*/
-	void OpenLinkedAnimSequence(FGuid Binding);
-
-	/** Can Open the linked Anim Sequence*/
-	bool CanOpenLinkedAnimSequence(FGuid Binding);
+	void OnPostPropertyChanged(UObject* InObject, struct FPropertyChangedEvent& InPropertyChangedEvent);
 
 	/** Common function used to build context and sidebar menus */
 	void BuildTrackContextMenu_Internal(FMenuBuilder& MenuBuilder, UMovieSceneTrack* const InTrack, const bool bAddSeparatorAtEnd);
-
-	friend class FMovieSceneSkeletalAnimationParamsDetailCustomization;
-
-private:
-	/* Was part of the the section but should be at the track level since it takes the final blended result at the current time, not the section instance value*/
-	bool CreatePoseAsset(const TArray<UObject*> NewAssets, FGuid InObjectBinding);
-	void HandleCreatePoseAsset(FGuid InObjectBinding);
-	bool CanCreatePoseAsset(FGuid InObjectBinding) const;
-
-private:
-	/* For Anim Sequence UI Option with be gc'd*/
-	TObjectPtr<UAnimSeqExportOption> AnimSeqExportOption;
-
-
-private:
-	/* Delegate to handle sequencer changes for auto baking of anim sequences*/
-	FDelegateHandle SequencerSavedHandle;
-	void OnSequencerSaved(ISequencer& InSequence);
-	FDelegateHandle SequencerChangedHandle;
-	void OnSequencerDataChanged(EMovieSceneDataChangeType DataChangeType);
-	void OnPostPropertyChanged(UObject* InObject, struct FPropertyChangedEvent& InPropertyChangedEvent);
 };
 
 
 /** Class for animation sections */
 class FSkeletalAnimationSection
-	: public ISequencerSection
-	, public TSharedFromThis<FSkeletalAnimationSection>
+	: public UE::Sequencer::FCommonAnimationSection
 {
 public:
 
 	/** Constructor. */
 	FSkeletalAnimationSection( UMovieSceneSection& InSection, TWeakPtr<ISequencer> InSequencer);
 
-	/** Virtual destructor. */
-	virtual ~FSkeletalAnimationSection();
-
-public:
-
-	// ISequencerSection interface
-
-	virtual UMovieSceneSection* GetSectionObject() override;
-	virtual FText GetSectionTitle() const override;
-	virtual FText GetSectionToolTip() const override;
-	virtual TOptional<FFrameTime> GetSectionTime(FSequencerSectionPainter& InPainter) const override;
-	virtual float GetSectionHeight(const UE::Sequencer::FViewDensityInfo& ViewDensity) const override;
-	virtual FMargin GetContentPadding() const override;
-	virtual int32 OnPaintSection( FSequencerSectionPainter& Painter ) const override;
-	virtual void BeginResizeSection() override;
-	virtual void ResizeSection(ESequencerSectionResizeMode ResizeMode, FFrameNumber ResizeTime) override;
-	virtual void BeginSlipSection() override;
-	virtual void SlipSection(FFrameNumber SlipTime) override;
-	virtual void CustomizePropertiesDetailsView(TSharedRef<IDetailsView> DetailsView, const FSequencerSectionPropertyDetailsViewCustomizationParams& InParams) const override;
 	virtual void BuildSectionContextMenu(FMenuBuilder& MenuBuilder, const FGuid& InObjectBinding) override;
-	virtual void BeginDilateSection() override;
-	virtual void DilateSection(const TRange<FFrameNumber>& NewRange, float DilationFactor) override;
-	virtual bool RequestDeleteKeyArea(const TArray<FName>& KeyAreaNamePath) override;
-
-
-private:
-	void FindBestBlendSection(FGuid InObjectBinding);
-private:
-
-	/** The section we are visualizing */
-	UMovieSceneSkeletalAnimationSection& Section;
-
-	/** Used to draw animation frame, need selection state and local time*/
-	TWeakPtr<ISequencer> Sequencer;
-
-	TUniquePtr<FMovieSceneSequenceTransform> InitialDragTransform;
-	TUniquePtr<FMovieSceneTimeWarpChannel> PreDilateChannel;
-	double PreDilatePlayRate;
 };
diff --git a/Engine/Source/Editor/MovieSceneTools/Public/TrackEditors/CommonAnimationTrackEditor.h b/Engine/Source/Editor/MovieSceneTools/Public/TrackEditors/CommonAnimationTrackEditor.h
new file mode 100644
--- /dev/null
+++ b/Engine/Source/Editor/MovieSceneTools/Public/TrackEditors/CommonAnimationTrackEditor.h
@@ -0,0 +1,190 @@
+// Copyright Epic Games, Inc. All Rights Reserved.
+
+#pragma once
+
+#include "CoreTypes.h"
+#include "Misc/Guid.h"
+#include "Templates/SubclassOf.h"
+#include "Templates/UniquePtr.h"
+#include "Widgets/SWidget.h"
+#include "ISequencer.h"
+#include "MovieSceneTrack.h"
+#include "ISequencerSection.h"
+#include "MovieSceneTrackEditor.h"
+#include "SequencerCoreFwd.h"
+
+struct FAssetData;
+struct FMovieSceneTimeWarpChannel;
+struct FMovieSceneSequenceTransform;
+
+class FMenuBuilder;
+class FSequencerSectionPainter;
+class UAnimSeqExportOption;
+class UAnimSequenceBase;
+class UMovieSceneCommonAnimationTrack;
+class UMovieSceneSkeletalAnimationSection;
+class USkeletalMeshComponent;
+class USkeleton;
+
+namespace UE::Sequencer
+{
+
+class ITrackExtension;
+
+/**
+ * Tools for animation tracks
+ */
+class MOVIESCENETOOLS_API FCommonAnimationTrackEditor : public FMovieSceneTrackEditor, public FGCObject
+{
+public:
+	/**
+	 * Constructor
+	 *
+	 * @param InSequencer The sequencer instance to be used by this tool
+	 */
+	FCommonAnimationTrackEditor( TSharedRef<ISequencer> InSequencer );
+
+	/** Virtual destructor. */
+	virtual ~FCommonAnimationTrackEditor() { }
+
+	//~ FGCObject
+	virtual void AddReferencedObjects(FReferenceCollector& Collector) override;
+	virtual FString GetReferencerName() const override
+	{
+		return TEXT("FCommonAnimationTrackEditor");
+	}
+
+	/**
+	* Keeps track of how many skeletal animation track editors we have*
+	*/
+	static int32 NumberActive;
+
+	static USkeletalMeshComponent* AcquireSkeletalMeshFromObjectGuid(const FGuid& Guid, TSharedPtr<ISequencer> SequencerPtr);
+	static USkeleton* AcquireSkeletonFromObjectGuid(const FGuid& Guid, TSharedPtr<ISequencer> SequencerPtr);
+
+public:
+
+	// ISequencerTrackEditor interface
+
+	virtual FText GetDisplayName() const override;
+	virtual void BuildObjectBindingContextMenu(FMenuBuilder& MenuBuilder, const TArray<FGuid>& ObjectBindings, const UClass* ObjectClass) override;
+	virtual void BuildObjectBindingTrackMenu(FMenuBuilder& MenuBuilder, const TArray<FGuid>& ObjectBindings, const UClass* ObjectClass) override;
+	virtual bool HandleAssetAdded(UObject* Asset, const FGuid& TargetObjectGuid) override;
+	virtual TSharedRef<ISequencerSection> MakeSectionInterface( UMovieSceneSection& SectionObject, UMovieSceneTrack& Track, FGuid ObjectBinding ) override;
+	virtual TSharedPtr<SWidget> BuildOutlinerEditWidget(const FGuid& ObjectBinding, UMovieSceneTrack* Track, const FBuildEditWidgetParams& Params) override;
+	virtual bool OnAllowDrop(const FDragDropEvent& DragDropEvent, FSequencerDragDropParams& DragDropParams) override;
+	virtual FReply OnDrop(const FDragDropEvent& DragDropEvent, const FSequencerDragDropParams& DragDropParams) override;
+	virtual void OnInitialize() override;
+	virtual void OnRelease() override;
+
+protected:
+
+	virtual TSubclassOf<UMovieSceneCommonAnimationTrack> GetTrackClass() const = 0;
+
+	/** Animation sub menu */
+	TSharedRef<SWidget> BuildAddAnimationSubMenu(FGuid ObjectBinding, USkeleton* Skeleton, UE::Sequencer::TWeakViewModelPtr<UE::Sequencer::ITrackExtension> TrackModel);
+	TSharedRef<SWidget> BuildAnimationSubMenu(FGuid ObjectBinding, USkeleton* Skeleton, UMovieSceneTrack* Track);
+	void AddAnimationSubMenu(FMenuBuilder& MenuBuilder, TArray<FGuid> ObjectBindings, USkeleton* Skeleton, UMovieSceneTrack* Track);
+
+	/** Filter only compatible skeletons */
+	bool FilterAnimSequences(const FAssetData& AssetData, USkeleton* Skeleton);
+
+	/** Animation sub menu filter function */
+	bool ShouldFilterAsset(const FAssetData& AssetData);
+
+	/** Animation asset selected */
+	void OnAnimationAssetSelected(const FAssetData& AssetData, TArray<FGuid> ObjectBindings, UMovieSceneTrack* Track);
+
+	/** Animation asset enter pressed */
+	void OnAnimationAssetEnterPressed(const TArray<FAssetData>& AssetData, TArray<FGuid> ObjectBindings, UMovieSceneTrack* Track);
+
+	/** Delegate for AnimatablePropertyChanged in AddKey */
+	FKeyPropertyResult AddKeyInternal(FFrameNumber KeyTime, UObject* Object, UAnimSequenceBase* AnimSequence, UMovieSceneTrack* Track, int32 RowIndex);
+	
+	/** Construct the binding menu*/
+	void ConstructObjectBindingTrackMenu(FMenuBuilder& MenuBuilder, TArray<FGuid> ObjectBindings);
+
+	/** Callback to Create the Animation Asset, pop open the dialog */
+	void HandleCreateAnimationSequence(USkeletalMeshComponent* SkelMeshComp, USkeleton* Skeleton, FGuid Binding, bool bCeateSoftLink);
+
+	/** Callback to Creae the Animation Asset after getting the name*/
+	bool CreateAnimationSequence(const TArray<UObject*> NewAssets,USkeletalMeshComponent* SkelMeshComp, FGuid Binding, bool bCreateSoftLink);
+
+	/** Open the linked Anim Sequence*/
+	void OpenLinkedAnimSequence(FGuid Binding);
+
+	/** Can Open the linked Anim Sequence*/
+	bool CanOpenLinkedAnimSequence(FGuid Binding);
+
+	friend class FMovieSceneSkeletalAnimationParamsDetailCustomization;
+
+protected:
+	/* Was part of the the section but should be at the track level since it takes the final blended result at the current time, not the section instance value*/
+	bool CreatePoseAsset(const TArray<UObject*> NewAssets, FGuid InObjectBinding);
+	void HandleCreatePoseAsset(FGuid InObjectBinding);
+	bool CanCreatePoseAsset(FGuid InObjectBinding) const;
+
+protected:
+	/* For Anim Sequence UI Option with be gc'd*/
+	TObjectPtr<UAnimSeqExportOption> AnimSeqExportOption;
+
+
+protected:
+	/* Delegate to handle sequencer changes for auto baking of anim sequences*/
+	FDelegateHandle SequencerSavedHandle;
+	void OnSequencerSaved(ISequencer& InSequence);
+	FDelegateHandle SequencerChangedHandle;
+	void OnSequencerDataChanged(EMovieSceneDataChangeType DataChangeType);
+};
+
+
+/** Class for animation sections */
+class MOVIESCENETOOLS_API FCommonAnimationSection
+	: public ISequencerSection
+	, public TSharedFromThis<FCommonAnimationSection>
+{
+public:
+
+	/** Constructor. */
+	FCommonAnimationSection( UMovieSceneSection& InSection, TWeakPtr<ISequencer> InSequencer);
+
+	/** Virtual destructor. */
+	virtual ~FCommonAnimationSection();
+
+public:
+
+	// ISequencerSection interface
+
+	virtual UMovieSceneSection* GetSectionObject() override;
+	virtual FText GetSectionTitle() const override;
+	virtual FText GetSectionToolTip() const override;
+	virtual TOptional<FFrameTime> GetSectionTime(FSequencerSectionPainter& InPainter) const override;
+	virtual float GetSectionHeight(const UE::Sequencer::FViewDensityInfo& ViewDensity) const override;
+	virtual FMargin GetContentPadding() const override;
+	virtual int32 OnPaintSection( FSequencerSectionPainter& Painter ) const override;
+	virtual void BeginResizeSection() override;
+	virtual void ResizeSection(ESequencerSectionResizeMode ResizeMode, FFrameNumber ResizeTime) override;
+	virtual void BeginSlipSection() override;
+	virtual void SlipSection(FFrameNumber SlipTime) override;
+	virtual void CustomizePropertiesDetailsView(TSharedRef<IDetailsView> DetailsView, const FSequencerSectionPropertyDetailsViewCustomizationParams& InParams) const override;
+	virtual void BeginDilateSection() override;
+	virtual void DilateSection(const TRange<FFrameNumber>& NewRange, float DilationFactor) override;
+	virtual bool RequestDeleteKeyArea(const TArray<FName>& KeyAreaNamePath) override;
+
+
+protected:
+	void FindBestBlendSection(FGuid InObjectBinding);
+protected:
+
+	/** The section we are visualizing */
+	UMovieSceneSkeletalAnimationSection& Section;
+
+	/** Used to draw animation frame, need selection state and local time*/
+	TWeakPtr<ISequencer> Sequencer;
+
+	TUniquePtr<FMovieSceneSequenceTransform> InitialDragTransform;
+	TUniquePtr<FMovieSceneTimeWarpChannel> PreDilateChannel;
+	double PreDilatePlayRate;
+};
+
+} // namespace UE::Sequencer
\ No newline at end of file
diff --git a/Engine/Source/Runtime/MovieScene/Private/MovieSceneTrack.cpp b/Engine/Source/Runtime/MovieScene/Private/MovieSceneTrack.cpp
--- a/Engine/Source/Runtime/MovieScene/Private/MovieSceneTrack.cpp
+++ b/Engine/Source/Runtime/MovieScene/Private/MovieSceneTrack.cpp
@@ -99,12 +99,14 @@ void UMovieSceneTrack::AddSection(FSectionParameter Section)
 {
 	CallAddSection(*Section.Section);
 	Section.Section->OnAddedToTrack(this);
+	OnSectionAddedImpl(Section.Section);
 }
 
 void UMovieSceneTrack::RemoveSection(FSectionParameter Section)
 {
 	CallRemoveSection(*Section.Section);
 	Section.Section->OnRemovedFromTrack();
+	OnSectionRemovedImpl(Section.Section);
 }
 
 void UMovieSceneTrack::RemoveSectionAt(FSectionIndexParameter SectionIndex)
@@ -114,7 +116,10 @@ void UMovieSceneTrack::RemoveSectionAt(FSectionIndexParameter SectionIndex)
 	{
 		UMovieSceneSection* Section = Sections[SectionIndex.SectionIndex];
 		CallRemoveSectionAt(SectionIndex.SectionIndex);
-		Section->OnRemovedFromTrack();
+		if (Section)
+		{
+			Section->OnRemovedFromTrack();
+		}
 	}
 }
 
diff --git a/Engine/Source/Runtime/MovieScene/Public/EntitySystem/MovieSceneEntitySystemTask.h b/Engine/Source/Runtime/MovieScene/Public/EntitySystem/MovieSceneEntitySystemTask.h
--- a/Engine/Source/Runtime/MovieScene/Public/EntitySystem/MovieSceneEntitySystemTask.h
+++ b/Engine/Source/Runtime/MovieScene/Public/EntitySystem/MovieSceneEntitySystemTask.h
@@ -1003,6 +1003,17 @@ public:
 	FCommonEntityTaskParams CommonParams;
 };
 
+struct FNoopTask
+{
+	template<typename... T>
+	FORCEINLINE static void ForEachAllocation(T&&...)
+	{}
+
+	template<typename... T>
+	FORCEINLINE static void ForEachEntity(T&&...)
+	{}
+};
+
 
 /**
  * Main entry point utility for create tasks that run over component data
diff --git a/Engine/Source/Runtime/MovieScene/Public/MovieSceneTrack.h b/Engine/Source/Runtime/MovieScene/Public/MovieSceneTrack.h
--- a/Engine/Source/Runtime/MovieScene/Public/MovieSceneTrack.h
+++ b/Engine/Source/Runtime/MovieScene/Public/MovieSceneTrack.h
@@ -422,7 +422,7 @@ public:
 	 * Updates the row indices of sections owned by this track so that all row indices which are used are consecutive with no gaps. 
 	 * @return Whether or not fixes were made. 
 	 */
-	MOVIESCENE_API bool FixRowIndices();
+	MOVIESCENE_API virtual bool FixRowIndices();
 
 	/** Called when row indices have been fixed up */
 	MOVIESCENE_API virtual void OnRowIndicesChanged(const TMap<int32, int32>& NewToOldRowIndices);
@@ -536,6 +536,18 @@ public:
 	 */
 	virtual bool HasSection(const UMovieSceneSection& Section) const PURE_VIRTUAL(UMovieSceneTrack::HasSection, return false;);
 
+private:
+
+	virtual void OnSectionAddedImpl(UMovieSceneSection* Section)
+	{
+	}
+
+	virtual void OnSectionRemovedImpl(UMovieSceneSection* Section)
+	{
+	}
+
+public:
+
 #if WITH_EDITOR
 
 	/**
diff --git a/Engine/Source/Runtime/MovieSceneTracks/Private/Sections/MovieSceneSkeletalAnimationSection.cpp b/Engine/Source/Runtime/MovieSceneTracks/Private/Sections/MovieSceneSkeletalAnimationSection.cpp
--- a/Engine/Source/Runtime/MovieSceneTracks/Private/Sections/MovieSceneSkeletalAnimationSection.cpp
+++ b/Engine/Source/Runtime/MovieSceneTracks/Private/Sections/MovieSceneSkeletalAnimationSection.cpp
@@ -2,6 +2,7 @@
 
 #include "Sections/MovieSceneSkeletalAnimationSection.h"
 #include "Channels/MovieSceneChannelProxy.h"
+#include "Decorations/IMovieSceneChannelDecoration.h"
 #include "Animation/AnimSequence.h"
 #include "AnimSequencerInstanceProxy.h"
 #include "Logging/MessageLog.h"
@@ -157,6 +158,16 @@ EMovieSceneChannelProxyType UMovieSceneSkeletalAnimationSection::CacheChannelPro
 		}
 	}
 
+	for (UObject* Decoration : GetDecorations())
+	{
+		if (IMovieSceneChannelDecoration* ChannelDecoration = Cast<IMovieSceneChannelDecoration>(Decoration))
+		{
+			// Result doesn't matter because we always return EMovieSceneChannelProxyType::Dynamic anyway
+			ChannelDecoration->PopulateChannelProxy(Channels);
+		}
+	}
+
+
 
 #if WITH_EDITOR
 
diff --git a/Engine/Source/Runtime/MovieSceneTracks/Private/Tracks/MovieSceneSkeletalAnimationTrack.cpp b/Engine/Source/Runtime/MovieSceneTracks/Private/Tracks/MovieSceneSkeletalAnimationTrack.cpp
--- a/Engine/Source/Runtime/MovieSceneTracks/Private/Tracks/MovieSceneSkeletalAnimationTrack.cpp
+++ b/Engine/Source/Runtime/MovieSceneTracks/Private/Tracks/MovieSceneSkeletalAnimationTrack.cpp
@@ -52,7 +52,7 @@ UMovieSceneSkeletalAnimationTrack::UMovieSceneSkeletalAnimationTrack(const FObje
 /* UMovieSceneSkeletalAnimationTrack interface
  *****************************************************************************/
 
-UMovieSceneSection* UMovieSceneSkeletalAnimationTrack::AddNewAnimationOnRow(FFrameNumber KeyTime, UAnimSequenceBase* AnimSequence, int32 RowIndex)
+UMovieSceneSection* UMovieSceneCommonAnimationTrack::AddNewAnimationOnRow(FFrameNumber KeyTime, UAnimSequenceBase* AnimSequence, int32 RowIndex)
 {
 	UMovieSceneSkeletalAnimationSection* NewSection = Cast<UMovieSceneSkeletalAnimationSection>(CreateNewSection());
 	{
@@ -191,42 +191,10 @@ void UMovieSceneSkeletalAnimationTrack::PostEditChangeProperty(FPropertyChangedE
 
 #endif
 
-
-const TArray<UMovieSceneSection*>& UMovieSceneSkeletalAnimationTrack::GetAllSections() const
-{
-	return AnimationSections;
-}
-
-
-bool UMovieSceneSkeletalAnimationTrack::SupportsMultipleRows() const
-{
-	return true;
-}
-
-bool UMovieSceneSkeletalAnimationTrack::SupportsType(TSubclassOf<UMovieSceneSection> SectionClass) const
-{
-	return SectionClass == UMovieSceneSkeletalAnimationSection::StaticClass();
-}
-
-UMovieSceneSection* UMovieSceneSkeletalAnimationTrack::CreateNewSection()
-{
-	return NewObject<UMovieSceneSkeletalAnimationSection>(this, NAME_None, RF_Transactional);
-}
-
-void UMovieSceneSkeletalAnimationTrack::RemoveAllAnimationData()
-{
-	AnimationSections.Empty();
-
-}
-
-bool UMovieSceneSkeletalAnimationTrack::HasSection(const UMovieSceneSection& Section) const
-{
-	return AnimationSections.Contains(&Section);
-}
-
 void UMovieSceneSkeletalAnimationTrack::AddSection(UMovieSceneSection& Section)
 {
-	AnimationSections.Add(&Section);
+	Super::AddSection(Section);
+
 	UMovieSceneSkeletalAnimationSection* AnimSection = Cast< UMovieSceneSkeletalAnimationSection>(&Section);
 	if (AnimSection)
 	{
@@ -243,21 +211,16 @@ void UMovieSceneSkeletalAnimationTrack::UpdateEasing()
 
 void UMovieSceneSkeletalAnimationTrack::RemoveSection(UMovieSceneSection& Section)
 {
-	AnimationSections.Remove(&Section);
+	Super::RemoveSection(Section);
 	SetUpRootMotions(true);
 }
 
 void UMovieSceneSkeletalAnimationTrack::RemoveSectionAt(int32 SectionIndex)
 {
-	AnimationSections.RemoveAt(SectionIndex);
+	Super::RemoveSectionAt(SectionIndex);
 	SetUpRootMotions(true);
 }
 
-bool UMovieSceneSkeletalAnimationTrack::IsEmpty() const
-{
-	return AnimationSections.Num() == 0;
-}
-
 #if WITH_EDITORONLY_DATA
 
 FText UMovieSceneSkeletalAnimationTrack::GetDefaultDisplayName() const
diff --git a/Engine/Source/Runtime/MovieSceneTracks/Public/Sections/MovieSceneSkeletalAnimationSection.h b/Engine/Source/Runtime/MovieSceneTracks/Public/Sections/MovieSceneSkeletalAnimationSection.h
--- a/Engine/Source/Runtime/MovieSceneTracks/Public/Sections/MovieSceneSkeletalAnimationSection.h
+++ b/Engine/Source/Runtime/MovieSceneTracks/Public/Sections/MovieSceneSkeletalAnimationSection.h
@@ -108,7 +108,7 @@ struct FMovieSceneSkeletalAnimationParams
 /**
  * Movie scene section that control skeletal animation
  */
-UCLASS( MinimalAPI )
+UCLASS( MinimalAPI, DisplayName="Anim Sequence" )
 class UMovieSceneSkeletalAnimationSection
 	: public UMovieSceneSection
 	, public IMovieSceneEntityProvider
diff --git a/Engine/Source/Runtime/MovieSceneTracks/Public/Tracks/MovieSceneCommonAnimationTrack.h b/Engine/Source/Runtime/MovieSceneTracks/Public/Tracks/MovieSceneCommonAnimationTrack.h
new file mode 100644
--- /dev/null
+++ b/Engine/Source/Runtime/MovieSceneTracks/Public/Tracks/MovieSceneCommonAnimationTrack.h
@@ -0,0 +1,41 @@
+// Copyright Epic Games, Inc. All Rights Reserved.
+
+#pragma once
+
+#include "CoreTypes.h"
+#include "Misc/FrameNumber.h"
+#include "MovieSceneNameableTrack.h"
+#include "MovieSceneCommonAnimationTrack.generated.h"
+
+class UAnimSequenceBase;
+class UMovieSceneSection;
+
+UCLASS()
+class MOVIESCENETRACKS_API UMovieSceneCommonAnimationTrack
+	: public UMovieSceneNameableTrack
+{
+	GENERATED_BODY()
+
+public:
+
+	/** Adds a new animation to this track */
+	virtual UMovieSceneSection* AddNewAnimationOnRow(FFrameNumber KeyTime, UAnimSequenceBase* AnimSequence, int32 RowIndex);
+
+	/** Adds a new animation to this track on the next available/non-overlapping row */
+	 virtual UMovieSceneSection* AddNewAnimation(FFrameNumber KeyTime, UAnimSequenceBase* AnimSequence) { return AddNewAnimationOnRow(KeyTime, AnimSequence, INDEX_NONE); }
+
+	virtual const TArray<UMovieSceneSection*>& GetAllSections() const override;
+	virtual UMovieSceneSection* CreateNewSection() override;
+	virtual bool SupportsType(TSubclassOf<UMovieSceneSection> SectionClass) const override;
+	virtual bool SupportsMultipleRows() const override;
+	virtual void RemoveAllAnimationData() override;
+	virtual bool HasSection(const UMovieSceneSection& Section) const override;
+	virtual void AddSection(UMovieSceneSection& Section) override;
+	virtual void RemoveSection(UMovieSceneSection& Section) override;
+	virtual void RemoveSectionAt(int32 SectionIndex) override;
+	virtual bool IsEmpty() const override;
+
+	/** List of all animation sections */
+	UPROPERTY()
+	TArray<TObjectPtr<UMovieSceneSection>> AnimationSections;
+};
\ No newline at end of file
diff --git a/Engine/Source/Runtime/MovieSceneTracks/Public/Tracks/MovieSceneSkeletalAnimationTrack.h b/Engine/Source/Runtime/MovieSceneTracks/Public/Tracks/MovieSceneSkeletalAnimationTrack.h
--- a/Engine/Source/Runtime/MovieSceneTracks/Public/Tracks/MovieSceneSkeletalAnimationTrack.h
+++ b/Engine/Source/Runtime/MovieSceneTracks/Public/Tracks/MovieSceneSkeletalAnimationTrack.h
@@ -6,6 +6,7 @@
 #include "UObject/ObjectMacros.h"
 #include "MovieSceneNameableTrack.h"
 #include "Sections/MovieSceneSkeletalAnimationSection.h"
+#include "Tracks/MovieSceneCommonAnimationTrack.h"
 #include "MovieSceneSkeletalAnimationTrack.generated.h"
 
 enum class ESwapRootBone : uint8;
@@ -39,18 +40,12 @@ struct FMovieSceneSkeletalAnimRootMotionTrackParams
  */
 UCLASS(MinimalAPI)
 class UMovieSceneSkeletalAnimationTrack
-	: public UMovieSceneNameableTrack
+	: public UMovieSceneCommonAnimationTrack
 {
 	GENERATED_UCLASS_BODY()
 
 public:
 
-	/** Adds a new animation to this track */
-	MOVIESCENETRACKS_API virtual UMovieSceneSection* AddNewAnimationOnRow(FFrameNumber KeyTime, class UAnimSequenceBase* AnimSequence, int32 RowIndex);
-
-	/** Adds a new animation to this track on the next available/non-overlapping row */
-	virtual UMovieSceneSection* AddNewAnimation(FFrameNumber KeyTime, class UAnimSequenceBase* AnimSequence) { return AddNewAnimationOnRow(KeyTime, AnimSequence, INDEX_NONE); }
-
 	/** Gets the animation sections at a certain time */
 	MOVIESCENETRACKS_API TArray<UMovieSceneSection*> GetAnimSectionsAtTime(FFrameNumber Time);
 
@@ -66,17 +61,10 @@ public:
 
 #endif
 
-	virtual void RemoveAllAnimationData() override;
-	virtual bool HasSection(const UMovieSceneSection& Section) const override;
 	virtual void AddSection(UMovieSceneSection& Section) override;
 	virtual void RemoveSection(UMovieSceneSection& Section) override;
 	virtual void RemoveSectionAt(int32 SectionIndex) override;
-	virtual bool IsEmpty() const override;
-	virtual const TArray<UMovieSceneSection*>& GetAllSections() const override;
-	virtual bool SupportsType(TSubclassOf<UMovieSceneSection> SectionClass) const override;
-	virtual UMovieSceneSection* CreateNewSection() override;
 	virtual bool PopulateEvaluationTree(TMovieSceneEvaluationTree<FMovieSceneTrackEvaluationData>& OutData) const override;
-	virtual bool SupportsMultipleRows() const override;
 	virtual void UpdateEasing() override;
 
 #if WITH_EDITORONLY_DATA
@@ -107,10 +95,6 @@ private:
 	void AutoMatchSectionRoot(UMovieSceneSkeletalAnimationSection* AnimSection);
 public:
 
-	/** List of all animation sections */
-	UPROPERTY()
-	TArray<TObjectPtr<UMovieSceneSection>> AnimationSections;
-
 	UPROPERTY()
 	bool bUseLegacySectionIndexBlend;
 
