Build fix from small merge mixtake in 40941769

diff --git a/Engine/Source/Editor/MovieSceneTools/Private/TrackEditors/SkeletalAnimationTrackEditor.cpp b/Engine/Source/Editor/MovieSceneTools/Private/TrackEditors/SkeletalAnimationTrackEditor.cpp
--- a/Engine/Source/Editor/MovieSceneTools/Private/TrackEditors/SkeletalAnimationTrackEditor.cpp
+++ b/Engine/Source/Editor/MovieSceneTools/Private/TrackEditors/SkeletalAnimationTrackEditor.cpp
@@ -53,7 +53,7 @@ bool FSkeletalAnimationTrackEditor::SupportsType( TSubclassOf<UMovieSceneTrack>
 
 void FSkeletalAnimationTrackEditor::OnPostPropertyChanged(UObject* InObject, FPropertyChangedEvent& InPropertyChangedEvent)
 {
-	if (SkeletalAnimationEditorConstants::CVarEvaluateSkeletalMeshOnPropertyChange->GetBool())
+	if (UE::Sequencer::SkeletalAnimationEditorConstants::CVarEvaluateSkeletalMeshOnPropertyChange->GetBool())
 	{
 		const FName PropertyName = (InPropertyChangedEvent.Property != nullptr) ? InPropertyChangedEvent.Property->GetFName() : NAME_None;
 		if (PropertyName == USceneComponent::GetRelativeLocationPropertyName() ||
