Sequencer: Fixed undo-redo not working for channels with owner objects that were not sections

diff --git a/Engine/Source/Editor/MovieSceneTools/Private/Channels/BuiltInChannelEditors.cpp b/Engine/Source/Editor/MovieSceneTools/Private/Channels/BuiltInChannelEditors.cpp
--- a/Engine/Source/Editor/MovieSceneTools/Private/Channels/BuiltInChannelEditors.cpp
+++ b/Engine/Source/Editor/MovieSceneTools/Private/Channels/BuiltInChannelEditors.cpp
@@ -316,7 +316,7 @@ TSharedRef<SWidget> CreateKeyEditor(const TMovieSceneChannelHandle<FMovieSceneBo
 		Params.ObjectBindingID, Channel,
 		Params.OwningSection, Params.Sequencer, Params.PropertyBindings, ExternalValue->OnGetExternalValue
 		);
-
+	KeyEditor.SetOwningObject(Cast<UMovieSceneSignedObject>(Params.OwningObject));
 	return SNew(SBoolCurveKeyEditor, KeyEditor);
 }
 
@@ -333,6 +333,7 @@ TSharedRef<SWidget> CreateKeyEditor(const TMovieSceneChannelHandle<FMovieSceneIn
 		Params.ObjectBindingID, Channel,
 		Params.OwningSection, Params.Sequencer, Params.PropertyBindings, ExternalValue->OnGetExternalValue
 		);
+	KeyEditor.SetOwningObject(Cast<UMovieSceneSignedObject>(Params.OwningObject));
 
 	typedef SNumericKeyEditor<FMovieSceneIntegerChannel, int32> KeyEditorType;
 	return SNew(KeyEditorType, KeyEditor);
@@ -351,6 +352,7 @@ TSharedRef<SWidget> CreateKeyEditor(const TMovieSceneChannelHandle<FMovieSceneFl
 		Params.ObjectBindingID, Channel,
 		Params.OwningSection, Params.Sequencer, Params.PropertyBindings, ExternalValue->OnGetExternalValue
 		);
+	KeyEditor.SetOwningObject(Cast<UMovieSceneSignedObject>(Params.OwningObject));
 
 	typedef SNumericKeyEditor<FMovieSceneFloatChannel, float> KeyEditorType;
 	return SNew(KeyEditorType, KeyEditor);
@@ -369,6 +371,7 @@ TSharedRef<SWidget> CreateKeyEditor(const TMovieSceneChannelHandle<FMovieSceneDo
 		Params.ObjectBindingID, Channel,
 		Params.OwningSection, Params.Sequencer, Params.PropertyBindings, ExternalValue
 		);
+	KeyEditor.SetOwningObject(Cast<UMovieSceneSignedObject>(Params.OwningObject));
 
 	typedef SNumericKeyEditor<FMovieSceneDoubleChannel, double> KeyEditorType;
 	return SNew(KeyEditorType, KeyEditor);
@@ -387,6 +390,7 @@ TSharedRef<SWidget> CreateKeyEditor(const TMovieSceneChannelHandle<FMovieSceneSt
 		Params.ObjectBindingID, Channel,
 		Params.OwningSection, Params.Sequencer, Params.PropertyBindings, ExternalValue->OnGetExternalValue
 		);
+	KeyEditor.SetOwningObject(Cast<UMovieSceneSignedObject>(Params.OwningObject));
 
 	return SNew(SStringCurveKeyEditor, KeyEditor);
 }
@@ -405,6 +409,7 @@ TSharedRef<SWidget> CreateKeyEditor(const TMovieSceneChannelHandle<FMovieSceneBy
 		Params.ObjectBindingID, Channel,
 		Params.OwningSection, Params.Sequencer, Params.PropertyBindings, ExternalValue->OnGetExternalValue
 		);
+	KeyEditor.SetOwningObject(Cast<UMovieSceneSignedObject>(Params.OwningObject));
 
 	if (UEnum* Enum = RawChannel->GetEnum())
 	{
@@ -426,6 +431,7 @@ TSharedRef<SWidget> CreateKeyEditor(const TMovieSceneChannelHandle<FMovieSceneOb
 	if (ExternalValue && RawChannel)
 	{
 		TSequencerKeyEditor<FMovieSceneObjectPathChannel, UObject*> KeyEditor(Params.ObjectBindingID, Channel, Params.OwningSection, Params.Sequencer, Params.PropertyBindings, ExternalValue->OnGetExternalValue);
+		KeyEditor.SetOwningObject(Cast<UMovieSceneSignedObject>(Params.OwningObject));
 
 		UClass* PropertyClass = ObjectPathTrack ? ObjectPathTrack->PropertyClass : nullptr;
 		const bool bClassPicker = ObjectPathTrack ? ObjectPathTrack->bClassProperty : false;
