Anim In Engine: Create Helper to be used when baking. Currently will just turn on/off evaluate sub sequences in isolation but could be used for other stuff in the future. Use it for snapper and anim baking currently.

diff --git a/Engine/Plugins/Animation/ControlRig/Source/ControlRigEditor/Private/Tools/ControlRigSnapper.cpp b/Engine/Plugins/Animation/ControlRig/Source/ControlRigEditor/Private/Tools/ControlRigSnapper.cpp
--- a/Engine/Plugins/Animation/ControlRig/Source/ControlRigEditor/Private/Tools/ControlRigSnapper.cpp
+++ b/Engine/Plugins/Animation/ControlRig/Source/ControlRigEditor/Private/Tools/ControlRigSnapper.cpp
@@ -32,6 +32,7 @@
 #include "Containers/SortedMap.h"
 #include "Sequencer/ControlRigSequencerHelpers.h"
 #include "Sequencer/EditModeAnimationUtil.h"
+#include "SequencerBakingSetupRestore.h"
 #include "Transform/TransformConstraint.h"
 #include "Transform/TransformConstraintUtil.h"
 #include "Constraints/MovieSceneConstraintChannelHelper.inl"
@@ -526,6 +527,7 @@ bool FControlRigSnapper::SnapIt(FFrameNumber StartFrame, FFrameNumber EndFrame,c
 		TSharedPtr<ISequencer> Sequencer = InSequencer.Pin();
 		UMovieScene* MovieScene = Sequencer->GetFocusedMovieSceneSequence()->GetMovieScene();
 
+	
 		TArray<FFrameNumber> Frames;
 		if (SnapSettings && SnapSettings->BakingKeySettings == EBakingKeySettings::AllFrames)
 		{
@@ -553,9 +555,12 @@ bool FControlRigSnapper::SnapIt(FFrameNumber StartFrame, FFrameNumber EndFrame,c
 			return false;
 		}
 
-		FScopedTransaction ScopedTransaction(LOCTEXT("SnapAnimation", "Snap Animation"));
+		//if in subsquence and we want to, turn on ShouldEvaluateSubSequencesInIsolation, use the following object for that
 
-		
+		UE::Sequencer::FSequencerBakingSetupRestore RestoreBaking(Sequencer);
+
+		FScopedTransaction ScopedTransaction(LOCTEXT("SnapAnimation", "Snap Animation"));
+	
 		MovieScene->Modify();
 
 		//need to make sure binding is setup
diff --git a/Engine/Source/Editor/MovieSceneTools/Private/TrackEditors/CommonAnimationTrackEditor.cpp b/Engine/Source/Editor/MovieSceneTools/Private/TrackEditors/CommonAnimationTrackEditor.cpp
--- a/Engine/Source/Editor/MovieSceneTools/Private/TrackEditors/CommonAnimationTrackEditor.cpp
+++ b/Engine/Source/Editor/MovieSceneTools/Private/TrackEditors/CommonAnimationTrackEditor.cpp
@@ -25,7 +25,7 @@
 #include "AssetRegistry/AssetRegistryModule.h"
 #include "IContentBrowserSingleton.h"
 #include "ContentBrowserModule.h"
-#include "SequencerSettings.h"
+#include "SequencerBakingSetupRestore.h"
 #include "MVVM/Views/ViewUtilities.h"
 #include "MVVM/ViewModels/ViewDensity.h"
 #include "MVVM/Extensions/ITrackExtension.h"
@@ -107,15 +107,6 @@ namespace CommonAnimationEditorConstants
 {
 	// @todo Sequencer Allow this to be customizable
 	const uint32 AnimationTrackHeight = 28;
-
-	const TCHAR* CVarName = TEXT("Sequencer.SubsequenceAnimBakeInIsolation");
-
-	static bool bSubsequenceAnimBakeInIsolation = false; //default to false for now will turn it on after more testing
-	static FAutoConsoleVariableRef CVarSetSubsequenceAnimBakeInIsolation(
-		CVarName,
-		bSubsequenceAnimBakeInIsolation,
-		TEXT("Automatically Force Evaluate Subsequence In Isolation when Baking Anim Sequences")
-	);
 }
 
 
@@ -1468,127 +1459,118 @@ void FCommonAnimationTrackEditor::OnSequencerSaved(ISequencer& )
 			ULevelSequenceAnimSequenceLink* LevelAnimLink = AssetUserDataInterface->GetAssetUserData< ULevelSequenceAnimSequenceLink >();
 			if (LevelAnimLink && LevelAnimLink->AnimSequenceLinks.Num() > 0)
 			{
-				
 
 				UMovieSceneSequence* MovieSceneSequence = SequencerPtr->GetFocusedMovieSceneSequence();
 				UMovieSceneSequence* RootMovieSceneSequence = SequencerPtr->GetRootMovieSceneSequence();
 				FMovieSceneSequenceIDRef Template = SequencerPtr->GetFocusedTemplateID();
 				FMovieSceneSequenceTransform RootToLocalTransform = SequencerPtr->GetFocusedMovieSceneSequenceTransform();
-				//if in subsquence and we want to, turn on ShouldEvaluateSubSequencesInIsolation
-				TOptional<bool> bRestoreShouldEvaluateSubSequencesInIsolation;
-				if (CommonAnimationEditorConstants::bSubsequenceAnimBakeInIsolation &&
-					MovieSceneSequence != RootMovieSceneSequence &&
-					SequencerPtr->GetSequencerSettings()->ShouldEvaluateSubSequencesInIsolation() == false)
-				{
-					bRestoreShouldEvaluateSubSequencesInIsolation = false;
-					SequencerPtr->GetSequencerSettings()->SetEvaluateSubSequencesInIsolation(true);
-				}
-				for (int32 Index = LevelAnimLink->AnimSequenceLinks.Num() -1; Index >=0 ; --Index)
+				//if in subsquence and we want to, turn on ShouldEvaluateSubSequencesInIsolation, use the following object for that
 				{
-					FLevelSequenceAnimSequenceLinkItem& Item = LevelAnimLink->AnimSequenceLinks[Index];
-					UAnimSequence* AnimSequence = Item.ResolveAnimSequence();
-					if (AnimSequence == nullptr)
-					{
-						LevelAnimLink->AnimSequenceLinks.RemoveAt(Index);
-						continue;
-					}
-					if (IInterface_AssetUserData* AnimAssetUserData = Cast< IInterface_AssetUserData >(AnimSequence))
+					UE::Sequencer::FSequencerBakingSetupRestore RestoreBaking(SequencerPtr);
+					for (int32 Index = LevelAnimLink->AnimSequenceLinks.Num() - 1; Index >= 0; --Index)
 					{
-						UAnimSequenceLevelSequenceLink* AnimLevelLink = AnimAssetUserData->GetAssetUserData< UAnimSequenceLevelSequenceLink >();
-						if (!AnimLevelLink)
+						FLevelSequenceAnimSequenceLinkItem& Item = LevelAnimLink->AnimSequenceLinks[Index];
+						UAnimSequence* AnimSequence = Item.ResolveAnimSequence();
+						if (AnimSequence == nullptr)
 						{
-							AnimLevelLink = NewObject<UAnimSequenceLevelSequenceLink>(AnimSequence, NAME_None, RF_Public | RF_Transactional);
-							AnimAssetUserData->AddAssetUserData(AnimLevelLink);
+							LevelAnimLink->AnimSequenceLinks.RemoveAt(Index);
+							continue;
+						}
+						if (IInterface_AssetUserData* AnimAssetUserData = Cast< IInterface_AssetUserData >(AnimSequence))
+						{
+							UAnimSequenceLevelSequenceLink* AnimLevelLink = AnimAssetUserData->GetAssetUserData< UAnimSequenceLevelSequenceLink >();
+							if (!AnimLevelLink)
+							{
+								AnimLevelLink = NewObject<UAnimSequenceLevelSequenceLink>(AnimSequence, NAME_None, RF_Public | RF_Transactional);
+								AnimAssetUserData->AddAssetUserData(AnimLevelLink);
+							}
+							AnimLevelLink->SetLevelSequence(LevelSequence);
+							AnimLevelLink->SkelTrackGuid = Item.SkelTrackGuid;
+						}
+						USkeletalMeshComponent* SkelMeshComp = UE::Sequencer::FCommonAnimationTrackEditor::AcquireSkeletalMeshFromObjectGuid(Item.SkelTrackGuid, GetSequencer());
+						if (AnimSequence && SkelMeshComp)
+						{
+							const bool bSavedExportMorphTargets = AnimSeqExportOption->bExportMorphTargets;
+							const bool bSavedExportAttributeCurves = AnimSeqExportOption->bExportAttributeCurves;
+							const bool bSavedExportMaterialCurves = AnimSeqExportOption->bExportMaterialCurves;
+							const bool bSavedExportTransforms = AnimSeqExportOption->bExportTransforms;
+							const bool bSavedIncludeComponentTransform = AnimSeqExportOption->bRecordInWorldSpace;
+							const bool bSavedEvaluateAllSkeletalMeshComponents = AnimSeqExportOption->bEvaluateAllSkeletalMeshComponents;
+							const EAnimInterpolationType SavedInterpolationType = AnimSeqExportOption->Interpolation;
+							const ERichCurveInterpMode SavedCurveInterpolationType = AnimSeqExportOption->CurveInterpolation;
+							const TArray<FString> SavedIncludeAnimationNames = AnimSeqExportOption->IncludeAnimationNames;
+							const TArray<FString> SavedExcludeAnimationNames = AnimSeqExportOption->ExcludeAnimationNames;
+							const FFrameNumber SavedWarmUpFrames = AnimSeqExportOption->WarmUpFrames;
+							const FFrameNumber SavedDelayBeforeStart = AnimSeqExportOption->DelayBeforeStart;
+							const bool bSavedUseCustomTimeRange = AnimSeqExportOption->bUseCustomTimeRange;
+							const FFrameNumber SavedCustomStartFrame = AnimSeqExportOption->CustomStartFrame;
+							const FFrameNumber SavedCustomEndFrame = AnimSeqExportOption->CustomEndFrame;
+							const FFrameRate SavedCustomDisplayRate = AnimSeqExportOption->CustomDisplayRate;
+							const bool bSavedUseCustomFrameRate = AnimSeqExportOption->bUseCustomFrameRate;
+							const FFrameRate SavedCustomFrameRate = AnimSeqExportOption->CustomFrameRate;
+
+							AnimSeqExportOption->bExportMorphTargets = Item.bExportMorphTargets;
+							AnimSeqExportOption->bExportAttributeCurves = Item.bExportAttributeCurves;
+							AnimSeqExportOption->bExportMaterialCurves = Item.bExportMaterialCurves;
+							AnimSeqExportOption->bExportTransforms = Item.bExportTransforms;
+							AnimSeqExportOption->bRecordInWorldSpace = Item.bRecordInWorldSpace;
+							AnimSeqExportOption->bEvaluateAllSkeletalMeshComponents = Item.bEvaluateAllSkeletalMeshComponents;
+							AnimSeqExportOption->Interpolation = Item.Interpolation;
+							AnimSeqExportOption->CurveInterpolation = Item.CurveInterpolation;
+
+							AnimSeqExportOption->IncludeAnimationNames = Item.IncludeAnimationNames;
+							AnimSeqExportOption->ExcludeAnimationNames = Item.ExcludeAnimationNames;
+							AnimSeqExportOption->WarmUpFrames = Item.WarmUpFrames;
+							AnimSeqExportOption->DelayBeforeStart = Item.DelayBeforeStart;
+							AnimSeqExportOption->bUseCustomTimeRange = Item.bUseCustomTimeRange;
+							AnimSeqExportOption->CustomStartFrame = Item.CustomStartFrame;
+							AnimSeqExportOption->CustomEndFrame = Item.CustomEndFrame;
+							AnimSeqExportOption->CustomDisplayRate = Item.CustomDisplayRate;
+							AnimSeqExportOption->bUseCustomFrameRate = Item.bUseCustomFrameRate;
+							AnimSeqExportOption->CustomFrameRate = Item.CustomFrameRate;
+
+							FAnimExportSequenceParameters AESP;
+							AESP.Player = SequencerPtr.Get();
+							AESP.RootToLocalTransform = RootToLocalTransform;
+							AESP.MovieSceneSequence = MovieSceneSequence;
+							AESP.RootMovieSceneSequence = RootMovieSceneSequence;
+							AESP.bForceUseOfMovieScenePlaybackRange = SequencerPtr.Get()->GetSequencerSettings()->ShouldEvaluateSubSequencesInIsolation();
+							//make sure all spanwables are present
+							SequencerPtr->ForceEvaluate();
+							bool bResult = MovieSceneToolHelpers::ExportToAnimSequence(AnimSequence, AnimSeqExportOption, AESP, SkelMeshComp);
+
+							AnimSeqExportOption->bExportMorphTargets = bSavedExportMorphTargets;
+							AnimSeqExportOption->bExportAttributeCurves = bSavedExportAttributeCurves;
+							AnimSeqExportOption->bExportMaterialCurves = bSavedExportMaterialCurves;
+							AnimSeqExportOption->bExportTransforms = bSavedExportTransforms;
+							AnimSeqExportOption->bRecordInWorldSpace = bSavedIncludeComponentTransform;
+							AnimSeqExportOption->bEvaluateAllSkeletalMeshComponents = bSavedEvaluateAllSkeletalMeshComponents;
+							AnimSeqExportOption->Interpolation = SavedInterpolationType;
+							AnimSeqExportOption->CurveInterpolation = SavedCurveInterpolationType;
+
+							AnimSeqExportOption->IncludeAnimationNames = SavedIncludeAnimationNames;
+							AnimSeqExportOption->ExcludeAnimationNames = SavedExcludeAnimationNames;
+							AnimSeqExportOption->WarmUpFrames = SavedWarmUpFrames;
+							AnimSeqExportOption->DelayBeforeStart = SavedDelayBeforeStart;
+							AnimSeqExportOption->bUseCustomTimeRange = bSavedUseCustomTimeRange;
+							AnimSeqExportOption->CustomStartFrame = SavedCustomStartFrame;
+							AnimSeqExportOption->CustomEndFrame = SavedCustomEndFrame;
+							AnimSeqExportOption->CustomDisplayRate = SavedCustomDisplayRate;
+							AnimSeqExportOption->bUseCustomFrameRate = bSavedUseCustomFrameRate;
+							AnimSeqExportOption->CustomFrameRate = SavedCustomFrameRate;
+							//save the anim sequence to disk to make sure they are in sync
+							UPackage* const Package = AnimSequence->GetOutermost();
+							FString const PackageName = Package->GetName();
+							FString const PackageFileName = FPackageName::LongPackageNameToFilename(PackageName, FPackageName::GetAssetPackageExtension());
+
+							FSavePackageArgs SaveArgs;
+							SaveArgs.TopLevelFlags = RF_Standalone;
+							SaveArgs.SaveFlags = SAVE_NoError;
+							UPackage::SavePackage(Package, NULL, *PackageFileName, SaveArgs);
 						}
-						AnimLevelLink->SetLevelSequence(LevelSequence);
-						AnimLevelLink->SkelTrackGuid = Item.SkelTrackGuid;
-					}
-					USkeletalMeshComponent* SkelMeshComp = UE::Sequencer::FCommonAnimationTrackEditor::AcquireSkeletalMeshFromObjectGuid(Item.SkelTrackGuid, GetSequencer());
-					if (AnimSequence && SkelMeshComp)
-					{
-						const bool bSavedExportMorphTargets = AnimSeqExportOption->bExportMorphTargets;
-						const bool bSavedExportAttributeCurves = AnimSeqExportOption->bExportAttributeCurves;
-						const bool bSavedExportMaterialCurves = AnimSeqExportOption->bExportMaterialCurves;
-						const bool bSavedExportTransforms = AnimSeqExportOption->bExportTransforms;
-						const bool bSavedIncludeComponentTransform = AnimSeqExportOption->bRecordInWorldSpace;
-						const bool bSavedEvaluateAllSkeletalMeshComponents = AnimSeqExportOption->bEvaluateAllSkeletalMeshComponents;
-						const EAnimInterpolationType SavedInterpolationType = AnimSeqExportOption->Interpolation;
-						const ERichCurveInterpMode SavedCurveInterpolationType = AnimSeqExportOption->CurveInterpolation;
-						const TArray<FString> SavedIncludeAnimationNames = AnimSeqExportOption->IncludeAnimationNames;
-						const TArray<FString> SavedExcludeAnimationNames = AnimSeqExportOption->ExcludeAnimationNames;
-						const FFrameNumber SavedWarmUpFrames = AnimSeqExportOption->WarmUpFrames;
-						const FFrameNumber SavedDelayBeforeStart = AnimSeqExportOption->DelayBeforeStart;
-						const bool bSavedUseCustomTimeRange = AnimSeqExportOption->bUseCustomTimeRange;
-						const FFrameNumber SavedCustomStartFrame = AnimSeqExportOption->CustomStartFrame;
-						const FFrameNumber SavedCustomEndFrame = AnimSeqExportOption->CustomEndFrame;
-						const FFrameRate SavedCustomDisplayRate = AnimSeqExportOption->CustomDisplayRate;
-						const bool bSavedUseCustomFrameRate = AnimSeqExportOption->bUseCustomFrameRate;
-						const FFrameRate SavedCustomFrameRate = AnimSeqExportOption->CustomFrameRate;
-
-						AnimSeqExportOption->bExportMorphTargets = Item.bExportMorphTargets;
-						AnimSeqExportOption->bExportAttributeCurves = Item.bExportAttributeCurves;
-						AnimSeqExportOption->bExportMaterialCurves = Item.bExportMaterialCurves;
-						AnimSeqExportOption->bExportTransforms = Item.bExportTransforms;
-						AnimSeqExportOption->bRecordInWorldSpace = Item.bRecordInWorldSpace;
-						AnimSeqExportOption->bEvaluateAllSkeletalMeshComponents = Item.bEvaluateAllSkeletalMeshComponents;
-						AnimSeqExportOption->Interpolation = Item.Interpolation;
-						AnimSeqExportOption->CurveInterpolation = Item.CurveInterpolation;
-
-						AnimSeqExportOption->IncludeAnimationNames = Item.IncludeAnimationNames;
-						AnimSeqExportOption->ExcludeAnimationNames = Item.ExcludeAnimationNames;
-						AnimSeqExportOption->WarmUpFrames = Item.WarmUpFrames;
-						AnimSeqExportOption->DelayBeforeStart = Item.DelayBeforeStart;
-						AnimSeqExportOption->bUseCustomTimeRange = Item.bUseCustomTimeRange;
-						AnimSeqExportOption->CustomStartFrame = Item.CustomStartFrame;
-						AnimSeqExportOption->CustomEndFrame = Item.CustomEndFrame;
-						AnimSeqExportOption->CustomDisplayRate = Item.CustomDisplayRate;
-						AnimSeqExportOption->bUseCustomFrameRate = Item.bUseCustomFrameRate;
-						AnimSeqExportOption->CustomFrameRate = Item.CustomFrameRate;
-
-						FAnimExportSequenceParameters AESP;
-						AESP.Player = SequencerPtr.Get();
-						AESP.RootToLocalTransform = RootToLocalTransform;
-						AESP.MovieSceneSequence = MovieSceneSequence;
-						AESP.RootMovieSceneSequence = RootMovieSceneSequence;
-						AESP.bForceUseOfMovieScenePlaybackRange = SequencerPtr.Get()->GetSequencerSettings()->ShouldEvaluateSubSequencesInIsolation();
-						//make sure all spanwables are present
-						SequencerPtr->ForceEvaluate();
-						bool bResult = MovieSceneToolHelpers::ExportToAnimSequence(AnimSequence, AnimSeqExportOption, AESP, SkelMeshComp);
-
-						AnimSeqExportOption->bExportMorphTargets = bSavedExportMorphTargets;
-						AnimSeqExportOption->bExportAttributeCurves = bSavedExportAttributeCurves;
-						AnimSeqExportOption->bExportMaterialCurves = bSavedExportMaterialCurves;
-						AnimSeqExportOption->bExportTransforms = bSavedExportTransforms;
-						AnimSeqExportOption->bRecordInWorldSpace = bSavedIncludeComponentTransform;
-						AnimSeqExportOption->bEvaluateAllSkeletalMeshComponents = bSavedEvaluateAllSkeletalMeshComponents;
-						AnimSeqExportOption->Interpolation = SavedInterpolationType;
-						AnimSeqExportOption->CurveInterpolation = SavedCurveInterpolationType;
-
-						AnimSeqExportOption->IncludeAnimationNames = SavedIncludeAnimationNames;
-						AnimSeqExportOption->ExcludeAnimationNames = SavedExcludeAnimationNames;
-						AnimSeqExportOption->WarmUpFrames = SavedWarmUpFrames;
-						AnimSeqExportOption->DelayBeforeStart = SavedDelayBeforeStart;
-						AnimSeqExportOption->bUseCustomTimeRange = bSavedUseCustomTimeRange;
-						AnimSeqExportOption->CustomStartFrame = SavedCustomStartFrame;
-						AnimSeqExportOption->CustomEndFrame = SavedCustomEndFrame;
-						AnimSeqExportOption->CustomDisplayRate = SavedCustomDisplayRate;
-						AnimSeqExportOption->bUseCustomFrameRate = bSavedUseCustomFrameRate;
-						AnimSeqExportOption->CustomFrameRate = SavedCustomFrameRate;
-						//save the anim sequence to disk to make sure they are in sync
-						UPackage* const Package = AnimSequence->GetOutermost();
-						FString const PackageName = Package->GetName();
-						FString const PackageFileName = FPackageName::LongPackageNameToFilename(PackageName, FPackageName::GetAssetPackageExtension());
-
-						FSavePackageArgs SaveArgs;
-						SaveArgs.TopLevelFlags = RF_Standalone;
-						SaveArgs.SaveFlags = SAVE_NoError;
-						UPackage::SavePackage(Package, NULL, *PackageFileName, SaveArgs);
 					}
-				}
-				if (bRestoreShouldEvaluateSubSequencesInIsolation.IsSet())
-				{
-					SequencerPtr->GetSequencerSettings()->SetEvaluateSubSequencesInIsolation(false);
-				}
+				}//restore bake settings
+				
 				//re-evaluate at current frame
 				SequencerPtr->ForceEvaluate();
 			}
@@ -1668,29 +1650,21 @@ bool FCommonAnimationTrackEditor::CreateAnimationSequence(const TArray<UObject*>
 
 			if (OptionWindow->ShouldExport())
 			{
-				const TSharedPtr<ISequencer> ParentSequencer = GetSequencer();
+				TSharedPtr<ISequencer> ParentSequencer = GetSequencer();
 				UMovieSceneSequence* MovieSceneSequence = ParentSequencer->GetFocusedMovieSceneSequence();
 				UMovieSceneSequence* RootMovieSceneSequence = ParentSequencer->GetRootMovieSceneSequence();
 				FMovieSceneSequenceTransform RootToLocalTransform = ParentSequencer->GetFocusedMovieSceneSequenceTransform();
-				TOptional<bool> bRestoreShouldEvaluateSubSequencesInIsolation;
-				if (CommonAnimationEditorConstants::bSubsequenceAnimBakeInIsolation &&
-					MovieSceneSequence != RootMovieSceneSequence &&
-					ParentSequencer->GetSequencerSettings()->ShouldEvaluateSubSequencesInIsolation() == false)
-				{
-					bRestoreShouldEvaluateSubSequencesInIsolation = false;
-					ParentSequencer->GetSequencerSettings()->SetEvaluateSubSequencesInIsolation(true);
-				}
-				FAnimExportSequenceParameters AESP;
-				AESP.Player = ParentSequencer.Get();
-				AESP.RootToLocalTransform = RootToLocalTransform;
-				AESP.MovieSceneSequence = MovieSceneSequence;
-				AESP.RootMovieSceneSequence = RootMovieSceneSequence;
-				AESP.bForceUseOfMovieScenePlaybackRange = ParentSequencer->GetSequencerSettings()->ShouldEvaluateSubSequencesInIsolation();
-				AnimSeqExportOption->CustomDisplayRate = ParentSequencer->GetFocusedDisplayRate();
-				bResult  = MovieSceneToolHelpers::ExportToAnimSequence(AnimSequence, AnimSeqExportOption, AESP, SkelMeshComp);
-				if (bRestoreShouldEvaluateSubSequencesInIsolation.IsSet())
+				//if in subsquence and we want to, turn on ShouldEvaluateSubSequencesInIsolation, use the following object for that
 				{
-					ParentSequencer->GetSequencerSettings()->SetEvaluateSubSequencesInIsolation(false);
+					UE::Sequencer::FSequencerBakingSetupRestore RestoreBaking(ParentSequencer);
+					FAnimExportSequenceParameters AESP;
+					AESP.Player = ParentSequencer.Get();
+					AESP.RootToLocalTransform = RootToLocalTransform;
+					AESP.MovieSceneSequence = MovieSceneSequence;
+					AESP.RootMovieSceneSequence = RootMovieSceneSequence;
+					AESP.bForceUseOfMovieScenePlaybackRange = ParentSequencer->GetSequencerSettings()->ShouldEvaluateSubSequencesInIsolation();
+					AnimSeqExportOption->CustomDisplayRate = ParentSequencer->GetFocusedDisplayRate();
+					bResult = MovieSceneToolHelpers::ExportToAnimSequence(AnimSequence, AnimSeqExportOption, AESP, SkelMeshComp);
 				}
 				//re-evaluate at current frame
 				ParentSequencer->ForceEvaluate();
diff --git a/Engine/Source/Editor/Sequencer/Private/SequencerBakingSetupRestore.cpp b/Engine/Source/Editor/Sequencer/Private/SequencerBakingSetupRestore.cpp
new file mode 100644
--- /dev/null
+++ b/Engine/Source/Editor/Sequencer/Private/SequencerBakingSetupRestore.cpp
@@ -0,0 +1,43 @@
+// Copyright Epic Games, Inc. All Rights Reserved.
+
+#include "SequencerBakingSetupRestore.h"
+#include "ISequencer.h"
+#include "SequencerSettings.h"
+
+namespace UE::Sequencer
+{
+	const TCHAR* CVarName = TEXT("Sequencer.SubsequenceAnimBakeInIsolation");
+
+	static bool bSubSequenceBakingInIsolation = true; //default to false for now will turn it on after more testing
+	static FAutoConsoleVariableRef CVarSetIsSequencerBakingInIsolation(
+		CVarName,
+		bSubSequenceBakingInIsolation,
+		TEXT("Automatically Force Evaluate Subsequence In Isolation when Baking")
+	);
+
+	FSequencerBakingSetupRestore::FSequencerBakingSetupRestore(TSharedPtr<ISequencer>& SequencerPtr)
+	{
+		if (!SequencerPtr || bSubSequenceBakingInIsolation == false)
+		{
+			return;
+		}
+		WeakSequencer = SequencerPtr;
+		UMovieSceneSequence* MovieSceneSequence = SequencerPtr->GetFocusedMovieSceneSequence();
+		UMovieSceneSequence* RootMovieSceneSequence = SequencerPtr->GetRootMovieSceneSequence();
+		if (MovieSceneSequence != RootMovieSceneSequence &&
+			SequencerPtr->GetSequencerSettings()->ShouldEvaluateSubSequencesInIsolation() == false)
+		{
+			bRestoreShouldEvaluateSubSequencesInIsolation = false;
+			SequencerPtr->GetSequencerSettings()->SetEvaluateSubSequencesInIsolation(true);
+		}
+	}
+	FSequencerBakingSetupRestore::~FSequencerBakingSetupRestore()
+	{
+		if (WeakSequencer.IsValid() && bRestoreShouldEvaluateSubSequencesInIsolation.IsSet())
+		{
+			TSharedPtr<ISequencer> SequencerPtr = WeakSequencer.Pin();
+			SequencerPtr->GetSequencerSettings()->SetEvaluateSubSequencesInIsolation(false);
+		}
+	}
+
+} // namespace UE::Sequencer
\ No newline at end of file
diff --git a/Engine/Source/Editor/Sequencer/Public/SequencerBakingSetupRestore.h b/Engine/Source/Editor/Sequencer/Public/SequencerBakingSetupRestore.h
new file mode 100644
--- /dev/null
+++ b/Engine/Source/Editor/Sequencer/Public/SequencerBakingSetupRestore.h
@@ -0,0 +1,27 @@
+// Copyright Epic Games, Inc. All Rights Reserved.
+
+#pragma once
+
+#include "CoreTypes.h"
+#include "Templates/SharedPointer.h"
+#include "Misc/Optional.h"
+
+class ISequencer;
+
+namespace UE::Sequencer
+{
+
+	class SEQUENCER_API FSequencerBakingSetupRestore 
+	{
+	public:
+		FSequencerBakingSetupRestore() = delete;
+		FSequencerBakingSetupRestore(TSharedPtr<ISequencer>& SequencerPtr);
+		~FSequencerBakingSetupRestore();
+		
+	private:
+		TWeakPtr<ISequencer> WeakSequencer;
+		TOptional<bool> bRestoreShouldEvaluateSubSequencesInIsolation;
+		
+	};
+
+} // namespace UE::Sequencer
\ No newline at end of file
