Sequencer: Expanded root motion controls when using the Animation Mixer Track

There are now 3 modes for the root motion control on an animation within the mixer track: Asset, Offset and Override.

- Asset: Takes the root motion from the asset directly.
- Offset: Takes the root motion from the asset, and offsets it by the channel values. New: there is an additional origin location around which the offset is transformed, this allows offsetting using a specific reference frame from the animation.
- Override: Ignores any root motion from the asset, and use the channel values directly.

These new settings complement the existing Animation/World space controls and allow for precise positioning of animation data that was authored in world space.

diff --git a/Engine/Plugins/Experimental/MovieSceneAnimMixer/Source/MovieSceneAnimMixer/Internal/MovieSceneAnimationMixerTrack.h b/Engine/Plugins/Experimental/MovieSceneAnimMixer/Source/MovieSceneAnimMixer/Internal/MovieSceneAnimationMixerTrack.h
--- a/Engine/Plugins/Experimental/MovieSceneAnimMixer/Source/MovieSceneAnimMixer/Internal/MovieSceneAnimationMixerTrack.h
+++ b/Engine/Plugins/Experimental/MovieSceneAnimMixer/Source/MovieSceneAnimMixer/Internal/MovieSceneAnimationMixerTrack.h
@@ -25,6 +25,42 @@ enum class EMovieSceneRootMotionSpace : uint8
 	WorldSpace
 };
 
+
+// Enum describing what space that root motion should be applied in.
+UENUM()
+enum class EMovieSceneRootMotionTransformMode : uint8
+{
+	Asset,
+	Offset,
+	Override,
+};
+
+
+USTRUCT()
+struct FMovieSceneByteChannelDefaultOnly : public FMovieSceneByteChannel
+{
+	GENERATED_BODY()
+
+	bool SerializeFromMismatchedTag(const FPropertyTag& Tag, FStructuredArchive::FSlot Slot);
+};
+
+template<>
+struct TStructOpsTypeTraits<FMovieSceneByteChannelDefaultOnly> : public TStructOpsTypeTraitsBase2<FMovieSceneByteChannelDefaultOnly>
+{
+	enum { WithStructuredSerializeFromMismatchedTag = true };
+};
+
+template<>
+struct TMovieSceneChannelTraits<FMovieSceneByteChannelDefaultOnly> : TMovieSceneChannelTraitsBase<FMovieSceneByteChannelDefaultOnly>
+{
+#if WITH_EDITOR
+
+	/** Byte channels can have external values (ie, they can get their values from external objects for UI purposes) */
+	typedef TMovieSceneExternalValue<uint8> ExtendedEditorDataType;
+
+#endif
+};
+
 UINTERFACE(MinimalAPI)
 class UMovieSceneAnimationSectionInterface : public UInterface
 {
@@ -63,10 +99,14 @@ class UMovieSceneAnimationBaseTransformDecoration
 	, public IMovieSceneChannelDecoration
 	, public IMovieSceneEntityDecorator
 {
+public:
+
 	GENERATED_BODY()
 
 	UMovieSceneAnimationBaseTransformDecoration(const FObjectInitializer& ObjInit);
 
+	MOVIESCENEANIMMIXER_API EMovieSceneRootMotionTransformMode GetRootTransformMode() const;
+
 	virtual EMovieSceneChannelProxyType PopulateChannelProxy(FMovieSceneChannelProxyData& OutProxyData) override;
 	virtual void ExtendEntityImpl(UMovieSceneEntitySystemLinker* EntityLinker, const FEntityImportParams& Params, FImportedEntity* OutImportedEntity) override;
 
@@ -76,8 +116,14 @@ class UMovieSceneAnimationBaseTransformDecoration
 	UPROPERTY()
 	FMovieSceneDoubleChannel Rotation[3];
 
+	UPROPERTY(EditAnywhere, Category="Root Transform")
+	FVector RootOriginLocation;
+
+	UPROPERTY()
+	FMovieSceneByteChannelDefaultOnly RootMotionSpace;
+
 	UPROPERTY()
-	FMovieSceneByteChannel RootMotionSpace;
+	FMovieSceneByteChannelDefaultOnly TransformMode;
 };
 
 
diff --git a/Engine/Plugins/Experimental/MovieSceneAnimMixer/Source/MovieSceneAnimMixer/Private/MovieSceneAnimationMixerTrack.cpp b/Engine/Plugins/Experimental/MovieSceneAnimMixer/Source/MovieSceneAnimMixer/Private/MovieSceneAnimationMixerTrack.cpp
--- a/Engine/Plugins/Experimental/MovieSceneAnimMixer/Source/MovieSceneAnimMixer/Private/MovieSceneAnimationMixerTrack.cpp
+++ b/Engine/Plugins/Experimental/MovieSceneAnimMixer/Source/MovieSceneAnimMixer/Private/MovieSceneAnimationMixerTrack.cpp
@@ -5,7 +5,23 @@
 #include "Sections/MovieSceneSkeletalAnimationSection.h"
 #include "Channels/MovieSceneChannelProxy.h"
 #include "EntitySystem/BuiltInComponentTypes.h"
+#include "AnimMixerComponentTypes.h"
 #include "Misc/AxisDisplayInfo.h"
+#include "Systems/MovieSceneRootMotionSystem.h"
+
+#define LOCTEXT_NAMESPACE "MovieSceneAnimationMixerTrack"
+
+bool FMovieSceneByteChannelDefaultOnly::SerializeFromMismatchedTag(const FPropertyTag& Tag, FStructuredArchive::FSlot Slot)
+{
+	static const FName MovieSceneByteChannel("MovieSceneByteChannel");
+	if (Tag.GetType().IsStruct(MovieSceneByteChannel))
+	{
+		StaticStruct()->SerializeItem(Slot, this, nullptr);
+		return true;
+	}
+
+	return false;
+}
 
 UMovieSceneAnimationBaseTransformDecoration::UMovieSceneAnimationBaseTransformDecoration(const FObjectInitializer& ObjInit)
 	: Super(ObjInit)
@@ -19,73 +35,105 @@ UMovieSceneAnimationBaseTransformDecoration::UMovieSceneAnimationBaseTransformDe
 
 	RootMotionSpace.SetEnum(StaticEnum<EMovieSceneRootMotionSpace>());
 	RootMotionSpace.SetDefault((uint8)EMovieSceneRootMotionSpace::AnimationSpace);
+
+	TransformMode.SetEnum(StaticEnum<EMovieSceneRootMotionTransformMode>());
+	TransformMode.SetDefault((uint8)EMovieSceneRootMotionTransformMode::Offset);
+}
+
+EMovieSceneRootMotionTransformMode UMovieSceneAnimationBaseTransformDecoration::GetRootTransformMode() const
+{
+	return (EMovieSceneRootMotionTransformMode)TransformMode.GetDefault().Get((uint8)EMovieSceneRootMotionTransformMode::Offset);
 }
 
 EMovieSceneChannelProxyType UMovieSceneAnimationBaseTransformDecoration::PopulateChannelProxy(FMovieSceneChannelProxyData& OutProxyData)
 {
-#if WITH_EDITOR
+	EMovieSceneRootMotionTransformMode CurrentTransformMode = GetRootTransformMode();
 
-	FText LocationGroup = NSLOCTEXT("BaseTransformDecoration", "RootBaseLocation", "Root Base Location");
+#if WITH_EDITOR
+	FText Group = LOCTEXT("RootMotion", "Root Motion");
 	{
 		FMovieSceneChannelMetaData ChannelMetaData;
-		ChannelMetaData.SetIdentifiers("RootBaseLocation.X", AxisDisplayInfo::GetAxisDisplayName(EAxisList::X), LocationGroup);
-		ChannelMetaData.Color = AxisDisplayInfo::GetAxisColor(EAxisList::X);
+		ChannelMetaData.SetIdentifiers("Space", LOCTEXT("Space", "Space"), Group);
 		ChannelMetaData.WeakOwningObject = this;
 		ChannelMetaData.SortOrder = 0;
-		OutProxyData.Add(Location[0], ChannelMetaData, TMovieSceneExternalValue<double>());
+		OutProxyData.Add(RootMotionSpace, ChannelMetaData, TMovieSceneExternalValue<uint8>());
 	}
 	{
 		FMovieSceneChannelMetaData ChannelMetaData;
-		ChannelMetaData.SetIdentifiers("RootBaseLocation.Y", AxisDisplayInfo::GetAxisDisplayName(EAxisList::Y), LocationGroup);
-		ChannelMetaData.Color = AxisDisplayInfo::GetAxisColor(EAxisList::Y);
+		ChannelMetaData.SetIdentifiers("Root Mode", LOCTEXT("Mode", "Mode"), Group);
 		ChannelMetaData.WeakOwningObject = this;
 		ChannelMetaData.SortOrder = 1;
-		OutProxyData.Add(Location[1], ChannelMetaData, TMovieSceneExternalValue<double>());
-	}
-	{
-		FMovieSceneChannelMetaData ChannelMetaData;
-		ChannelMetaData.SetIdentifiers("RootBaseLocation.Z", AxisDisplayInfo::GetAxisDisplayName(EAxisList::Z), LocationGroup);
-		ChannelMetaData.Color = AxisDisplayInfo::GetAxisColor(EAxisList::Z);
-		ChannelMetaData.WeakOwningObject = this;
-		ChannelMetaData.SortOrder = 2;
-		OutProxyData.Add(Location[2], ChannelMetaData, TMovieSceneExternalValue<double>());
+		OutProxyData.Add(TransformMode, ChannelMetaData, TMovieSceneExternalValue<uint8>());
 	}
 
-	FText RotationGroup = NSLOCTEXT("BaseTransformDecoration", "RootBaseRotation", "Root Base Rotation");
-	{
-		FMovieSceneChannelMetaData ChannelMetaData;
-		ChannelMetaData.SetIdentifiers("RootBaseRotation.X", NSLOCTEXT("BaseTransformDecoration", "RotationX", "Roll"), RotationGroup);
-		ChannelMetaData.Color = AxisDisplayInfo::GetAxisColor(EAxisList::X);
-		ChannelMetaData.WeakOwningObject = this;
-		ChannelMetaData.SortOrder = 3;
-		OutProxyData.Add(Rotation[0], ChannelMetaData, TMovieSceneExternalValue<double>());
-	}
-	{
-		FMovieSceneChannelMetaData ChannelMetaData;
-		ChannelMetaData.SetIdentifiers("RootBaseRotation.Y", NSLOCTEXT("BaseTransformDecoration", "RotationY", "Pitch"), RotationGroup);
-		ChannelMetaData.Color = AxisDisplayInfo::GetAxisColor(EAxisList::Y);
-		ChannelMetaData.WeakOwningObject = this;
-		ChannelMetaData.SortOrder = 4;
-		OutProxyData.Add(Rotation[1], ChannelMetaData, TMovieSceneExternalValue<double>());
-	}
-	{
-		FMovieSceneChannelMetaData ChannelMetaData;
-		ChannelMetaData.SetIdentifiers("RootBaseRotation.Z", NSLOCTEXT("BaseTransformDecoration", "RotationZ", "Yaw"), RotationGroup);
-		ChannelMetaData.Color = AxisDisplayInfo::GetAxisColor(EAxisList::Z);
-		ChannelMetaData.WeakOwningObject = this;
-		ChannelMetaData.SortOrder = 5;
-		OutProxyData.Add(Rotation[2], ChannelMetaData, TMovieSceneExternalValue<double>());
-	}
+	if (CurrentTransformMode != EMovieSceneRootMotionTransformMode::Asset)
 	{
-		FMovieSceneChannelMetaData ChannelMetaData;
-		ChannelMetaData.SetIdentifiers("Space", NSLOCTEXT("BaseTransformDecoration", "Space", "Space"));
-		ChannelMetaData.WeakOwningObject = this;
-		ChannelMetaData.SortOrder = 6;
-		OutProxyData.Add(RootMotionSpace, ChannelMetaData, TMovieSceneExternalValue<uint8>());
+		FName SliderExponent("SliderExponent");
+
+		{
+			FMovieSceneChannelMetaData ChannelMetaData;
+			ChannelMetaData.SetIdentifiers("RootBaseLocation.X", AxisDisplayInfo::GetAxisDisplayName(EAxisList::X), Group);
+			ChannelMetaData.Color = AxisDisplayInfo::GetAxisColor(EAxisList::X);
+			ChannelMetaData.WeakOwningObject = this;
+			ChannelMetaData.SortOrder = 2;
+			ChannelMetaData.PropertyMetaData.Add(SliderExponent, TEXT("0.2"));
+			OutProxyData.Add(Location[0], ChannelMetaData, TMovieSceneExternalValue<double>());
+		}
+		{
+			FMovieSceneChannelMetaData ChannelMetaData;
+			ChannelMetaData.SetIdentifiers("RootBaseLocation.Y", AxisDisplayInfo::GetAxisDisplayName(EAxisList::Y), Group);
+			ChannelMetaData.Color = AxisDisplayInfo::GetAxisColor(EAxisList::Y);
+			ChannelMetaData.WeakOwningObject = this;
+			ChannelMetaData.SortOrder = 3;
+			ChannelMetaData.PropertyMetaData.Add(SliderExponent, TEXT("0.2"));
+			OutProxyData.Add(Location[1], ChannelMetaData, TMovieSceneExternalValue<double>());
+		}
+		{
+			FMovieSceneChannelMetaData ChannelMetaData;
+			ChannelMetaData.SetIdentifiers("RootBaseLocation.Z", AxisDisplayInfo::GetAxisDisplayName(EAxisList::Z), Group);
+			ChannelMetaData.Color = AxisDisplayInfo::GetAxisColor(EAxisList::Z);
+			ChannelMetaData.WeakOwningObject = this;
+			ChannelMetaData.SortOrder = 4;
+			ChannelMetaData.PropertyMetaData.Add(SliderExponent, TEXT("0.2"));
+			OutProxyData.Add(Location[2], ChannelMetaData, TMovieSceneExternalValue<double>());
+		}
+		{
+			FMovieSceneChannelMetaData ChannelMetaData;
+			ChannelMetaData.SetIdentifiers("RootBaseRotation.X", LOCTEXT("RotationX", "Roll"), Group);
+			ChannelMetaData.Color = AxisDisplayInfo::GetAxisColor(EAxisList::X);
+			ChannelMetaData.WeakOwningObject = this;
+			ChannelMetaData.SortOrder = 5;
+			ChannelMetaData.PropertyMetaData.Add(SliderExponent, TEXT("0.2"));
+			OutProxyData.Add(Rotation[0], ChannelMetaData, TMovieSceneExternalValue<double>());
+		}
+		{
+			FMovieSceneChannelMetaData ChannelMetaData;
+			ChannelMetaData.SetIdentifiers("RootBaseRotation.Y", LOCTEXT("RotationY", "Pitch"), Group);
+			ChannelMetaData.Color = AxisDisplayInfo::GetAxisColor(EAxisList::Y);
+			ChannelMetaData.WeakOwningObject = this;
+			ChannelMetaData.SortOrder = 6;
+			ChannelMetaData.PropertyMetaData.Add(SliderExponent, TEXT("0.2"));
+			OutProxyData.Add(Rotation[1], ChannelMetaData, TMovieSceneExternalValue<double>());
+		}
+		{
+			FMovieSceneChannelMetaData ChannelMetaData;
+			ChannelMetaData.SetIdentifiers("RootBaseRotation.Z", LOCTEXT("RotationZ", "Yaw"), Group);
+			ChannelMetaData.Color = AxisDisplayInfo::GetAxisColor(EAxisList::Z);
+			ChannelMetaData.WeakOwningObject = this;
+			ChannelMetaData.SortOrder = 7;
+			ChannelMetaData.PropertyMetaData.Add(SliderExponent, TEXT("0.2"));
+			OutProxyData.Add(Rotation[2], ChannelMetaData, TMovieSceneExternalValue<double>());
+		}
 	}
 
 #else
 
+	if (CurrentTransformMode != EMovieSceneRootMotionTransformMode::Asset)
+	{
+		OutProxyData.Add(RootMotionSpace);
+		OutProxyData.Add(TransformMode);
+	}
+
 	OutProxyData.Add(Location[0]);
 	OutProxyData.Add(Location[1]);
 	OutProxyData.Add(Location[2]);
@@ -94,27 +142,41 @@ EMovieSceneChannelProxyType UMovieSceneAnimationBaseTransformDecoration::Populat
 	OutProxyData.Add(Rotation[1]);
 	OutProxyData.Add(Rotation[2]);
 
-	OutProxyData.Add(RootMotionSpace);
-
 #endif
 
-	return EMovieSceneChannelProxyType::Static;
+	return EMovieSceneChannelProxyType::Dynamic;
 }
 
 void UMovieSceneAnimationBaseTransformDecoration::ExtendEntityImpl(UMovieSceneEntitySystemLinker* EntityLinker, const FEntityImportParams& Params, FImportedEntity* OutImportedEntity)
 {
 	using namespace UE::MovieScene;
 	FBuiltInComponentTypes* BuiltInComponents = FBuiltInComponentTypes::Get();
+	FAnimMixerComponentTypes* AnimMixerComponents = FAnimMixerComponentTypes::Get();
+
+	FMovieSceneRootMotionSettings RootMotionSettings;
+	{
+		RootMotionSettings.RootMotionSpace = (EMovieSceneRootMotionSpace)RootMotionSpace.GetDefault().Get((uint8)EMovieSceneRootMotionSpace::AnimationSpace);
+		RootMotionSettings.TransformMode   = GetRootTransformMode();
 
+		UMovieSceneSkeletalAnimationSection* AnimSection = GetTypedOuter<UMovieSceneSkeletalAnimationSection>();
+		if (AnimSection)
+		{
+			RootMotionSettings.LegacySwapRootBone = AnimSection->Params.SwapRootBone;
+		}
+	}
+
+	RootMotionSettings.RootOriginLocation = RootOriginLocation;
+
+	const bool bAddTransform = RootMotionSettings.TransformMode != EMovieSceneRootMotionTransformMode::Asset;
 	OutImportedEntity->AddBuilder(
 		FEntityBuilder()
-		.Add(BuiltInComponents->DoubleChannel[0], &Location[0])
-		.Add(BuiltInComponents->DoubleChannel[1], &Location[1])
-		.Add(BuiltInComponents->DoubleChannel[2], &Location[2])
-		.Add(BuiltInComponents->DoubleChannel[3], &Rotation[0])
-		.Add(BuiltInComponents->DoubleChannel[4], &Rotation[1])
-		.Add(BuiltInComponents->DoubleChannel[5], &Rotation[2])
-		.Add(BuiltInComponents->ByteChannel, &RootMotionSpace)
+		.AddConditional(BuiltInComponents->DoubleChannel[0], &Location[0], bAddTransform)
+		.AddConditional(BuiltInComponents->DoubleChannel[1], &Location[1], bAddTransform)
+		.AddConditional(BuiltInComponents->DoubleChannel[2], &Location[2], bAddTransform)
+		.AddConditional(BuiltInComponents->DoubleChannel[3], &Rotation[0], bAddTransform)
+		.AddConditional(BuiltInComponents->DoubleChannel[4], &Rotation[1], bAddTransform)
+		.AddConditional(BuiltInComponents->DoubleChannel[5], &Rotation[2], bAddTransform)
+		.Add(AnimMixerComponents->RootMotionSettings, RootMotionSettings)
 	);
 }
 
@@ -227,3 +289,6 @@ bool UMovieSceneAnimationMixerTrack::CanRename() const
 }
 
 #endif // WITH_EDITORONLY_DATA
+
+
+#undef LOCTEXT_NAMESPACE
\ No newline at end of file
diff --git a/Engine/Plugins/Experimental/MovieSceneAnimMixer/Source/MovieSceneAnimMixer/Private/Systems/MovieSceneAnimMixerSystem.cpp b/Engine/Plugins/Experimental/MovieSceneAnimMixer/Source/MovieSceneAnimMixer/Private/Systems/MovieSceneAnimMixerSystem.cpp
--- a/Engine/Plugins/Experimental/MovieSceneAnimMixer/Source/MovieSceneAnimMixer/Private/Systems/MovieSceneAnimMixerSystem.cpp
+++ b/Engine/Plugins/Experimental/MovieSceneAnimMixer/Source/MovieSceneAnimMixer/Private/Systems/MovieSceneAnimMixerSystem.cpp
@@ -222,39 +222,61 @@ namespace UE::MovieScene
 					// Evaluate the pose itself
 					Mixer->EvaluationProgram->AppendTaskPtr(MixerEntry->EvalTask);
 
-					FTransform TransformOrigin;
-					ESpaceConversions Conversion = ESpaceConversions::None;
-
-					EMovieSceneRootMotionSpace ThisRootSpace = MixerEntry->RootMotionSettings
-						? MixerEntry->RootMotionSettings->RootMotionSpace
-						: EMovieSceneRootMotionSpace::AnimationSpace;
+					// Handle root space conversions and manipulations
+					{
+						FTransform TransformOrigin  = FTransform::Identity;
+						FTransform RootTransform    = FTransform::Identity;
+						FVector    RootOffsetOrigin = FVector(ForceInitToZero);
 
-					FTransform RootBaseTransform = MixerEntry->RootMotionSettings
-						? FTransform(MixerEntry->RootMotionSettings->RootBaseRotation, MixerEntry->RootMotionSettings->RootBaseLocation)
-						: FTransform::Identity;
+						ESpaceConversions Conversion = ESpaceConversions::None;
+						EMovieSceneRootMotionSpace ThisRootSpace = EMovieSceneRootMotionSpace::AnimationSpace;
 
-					if (ThisRootSpace == EMovieSceneRootMotionSpace::AnimationSpace)
-					{
-						Conversion = ESpaceConversions::ComponentToActorRotation | ESpaceConversions::PreBaseTransform;
-						if (bRootComponentHasKeyedTransform)
+						if (MixerEntry->RootMotionSettings)
 						{
-							Conversion |= ESpaceConversions::AnimationToWorld;
+							ThisRootSpace = MixerEntry->RootMotionSettings->RootMotionSpace;
+
+							// If we have a root location or rotation offset/override, set that up
+							const bool bHasLocation = !MixerEntry->RootMotionSettings->RootLocation.IsZero();
+							const bool bHasRotation = !MixerEntry->RootMotionSettings->RootRotation.IsIdentity();
+							if (bHasLocation || bHasRotation)
+							{
+								if (MixerEntry->RootMotionSettings->TransformMode == EMovieSceneRootMotionTransformMode::Offset)
+								{
+									Conversion |= ESpaceConversions::RootTransformOffset;
+									RootOffsetOrigin = MixerEntry->RootMotionSettings->RootOriginLocation;
+								}
+								else
+								{
+									Conversion |= ESpaceConversions::RootTransformOverride;
+								}
+
+								RootTransform = FTransform(
+									MixerEntry->RootMotionSettings->RootRotation,
+									MixerEntry->RootMotionSettings->RootLocation
+									);
+							}
 						}
-						else if (TransformOriginSystem && TransformOriginSystem->GetTransformOrigin(MixerEntry->InstanceHandle, TransformOrigin))
+
+						if (ThisRootSpace == EMovieSceneRootMotionSpace::AnimationSpace)
 						{
-							Conversion |= ESpaceConversions::TransformOriginToWorld;
+							Conversion |= ESpaceConversions::ComponentToActorRotation;
+
+							if (bRootComponentHasKeyedTransform)
+							{
+								Conversion |= ESpaceConversions::AnimationToWorld;
+							}
+							else if (TransformOriginSystem && TransformOriginSystem->GetTransformOrigin(MixerEntry->InstanceHandle, TransformOrigin))
+							{
+								Conversion |= ESpaceConversions::TransformOriginToWorld;
+							}
 						}
-					}
-					else if (RootBaseTransform.GetLocation().IsZero() || !RootBaseTransform.GetRotation().IsIdentity())
-					{
-						Conversion = ESpaceConversions::PostBaseTransform;
-					}
 
-					if (Conversion != ESpaceConversions::None)
-					{
-						// Need to convert root motion attribute to world space for blending, insert task to do that
-						Mixer->EvaluationProgram->AppendTask(FAnimNextConvertRootMotionToWorldSpaceTask::Make(
-							RootMotion, TransformOrigin, RootBaseTransform, Conversion));
+						if (Conversion != ESpaceConversions::None)
+						{
+							// Need to convert root motion attribute to world space for blending, insert task to do that
+							Mixer->EvaluationProgram->AppendTask(FAnimNextConvertRootMotionToWorldSpaceTask::Make(
+								RootMotion, TransformOrigin, RootTransform, RootOffsetOrigin, Conversion));
+						}
 					}
 
 					if (MixerEntry->bRequiresBlend)
diff --git a/Engine/Plugins/Experimental/MovieSceneAnimMixer/Source/MovieSceneAnimMixer/Private/Systems/MovieSceneMixedSkeletalAnimationSystem.cpp b/Engine/Plugins/Experimental/MovieSceneAnimMixer/Source/MovieSceneAnimMixer/Private/Systems/MovieSceneMixedSkeletalAnimationSystem.cpp
--- a/Engine/Plugins/Experimental/MovieSceneAnimMixer/Source/MovieSceneAnimMixer/Private/Systems/MovieSceneMixedSkeletalAnimationSystem.cpp
+++ b/Engine/Plugins/Experimental/MovieSceneAnimMixer/Source/MovieSceneAnimMixer/Private/Systems/MovieSceneMixedSkeletalAnimationSystem.cpp
@@ -28,7 +28,6 @@
 #include "Component/AnimNextComponent.h"
 #include "Systems/MovieSceneComponentTransformSystem.h"
 #include "Systems/MovieSceneQuaternionInterpolationRotationSystem.h"
-#include "Systems/ByteChannelEvaluatorSystem.h"
 #include "Systems/MovieSceneTransformOriginSystem.h"
 #include "EntitySystem/MovieSceneBoundObjectInstantiator.h"
 #include "EntitySystem/MovieSceneBoundSceneComponentInstantiator.h"
@@ -133,10 +132,7 @@ struct FSkeletalAnimMixerMutation : IMovieSceneEntityMutation
 				}
 
 				OutPriorities[Index] = SkeletalAnim.Section->MixedAnimationPriority;
-				OutRootMotionSettings[Index].RootMotionSpace = EMovieSceneRootMotionSpace::AnimationSpace;
 
-
-				bool bHasRootMotionOverride = false;
 				UMovieSceneCommonAnimationTrack* Track = SkeletalAnim.Section->GetTypedOuter<UMovieSceneCommonAnimationTrack>();
 				UMovieSceneSkeletalAnimationTrack* SkelAnimTrack = Cast<UMovieSceneSkeletalAnimationTrack>(Track);
 
@@ -167,6 +163,8 @@ struct FSkeletalAnimMixerMutation : IMovieSceneEntityMutation
 					}
 				}
 
+				OutRootMotionSettings[Index].bHasRootMotionOverride = false;
+
 				if (Track)
 				{
 					if (Track->RootMotionParams.bRootMotionsDirty) 
@@ -177,15 +175,9 @@ struct FSkeletalAnimMixerMutation : IMovieSceneEntityMutation
 
 					if (Track->RootMotionParams.bHaveRootMotion)
 					{
-						bHasRootMotionOverride = true;
-						OutRootMotionSettings[Index].RootMotionOverride.Emplace();
+						OutRootMotionSettings[Index].bHasRootMotionOverride = true;
 					}
 				}
-
-				if (!bHasRootMotionOverride)
-				{
-					OutRootMotionSettings[Index].RootMotionOverride.Reset();
-				}
 			}
 		}
 	}
@@ -198,26 +190,20 @@ struct FGatherMixableRootMotion
 	static void ForEachEntity(
 		const FMovieSceneSkeletalAnimationComponentData& SkeletalAnimation,
 		FFrameTime EvalTime,
-		const uint8* RootMotionSpace,
 		FMovieSceneRootMotionSettings& RootMotionSettings
 		)
 	{
-		if (RootMotionSpace)
-		{
-			RootMotionSettings.RootMotionSpace = (EMovieSceneRootMotionSpace)*RootMotionSpace;
-		}
-
-		if (RootMotionSettings.RootMotionOverride.IsSet())
+		if (RootMotionSettings.bHasRootMotionOverride)
 		{
 			UMovieSceneSkeletalAnimationSection::FRootMotionParams RootMotionParams;
 
 			SkeletalAnimation.Section->GetRootMotion(EvalTime, RootMotionParams);
 			if (RootMotionParams.Transform.IsSet())
 			{
-				RootMotionSettings.RootMotionOverride = FSkeletalAnimationRootMotionOverride();
-				RootMotionSettings.RootMotionOverride->RootMotion = RootMotionParams.Transform.GetValue();
-				RootMotionSettings.RootMotionOverride->bBlendFirstChildOfRoot = RootMotionParams.bBlendFirstChildOfRoot;
-				RootMotionSettings.RootMotionOverride->ChildBoneIndex = RootMotionParams.ChildBoneIndex;
+				RootMotionSettings.RootOverrideLocation = RootMotionParams.Transform->GetLocation();
+				RootMotionSettings.RootOverrideRotation = RootMotionParams.Transform->GetRotation();
+				RootMotionSettings.bBlendFirstChildOfRoot = RootMotionParams.bBlendFirstChildOfRoot;
+				RootMotionSettings.ChildBoneIndex = RootMotionParams.ChildBoneIndex;
 			}
 		}
 	}
@@ -247,7 +233,6 @@ struct FGatherMixableSkeletalAnimations
 		TWriteOptional<FMovieSceneRootMotionSettings> OptionalRootMotionSettings,
 		const double* OptLocationX, const double* OptLocationY, const double* OptLocationZ,
 		const double* OptRotationX, const double* OptRotationY, const double* OptRotationZ,
-		const uint8* OptSpaces,
 		TWrite<TSharedPtr<FAnimNextEvaluationTask>> EvalTask) const
 	{
 		// Gather all the skeletal animations currently active in all sequences.
@@ -348,8 +333,19 @@ struct FGatherMixableSkeletalAnimations
 			AnimData.AnimSequence = AnimSequence;
 			if (OptionalRootMotionSettings)
 			{
-				AnimData.RootMotionOverride = OptionalRootMotionSettings[Index].RootMotionOverride;
-				AnimData.LegacySwapRootBone = OptionalRootMotionSettings[Index].LegacySwapRootBone;
+				if (OptionalRootMotionSettings[Index].bHasRootMotionOverride)
+				{
+					FSkeletalAnimationRootMotionOverride NewOverride;
+					NewOverride.RootMotion = FTransform(OptionalRootMotionSettings[Index].RootOverrideRotation, OptionalRootMotionSettings[Index].RootOverrideLocation);
+					NewOverride.ChildBoneIndex = OptionalRootMotionSettings[Index].ChildBoneIndex;
+					NewOverride.bBlendFirstChildOfRoot = OptionalRootMotionSettings[Index].bBlendFirstChildOfRoot;
+
+					AnimData.RootMotionOverride = NewOverride;
+				}
+				else
+				{
+					AnimData.RootMotionOverride.Reset();
+				}
 			}
 			AnimData.FromPosition = PreviousEvalTime;
 			AnimData.ToPosition = EvalTime;
@@ -370,8 +366,8 @@ struct FGatherMixableSkeletalAnimations
 
 			if (OptionalRootMotionSettings)
 			{
-				FVector  CurrentLocation = OptionalRootMotionSettings[Index].RootBaseLocation;
-				FRotator CurrentRotation = OptionalRootMotionSettings[Index].RootBaseRotation.Rotator();
+				FVector  CurrentLocation = OptionalRootMotionSettings[Index].RootLocation;
+				FRotator CurrentRotation = OptionalRootMotionSettings[Index].RootRotation.Rotator();
 
 				if (OptLocationX) { CurrentLocation.X = OptLocationX[Index]; }
 				if (OptLocationY) { CurrentLocation.Y = OptLocationY[Index]; }
@@ -380,8 +376,8 @@ struct FGatherMixableSkeletalAnimations
 				if (OptRotationY) { CurrentRotation.Pitch = OptRotationY[Index]; }
 				if (OptRotationZ) { CurrentRotation.Yaw   = OptRotationZ[Index]; }
 
-				OptionalRootMotionSettings[Index].RootBaseLocation = CurrentLocation;
-				OptionalRootMotionSettings[Index].RootBaseRotation = CurrentRotation.Quaternion();
+				OptionalRootMotionSettings[Index].RootLocation = CurrentLocation;
+				OptionalRootMotionSettings[Index].RootRotation = CurrentRotation.Quaternion();
 			}
 		}
 	}
@@ -518,7 +514,6 @@ UMovieSceneMixedSkeletalAnimationSystem::UMovieSceneMixedSkeletalAnimationSystem
 		FBuiltInComponentTypes* BuiltInComponents = FBuiltInComponentTypes::Get(); 
 
 		DefineImplicitPrerequisite(UMovieSceneEvalTimeSystem::StaticClass(), GetClass());
-		DefineImplicitPrerequisite(UByteChannelEvaluatorSystem::StaticClass(), GetClass());
 		DefineImplicitPrerequisite(UMovieSceneGenericBoundObjectInstantiator::StaticClass(), GetClass());
 		DefineImplicitPrerequisite(UMovieSceneBoundSceneComponentInstantiator::StaticClass(), GetClass());
 		DefineImplicitPrerequisite(UMovieSceneTransformOriginSystem::StaticClass(), GetClass());
@@ -572,7 +567,6 @@ void UMovieSceneMixedSkeletalAnimationSystem::OnSchedulePersistentTasks(UE::Movi
 	FTaskID GatherRootMotionTask = FEntityTaskBuilder()
 		.Read(TrackComponents->SkeletalAnimation)
 		.Read(BuiltInComponents->EvalTime)
-		.ReadOptional(BuiltInComponents->ByteResult)      // Root Space
 		.Write(AnimMixerComponents->RootMotionSettings)
 		.SetDesiredThread(ENamedThreads::GameThread_Local)
 		.FilterNone({ BuiltInComponents->Tags.Ignored })
@@ -592,7 +586,6 @@ void UMovieSceneMixedSkeletalAnimationSystem::OnSchedulePersistentTasks(UE::Movi
 		.ReadOptional(BuiltInComponents->DoubleResult[3]) // Root Rotation Offset X
 		.ReadOptional(BuiltInComponents->DoubleResult[4]) // Root Rotation Offset Y
 		.ReadOptional(BuiltInComponents->DoubleResult[5]) // Root Rotation Offset Z
-		.ReadOptional(BuiltInComponents->ByteResult)      // Root Space
 		.Write(AnimMixerComponents->Task)
 		.FilterAll( { TrackComponents->Tags.AnimMixerPoseProducer })
 		.FilterNone({ BuiltInComponents->Tags.Ignored})
diff --git a/Engine/Plugins/Experimental/MovieSceneAnimMixer/Source/MovieSceneAnimMixer/Private/Systems/MovieSceneRootMotionSystem.cpp b/Engine/Plugins/Experimental/MovieSceneAnimMixer/Source/MovieSceneAnimMixer/Private/Systems/MovieSceneRootMotionSystem.cpp
--- a/Engine/Plugins/Experimental/MovieSceneAnimMixer/Source/MovieSceneAnimMixer/Private/Systems/MovieSceneRootMotionSystem.cpp
+++ b/Engine/Plugins/Experimental/MovieSceneAnimMixer/Source/MovieSceneAnimMixer/Private/Systems/MovieSceneRootMotionSystem.cpp
@@ -178,17 +178,18 @@ void FMovieSceneMixerRootMotionComponentData::Initialize()
 		|| (RootDestination == EMovieSceneRootMotionDestination::Component && RootComponent != Target);
 }
 
-FAnimNextConvertRootMotionToWorldSpaceTask::FAnimNextConvertRootMotionToWorldSpaceTask(const TSharedPtr<FMovieSceneMixerRootMotionComponentData>& InRootMotionData, const FTransform& InTransformOrigin, const FTransform& InBaseTransform, ESpaceConversions InConversions)
+FAnimNextConvertRootMotionToWorldSpaceTask::FAnimNextConvertRootMotionToWorldSpaceTask(const TSharedPtr<FMovieSceneMixerRootMotionComponentData>& InRootMotionData, const FTransform& InTransformOrigin, const FTransform& InRootTransform, const FVector& InRootOffsetOrigin, ESpaceConversions InConversions)
 	: WeakRootMotionData(InRootMotionData)
-	, BaseTransform(InBaseTransform)
+	, RootTransform(InRootTransform)
 	, TransformOrigin(InTransformOrigin)
+	, RootOffsetOrigin(InRootOffsetOrigin)
 	, Conversions(InConversions)
 {
 }
 
-FAnimNextConvertRootMotionToWorldSpaceTask FAnimNextConvertRootMotionToWorldSpaceTask::Make(const TSharedPtr<FMovieSceneMixerRootMotionComponentData>& InRootMotionData, const FTransform& InTransformOrigin, const FTransform& InBaseTransform, ESpaceConversions InConversions)
+FAnimNextConvertRootMotionToWorldSpaceTask FAnimNextConvertRootMotionToWorldSpaceTask::Make(const TSharedPtr<FMovieSceneMixerRootMotionComponentData>& InRootMotionData, const FTransform& InTransformOrigin, const FTransform& InRootTransform, const FVector& InRootOffsetOrigin, ESpaceConversions InConversions)
 {
-	return FAnimNextConvertRootMotionToWorldSpaceTask(InRootMotionData, InTransformOrigin, InBaseTransform, InConversions);
+	return FAnimNextConvertRootMotionToWorldSpaceTask(InRootMotionData, InTransformOrigin, InRootTransform, InRootOffsetOrigin, InConversions);
 }
 
 void FAnimNextConvertRootMotionToWorldSpaceTask::Execute(UE::AnimNext::FEvaluationVM& VM) const
@@ -215,7 +216,9 @@ void FAnimNextConvertRootMotionToWorldSpaceTask::Execute(UE::AnimNext::FEvaluati
 		return;
 	}
 
-	FTransform RelativeTransform = RootMotionAttribute->Value;
+	FTransform RelativeTransform = EnumHasAnyFlags(Conversions, ESpaceConversions::RootTransformOverride)
+		? RootTransform
+		: RootMotionAttribute->Value;
 
 	TSharedPtr<FMovieSceneMixerRootMotionComponentData> RootMotionData = WeakRootMotionData.Pin();
 	if (RootMotionData)
@@ -228,16 +231,20 @@ void FAnimNextConvertRootMotionToWorldSpaceTask::Execute(UE::AnimNext::FEvaluati
 				RelativeTransform.SetTranslation(InverseMeshToActorRotation.GetValue() * RelativeTransform.GetTranslation());
 			}
 		}
+	}
 
-		if (EnumHasAnyFlags(Conversions, ESpaceConversions::PreBaseTransform))
-		{
-			RelativeTransform = RelativeTransform * BaseTransform;
-		}
-		else if (EnumHasAnyFlags(Conversions, ESpaceConversions::PostBaseTransform))
-		{
-			RelativeTransform = BaseTransform * RelativeTransform;
-		}
+	// Apply root transform offsets if necessary
+	if (EnumHasAnyFlags(Conversions, ESpaceConversions::RootTransformOffset))
+	{
+		// Multiply by the inverse of the root origin to get the root motion relative to the origin
+		RelativeTransform *= FTransform(-RootOffsetOrigin);
+		RelativeTransform = RelativeTransform * RootTransform;
+		// Put the transform back in component space
+		RelativeTransform *= FTransform(RootOffsetOrigin);
+	}
 
+	if (RootMotionData)
+	{
 		if (EnumHasAnyFlags(Conversions, ESpaceConversions::AnimationToWorld))
 		{
 			RelativeTransform *= RootMotionData->ActorTransform;
diff --git a/Engine/Plugins/Experimental/MovieSceneAnimMixer/Source/MovieSceneAnimMixer/Public/Systems/MovieSceneMixedSkeletalAnimationSystem.h b/Engine/Plugins/Experimental/MovieSceneAnimMixer/Source/MovieSceneAnimMixer/Public/Systems/MovieSceneMixedSkeletalAnimationSystem.h
--- a/Engine/Plugins/Experimental/MovieSceneAnimMixer/Source/MovieSceneAnimMixer/Public/Systems/MovieSceneMixedSkeletalAnimationSystem.h
+++ b/Engine/Plugins/Experimental/MovieSceneAnimMixer/Source/MovieSceneAnimMixer/Public/Systems/MovieSceneMixedSkeletalAnimationSystem.h
@@ -33,9 +33,6 @@ struct FMixedAnimSkeletalAnimationData
 	UPROPERTY()
 	double ToPosition = -1.0f;
 
-	UPROPERTY()
-	ESwapRootBone LegacySwapRootBone = ESwapRootBone::SwapRootBone_None;
-
 	UPROPERTY()
 	bool bFireNotifies = true;
 
diff --git a/Engine/Plugins/Experimental/MovieSceneAnimMixer/Source/MovieSceneAnimMixer/Public/Systems/MovieSceneRootMotionSystem.h b/Engine/Plugins/Experimental/MovieSceneAnimMixer/Source/MovieSceneAnimMixer/Public/Systems/MovieSceneRootMotionSystem.h
--- a/Engine/Plugins/Experimental/MovieSceneAnimMixer/Source/MovieSceneAnimMixer/Public/Systems/MovieSceneRootMotionSystem.h
+++ b/Engine/Plugins/Experimental/MovieSceneAnimMixer/Source/MovieSceneAnimMixer/Public/Systems/MovieSceneRootMotionSystem.h
@@ -35,15 +35,23 @@ struct FSkeletalAnimationRootMotionOverride
 // Structure used for animation tracks to communicate to the mixer how they would like root motion handled if at all.
 struct FMovieSceneRootMotionSettings
 {
-	TOptional<FSkeletalAnimationRootMotionOverride> RootMotionOverride;
+	FVector RootLocation = FVector(EForceInit::ForceInit);
+	FVector RootOriginLocation = FVector(EForceInit::ForceInit);
+	FVector RootOverrideLocation = FVector(EForceInit::ForceInit);
 
-	FVector RootBaseLocation = FVector(EForceInit::ForceInit);
-	FQuat   RootBaseRotation = FQuat(EForceInit::ForceInit);
+	FQuat   RootRotation = FQuat(EForceInit::ForceInit);
+	FQuat   RootOverrideRotation = FQuat(EForceInit::ForceInit);
+
+	int32 ChildBoneIndex = INDEX_NONE;
 
 	// What space to apply root motion in. Defaults to animation space.
 	EMovieSceneRootMotionSpace RootMotionSpace = EMovieSceneRootMotionSpace::AnimationSpace;
-
+	EMovieSceneRootMotionTransformMode TransformMode = EMovieSceneRootMotionTransformMode::Offset;
 	ESwapRootBone LegacySwapRootBone = ESwapRootBone::SwapRootBone_None;
+
+	uint8 bHasRootMotionOverride : 1 = 0;
+	/** If true we use the ChildBoneIndex otherwise we use the root*/
+	uint8 bBlendFirstChildOfRoot : 1 = 0;
 };
 
 // Structure that is shared between entities for handling the mixer's root motion result. 
@@ -101,25 +109,27 @@ struct FAnimNextConvertRootMotionToWorldSpaceTask : public FAnimNextEvaluationTa
 		TransformOriginToWorld = 1 << 1,
 		/** Convert the root motion from component -> actor space using the inverse component rotation only */
 		ComponentToActorRotation = 1 << 2,
-		/** Apply RootBaseTransform as Root = Root * RootBaseTransform */
-		PreBaseTransform = 1 << 3,
-		/** Apply RootBaseTransform as Root = RootBaseTransform * Root */
-		PostBaseTransform = 1 << 4,
+		/** Apply RootBaseTransform as an offset around RootOffsetOrigin where Root = Root * RootBaseTransform */
+		RootTransformOffset = 1 << 3,
+		/** Completely override the root transform with RootTransform */
+		RootTransformOverride = 1 << 4,
 	};
 
 	FAnimNextConvertRootMotionToWorldSpaceTask() = default;
-	FAnimNextConvertRootMotionToWorldSpaceTask(const TSharedPtr<FMovieSceneMixerRootMotionComponentData>& InRootMotionData, const FTransform& InTransformOrigin, const FTransform& InBaseTransform, ESpaceConversions InConversions);
+	FAnimNextConvertRootMotionToWorldSpaceTask(const TSharedPtr<FMovieSceneMixerRootMotionComponentData>& InRootMotionData, const FTransform& InTransformOrigin, const FTransform& InRootTransform, const FVector& InRootOffsetOrigin, ESpaceConversions InConversions);
 
-	static FAnimNextConvertRootMotionToWorldSpaceTask Make(const TSharedPtr<FMovieSceneMixerRootMotionComponentData>& InRootMotionData, const FTransform& InTransformOrigin, const FTransform& InBaseTransform, ESpaceConversions InConversions);
+	static FAnimNextConvertRootMotionToWorldSpaceTask Make(const TSharedPtr<FMovieSceneMixerRootMotionComponentData>& InRootMotionData, const FTransform& InTransformOrigin, const FTransform& InRootTransform, const FVector& InRootOffsetOrigin, ESpaceConversions InConversions);
 
 	virtual void Execute(UE::AnimNext::FEvaluationVM& VM) const override;
 
 	/** Weak pointer to the root motion data for all mixed animations. May be null if only a transform origin transformation is required. */
 	TWeakPtr<FMovieSceneMixerRootMotionComponentData> WeakRootMotionData;
 	/** Base transformation to apply to the root in Actor space, before transform origins. Can be used in place of a transform track. */
-	FTransform BaseTransform;
+	FTransform RootTransform;
 	/** Transform origin to apply to the root, if Conversions & ESpaceConversions::TransformOriginToWorld */
 	FTransform TransformOrigin;
+	/** Origin around which to apply RootTransform when space conversion is RootTransformOffset. */
+	FVector RootOffsetOrigin;
 	/** Bitmask specifying the conversions to perform */
 	ESpaceConversions Conversions;
 };
diff --git a/Engine/Plugins/Experimental/MovieSceneAnimMixer/Source/MovieSceneAnimMixerEditor/Private/MovieSceneAnimMixerEditorModule.cpp b/Engine/Plugins/Experimental/MovieSceneAnimMixer/Source/MovieSceneAnimMixerEditor/Private/MovieSceneAnimMixerEditorModule.cpp
--- a/Engine/Plugins/Experimental/MovieSceneAnimMixer/Source/MovieSceneAnimMixerEditor/Private/MovieSceneAnimMixerEditorModule.cpp
+++ b/Engine/Plugins/Experimental/MovieSceneAnimMixer/Source/MovieSceneAnimMixerEditor/Private/MovieSceneAnimMixerEditorModule.cpp
@@ -5,7 +5,76 @@
 #include "ISequencerModule.h"
 #include "Modules/ModuleInterface.h"
 #include "Modules/ModuleManager.h"
+#include "MovieSceneAnimationMixerTrack.h"
 #include "MovieSceneAnimationMixerTrackEditor.h"
+#include "Channels/BuiltInChannelEditors.h"
+
+
+namespace Sequencer
+{
+	FKeyHandle EvaluateAndAddKey(FMovieSceneByteChannelDefaultOnly* InChannel, const TMovieSceneChannelData<uint8>& InChannelData, FFrameNumber InTime, ISequencer& InSequencer, uint8 InDefaultValue = 0)
+	{
+		return FKeyHandle::Invalid();
+	}
+	TOptional<FKeyHandle> AddKeyForExternalValue(
+		FMovieSceneByteChannelDefaultOnly*         InChannel,
+		const TMovieSceneExternalValue<uint8>&     InExternalValue,
+		FFrameNumber                               InTime,
+		ISequencer&                                InSequencer,
+		const FGuid&                               InObjectBindingID,
+		FTrackInstancePropertyBindings*            InPropertyBindings
+		)
+	{
+		return TOptional<FKeyHandle>();
+	}
+
+	FKeyHandle AddOrUpdateKey(
+		FMovieSceneByteChannelDefaultOnly* InChannel,
+		UMovieSceneSection*                InSectionToKey,
+		FFrameNumber                       InTime,
+		ISequencer&                        InSequencer,
+		const FGuid&                       InObjectBindingID,
+		FTrackInstancePropertyBindings*    InPropertyBindings
+		)
+	{
+		return FKeyHandle::Invalid();
+	}
+
+	FKeyHandle AddOrUpdateKey(
+		FMovieSceneByteChannelDefaultOnly*         InChannel,
+		UMovieSceneSection*                        SectionToKey,
+		const TMovieSceneExternalValue<uint8>&     InExternalValue,
+		FFrameNumber                               InTime,
+		ISequencer&                                InSequencer,
+		const FGuid&                               InObjectBindingID,
+		FTrackInstancePropertyBindings*            InPropertyBindings)
+	{
+		return FKeyHandle::Invalid();
+	}
+
+	void CopyKeys(FMovieSceneByteChannelDefaultOnly* InChannel, const UMovieSceneSection* InSection, FName KeyAreaName, FMovieSceneClipboardBuilder& ClipboardBuilder, TArrayView<const FKeyHandle> InHandles)
+	{
+	}
+	void PasteKeys(FMovieSceneByteChannelDefaultOnly* InChannel, UMovieSceneSection* Section, const FMovieSceneClipboardKeyTrack& KeyTrack, const FMovieSceneClipboardEnvironment& SrcEnvironment, const FSequencerPasteEnvironment& DstEnvironment, TArray<FKeyHandle>& OutPastedKeys)
+	{
+	}
+	bool SupportsCurveEditorModels(const TMovieSceneChannelHandle<FMovieSceneByteChannelDefaultOnly>& ChannelHandle)
+	{
+		return false;
+	}
+	TSharedRef<SWidget> CreateKeyEditor(const TMovieSceneChannelHandle<FMovieSceneByteChannelDefaultOnly>& Channel, const UE::Sequencer::FCreateKeyEditorParams& Params)
+	{
+		return CreateKeyEditor(Channel.Cast<FMovieSceneByteChannel>(), Params);
+	}
+	bool CanCreateKeyEditor(const FMovieSceneByteChannelDefaultOnly* Channel)
+	{
+		return true;
+	}
+} // namespace Sequencer
+
+
+// Include order matters here since this header needs to be included after the above overloads in order for ADL to see them correctly
+#include "SequencerChannelInterface.h"
 
 namespace UE::MovieScene
 {
@@ -17,6 +86,8 @@ namespace UE::MovieScene
 
 		ISequencerModule& SequencerModule = FModuleManager::Get().LoadModuleChecked<ISequencerModule>(SequencerModuleName.Resolve());
 		AnimationTrackEditorHandle = SequencerModule.RegisterTrackEditor(FOnCreateTrackEditor::CreateStatic( &FAnimationMixerTrackEditor::CreateTrackEditor));
+
+		SequencerModule.RegisterChannelInterface<FMovieSceneByteChannelDefaultOnly>();
 	}
 
 	void FMovieSceneAnimMixerEditorModule::ShutdownModule()
diff --git a/Engine/Plugins/Experimental/MovieSceneAnimMixer/Source/MovieSceneAnimMixerEditor/Private/MovieSceneAnimationMixerTrackEditor.cpp b/Engine/Plugins/Experimental/MovieSceneAnimMixer/Source/MovieSceneAnimMixerEditor/Private/MovieSceneAnimationMixerTrackEditor.cpp
--- a/Engine/Plugins/Experimental/MovieSceneAnimMixer/Source/MovieSceneAnimMixerEditor/Private/MovieSceneAnimationMixerTrackEditor.cpp
+++ b/Engine/Plugins/Experimental/MovieSceneAnimMixer/Source/MovieSceneAnimMixerEditor/Private/MovieSceneAnimationMixerTrackEditor.cpp
@@ -2,6 +2,7 @@
 
 #include "MovieSceneAnimationMixerTrackEditor.h"
 
+#include "Animation/AnimSequenceBase.h"
 #include "AssetRegistry/AssetRegistryModule.h"
 #include "ClassIconFinder.h"
 #include "Components/SkeletalMeshComponent.h"
@@ -26,6 +27,189 @@
 namespace UE::Sequencer
 {
 
+class FAnimMixerAnimationSection
+	: public FCommonAnimationSection
+{
+public:
+
+	FAnimMixerAnimationSection(UMovieSceneSection& InSection, TWeakPtr<ISequencer> InSequencer)
+		: FCommonAnimationSection(InSection, InSequencer)
+	{}
+
+	virtual void BuildSectionContextMenu(FMenuBuilder& MenuBuilder, const FGuid& ObjectBinding) override
+	{
+		FCommonAnimationSection::BuildSectionContextMenu(MenuBuilder, ObjectBinding);
+
+		UMovieSceneSkeletalAnimationSection* AnimSection = Cast<UMovieSceneSkeletalAnimationSection>(WeakSection.Get());
+		if (AnimSection)
+		{
+			MenuBuilder.AddSubMenu(
+				LOCTEXT("RootTransform_Label", "Root Transform"),
+				LOCTEXT("RootTransform_Tooltip", "Options for root transform behavior from this anim clip"),
+				FNewMenuDelegate::CreateSP(this, &FAnimMixerAnimationSection::PopulateRootTransformMenu)
+			);
+		}
+	}
+
+	UMovieSceneAnimationBaseTransformDecoration* FindRootDecoration() const
+	{
+		UMovieSceneSkeletalAnimationSection* AnimSection = Cast<UMovieSceneSkeletalAnimationSection>(WeakSection.Get());
+		return AnimSection ? AnimSection->FindDecoration<UMovieSceneAnimationBaseTransformDecoration>() : nullptr;
+	}
+
+	UMovieSceneAnimationBaseTransformDecoration* FindOrCreateRootDecoration() const
+	{
+		UMovieSceneSkeletalAnimationSection* AnimSection = Cast<UMovieSceneSkeletalAnimationSection>(WeakSection.Get());
+		if (AnimSection)
+		{
+			UMovieSceneAnimationBaseTransformDecoration* Decoration = AnimSection->FindDecoration<UMovieSceneAnimationBaseTransformDecoration>();
+			if (!Decoration)
+			{
+				AnimSection->Modify();
+				Decoration = AnimSection->GetOrCreateDecoration<UMovieSceneAnimationBaseTransformDecoration>();
+			}
+			return Decoration;
+		}
+		return nullptr;
+	}
+
+	void PopulateRootTransformMenu(FMenuBuilder& MenuBuilder)
+	{
+		MenuBuilder.BeginSection(NAME_None, LOCTEXT("RootTransformModeSection_Label", "Root Transform Mode"));
+		{
+			FCanExecuteAction AlwaysExecute = FCanExecuteAction::CreateLambda([] { return true; });
+
+			MenuBuilder.AddMenuEntry(
+				LOCTEXT("RootTransformMode_None", "From Animation"),
+				LOCTEXT("RootTransformMode_None_Tip", "Use the root motion transform directly from the animation asset"),
+				FSlateIcon(),
+				FUIAction(
+					FExecuteAction::CreateSP(this, &FAnimMixerAnimationSection::SetRootTransformNone),
+					AlwaysExecute,
+					FIsActionChecked::CreateSP(this, &FAnimMixerAnimationSection::IsTransformMode, EMovieSceneRootMotionTransformMode::Asset)
+				),
+				NAME_None,
+				EUserInterfaceActionType::RadioButton
+			);
+
+			MenuBuilder.AddMenuEntry(
+				LOCTEXT("RootTransformMode_Offset", "Offset From Animation"),
+				LOCTEXT("RootTransformMode_Offset_Toolitp", "Offset the root motion transform using keyframed values"),
+				FSlateIcon(),
+				FUIAction(
+					FExecuteAction::CreateSP(this, &FAnimMixerAnimationSection::SetRootTransformOffset),
+					AlwaysExecute,
+					FIsActionChecked::CreateSP(this, &FAnimMixerAnimationSection::IsTransformMode, EMovieSceneRootMotionTransformMode::Offset)
+				),
+				NAME_None,
+				EUserInterfaceActionType::RadioButton
+			);
+
+			MenuBuilder.AddMenuEntry(
+				LOCTEXT("RootTransformMode_Override", "Manual Override"),
+				LOCTEXT("RootTransformMode_Override_Toolitp", "Completely override the root motion transform using keyframed values"),
+				FSlateIcon(),
+				FUIAction(
+					FExecuteAction::CreateSP(this, &FAnimMixerAnimationSection::SetRootTransformOverride),
+					AlwaysExecute,
+					FIsActionChecked::CreateSP(this, &FAnimMixerAnimationSection::IsTransformMode, EMovieSceneRootMotionTransformMode::Override)
+				),
+				NAME_None,
+				EUserInterfaceActionType::RadioButton
+			);
+		}
+		MenuBuilder.EndSection();
+
+		UMovieSceneAnimationBaseTransformDecoration* Decoration = FindRootDecoration();
+		if (Decoration && Decoration->GetRootTransformMode() == EMovieSceneRootMotionTransformMode::Offset)
+		{
+			MenuBuilder.BeginSection(NAME_None, LOCTEXT("RootTransformOffset_Label", "Offset"));
+			{
+				MenuBuilder.AddMenuEntry(
+					LOCTEXT("RecenterRootTransform", "Re-center Root Transform"),
+					LOCTEXT("RecenterRootTransform_Tooltip", "Center the root transform for this animation around its current position"),
+					FSlateIcon(),
+					FUIAction(
+						FExecuteAction::CreateSP(this, &FAnimMixerAnimationSection::RecenterRootTransform)
+					)
+				);
+			}
+		}
+	}
+
+	bool IsTransformMode(EMovieSceneRootMotionTransformMode Mode) const
+	{
+		if (UMovieSceneAnimationBaseTransformDecoration* Decoration = FindRootDecoration())
+		{
+			return Decoration->GetRootTransformMode() == Mode;
+		}
+		return Mode == EMovieSceneRootMotionTransformMode::Asset;
+	}
+
+	void SetRootTransformNone()
+	{
+		FScopedTransaction Transaction(LOCTEXT("SetRootTransformNone", "Use Root Transform from Asset"));
+
+		UMovieSceneSkeletalAnimationSection*         AnimSection = Cast<UMovieSceneSkeletalAnimationSection>(WeakSection.Get());
+		UMovieSceneAnimationBaseTransformDecoration* Decoration  = FindOrCreateRootDecoration();
+		if (Decoration && AnimSection)
+		{
+			Decoration->Modify();
+			Decoration->TransformMode.SetDefault((uint8)EMovieSceneRootMotionTransformMode::Asset);
+			AnimSection->InvalidateChannelProxy();
+		}
+	}
+	void SetRootTransformOffset()
+	{
+		FScopedTransaction Transaction(LOCTEXT("SetRootTransformOffset", "Offset Root Transform from Asset"));
+
+		UMovieSceneSkeletalAnimationSection*         AnimSection = Cast<UMovieSceneSkeletalAnimationSection>(WeakSection.Get());
+		UMovieSceneAnimationBaseTransformDecoration* Decoration  = FindOrCreateRootDecoration();
+		if (Decoration && AnimSection)
+		{
+			Decoration->Modify();
+			Decoration->TransformMode.SetDefault((uint8)EMovieSceneRootMotionTransformMode::Offset);
+			AnimSection->InvalidateChannelProxy();
+		}
+	}
+	void SetRootTransformOverride()
+	{
+		FScopedTransaction Transaction(LOCTEXT("SetRootTransformOffset", "Override Root Transform"));
+
+		UMovieSceneSkeletalAnimationSection*         AnimSection = Cast<UMovieSceneSkeletalAnimationSection>(WeakSection.Get());
+		UMovieSceneAnimationBaseTransformDecoration* Decoration  = FindOrCreateRootDecoration();
+		if (Decoration && AnimSection)
+		{
+			Decoration->Modify();
+			Decoration->TransformMode.SetDefault((uint8)EMovieSceneRootMotionTransformMode::Override);
+			AnimSection->InvalidateChannelProxy();
+		}
+	}
+	void RecenterRootTransform()
+	{
+		FScopedTransaction Transaction(LOCTEXT("SetRootTransformOffset", "Override Root Transform"));
+
+		TSharedPtr<ISequencer> SequencerPtr = Sequencer.Pin();
+
+		UMovieSceneSkeletalAnimationSection*         AnimSection = Cast<UMovieSceneSkeletalAnimationSection>(WeakSection.Get());
+		UMovieSceneAnimationBaseTransformDecoration* Decoration  = FindRootDecoration();
+		UAnimSequenceBase*                           Animation   = AnimSection ? AnimSection->GetAnimation() : nullptr;
+
+		if (SequencerPtr && Decoration && Animation)
+		{
+			Decoration->Modify();
+
+			FQualifiedFrameTime Time = SequencerPtr->GetLocalTime();
+			double AnimTime = AnimSection->MapTimeToAnimation(Time.Time, Time.Rate);
+
+			FAnimExtractContext CurrentExtractContext(AnimTime);
+			FTransform CurrentTransform  = Animation->ExtractRootTrackTransform(CurrentExtractContext, nullptr);
+
+			Decoration->RootOriginLocation = CurrentTransform.GetTranslation();
+		}
+	}
+};
+
 TSharedRef<ISequencerTrackEditor> FAnimationMixerTrackEditor::CreateTrackEditor(TSharedRef<ISequencer> InSequencer)
 {
 	return MakeShareable(new FAnimationMixerTrackEditor(InSequencer));
@@ -109,7 +293,7 @@ TSharedRef<ISequencerSection> FAnimationMixerTrackEditor::MakeSectionInterface(U
 	
 	if (SectionObject.IsA<UMovieSceneSkeletalAnimationSection>())
 	{
-		return MakeShared<FCommonAnimationSection>(SectionObject, GetSequencer());
+		return MakeShared<FAnimMixerAnimationSection>(SectionObject, GetSequencer());
 	}
 
 	return FMovieSceneTrackEditor::MakeSectionInterface(SectionObject, Track, ObjectBinding);
diff --git a/Engine/Source/Editor/MovieSceneTools/Private/Channels/BuiltInChannelEditors.cpp b/Engine/Source/Editor/MovieSceneTools/Private/Channels/BuiltInChannelEditors.cpp
--- a/Engine/Source/Editor/MovieSceneTools/Private/Channels/BuiltInChannelEditors.cpp
+++ b/Engine/Source/Editor/MovieSceneTools/Private/Channels/BuiltInChannelEditors.cpp
@@ -1,6 +1,6 @@
 // Copyright Epic Games, Inc. All Rights Reserved.
 
-#include "BuiltInChannelEditors.h"
+#include "Channels/BuiltInChannelEditors.h"
 #include "MovieSceneSequence.h"
 #include "MovieSceneSequenceEditor.h"
 #include "MovieSceneEventUtils.h"
diff --git a/Engine/Source/Editor/MovieSceneTools/Private/Channels/BuiltInChannelEditors.h b/Engine/Source/Editor/MovieSceneTools/Private/Channels/BuiltInChannelEditors.h
deleted file mode 100644
diff --git a/Engine/Source/Editor/MovieSceneTools/Private/TrackEditors/CommonAnimationTrackEditor.cpp b/Engine/Source/Editor/MovieSceneTools/Private/TrackEditors/CommonAnimationTrackEditor.cpp
--- a/Engine/Source/Editor/MovieSceneTools/Private/TrackEditors/CommonAnimationTrackEditor.cpp
+++ b/Engine/Source/Editor/MovieSceneTools/Private/TrackEditors/CommonAnimationTrackEditor.cpp
@@ -598,7 +598,7 @@ private:
 
 
 FCommonAnimationSection::FCommonAnimationSection( UMovieSceneSection& InSection, TWeakPtr<ISequencer> InSequencer)
-	: Section(*CastChecked<UMovieSceneSkeletalAnimationSection>(&InSection))
+	: WeakSection(CastChecked<UMovieSceneSkeletalAnimationSection>(&InSection))
 	, Sequencer(InSequencer)
 	, PreDilatePlayRate(1.0)
 {
@@ -610,16 +610,22 @@ FCommonAnimationSection::~FCommonAnimationSection()
 
 void FCommonAnimationSection::BeginDilateSection()
 {
-	if (Section.Params.PlayRate.GetType() == EMovieSceneTimeWarpType::FixedPlayRate)
+	UMovieSceneSkeletalAnimationSection* Section = WeakSection.Get();
+	if (!Section)
 	{
-		PreDilatePlayRate = Section.Params.PlayRate.AsFixedPlayRate(); //make sure to cache the play rate
+		return;
+	}
+
+	if (Section->Params.PlayRate.GetType() == EMovieSceneTimeWarpType::FixedPlayRate)
+	{
+		PreDilatePlayRate = Section->Params.PlayRate.AsFixedPlayRate(); //make sure to cache the play rate
 	}
-	else if (Section.Params.PlayRate.GetType() == EMovieSceneTimeWarpType::Custom)
+	else if (Section->Params.PlayRate.GetType() == EMovieSceneTimeWarpType::Custom)
 	{
-		FMovieSceneTimeWarpChannel* Channel = Section.GetChannelProxy().GetChannel<FMovieSceneTimeWarpChannel>(0);
+		FMovieSceneTimeWarpChannel* Channel = Section->GetChannelProxy().GetChannel<FMovieSceneTimeWarpChannel>(0);
 		if (Channel)
 		{
-			Section.Params.PlayRate.AsCustom()->Modify();
+			Section->Params.PlayRate.AsCustom()->Modify();
 			PreDilateChannel = MakeUnique<FMovieSceneTimeWarpChannel>(*Channel);
 		}
 	}
@@ -627,13 +633,19 @@ void FCommonAnimationSection::BeginDilateSection()
 
 void FCommonAnimationSection::DilateSection(const TRange<FFrameNumber>& NewRange, float DilationFactor)
 {
-	if (Section.Params.PlayRate.GetType() == EMovieSceneTimeWarpType::FixedPlayRate)
+	UMovieSceneSkeletalAnimationSection* Section = WeakSection.Get();
+	if (!Section)
+	{
+		return;
+	}
+
+	if (Section->Params.PlayRate.GetType() == EMovieSceneTimeWarpType::FixedPlayRate)
 	{
-		Section.Params.PlayRate.Set(PreDilatePlayRate / DilationFactor);
+		Section->Params.PlayRate.Set(PreDilatePlayRate / DilationFactor);
 	}
-	else if (Section.Params.PlayRate.GetType() == EMovieSceneTimeWarpType::Custom)
+	else if (Section->Params.PlayRate.GetType() == EMovieSceneTimeWarpType::Custom)
 	{
-		FMovieSceneTimeWarpChannel* Channel = Section.GetChannelProxy().GetChannel<FMovieSceneTimeWarpChannel>(0);
+		FMovieSceneTimeWarpChannel* Channel = Section->GetChannelProxy().GetChannel<FMovieSceneTimeWarpChannel>(0);
 		if (Channel)
 		{
 			*Channel = *PreDilateChannel;
@@ -641,28 +653,34 @@ void FCommonAnimationSection::DilateSection(const TRange<FFrameNumber>& NewRange
 			// Dilate the times
 			Dilate(Channel, FFrameNumber(0), DilationFactor);
 
-			Section.Params.PlayRate.AsCustom()->MarkAsChanged();
+			Section->Params.PlayRate.AsCustom()->MarkAsChanged();
 		}
 	}
-	Section.SetRange(NewRange);
+	Section->SetRange(NewRange);
 }
 
 UMovieSceneSection* FCommonAnimationSection::GetSectionObject()
 { 
-	return &Section;
+	return WeakSection.Get();
 }
 
 FText FCommonAnimationSection::GetSectionTitle() const
 {
-	if (Section.Params.Animation != nullptr)
+	UMovieSceneSkeletalAnimationSection* Section = WeakSection.Get();
+	if (!Section)
+	{
+		return FText();
+	}
+
+	if (Section->Params.Animation != nullptr)
 	{
-		if (!Section.Params.MirrorDataTable)
+		if (!Section->Params.MirrorDataTable)
 		{
-			return FText::FromString( Section.Params.Animation->GetName() );
+			return FText::FromString( Section->Params.Animation->GetName() );
 		}
 		else
 		{
-			return FText::Format(LOCTEXT("SectionTitleContentFormat", "{0} mirrored with {1}"), FText::FromString(Section.Params.Animation->GetName()), FText::FromString(Section.Params.MirrorDataTable->GetName()));
+			return FText::Format(LOCTEXT("SectionTitleContentFormat", "{0} mirrored with {1}"), FText::FromString(Section->Params.Animation->GetName()), FText::FromString(Section->Params.MirrorDataTable->GetName()));
 		}
 	}
 	return LOCTEXT("NoAnimationSection", "No Animation");
@@ -670,22 +688,23 @@ FText FCommonAnimationSection::GetSectionTitle() const
 
 FText FCommonAnimationSection::GetSectionToolTip() const
 {
-	if (Section.Params.Animation != nullptr && Section.HasStartFrame() && Section.HasEndFrame())
+	UMovieSceneSkeletalAnimationSection* Section = WeakSection.Get();
+	if (Section && Section->Params.Animation != nullptr && Section->HasStartFrame() && Section->HasEndFrame())
 	{
-		UMovieScene* MovieScene = Section.GetTypedOuter<UMovieScene>();
+		UMovieScene* MovieScene = Section->GetTypedOuter<UMovieScene>();
 		FFrameRate TickResolution = MovieScene->GetTickResolution();
 
-		FMovieSceneSequenceTransform Transform = Section.Params.MakeTransform(TickResolution, Section.GetRange());
+		FMovieSceneSequenceTransform Transform = Section->Params.MakeTransform(TickResolution, Section->GetRange());
 
 		const double StartTime     = Transform.TransformTime(0).AsDecimal();
-		const float  SectionLength = Section.GetRange().Size<FFrameTime>() / TickResolution;
+		const float  SectionLength = Section->GetRange().Size<FFrameTime>() / TickResolution;
 
-		if (Section.Params.PlayRate.GetType() == EMovieSceneTimeWarpType::FixedPlayRate)
+		if (Section->Params.PlayRate.GetType() == EMovieSceneTimeWarpType::FixedPlayRate)
 		{
-			const double PlayRate = Section.Params.PlayRate.AsFixedPlayRate();
+			const double PlayRate = Section->Params.PlayRate.AsFixedPlayRate();
 			return FText::Format(LOCTEXT("ToolTipContentFormat_FixedPlayRate", "Start: {0}s\nDuration: {1}s\nPlay Rate: {2}x"), StartTime, SectionLength, PlayRate);
 		}
-		else if (Section.Params.PlayRate.GetType() == EMovieSceneTimeWarpType::Custom)
+		else if (Section->Params.PlayRate.GetType() == EMovieSceneTimeWarpType::Custom)
 		{
 			return FText::Format(LOCTEXT("ToolTipContentFormat_TimwWarp", "Start: {0}s\nDuration: {1}s\nPlay Rate: Variable"), StartTime, SectionLength);
 		}
@@ -695,13 +714,15 @@ FText FCommonAnimationSection::GetSectionToolTip() const
 
 TOptional<FFrameTime> FCommonAnimationSection::GetSectionTime(FSequencerSectionPainter& InPainter) const
 {
-	if (!InPainter.bIsSelected || !Sequencer.Pin() || Section.Params.Animation == nullptr)
+	UMovieSceneSkeletalAnimationSection* Section = WeakSection.Get();
+
+	if (!Section || !InPainter.bIsSelected || !Sequencer.Pin() || Section->Params.Animation == nullptr)
 	{
 		return TOptional<FFrameTime>();
 	}
 
 	FFrameTime CurrentTime = Sequencer.Pin()->GetLocalTime().Time;
-	if (!Section.GetRange().Contains(CurrentTime.FrameNumber))
+	if (!Section->GetRange().Contains(CurrentTime.FrameNumber))
 	{
 		return TOptional<FFrameTime>();
 	}
@@ -710,11 +731,11 @@ TOptional<FFrameTime> FCommonAnimationSection::GetSectionTime(FSequencerSectionP
 	FFrameRate TickResolution = TimeToPixelConverter.GetTickResolution();
 
 	// Draw the current time next to the scrub handle
-	const double AnimTime = Section.MapTimeToAnimation(CurrentTime, TickResolution);
-	const FFrameRate SamplingFrameRate = Section.Params.Animation->GetSamplingFrameRate();
+	const double AnimTime = Section->MapTimeToAnimation(CurrentTime, TickResolution);
+	const FFrameRate SamplingFrameRate = Section->Params.Animation->GetSamplingFrameRate();
 
 	FQualifiedFrameTime HintFrameTime;
-	if (!UAnimationBlueprintLibrary::EvaluateRootBoneTimecodeAttributesAtTime(Section.Params.Animation, static_cast<float>(AnimTime), HintFrameTime))
+	if (!UAnimationBlueprintLibrary::EvaluateRootBoneTimecodeAttributesAtTime(Section->Params.Animation, static_cast<float>(AnimTime), HintFrameTime))
 	{
 		const FFrameTime FrameTime = SamplingFrameRate.AsFrameTime(AnimTime);
 		HintFrameTime = FQualifiedFrameTime(FrameTime, SamplingFrameRate);
@@ -741,7 +762,7 @@ TOptional<FFrameTime> FCommonAnimationSection::GetSectionTime(FSequencerSectionP
 	FFrameNumberInterface FrameNumberInterface(DisplayFormatAttr, ZeroPadFrameNumbersAttr, TickResolutionAttr, DisplayRateAttr);
 
 	float Subframe = 0.0f;
-	if (UAnimationBlueprintLibrary::EvaluateRootBoneTimecodeSubframeAttributeAtTime(Section.Params.Animation, static_cast<float>(AnimTime), Subframe))
+	if (UAnimationBlueprintLibrary::EvaluateRootBoneTimecodeSubframeAttributeAtTime(Section->Params.Animation, static_cast<float>(AnimTime), Subframe))
 	{
 		if (FMath::IsNearlyEqual(Subframe, FMath::RoundToFloat(Subframe)))
 		{
@@ -773,29 +794,35 @@ int32 FCommonAnimationSection::OnPaintSection( FSequencerSectionPainter& Painter
 {
 	using namespace UE::MovieScene;
 
+	UMovieSceneSkeletalAnimationSection* Section = WeakSection.Get();
+	if (!Section)
+	{
+		return Painter.LayerId;
+	}
+
 	const ESlateDrawEffect DrawEffects = Painter.bParentEnabled ? ESlateDrawEffect::None : ESlateDrawEffect::DisabledEffect;
 	
 	const FTimeToPixel& TimeToPixelConverter = Painter.GetTimeConverter();
 
 	int32 LayerId = Painter.PaintSectionBackground();
 
-	if (!Section.HasStartFrame() || !Section.HasEndFrame())
+	if (!Section->HasStartFrame() || !Section->HasEndFrame())
 	{
 		return LayerId;
 	}
 
 	FFrameRate TickResolution = TimeToPixelConverter.GetTickResolution();
 
-	const FFrameNumber StartFrame = Section.GetInclusiveStartFrame();
-	const FFrameNumber EndFrame   = Section.GetExclusiveEndFrame();
+	const FFrameNumber StartFrame = Section->GetInclusiveStartFrame();
+	const FFrameNumber EndFrame   = Section->GetExclusiveEndFrame();
 
-	if (UMovieScene* MovieScene = Section.GetTypedOuter<UMovieScene>())
+	if (UMovieScene* MovieScene = Section->GetTypedOuter<UMovieScene>())
 	{
 		FFrameRate FrameRate = MovieScene->GetTickResolution();
-		FMovieSceneSequenceTransform OuterToInnerTransform = Section.Params.MakeTransform(FrameRate, Section.GetRange());
+		FMovieSceneSequenceTransform OuterToInnerTransform = Section->Params.MakeTransform(FrameRate, Section->GetRange());
 
 		// As seconds represented as a FFrameTime
-		FFrameTime LoopStart = FFrameTime::FromDecimal(Section.Params.StartFrameOffset / FrameRate);
+		FFrameTime LoopStart = FFrameTime::FromDecimal(Section->Params.StartFrameOffset / FrameRate);
 
 		FLinearColor SectionTint = Painter.GetSectionColor().LinearRGBToHSV();
 		SectionTint.B *= 0.1f;
@@ -837,16 +864,24 @@ int32 FCommonAnimationSection::OnPaintSection( FSequencerSectionPainter& Painter
 
 void FCommonAnimationSection::BeginResizeSection()
 {
-	const FFrameRate FrameRate = Section.GetTypedOuter<UMovieScene>()->GetTickResolution();
-	InitialDragTransform = MakeUnique<FMovieSceneSequenceTransform>(Section.Params.MakeTransform(FrameRate, Section.GetRange()));
+	UMovieSceneSkeletalAnimationSection* Section = WeakSection.Get();
+	if (!Section)
+	{
+		return;
+	}
+
+	const FFrameRate FrameRate = Section->GetTypedOuter<UMovieScene>()->GetTickResolution();
+	InitialDragTransform = MakeUnique<FMovieSceneSequenceTransform>(Section->Params.MakeTransform(FrameRate, Section->GetRange()));
 }
 
 void FCommonAnimationSection::ResizeSection(ESequencerSectionResizeMode ResizeMode, FFrameNumber ResizeTime)
 {
-	if (ResizeMode == SSRM_LeadingEdge && Section.Params.PlayRate.GetType() != EMovieSceneTimeWarpType::Custom)
+	UMovieSceneSkeletalAnimationSection* Section = WeakSection.Get();
+
+	if (Section && ResizeMode == SSRM_LeadingEdge && Section->Params.PlayRate.GetType() != EMovieSceneTimeWarpType::Custom)
 	{
-		const FFrameRate FrameRate = Section.GetTypedOuter<UMovieScene>()->GetTickResolution();
-		Section.Params.FirstLoopStartFrameOffset = (InitialDragTransform->TransformTime(ResizeTime).AsDecimal() * FrameRate).RoundToFrame();
+		const FFrameRate FrameRate = Section->GetTypedOuter<UMovieScene>()->GetTickResolution();
+		Section->Params.FirstLoopStartFrameOffset = (InitialDragTransform->TransformTime(ResizeTime).AsDecimal() * FrameRate).RoundToFrame();
 	}
 
 	ISequencerSection::ResizeSection(ResizeMode, ResizeTime);
@@ -859,15 +894,23 @@ void FCommonAnimationSection::BeginSlipSection()
 
 void FCommonAnimationSection::SlipSection(FFrameNumber SlipTime)
 {
-	const FFrameRate FrameRate = Section.GetTypedOuter<UMovieScene>()->GetTickResolution();
-	Section.Params.FirstLoopStartFrameOffset = (InitialDragTransform->TransformTime(SlipTime).AsDecimal() * FrameRate).RoundToFrame();
+	UMovieSceneSkeletalAnimationSection* Section = WeakSection.Get();
+	if (Section)
+	{
+		const FFrameRate FrameRate = Section->GetTypedOuter<UMovieScene>()->GetTickResolution();
+		Section->Params.FirstLoopStartFrameOffset = (InitialDragTransform->TransformTime(SlipTime).AsDecimal() * FrameRate).RoundToFrame();
+	}
 	ISequencerSection::SlipSection(SlipTime);
 }
 
 bool FCommonAnimationSection::RequestDeleteKeyArea(const TArray<FName>& KeyAreaNamePath)
 {
-	Section.Modify();
-	Section.DeleteChannels(KeyAreaNamePath);
+	UMovieSceneSkeletalAnimationSection* Section = WeakSection.Get();
+	if (Section)
+	{
+		Section->Modify();
+		Section->DeleteChannels(KeyAreaNamePath);
+	}
 	return true;
 }
 
@@ -886,7 +929,13 @@ void FCommonAnimationSection::BuildSectionContextMenu(FMenuBuilder& MenuBuilder,
 		return;
 	}
 
-	UMovieSceneCommonAnimationTrack* const Track = Section.GetTypedOuter<UMovieSceneCommonAnimationTrack>();
+	UMovieSceneSkeletalAnimationSection* Section = WeakSection.Get();
+	if (!Section)
+	{
+		return;
+	}
+
+	UMovieSceneCommonAnimationTrack* const Track = Section->GetTypedOuter<UMovieSceneCommonAnimationTrack>();
 	if (!IsValid(Track))
 	{
 		return;
@@ -906,10 +955,10 @@ void FCommonAnimationSection::BuildSectionContextMenu(FMenuBuilder& MenuBuilder,
 		BoneNames.Add(Skeleton->GetReferenceSkeleton().GetBoneName(BoneIndex));
 	}
 
-	auto MatchToBone = [this, ObjectBinding, Skeleton](bool bMatchPrevious, int32 Index)
+	auto MatchToBone = [this, ObjectBinding, Skeleton, Section](bool bMatchPrevious, int32 Index)
 	{
 		return FUIAction(
-			FExecuteAction::CreateLambda([this, ObjectBinding, Skeleton, bMatchPrevious, Index]
+			FExecuteAction::CreateLambda([this, ObjectBinding, Skeleton, bMatchPrevious, Index, Section]
 				{
 					const TSharedPtr<ISequencer> SequencerPtr = Sequencer.Pin();
 					if (!SequencerPtr.IsValid())
@@ -917,25 +966,25 @@ void FCommonAnimationSection::BuildSectionContextMenu(FMenuBuilder& MenuBuilder,
 						return;
 					}
 
-					UMovieSceneCommonAnimationTrack* const Track = Section.GetTypedOuter<UMovieSceneCommonAnimationTrack>();
+					UMovieSceneCommonAnimationTrack* const Track = Section->GetTypedOuter<UMovieSceneCommonAnimationTrack>();
 					if (!IsValid(Track))
 					{
 						return;
 					}
 
 					FScopedTransaction MatchSection(LOCTEXT("MatchSectionByBone_Transaction", "Match Section By Bone"));
-					Section.Modify();	
-					Section.bMatchWithPrevious = bMatchPrevious;
+					Section->Modify();	
+					Section->bMatchWithPrevious = bMatchPrevious;
 					if (Index >= 0)
 					{
 						USkeletalMeshComponent* const SkelMeshComp = UE::Sequencer::FCommonAnimationTrackEditor::AcquireSkeletalMeshFromObjectGuid(ObjectBinding, SequencerPtr);
 
 						FName Name = Skeleton->GetReferenceSkeleton().GetBoneName(Index);
-						Section.MatchSectionByBoneTransform(SkelMeshComp, SequencerPtr->GetLocalTime().Time, SequencerPtr->GetLocalTime().Rate, Name);
+						Section->MatchSectionByBoneTransform(SkelMeshComp, SequencerPtr->GetLocalTime().Time, SequencerPtr->GetLocalTime().Rate, Name);
 					}
 					else
 					{
-						Section.ClearMatchedOffsetTransforms();
+						Section->ClearMatchedOffsetTransforms();
 					}
 					SequencerPtr->NotifyMovieSceneDataChanged(EMovieSceneDataChangeType::TrackValueChanged);
 
@@ -945,7 +994,7 @@ void FCommonAnimationSection::BuildSectionContextMenu(FMenuBuilder& MenuBuilder,
 				{ 
 					return Sequencer.IsValid(); 
 				}),
-			FIsActionChecked::CreateLambda([this, ObjectBinding, Index]() -> bool
+			FIsActionChecked::CreateLambda([this, ObjectBinding, Index, Section]() -> bool
 				{
 					const TSharedPtr<ISequencer> SequencerPtr = Sequencer.Pin();
 					if (!SequencerPtr.IsValid())
@@ -962,10 +1011,10 @@ void FCommonAnimationSection::BuildSectionContextMenu(FMenuBuilder& MenuBuilder,
 					if (Index >= 0)
 					{
 						FName Name = Skeleton->GetReferenceSkeleton().GetBoneName(Index);
-						return (Section.MatchedBoneName == Name);
+						return (Section->MatchedBoneName == Name);
 					}
 
-					return (Section.MatchedBoneName == NAME_None);
+					return (Section->MatchedBoneName == NAME_None);
 				})
 			);
 	};
@@ -1022,7 +1071,7 @@ void FCommonAnimationSection::BuildSectionContextMenu(FMenuBuilder& MenuBuilder,
 			LOCTEXT("MatchTranslationTooltip", "Match the Translation to the Specified Bone"),
 			FSlateIcon(),
 			FUIAction(
-				FExecuteAction::CreateLambda([this, ObjectBinding]()
+				FExecuteAction::CreateLambda([this, ObjectBinding, Section]()
 					{
 						const TSharedPtr<ISequencer> SequencerPtr = Sequencer.Pin();
 						if (!SequencerPtr.IsValid())
@@ -1033,18 +1082,18 @@ void FCommonAnimationSection::BuildSectionContextMenu(FMenuBuilder& MenuBuilder,
 						USkeletalMeshComponent* const SkelMeshComp = UE::Sequencer::FCommonAnimationTrackEditor::AcquireSkeletalMeshFromObjectGuid(ObjectBinding, SequencerPtr);
 
 						FScopedTransaction MatchTransaction(LOCTEXT("MatchTranslation_Transaction", "Match Translation"));
-						Section.Modify();	
-						Section.ToggleMatchTranslation();
-						Section.MatchSectionByBoneTransform(SkelMeshComp, SequencerPtr->GetLocalTime().Time, SequencerPtr->GetLocalTime().Rate, Section.MatchedBoneName);
+						Section->Modify();	
+						Section->ToggleMatchTranslation();
+						Section->MatchSectionByBoneTransform(SkelMeshComp, SequencerPtr->GetLocalTime().Time, SequencerPtr->GetLocalTime().Rate, Section->MatchedBoneName);
 						SequencerPtr->NotifyMovieSceneDataChanged(EMovieSceneDataChangeType::TrackValueChanged);
 					}),
 				FCanExecuteAction::CreateLambda([]() -> bool
 					{
 						return true;
 					}),
-				FIsActionChecked::CreateLambda([this]() -> bool
+				FIsActionChecked::CreateLambda([Section]() -> bool
 					{
-						return Section.bMatchTranslation;
+						return Section->bMatchTranslation;
 					})),
 			NAME_None,
 			EUserInterfaceActionType::ToggleButton);
@@ -1054,7 +1103,7 @@ void FCommonAnimationSection::BuildSectionContextMenu(FMenuBuilder& MenuBuilder,
 			LOCTEXT("MatchZHeightTooltip", "Match the Z Height, may want this off for better matching"),
 			FSlateIcon(),
 			FUIAction(
-				FExecuteAction::CreateLambda([this, ObjectBinding]()
+				FExecuteAction::CreateLambda([this, ObjectBinding, Section]()
 					{
 						const TSharedPtr<ISequencer> SequencerPtr = Sequencer.Pin();
 						if (!SequencerPtr.IsValid())
@@ -1065,18 +1114,18 @@ void FCommonAnimationSection::BuildSectionContextMenu(FMenuBuilder& MenuBuilder,
 						USkeletalMeshComponent* const SkelMeshComp = UE::Sequencer::FCommonAnimationTrackEditor::AcquireSkeletalMeshFromObjectGuid(ObjectBinding, SequencerPtr);
 
 						FScopedTransaction MatchTransaction(LOCTEXT("MatchZHeight_Transaction", "Match Z Height"));
-						Section.Modify();
-						Section.ToggleMatchIncludeZHeight(); 
-						Section.MatchSectionByBoneTransform(SkelMeshComp, SequencerPtr->GetLocalTime().Time, SequencerPtr->GetLocalTime().Rate, Section.MatchedBoneName);
+						Section->Modify();
+						Section->ToggleMatchIncludeZHeight(); 
+						Section->MatchSectionByBoneTransform(SkelMeshComp, SequencerPtr->GetLocalTime().Time, SequencerPtr->GetLocalTime().Rate, Section->MatchedBoneName);
 						SequencerPtr->NotifyMovieSceneDataChanged(EMovieSceneDataChangeType::TrackValueChanged);
 					}),
 					FCanExecuteAction::CreateLambda([]() -> bool
 						{
 							return true;
 						}),
-					FIsActionChecked::CreateLambda([this]() -> bool
+					FIsActionChecked::CreateLambda([Section]() -> bool
 						{
-							return Section.bMatchIncludeZHeight;
+							return Section->bMatchIncludeZHeight;
 						})),
 			NAME_None,
 			EUserInterfaceActionType::ToggleButton
@@ -1087,7 +1136,7 @@ void FCommonAnimationSection::BuildSectionContextMenu(FMenuBuilder& MenuBuilder,
 			LOCTEXT("MatchYawRotationTooltip", "Match the Yaw Rotation, may want this off for better matching"),
 			FSlateIcon(),
 			FUIAction(
-				FExecuteAction::CreateLambda([this, ObjectBinding]()
+				FExecuteAction::CreateLambda([this, ObjectBinding, Section]()
 					{
 						const TSharedPtr<ISequencer> SequencerPtr = Sequencer.Pin();
 						if (!SequencerPtr.IsValid())
@@ -1098,17 +1147,17 @@ void FCommonAnimationSection::BuildSectionContextMenu(FMenuBuilder& MenuBuilder,
 						USkeletalMeshComponent* const SkelMeshComp = UE::Sequencer::FCommonAnimationTrackEditor::AcquireSkeletalMeshFromObjectGuid(ObjectBinding, SequencerPtr);
 
 						FScopedTransaction MatchTransaction(LOCTEXT("MatchYawRotation_Transaction", "Match Yaw Rotation"));
-						Section.Modify();
-						Section.ToggleMatchIncludeYawRotation();
-						Section.MatchSectionByBoneTransform(SkelMeshComp, SequencerPtr->GetLocalTime().Time, SequencerPtr->GetLocalTime().Rate, Section.MatchedBoneName);
+						Section->Modify();
+						Section->ToggleMatchIncludeYawRotation();
+						Section->MatchSectionByBoneTransform(SkelMeshComp, SequencerPtr->GetLocalTime().Time, SequencerPtr->GetLocalTime().Rate, Section->MatchedBoneName);
 						SequencerPtr->NotifyMovieSceneDataChanged(EMovieSceneDataChangeType::TrackValueChanged); }),
 					FCanExecuteAction::CreateLambda([]() -> bool
 						{
 							return true;
 						}),
-					FIsActionChecked::CreateLambda([this]() -> bool
+					FIsActionChecked::CreateLambda([Section]() -> bool
 						{
-							return Section.bMatchRotationYaw;
+							return Section->bMatchRotationYaw;
 						})),
 				NAME_None,
 				EUserInterfaceActionType::ToggleButton);
@@ -1118,7 +1167,7 @@ void FCommonAnimationSection::BuildSectionContextMenu(FMenuBuilder& MenuBuilder,
 			LOCTEXT("MatchPitchRotationTooltip", "Match the Pitch Rotation, may want this off for better matching"),
 			FSlateIcon(),
 			FUIAction(
-				FExecuteAction::CreateLambda([this, ObjectBinding]()
+				FExecuteAction::CreateLambda([this, ObjectBinding, Section]()
 					{
 						const TSharedPtr<ISequencer> SequencerPtr = Sequencer.Pin();
 						if (!SequencerPtr.IsValid())
@@ -1129,18 +1178,18 @@ void FCommonAnimationSection::BuildSectionContextMenu(FMenuBuilder& MenuBuilder,
 						USkeletalMeshComponent* const SkelMeshComp = UE::Sequencer::FCommonAnimationTrackEditor::AcquireSkeletalMeshFromObjectGuid(ObjectBinding, SequencerPtr);
 
 						FScopedTransaction MatchTransaction(LOCTEXT("MatchPitchRotation_Transaction", "Match Pitch Rotation"));
-						Section.Modify();
-						Section.ToggleMatchIncludePitchRotation();
-						Section.MatchSectionByBoneTransform(SkelMeshComp, SequencerPtr->GetLocalTime().Time, SequencerPtr->GetLocalTime().Rate, Section.MatchedBoneName);
+						Section->Modify();
+						Section->ToggleMatchIncludePitchRotation();
+						Section->MatchSectionByBoneTransform(SkelMeshComp, SequencerPtr->GetLocalTime().Time, SequencerPtr->GetLocalTime().Rate, Section->MatchedBoneName);
 						SequencerPtr->NotifyMovieSceneDataChanged(EMovieSceneDataChangeType::TrackValueChanged);
 					}),
 					FCanExecuteAction::CreateLambda([]() -> bool
 						{
 							return true;
 						}),
-					FIsActionChecked::CreateLambda([this]() -> bool
+					FIsActionChecked::CreateLambda([Section]() -> bool
 						{
-							return Section.bMatchRotationPitch;
+							return Section->bMatchRotationPitch;
 						})),
 			NAME_None,
 			EUserInterfaceActionType::ToggleButton);
@@ -1150,7 +1199,7 @@ void FCommonAnimationSection::BuildSectionContextMenu(FMenuBuilder& MenuBuilder,
 			LOCTEXT("MatchRollRotationTooltip", "Match the Roll Rotation, may want this off for better matching"),
 			FSlateIcon(),
 			FUIAction(
-				FExecuteAction::CreateLambda([this, ObjectBinding]()
+				FExecuteAction::CreateLambda([this, ObjectBinding, Section]()
 					{
 						const TSharedPtr<ISequencer> SequencerPtr = Sequencer.Pin();
 						if (!SequencerPtr.IsValid())
@@ -1161,18 +1210,18 @@ void FCommonAnimationSection::BuildSectionContextMenu(FMenuBuilder& MenuBuilder,
 						USkeletalMeshComponent* const SkelMeshComp = UE::Sequencer::FCommonAnimationTrackEditor::AcquireSkeletalMeshFromObjectGuid(ObjectBinding, SequencerPtr);
 
 						FScopedTransaction MatchTransaction(LOCTEXT("MatchRollRotation_Transaction", "Match Roll Rotation"));
-						Section.Modify();
-						Section.ToggleMatchIncludeRollRotation();
-						Section.MatchSectionByBoneTransform(SkelMeshComp, SequencerPtr->GetLocalTime().Time, SequencerPtr->GetLocalTime().Rate, Section.MatchedBoneName);
+						Section->Modify();
+						Section->ToggleMatchIncludeRollRotation();
+						Section->MatchSectionByBoneTransform(SkelMeshComp, SequencerPtr->GetLocalTime().Time, SequencerPtr->GetLocalTime().Rate, Section->MatchedBoneName);
 						SequencerPtr->NotifyMovieSceneDataChanged(EMovieSceneDataChangeType::TrackValueChanged);
 					}),
 				FCanExecuteAction::CreateLambda([]() -> bool
 					{
 						return true;
 					}),
-				FIsActionChecked::CreateLambda([this]() -> bool
+				FIsActionChecked::CreateLambda([Section]() -> bool
 					{
-						return Section.bMatchRotationRoll;
+						return Section->bMatchRotationRoll;
 					})),
 			NAME_None,
 			EUserInterfaceActionType::ToggleButton);
@@ -1186,7 +1235,7 @@ void FCommonAnimationSection::BuildSectionContextMenu(FMenuBuilder& MenuBuilder,
 			NSLOCTEXT("Sequencer", "ShowSkeletonsTooltip", "Show A Skeleton for this Section."),
 			FSlateIcon(),
 			FUIAction(
-				FExecuteAction::CreateLambda([this]()
+				FExecuteAction::CreateLambda([this, Section]()
 					{
 						const TSharedPtr<ISequencer> SequencerPtr = Sequencer.Pin();
 						if (!SequencerPtr.IsValid())
@@ -1194,16 +1243,16 @@ void FCommonAnimationSection::BuildSectionContextMenu(FMenuBuilder& MenuBuilder,
 							return;
 						}
 
-						Section.ToggleShowSkeleton();
+						Section->ToggleShowSkeleton();
 						SequencerPtr->NotifyMovieSceneDataChanged(EMovieSceneDataChangeType::TrackValueChanged);
 					}),
 				FCanExecuteAction::CreateLambda([this]() -> bool
 					{
 						return Sequencer.IsValid();
 					}),
-				FIsActionChecked::CreateLambda([this]() -> bool
+				FIsActionChecked::CreateLambda([Section]() -> bool
 					{
-						return Section.bShowSkeleton;
+						return Section->bShowSkeleton;
 					})),
 				NAME_None,
 				EUserInterfaceActionType::ToggleButton);
diff --git a/Engine/Source/Editor/MovieSceneTools/Public/Channels/BuiltInChannelEditors.h b/Engine/Source/Editor/MovieSceneTools/Public/Channels/BuiltInChannelEditors.h
new file mode 100644
--- /dev/null
+++ b/Engine/Source/Editor/MovieSceneTools/Public/Channels/BuiltInChannelEditors.h
@@ -0,0 +1,114 @@
+// Copyright Epic Games, Inc. All Rights Reserved.
+
+#pragma once
+
+#include "CoreTypes.h"
+#include "Templates/SharedPointer.h"
+#include "Containers/Array.h"
+#include "Containers/ContainersFwd.h"
+#include "UObject/ObjectMacros.h"
+#include "UObject/Class.h"
+
+#include "MovieSceneKeyStruct.h"
+#include "SequencerChannelTraits.h"
+#include "Channels/MovieSceneChannelHandle.h"
+
+#include "Channels/MovieSceneBoolChannel.h"
+#include "Channels/MovieSceneByteChannel.h"
+#include "Channels/MovieSceneDoubleChannel.h"
+#include "Channels/MovieSceneFloatChannel.h"
+#include "Channels/MovieSceneIntegerChannel.h"
+#include "Channels/MovieSceneEventChannel.h"
+#include "Channels/MovieSceneObjectPathChannel.h"
+#include "Channels/MovieSceneTimeWarpChannel.h"
+#include "Sections/MovieSceneStringSection.h"
+#include "Sections/MovieSceneParticleSection.h"
+#include "Sections/MovieSceneActorReferenceSection.h"
+
+struct FKeyHandle;
+struct FKeyDrawParams;
+
+class SWidget;
+class ISequencer;
+class FMenuBuilder;
+class FStructOnScope;
+class ISectionLayoutBuilder;
+
+/** Overrides for adding or updating a key for non-standard channels */
+MOVIESCENETOOLS_API FKeyHandle AddOrUpdateKey(FMovieSceneDoubleChannel* Channel, UMovieSceneSection* SectionToKey,  const TMovieSceneExternalValue<double>& EditorData, FFrameNumber InTime, ISequencer& Sequencer, const FGuid& InObjectBindingID, FTrackInstancePropertyBindings* PropertyBindings);
+MOVIESCENETOOLS_API FKeyHandle AddOrUpdateKey(FMovieSceneFloatChannel* Channel, UMovieSceneSection* SectionToKey,  const TMovieSceneExternalValue<float>& EditorData, FFrameNumber InTime, ISequencer& Sequencer, const FGuid& InObjectBindingID, FTrackInstancePropertyBindings* PropertyBindings);
+MOVIESCENETOOLS_API FKeyHandle AddOrUpdateKey(FMovieSceneActorReferenceData* Channel, UMovieSceneSection* SectionToKey, FFrameNumber InTime, ISequencer& Sequencer, const FGuid& InObjectBindingID, FTrackInstancePropertyBindings* PropertyBindings);
+MOVIESCENETOOLS_API FKeyHandle AddOrUpdateKey(FMovieSceneTimeWarpChannel* Channel, UMovieSceneSection* SectionToKey, FFrameNumber InTime, ISequencer& Sequencer, const FGuid& InObjectBindingID, FTrackInstancePropertyBindings* PropertyBindings);
+
+/** Key editor overrides */
+MOVIESCENETOOLS_API bool CanCreateKeyEditor(const FMovieSceneBoolChannel*       Channel);
+MOVIESCENETOOLS_API bool CanCreateKeyEditor(const FMovieSceneByteChannel*       Channel);
+MOVIESCENETOOLS_API bool CanCreateKeyEditor(const FMovieSceneIntegerChannel*    Channel);
+MOVIESCENETOOLS_API bool CanCreateKeyEditor(const FMovieSceneDoubleChannel*      Channel);
+MOVIESCENETOOLS_API bool CanCreateKeyEditor(const FMovieSceneFloatChannel*      Channel);
+MOVIESCENETOOLS_API bool CanCreateKeyEditor(const FMovieSceneStringChannel*     Channel);
+MOVIESCENETOOLS_API bool CanCreateKeyEditor(const FMovieSceneObjectPathChannel* Channel);
+MOVIESCENETOOLS_API bool CanCreateKeyEditor(const FMovieSceneActorReferenceData* Channel);
+
+MOVIESCENETOOLS_API TSharedRef<SWidget> CreateKeyEditor(const TMovieSceneChannelHandle<FMovieSceneBoolChannel>&        Channel, const UE::Sequencer::FCreateKeyEditorParams& Params);
+MOVIESCENETOOLS_API TSharedRef<SWidget> CreateKeyEditor(const TMovieSceneChannelHandle<FMovieSceneByteChannel>&        Channel, const UE::Sequencer::FCreateKeyEditorParams& Params);
+MOVIESCENETOOLS_API TSharedRef<SWidget> CreateKeyEditor(const TMovieSceneChannelHandle<FMovieSceneIntegerChannel>&     Channel, const UE::Sequencer::FCreateKeyEditorParams& Params);
+MOVIESCENETOOLS_API TSharedRef<SWidget> CreateKeyEditor(const TMovieSceneChannelHandle<FMovieSceneDoubleChannel>&      Channel, const UE::Sequencer::FCreateKeyEditorParams& Params);
+MOVIESCENETOOLS_API TSharedRef<SWidget> CreateKeyEditor(const TMovieSceneChannelHandle<FMovieSceneFloatChannel>&       Channel, const UE::Sequencer::FCreateKeyEditorParams& Params);
+MOVIESCENETOOLS_API TSharedRef<SWidget> CreateKeyEditor(const TMovieSceneChannelHandle<FMovieSceneStringChannel>&      Channel, const UE::Sequencer::FCreateKeyEditorParams& Params);
+MOVIESCENETOOLS_API TSharedRef<SWidget> CreateKeyEditor(const TMovieSceneChannelHandle<FMovieSceneObjectPathChannel>&  Channel, const UE::Sequencer::FCreateKeyEditorParams& Params);
+MOVIESCENETOOLS_API TSharedRef<SWidget> CreateKeyEditor(const TMovieSceneChannelHandle<FMovieSceneActorReferenceData>& Channel, const UE::Sequencer::FCreateKeyEditorParams& Params);
+MOVIESCENETOOLS_API TSharedRef<SWidget> CreateKeyEditor(const TMovieSceneChannelHandle<FMovieSceneTimeWarpChannel>&    Channel, const UE::Sequencer::FCreateKeyEditorParams& Params);
+
+MOVIESCENETOOLS_API UMovieSceneKeyStructType* InstanceGeneratedStruct(FMovieSceneByteChannel* Channel,       FSequencerKeyStructGenerator* Generator);
+MOVIESCENETOOLS_API UMovieSceneKeyStructType* InstanceGeneratedStruct(FMovieSceneTimeWarpChannel* Channel,   FSequencerKeyStructGenerator* Generator);
+MOVIESCENETOOLS_API UMovieSceneKeyStructType* InstanceGeneratedStruct(FMovieSceneObjectPathChannel* Channel, FSequencerKeyStructGenerator* Generator);
+
+MOVIESCENETOOLS_API void PostConstructKeyInstance(const TMovieSceneChannelHandle<FMovieSceneTimeWarpChannel>& ChannelHandle, FKeyHandle InHandle, FStructOnScope* Struct);
+MOVIESCENETOOLS_API void PostConstructKeyInstance(const TMovieSceneChannelHandle<FMovieSceneObjectPathChannel>& ChannelHandle, FKeyHandle InHandle, FStructOnScope* Struct);
+
+
+/** Key drawing overrides */
+MOVIESCENETOOLS_API void DrawKeys(FMovieSceneDoubleChannel*   Channel, TArrayView<const FKeyHandle> InKeyHandles, const UMovieSceneSection* InOwner, TArrayView<FKeyDrawParams> OutKeyDrawParams);
+MOVIESCENETOOLS_API void DrawKeys(FMovieSceneFloatChannel*    Channel, TArrayView<const FKeyHandle> InKeyHandles, const UMovieSceneSection* InOwner, TArrayView<FKeyDrawParams> OutKeyDrawParams);
+MOVIESCENETOOLS_API void DrawKeys(FMovieSceneParticleChannel* Channel, TArrayView<const FKeyHandle> InKeyHandles, const UMovieSceneSection* InOwner, TArrayView<FKeyDrawParams> OutKeyDrawParams);
+MOVIESCENETOOLS_API void DrawKeys(FMovieSceneEventChannel*    Channel, TArrayView<const FKeyHandle> InKeyHandles, const UMovieSceneSection* InOwner, TArrayView<FKeyDrawParams> OutKeyDrawParams);
+
+/** Context menu overrides */
+MOVIESCENETOOLS_API void ExtendSectionMenu(FMenuBuilder& OuterMenuBuilder, TSharedPtr<FExtender> MenuExtender, TArray<TMovieSceneChannelHandle<FMovieSceneDoubleChannel>>&& Channels, const TArray<TWeakObjectPtr<UMovieSceneSection>>& InWeakSections, TWeakPtr<ISequencer> InWeakSequencer);
+MOVIESCENETOOLS_API void ExtendSectionMenu(FMenuBuilder& OuterMenuBuilder, TSharedPtr<FExtender> MenuExtender, TArray<TMovieSceneChannelHandle<FMovieSceneFloatChannel>>&& Channels, const TArray<TWeakObjectPtr<UMovieSceneSection>>& InWeakSections, TWeakPtr<ISequencer> InWeakSequencer);
+MOVIESCENETOOLS_API void ExtendSectionMenu(FMenuBuilder& OuterMenuBuilder, TSharedPtr<FExtender> MenuExtender, TArray<TMovieSceneChannelHandle<FMovieSceneIntegerChannel>>&& Channels, const TArray<TWeakObjectPtr<UMovieSceneSection>>& InWeakSections, TWeakPtr<ISequencer> InWeakSequencer);
+MOVIESCENETOOLS_API void ExtendSectionMenu(FMenuBuilder& OuterMenuBuilder, TSharedPtr<FExtender> MenuExtender, TArray<TMovieSceneChannelHandle<FMovieSceneBoolChannel>>&& Channels, const TArray<TWeakObjectPtr<UMovieSceneSection>>& InWeakSections, TWeakPtr<ISequencer> InWeakSequencer);
+MOVIESCENETOOLS_API void ExtendSectionMenu(FMenuBuilder& OuterMenuBuilder, TSharedPtr<FExtender> MenuExtender, TArray<TMovieSceneChannelHandle<FMovieSceneByteChannel>>&& Channels, const TArray<TWeakObjectPtr<UMovieSceneSection>>& InWeakSections, TWeakPtr<ISequencer> InWeakSequencer);
+
+MOVIESCENETOOLS_API TSharedPtr<ISidebarChannelExtension> ExtendSidebarMenu(FMenuBuilder& OuterMenuBuilder, TSharedPtr<FExtender> MenuExtender, TArray<TMovieSceneChannelHandle<FMovieSceneDoubleChannel>>&& InChannels, const TArray<TWeakObjectPtr<UMovieSceneSection>>& InWeakSections, TWeakPtr<ISequencer> InWeakSequencer);
+MOVIESCENETOOLS_API TSharedPtr<ISidebarChannelExtension> ExtendSidebarMenu(FMenuBuilder& OuterMenuBuilder, TSharedPtr<FExtender> MenuExtender, TArray<TMovieSceneChannelHandle<FMovieSceneFloatChannel>>&& InChannels,const  TArray<TWeakObjectPtr<UMovieSceneSection>>& InWeakSections, TWeakPtr<ISequencer> InWeakSequencer);
+MOVIESCENETOOLS_API TSharedPtr<ISidebarChannelExtension> ExtendSidebarMenu(FMenuBuilder& OuterMenuBuilder, TSharedPtr<FExtender> MenuExtender, TArray<TMovieSceneChannelHandle<FMovieSceneIntegerChannel>>&& InChannels, const TArray<TWeakObjectPtr<UMovieSceneSection>>& InWeakSections, TWeakPtr<ISequencer> InWeakSequencer);
+MOVIESCENETOOLS_API TSharedPtr<ISidebarChannelExtension> ExtendSidebarMenu(FMenuBuilder& OuterMenuBuilder, TSharedPtr<FExtender> MenuExtender, TArray<TMovieSceneChannelHandle<FMovieSceneBoolChannel>>&& InChannels, const TArray<TWeakObjectPtr<UMovieSceneSection>>& InWeakSections, TWeakPtr<ISequencer> InWeakSequencer);
+MOVIESCENETOOLS_API TSharedPtr<ISidebarChannelExtension> ExtendSidebarMenu(FMenuBuilder& OuterMenuBuilder, TSharedPtr<FExtender> MenuExtender, TArray<TMovieSceneChannelHandle<FMovieSceneByteChannel>>&& InChannels, const TArray<TWeakObjectPtr<UMovieSceneSection>>& InWeakSections, TWeakPtr<ISequencer> InWeakSequencer);
+
+MOVIESCENETOOLS_API void ExtendKeyMenu(FMenuBuilder& OuterMenuBuilder, TSharedPtr<FExtender> MenuExtender, TArray<TExtendKeyMenuParams<FMovieSceneDoubleChannel>>&& Channels, TWeakPtr<ISequencer> InSequencer);
+MOVIESCENETOOLS_API void ExtendKeyMenu(FMenuBuilder& OuterMenuBuilder, TSharedPtr<FExtender> MenuExtender, TArray<TExtendKeyMenuParams<FMovieSceneFloatChannel>>&& Channels, TWeakPtr<ISequencer> InSequencer);
+MOVIESCENETOOLS_API void ExtendKeyMenu(FMenuBuilder& OuterMenuBuilder, TSharedPtr<FExtender> MenuExtender, TArray<TExtendKeyMenuParams<FMovieSceneTimeWarpChannel>>&& Channels, TWeakPtr<ISequencer> InSequencer);
+
+/** Curve editor models */
+MOVIESCENETOOLS_API inline bool SupportsCurveEditorModels(const TMovieSceneChannelHandle<FMovieSceneDoubleChannel>& DoubleChannel) { return true; }
+MOVIESCENETOOLS_API inline bool SupportsCurveEditorModels(const TMovieSceneChannelHandle<FMovieSceneFloatChannel>& FloatChannel) { return true; }
+MOVIESCENETOOLS_API inline bool SupportsCurveEditorModels(const TMovieSceneChannelHandle<FMovieSceneIntegerChannel>& IntegerChannel) { return true; }
+MOVIESCENETOOLS_API inline bool SupportsCurveEditorModels(const TMovieSceneChannelHandle<FMovieSceneBoolChannel>& BoolChannel) { return true; }
+MOVIESCENETOOLS_API inline bool SupportsCurveEditorModels(const TMovieSceneChannelHandle<FMovieSceneByteChannel>& ByteChannel) { return true; }
+MOVIESCENETOOLS_API inline bool SupportsCurveEditorModels(const TMovieSceneChannelHandle<FMovieSceneEventChannel>& EventChannel) { return true; }
+MOVIESCENETOOLS_API inline bool SupportsCurveEditorModels(const TMovieSceneChannelHandle<FMovieSceneTimeWarpChannel>& TimeWarpChannel) { return true; }
+
+MOVIESCENETOOLS_API TUniquePtr<FCurveModel> CreateCurveEditorModel(const TMovieSceneChannelHandle<FMovieSceneDoubleChannel>& DoubleChannel, const UE::Sequencer::FCreateCurveEditorModelParams& Params);
+MOVIESCENETOOLS_API TUniquePtr<FCurveModel> CreateCurveEditorModel(const TMovieSceneChannelHandle<FMovieSceneFloatChannel>& FloatChannel, const UE::Sequencer::FCreateCurveEditorModelParams& Params);
+MOVIESCENETOOLS_API TUniquePtr<FCurveModel> CreateCurveEditorModel(const TMovieSceneChannelHandle<FMovieSceneIntegerChannel>& IntegerChannel, const UE::Sequencer::FCreateCurveEditorModelParams& Params);
+MOVIESCENETOOLS_API TUniquePtr<FCurveModel> CreateCurveEditorModel(const TMovieSceneChannelHandle<FMovieSceneBoolChannel>& BoolChannel, const UE::Sequencer::FCreateCurveEditorModelParams& Params);
+MOVIESCENETOOLS_API TUniquePtr<FCurveModel> CreateCurveEditorModel(const TMovieSceneChannelHandle<FMovieSceneByteChannel>& ByteChannel, const UE::Sequencer::FCreateCurveEditorModelParams& Params);
+MOVIESCENETOOLS_API TUniquePtr<FCurveModel> CreateCurveEditorModel(const TMovieSceneChannelHandle<FMovieSceneEventChannel>& EventChannel, const UE::Sequencer::FCreateCurveEditorModelParams& Params);
+MOVIESCENETOOLS_API TUniquePtr<FCurveModel> CreateCurveEditorModel(const TMovieSceneChannelHandle<FMovieSceneTimeWarpChannel>& TimeWarpChannel, const UE::Sequencer::FCreateCurveEditorModelParams& Params);
+
+MOVIESCENETOOLS_API bool ShouldShowCurve(const FMovieSceneFloatChannel* Channel, UMovieSceneSection* InSection);
+MOVIESCENETOOLS_API bool ShouldShowCurve(const FMovieSceneDoubleChannel* Channel, UMovieSceneSection* InSection);
+
+MOVIESCENETOOLS_API TSharedPtr<UE::Sequencer::FChannelModel> CreateChannelModel(const TMovieSceneChannelHandle<FMovieSceneTimeWarpChannel>& InChannelHandle, const UE::Sequencer::FSectionModel& InSection, FName InChannelName);
\ No newline at end of file
diff --git a/Engine/Source/Editor/MovieSceneTools/Public/TrackEditors/CommonAnimationTrackEditor.h b/Engine/Source/Editor/MovieSceneTools/Public/TrackEditors/CommonAnimationTrackEditor.h
--- a/Engine/Source/Editor/MovieSceneTools/Public/TrackEditors/CommonAnimationTrackEditor.h
+++ b/Engine/Source/Editor/MovieSceneTools/Public/TrackEditors/CommonAnimationTrackEditor.h
@@ -178,7 +178,7 @@ protected:
 protected:
 
 	/** The section we are visualizing */
-	UMovieSceneSkeletalAnimationSection& Section;
+	TWeakObjectPtr<UMovieSceneSkeletalAnimationSection> WeakSection;
 
 	/** Used to draw animation frame, need selection state and local time*/
 	TWeakPtr<ISequencer> Sequencer;
