Sequencer: Fix call to BuildEditSectionMenu with WeakSequencer and crashes in MovieSceneConditionCustomization when Sequencer is null

#rnx

diff --git a/Engine/Source/Editor/MovieSceneTools/Private/Conditions/MovieSceneConditionCustomization.cpp b/Engine/Source/Editor/MovieSceneTools/Private/Conditions/MovieSceneConditionCustomization.cpp
--- a/Engine/Source/Editor/MovieSceneTools/Private/Conditions/MovieSceneConditionCustomization.cpp
+++ b/Engine/Source/Editor/MovieSceneTools/Private/Conditions/MovieSceneConditionCustomization.cpp
@@ -83,6 +83,13 @@ TSharedRef<IPropertyTypeCustomization> FMovieSceneConditionCustomization::MakeIn
 void FMovieSceneConditionCustomization::CustomizeHeader(TSharedRef<IPropertyHandle> InPropertyHandle, FDetailWidgetRow& HeaderRow, IPropertyTypeCustomizationUtils& CustomizationUtils)
 {
 	ConditionContainerPropertyHandle = InPropertyHandle;
+
+	if (!Sequencer.IsValid())
+	{
+		ConditionContainerPropertyHandle->MarkHiddenByCustomization();
+		return;
+	}
+
 	ConditionPropertyHandle = InPropertyHandle->GetChildHandle(GET_MEMBER_NAME_CHECKED(FMovieSceneConditionContainer, Condition));
 
 	if (!Sequence.IsValid())
@@ -165,6 +172,11 @@ void FMovieSceneConditionCustomization::CustomizeHeader(TSharedRef<IPropertyHand
 
 void FMovieSceneConditionCustomization::CustomizeChildren(TSharedRef<IPropertyHandle> InPropertyHandle, IDetailChildrenBuilder& ChildBuilder, IPropertyTypeCustomizationUtils& CustomizationUtils)
 {
+	if (!Sequencer.IsValid())
+	{
+		return;
+	}
+
 	// If conditions not allowed, hide condition property functionality
 	if (!Sequence.IsValid() || !Sequence->GetMovieScene()->IsConditionClassAllowed(UMovieSceneCondition::StaticClass()) || (Track.IsValid() && !Track->SupportsConditions()))
 	{
diff --git a/Engine/Source/Editor/Sequencer/Private/Widgets/Sidebar/SequencerSelectionDrawer.cpp b/Engine/Source/Editor/Sequencer/Private/Widgets/Sidebar/SequencerSelectionDrawer.cpp
--- a/Engine/Source/Editor/Sequencer/Private/Widgets/Sidebar/SequencerSelectionDrawer.cpp
+++ b/Engine/Source/Editor/Sequencer/Private/Widgets/Sidebar/SequencerSelectionDrawer.cpp
@@ -436,7 +436,7 @@ void FSequencerSelectionDrawer::BuildOutlinerDetails(const TSharedRef<FSequencer
 
 	if (!WeakTrackObjects.IsEmpty())
 	{
-		SequencerHelpers::BuildEditTrackMenu(Sequencer, WeakTrackObjects, MenuBuilder, false);
+		SequencerHelpers::BuildEditTrackMenu(WeakSequencer, WeakTrackObjects, MenuBuilder, false);
 	}
 
 	// Only show section details for a selected track/trackrow if a single infinite section on the row
