Sequencer: Fix rotations 360 for additive sections by disabling winding fixes for addiitves.

diff --git a/Engine/Source/Editor/MovieSceneTools/Private/TrackEditors/TransformTrackEditor.cpp b/Engine/Source/Editor/MovieSceneTools/Private/TrackEditors/TransformTrackEditor.cpp
--- a/Engine/Source/Editor/MovieSceneTools/Private/TrackEditors/TransformTrackEditor.cpp
+++ b/Engine/Source/Editor/MovieSceneTools/Private/TrackEditors/TransformTrackEditor.cpp
@@ -931,7 +931,8 @@ void F3DTransformTrackEditor::GetTransformKeys( const TOptional<FTransformData>&
 		}
 
 		// Do we need to unwind re-composed rotations?
-		KeyRotator = UnwindRotator(CurrentTransform.Rotation, RecomposedTransform.Rotation);
+		KeyRotator = (Section->GetBlendType() != EMovieSceneBlendType::Additive && Section->GetBlendType() != EMovieSceneBlendType::AdditiveFromBase) ?
+			UnwindRotator(CurrentTransform.Rotation, RecomposedTransform.Rotation) : RecomposedTransform.Rotation;
 		OutGeneratedKeys.Add(FMovieSceneChannelValueSetter::Create<FMovieSceneDoubleChannel>(ChannelHandles[3].GetChannelIndex(), (double)KeyRotator.Roll, bKeyX));
 		OutGeneratedKeys.Add(FMovieSceneChannelValueSetter::Create<FMovieSceneDoubleChannel>(ChannelHandles[4].GetChannelIndex(), (double)KeyRotator.Pitch, bKeyY));
 		OutGeneratedKeys.Add(FMovieSceneChannelValueSetter::Create<FMovieSceneDoubleChannel>(ChannelHandles[5].GetChannelIndex(), (double)KeyRotator.Yaw, bKeyZ));
@@ -1976,3 +1977,4 @@ void F3DTransformTrackEditor::ClearOutConstraintDelegates()
 	UTickableTransformConstraint::GetOnConstraintChanged().RemoveAll(this);
 }
 #undef LOCTEXT_NAMESPACE
+
