Take Recorder: Include TransformOrigin while recording non-attached transforms

diff --git a/Engine/Plugins/VirtualProduction/Takes/Source/TakeRecorder/Private/Recorder/TakeRecorder.cpp b/Engine/Plugins/VirtualProduction/Takes/Source/TakeRecorder/Private/Recorder/TakeRecorder.cpp
--- a/Engine/Plugins/VirtualProduction/Takes/Source/TakeRecorder/Private/Recorder/TakeRecorder.cpp
+++ b/Engine/Plugins/VirtualProduction/Takes/Source/TakeRecorder/Private/Recorder/TakeRecorder.cpp
@@ -53,6 +53,7 @@
 // LevelEditor includes
 #include "IAssetViewport.h"
 #include "LevelEditor.h"
+#include "MovieSceneToolHelpers.h"
 #include "Misc/GuardCVar.h"
 #include "Subsystems/AssetEditorSubsystem.h"
 
@@ -843,9 +844,14 @@ void UTakeRecorder::PreRecord()
 	TakeRecorderSourcesSettings.bRemoveRedundantTracks = Parameters.User.bRemoveRedundantTracks;
 	TakeRecorderSourcesSettings.bAutoLock = Parameters.User.bAutoLock;
 
+	TSharedPtr<ISequencer> Sequencer = WeakSequencer.Pin();
+	if (Sequencer)
+	{
+		TakeRecorderSourcesSettings.TransformOrigin = MovieSceneToolHelpers::GetTransformOriginForFocusedSequence(Sequencer);
+	}
+
 	Sources->SetSettings(TakeRecorderSourcesSettings);
 
-	TSharedPtr<ISequencer> Sequencer = WeakSequencer.Pin();
 	if (!Parameters.bDisableRecordingAndSave)
 	{
 		UMovieScene* MovieScene = SequenceAsset->GetMovieScene();
diff --git a/Engine/Plugins/VirtualProduction/Takes/Source/TakeRecorderSources/Private/TakeRecorderActorSource.cpp b/Engine/Plugins/VirtualProduction/Takes/Source/TakeRecorderSources/Private/TakeRecorderActorSource.cpp
--- a/Engine/Plugins/VirtualProduction/Takes/Source/TakeRecorderSources/Private/TakeRecorderActorSource.cpp
+++ b/Engine/Plugins/VirtualProduction/Takes/Source/TakeRecorderSources/Private/TakeRecorderActorSource.cpp
@@ -1712,6 +1712,8 @@ FTrackRecorderSettings UTakeRecorderActorSource::GetTrackRecorderSettings() cons
 	TrackRecorderSettings.IncludeAnimationNames = IncludeAnimationNames;
 	TrackRecorderSettings.ExcludeAnimationNames = ExcludeAnimationNames;
 
+	TrackRecorderSettings.TransformOrigin = Sources->GetSettings().TransformOrigin;
+
 	return TrackRecorderSettings;
 }
 
diff --git a/Engine/Plugins/VirtualProduction/Takes/Source/TakeTrackRecorders/Private/MovieScene3DTransformTrackRecorder.cpp b/Engine/Plugins/VirtualProduction/Takes/Source/TakeTrackRecorders/Private/MovieScene3DTransformTrackRecorder.cpp
--- a/Engine/Plugins/VirtualProduction/Takes/Source/TakeTrackRecorders/Private/MovieScene3DTransformTrackRecorder.cpp
+++ b/Engine/Plugins/VirtualProduction/Takes/Source/TakeTrackRecorders/Private/MovieScene3DTransformTrackRecorder.cpp
@@ -158,6 +158,10 @@ void UMovieScene3DTransformTrackRecorder::FinalizeTrackImpl()
 				}
 			}
 		}
+		else if (BufferedTransforms.Times.Num() > 0)
+		{
+			BufferedTransforms.PostMultTransform(OwningTakeRecorderSource->GetTrackRecorderSettings().TransformOrigin.Inverse());
+		}
 	}
 
  	// Try to 're-wind' rotations that look like axis flips
diff --git a/Engine/Plugins/VirtualProduction/Takes/Source/TakeTrackRecorders/Public/TrackRecorders/IMovieSceneTrackRecorderHost.h b/Engine/Plugins/VirtualProduction/Takes/Source/TakeTrackRecorders/Public/TrackRecorders/IMovieSceneTrackRecorderHost.h
--- a/Engine/Plugins/VirtualProduction/Takes/Source/TakeTrackRecorders/Public/TrackRecorders/IMovieSceneTrackRecorderHost.h
+++ b/Engine/Plugins/VirtualProduction/Takes/Source/TakeTrackRecorders/Public/TrackRecorders/IMovieSceneTrackRecorderHost.h
@@ -54,6 +54,7 @@ struct FTrackRecorderSettings
 	TArray<FTakeRecorderTrackSettings> DefaultTracks;
 	TArray<FString> IncludeAnimationNames;
 	TArray<FString> ExcludeAnimationNames;
+	FTransform TransformOrigin;
 
 	static bool IsDefaultPropertyTrack(UObject* InObjectToRecord, const FString& InPropertyPath, const TArray<FTakeRecorderTrackSettings>& DefaultTracks)
 	{
diff --git a/Engine/Plugins/VirtualProduction/Takes/Source/TakesCore/Public/TakeRecorderSources.h b/Engine/Plugins/VirtualProduction/Takes/Source/TakesCore/Public/TakeRecorderSources.h
--- a/Engine/Plugins/VirtualProduction/Takes/Source/TakesCore/Public/TakeRecorderSources.h
+++ b/Engine/Plugins/VirtualProduction/Takes/Source/TakesCore/Public/TakeRecorderSources.h
@@ -24,6 +24,7 @@ struct FTakeRecorderSourcesSettings
 	bool bSaveRecordedAssets;
 	bool bAutoLock;
 	bool bRemoveRedundantTracks;
+	FTransform TransformOrigin;
 };
 
 /**
diff --git a/Engine/Source/Editor/MovieSceneTools/Private/MovieSceneToolHelpers.cpp b/Engine/Source/Editor/MovieSceneTools/Private/MovieSceneToolHelpers.cpp
--- a/Engine/Source/Editor/MovieSceneTools/Private/MovieSceneToolHelpers.cpp
+++ b/Engine/Source/Editor/MovieSceneTools/Private/MovieSceneToolHelpers.cpp
@@ -96,6 +96,8 @@
 #include "Tracks/MovieSceneBindingLifetimeTrack.h"
 #include "Sections/MovieSceneBindingLifetimeSection.h"
 #include "GameFramework/Character.h"
+#include "Systems/MovieSceneTransformOriginSystem.h"
+#include "Tracks/IMovieSceneTransformOrigin.h"
 
 /* FSkelMeshRecorder
  ***********/
@@ -330,6 +332,54 @@ bool MovieSceneToolHelpers::CanSplitSection(const TSet<UMovieSceneSection*>& Sec
 	return false;
 }
 
+FTransform MovieSceneToolHelpers::GetTransformOriginForFocusedSequence(TSharedPtr<ISequencer> InSequencer)
+{
+	FTransform TransformOrigin;
+
+	const IMovieScenePlaybackClient*  Client       = InSequencer->GetPlaybackClient();
+	const UObject*                    InstanceData = Client ? Client->GetInstanceData() : nullptr;
+	const IMovieSceneTransformOrigin* RawInterface = Cast<const IMovieSceneTransformOrigin>(InstanceData);
+
+	const bool bHasInterface = RawInterface || (InstanceData && InstanceData->GetClass()->ImplementsInterface(UMovieSceneTransformOrigin::StaticClass()));
+	if (bHasInterface)
+	{
+		// Retrieve the current origin
+		TransformOrigin = RawInterface ? RawInterface->GetTransformOrigin() : IMovieSceneTransformOrigin::Execute_BP_GetTransformOrigin(InstanceData);
+	}
+	
+	const TArray<FMovieSceneSequenceID>& Hierarchy = InSequencer->GetSubSequenceHierarchy();
+
+	const FMovieSceneRootEvaluationTemplateInstance& EvaluationTemplate = InSequencer->GetEvaluationTemplate();
+	const UMovieSceneEntitySystemLinker* EntityLinker = EvaluationTemplate.GetEntitySystemLinker();
+	if(!EntityLinker)
+	{
+		return TransformOrigin;
+	}
+
+	const UMovieSceneTransformOriginSystem* TransformOriginSystem = EntityLinker->FindSystem<UMovieSceneTransformOriginSystem>();
+
+	if(!TransformOriginSystem)
+	{
+		return TransformOrigin;
+	}
+
+	const TMap<FMovieSceneSequenceID, UE::MovieScene::FInstanceHandle> SequenceIDToInstanceHandle = TransformOriginSystem->GetSequenceIDToInstanceHandle();
+
+	// Transform Origins will be pre-multiplied at this step, so only retrieve the entry for the currently focused sub-sequence.
+	if(Hierarchy.Num())
+	{
+		const FMovieSceneSequenceID CurrentSequence = Hierarchy.Last();
+		if(SequenceIDToInstanceHandle.Contains(CurrentSequence))
+		{
+			const UE::MovieScene::FInstanceHandle CurrentHandle = SequenceIDToInstanceHandle[CurrentSequence];
+			// Overrides the origin only if entry is present
+			TransformOriginSystem->GetTransformOrigin(CurrentHandle, TransformOrigin);
+		}
+	}
+
+	return TransformOrigin;
+}
+
 bool MovieSceneToolHelpers::ParseShotName(const FString& InShotName, FString& ShotPrefix, uint32& ShotNumber, uint32& TakeNumber, uint32& ShotNumberDigits, uint32& TakeNumberDigits)
 {
 	// Parse a shot name
diff --git a/Engine/Source/Editor/MovieSceneTools/Private/TrackEditors/TransformTrackEditor.cpp b/Engine/Source/Editor/MovieSceneTools/Private/TrackEditors/TransformTrackEditor.cpp
--- a/Engine/Source/Editor/MovieSceneTools/Private/TrackEditors/TransformTrackEditor.cpp
+++ b/Engine/Source/Editor/MovieSceneTools/Private/TrackEditors/TransformTrackEditor.cpp
@@ -978,54 +978,7 @@ void F3DTransformTrackEditor::GetTransformKeys( const TOptional<FTransformData>&
 
 FTransform F3DTransformTrackEditor::GetTransformOrigin() const
 {
-	FTransform TransformOrigin;
-
-	const IMovieScenePlaybackClient*  Client       = GetSequencer()->GetPlaybackClient();
-	const UObject*                    InstanceData = Client ? Client->GetInstanceData() : nullptr;
-	const IMovieSceneTransformOrigin* RawInterface = Cast<const IMovieSceneTransformOrigin>(InstanceData);
-
-	const bool bHasInterface = RawInterface || (InstanceData && InstanceData->GetClass()->ImplementsInterface(UMovieSceneTransformOrigin::StaticClass()));
-	if (bHasInterface)
-	{
-		// Retrieve the current origin
-		TransformOrigin = RawInterface ? RawInterface->GetTransformOrigin() : IMovieSceneTransformOrigin::Execute_BP_GetTransformOrigin(InstanceData);
-	}
-	
-	const TArray<FMovieSceneSequenceID>& Hierarchy = GetSequencer()->GetSubSequenceHierarchy();
-
-	const FMovieSceneRootEvaluationTemplateInstance& EvaluationTemplate = GetSequencer()->GetEvaluationTemplate();
-	const UMovieSceneEntitySystemLinker* EntityLinker = EvaluationTemplate.GetEntitySystemLinker();
-	if(!EntityLinker)
-	{
-		return TransformOrigin;
-	}
-
-	const UMovieSceneTransformOriginSystem* TransformOriginSystem = EntityLinker->FindSystem<UMovieSceneTransformOriginSystem>();
-
-	if(!TransformOriginSystem)
-	{
-		return TransformOrigin;
-	}
-	
-	const TSparseArray<FTransform>& TransformOrigins = TransformOriginSystem->GetTransformOriginsByInstanceID();
-	const TMap<FMovieSceneSequenceID, UE::MovieScene::FInstanceHandle> SequenceIDToInstanceHandle = TransformOriginSystem->GetSequenceIDToInstanceHandle();
-
-	// Transform Origins will be pre-multiplied at this step, so only retrieve the entry for the currently focused sub-sequence.
-	if(Hierarchy.Num())
-	{
-		const FMovieSceneSequenceID CurrentSequence = Hierarchy.Last();
-		if(SequenceIDToInstanceHandle.Contains(CurrentSequence))
-		{
-			const UE::MovieScene::FInstanceHandle CurrentHandle = SequenceIDToInstanceHandle[CurrentSequence];
-			if(TransformOrigins.IsValidIndex(CurrentHandle.InstanceID))
-			{
-				// Override the root origin 
-				TransformOrigin = TransformOrigins[CurrentHandle.InstanceID];
-			}
-		}
-	}
-
-	return TransformOrigin;
+	return MovieSceneToolHelpers::GetTransformOriginForFocusedSequence(GetSequencer());
 }
 
 FTransform F3DTransformTrackEditor::GetRootMotionOffset(UObject* AnimatedObject) const
diff --git a/Engine/Source/Editor/MovieSceneTools/Public/MovieSceneToolHelpers.h b/Engine/Source/Editor/MovieSceneTools/Public/MovieSceneToolHelpers.h
--- a/Engine/Source/Editor/MovieSceneTools/Public/MovieSceneToolHelpers.h
+++ b/Engine/Source/Editor/MovieSceneTools/Public/MovieSceneToolHelpers.h
@@ -153,6 +153,8 @@ public:
 	static void SplitSection(const TSet<UMovieSceneSection*>& Sections, FQualifiedFrameTime Time, bool bDeleteKeys);
 	static bool CanSplitSection(const TSet<UMovieSceneSection*>& Sections, FQualifiedFrameTime Time);
 
+	static FTransform GetTransformOriginForFocusedSequence(TSharedPtr<ISequencer> InSequencer);
+
 	/**
 	 * Parse a shot name into its components.
 	 *
