Sequencer: Fix styling on the range toggle buttons

#rnx

diff --git a/Engine/Source/Editor/MovieSceneTools/Private/MovieSceneSectionDetailsCustomization.cpp b/Engine/Source/Editor/MovieSceneTools/Private/MovieSceneSectionDetailsCustomization.cpp
--- a/Engine/Source/Editor/MovieSceneTools/Private/MovieSceneSectionDetailsCustomization.cpp
+++ b/Engine/Source/Editor/MovieSceneTools/Private/MovieSceneSectionDetailsCustomization.cpp
@@ -97,8 +97,7 @@ void FMovieSceneSectionDetailsCustomization::CustomizeDetails(IDetailLayoutBuild
 			.OnClicked(this, &FMovieSceneSectionDetailsCustomization::ToggleRangeStartBounded)
 			.ContentPadding(0)
 			.ToolTipText(LOCTEXT("LockedRangeBounds", "Some sections support infinite ranges and fixed ranges. Toggling this will change the bound type."))
-			.ForegroundColor(FSlateColor::UseForeground())
-			.ButtonStyle(FAppStyle::Get(), "ToggleButton")
+			.ButtonStyle(FAppStyle::Get(), "SimpleButton")
 			[
 				SNew(STextBlock)
 				.Font(FAppStyle::Get().GetFontStyle("FontAwesome.11"))
@@ -143,8 +142,7 @@ void FMovieSceneSectionDetailsCustomization::CustomizeDetails(IDetailLayoutBuild
 			.OnClicked(this, &FMovieSceneSectionDetailsCustomization::ToggleRangeEndBounded)
 			.ContentPadding(0)
 			.ToolTipText(LOCTEXT("LockedRangeBounds", "Some sections support infinite ranges and fixed ranges. Toggling this will change the bound type."))
-			.ForegroundColor(FSlateColor::UseForeground())
-			.ButtonStyle(FAppStyle::Get(), "ToggleButton")
+			.ButtonStyle(FAppStyle::Get(), "SimpleButton")
 			[
 				SNew(STextBlock)
 				.Font(FAppStyle::Get().GetFontStyle("FontAwesome.11"))
