Sequencer: Fix key editor by only changing the value when adjusting through the spin box.
OnValueChanged is too aggressive and causes value changes on focus which then get committed incorrectly.

diff --git a/Engine/Source/Editor/MovieSceneTools/Private/CurveKeyEditors/SNumericKeyEditor.h b/Engine/Source/Editor/MovieSceneTools/Private/CurveKeyEditors/SNumericKeyEditor.h
--- a/Engine/Source/Editor/MovieSceneTools/Private/CurveKeyEditors/SNumericKeyEditor.h
+++ b/Engine/Source/Editor/MovieSceneTools/Private/CurveKeyEditors/SNumericKeyEditor.h
@@ -115,6 +115,7 @@ private:
 
 	void OnBeginSliderMovement()
 	{
+		bSliding = true;
 		GEditor->BeginTransaction(LOCTEXT("SetNumericKey", "Set Key Value"));
 	}
 
@@ -125,15 +126,17 @@ private:
 			KeyEditor->SetValue(Value);
 			GEditor->EndTransaction();
 		}
+		bSliding = false;
 	}
 
 	void OnValueChanged(NumericType Value)
 	{
-		if (!bCommitted)
+		// Update the key editor only when spinning/sliding the value since the value hasn't been committed yet.
+		// This should only happen in this case because OnValueChanged is called aggressively (ie. on focus changing)
+		if (bSliding)
 		{
 			KeyEditor->SetValueWithNotify(Value, EMovieSceneDataChangeType::TrackValueChanged);
 		}
-		bCommitted = false;
 	}
 
 	void OnValueCommitted(NumericType Value, ETextCommit::Type CommitInfo)
@@ -142,14 +145,13 @@ private:
 		{
 			const FScopedTransaction Transaction( LOCTEXT("SetNumericKey", "Set Key Value") );
 			KeyEditor->SetValueWithNotify(Value, EMovieSceneDataChangeType::TrackValueChangedRefreshImmediately);
-			bCommitted = true;
 		}
 	}
 
 private:
 
 	TSharedPtr<ISequencerKeyEditor<NumericType>> KeyEditor;
-	bool bCommitted = false;
+	bool bSliding = false;
 };
 
 template<typename ChannelType, typename NumericType>
