Anim In Engine: Reget skel mesh component after setting up restore since that may make the current one invalid.

diff --git a/Engine/Source/Editor/MovieSceneTools/Private/TrackEditors/CommonAnimationTrackEditor.cpp b/Engine/Source/Editor/MovieSceneTools/Private/TrackEditors/CommonAnimationTrackEditor.cpp
--- a/Engine/Source/Editor/MovieSceneTools/Private/TrackEditors/CommonAnimationTrackEditor.cpp
+++ b/Engine/Source/Editor/MovieSceneTools/Private/TrackEditors/CommonAnimationTrackEditor.cpp
@@ -1706,6 +1706,9 @@ bool FCommonAnimationTrackEditor::CreateAnimationSequence(const TArray<UObject*>
 				//if in subsquence and we want to, turn on ShouldEvaluateSubSequencesInIsolation, use the following object for that
 				{
 					UE::Sequencer::FSequencerBakingSetupRestore RestoreBaking(ParentSequencer);
+					//reacquire the above function may force spawnables to get rebound.
+					SkelMeshComp = UE::Sequencer::FCommonAnimationTrackEditor::AcquireSkeletalMeshFromObjectGuid(Binding,ParentSequencer);
+
 					FAnimExportSequenceParameters AESP;
 					AESP.Player = ParentSequencer.Get();
 					AESP.RootToLocalTransform = RootToLocalTransform;
