Anim in Engine: Curve cache pool now only interpolates the visible part of interactively edited curves

* Better key editing performance on large data sets when zoomed in closely
* New CVar MovieSceneTools.ZoomInterpolatesCurvePointsAsync, when enabled curves are interpolated async when zooming, leading to smoother zoom in favor of instantly interpolated curve lines

diff --git a/Engine/Source/Editor/CurveEditor/Private/CurveEditorCurveDrawParamsCache.cpp b/Engine/Source/Editor/CurveEditor/Private/CurveEditorCurveDrawParamsCache.cpp
--- a/Engine/Source/Editor/CurveEditor/Private/CurveEditorCurveDrawParamsCache.cpp
+++ b/Engine/Source/Editor/CurveEditor/Private/CurveEditorCurveDrawParamsCache.cpp
@@ -12,13 +12,12 @@
 
 namespace UE::CurveEditor
 {
-	int32 GCurveEditorUseCurveCachePool = 0;
+	int32 GCurveEditorUseCurveCachePool = 1;
 
 	static FAutoConsoleVariableRef CCurveEditorUseCurveCachePool(
 		TEXT("CurveEditor.UseCurveCachePool"),
 		GCurveEditorUseCurveCachePool,
-		TEXT("Enables support for curve editor cache Pool"),
-		ECVF_Default);
+		TEXT("Enables improved curve editor performance by using a curve cache pool"));
 
 	void FCurveDrawParamsCache::Invalidate(const TSharedRef<SCurveEditorView>& CurveEditorView, const TArray<FCurveModelID>& ModelIDs)
 	{
@@ -212,7 +211,7 @@ namespace UE::CurveEditor
 
 	void FCurveDrawParamsCache::UpdateCurveDrawParams(const TSharedRef<SCurveEditorView>& CurveEditorView, const TArray<FCurveDrawParamsHandle>& CurveDrawParamsHandles)
 	{
-		if (!WeakCurveEditor.IsValid())
+		if (!WeakCurveEditor.IsValid() || IsEngineExitRequested())
 		{
 			return;
 		}
diff --git a/Engine/Source/Editor/MovieSceneTools/Private/Cache/MovieSceneAsyncCurvePainter.cpp b/Engine/Source/Editor/MovieSceneTools/Private/Cache/MovieSceneAsyncCurvePainter.cpp
deleted file mode 100644
diff --git a/Engine/Source/Editor/MovieSceneTools/Private/Cache/MovieSceneAsyncCurvePainter.h b/Engine/Source/Editor/MovieSceneTools/Private/Cache/MovieSceneAsyncCurvePainter.h
deleted file mode 100644
diff --git a/Engine/Source/Editor/MovieSceneTools/Private/Cache/MovieSceneCachedCurve.h b/Engine/Source/Editor/MovieSceneTools/Private/Cache/MovieSceneCachedCurve.h
--- a/Engine/Source/Editor/MovieSceneTools/Private/Cache/MovieSceneCachedCurve.h
+++ b/Engine/Source/Editor/MovieSceneTools/Private/Cache/MovieSceneCachedCurve.h
@@ -2,8 +2,8 @@
 
 #pragma once
 
-#include "Cache/MovieSceneAsyncCurvePainter.h"
 #include "Cache/MovieSceneCurveCachePool.h"
+#include "Cache/MovieSceneInterpolatingPointsDrawTask.h"
 #include "Cache/MovieSceneUpdateCachedCurveData.h"
 #include "Channels/MovieSceneInterpolation.h"
 #include "CurveDataAbstraction.h"
@@ -16,17 +16,24 @@ namespace UE::MovieSceneTools
 {
 	template <typename ChannelType> struct FMovieSceneUpdateCachedCurveData;
 
+	int32 GZoomInterpolatesCurvePointsAsync = 0;
+	static FAutoConsoleVariableRef CZoomInterpolatesCurvePointsAsync(
+		TEXT("MovieSceneTools.ZoomInterpolatesCurvePointsAsync"),
+		GZoomInterpolatesCurvePointsAsync,
+		TEXT("When set to true, zoom is smoother but curve lines will not update instantaneously"),
+		ECVF_Default);
+
 	/** Flags defining how the cache changed when it was last updated */
 	enum class EMovieSceneCurveCacheChangeFlags : uint32
 	{
 		None = 0,
 		ChangedPosition = 1 << 0,
 		ChangedSize = 1 << 1,
-		ChangedKeyIndices = 1 << 3,
-		ChangedTangentVisibility = 1 << 4,
-		ChangedSelection = 1 << 5,
-		ChangedCurveData = 1 << 6,
-		ChangedInterpolatingPoints = 1 << 7,
+		ChangedKeyIndices = 1 << 2,
+		ChangedTangentVisibility = 1 << 3,
+		ChangedSelection = 1 << 4,
+		ChangedCurveData = 1 << 5,
+		ChangedInterpolatingPoints = 1 << 6,
 	};
 	ENUM_CLASS_FLAGS(EMovieSceneCurveCacheChangeFlags);
 
@@ -64,18 +71,27 @@ namespace UE::MovieSceneTools
 		/** Returns the cached tick resolution */
 		const FFrameRate& GetTickResolution() const { return TickResolution; }
 
+		/** Cached first index of the visible space in keys / values */
+		int32 GetStartingIndex() const { return StartingIndex; }
+
+		/** Cached last index of the visible space in keys / values */
+		int32 GetEndingIndex() const { return EndingIndex; }
+
+		/** Returns cached curve times */
+		TArray<const FFrameNumber> GetTimes() const { return Times; }
+
+		/** Returns cached curve values */
+		TArray<const ChannelValueType> GetValues() const { return Values; }
+
 		/** Returns the threshold of visible pixels per time */
 		double GetTimeThreshold() const { return FMath::Max(0.0001, 1.0 / ScreenSpace.PixelsPerInput()); }
 
 		/** Returns the threshold of visible pixels per value */
-		double GetValueThreshold() const { return FMath::Max(0.0001, 1.0 / ScreenSpace.PixelsPerInput()); }
+		double GetValueThreshold() const { return FMath::Max(0.0001, 1.0 / ScreenSpace.PixelsPerOutput()); }
 
 		/** Returns the cached piecewise curve */
 		const TSharedPtr<UE::MovieScene::FPiecewiseCurve>& GetPiecewiseCurve() const { return PiecewiseCurve; }
 
-		/** Returns the points where keys reside */
-		TArray<FVector2D> GetKeyPoints();
-
 		// IMovieSceneCachedCurve interface
 		virtual uint32 GetInterpolatingPointsHash() const override;
 		virtual const UE::CurveEditor::FCurveDrawParamsHandle& GetDrawParamsHandle() const override;
@@ -115,8 +131,8 @@ namespace UE::MovieSceneTools
 		/** Updates the piecewise cuve depending on flags */
 		void ConditionallyUpdatePiecewiseCurve(const FMovieSceneUpdateCachedCurveData<ChannelType>& Data);
 
-		/** Draws the full range of interpolating points, possibly async. */
-		void DrawFullRangeInterpolatingPoints();
+		/** Creates a task to draw full range interpolating points, possibly async. */
+		void CreateFullRangeInterpolatingPointsDrawTask(const ChannelType& Channel);
 
 		/** Draws the interpolating points from the currently cached full range interpolating points */
 		void DrawInterpolatingPointsFromFullRange();
@@ -148,10 +164,10 @@ namespace UE::MovieSceneTools
 		/** Cached input display offset */
 		double InputDisplayOffset = TNumericLimits<double>::Max();
 
-		/** Cached first index in keys / values */
+		/** Cached first index of the visible space in keys / values */
 		int32 StartingIndex = INDEX_NONE;
 
-		/** Cached last index in keys / values */
+		/** Cached last index of the visible space in keys / values */
 		int32 EndingIndex = INDEX_NONE;
 
 		/** Cached selection */
@@ -214,7 +230,7 @@ namespace UE::MovieSceneTools
 		UE::CurveEditor::FCurveDrawParamsHandle DrawParamsHandle;
 
 		/** Flags defining how the cache changed since it last was updated */
-		std::atomic<EMovieSceneCurveCacheChangeFlags> Flags = EMovieSceneCurveCacheChangeFlags::None;
+		std::atomic<EMovieSceneCurveCacheChangeFlags> Flags = EMovieSceneCurveCacheChangeFlags::ChangedCurveData;
 	};
 
 	template <typename ChannelType>
@@ -236,12 +252,6 @@ namespace UE::MovieSceneTools
 		FMovieSceneCurveCachePool::Get().Join(WeakCurveEditor, AsShared());
 	}
 
-	template <typename ChannelType>
-	TArray<FVector2D> FMovieSceneCachedCurve<ChannelType>::GetKeyPoints()
-	{
-		return CachedDrawParams.InterpolatingPoints;
-	}
-
 	template <typename ChannelType>
 	void FMovieSceneCachedCurve<ChannelType>::UpdateCachedCurve(const FMovieSceneUpdateCachedCurveData<ChannelType>& UpdateData, const UE::CurveEditor::FCurveDrawParamsHandle& CurveDrawParamsHandle)
 	{
@@ -268,6 +278,9 @@ namespace UE::MovieSceneTools
 		// Update depending on above data
 		ConditionallyUpdateKeyDrawInfos(UpdateData);
 		ConditionallyUpdatePiecewiseCurve(UpdateData);
+
+		// Create a task for the curve painter
+		CreateFullRangeInterpolatingPointsDrawTask(UpdateData.Channel);
 	}
 
 	template <typename ChannelType>
@@ -281,7 +294,6 @@ namespace UE::MovieSceneTools
 		else
 		{
 			DrawKeys();
-			DrawFullRangeInterpolatingPoints();
 			DrawInterpolatingPointsFromFullRange();
 
 			ApplyDrawParams();
@@ -521,32 +533,36 @@ namespace UE::MovieSceneTools
 	}
 
 	template <typename ChannelType>
-	void FMovieSceneCachedCurve<ChannelType>::DrawFullRangeInterpolatingPoints()
+	void FMovieSceneCachedCurve<ChannelType>::CreateFullRangeInterpolatingPointsDrawTask(const ChannelType& Channel)
 	{
-		// Only draw the full range interpolating points when the curve was edited or when zooming
+		// Optimization: If there's less than two keys there's nothing to interpolate. 
+		// FMovieSceneInterpolatingPointsDrawTask ensures no such tasks are created to avoid needless overhead. 
+		// Adopt this principle here.
+		if (Times.Num() < 2)
+		{
+			return;
+		}
 
+		// Only draw the full range interpolating points when the curve was edited or when zooming
 		const bool bCurveIsBeingEdited = EnumHasAnyFlags(Flags.load(), EMovieSceneCurveCacheChangeFlags::ChangedCurveData);
 		const bool bZooming = EnumHasAnyFlags(Flags.load(), EMovieSceneCurveCacheChangeFlags::ChangedSize);
 		if (bCurveIsBeingEdited || bZooming)
 		{
-			const FScopeLock Lock(&LockInterpolatingPoints);
-
-			CachedDrawParams.InterpolatingPoints.Reset();
-
-			CachedDrawParams.InterpolatingPoints.Reserve(Times.Num());
-			for (int32 DataIndex = 0; DataIndex < Times.Num(); DataIndex++)
-			{
-				CachedDrawParams.InterpolatingPoints.Emplace(Times[DataIndex] / TickResolution, double(Values[DataIndex].Value));
-			}
-
-			// Paint the curve synchronous when the it is being edited
-			const bool bSynchronous = bCurveIsBeingEdited;
+			// Paint the visible part of the curve interactively when the it is being edited or when GZoomInterpolatesCurvePointsAsync
+			const bool bDeferredZooming = GZoomInterpolatesCurvePointsAsync > 0 ? bZooming : false;
+			const bool bInteractive = bCurveIsBeingEdited || !bDeferredZooming;
 
-			FMovieSceneAsyncCurvePainter<ChannelType>::PaintCurve(
+			const TSharedRef<FMovieSceneInterpolatingPointsDrawTask<ChannelType>> Task = MakeShared<FMovieSceneInterpolatingPointsDrawTask<ChannelType>>(
 				SharedThis(this),
-				bSynchronous,
-				[this](TArray<FVector2D> NewInterpolatingPoints, TArray<int32> NewKeyOffsets)
+				Channel,
+				bInteractive,
+				[WeakThis = AsWeak(), this](TArray<FVector2D> NewInterpolatingPoints, TArray<int32> NewKeyOffsets)
 				{
+					if (!WeakThis.IsValid())
+					{
+						return;
+					}
+
 					const FScopeLock Lock(&LockInterpolatingPoints);
 
 					FullRangeInterpolation.Points = NewInterpolatingPoints;
@@ -563,6 +579,8 @@ namespace UE::MovieSceneTools
 
 					Flags.store(Flags.load() | EMovieSceneCurveCacheChangeFlags::ChangedInterpolatingPoints);
 				});
+
+			FMovieSceneCurveCachePool::Get().AddTask(SharedThis(this), Task);
 		}
 	}
 
diff --git a/Engine/Source/Editor/MovieSceneTools/Private/Cache/MovieSceneCurveCachePool.cpp b/Engine/Source/Editor/MovieSceneTools/Private/Cache/MovieSceneCurveCachePool.cpp
--- a/Engine/Source/Editor/MovieSceneTools/Private/Cache/MovieSceneCurveCachePool.cpp
+++ b/Engine/Source/Editor/MovieSceneTools/Private/Cache/MovieSceneCurveCachePool.cpp
@@ -84,6 +84,32 @@ namespace UE::MovieSceneTools
 		}
 	}
 
+	void FMovieSceneCurveCachePool::AddTask(const TSharedRef<IMovieSceneCachedCurve>& CachedCurve, const TSharedRef<IMovieSceneInterpolatingPointsDrawTask>& Task)
+	{
+		check(IsInGameThread()); // Note, if concurrency is ever desired here, we can lock access to CachedCurveToTaskMap
+
+		if (CachedCurveToTaskMap.IsEmpty() &&
+			!FCoreDelegates::OnEndFrame.IsBoundToObject(this))
+		{
+			// Ticking was disabled if there is was no task
+			FCoreDelegates::OnEndFrame.AddSP(this, &FMovieSceneCurveCachePool::OnEndFrame);
+		}
+		
+		TSharedPtr<IMovieSceneInterpolatingPointsDrawTask>* RunningTaskPtr = CachedCurveToTaskMap.Find(CachedCurve);
+		if (RunningTaskPtr && (*RunningTaskPtr).IsValid())
+		{
+			// Let a currently running task know he's now void
+			(*RunningTaskPtr)->SetFlags(ECurvePainterTaskStateFlags::Void);
+			// Set the new task
+			(*RunningTaskPtr) = Task;
+		}
+		else
+		{
+			// Add the new task
+			CachedCurveToTaskMap.Add(CachedCurve, Task);
+		}
+	}
+
 	void FMovieSceneCurveCachePool::DrawCachedCurves(const TArray<TWeakPtr<IMovieSceneCachedCurve>>& CachedCurves) const
 	{
 		ParallelFor(CachedCurves.Num(),
@@ -175,4 +201,54 @@ namespace UE::MovieSceneTools
 			}
 		}
 	}
+
+	void FMovieSceneCurveCachePool::OnEndFrame()
+	{
+		check(IsInGameThread()); // Note, if concurrency is ever desired here, we can lock access to CachedCurveToTaskMap
+
+		// Avoid any concurrency while work is ongoing
+		if (bWorking)
+		{
+			return;
+		}
+
+		// Remove completed and void tasks
+		for (TMap<TWeakPtr<IMovieSceneCachedCurve>, TSharedPtr<IMovieSceneInterpolatingPointsDrawTask>>::TIterator It(CachedCurveToTaskMap); It; ++It)
+		{
+			if ((*It).Value->HasAnyFlags(ECurvePainterTaskStateFlags::Completed | ECurvePainterTaskStateFlags::Void))
+			{
+				It.RemoveCurrent();
+			}
+		}
+		if (CachedCurveToTaskMap.IsEmpty())
+		{
+			// All tasks have terminated, stop ticking
+			FCoreDelegates::OnEndFrame.RemoveAll(this);
+			return;
+		}
+
+		TArray<TSharedPtr<IMovieSceneInterpolatingPointsDrawTask>> Tasks;
+		CachedCurveToTaskMap.GenerateValueArray(Tasks);
+
+		UE::Tasks::Launch(UE_SOURCE_LOCATION,
+			[Tasks, this]()
+			{
+				TRACE_CPUPROFILER_EVENT_SCOPE(FMovieSceneAsyncCurvePainter DrawInterpolatiedPointsAsync);
+
+				bWorking = true;
+
+				ParallelFor(Tasks.Num(),
+					[&Tasks](int32 TaskIndex)
+					{
+						const TSharedPtr<IMovieSceneInterpolatingPointsDrawTask>& Task = Tasks[TaskIndex];
+						if (!Task->HasAnyFlags(ECurvePainterTaskStateFlags::Void))
+						{
+							Task->RefineFullRangeInterpolatingPoints();
+						}
+					});
+
+
+				bWorking = false;
+			}, LowLevelTasks::ETaskPriority::BackgroundNormal);
+	}
 }
diff --git a/Engine/Source/Editor/MovieSceneTools/Private/Cache/MovieSceneCurveCachePool.h b/Engine/Source/Editor/MovieSceneTools/Private/Cache/MovieSceneCurveCachePool.h
--- a/Engine/Source/Editor/MovieSceneTools/Private/Cache/MovieSceneCurveCachePool.h
+++ b/Engine/Source/Editor/MovieSceneTools/Private/Cache/MovieSceneCurveCachePool.h
@@ -30,6 +30,39 @@ namespace UE::MovieSceneTools
 		virtual const UE::CurveEditor::FCurveDrawParamsHandle& GetDrawParamsHandle() const = 0;
 	};
 
+	/** States of a task in the pool */
+	enum class ECurvePainterTaskStateFlags : uint8
+	{
+		None = 0,
+
+		/** The visible part was written immediately, synchronous */
+		Interactive = 1 << 0,
+
+		/** The task successfully finished */
+		Completed = 1 << 1,
+
+		/** The task is running but no longer useful */
+		Void = 1 << 2,
+	};
+	ENUM_CLASS_FLAGS(ECurvePainterTaskStateFlags);
+
+	/** Interface for tasks the curve cache pool can handle */
+	class IMovieSceneInterpolatingPointsDrawTask
+		: public TSharedFromThis<IMovieSceneInterpolatingPointsDrawTask>
+	{
+	public:
+		virtual ~IMovieSceneInterpolatingPointsDrawTask() = default;
+
+		/** Sets flags for this task */
+		virtual void SetFlags(ECurvePainterTaskStateFlags NewFlags) = 0;
+
+		/** Returns true if the task has the specified flag set */
+		virtual bool HasAnyFlags(ECurvePainterTaskStateFlags Flags) const = 0;
+
+		/** Refines the full range of the interpolating points */
+		virtual void RefineFullRangeInterpolatingPoints() = 0;
+	};
+
 	/** A pool of cached curves in a curve editor */
 	class FMovieSceneCurveCachePool
 		: public UE::CurveEditor::ICurveEditorCurveCachePool
@@ -50,6 +83,9 @@ namespace UE::MovieSceneTools
 		/** Lets a cached curve leave the cache pool */
 		void Leave(const IMovieSceneCachedCurve& CachedCurve);
 
+		/** Adds a task to the curve cache pool */
+		void AddTask(const TSharedRef<IMovieSceneCachedCurve>& CachedCurve, const TSharedRef<IMovieSceneInterpolatingPointsDrawTask>& Task);
+
 	private:
 		/** Lets cached curves draw to their draw params */
 		void DrawCachedCurves(const TArray<TWeakPtr<IMovieSceneCachedCurve>>& CachedCurves) const;
@@ -57,7 +93,19 @@ namespace UE::MovieSceneTools
 		/** Culls cached curves */
 		void CullCachedCurves(const TArray<TWeakPtr<IMovieSceneCachedCurve>>& CachedCurves) const;
 
+		/** Called on end frame. Updates work on tick */
+		void OnEndFrame();
+
 		/** Cached curves per curve editor */
 		TMap<TWeakPtr<const FCurveEditor>, TArray<TWeakPtr<IMovieSceneCachedCurve>>> CurveEditorToCachedCurvesMap;
+
+		/** A map of cached curves with their related tasks that need to be performed asynchronous */
+		TMap<TWeakPtr<IMovieSceneCachedCurve>, TSharedPtr<IMovieSceneInterpolatingPointsDrawTask>> CachedCurveToTaskMap;
+
+		/** True while work is ongoing */
+		std::atomic<bool> bWorking = false;
+
+		/** Critical section to access tasks */
+		FCriticalSection AccessTasksCritSec;
 	};
 }
diff --git a/Engine/Source/Editor/MovieSceneTools/Private/Cache/MovieSceneInterpolatingPointsDrawTask.cpp b/Engine/Source/Editor/MovieSceneTools/Private/Cache/MovieSceneInterpolatingPointsDrawTask.cpp
new file mode 100644
--- /dev/null
+++ b/Engine/Source/Editor/MovieSceneTools/Private/Cache/MovieSceneInterpolatingPointsDrawTask.cpp
@@ -0,0 +1,313 @@
+// Copyright Epic Games, Inc. All Rights Reserved.
+
+#include "MovieSceneInterpolatingPointsDrawTask.h"
+
+#include "Cache/MovieSceneCachedCurve.h"
+#include "Channels/MovieSceneDoubleChannel.h"
+#include "Channels/MovieSceneFloatChannel.h"
+#include "Channels/MovieScenePiecewiseCurve.h"
+
+namespace UE::MovieSceneTools
+{
+	template struct FMovieSceneInterpolatingPointsDrawTask<FMovieSceneFloatChannel>;
+	template struct FMovieSceneInterpolatingPointsDrawTask<FMovieSceneDoubleChannel>;
+
+	template <typename ChannelType>
+	FMovieSceneInterpolatingPointsDrawTask<ChannelType>::FMovieSceneInterpolatingPointsDrawTask(
+		const TSharedRef<FMovieSceneCachedCurve<ChannelType>>& CachedCurve,
+		const ChannelType& MovieSceneChannel,
+		const bool bInteractive,
+		const TFunction<void(TArray<FVector2D>, TArray<int32>)>& InCallback)
+		: Callback(InCallback)
+		, ScreenSpace(CachedCurve->GetScreenSpace())
+		, TickResolution(CachedCurve->GetTickResolution())
+		, TimeThreshold(CachedCurve->GetTimeThreshold())
+		, ValueThreshold(CachedCurve->GetValueThreshold())
+		, PiecewiseCurve(CachedCurve->GetPiecewiseCurve())
+	{
+		const TArray<const FFrameNumber>& Times = CachedCurve->GetTimes();
+		const TArray<const ChannelValueType>& Values = CachedCurve->GetValues();
+		if (!ensureMsgf(Times.Num() > 1, TEXT("Curve paint tasks should only be created for curves with more than one key")))
+		{
+			SetFlags(ECurvePainterTaskStateFlags::Completed);
+			return;
+		}
+
+		// Remember key points
+		KeyPoints.Reserve(Times.Num());
+		for (int32 DataIndex = 0; DataIndex < Times.Num(); DataIndex++)
+		{
+			KeyPoints.Emplace(Times[DataIndex] / TickResolution, double(Values[DataIndex].Value));
+		}
+		InterpolatingPoints = KeyPoints;
+
+		// Process the visible part synchronous if interactive
+		if (bInteractive)
+		{		
+			DrawVisibleCurvePointsSynchronous(CachedCurve, MovieSceneChannel);
+			InvokeCallback();
+
+			SetFlags(ECurvePainterTaskStateFlags::Interactive);
+		}
+	}
+
+	template <typename ChannelType>
+	void FMovieSceneInterpolatingPointsDrawTask<ChannelType>::RefineFullRangeInterpolatingPoints()
+	{
+		// Make sure there's no concurrency
+		check(!bWorking);
+		bWorking = true;
+
+		const int32 OldSize = InterpolatingPoints.Num();
+
+		if (HasAnyFlags(ECurvePainterTaskStateFlags::Interactive))
+		{
+			RefineFullRangeInterpolatingPointsInternal(ECurvePainterInterpolationRange::StartToInteractive);
+			RefineFullRangeInterpolatingPointsInternal(ECurvePainterInterpolationRange::InteractiveToEnd);
+		}
+		else
+		{
+			RefineFullRangeInterpolatingPointsInternal(ECurvePainterInterpolationRange::FullRange);
+		}
+
+		const int32 NewSize = InterpolatingPoints.Num();
+		if (OldSize == NewSize)
+		{
+			// Remove straight lines to reduce slate elements needed to draw the curve
+			const float PixelArea = FMath::Pow(FMath::Min(TimeThreshold, ValueThreshold), 2.f);
+			
+			for (int32 PointIndex = 0; PointIndex < InterpolatingPoints.Num() - 3;)
+			{
+				const FVector2D& First = InterpolatingPoints[0];
+				const FVector2D& Second = InterpolatingPoints[1];
+				const FVector2D& Third = InterpolatingPoints[2];
+
+				const float LineArea = FMath::Abs(First.X * (Second.Y - Third.Y) + Second.X * (Third.Y - First.Y) + Third.X * (First.Y - Third.Y));
+				if (LineArea < PixelArea)
+				{
+					InterpolatingPoints.RemoveAt(PointIndex + 1);
+				}
+				else
+				{
+					PointIndex++;
+				}
+			}
+
+			InvokeCallback();
+			SetFlags(ECurvePainterTaskStateFlags::Completed);
+		}
+
+		bWorking = false;
+	}
+
+	template <typename ChannelType>
+	void FMovieSceneInterpolatingPointsDrawTask<ChannelType>::SetFlags(ECurvePainterTaskStateFlags NewFlags)
+	{
+		StateFlags.store(NewFlags);
+	}
+
+	template <typename ChannelType>
+	bool FMovieSceneInterpolatingPointsDrawTask<ChannelType>::HasAnyFlags(ECurvePainterTaskStateFlags Flags) const
+	{
+		return EnumHasAnyFlags(StateFlags.load(), Flags);
+	}
+
+	template<typename ChannelType>
+	void FMovieSceneInterpolatingPointsDrawTask<ChannelType>::RefineFullRangeInterpolatingPointsInternal(const ECurvePainterInterpolationRange InterpolationRange)
+	{			
+		const int32 StartingIndex = [&InterpolationRange, this]()
+			{
+				if (InterpolationRange == ECurvePainterInterpolationRange::FullRange ||
+					InterpolationRange == ECurvePainterInterpolationRange::StartToInteractive)
+				{
+					return 0;
+				}
+				else
+				{
+					const double InteractiveEnd = InteractiveTimeEnd.GetValue();
+					const int32 StartingIndex = InterpolatingPoints.IndexOfByPredicate(
+						[this, &InteractiveEnd](const FVector2D& InterpolatingPoint)
+						{
+							return InteractiveEnd <= InterpolatingPoint.X;
+						});
+
+					return StartingIndex != INDEX_NONE ? StartingIndex : InterpolatingPoints.Num() - 1;
+				}
+			}();
+
+		const int32 EndingIndex = [&InterpolationRange, this]()
+			{
+				if (InterpolationRange == ECurvePainterInterpolationRange::FullRange ||
+					InterpolationRange == ECurvePainterInterpolationRange::InteractiveToEnd)
+				{
+					return InterpolatingPoints.Num() - 1;
+				}
+				else
+				{
+					const double InteractiveStart = InteractiveTimeStart.GetValue();
+					const int32 EndingIndex = InterpolatingPoints.IndexOfByPredicate(
+						[this, &InteractiveStart](const FVector2D& InterpolatingPoint)
+						{
+							return InteractiveStart >= InterpolatingPoint.X;
+						});
+
+					return EndingIndex != INDEX_NONE ? EndingIndex : InterpolatingPoints.Num() - 1;
+				}
+			}();
+
+		constexpr float InterpTimes[] = { 0.25f, 0.5f, 0.6f };
+		for (int32 Index = StartingIndex; Index < EndingIndex; Index++)
+		{
+			const FVector2D& Lower = InterpolatingPoints[Index];
+			const FVector2D& Upper = InterpolatingPoints[Index + 1];
+
+			if ((Upper.X - Lower.X) >= TimeThreshold)
+			{
+				bool bSegmentIsLinear = true;
+
+				FVector2D Evaluated[UE_ARRAY_COUNT(InterpTimes)] = { FVector2D::ZeroVector };
+				for (int32 InterpIndex = 0; InterpIndex < UE_ARRAY_COUNT(InterpTimes); ++InterpIndex)
+				{
+					double& EvalTime = Evaluated[InterpIndex].X;
+					EvalTime = FMath::Lerp(Lower.X, Upper.X, InterpTimes[InterpIndex]);
+					double Value = 0.0;
+
+					PiecewiseCurve->Evaluate(EvalTime * TickResolution, Value);
+					const double LinearValue = FMath::Lerp(Lower.Y, Upper.Y, InterpTimes[InterpIndex]);
+					if (bSegmentIsLinear)
+					{
+						bSegmentIsLinear = FMath::IsNearlyEqual(Value, LinearValue, ValueThreshold);
+					}
+
+					Evaluated[InterpIndex].Y = Value;
+				}
+
+				if (!bSegmentIsLinear)
+				{
+					// Add the point
+					InterpolatingPoints.Insert(Evaluated, UE_ARRAY_COUNT(Evaluated), Index + 1);
+					--Index;
+				}
+			}
+		}
+	}
+
+	template<typename ChannelType>
+	void FMovieSceneInterpolatingPointsDrawTask<ChannelType>::DrawVisibleCurvePointsSynchronous(const TSharedRef<FMovieSceneCachedCurve<ChannelType>>& CachedCurve, const ChannelType& Channel)
+	{
+		TRACE_CPUPROFILER_EVENT_SCOPE(FMovieSceneCurvePainterTas::DrawCurvePointsSynchronous);
+
+		const TArray<const FFrameNumber>& Times = CachedCurve->GetTimes();
+		const TArray<const ChannelValueType>& Values = CachedCurve->GetValues();
+
+		const int32 StartingIndex = CachedCurve->GetStartingIndex();
+		const int32 EndingIndex = CachedCurve->GetEndingIndex();
+
+		const double StartTime = ScreenSpace.GetInputMin();
+		const double EndTime = ScreenSpace.GetInputMax();
+
+		const FFrameNumber StartFrame = (StartTime * TickResolution).FloorToFrame();
+		const FFrameNumber EndFrame = (EndTime * TickResolution).CeilToFrame();
+
+		// Add the lower bound of the visible space
+		CurveValueType EvaluatedValue = 0;
+		if (Channel.Evaluate(StartFrame, EvaluatedValue))
+		{
+			InterpolatingPoints.Emplace(StartTime, double(EvaluatedValue));
+		}
+
+		// Add all keys in-between
+		for (int32 KeyIndex = StartingIndex; KeyIndex < EndingIndex; ++KeyIndex)
+		{
+			InterpolatingPoints.Emplace(Times[KeyIndex] / TickResolution, double(Values[KeyIndex].Value));
+		}
+
+		// Add the upper bound of the visible space
+		if (Channel.Evaluate(EndFrame, EvaluatedValue))
+		{
+			InterpolatingPoints.Emplace(EndTime, double(EvaluatedValue));
+		}
+
+		int32 OldSize;
+		do
+		{
+			OldSize = InterpolatingPoints.Num();
+			RefineVisibleCurvePointsSynchronous(CachedCurve, Channel);
+		} while (OldSize != InterpolatingPoints.Num());
+
+		// Remember what was interpolated already
+		InteractiveTimeStart = StartTime;
+		InteractiveTimeEnd = EndTime;
+	}
+
+	template<typename ChannelType>
+	void FMovieSceneInterpolatingPointsDrawTask<ChannelType>::RefineVisibleCurvePointsSynchronous(const TSharedRef<FMovieSceneCachedCurve<ChannelType>>& CachedCurve, const ChannelType& Channel)
+	{
+		const float InterpTimes[] = { 0.25f, 0.5f, 0.6f };
+
+		for (int32 Index = 0; Index < InterpolatingPoints.Num() - 1; ++Index)
+		{
+			const FVector2D& Lower = InterpolatingPoints[Index];
+			const FVector2D& Upper = InterpolatingPoints[Index + 1];
+
+			if ((Upper.X - Lower.X) >= TimeThreshold)
+			{
+				bool bSegmentIsLinear = true;
+
+				FVector2D Evaluated[UE_ARRAY_COUNT(InterpTimes)] = {};
+
+				for (int32 InterpIndex = 0; InterpIndex < UE_ARRAY_COUNT(InterpTimes); ++InterpIndex)
+				{
+					double& EvalTime = Evaluated[InterpIndex].X;
+
+					EvalTime = FMath::Lerp(Lower.X, Upper.X, InterpTimes[InterpIndex]);
+
+					CurveValueType Value = 0.0;
+					Channel.Evaluate(EvalTime * TickResolution, Value);
+
+					const CurveValueType LinearValue = FMath::Lerp(Lower.Y, Upper.Y, InterpTimes[InterpIndex]);
+					if (bSegmentIsLinear)
+					{
+						bSegmentIsLinear = FMath::IsNearlyEqual(Value, LinearValue, ValueThreshold);
+					}
+
+					Evaluated[InterpIndex].Y = Value;
+				}
+
+				if (!bSegmentIsLinear)
+				{
+					// Add the point
+					InterpolatingPoints.Insert(Evaluated, UE_ARRAY_COUNT(Evaluated), Index + 1);
+					--Index;
+				}
+			}
+		}
+	}
+
+	template<typename ChannelType>
+	void FMovieSceneInterpolatingPointsDrawTask<ChannelType>::InvokeCallback()
+	{
+		TArray<int32> KeyOffsets;
+		int32 InterpolatingPointIndex = 0;
+		for (const FVector2D& KeyPoint : KeyPoints)
+		{
+			while (
+				InterpolatingPoints.IsValidIndex(InterpolatingPointIndex) &&
+				KeyPoint.X > InterpolatingPoints[InterpolatingPointIndex].X)
+			{
+				InterpolatingPointIndex++;
+			}
+
+			if (InterpolatingPoints.IsValidIndex(InterpolatingPointIndex))
+			{
+				KeyOffsets.Add(InterpolatingPointIndex);
+			}
+			else if (!InterpolatingPoints.IsEmpty())
+			{
+				KeyOffsets.Add(InterpolatingPoints.Num() - 1);
+			}
+		}
+
+		Callback(InterpolatingPoints, KeyOffsets);
+	}
+}
diff --git a/Engine/Source/Editor/MovieSceneTools/Private/Cache/MovieSceneInterpolatingPointsDrawTask.h b/Engine/Source/Editor/MovieSceneTools/Private/Cache/MovieSceneInterpolatingPointsDrawTask.h
new file mode 100644
--- /dev/null
+++ b/Engine/Source/Editor/MovieSceneTools/Private/Cache/MovieSceneInterpolatingPointsDrawTask.h
@@ -0,0 +1,110 @@
+// Copyright Epic Games, Inc. All Rights Reserved.
+
+#pragma once
+
+#include "Cache/MovieSceneCurveCachePool.h"
+#include "CurveEditorScreenSpace.h"
+#include "Misc/FrameRate.h"
+#include "Misc/Optional.h"
+#include "Templates/SharedPointer.h"
+
+namespace UE::MovieScene { struct FPiecewiseCurve; }
+
+namespace UE::MovieSceneTools
+{
+	template <typename ChannelType> class FMovieSceneCachedCurve;
+
+	/** 
+	 * A task to draw interpolating points.
+	 * Draws interpolation for the visible space when interactive, but also the full range of the curve.
+	 */
+	template <typename ChannelType>
+	struct FMovieSceneInterpolatingPointsDrawTask
+		: public IMovieSceneInterpolatingPointsDrawTask
+	{		
+		/** The type of channel value structs */
+		using ChannelValueType = typename ChannelType::ChannelValueType;
+		using CurveValueType = typename ChannelType::CurveValueType;
+
+	public:
+		FMovieSceneInterpolatingPointsDrawTask(
+			const TSharedRef<FMovieSceneCachedCurve<ChannelType>>& CachedCurve,
+			const ChannelType& MovieSceneChannel,
+			const bool bInteractive,
+			const TFunction<void(TArray<FVector2D> /** OutInterpolatedPoints */, TArray<int32> /** OutKeyOffsets */)>& InCallback);
+
+		//~ Begin IMovieSceneInterpolatingPointsDrawTask interface
+		virtual void SetFlags(ECurvePainterTaskStateFlags NewFlags) override;
+		virtual bool HasAnyFlags(ECurvePainterTaskStateFlags Flags) const override;
+		virtual void RefineFullRangeInterpolatingPoints() override;
+		//~ End IMovieSceneInterpolatingPointsDrawTask interface
+
+	private:
+		/** Describs wich part of the curve is interpolated in RefineFullRangeInterpolatingPoints */
+		enum class ECurvePainterInterpolationRange : uint8
+		{
+			/** Draws the full range of interpolating points */
+			FullRange,
+
+			/** Draws from the start to the already interactively interpolated range */
+			StartToInteractive,
+
+			/** Draws from the already interactively interpolated range to the end */
+			InteractiveToEnd
+		};
+
+		/** 
+		 * Refines the full range of the interpolating points. 
+		 * Internal logic to allow to skip over the interactively painted area if this is an interactive task.
+		 */
+		void RefineFullRangeInterpolatingPointsInternal(const ECurvePainterInterpolationRange InterpolationRange);
+
+		/** 
+		 * Utility to darw the visible part of the curve synchronously. Useful for interactive edits.
+		 * Note this uses the channel, this is faster than using the piecewise curve in the async counterpart.
+		 */
+		void DrawVisibleCurvePointsSynchronous(const TSharedRef<FMovieSceneCachedCurve<ChannelType>>& CachedCurve, const ChannelType& Channel);
+
+		/** Utility to refine the visible part of the curve synchronously. Useful for interactive edits. */
+		void RefineVisibleCurvePointsSynchronous(const TSharedRef<FMovieSceneCachedCurve<ChannelType>>& CachedCurve, const ChannelType& Channel);
+
+		/** Calls back to the cached curve, forwarding interpolated points */
+		void InvokeCallback();
+
+		/** The callback to execute on completion */
+		const TFunction<void(TArray<FVector2D> /** OutInterpolatedPoints */, TArray<int32> /** OutKeyOffsets */)> Callback;
+
+		/** The screenspace in which we paint */
+		FCurveEditorScreenSpace ScreenSpace;
+
+		/** The current tick resolution */
+		FFrameRate TickResolution;
+
+		/** Threshold of visible times */
+		double TimeThreshold = 0.0;
+
+		/** Threshold of visible values */
+		double ValueThreshold = 0.0;
+
+		/** The curve as piecewise curve */
+		TSharedPtr<UE::MovieScene::FPiecewiseCurve> PiecewiseCurve;
+
+		/** The initial points that are being interpolated */
+		TArray<FVector2D> KeyPoints;
+
+		/** The (possibly not yet fully interpolated) interpolating points */
+		TArray<FVector2D> InterpolatingPoints;
+
+		/** The first frame that was painted interactively. Only set if this is an interactive task */
+		TOptional<double> InteractiveTimeStart;
+
+		/** The last frame that was painted interactively. Only set if this is an interactive task */
+		TOptional<double> InteractiveTimeEnd;
+
+		/** True while work is performed in this task. Useful to prevent from concurrency by design. */
+		std::atomic<bool> bWorking = false;
+
+		/** True if the task needs to iterate again */
+		std::atomic<ECurvePainterTaskStateFlags> StateFlags = ECurvePainterTaskStateFlags::None;
+	};
+}
