Curve Editor: Fix an issue where curves are not painted correctly in some cases after deleting keys

#rnx

diff --git a/Engine/Source/Editor/MovieSceneTools/Private/Cache/MovieSceneCachedCurve.h b/Engine/Source/Editor/MovieSceneTools/Private/Cache/MovieSceneCachedCurve.h
--- a/Engine/Source/Editor/MovieSceneTools/Private/Cache/MovieSceneCachedCurve.h
+++ b/Engine/Source/Editor/MovieSceneTools/Private/Cache/MovieSceneCachedCurve.h
@@ -534,54 +534,73 @@ namespace UE::MovieSceneTools
 
 	template <typename ChannelType>
 	void FMovieSceneCachedCurve<ChannelType>::CreateFullRangeInterpolatingPointsDrawTask(const ChannelType& Channel)
-	{
+	{		
+		// Only draw the full range interpolating points when the curve was edited or when zooming
+		const bool bCurveIsBeingEdited = EnumHasAnyFlags(Flags.load(), EMovieSceneCurveCacheChangeFlags::ChangedCurveData);
+		const bool bZooming = EnumHasAnyFlags(Flags.load(), EMovieSceneCurveCacheChangeFlags::ChangedSize);
+		if (!bCurveIsBeingEdited && !bZooming)
+		{
+			return;
+		}
+
 		// Optimization: If there's less than two keys there's nothing to interpolate. 
 		// FMovieSceneInterpolatingPointsDrawTask ensures no such tasks are created to avoid needless overhead. 
 		// Adopt this principle here.
-		if (Times.Num() < 2)
+		if (Times.IsEmpty())
 		{
+			FullRangeInterpolation.KeyOffsets.Reset();
+			FullRangeInterpolation.Points.Reset();
+			FullRangeInterpolation.Hash = 0;
+
 			return;
 		}
-
-		// Only draw the full range interpolating points when the curve was edited or when zooming
-		const bool bCurveIsBeingEdited = EnumHasAnyFlags(Flags.load(), EMovieSceneCurveCacheChangeFlags::ChangedCurveData);
-		const bool bZooming = EnumHasAnyFlags(Flags.load(), EMovieSceneCurveCacheChangeFlags::ChangedSize);
-		if (bCurveIsBeingEdited || bZooming)
+		else if (Times.Num() == 1)
 		{
-			// Paint the visible part of the curve interactively when the it is being edited or when GZoomInterpolatesCurvePointsAsync
-			const bool bDeferredZooming = GZoomInterpolatesCurvePointsAsync > 0 ? bZooming : false;
-			const bool bInteractive = bCurveIsBeingEdited || !bDeferredZooming;
+			FullRangeInterpolation.KeyOffsets = { 0 };
+
+			FullRangeInterpolation.Points.Reset();
+			FullRangeInterpolation.Points.Emplace(Times[0].Value, Values[0].Value);
 
-			const TSharedRef<FMovieSceneInterpolatingPointsDrawTask<ChannelType>> Task = MakeShared<FMovieSceneInterpolatingPointsDrawTask<ChannelType>>(
-				SharedThis(this),
-				Channel,
-				bInteractive,
-				[WeakThis = AsWeak(), this](TArray<FVector2D> NewInterpolatingPoints, TArray<int32> NewKeyOffsets)
+			const int32 X = Times[0].Value;
+			const int32 Y = FMath::RoundToInt32(Values[0].Value);
+			FullRangeInterpolation.Hash = HashCombine(X, Y);
+
+			return;
+		}
+
+		// Create draw tasks for complex curves
+		const bool bDeferredZooming = GZoomInterpolatesCurvePointsAsync > 0 ? bZooming : false;
+		const bool bInteractive = bCurveIsBeingEdited || !bDeferredZooming;
+
+		const TSharedRef<FMovieSceneInterpolatingPointsDrawTask<ChannelType>> Task = MakeShared<FMovieSceneInterpolatingPointsDrawTask<ChannelType>>(
+			SharedThis(this),
+			Channel,
+			bInteractive,
+			[WeakThis = AsWeak(), this](TArray<FVector2D> NewInterpolatingPoints, TArray<int32> NewKeyOffsets)
+			{
+				if (!WeakThis.IsValid())
 				{
-					if (!WeakThis.IsValid())
-					{
-						return;
-					}
+					return;
+				}
 
-					const FScopeLock Lock(&LockInterpolatingPoints);
+				const FScopeLock Lock(&LockInterpolatingPoints);
 
-					FullRangeInterpolation.Points = NewInterpolatingPoints;
-					FullRangeInterpolation.KeyOffsets = NewKeyOffsets;
+				FullRangeInterpolation.Points = NewInterpolatingPoints;
+				FullRangeInterpolation.KeyOffsets = NewKeyOffsets;
 
-					uint32 Hash = 0;
-					for (FVector2D& InterpolatingPoint : FullRangeInterpolation.Points)
-					{
-						const int32 X = FMath::RoundToInt32(InterpolatingPoint.X);
-						const int32 Y = FMath::RoundToInt32(InterpolatingPoint.Y);
-						Hash = HashCombine(FullRangeInterpolation.Hash, HashCombine(X, Y));
-					}
-					FullRangeInterpolation.Hash.store(Hash);
+				uint32 Hash = 0;
+				for (FVector2D& InterpolatingPoint : FullRangeInterpolation.Points)
+				{
+					const int32 X = FMath::RoundToInt32(InterpolatingPoint.X);
+					const int32 Y = FMath::RoundToInt32(InterpolatingPoint.Y);
+					Hash = HashCombine(FullRangeInterpolation.Hash, HashCombine(X, Y));
+				}
+				FullRangeInterpolation.Hash.store(Hash);
 
-					Flags.store(Flags.load() | EMovieSceneCurveCacheChangeFlags::ChangedInterpolatingPoints);
-				});
+				Flags.store(Flags.load() | EMovieSceneCurveCacheChangeFlags::ChangedInterpolatingPoints);
+			});
 
-			FMovieSceneCurveCachePool::Get().AddTask(SharedThis(this), Task);
-		}
+		FMovieSceneCurveCachePool::Get().AddTask(SharedThis(this), Task);
 	}
 
 	template <typename ChannelType>
@@ -591,7 +610,7 @@ namespace UE::MovieSceneTools
 
 		CachedDrawParams.InterpolatingPoints.Reset(FullRangeInterpolation.Points.Num() + 2);
 
-		if (FullRangeInterpolation.Points.IsEmpty())
+		if (FullRangeInterpolation.Points.Num() < 2)
 		{
 			DrawPreInfinityInterpolatingPoints();
 			DrawPostInfinityInterpolatingPoints();
