Sequencer: Added user defined colors for playback start/end range.

Fix painting the start/end lines after the dimmed start/end regions

diff --git a/Engine/Plugins/MovieScene/LevelSequenceEditor/Source/LevelSequenceEditor/Private/CinematicViewport/SCinematicTransportRange.cpp b/Engine/Plugins/MovieScene/LevelSequenceEditor/Source/LevelSequenceEditor/Private/CinematicViewport/SCinematicTransportRange.cpp
--- a/Engine/Plugins/MovieScene/LevelSequenceEditor/Source/LevelSequenceEditor/Private/CinematicViewport/SCinematicTransportRange.cpp
+++ b/Engine/Plugins/MovieScene/LevelSequenceEditor/Source/LevelSequenceEditor/Private/CinematicViewport/SCinematicTransportRange.cpp
@@ -345,8 +345,8 @@ int32 SCinematicTransportRange::OnPaint(const FPaintArgs& Args, const FGeometry&
 			LayerId+1,
 			AllottedGeometry.ToPaintGeometry(FVector2f(BrushWidth, TrackHeight), FSlateLayoutTransform(FVector2f(AllottedGeometry.GetLocalSize().X*PlaybackStartLerp, TrackOffsetY))),
 			FLevelSequenceEditorStyle::Get()->GetBrush("LevelSequenceEditor.CinematicViewportRangeStart"),
-			DrawEffects,
-			FColor(32, 128, 32)	// 120, 75, 50 (HSV)
+			DrawEffects, 
+			Sequencer->GetSequencerSettings()->GetPlaybackRangeStartColor().ToFColor(true)
 		);
 
 		FSlateDrawElement::MakeBox(
@@ -355,7 +355,7 @@ int32 SCinematicTransportRange::OnPaint(const FPaintArgs& Args, const FGeometry&
 			AllottedGeometry.ToPaintGeometry(FVector2f(BrushWidth, TrackHeight), FSlateLayoutTransform(FVector2f(AllottedGeometry.GetLocalSize().X*PlaybackEndLerp - BrushWidth, TrackOffsetY))),
 			FLevelSequenceEditorStyle::Get()->GetBrush("LevelSequenceEditor.CinematicViewportRangeEnd"),
 			DrawEffects,
-			FColor(128, 32, 32)	// 0, 75, 50 (HSV)
+			Sequencer->GetSequencerSettings()->GetPlaybackRangeEndColor().ToFColor(true)
 		);
 	}
 
diff --git a/Engine/Source/Editor/MovieSceneTools/Private/TrackEditors/SubTrackEditorBase.cpp b/Engine/Source/Editor/MovieSceneTools/Private/TrackEditors/SubTrackEditorBase.cpp
--- a/Engine/Source/Editor/MovieSceneTools/Private/TrackEditors/SubTrackEditorBase.cpp
+++ b/Engine/Source/Editor/MovieSceneTools/Private/TrackEditors/SubTrackEditorBase.cpp
@@ -41,7 +41,7 @@ FSubSectionPainterResult FSubSectionPainterUtil::PaintSection(TSharedPtr<const I
 
 	const ESlateDrawEffect DrawEffects = InPainter.bParentEnabled ? ESlateDrawEffect::None : ESlateDrawEffect::DisabledEffect;
 
-	PaintSectionBounds(SectionObject, *InnerSequence, InPainter, DrawEffects);
+	PaintSectionBounds(Sequencer, SectionObject, *InnerSequence, InPainter, DrawEffects);
 
 	UMovieScene* MovieScene = InnerSequence->GetMovieScene();
 	const int32 NumTracks = MovieScene->GetPossessableCount() + MovieScene->GetSpawnableCount() + MovieScene->GetTracks().Num();
@@ -122,7 +122,7 @@ FSubSectionPainterResult FSubSectionPainterUtil::PaintSection(TSharedPtr<const I
 	return FSSPR_Success;
 }
 
-void FSubSectionPainterUtil::PaintSectionBounds(const UMovieSceneSubSection& SectionObject, const UMovieSceneSequence& InnerSequence, FSequencerSectionPainter& InPainter, ESlateDrawEffect DrawEffects)
+void FSubSectionPainterUtil::PaintSectionBounds(TSharedPtr<const ISequencer> Sequencer, const UMovieSceneSubSection& SectionObject, const UMovieSceneSequence& InnerSequence, FSequencerSectionPainter& InPainter, ESlateDrawEffect DrawEffects)
 {
 	const uint8        GhostAlpha        = static_cast<uint8>(255 * InPainter.GhostAlpha);
 	const FFrameNumber SectionStartFrame = SectionObject.GetInclusiveStartFrame();
@@ -134,6 +134,10 @@ void FSubSectionPainterUtil::PaintSectionBounds(const UMovieSceneSubSection& Sec
 	const TRange<FFrameTime> SectionTimeRange = UE::MovieScene::ConvertToFrameTimeRange(SectionObject.GetRange());
 	const int32              SectionSize      = UE::MovieScene::DiscreteSize(SectionObject.GetRange());
 
+	int32 LayerId = InPainter.LayerId;
+	const int32 LineLayerIdOffset = 10; // Draw the boundary lines on a layer after the dimmed boxes
+	const int32 BoxLayerIdOffset = 5; // Draw the dimmed boxes on a layer before the boundary lines
+	InPainter.LayerId = LayerId + LineLayerIdOffset + BoxLayerIdOffset + 1;
 
 	TOptional<FFrameTime> FirstTime, LastTime;
 	auto PaintDashedLine = [&](FFrameTime InTime, FColor Tint, float LineNudge)
@@ -162,7 +166,7 @@ void FSubSectionPainterUtil::PaintSectionBounds(const UMovieSceneSubSection& Sec
 		constexpr float DashLengthPx = 3.f;
 		FSlateDrawElement::MakeDashedLines(
 			InPainter.DrawElements,
-			InPainter.LayerId++,
+			LayerId + LineLayerIdOffset,
 			InPainter.SectionGeometry.ToPaintGeometry(),
 			MoveTemp(NewVector),
 			DrawEffects,
@@ -172,8 +176,6 @@ void FSubSectionPainterUtil::PaintSectionBounds(const UMovieSceneSubSection& Sec
 		);
 	};
 
-
-
 	auto PaintLine = [&](FFrameTime InTime, FColor Tint, float LineNudge)
 	{
 		float Offset = InPainter.GetTimeConverter().FrameToPixel(InTime) + LineNudge;
@@ -193,7 +195,7 @@ void FSubSectionPainterUtil::PaintSectionBounds(const UMovieSceneSubSection& Sec
 		// add green line for playback start
 		FSlateDrawElement::MakeBox(
 			InPainter.DrawElements,
-			InPainter.LayerId++,
+			LayerId + LineLayerIdOffset,
 			InPainter.SectionGeometry.ToPaintGeometry(
 				FVector2f(1.0f, InPainter.SectionGeometry.Size.Y),
 				FSlateLayoutTransform(FVector2f(Offset, 0.f))
@@ -205,31 +207,33 @@ void FSubSectionPainterUtil::PaintSectionBounds(const UMovieSceneSubSection& Sec
 	};
 
 
-	FColor GreenTint(32, 128, 32, GhostAlpha);	// 120, 75, 50 (HSV)
-	FColor RedTint(128, 32, 32, GhostAlpha);	// 0,   75, 50 (HSV)
+	FLinearColor PlaybackRangeStartColor = Sequencer->GetSequencerSettings()->GetPlaybackRangeStartColor();
+	FLinearColor PlaybackRangeEndColor = Sequencer->GetSequencerSettings()->GetPlaybackRangeEndColor();
+
+	FColor RedTint(192, 48, 48, GhostAlpha);	// 0,   75, 75 (HSV)
 	FLinearColor SectionTint = InPainter.GetSectionColor().LinearRGBToHSV();
 	SectionTint.B *= 0.1f;
 	SectionTint = SectionTint.HSVToLinearRGB();
 
-	auto PaintBoundary = [PaintDashedLine, SectionTint](FFrameTime Time)
+	auto PaintBoundary = [PaintDashedLine, SectionTint, LayerId](FFrameTime Time)
 	{
 		PaintDashedLine(Time, SectionTint.ToFColor(true), 0.f);
 		return true;
 	};
 
-	if (!OuterToInnerTransform.ExtractBoundariesWithinRange(SectionStartFrame, SectionStartFrame+SectionSize, PaintBoundary))
+	if (!OuterToInnerTransform.ExtractBoundariesWithinRange(SectionStartFrame, SectionStartFrame + SectionSize, PaintBoundary))
 	{
 		// Just use the playback range
 		TOptional<FFrameTime> StartBound = InnerToOuterTransform.TryTransformTime(MovieScene->GetPlaybackRange().GetLowerBoundValue());
-		TOptional<FFrameTime> EndBound   = InnerToOuterTransform.TryTransformTime(MovieScene->GetPlaybackRange().GetUpperBoundValue());
+		TOptional<FFrameTime> EndBound = InnerToOuterTransform.TryTransformTime(MovieScene->GetPlaybackRange().GetUpperBoundValue());
 
 		if (StartBound)
 		{
-			PaintLine(StartBound.GetValue(), GreenTint, 0.f);
+			PaintLine(StartBound.GetValue(), PlaybackRangeStartColor.CopyWithNewOpacity(GhostAlpha).ToFColor(true), 0.f);
 		}
 		if (EndBound)
 		{
-			PaintLine(EndBound.GetValue(), RedTint, 1.f);
+			PaintLine(EndBound.GetValue(), PlaybackRangeEndColor.CopyWithNewOpacity(GhostAlpha).ToFColor(true), 1.f);
 		}
 	}
 
@@ -250,7 +254,7 @@ void FSubSectionPainterUtil::PaintSectionBounds(const UMovieSceneSubSection& Sec
 
 			FSlateDrawElement::MakeBox(
 				InPainter.DrawElements,
-				InPainter.LayerId++,
+				LayerId + BoxLayerIdOffset,
 				InPainter.SectionGeometry.ToPaintGeometry(
 					FVector2f(InPainter.SectionGeometry.Size.X * WidthFactor-1, InPainter.SectionGeometry.Size.Y),
 					FSlateLayoutTransform()
@@ -280,7 +284,7 @@ void FSubSectionPainterUtil::PaintSectionBounds(const UMovieSceneSubSection& Sec
 
 			FSlateDrawElement::MakeBox(
 				InPainter.DrawElements,
-				InPainter.LayerId++,
+				LayerId + BoxLayerIdOffset,
 				InPainter.SectionGeometry.ToPaintGeometry(
 					FVector2f(InPainter.SectionGeometry.Size.X * WidthFactor-1, InPainter.SectionGeometry.Size.Y),
 					FSlateLayoutTransform(FVector2f(Offset, 0.f))
diff --git a/Engine/Source/Editor/MovieSceneTools/Public/TrackEditors/SubTrackEditorBase.h b/Engine/Source/Editor/MovieSceneTools/Public/TrackEditors/SubTrackEditorBase.h
--- a/Engine/Source/Editor/MovieSceneTools/Public/TrackEditors/SubTrackEditorBase.h
+++ b/Engine/Source/Editor/MovieSceneTools/Public/TrackEditors/SubTrackEditorBase.h
@@ -61,7 +61,7 @@ public:
     static FSubSectionPainterResult PaintSection(TSharedPtr<const ISequencer> Sequencer, const UMovieSceneSubSection& SectionObject, FSequencerSectionPainter& InPainter, FSubSectionPainterParams Params);
 
 private:
-    static void PaintSectionBounds(const UMovieSceneSubSection& SectionObject, const UMovieSceneSequence& InnerSequence, FSequencerSectionPainter& InPainter, ESlateDrawEffect DrawEffects);
+    static void PaintSectionBounds(TSharedPtr<const ISequencer> Sequencer, const UMovieSceneSubSection& SectionObject, const UMovieSceneSequence& InnerSequence, FSequencerSectionPainter& InPainter, ESlateDrawEffect DrawEffects);
 };
 
 /**
diff --git a/Engine/Source/Editor/Sequencer/Private/SSequencer.cpp b/Engine/Source/Editor/Sequencer/Private/SSequencer.cpp
--- a/Engine/Source/Editor/Sequencer/Private/SSequencer.cpp
+++ b/Engine/Source/Editor/Sequencer/Private/SSequencer.cpp
@@ -711,7 +711,9 @@ TSharedRef<SGridPanel> SSequencer::ConstructTrackAreaGridPanel(const FArguments&
 			Ranges,
 			TimeSliderControllerRef,
 			TAttribute<EVisibility>(this, &SSequencer::GetTimeRangeVisibility),
-			MakeAttributeSP(SequencerRef, &FSequencer::GetNumericTypeInterface, UE::Sequencer::ENumericIntent::Position)
+			MakeAttributeSP(SequencerRef, &FSequencer::GetNumericTypeInterface, UE::Sequencer::ENumericIntent::Position),
+			GetSequencerSettings()->GetPlaybackRangeStartColor(),
+			GetSequencerSettings()->GetPlaybackRangeEndColor()
 		),
 		SequencerWidgets.CreateTimeRangeSlider(TimeSliderControllerRef)
 	);
@@ -2572,8 +2574,9 @@ TSharedRef<SWidget> SSequencer::MakeTimeRange(const TSharedRef<SWidget>& InnerCo
 		ShowRange,
 		TimeSliderController.ToSharedRef(),
 		EVisibility::Visible,
-		MakeAttributeSP(SequencerPtr.Pin().Get(), &FSequencer::GetNumericTypeInterface, UE::Sequencer::ENumericIntent::Position)
-		);
+		MakeAttributeSP(SequencerPtr.Pin().Get(), &FSequencer::GetNumericTypeInterface, UE::Sequencer::ENumericIntent::Position),
+		GetSequencerSettings()->GetPlaybackRangeStartColor(),
+		GetSequencerSettings()->GetPlaybackRangeEndColor());
 	return SequencerWidgets.CreateTimeRange(Args, InnerContent);
 }
 
diff --git a/Engine/Source/Editor/Sequencer/Private/Sequencer.cpp b/Engine/Source/Editor/Sequencer/Private/Sequencer.cpp
--- a/Engine/Source/Editor/Sequencer/Private/Sequencer.cpp
+++ b/Engine/Source/Editor/Sequencer/Private/Sequencer.cpp
@@ -4758,8 +4758,8 @@ TSharedRef<SWidget> FSequencer::OnCreateTransportSetPlaybackStart()
 		.IsFocusable(false)
 		[
 			SNew(SImage)
-			.ColorAndOpacity(FLinearColor::Green)
-			.Image(FAppStyle::Get().GetBrush("Sequencer.Transport.SetPlayStart"))
+				.ColorAndOpacity_Lambda([this]() { return Settings->GetPlaybackRangeStartColor(); })
+				.Image(FAppStyle::Get().GetBrush("Sequencer.Transport.SetPlayStart"))
 		];
 }
 
@@ -4809,7 +4809,7 @@ TSharedRef<SWidget> FSequencer::OnCreateTransportSetPlaybackEnd()
 		.IsFocusable(false)
 		[
 			SNew(SImage)
-			.ColorAndOpacity(FLinearColor::Red)
+			.ColorAndOpacity_Lambda([this]() { return Settings->GetPlaybackRangeEndColor(); })
 			.Image(FAppStyle::Get().GetBrush("Sequencer.Transport.SetPlayEnd"))
 		];
 }
diff --git a/Engine/Source/Editor/Sequencer/Private/SequencerSettings.cpp b/Engine/Source/Editor/Sequencer/Private/SequencerSettings.cpp
--- a/Engine/Source/Editor/Sequencer/Private/SequencerSettings.cpp
+++ b/Engine/Source/Editor/Sequencer/Private/SequencerSettings.cpp
@@ -57,6 +57,8 @@ USequencerSettings::USequencerSettings( const FObjectInitializer& ObjectInitiali
 	KeyAreaHeightWithCurves = SequencerLayoutConstants::KeyAreaHeight;
 	bDeleteKeysWhenTrimming = true;
 	bDisableSectionsAfterBaking = true;
+	PlaybackRangeStartColor = FLinearColor(0.1f, 0.5f, 0.1f, 1.0f);
+	PlaybackRangeEndColor = FLinearColor(0.5f, 0.01f, 0.01f, 1.0f);
 	MarkedFrameColor = FLinearColor(0.f, 1.f, 1.f, 0.4f);
 	bCleanPlaybackMode = true;
 	bActivateRealtimeViewports = true;
@@ -856,6 +858,34 @@ void USequencerSettings::SetDisableSectionsAfterBaking(bool bInDisableSectionsAf
 	}
 }
 
+FLinearColor USequencerSettings::GetPlaybackRangeStartColor() const
+{
+	return PlaybackRangeStartColor;
+}
+
+void USequencerSettings::SetPlaybackRangeStartColor(const FLinearColor& InPlaybackRangeStartColor)
+{
+	if (PlaybackRangeStartColor != InPlaybackRangeStartColor)
+	{
+		PlaybackRangeStartColor = InPlaybackRangeStartColor;
+		SaveConfig();
+	}
+}
+
+FLinearColor USequencerSettings::GetPlaybackRangeEndColor() const
+{
+	return PlaybackRangeEndColor;
+}
+
+void USequencerSettings::SetPlaybackRangeEndColor(const FLinearColor& InPlaybackRangeEndColor)
+{
+	if (PlaybackRangeEndColor != InPlaybackRangeEndColor)
+	{
+		PlaybackRangeEndColor = InPlaybackRangeEndColor;
+		SaveConfig();
+	}
+}
+
 FLinearColor USequencerSettings::GetMarkedFrameColor() const
 {
 	return MarkedFrameColor;
diff --git a/Engine/Source/Editor/Sequencer/Private/SequencerTimeSliderController.cpp b/Engine/Source/Editor/Sequencer/Private/SequencerTimeSliderController.cpp
--- a/Engine/Source/Editor/Sequencer/Private/SequencerTimeSliderController.cpp
+++ b/Engine/Source/Editor/Sequencer/Private/SequencerTimeSliderController.cpp
@@ -811,13 +811,16 @@ int32 FSequencerTimeSliderController::DrawPlaybackRange(const FGeometry& Allotte
 	const float PlaybackRangeL = RangeToScreen.InputToLocalX(PlaybackRange.GetLowerBoundValue() / TickResolution);
 	const float PlaybackRangeR = RangeToScreen.InputToLocalX(PlaybackRange.GetUpperBoundValue() / TickResolution) - 1;
 
+	const FLinearColor PlaybackRangeStartColor = Sequencer->GetSequencerSettings()->GetPlaybackRangeStartColor();
+	const FLinearColor PlaybackRangeEndColor = Sequencer->GetSequencerSettings()->GetPlaybackRangeEndColor();
+
 	FSlateDrawElement::MakeBox(
 		OutDrawElements,
 		LayerId+1,
 		AllottedGeometry.ToPaintGeometry(FVector2f(Args.BrushWidth, AllottedGeometry.Size.Y), FSlateLayoutTransform(FVector2f(PlaybackRangeL, 0.f))),
 		Args.StartBrush,
 		ESlateDrawEffect::None,
-		FColor(32, 128, 32, OpacityBlend)	// 120, 75, 50 (HSV)
+		PlaybackRangeStartColor.CopyWithNewOpacity(OpacityBlend)
 	);
 
 	FSlateDrawElement::MakeBox(
@@ -826,7 +829,7 @@ int32 FSequencerTimeSliderController::DrawPlaybackRange(const FGeometry& Allotte
 		AllottedGeometry.ToPaintGeometry(FVector2f(Args.BrushWidth, AllottedGeometry.Size.Y), FSlateLayoutTransform(FVector2f(PlaybackRangeR - Args.BrushWidth, 0.f))),
 		Args.EndBrush,
 		ESlateDrawEffect::None,
-		FColor(128, 32, 32, OpacityBlend)	// 0, 75, 50 (HSV)
+		PlaybackRangeEndColor.CopyWithNewOpacity(OpacityBlend)
 	);
 
 	// Black tint for excluded regions
@@ -877,6 +880,9 @@ int32 FSequencerTimeSliderController::DrawSubSequenceRange(const FGeometry& Allo
 	static const FSlateBrush* LineBrushL(FAppStyle::GetBrush("Sequencer.Timeline.PlayRange_L"));
 	static const FSlateBrush* LineBrushR(FAppStyle::GetBrush("Sequencer.Timeline.PlayRange_R"));
 
+	const FLinearColor PlaybackRangeStartColor = Sequencer->GetSequencerSettings()->GetPlaybackRangeStartColor();
+	const FLinearColor PlaybackRangeEndColor = Sequencer->GetSequencerSettings()->GetPlaybackRangeEndColor();
+
 	FColor GreenTint(32, 128, 32);	// 120, 75, 50 (HSV)
 	FSlateDrawElement::MakeBox(
 		OutDrawElements,
@@ -884,17 +890,17 @@ int32 FSequencerTimeSliderController::DrawSubSequenceRange(const FGeometry& Allo
 		AllottedGeometry.ToPaintGeometry(FVector2f(Args.BrushWidth, AllottedGeometry.Size.Y), FSlateLayoutTransform(FVector2f(SubSequenceRangeL, 0.f))),
 		LineBrushL,
 		ESlateDrawEffect::None,
-		GreenTint
+		PlaybackRangeStartColor
 	);
 
-	FColor RedTint(128, 32, 32);	// 0, 75, 50 (HSV)
+	FColor RedTint(192, 48, 48);	// 0, 75, 75 (HSV)
 	FSlateDrawElement::MakeBox(
 		OutDrawElements,
 		LayerId+1,
 		AllottedGeometry.ToPaintGeometry(FVector2f(Args.BrushWidth, AllottedGeometry.Size.Y), FSlateLayoutTransform(FVector2f(SubSequenceRangeR - Args.BrushWidth, 0.f))),
 		LineBrushR,
 		ESlateDrawEffect::None,
-		RedTint
+		PlaybackRangeEndColor
 	);
 
 	// Black tint for excluded regions
@@ -923,7 +929,7 @@ int32 FSequencerTimeSliderController::DrawSubSequenceRange(const FGeometry& Allo
 		AllottedGeometry.ToPaintGeometry(FVector2f(16.f, AllottedGeometry.Size.Y), FSlateLayoutTransform(FVector2f(SubSequenceRangeL - 16.f, 0.f))),
 		FAppStyle::GetBrush("Sequencer.Timeline.SubSequenceRangeHashL"),
 		ESlateDrawEffect::None,
-		GreenTint
+		PlaybackRangeStartColor
 	);
 
 	FSlateDrawElement::MakeBox(
@@ -932,7 +938,7 @@ int32 FSequencerTimeSliderController::DrawSubSequenceRange(const FGeometry& Allo
 		AllottedGeometry.ToPaintGeometry(FVector2f(16.f, AllottedGeometry.Size.Y), FSlateLayoutTransform(FVector2f(SubSequenceRangeR, 0.f))),
 		FAppStyle::GetBrush("Sequencer.Timeline.SubSequenceRangeHashR"),
 		ESlateDrawEffect::None,
-		RedTint
+		PlaybackRangeEndColor
 	);
 
 	return LayerId + 1;
diff --git a/Engine/Source/Editor/Sequencer/Public/SequencerSettings.h b/Engine/Source/Editor/Sequencer/Public/SequencerSettings.h
--- a/Engine/Source/Editor/Sequencer/Public/SequencerSettings.h
+++ b/Engine/Source/Editor/Sequencer/Public/SequencerSettings.h
@@ -404,6 +404,16 @@ public:
 	/** Set whether to disable sections when baking, as opposed to deleting */
 	void SetDisableSectionsAfterBaking(bool bInDisableSectionsAfterBaking);
 
+	/** Set the playback range start color */
+	FLinearColor GetPlaybackRangeStartColor() const;
+	/** Set the playback range start color */
+	void SetPlaybackRangeStartColor(const FLinearColor& InColor);
+
+	/** Set the playback range end color */
+	FLinearColor GetPlaybackRangeEndColor() const;
+	/** Set the playback range end color */
+	void SetPlaybackRangeEndColor(const FLinearColor& InColor);
+
 	/** @return the default marked frame color */
 	FLinearColor GetMarkedFrameColor() const;
 	/** Set the default marked frame color */
@@ -771,6 +781,14 @@ protected:
 	UPROPERTY(config, EditAnywhere, Category = Timeline)
 	bool bDisableSectionsAfterBaking;
 	
+	/** PLayback range start color */
+	UPROPERTY(config, EditAnywhere, Category = Timeline)
+	FLinearColor PlaybackRangeStartColor;
+
+	/** PLayback range end color */
+	UPROPERTY(config, EditAnywhere, Category = Timeline)
+	FLinearColor PlaybackRangeEndColor;
+
 	/** Default marked frame color */
 	UPROPERTY(config, EditAnywhere, Category = Timeline)
 	FLinearColor MarkedFrameColor;
diff --git a/Engine/Source/Editor/SequencerWidgets/Private/STimeRange.cpp b/Engine/Source/Editor/SequencerWidgets/Private/STimeRange.cpp
--- a/Engine/Source/Editor/SequencerWidgets/Private/STimeRange.cpp
+++ b/Engine/Source/Editor/SequencerWidgets/Private/STimeRange.cpp
@@ -38,6 +38,8 @@ template <typename NumericType> struct INumericTypeInterface;
 void STimeRange::Construct( const STimeRange::FArguments& InArgs, TSharedRef<ITimeSliderController> InTimeSliderController )
 {
 	TimeSliderController = InTimeSliderController;
+	PlaybackRangeStartColor = InArgs._PlaybackRangeStartColor;
+	PlaybackRangeEndColor = InArgs._PlaybackRangeEndColor;
 
 	TSharedRef<SWidget> WorkingRangeStart = SNullWidget::NullWidget, WorkingRangeEnd = SNullWidget::NullWidget;
 	if (InArgs._ShowWorkingRange)
@@ -169,7 +171,7 @@ void STimeRange::Construct( const STimeRange::FArguments& InArgs, TSharedRef<ITi
 				SNew(SBorder)
 				.Padding(0.f)
 				.BorderImage(nullptr)
-				.ForegroundColor(FLinearColor::Green)
+				.ForegroundColor(PlaybackRangeStartColor)
 				[
 					PlaybackRangeStart
 				]
@@ -225,7 +227,7 @@ void STimeRange::Construct( const STimeRange::FArguments& InArgs, TSharedRef<ITi
 				SNew(SBorder)
 				.Padding(0.f)
 				.BorderImage(nullptr)
-				.ForegroundColor(FLinearColor::Red)
+				.ForegroundColor(PlaybackRangeEndColor)
 				[
 					PlaybackRangeEnd
 				]
diff --git a/Engine/Source/Editor/SequencerWidgets/Private/STimeRange.h b/Engine/Source/Editor/SequencerWidgets/Private/STimeRange.h
--- a/Engine/Source/Editor/SequencerWidgets/Private/STimeRange.h
+++ b/Engine/Source/Editor/SequencerWidgets/Private/STimeRange.h
@@ -14,7 +14,11 @@ class STimeRange : public ITimeSlider
 {
 public:
 	SLATE_BEGIN_ARGS(STimeRange)
-		: _ShowWorkingRange(true), _ShowViewRange(false), _ShowPlaybackRange(true)
+		: _ShowWorkingRange(true)
+		, _ShowViewRange(false)
+		, _ShowPlaybackRange(true)
+		, _PlaybackRangeStartColor(FLinearColor::Red)
+		, _PlaybackRangeEndColor(FLinearColor::Green)
 	{}
 		/** Whether to show the working range */
 		SLATE_ATTRIBUTE( TSharedPtr<INumericTypeInterface<double>>, NumericTypeInterface )
@@ -30,6 +34,10 @@ public:
 		SLATE_ARGUMENT( bool, EnableViewRange )
 		/** Whether to enable the playback range */
 		SLATE_ARGUMENT( bool, EnablePlaybackRange )
+		/** The color of the playback start indicator */
+		SLATE_ARGUMENT( FLinearColor, PlaybackRangeStartColor )
+		/** The color of the playback end indicator */
+		SLATE_ARGUMENT( FLinearColor, PlaybackRangeEndColor )
 		/* Content to display inside the time range */
 		SLATE_DEFAULT_SLOT( FArguments, CenterContent )
 	SLATE_END_ARGS()
@@ -79,4 +87,7 @@ protected:
 
 private:
 	TSharedPtr<ITimeSliderController> TimeSliderController;
+
+	FLinearColor PlaybackRangeStartColor;
+	FLinearColor PlaybackRangeEndColor;
 };
diff --git a/Engine/Source/Editor/SequencerWidgets/Private/SequencerWidgetsModule.cpp b/Engine/Source/Editor/SequencerWidgets/Private/SequencerWidgetsModule.cpp
--- a/Engine/Source/Editor/SequencerWidgets/Private/SequencerWidgetsModule.cpp
+++ b/Engine/Source/Editor/SequencerWidgets/Private/SequencerWidgetsModule.cpp
@@ -54,6 +54,8 @@ public:
 		.EnableWorkingRange(!!(InArgs.EnableRanges & EShowRange::WorkingRange))
 		.EnableViewRange(!!(InArgs.EnableRanges & EShowRange::ViewRange))
 		.EnablePlaybackRange(!!(InArgs.EnableRanges & EShowRange::PlaybackRange))
+		.PlaybackRangeStartColor(InArgs.PlaybackRangeStartColor)
+		.PlaybackRangeEndColor(InArgs.PlaybackRangeEndColor)
 		[
 			Content
 		];
diff --git a/Engine/Source/Editor/SequencerWidgets/Public/ISequencerWidgetsModule.h b/Engine/Source/Editor/SequencerWidgets/Public/ISequencerWidgetsModule.h
--- a/Engine/Source/Editor/SequencerWidgets/Public/ISequencerWidgetsModule.h
+++ b/Engine/Source/Editor/SequencerWidgets/Public/ISequencerWidgetsModule.h
@@ -29,13 +29,17 @@ struct FTimeRangeArgs
 		EShowRange InShowRanges,
 		TSharedRef<ITimeSliderController> InController,
 		TAttribute<EVisibility> InVisibilityDelegate,
-		TAttribute<TSharedPtr<INumericTypeInterface<double>>> InNumericTypeInterface
-		)
+		TAttribute<TSharedPtr<INumericTypeInterface<double>>> InNumericTypeInterface,
+		FLinearColor InPlaybackRangeStartColor = FLinearColor::Green,
+		FLinearColor InPlaybackRangeEndColor = FLinearColor::Red
+	)
 		: ShowRanges(InShowRanges)
 		, EnableRanges(InShowRanges)
 		, Controller(InController)
 		, VisibilityDelegate(InVisibilityDelegate)
 		, NumericTypeInterface(InNumericTypeInterface)
+		, PlaybackRangeStartColor(InPlaybackRangeStartColor)
+		, PlaybackRangeEndColor(InPlaybackRangeEndColor)
 	{
 	}
 
@@ -44,13 +48,17 @@ struct FTimeRangeArgs
 		EShowRange InEnableRanges,
 		TSharedRef<ITimeSliderController> InController,
 		TAttribute<EVisibility> InVisibilityDelegate,
-		TAttribute<TSharedPtr<INumericTypeInterface<double>>> InNumericTypeInterface
+		TAttribute<TSharedPtr<INumericTypeInterface<double>>> InNumericTypeInterface,
+		FLinearColor InPlaybackRangeStartColor = FLinearColor::Green,
+		FLinearColor InPlaybackRangeEndColor = FLinearColor::Red
 		)
 		: ShowRanges(InShowRanges)
 		, EnableRanges(InEnableRanges)
 		, Controller(InController)
 		, VisibilityDelegate(InVisibilityDelegate)
 		, NumericTypeInterface(InNumericTypeInterface)
+		, PlaybackRangeStartColor(InPlaybackRangeStartColor)
+		, PlaybackRangeEndColor(InPlaybackRangeEndColor)
 	{
 	}
 
@@ -64,6 +72,10 @@ struct FTimeRangeArgs
 	TAttribute<EVisibility> VisibilityDelegate;
 	/** Numeric type interface to use for frame<->time conversion */
 	TAttribute<TSharedPtr<INumericTypeInterface<double>>> NumericTypeInterface;
+	/** Playback range start color */
+	FLinearColor PlaybackRangeStartColor = FLinearColor::Green;
+	/** Playback range end color */
+	FLinearColor PlaybackRangeEndColor = FLinearColor::Red;
 };
 
 /**
