Anim Layer: Bake Fix for absolute with additives when doing per frame bakes.

diff --git a/Engine/Source/Editor/MovieSceneTools/Public/MovieSceneToolHelpers.h b/Engine/Source/Editor/MovieSceneTools/Public/MovieSceneToolHelpers.h
--- a/Engine/Source/Editor/MovieSceneTools/Public/MovieSceneToolHelpers.h
+++ b/Engine/Source/Editor/MovieSceneTools/Public/MovieSceneToolHelpers.h
@@ -914,9 +914,6 @@ public:
 							continue;
 						}
 						CurveValueType Value = 0.0;
-						const FFrameTime FrameTime(Frame);
-						Channel->Evaluate(FrameTime, Value);
-
 						MergeValue(Value, Channels, Frame, SectionChannelIndex, Sections, AbsoluteSections,
 							OtherSections, OverrideChannelIndex, MergeAlgorithm);
 						ValuesToSet.Add(TPair<FFrameNumber, CurveValueType>(Frame, Value));
