Anim In Engine: Only compensate from game thread.
#trivial

diff --git a/Engine/Source/Editor/MovieSceneTools/Private/Constraints/MovieSceneConstraintChannelHelper.cpp b/Engine/Source/Editor/MovieSceneTools/Private/Constraints/MovieSceneConstraintChannelHelper.cpp
--- a/Engine/Source/Editor/MovieSceneTools/Private/Constraints/MovieSceneConstraintChannelHelper.cpp
+++ b/Engine/Source/Editor/MovieSceneTools/Private/Constraints/MovieSceneConstraintChannelHelper.cpp
@@ -890,7 +890,7 @@ void FMovieSceneConstraintChannelHelper::HandleConstraintKeyDeleted(
 	UMovieSceneSection* InSection,
 	const FFrameNumber& InTime)
 {
-	if (bDoNotCompensate)
+	if (bDoNotCompensate == true || IsInGameThread() == false) //this may happen in an non game thread via a parallel for when we are deleting many channels
 	{
 		return;
 	}
@@ -899,6 +899,7 @@ void FMovieSceneConstraintChannelHelper::HandleConstraintKeyDeleted(
 	{
 		return;
 	}
+
 	
 	const FFrameNumber TimeMinusOne(InTime - 1);
 
