Sequencer: Fixed object binding IDs not working when set to a binding within a subsequence inside a director BP

diff --git a/Engine/Source/Editor/MovieSceneTools/Private/K2Node_GetSequenceBinding.cpp b/Engine/Source/Editor/MovieSceneTools/Private/K2Node_GetSequenceBinding.cpp
--- a/Engine/Source/Editor/MovieSceneTools/Private/K2Node_GetSequenceBinding.cpp
+++ b/Engine/Source/Editor/MovieSceneTools/Private/K2Node_GetSequenceBinding.cpp
@@ -91,7 +91,10 @@ public:
 			if (Pin->Direction == EGPD_Output && Pin->LinkedTo.Num())
 			{
 				FBPTerminal* Term = RegisterLiteral(Context, Pin);
-				FMovieSceneObjectBindingID::StaticStruct()->ExportText(Term->Name, &BindingNode->Binding, nullptr, nullptr, 0, nullptr);
+
+				// Intentionally pass BindingPtr to both the data and defaults for ExportText to ensure that we always disable delta-serialization on the struct
+				FMovieSceneObjectBindingID* BindingPtr = &BindingNode->Binding;
+				FMovieSceneObjectBindingID::StaticStruct()->ExportText(Term->Name, BindingPtr, BindingPtr, nullptr, 0, nullptr);
 			}
 		}
 	}
diff --git a/Engine/Source/Editor/MovieSceneTools/Private/MovieSceneObjectBindingIDPicker.cpp b/Engine/Source/Editor/MovieSceneTools/Private/MovieSceneObjectBindingIDPicker.cpp
--- a/Engine/Source/Editor/MovieSceneTools/Private/MovieSceneObjectBindingIDPicker.cpp
+++ b/Engine/Source/Editor/MovieSceneTools/Private/MovieSceneObjectBindingIDPicker.cpp
@@ -346,10 +346,12 @@ void FMovieSceneObjectBindingIDPicker::SetCurrentValueFromFixed(UE::MovieScene::
 	TSharedPtr<ISequencer>              Sequencer = WeakSequencer.Pin();
 	const FMovieSceneSequenceHierarchy* Hierarchy = Sequencer.IsValid() ? Sequencer->GetEvaluationTemplate().GetHierarchy() : nullptr;
 
-	// If there is no sequencer, just set the ID directly 
+	// If there is no sequencer, use a relative binding since there should be no contexts where this is different from a fixed binding
+	//    but relative bindings will support cases where this needs to be resolved within the context of a different subsequence (ie, if Sequence A
+	//    has a GetSequenceBinding node for itself inside its director BP, and that sequence is added to another sequence).
 	if (!Sequencer.IsValid())
 	{
-		SetCurrentValue(InValue);
+		SetCurrentValue(UE::MovieScene::FRelativeObjectBindingID(InValue.Guid, InValue.SequenceID));
 	}
 	// If we don't know the local sequence ID, or we have no hierarchy, or we're resetting the binding; set a relative ID
 	else if (LocalSequenceID == MovieSceneSequenceID::Invalid || !InValue.Guid.IsValid() || Hierarchy == nullptr)
