Curve Editor: Fix issues where latest improvements in Curve Cache Pool did not always draw pre- and post-infinity extrapolation correctly. Curve Cache Pool is now enabled again by default, can be disabled with CVar CurveEditor.UseCurveCachePool 0

#rnx

diff --git a/Engine/Source/Editor/CurveEditor/Private/CurveEditorCurveDrawParamsCache.cpp b/Engine/Source/Editor/CurveEditor/Private/CurveEditorCurveDrawParamsCache.cpp
--- a/Engine/Source/Editor/CurveEditor/Private/CurveEditorCurveDrawParamsCache.cpp
+++ b/Engine/Source/Editor/CurveEditor/Private/CurveEditorCurveDrawParamsCache.cpp
@@ -12,7 +12,7 @@
 
 namespace UE::CurveEditor
 {
-	int32 GCurveEditorUseCurveCachePool = 0;
+	int32 GCurveEditorUseCurveCachePool = 1;
 
 	static FAutoConsoleVariableRef CCurveEditorUseCurveCachePool(
 		TEXT("CurveEditor.UseCurveCachePool"),
diff --git a/Engine/Source/Editor/MovieSceneTools/Private/Cache/MovieSceneCachedCurve.h b/Engine/Source/Editor/MovieSceneTools/Private/Cache/MovieSceneCachedCurve.h
--- a/Engine/Source/Editor/MovieSceneTools/Private/Cache/MovieSceneCachedCurve.h
+++ b/Engine/Source/Editor/MovieSceneTools/Private/Cache/MovieSceneCachedCurve.h
@@ -131,8 +131,11 @@ namespace UE::MovieSceneTools
 		/** Updates the piecewise cuve depending on flags */
 		void ConditionallyUpdatePiecewiseCurve(const FMovieSceneUpdateCachedCurveData<ChannelType>& Data);
 
-		/** Creates a task to draw full range interpolating points, possibly async. */
-		void CreateFullRangeInterpolatingPointsDrawTask(const ChannelType& Channel);
+		/** 
+		 * When the curve is being edited interactively, paints the visible part synchronous. 
+		 * When curve data changed in any way, paints the whole range of the curve async.
+		 */
+		void ConditionallyPaintCurve(const ChannelType& Channel);
 
 		/** Draws the interpolating points from the currently cached full range interpolating points */
 		void DrawInterpolatingPointsFromFullRange();
@@ -278,9 +281,7 @@ namespace UE::MovieSceneTools
 		// Update depending on above data
 		ConditionallyUpdateKeyDrawInfos(UpdateData);
 		ConditionallyUpdatePiecewiseCurve(UpdateData);
-
-		// Create a task for the curve painter
-		CreateFullRangeInterpolatingPointsDrawTask(UpdateData.Channel);
+		ConditionallyPaintCurve(UpdateData.Channel);
 	}
 
 	template <typename ChannelType>
@@ -467,6 +468,8 @@ namespace UE::MovieSceneTools
 		{
 			FullRangeInterpolation.PreInfinityExtrapolation = NewPreInfinityExtrapolation;
 			FullRangeInterpolation.PostInfinityExtrapolation = NewPostInfinityExtrapolation;
+
+			Flags.store(Flags.load() | EMovieSceneCurveCacheChangeFlags::ChangedCurveData);
 		}
 	}
 
@@ -533,7 +536,7 @@ namespace UE::MovieSceneTools
 	}
 
 	template <typename ChannelType>
-	void FMovieSceneCachedCurve<ChannelType>::CreateFullRangeInterpolatingPointsDrawTask(const ChannelType& Channel)
+	void FMovieSceneCachedCurve<ChannelType>::ConditionallyPaintCurve(const ChannelType& Channel)
 	{		
 		// Only draw the full range interpolating points when the curve was edited or when zooming
 		const bool bCurveIsBeingEdited = EnumHasAnyFlags(Flags.load(), EMovieSceneCurveCacheChangeFlags::ChangedCurveData);
@@ -568,14 +571,41 @@ namespace UE::MovieSceneTools
 			return;
 		}
 
-		// Create draw tasks for complex curves
 		const bool bDeferredZooming = GZoomInterpolatesCurvePointsAsync > 0 ? bZooming : false;
 		const bool bInteractive = bCurveIsBeingEdited || !bDeferredZooming;
 
+		// Draw interactive changes to the draw params directly
+		if (bInteractive)
+		{
+			const double StartTimeSeconds = ScreenSpace.GetInputMin();
+			const double EndTimeSeconds = ScreenSpace.GetInputMax();
+			const double TimeThreshold = GetTimeThreshold();
+			const double ValueThreshold = GetValueThreshold();
+
+			TArray<TTuple<double, double>> InteractiveInterpolatingPoints;
+			Channel.PopulateCurvePoints(StartTimeSeconds, EndTimeSeconds, TimeThreshold, ValueThreshold, TickResolution, InteractiveInterpolatingPoints);
+
+			CachedDrawParams.InterpolatingPoints.Reset();
+			CachedDrawParams.InterpolatingPoints.Reserve(InteractiveInterpolatingPoints.Num());
+
+			// Convert the interpolating points to screen space
+			for (const TTuple<double, double>& Point : InteractiveInterpolatingPoints)
+			{
+				CachedDrawParams.InterpolatingPoints.Emplace(
+					ScreenSpace.SecondsToScreen(Point.Get<0>() + InputDisplayOffset),
+					ScreenSpace.ValueToScreen(Point.Get<1>())
+				);
+			}
+
+			DrawKeys();
+			ApplyDrawParams();
+
+			Flags = EMovieSceneCurveCacheChangeFlags::None;
+		}
+
+		// Create a draw task to draw interpolating points for the full range of the curve
 		const TSharedRef<FMovieSceneInterpolatingPointsDrawTask<ChannelType>> Task = MakeShared<FMovieSceneInterpolatingPointsDrawTask<ChannelType>>(
 			SharedThis(this),
-			Channel,
-			bInteractive,
 			[WeakThis = AsWeak(), this](TArray<FVector2D> NewInterpolatingPoints, TArray<int32> NewKeyOffsets)
 			{
 				if (!WeakThis.IsValid())
@@ -669,7 +699,7 @@ namespace UE::MovieSceneTools
 			FullRangeInterpolation.PreInfinityExtrapolation == ERichCurveExtrapolation::RCCE_CycleWithOffset ||
 			FullRangeInterpolation.PreInfinityExtrapolation == ERichCurveExtrapolation::RCCE_Oscillate)
 		{
-			const double InfinityOffset = FullRangeInterpolation.Points[0].X - ScreenSpace.GetInputMin();
+			const double InfinityOffset = FullRangeInterpolation.Points[0].X + InputDisplayOffset - ScreenSpace.GetInputMin();
 			if (InfinityOffset < 0.0)
 			{
 				// Don't draw if pre-infinity is not visible
@@ -699,7 +729,7 @@ namespace UE::MovieSceneTools
 				const int32 NumIterations = static_cast<int32>(InfinityOffset / Duration) + 1;
 				for (int32 Iteration = NumIterations; Iteration > 0; Iteration--)
 				{
-					const bool bReverse = FullRangeInterpolation.PostInfinityExtrapolation == ERichCurveExtrapolation::RCCE_Oscillate && Iteration % 2 != 0;
+					const bool bReverse = FullRangeInterpolation.PreInfinityExtrapolation == ERichCurveExtrapolation::RCCE_Oscillate && Iteration % 2 != 0;
 
 					if (bReverse)
 					{
@@ -741,10 +771,11 @@ namespace UE::MovieSceneTools
 
 		if (FullRangeInterpolation.Points.Num() > 1 &&
 			(FullRangeInterpolation.PostInfinityExtrapolation == ERichCurveExtrapolation::RCCE_Cycle ||
-			FullRangeInterpolation.PostInfinityExtrapolation == ERichCurveExtrapolation::RCCE_CycleWithOffset ||
-			FullRangeInterpolation.PostInfinityExtrapolation == ERichCurveExtrapolation::RCCE_Oscillate))
+				FullRangeInterpolation.PostInfinityExtrapolation == ERichCurveExtrapolation::RCCE_CycleWithOffset ||
+				FullRangeInterpolation.PostInfinityExtrapolation == ERichCurveExtrapolation::RCCE_Oscillate))
 		{
-			const double InfinityOffset = ScreenSpace.GetInputMax() - FullRangeInterpolation.Points.Last().X;
+			const double InfinityOffset = ScreenSpace.GetInputMax() + InputDisplayOffset - FullRangeInterpolation.Points.Last().X;
+
 			if (InfinityOffset < 0.0)
 			{
 				// Don't draw if post-infinity is not visible
@@ -767,7 +798,7 @@ namespace UE::MovieSceneTools
 			{
 				// Cycle or oscillate
 				const double StartTime = FullRangeInterpolation.Points[0].X;
-				const double ValueOffset = FullRangeInterpolation.PreInfinityExtrapolation == ERichCurveExtrapolation::RCCE_CycleWithOffset ?
+				const double ValueOffset = FullRangeInterpolation.PostInfinityExtrapolation == ERichCurveExtrapolation::RCCE_CycleWithOffset ?
 					FullRangeInterpolation.Points.Last().Y - FullRangeInterpolation.Points[0].Y :
 					0.0;
 
diff --git a/Engine/Source/Editor/MovieSceneTools/Private/Cache/MovieSceneInterpolatingPointsDrawTask.cpp b/Engine/Source/Editor/MovieSceneTools/Private/Cache/MovieSceneInterpolatingPointsDrawTask.cpp
--- a/Engine/Source/Editor/MovieSceneTools/Private/Cache/MovieSceneInterpolatingPointsDrawTask.cpp
+++ b/Engine/Source/Editor/MovieSceneTools/Private/Cache/MovieSceneInterpolatingPointsDrawTask.cpp
@@ -15,8 +15,6 @@ namespace UE::MovieSceneTools
 	template <typename ChannelType>
 	FMovieSceneInterpolatingPointsDrawTask<ChannelType>::FMovieSceneInterpolatingPointsDrawTask(
 		const TSharedRef<FMovieSceneCachedCurve<ChannelType>>& CachedCurve,
-		const ChannelType& MovieSceneChannel,
-		const bool bInteractive,
 		const TFunction<void(TArray<FVector2D>, TArray<int32>)>& InCallback)
 		: Callback(InCallback)
 		, ScreenSpace(CachedCurve->GetScreenSpace())
@@ -39,16 +37,8 @@ namespace UE::MovieSceneTools
 		{
 			KeyPoints.Emplace(Times[DataIndex] / TickResolution, double(Values[DataIndex].Value));
 		}
-		InterpolatingPoints = KeyPoints;
-
-		// Process the visible part synchronous if interactive
-		if (bInteractive)
-		{		
-			DrawVisibleCurvePointsSynchronous(CachedCurve, MovieSceneChannel);
-			InvokeCallback();
 
-			SetFlags(ECurvePainterTaskStateFlags::Interactive);
-		}
+		InterpolatingPoints = KeyPoints;
 	}
 
 	template <typename ChannelType>
@@ -59,16 +49,7 @@ namespace UE::MovieSceneTools
 		bWorking = true;
 
 		const int32 OldSize = InterpolatingPoints.Num();
-
-		if (HasAnyFlags(ECurvePainterTaskStateFlags::Interactive))
-		{
-			RefineFullRangeInterpolatingPointsInternal(ECurvePainterInterpolationRange::StartToInteractive);
-			RefineFullRangeInterpolatingPointsInternal(ECurvePainterInterpolationRange::InteractiveToEnd);
-		}
-		else
-		{
-			RefineFullRangeInterpolatingPointsInternal(ECurvePainterInterpolationRange::FullRange);
-		}
+		RefineFullRangeInterpolatingPointsInternal();
 
 		const int32 NewSize = InterpolatingPoints.Num();
 		if (OldSize == NewSize)
@@ -113,50 +94,10 @@ namespace UE::MovieSceneTools
 	}
 
 	template<typename ChannelType>
-	void FMovieSceneInterpolatingPointsDrawTask<ChannelType>::RefineFullRangeInterpolatingPointsInternal(const ECurvePainterInterpolationRange InterpolationRange)
+	void FMovieSceneInterpolatingPointsDrawTask<ChannelType>::RefineFullRangeInterpolatingPointsInternal()
 	{			
-		const int32 StartingIndex = [&InterpolationRange, this]()
-			{
-				if (InterpolationRange == ECurvePainterInterpolationRange::FullRange ||
-					InterpolationRange == ECurvePainterInterpolationRange::StartToInteractive)
-				{
-					return 0;
-				}
-				else
-				{
-					const double InteractiveEnd = InteractiveTimeEnd.GetValue();
-					const int32 StartingIndex = InterpolatingPoints.IndexOfByPredicate(
-						[this, &InteractiveEnd](const FVector2D& InterpolatingPoint)
-						{
-							return InteractiveEnd <= InterpolatingPoint.X;
-						});
-
-					return StartingIndex != INDEX_NONE ? StartingIndex : InterpolatingPoints.Num() - 1;
-				}
-			}();
-
-		const int32 EndingIndex = [&InterpolationRange, this]()
-			{
-				if (InterpolationRange == ECurvePainterInterpolationRange::FullRange ||
-					InterpolationRange == ECurvePainterInterpolationRange::InteractiveToEnd)
-				{
-					return InterpolatingPoints.Num() - 1;
-				}
-				else
-				{
-					const double InteractiveStart = InteractiveTimeStart.GetValue();
-					const int32 EndingIndex = InterpolatingPoints.IndexOfByPredicate(
-						[this, &InteractiveStart](const FVector2D& InterpolatingPoint)
-						{
-							return InteractiveStart >= InterpolatingPoint.X;
-						});
-
-					return EndingIndex != INDEX_NONE ? EndingIndex : InterpolatingPoints.Num() - 1;
-				}
-			}();
-
 		constexpr float InterpTimes[] = { 0.25f, 0.5f, 0.6f };
-		for (int32 Index = StartingIndex; Index < EndingIndex; Index++)
+		for (int32 Index = 0; Index < InterpolatingPoints.Num() - 1; Index++)
 		{
 			const FVector2D& Lower = InterpolatingPoints[Index];
 			const FVector2D& Upper = InterpolatingPoints[Index + 1];
@@ -192,98 +133,6 @@ namespace UE::MovieSceneTools
 		}
 	}
 
-	template<typename ChannelType>
-	void FMovieSceneInterpolatingPointsDrawTask<ChannelType>::DrawVisibleCurvePointsSynchronous(const TSharedRef<FMovieSceneCachedCurve<ChannelType>>& CachedCurve, const ChannelType& Channel)
-	{
-		TRACE_CPUPROFILER_EVENT_SCOPE(FMovieSceneCurvePainterTas::DrawCurvePointsSynchronous);
-
-		const TArray<const FFrameNumber>& Times = CachedCurve->GetTimes();
-		const TArray<const ChannelValueType>& Values = CachedCurve->GetValues();
-
-		const int32 StartingIndex = CachedCurve->GetStartingIndex();
-		const int32 EndingIndex = CachedCurve->GetEndingIndex();
-
-		const double StartTime = ScreenSpace.GetInputMin();
-		const double EndTime = ScreenSpace.GetInputMax();
-
-		const FFrameNumber StartFrame = (StartTime * TickResolution).FloorToFrame();
-		const FFrameNumber EndFrame = (EndTime * TickResolution).CeilToFrame();
-
-		// Add the lower bound of the visible space
-		CurveValueType EvaluatedValue = 0;
-		if (Channel.Evaluate(StartFrame, EvaluatedValue))
-		{
-			InterpolatingPoints.Emplace(StartTime, double(EvaluatedValue));
-		}
-
-		// Add all keys in-between
-		for (int32 KeyIndex = StartingIndex; KeyIndex < EndingIndex; ++KeyIndex)
-		{
-			InterpolatingPoints.Emplace(Times[KeyIndex] / TickResolution, double(Values[KeyIndex].Value));
-		}
-
-		// Add the upper bound of the visible space
-		if (Channel.Evaluate(EndFrame, EvaluatedValue))
-		{
-			InterpolatingPoints.Emplace(EndTime, double(EvaluatedValue));
-		}
-
-		int32 OldSize;
-		do
-		{
-			OldSize = InterpolatingPoints.Num();
-			RefineVisibleCurvePointsSynchronous(CachedCurve, Channel);
-		} while (OldSize != InterpolatingPoints.Num());
-
-		// Remember what was interpolated already
-		InteractiveTimeStart = StartTime;
-		InteractiveTimeEnd = EndTime;
-	}
-
-	template<typename ChannelType>
-	void FMovieSceneInterpolatingPointsDrawTask<ChannelType>::RefineVisibleCurvePointsSynchronous(const TSharedRef<FMovieSceneCachedCurve<ChannelType>>& CachedCurve, const ChannelType& Channel)
-	{
-		const float InterpTimes[] = { 0.25f, 0.5f, 0.6f };
-
-		for (int32 Index = 0; Index < InterpolatingPoints.Num() - 1; ++Index)
-		{
-			const FVector2D& Lower = InterpolatingPoints[Index];
-			const FVector2D& Upper = InterpolatingPoints[Index + 1];
-
-			if ((Upper.X - Lower.X) >= TimeThreshold)
-			{
-				bool bSegmentIsLinear = true;
-
-				FVector2D Evaluated[UE_ARRAY_COUNT(InterpTimes)] = {};
-
-				for (int32 InterpIndex = 0; InterpIndex < UE_ARRAY_COUNT(InterpTimes); ++InterpIndex)
-				{
-					double& EvalTime = Evaluated[InterpIndex].X;
-
-					EvalTime = FMath::Lerp(Lower.X, Upper.X, InterpTimes[InterpIndex]);
-
-					CurveValueType Value = 0.0;
-					Channel.Evaluate(EvalTime * TickResolution, Value);
-
-					const CurveValueType LinearValue = FMath::Lerp(Lower.Y, Upper.Y, InterpTimes[InterpIndex]);
-					if (bSegmentIsLinear)
-					{
-						bSegmentIsLinear = FMath::IsNearlyEqual(Value, LinearValue, ValueThreshold);
-					}
-
-					Evaluated[InterpIndex].Y = Value;
-				}
-
-				if (!bSegmentIsLinear)
-				{
-					// Add the point
-					InterpolatingPoints.Insert(Evaluated, UE_ARRAY_COUNT(Evaluated), Index + 1);
-					--Index;
-				}
-			}
-		}
-	}
-
 	template<typename ChannelType>
 	void FMovieSceneInterpolatingPointsDrawTask<ChannelType>::InvokeCallback()
 	{
diff --git a/Engine/Source/Editor/MovieSceneTools/Private/Cache/MovieSceneInterpolatingPointsDrawTask.h b/Engine/Source/Editor/MovieSceneTools/Private/Cache/MovieSceneInterpolatingPointsDrawTask.h
--- a/Engine/Source/Editor/MovieSceneTools/Private/Cache/MovieSceneInterpolatingPointsDrawTask.h
+++ b/Engine/Source/Editor/MovieSceneTools/Private/Cache/MovieSceneInterpolatingPointsDrawTask.h
@@ -29,8 +29,6 @@ namespace UE::MovieSceneTools
 	public:
 		FMovieSceneInterpolatingPointsDrawTask(
 			const TSharedRef<FMovieSceneCachedCurve<ChannelType>>& CachedCurve,
-			const ChannelType& MovieSceneChannel,
-			const bool bInteractive,
 			const TFunction<void(TArray<FVector2D> /** OutInterpolatedPoints */, TArray<int32> /** OutKeyOffsets */)>& InCallback);
 
 		//~ Begin IMovieSceneInterpolatingPointsDrawTask interface
@@ -40,33 +38,8 @@ namespace UE::MovieSceneTools
 		//~ End IMovieSceneInterpolatingPointsDrawTask interface
 
 	private:
-		/** Describs wich part of the curve is interpolated in RefineFullRangeInterpolatingPoints */
-		enum class ECurvePainterInterpolationRange : uint8
-		{
-			/** Draws the full range of interpolating points */
-			FullRange,
-
-			/** Draws from the start to the already interactively interpolated range */
-			StartToInteractive,
-
-			/** Draws from the already interactively interpolated range to the end */
-			InteractiveToEnd
-		};
-
-		/** 
-		 * Refines the full range of the interpolating points. 
-		 * Internal logic to allow to skip over the interactively painted area if this is an interactive task.
-		 */
-		void RefineFullRangeInterpolatingPointsInternal(const ECurvePainterInterpolationRange InterpolationRange);
-
-		/** 
-		 * Utility to darw the visible part of the curve synchronously. Useful for interactive edits.
-		 * Note this uses the channel, this is faster than using the piecewise curve in the async counterpart.
-		 */
-		void DrawVisibleCurvePointsSynchronous(const TSharedRef<FMovieSceneCachedCurve<ChannelType>>& CachedCurve, const ChannelType& Channel);
-
-		/** Utility to refine the visible part of the curve synchronously. Useful for interactive edits. */
-		void RefineVisibleCurvePointsSynchronous(const TSharedRef<FMovieSceneCachedCurve<ChannelType>>& CachedCurve, const ChannelType& Channel);
+		/** Refines the full range of the interpolating points. */
+		void RefineFullRangeInterpolatingPointsInternal();
 
 		/** Calls back to the cached curve, forwarding interpolated points */
 		void InvokeCallback();
@@ -95,12 +68,6 @@ namespace UE::MovieSceneTools
 		/** The (possibly not yet fully interpolated) interpolating points */
 		TArray<FVector2D> InterpolatingPoints;
 
-		/** The first frame that was painted interactively. Only set if this is an interactive task */
-		TOptional<double> InteractiveTimeStart;
-
-		/** The last frame that was painted interactively. Only set if this is an interactive task */
-		TOptional<double> InteractiveTimeEnd;
-
 		/** True while work is performed in this task. Useful to prevent from concurrency by design. */
 		std::atomic<bool> bWorking = false;
 
