Anim In Engine: When merging sections no longer see if it's in range, just extend it so we can match up how it works without the merge
#trivial

diff --git a/Engine/Source/Editor/MovieSceneTools/Private/MovieSceneToolHelpers.cpp b/Engine/Source/Editor/MovieSceneTools/Private/MovieSceneToolHelpers.cpp
--- a/Engine/Source/Editor/MovieSceneTools/Private/MovieSceneToolHelpers.cpp
+++ b/Engine/Source/Editor/MovieSceneTools/Private/MovieSceneToolHelpers.cpp
@@ -5514,13 +5514,19 @@ bool MovieSceneToolHelpers::CollapseSection(TSharedPtr<ISequencer>& SequencerPtr
 		UMovieSceneSection* BaseSection = nullptr;
 		
 		FScopedTransaction Transaction(NSLOCTEXT("MovieSceneTools", "CollapseAllSections", "Collapse All Sections"));
+		bool bTrackModified = false;
 		for (int32 SectionIndex = Sections.Num() - 1; SectionIndex > 0; --SectionIndex)
 		{
 			BaseSection = Sections[SectionIndex - 1];
 			UMovieSceneSection* Section = Sections[SectionIndex];
 			if (BaseSection && Section && BaseSection->GetBlendType().IsValid() && Section->GetBlendType().IsValid())
 			{
-				OwnerTrack->Modify();
+				if (bTrackModified == false)
+				{
+					OwnerTrack->Modify();
+					bTrackModified = true;
+				}
+				BaseSection->Modify();
 				TArray<UMovieSceneSection*> TrackSections = OwnerTrack->GetAllSections();
 				if (Section->IsActive())//active sections merge them
 				{
diff --git a/Engine/Source/Editor/MovieSceneTools/Public/MovieSceneToolHelpers.h b/Engine/Source/Editor/MovieSceneTools/Public/MovieSceneToolHelpers.h
--- a/Engine/Source/Editor/MovieSceneTools/Public/MovieSceneToolHelpers.h
+++ b/Engine/Source/Editor/MovieSceneTools/Public/MovieSceneToolHelpers.h
@@ -941,9 +941,19 @@ public:
 				for (int32 FrameIndex = 0; FrameIndex < KeyTimes.Num(); ++FrameIndex)
 				{
 					const FFrameNumber& Frame = KeyTimes[FrameIndex];
-					if (Sections[0]->GetRange().Contains(Frame) == false)  //frame is outside base range so skip
+					if (Sections[0]->GetRange().Contains(Frame) == false)  //frame is outside base range so extend it
 					{
-						continue;
+						if (Sections[0]->HasEndFrame() && Sections[0]->GetExclusiveEndFrame() <= Frame)
+						{
+							if (Sections[0]->GetExclusiveEndFrame() != Frame)
+							{
+								Sections[0]->SetEndFrame(Frame);
+							}
+						}
+						else
+						{
+							Sections[0]->SetStartFrame(Frame);
+						}
 					}
 					const FFrameTime FrameTime(Frame);
 					int32 KeyIndex = Channel->GetData().GetIndex(Handles[FrameIndex]);
