When importing cameras apply uniform scale modifier to distance based parameters (Manual Focus Distance).

diff --git a/Engine/Source/Editor/MovieSceneTools/Private/MovieSceneToolHelpers.cpp b/Engine/Source/Editor/MovieSceneTools/Private/MovieSceneToolHelpers.cpp
--- a/Engine/Source/Editor/MovieSceneTools/Private/MovieSceneToolHelpers.cpp
+++ b/Engine/Source/Editor/MovieSceneTools/Private/MovieSceneToolHelpers.cpp
@@ -1304,6 +1304,22 @@ bool ImportFBXPropertyToPropertyTrack(UMovieScene* MovieScene, const FGuid Prope
 	return true;
 }
 
+
+// Static function to handle applying uniform scale data to keys inside of camera imports that are distance based.
+void ApplyUniformScaleToCameraKeys(const FString& AnimatedPropertyName, FRichCurve& Source)
+{
+	const UMovieSceneUserImportFBXSettings* ImportFBXSettings = GetDefault<UMovieSceneUserImportFBXSettings>();
+
+	// This is the only property right now.
+	if (AnimatedPropertyName == "FocusDistance")
+	{
+		for (int i=0; i<Source.Keys.Num(); i++)
+		{
+			Source.Keys[i].Value = Source.Keys[i].Value * ImportFBXSettings->ImportUniformScale;
+		}
+	}
+};
+
 bool ImportFBXProperty(FString NodeName, FString AnimatedPropertyName, FGuid ObjectBinding, UnFbx::FFbxCurvesAPI& CurveAPI, UMovieSceneSequence* InSequence, IMovieScenePlayer* Player, FMovieSceneSequenceIDRef TemplateID)
 {
 	UMovieScene* MovieScene = InSequence->GetMovieScene();
@@ -1314,6 +1330,9 @@ bool ImportFBXProperty(FString NodeName, FString AnimatedPropertyName, FGuid Obj
 	const bool bNegative = false;
 	CurveAPI.GetCurveDataForSequencer(NodeName, AnimatedPropertyName, ChannelIndex, CompositeIndex, Source, bNegative);
 
+	// Modify distance based camera curve info to account for Uniform Scale.
+	ApplyUniformScaleToCameraKeys(AnimatedPropertyName, Source);
+
 	// First, see if any of the custom importers can import this named property
 	if (FMovieSceneToolsModule::Get().ImportAnimatedProperty(AnimatedPropertyName, Source, ObjectBinding, MovieScene))
 	{
