Curve Editor: Fix warning when packaging projects that contain a level sequence

#rnx

diff --git a/Engine/Source/Editor/MovieSceneTools/Private/Cache/MovieSceneCachedCurve.h b/Engine/Source/Editor/MovieSceneTools/Private/Cache/MovieSceneCachedCurve.h
--- a/Engine/Source/Editor/MovieSceneTools/Private/Cache/MovieSceneCachedCurve.h
+++ b/Engine/Source/Editor/MovieSceneTools/Private/Cache/MovieSceneCachedCurve.h
@@ -16,13 +16,6 @@ namespace UE::MovieSceneTools
 {
 	template <typename ChannelType> struct FMovieSceneUpdateCachedCurveData;
 
-	int32 GZoomInterpolatesCurvePointsAsync = 0;
-	static FAutoConsoleVariableRef CZoomInterpolatesCurvePointsAsync(
-		TEXT("MovieSceneTools.ZoomInterpolatesCurvePointsAsync"),
-		GZoomInterpolatesCurvePointsAsync,
-		TEXT("When set to true, zoom is smoother but curve lines will not update instantaneously"),
-		ECVF_Default);
-
 	/** Flags defining how the cache changed when it was last updated */
 	enum class EMovieSceneCurveCacheChangeFlags : uint32
 	{
@@ -541,7 +534,8 @@ namespace UE::MovieSceneTools
 		// Only draw the full range interpolating points when the curve was edited or when zooming
 		const bool bCurveIsBeingEdited = EnumHasAnyFlags(Flags.load(), EMovieSceneCurveCacheChangeFlags::ChangedCurveData);
 		const bool bZooming = EnumHasAnyFlags(Flags.load(), EMovieSceneCurveCacheChangeFlags::ChangedSize);
-		if (!bCurveIsBeingEdited && !bZooming)
+		const bool bInteractive = bCurveIsBeingEdited || bZooming;
+		if (!bInteractive)
 		{
 			return;
 		}
@@ -571,37 +565,31 @@ namespace UE::MovieSceneTools
 			return;
 		}
 
-		const bool bDeferredZooming = GZoomInterpolatesCurvePointsAsync > 0 ? bZooming : false;
-		const bool bInteractive = bCurveIsBeingEdited || !bDeferredZooming;
-
 		// Draw interactive changes to the draw params directly
-		if (bInteractive)
-		{
-			const double StartTimeSeconds = ScreenSpace.GetInputMin();
-			const double EndTimeSeconds = ScreenSpace.GetInputMax();
-			const double TimeThreshold = GetTimeThreshold();
-			const double ValueThreshold = GetValueThreshold();
+		const double StartTimeSeconds = ScreenSpace.GetInputMin();
+		const double EndTimeSeconds = ScreenSpace.GetInputMax();
+		const double TimeThreshold = GetTimeThreshold();
+		const double ValueThreshold = GetValueThreshold();
 
-			TArray<TTuple<double, double>> InteractiveInterpolatingPoints;
-			Channel.PopulateCurvePoints(StartTimeSeconds, EndTimeSeconds, TimeThreshold, ValueThreshold, TickResolution, InteractiveInterpolatingPoints);
+		TArray<TTuple<double, double>> InteractiveInterpolatingPoints;
+		Channel.PopulateCurvePoints(StartTimeSeconds, EndTimeSeconds, TimeThreshold, ValueThreshold, TickResolution, InteractiveInterpolatingPoints);
 
-			CachedDrawParams.InterpolatingPoints.Reset();
-			CachedDrawParams.InterpolatingPoints.Reserve(InteractiveInterpolatingPoints.Num());
+		CachedDrawParams.InterpolatingPoints.Reset();
+		CachedDrawParams.InterpolatingPoints.Reserve(InteractiveInterpolatingPoints.Num());
 
-			// Convert the interpolating points to screen space
-			for (const TTuple<double, double>& Point : InteractiveInterpolatingPoints)
-			{
-				CachedDrawParams.InterpolatingPoints.Emplace(
-					ScreenSpace.SecondsToScreen(Point.Get<0>() + InputDisplayOffset),
-					ScreenSpace.ValueToScreen(Point.Get<1>())
-				);
-			}
+		// Convert the interpolating points to screen space
+		for (const TTuple<double, double>& Point : InteractiveInterpolatingPoints)
+		{
+			CachedDrawParams.InterpolatingPoints.Emplace(
+				ScreenSpace.SecondsToScreen(Point.Get<0>() + InputDisplayOffset),
+				ScreenSpace.ValueToScreen(Point.Get<1>())
+			);
+		}
 
-			DrawKeys();
-			ApplyDrawParams();
+		DrawKeys();
+		ApplyDrawParams();
 
-			Flags = EMovieSceneCurveCacheChangeFlags::None;
-		}
+		Flags = EMovieSceneCurveCacheChangeFlags::None;
 
 		// Create a draw task to draw interpolating points for the full range of the curve
 		const TSharedRef<FMovieSceneInterpolatingPointsDrawTask<ChannelType>> Task = MakeShared<FMovieSceneInterpolatingPointsDrawTask<ChannelType>>(
