Sequencer LUF: Follow up to 42488876 to check for valid metadata

#rnx

diff --git a/Engine/Source/Editor/MovieSceneTools/Private/Channels/BoolChannelKeyProxy.cpp b/Engine/Source/Editor/MovieSceneTools/Private/Channels/BoolChannelKeyProxy.cpp
--- a/Engine/Source/Editor/MovieSceneTools/Private/Channels/BoolChannelKeyProxy.cpp
+++ b/Engine/Source/Editor/MovieSceneTools/Private/Channels/BoolChannelKeyProxy.cpp
@@ -20,7 +20,7 @@ void UBoolChannelKeyProxy::PostEditChangeProperty(FPropertyChangedEvent& Propert
 {
 	bool bNewValue = bValue;
 
-	if (ChannelHandle.GetMetaData()->bInvertValue)
+	if (ChannelHandle.GetMetaData() != nullptr && ChannelHandle.GetMetaData()->bInvertValue)
 	{
 		bNewValue = !bNewValue;
 	}
@@ -34,7 +34,7 @@ void UBoolChannelKeyProxy::UpdateValuesFromRawData()
 {
 	RefreshCurrentValue(ChannelHandle, KeyHandle, bValue, Time);
 
-	if (ChannelHandle.GetMetaData()->bInvertValue)
+	if (ChannelHandle.GetMetaData() != nullptr && ChannelHandle.GetMetaData()->bInvertValue)
 	{
 		bValue = !bValue;
 	}
diff --git a/Engine/Source/Editor/MovieSceneTools/Private/Channels/BuiltInChannelEditors.cpp b/Engine/Source/Editor/MovieSceneTools/Private/Channels/BuiltInChannelEditors.cpp
--- a/Engine/Source/Editor/MovieSceneTools/Private/Channels/BuiltInChannelEditors.cpp
+++ b/Engine/Source/Editor/MovieSceneTools/Private/Channels/BuiltInChannelEditors.cpp
@@ -2210,7 +2210,7 @@ void ExtendKeyMenu(FMenuBuilder& OuterMenuBuilder, TSharedPtr<FExtender> MenuExt
 
 TUniquePtr<FCurveModel> CreateCurveEditorModel(const TMovieSceneChannelHandle<FMovieSceneFloatChannel>& FloatChannel, const UE::Sequencer::FCreateCurveEditorModelParams& Params)
 {
-	if (FloatChannel.GetMetaData()->bInvertValue)
+	if (FloatChannel.GetMetaData() != nullptr && FloatChannel.GetMetaData()->bInvertValue)
 	{
 		return MakeUnique<UE::CurveEditor::TInvertedCurveModel<FFloatChannelCurveModel>>(FloatChannel, Params.OwningSection, Params.Sequencer);
 	}
@@ -2219,7 +2219,7 @@ TUniquePtr<FCurveModel> CreateCurveEditorModel(const TMovieSceneChannelHandle<FM
 
 TUniquePtr<FCurveModel> CreateCurveEditorModel(const TMovieSceneChannelHandle<FMovieSceneDoubleChannel>& DoubleChannel, const UE::Sequencer::FCreateCurveEditorModelParams& Params)
 {
-	if (DoubleChannel.GetMetaData()->bInvertValue)
+	if (DoubleChannel.GetMetaData() != nullptr && DoubleChannel.GetMetaData()->bInvertValue)
 	{
 		return MakeUnique<UE::CurveEditor::TInvertedCurveModel<FDoubleChannelCurveModel>>(DoubleChannel, Params.OwningSection, Params.Sequencer);
 	}
@@ -2228,7 +2228,7 @@ TUniquePtr<FCurveModel> CreateCurveEditorModel(const TMovieSceneChannelHandle<FM
 
 TUniquePtr<FCurveModel> CreateCurveEditorModel(const TMovieSceneChannelHandle<FMovieSceneIntegerChannel>& IntegerChannel, const UE::Sequencer::FCreateCurveEditorModelParams& Params)
 {
-	if (IntegerChannel.GetMetaData()->bInvertValue)
+	if (IntegerChannel.GetMetaData() != nullptr && IntegerChannel.GetMetaData()->bInvertValue)
 	{
 		return MakeUnique<UE::CurveEditor::TInvertedCurveModel<FIntegerChannelCurveModel>>(IntegerChannel, Params.OwningSection, Params.Sequencer);
 	}
@@ -2237,7 +2237,7 @@ TUniquePtr<FCurveModel> CreateCurveEditorModel(const TMovieSceneChannelHandle<FM
 
 TUniquePtr<FCurveModel> CreateCurveEditorModel(const TMovieSceneChannelHandle<FMovieSceneBoolChannel>& BoolChannel, const UE::Sequencer::FCreateCurveEditorModelParams& Params)
 {
-	if (BoolChannel.GetMetaData()->bInvertValue)
+	if (BoolChannel.GetMetaData() != nullptr && BoolChannel.GetMetaData()->bInvertValue)
 	{
 		return MakeUnique<UE::CurveEditor::TInvertedCurveModel<FBoolChannelCurveModel>>(BoolChannel, Params.OwningSection, Params.Sequencer);
 	}
@@ -2246,7 +2246,7 @@ TUniquePtr<FCurveModel> CreateCurveEditorModel(const TMovieSceneChannelHandle<FM
 
 TUniquePtr<FCurveModel> CreateCurveEditorModel(const TMovieSceneChannelHandle<FMovieSceneByteChannel>& ByteChannel, const UE::Sequencer::FCreateCurveEditorModelParams& Params)
 {
-	if (ByteChannel.GetMetaData()->bInvertValue)
+	if (ByteChannel.GetMetaData() != nullptr && ByteChannel.GetMetaData()->bInvertValue)
 	{
 		return MakeUnique<UE::CurveEditor::TInvertedCurveModel<FByteChannelCurveModel>>(ByteChannel, Params.OwningSection, Params.Sequencer);
 	}
@@ -2255,7 +2255,7 @@ TUniquePtr<FCurveModel> CreateCurveEditorModel(const TMovieSceneChannelHandle<FM
 
 TUniquePtr<FCurveModel> CreateCurveEditorModel(const TMovieSceneChannelHandle<FMovieSceneEventChannel>& EventChannel, const UE::Sequencer::FCreateCurveEditorModelParams& Params)
 {
-	if (EventChannel.GetMetaData()->bInvertValue)
+	if (EventChannel.GetMetaData() != nullptr && EventChannel.GetMetaData()->bInvertValue)
 	{
 		return MakeUnique<UE::CurveEditor::TInvertedCurveModel<FEventChannelCurveModel>>(EventChannel, Params.OwningSection, Params.Sequencer);
 	}
@@ -2264,7 +2264,7 @@ TUniquePtr<FCurveModel> CreateCurveEditorModel(const TMovieSceneChannelHandle<FM
 
 TUniquePtr<FCurveModel> CreateCurveEditorModel(const TMovieSceneChannelHandle<FMovieSceneTimeWarpChannel>& TimeWarpChannel, const UE::Sequencer::FCreateCurveEditorModelParams& Params)
 {
-	if (TimeWarpChannel.GetMetaData()->bInvertValue)
+	if (TimeWarpChannel.GetMetaData() != nullptr && TimeWarpChannel.GetMetaData()->bInvertValue)
 	{
 		return MakeUnique<UE::CurveEditor::TInvertedCurveModel<FTimeWarpChannelCurveModel>>(TimeWarpChannel, Params.OwningSection, Params.OwningObject, Params.Sequencer);
 	}
diff --git a/Engine/Source/Editor/MovieSceneTools/Private/Channels/DoubleChannelKeyProxy.cpp b/Engine/Source/Editor/MovieSceneTools/Private/Channels/DoubleChannelKeyProxy.cpp
--- a/Engine/Source/Editor/MovieSceneTools/Private/Channels/DoubleChannelKeyProxy.cpp
+++ b/Engine/Source/Editor/MovieSceneTools/Private/Channels/DoubleChannelKeyProxy.cpp
@@ -19,7 +19,7 @@ void UDoubleChannelKeyProxy::PostEditChangeProperty(FPropertyChangedEvent& Prope
 {
 	FMovieSceneDoubleValue NewValue = Value;
 
-	if (ChannelHandle.GetMetaData()->bInvertValue)
+	if (ChannelHandle.GetMetaData() != nullptr && ChannelHandle.GetMetaData()->bInvertValue)
 	{
 		NewValue.Value = -NewValue.Value;
 	}
@@ -33,7 +33,7 @@ void UDoubleChannelKeyProxy::UpdateValuesFromRawData()
 {
 	RefreshCurrentValue(ChannelHandle, KeyHandle, Value, Time);
 
-	if (ChannelHandle.GetMetaData()->bInvertValue)
+	if (ChannelHandle.GetMetaData() != nullptr && ChannelHandle.GetMetaData()->bInvertValue)
 	{
 		Value.Value = -Value.Value;
 	}
diff --git a/Engine/Source/Editor/MovieSceneTools/Private/Channels/FloatChannelKeyProxy.cpp b/Engine/Source/Editor/MovieSceneTools/Private/Channels/FloatChannelKeyProxy.cpp
--- a/Engine/Source/Editor/MovieSceneTools/Private/Channels/FloatChannelKeyProxy.cpp
+++ b/Engine/Source/Editor/MovieSceneTools/Private/Channels/FloatChannelKeyProxy.cpp
@@ -19,7 +19,7 @@ void UFloatChannelKeyProxy::PostEditChangeProperty(FPropertyChangedEvent& Proper
 {
 	FMovieSceneFloatValue NewValue = Value;
 
-	if (ChannelHandle.GetMetaData()->bInvertValue)
+	if (ChannelHandle.GetMetaData() != nullptr && ChannelHandle.GetMetaData()->bInvertValue)
 	{
 		NewValue.Value = -NewValue.Value;
 	}
@@ -33,7 +33,7 @@ void UFloatChannelKeyProxy::UpdateValuesFromRawData()
 {
 	RefreshCurrentValue(ChannelHandle, KeyHandle, Value, Time);
 
-	if (ChannelHandle.GetMetaData()->bInvertValue)
+	if (ChannelHandle.GetMetaData() != nullptr && ChannelHandle.GetMetaData()->bInvertValue)
 	{
 		Value.Value = -Value.Value;
 	}
diff --git a/Engine/Source/Editor/MovieSceneTools/Private/Channels/IntegerChannelKeyProxy.cpp b/Engine/Source/Editor/MovieSceneTools/Private/Channels/IntegerChannelKeyProxy.cpp
--- a/Engine/Source/Editor/MovieSceneTools/Private/Channels/IntegerChannelKeyProxy.cpp
+++ b/Engine/Source/Editor/MovieSceneTools/Private/Channels/IntegerChannelKeyProxy.cpp
@@ -20,7 +20,7 @@ void UIntegerChannelKeyProxy::PostEditChangeProperty(FPropertyChangedEvent& Prop
 {
 	int32 NewValue = Value;
 
-	if (ChannelHandle.GetMetaData()->bInvertValue)
+	if (ChannelHandle.GetMetaData() != nullptr && ChannelHandle.GetMetaData()->bInvertValue)
 	{
 		NewValue = -NewValue;
 	}
@@ -34,7 +34,7 @@ void UIntegerChannelKeyProxy::UpdateValuesFromRawData()
 {
 	RefreshCurrentValue(ChannelHandle, KeyHandle, Value, Time);
 
-	if (ChannelHandle.GetMetaData()->bInvertValue)
+	if (ChannelHandle.GetMetaData() != nullptr && ChannelHandle.GetMetaData()->bInvertValue)
 	{
 		Value = -Value;
 	}
diff --git a/Engine/Source/Editor/MovieSceneTools/Private/CurveKeyEditors/SequencerKeyEditor.h b/Engine/Source/Editor/MovieSceneTools/Private/CurveKeyEditors/SequencerKeyEditor.h
--- a/Engine/Source/Editor/MovieSceneTools/Private/CurveKeyEditors/SequencerKeyEditor.h
+++ b/Engine/Source/Editor/MovieSceneTools/Private/CurveKeyEditors/SequencerKeyEditor.h
@@ -202,7 +202,7 @@ struct TSequencerKeyEditor
 		const FKeySelection& KeySelection = Sequencer->GetViewModel()->GetSelection()->KeySelection;
 
 		ValueType NewValue = InValue;
-		if (ChannelMetaData->bInvertValue)
+		if (ChannelMetaData && ChannelMetaData->bInvertValue)
 		{
 			InvertValue(NewValue);
 		}
