Constraints & scene components: make sure that the transform stored in OnPreTransformChanged is expressed in the constraints space if needed

#rnx

diff --git a/Engine/Source/Editor/MovieSceneTools/Private/TrackEditors/TransformTrackEditor.cpp b/Engine/Source/Editor/MovieSceneTools/Private/TrackEditors/TransformTrackEditor.cpp
--- a/Engine/Source/Editor/MovieSceneTools/Private/TrackEditors/TransformTrackEditor.cpp
+++ b/Engine/Source/Editor/MovieSceneTools/Private/TrackEditors/TransformTrackEditor.cpp
@@ -296,6 +296,20 @@ void F3DTransformTrackEditor::OnPreTransformChanged( UObject& InObject )
 			// Cache off the existing transform so we can detect which components have changed
 			// and keys only when something has changed
 			FTransformData Transform( SceneComponentThatChanged );
+
+			{
+				// switch to constraint space if needed so that the optional LastTransform passed to F3DTransformTrackEditor::GetTransformKeys
+				// is expressed in the right space to compare values
+				const FTransform Relative(Transform.Rotation, Transform.Translation, Transform.Scale);
+				FTransform ConstraintRelative(Relative);
+				UE::TransformConstraintUtil::UpdateTransformBasedOnConstraint(ConstraintRelative, SceneComponentThatChanged);
+				if (!ConstraintRelative.Equals(Relative, UE_SMALL_NUMBER))
+				{
+					Transform.Translation = ConstraintRelative.GetTranslation();
+					Transform.Rotation = ConstraintRelative.Rotator();
+					Transform.Scale = ConstraintRelative.GetScale3D();
+				}
+			}
 			
 			ObjectToExistingTransform.Add(&InObject, Transform);
 			
