Anim In Engine:  Force LOD to 1 when ticking skel meshes when baking anim sequences. Sequencer may restore state which can change the lod during the bake process.

diff --git a/Engine/Source/Editor/MovieSceneTools/Private/MovieSceneToolHelpers.cpp b/Engine/Source/Editor/MovieSceneTools/Private/MovieSceneToolHelpers.cpp
--- a/Engine/Source/Editor/MovieSceneTools/Private/MovieSceneToolHelpers.cpp
+++ b/Engine/Source/Editor/MovieSceneTools/Private/MovieSceneToolHelpers.cpp
@@ -4165,6 +4165,7 @@ static void TickFrame(const FFrameNumber& FrameNumber,float DeltaTime, UMovieSce
 	// Update space bases so new animation position has an effect.
 	for (USkeletalMeshComponent* SkelMeshComp : SkelMeshComps)
 	{
+		SkelMeshComp->SetForcedLOD(1); //we need to do this every tick since a restore state may happen while evaluating and bummp us out ot a different LOD
 		SkelMeshComp->TickAnimation(DeltaTime, false);
 
 		SkelMeshComp->RefreshBoneTransforms();
