Sequencer x Curve Editor: Allow setting custom colours for tracks that support names. The most important one is property tracks.

diff --git a/Engine/Source/Editor/MovieSceneTools/Private/Channels/ChannelCurveModel.cpp b/Engine/Source/Editor/MovieSceneTools/Private/Channels/ChannelCurveModel.cpp
--- a/Engine/Source/Editor/MovieSceneTools/Private/Channels/ChannelCurveModel.cpp
+++ b/Engine/Source/Editor/MovieSceneTools/Private/Channels/ChannelCurveModel.cpp
@@ -25,6 +25,7 @@
 #include "Misc/FrameTime.h"
 #include "Misc/Optional.h"
 #include "MovieScene.h"
+#include "MovieSceneNameableTrack.h"
 #include "MovieSceneSection.h"
 #include "Styling/AppStyle.h"
 #include "Styling/ISlateStyle.h"
@@ -599,16 +600,30 @@ void FChannelCurveModel<ChannelType, ChannelValue, KeyType>::GetCurveColorObject
 	{
 		*OutObject = Section->GetImplicitObjectOwner();
 
-		if (const FMovieSceneChannelMetaData* MetaData = ChannelHandle.GetMetaData())
+		const FMovieSceneChannelMetaData* MetaData = ChannelHandle.GetMetaData();
+		const bool bHasMetaData = MetaData && !MetaData->Group.IsEmpty() && !MetaData->DisplayText.IsEmpty();
+		if (bHasMetaData)
 		{
-			OutName = FString::Printf(TEXT( "%s.%s" ), *MetaData->Group.ToString(), *MetaData->DisplayText.ToString());
-			return;
+			OutName = FString::Printf(TEXT("%s.%s"), *MetaData->Group.ToString(), *MetaData->DisplayText.ToString());
+		}
+		else if (!GetIntentionName().IsEmpty())
+		{
+			OutName = GetIntentionName();
+		}
+		else if (UMovieSceneNameableTrack* NamedTrack = Section->GetTypedOuter<UMovieSceneNameableTrack>())
+		{
+			OutName = NamedTrack->GetDisplayName().ToString();
 		}
-		OutName = GetIntentionName();
-		return;
+		else
+		{
+			OutName = FString();
+		}
+	}
+	else
+	{
+		// Just call base if it doesn't work
+		FCurveModel::GetCurveColorObjectAndName(OutObject, OutName);
 	}
-	// Just call base if it doesn't work
-	FCurveModel::GetCurveColorObjectAndName(OutObject, OutName);
 }
 
 
