Shortening this button text so it fits.

diff --git a/Engine/Source/Editor/MovieSceneTools/Private/TrackEditors/CommonAnimationTrackEditor.cpp b/Engine/Source/Editor/MovieSceneTools/Private/TrackEditors/CommonAnimationTrackEditor.cpp
--- a/Engine/Source/Editor/MovieSceneTools/Private/TrackEditors/CommonAnimationTrackEditor.cpp
+++ b/Engine/Source/Editor/MovieSceneTools/Private/TrackEditors/CommonAnimationTrackEditor.cpp
@@ -250,7 +250,7 @@ void SAnimSequenceOptionsWindow::Construct(const FArguments& InArgs)
 		[
 			SNew(SButton)
 			.HAlign(HAlign_Center)
-			.Text(LOCTEXT("AnimExportOptionsWindow_Export", "Export To Animation Sequence"))
+			.Text(LOCTEXT("AnimExportOptionsWindow_Export", "Export"))
 			.OnClicked(this, &SAnimSequenceOptionsWindow::OnExport)
 		]
 	   + SUniformGridPanel::Slot(2, 0)
