Curve Editor: Fix an issue where in some cases curves do not span the entire viewport, improve stability of cached curve to be tolerant for curves where times, values, key draw infos and other arrays unexpectedly do not match up in size.

#rnx

diff --git a/Engine/Source/Editor/MovieSceneTools/Private/Cache/MovieSceneCachedCurve.h b/Engine/Source/Editor/MovieSceneTools/Private/Cache/MovieSceneCachedCurve.h
--- a/Engine/Source/Editor/MovieSceneTools/Private/Cache/MovieSceneCachedCurve.h
+++ b/Engine/Source/Editor/MovieSceneTools/Private/Cache/MovieSceneCachedCurve.h
@@ -103,9 +103,6 @@ namespace UE::MovieSceneTools
 		/** Updates the cached input display offset */
 		void UpdateInputDisplayOffset(const double NewInputDisplayOffset);
 
-		/** Updates cached indices of key in the times and values array */
-		void UpdateIndices(const FMovieSceneUpdateCachedCurveData<ChannelType>& UpdateData, const FFrameNumber& StartFrame, const FFrameNumber& EndFrame);
-
 		/** Updates the cached selection */
 		void UpdateSelection(const FMovieSceneUpdateCachedCurveData<ChannelType>& UpdateData);
 
@@ -113,7 +110,7 @@ namespace UE::MovieSceneTools
 		void UpdateTangentVisibility(const ECurveEditorTangentVisibility NewTangentVisibility);
 
 		/** Updates cached curve data such as keys, values, key handles and key attributes */
-		void UpdateCurveData(const FMovieSceneUpdateCachedCurveData<ChannelType>& UpdateData);
+		void UpdateCurveData(const FMovieSceneUpdateCachedCurveData<ChannelType>& UpdateData, const FFrameNumber& StartFrame, const FFrameNumber& EndFrame);
 
 		/** Updates the cached pre and post inifinity extrapolation */
 		void UpdatePrePostInfinityExtrapolation(const FMovieSceneUpdateCachedCurveData<ChannelType>& UpdateData, const FFrameNumber& StartFrame, const FFrameNumber& EndFrame);
@@ -266,9 +263,8 @@ namespace UE::MovieSceneTools
 		UpdateTickResolution(UpdateData.TickResolution);
 		UpdateTangentVisibility(UpdateData.CurveEditor.GetSettings()->GetTangentVisibility());
 		UpdateInputDisplayOffset(UpdateData.CurveModel.GetInputDisplayOffset());
-		UpdateIndices(UpdateData, StartFrame, EndFrame);
 		UpdateSelection(UpdateData);
-		UpdateCurveData(UpdateData);
+		UpdateCurveData(UpdateData, StartFrame, EndFrame);
 		UpdatePrePostInfinityExtrapolation(UpdateData, StartFrame, EndFrame);
 
 		// Update depending on above data
@@ -353,22 +349,6 @@ namespace UE::MovieSceneTools
 		}
 	}
 
-	template <typename ChannelType>
-	void FMovieSceneCachedCurve<ChannelType>::UpdateIndices(const FMovieSceneUpdateCachedCurveData<ChannelType>& UpdateData, const FFrameNumber& StartFrame, const FFrameNumber& EndFrame)
-	{
-		const int32 NewStartingIndex = Algo::LowerBound(UpdateData.Times, StartFrame);
-		const int32 NewEndingIndex = Algo::UpperBound(UpdateData.Times, EndFrame) - 1;
-
-		if (StartingIndex != NewStartingIndex ||
-			EndingIndex != NewEndingIndex)
-		{
-			StartingIndex = NewStartingIndex;
-			EndingIndex = NewEndingIndex;
-
-			Flags.store(Flags.load() | EMovieSceneCurveCacheChangeFlags::ChangedKeyIndices);
-		}
-	}
-
 	template <typename ChannelType>
 	void FMovieSceneCachedCurve<ChannelType>::UpdateSelection(const FMovieSceneUpdateCachedCurveData<ChannelType>& UpdateData)
 	{
@@ -397,7 +377,7 @@ namespace UE::MovieSceneTools
 	}
 
 	template <typename ChannelType>
-	void FMovieSceneCachedCurve<ChannelType>::UpdateCurveData(const FMovieSceneUpdateCachedCurveData<ChannelType>& UpdateData)
+	void FMovieSceneCachedCurve<ChannelType>::UpdateCurveData(const FMovieSceneUpdateCachedCurveData<ChannelType>& UpdateData, const FFrameNumber& StartFrame, const FFrameNumber& EndFrame)
 	{
 		// DefaultValue
 		const double NewDefaultValue = UpdateData.Channel.GetDefault().IsSet() ? UpdateData.Channel.GetDefault().GetValue() : 0.0;
@@ -448,6 +428,28 @@ namespace UE::MovieSceneTools
 
 			Flags.store(Flags.load() | EMovieSceneCurveCacheChangeFlags::ChangedCurveData);
 		}
+
+		// Indices
+		const int32 NewStartingIndex = UpdateData.Times.IsEmpty() ? INDEX_NONE : Algo::LowerBound(UpdateData.Times, StartFrame);
+		const int32 NewEndingIndex = [&UpdateData, &EndFrame, this]()
+			{
+				int32 UpperBoundIndex = UpdateData.Times.IsEmpty() ? INDEX_NONE : Algo::UpperBound(UpdateData.Times, EndFrame) - 1;
+				UpperBoundIndex = FMath::Min(UpperBoundIndex, Times.Num() - 1);
+				UpperBoundIndex = FMath::Min(UpperBoundIndex, Values.Num() - 1);
+				UpperBoundIndex = FMath::Min(UpperBoundIndex, KeyHandles.Num() - 1);
+				UpperBoundIndex = FMath::Min(UpperBoundIndex, KeyAttributes.Num() - 1);
+				
+				return UpperBoundIndex;
+			}();
+
+		if (StartingIndex != NewStartingIndex ||
+			EndingIndex != NewEndingIndex)
+		{
+			StartingIndex = NewStartingIndex;
+			EndingIndex = NewEndingIndex;
+
+			Flags.store(Flags.load() | EMovieSceneCurveCacheChangeFlags::ChangedKeyIndices);
+		}
 	}
 
 	template <typename ChannelType>
@@ -485,17 +487,20 @@ namespace UE::MovieSceneTools
 
 				// Tangent draw info
 				const FKeyAttributes* KeyAttributesPtr = nullptr;
-				if (TangentVisibility == ECurveEditorTangentVisibility::SelectedKeys)
+				if (KeyAttributes.IsValidIndex(DataIndex))
 				{
-					if (Selection.Contains(KeyHandle, ECurvePointType::Any))
+					if (TangentVisibility == ECurveEditorTangentVisibility::SelectedKeys)
+					{
+						if (Selection.Contains(KeyHandle, ECurvePointType::Any))
+						{
+							KeyAttributesPtr = &KeyAttributes[DataIndex];
+						}
+					}
+					else if (TangentVisibility == ECurveEditorTangentVisibility::AllTangents)
 					{
 						KeyAttributesPtr = &KeyAttributes[DataIndex];
 					}
 				}
-				else if (TangentVisibility == ECurveEditorTangentVisibility::AllTangents)
-				{
-					KeyAttributesPtr = &KeyAttributes[DataIndex];
-				}
 
 				TOptional<FKeyDrawInfo> ArriveTangentDrawInfo;
 				if (KeyAttributesPtr && KeyAttributesPtr->HasArriveTangent())
@@ -543,7 +548,7 @@ namespace UE::MovieSceneTools
 		// Optimization: If there's less than two keys there's nothing to interpolate. 
 		// FMovieSceneInterpolatingPointsDrawTask ensures no such tasks are created to avoid needless overhead. 
 		// Adopt this principle here.
-		if (Times.IsEmpty())
+		if (Times.IsEmpty() || Values.IsEmpty())
 		{
 			FullRangeInterpolation.KeyOffsets.Reset();
 			FullRangeInterpolation.Points.Reset();
@@ -551,7 +556,7 @@ namespace UE::MovieSceneTools
 
 			return;
 		}
-		else if (Times.Num() == 1)
+		else if (Times.Num() == 1 || Values.Num() == 1)
 		{
 			FullRangeInterpolation.KeyOffsets = { 0 };
 
@@ -633,6 +638,22 @@ namespace UE::MovieSceneTools
 			DrawPreInfinityInterpolatingPoints();
 			DrawPostInfinityInterpolatingPoints();
 		}
+		else if (!FullRangeInterpolation.KeyOffsets.IsValidIndex(StartingIndex) &&
+			!FullRangeInterpolation.KeyOffsets.IsValidIndex(EndingIndex))
+		{
+			const bool bHasKeyBefore = FullRangeInterpolation.KeyOffsets.IsValidIndex(StartingIndex - 1);
+			if (!bHasKeyBefore)
+			{
+				DrawPreInfinityInterpolatingPoints();
+			}
+
+			const bool bHasKeyAfter = FullRangeInterpolation.KeyOffsets.IsValidIndex(EndingIndex + 1);
+
+			if (!bHasKeyAfter)
+			{
+				DrawPostInfinityInterpolatingPoints();
+			}
+		}
 		else
 		{
 			const bool bHasKeyBefore = FullRangeInterpolation.KeyOffsets.IsValidIndex(StartingIndex - 1);
@@ -644,25 +665,29 @@ namespace UE::MovieSceneTools
 				DrawPreInfinityInterpolatingPoints();
 			}
 
-			// Overdraw to the previous key if possible
-			const int32 DataIndex = bHasKeyBefore ?
-				FullRangeInterpolation.KeyOffsets[StartingIndex - 1] :
-				FullRangeInterpolation.KeyOffsets[StartingIndex];
+			if (FullRangeInterpolation.KeyOffsets.IsValidIndex(StartingIndex) && 
+				FullRangeInterpolation.KeyOffsets.IsValidIndex(EndingIndex))
+			{
+				// Overdraw to the previous key if possible
+				const int32 DataIndex = bHasKeyBefore ?
+					FullRangeInterpolation.KeyOffsets[StartingIndex - 1] :
+					FullRangeInterpolation.KeyOffsets[StartingIndex];
 
-			// Overdraw to the next key if possible
-			const int32 InterpolatingPointsDataSize = bHasKeyAfter ?
-				FullRangeInterpolation.KeyOffsets[EndingIndex + 1] - DataIndex + 1 :
-				FullRangeInterpolation.KeyOffsets[EndingIndex] - DataIndex + 1;
+				// Overdraw to the next key if possible
+				const int32 InterpolatingPointsDataSize = bHasKeyAfter ?
+					FullRangeInterpolation.KeyOffsets[EndingIndex + 1] - DataIndex + 1 :
+					FullRangeInterpolation.KeyOffsets[EndingIndex] - DataIndex + 1;
 
-			// Add points in between
-			check(FullRangeInterpolation.Points.IsValidIndex(DataIndex) && FullRangeInterpolation.Points.Num() >= DataIndex + InterpolatingPointsDataSize);
+				// Add points in between
+				check(FullRangeInterpolation.Points.IsValidIndex(DataIndex) && FullRangeInterpolation.Points.Num() >= DataIndex + InterpolatingPointsDataSize);
 
-			const TArrayView<FVector2D> InterpolatingPointsView(FullRangeInterpolation.Points.GetData() + DataIndex, InterpolatingPointsDataSize);
-			for (const FVector2D& InterpolatingPoint : InterpolatingPointsView)
-			{
-				CachedDrawParams.InterpolatingPoints.Emplace(
-					ScreenSpace.SecondsToScreen(InterpolatingPoint.X + InputDisplayOffset),
-					ScreenSpace.ValueToScreen(InterpolatingPoint.Y));
+				const TArrayView<FVector2D> InterpolatingPointsView(FullRangeInterpolation.Points.GetData() + DataIndex, InterpolatingPointsDataSize);
+				for (const FVector2D& InterpolatingPoint : InterpolatingPointsView)
+				{
+					CachedDrawParams.InterpolatingPoints.Emplace(
+						ScreenSpace.SecondsToScreen(InterpolatingPoint.X + InputDisplayOffset),
+						ScreenSpace.ValueToScreen(InterpolatingPoint.Y));
+				}
 			}
 
 			// Add the upper bound of the visible space if there is no key after
@@ -683,6 +708,11 @@ namespace UE::MovieSceneTools
 
 		const FScopeLock Lock(&LockInterpolatingPoints);
 
+		if (FullRangeInterpolation.Points.IsEmpty())
+		{
+			return;
+		}
+
 		if (FullRangeInterpolation.PreInfinityExtrapolation == ERichCurveExtrapolation::RCCE_Cycle ||
 			FullRangeInterpolation.PreInfinityExtrapolation == ERichCurveExtrapolation::RCCE_CycleWithOffset ||
 			FullRangeInterpolation.PreInfinityExtrapolation == ERichCurveExtrapolation::RCCE_Oscillate)
@@ -703,7 +733,7 @@ namespace UE::MovieSceneTools
 					ScreenSpace.ValueToScreen(FullRangeInterpolation.Points[0].Y));
 
 				CachedDrawParams.InterpolatingPoints.Emplace(
-					ScreenSpace.SecondsToScreen(ScreenSpace.GetInputMax()),
+					ScreenSpace.SecondsToScreen(FullRangeInterpolation.Points[0].X),
 					ScreenSpace.ValueToScreen(FullRangeInterpolation.Points[0].Y));
 			}
 			else
@@ -757,10 +787,14 @@ namespace UE::MovieSceneTools
 
 		const FScopeLock Lock(&LockInterpolatingPoints);
 
-		if (FullRangeInterpolation.Points.Num() > 1 &&
-			(FullRangeInterpolation.PostInfinityExtrapolation == ERichCurveExtrapolation::RCCE_Cycle ||
-				FullRangeInterpolation.PostInfinityExtrapolation == ERichCurveExtrapolation::RCCE_CycleWithOffset ||
-				FullRangeInterpolation.PostInfinityExtrapolation == ERichCurveExtrapolation::RCCE_Oscillate))
+		if (FullRangeInterpolation.Points.IsEmpty())
+		{
+			return;
+		}
+
+		if (FullRangeInterpolation.PostInfinityExtrapolation == ERichCurveExtrapolation::RCCE_Cycle ||
+			FullRangeInterpolation.PostInfinityExtrapolation == ERichCurveExtrapolation::RCCE_CycleWithOffset ||
+			FullRangeInterpolation.PostInfinityExtrapolation == ERichCurveExtrapolation::RCCE_Oscillate)
 		{
 			const double InfinityOffset = ScreenSpace.GetInputMax() + InputDisplayOffset - FullRangeInterpolation.Points.Last().X;
 
@@ -775,7 +809,7 @@ namespace UE::MovieSceneTools
 			{
 				// Draw a single line in the odd case where there are many keys with a nearly zero duration.
 				CachedDrawParams.InterpolatingPoints.Emplace(
-					ScreenSpace.SecondsToScreen(ScreenSpace.GetInputMin()),
+					ScreenSpace.SecondsToScreen(FullRangeInterpolation.Points.Last().X),
 					ScreenSpace.ValueToScreen(FullRangeInterpolation.Points.Last().Y));
 
 				CachedDrawParams.InterpolatingPoints.Emplace(
@@ -826,7 +860,7 @@ namespace UE::MovieSceneTools
 	template <typename ChannelType>
 	bool FMovieSceneCachedCurve<ChannelType>::TryDrawPreInfinityExtentFast()
 	{
-		if (Values.IsEmpty())
+		if (Times.IsEmpty() || Values.IsEmpty())
 		{
 			CachedDrawParams.InterpolatingPoints.Emplace(
 				ScreenSpace.SecondsToScreen(ScreenSpace.GetInputMin()), 
@@ -834,7 +868,7 @@ namespace UE::MovieSceneTools
 
 			return true;
 		}
-		else if (Values.Num() == 1)
+		else if (Times.Num() == 1 || Values.Num() == 1)
 		{
 			const double SingleValue = Values[0].Value;
 			CachedDrawParams.InterpolatingPoints.Emplace(
@@ -846,14 +880,18 @@ namespace UE::MovieSceneTools
 		else
 		{
 			const bool bPreinfinityVisible = Times[0].Value > ScreenSpace.GetInputMin() * TickResolution;
+			const bool bLinearWithoutArriveTangent =
+				FullRangeInterpolation.PreInfinityExtrapolation == ERichCurveExtrapolation::RCCE_Linear &&
+				(KeyAttributes.IsEmpty() || !KeyAttributes[0].HasArriveTangent());
+
 			if (bPreinfinityVisible &&
 				(FullRangeInterpolation.PreInfinityExtrapolation == ERichCurveExtrapolation::RCCE_None ||
 				FullRangeInterpolation.PreInfinityExtrapolation == ERichCurveExtrapolation::RCCE_Constant ||
-				!KeyAttributes[0].HasArriveTangent()))
+				bLinearWithoutArriveTangent))
 			{
 				const double ExtentValue = Values[0].Value;
 				CachedDrawParams.InterpolatingPoints.Emplace(
-					ScreenSpace.SecondsToScreen(ScreenSpace.GetInputMin()), 
+					ScreenSpace.SecondsToScreen(ScreenSpace.GetInputMin()),
 					ScreenSpace.ValueToScreen(ExtentValue));
 
 				return true;
@@ -861,9 +899,9 @@ namespace UE::MovieSceneTools
 			else if (bPreinfinityVisible &&
 				FullRangeInterpolation.PreInfinityExtrapolation == ERichCurveExtrapolation::RCCE_Linear)
 			{
-				const double ExtentValue = Values[0].Value + KeyAttributes[0].GetArriveTangent() * (ScreenSpace.GetInputMin() - Times[0].Value);
+				const double ExtentValue = Values[0].Value + KeyAttributes[0].GetArriveTangent() * (ScreenSpace.GetInputMin() - Times[0] / TickResolution);
 				CachedDrawParams.InterpolatingPoints.Emplace(
-					ScreenSpace.SecondsToScreen(ScreenSpace.GetInputMin()), 
+					ScreenSpace.SecondsToScreen(ScreenSpace.GetInputMin()),
 					ScreenSpace.ValueToScreen(ExtentValue));
 
 				return true;
@@ -876,7 +914,7 @@ namespace UE::MovieSceneTools
 	template <typename ChannelType>
 	bool FMovieSceneCachedCurve<ChannelType>::TryDrawPostInfinityExtentFast()
 	{
-		if (Values.IsEmpty())
+		if (Times.IsEmpty() || Values.IsEmpty())
 		{
 			CachedDrawParams.InterpolatingPoints.Emplace(
 				ScreenSpace.SecondsToScreen(ScreenSpace.GetInputMax()), 
@@ -884,7 +922,7 @@ namespace UE::MovieSceneTools
 
 			return true;
 		}
-		else if (Values.Num() == 1)
+		else if (Times.Num() == 1 || Values.Num() == 1)
 		{
 			const double SingleValue = Values[0].Value;
 			CachedDrawParams.InterpolatingPoints.Emplace(
@@ -896,10 +934,14 @@ namespace UE::MovieSceneTools
 		else
 		{
 			const bool bPostinfinityVisible = Times.Last().Value < ScreenSpace.GetInputMax() * TickResolution;
+			const bool bLinearWithoutLeaveTangent =
+				FullRangeInterpolation.PostInfinityExtrapolation == ERichCurveExtrapolation::RCCE_Linear &&
+				(KeyAttributes.IsEmpty() || !KeyAttributes.Last().HasLeaveTangent());
+
 			if (bPostinfinityVisible &&
 				(FullRangeInterpolation.PostInfinityExtrapolation == ERichCurveExtrapolation::RCCE_None ||
 				FullRangeInterpolation.PostInfinityExtrapolation == ERichCurveExtrapolation::RCCE_Constant ||
-				!KeyAttributes.Last().HasLeaveTangent()))
+				bLinearWithoutLeaveTangent))
 			{
 				const double ExtentValue = Values.Last().Value;
 				CachedDrawParams.InterpolatingPoints.Emplace(
@@ -911,7 +953,7 @@ namespace UE::MovieSceneTools
 			else if (bPostinfinityVisible &&
 				FullRangeInterpolation.PostInfinityExtrapolation == ERichCurveExtrapolation::RCCE_Linear)
 			{
-				const double ExtentValue = Values.Last().Value + KeyAttributes.Last().GetLeaveTangent() * (ScreenSpace.GetInputMax() - Times.Last().Value);
+				const double ExtentValue = Values.Last().Value + KeyAttributes.Last().GetLeaveTangent() * (ScreenSpace.GetInputMax() - Times.Last() / TickResolution);
 				CachedDrawParams.InterpolatingPoints.Emplace(
 					ScreenSpace.SecondsToScreen(ScreenSpace.GetInputMax()),
 					ScreenSpace.ValueToScreen(ExtentValue));
@@ -926,6 +968,21 @@ namespace UE::MovieSceneTools
 	template <typename ChannelType>
 	void FMovieSceneCachedCurve<ChannelType>::DrawKeys()
 	{
+		if (!Times.IsValidIndex(StartingIndex) ||
+			!Times.IsValidIndex(EndingIndex) ||
+			!Values.IsValidIndex(StartingIndex) ||
+			!Values.IsValidIndex(EndingIndex) ||
+			!KeyHandles.IsValidIndex(StartingIndex) ||
+			!KeyHandles.IsValidIndex(EndingIndex) ||
+			!KeyAttributes.IsValidIndex(StartingIndex) ||
+			!KeyAttributes.IsValidIndex(EndingIndex) ||
+			!KeyDrawInfos.IsValidIndex(StartingIndex) ||
+			!KeyDrawInfos.IsValidIndex(EndingIndex))
+		{
+			CachedDrawParams.Points.Reset();
+			return;
+		}
+
 		const int32 DataSize = EndingIndex - StartingIndex + 1;
 		CachedDrawParams.Points.Reset(DataSize);
 
@@ -949,7 +1006,8 @@ namespace UE::MovieSceneTools
 
 			// Arrive Tangent
 			{
-				if (KeyDrawInfos[DataIndex].template Get<1>().IsSet())
+				if (KeyAttribute.HasArriveTangent() &&
+					KeyDrawInfos[DataIndex].template Get<1>().IsSet())
 				{
 					FCurvePointInfo ArriveTangentInfo(KeyHandle);
 
@@ -981,7 +1039,8 @@ namespace UE::MovieSceneTools
 
 			// Leave Tangent
 			{
-				if (KeyDrawInfos[DataIndex].template Get<2>().IsSet())
+				if (KeyAttribute.HasLeaveTangent() &&
+					KeyDrawInfos[DataIndex].template Get<2>().IsSet())
 				{
 					FCurvePointInfo LeaveTangentInfo(KeyHandle);
 
diff --git a/Engine/Source/Editor/MovieSceneTools/Private/Cache/MovieSceneInterpolatingPointsDrawTask.cpp b/Engine/Source/Editor/MovieSceneTools/Private/Cache/MovieSceneInterpolatingPointsDrawTask.cpp
--- a/Engine/Source/Editor/MovieSceneTools/Private/Cache/MovieSceneInterpolatingPointsDrawTask.cpp
+++ b/Engine/Source/Editor/MovieSceneTools/Private/Cache/MovieSceneInterpolatingPointsDrawTask.cpp
@@ -43,10 +43,12 @@ namespace UE::MovieSceneTools
 
 	template <typename ChannelType>
 	void FMovieSceneInterpolatingPointsDrawTask<ChannelType>::RefineFullRangeInterpolatingPoints()
-	{
+	{		
 		// Make sure there's no concurrency
-		check(!bWorking);
-		bWorking = true;
+		if (!AccessCurvePointsMutex.TryLock())
+		{
+			return;
+		}
 
 		const int32 OldSize = InterpolatingPoints.Num();
 		RefineFullRangeInterpolatingPointsInternal();
@@ -54,31 +56,11 @@ namespace UE::MovieSceneTools
 		const int32 NewSize = InterpolatingPoints.Num();
 		if (OldSize == NewSize)
 		{
-			// Remove straight lines to reduce slate elements needed to draw the curve
-			const float PixelArea = FMath::Pow(FMath::Min(TimeThreshold, ValueThreshold), 2.f);
-			
-			for (int32 PointIndex = 0; PointIndex < InterpolatingPoints.Num() - 3;)
-			{
-				const FVector2D& First = InterpolatingPoints[0];
-				const FVector2D& Second = InterpolatingPoints[1];
-				const FVector2D& Third = InterpolatingPoints[2];
-
-				const float LineArea = FMath::Abs(First.X * (Second.Y - Third.Y) + Second.X * (Third.Y - First.Y) + Third.X * (First.Y - Third.Y));
-				if (LineArea < PixelArea)
-				{
-					InterpolatingPoints.RemoveAt(PointIndex + 1);
-				}
-				else
-				{
-					PointIndex++;
-				}
-			}
-
 			InvokeCallback();
 			SetFlags(ECurvePainterTaskStateFlags::Completed);
 		}
 
-		bWorking = false;
+		AccessCurvePointsMutex.Unlock();
 	}
 
 	template <typename ChannelType>
@@ -137,23 +119,19 @@ namespace UE::MovieSceneTools
 	void FMovieSceneInterpolatingPointsDrawTask<ChannelType>::InvokeCallback()
 	{
 		TArray<int32> KeyOffsets;
-		int32 InterpolatingPointIndex = 0;
-		for (const FVector2D& KeyPoint : KeyPoints)
-		{
-			while (
-				InterpolatingPoints.IsValidIndex(InterpolatingPointIndex) &&
-				KeyPoint.X > InterpolatingPoints[InterpolatingPointIndex].X)
-			{
-				InterpolatingPointIndex++;
-			}
 
-			if (InterpolatingPoints.IsValidIndex(InterpolatingPointIndex))
+		int32 Offset = 0;
+		for (int32 InterpPointIndex = 0; InterpPointIndex < InterpolatingPoints.Num(); InterpPointIndex++)
+		{
+			if (KeyPoints.Num() <= Offset)
 			{
-				KeyOffsets.Add(InterpolatingPointIndex);
+				KeyOffsets.Add(InterpPointIndex);
+				break;
 			}
-			else if (!InterpolatingPoints.IsEmpty())
+			else if (InterpolatingPoints[InterpPointIndex].X >= KeyPoints[Offset].X)
 			{
-				KeyOffsets.Add(InterpolatingPoints.Num() - 1);
+				KeyOffsets.Add(InterpPointIndex);
+				Offset++;
 			}
 		}
 
diff --git a/Engine/Source/Editor/MovieSceneTools/Private/Cache/MovieSceneInterpolatingPointsDrawTask.h b/Engine/Source/Editor/MovieSceneTools/Private/Cache/MovieSceneInterpolatingPointsDrawTask.h
--- a/Engine/Source/Editor/MovieSceneTools/Private/Cache/MovieSceneInterpolatingPointsDrawTask.h
+++ b/Engine/Source/Editor/MovieSceneTools/Private/Cache/MovieSceneInterpolatingPointsDrawTask.h
@@ -4,6 +4,7 @@
 
 #include "Cache/MovieSceneCurveCachePool.h"
 #include "CurveEditorScreenSpace.h"
+#include "HAL/CriticalSection.h"
 #include "Misc/FrameRate.h"
 #include "Misc/Optional.h"
 #include "Templates/SharedPointer.h"
@@ -68,8 +69,8 @@ namespace UE::MovieSceneTools
 		/** The (possibly not yet fully interpolated) interpolating points */
 		TArray<FVector2D> InterpolatingPoints;
 
-		/** True while work is performed in this task. Useful to prevent from concurrency by design. */
-		std::atomic<bool> bWorking = false;
+		/** Critical section to enter when accessing curve points */
+		FCriticalSection AccessCurvePointsMutex;
 
 		/** True if the task needs to iterate again */
 		std::atomic<ECurvePainterTaskStateFlags> StateFlags = ECurvePainterTaskStateFlags::None;
