Updated headers using UnrealCodeFixup to make sure dllstorage is on methods/staticvars instead of types. Part 5/n

diff --git a/Engine/Source/Editor/MovieSceneTools/Public/Channels/SCurveEditorKeyBarView.h b/Engine/Source/Editor/MovieSceneTools/Public/Channels/SCurveEditorKeyBarView.h
--- a/Engine/Source/Editor/MovieSceneTools/Public/Channels/SCurveEditorKeyBarView.h
+++ b/Engine/Source/Editor/MovieSceneTools/Public/Channels/SCurveEditorKeyBarView.h
@@ -10,6 +10,8 @@
 #include "Templates/SharedPointer.h"
 #include "Views/SInteractiveCurveEditorView.h"
 
+#define UE_API MOVIESCENETOOLS_API
+
 class FCurveEditor;
 class FMenuBuilder;
 class FPaintArgs;
@@ -21,25 +23,27 @@ struct FCurveModelID;
 struct FCurvePointHandle;
 struct FGeometry;
 
-class MOVIESCENETOOLS_API SCurveEditorKeyBarView : public SInteractiveCurveEditorView
+class SCurveEditorKeyBarView : public SInteractiveCurveEditorView
 {
 public:
-	void Construct(const FArguments& InArgs, TWeakPtr<FCurveEditor> InCurveEditor);
+	UE_API void Construct(const FArguments& InArgs, TWeakPtr<FCurveEditor> InCurveEditor);
 
-	virtual void GetGridLinesY(TSharedRef<const FCurveEditor> CurveEditor, TArray<float>& MajorGridLines, TArray<float>& MinorGridLines, TArray<FText>* MajorGridLabels) const override;
+	UE_API virtual void GetGridLinesY(TSharedRef<const FCurveEditor> CurveEditor, TArray<float>& MajorGridLines, TArray<float>& MinorGridLines, TArray<FText>* MajorGridLabels) const override;
 
 private:
 
 	//~ SWidget Interface
-	virtual void Tick(const FGeometry& AllottedGeometry, const double InCurrentTime, const float InDeltaTime) override;
-	virtual FVector2D ComputeDesiredSize(float LayoutScaleMultiplier) const override;
+	UE_API virtual void Tick(const FGeometry& AllottedGeometry, const double InCurrentTime, const float InDeltaTime) override;
+	UE_API virtual FVector2D ComputeDesiredSize(float LayoutScaleMultiplier) const override;
 
 	//~ SInteractiveCurveEditorView interface
-	virtual void PaintView(const FPaintArgs& Args, const FGeometry& AllottedGeometry, const FSlateRect& MyCullingRect, FSlateWindowElementList& OutDrawElements, int32 BaseLayerId, const FWidgetStyle& InWidgetStyle, bool bParentEnabled) const override;
+	UE_API virtual void PaintView(const FPaintArgs& Args, const FGeometry& AllottedGeometry, const FSlateRect& MyCullingRect, FSlateWindowElementList& OutDrawElements, int32 BaseLayerId, const FWidgetStyle& InWidgetStyle, bool bParentEnabled) const override;
 	virtual bool IsValueSnapEnabled() const override { return false; }
-	virtual void BuildContextMenu(FMenuBuilder& MenuBuilder, TOptional<FCurvePointHandle> ClickedPoint, TOptional<FCurveModelID> HoveredCurveID) override;
+	UE_API virtual void BuildContextMenu(FMenuBuilder& MenuBuilder, TOptional<FCurvePointHandle> ClickedPoint, TOptional<FCurveModelID> HoveredCurveID) override;
+
+	UE_API void DrawLabels(const FGeometry& AllottedGeometry, FSlateWindowElementList& OutDrawElements, int32 BaseLayerId, ESlateDrawEffect DrawEffects) const;
 
-	void DrawLabels(const FGeometry& AllottedGeometry, FSlateWindowElementList& OutDrawElements, int32 BaseLayerId, ESlateDrawEffect DrawEffects) const;
+	static UE_API float TrackHeight;
+};
 
-	static float TrackHeight;
-};
\ No newline at end of file
+#undef UE_API
diff --git a/Engine/Source/Editor/MovieSceneTools/Public/Conditions/MovieSceneConditionCustomization.h b/Engine/Source/Editor/MovieSceneTools/Public/Conditions/MovieSceneConditionCustomization.h
--- a/Engine/Source/Editor/MovieSceneTools/Public/Conditions/MovieSceneConditionCustomization.h
+++ b/Engine/Source/Editor/MovieSceneTools/Public/Conditions/MovieSceneConditionCustomization.h
@@ -23,8 +23,8 @@ class IDetailsView;
 class ISequencer;
 
 // Helper UObject for editing optional track row metadata not in-place. A UObject instead of a UStruct because we need to support instanced sub objects (conditions)
-UCLASS(CollapseCategories)
-class MOVIESCENETOOLS_API UMovieSceneTrackRowMetadataHelper : public UObject
+UCLASS(MinimalAPI, CollapseCategories)
+class UMovieSceneTrackRowMetadataHelper : public UObject
 {
 	GENERATED_BODY()
 
diff --git a/Engine/Source/Editor/MovieSceneTools/Public/Conditions/MovieSceneDirectorBlueprintConditionCustomization.h b/Engine/Source/Editor/MovieSceneTools/Public/Conditions/MovieSceneDirectorBlueprintConditionCustomization.h
--- a/Engine/Source/Editor/MovieSceneTools/Public/Conditions/MovieSceneDirectorBlueprintConditionCustomization.h
+++ b/Engine/Source/Editor/MovieSceneTools/Public/Conditions/MovieSceneDirectorBlueprintConditionCustomization.h
@@ -7,6 +7,8 @@
 #include "IPropertyTypeCustomization.h"
 #include "Types/SlateEnums.h"
 
+#define UE_API MOVIESCENETOOLS_API
+
 class UK2Node;
 class UBlueprint;
 class UEdGraphPin;
@@ -21,37 +23,38 @@ enum class ECheckBoxState : uint8;
 /**
  * Customization for director blueprint condition endpoint picker
  */
-class MOVIESCENETOOLS_API FMovieSceneDirectorBlueprintConditionCustomization : public FMovieSceneDirectorBlueprintEndpointCustomization
+class FMovieSceneDirectorBlueprintConditionCustomization : public FMovieSceneDirectorBlueprintEndpointCustomization
 {
 public:
 
-	static TSharedRef<IPropertyTypeCustomization> MakeInstance();
-	static TSharedRef<IPropertyTypeCustomization> MakeInstance(UMovieScene* InMovieScene);
-	static TSharedRef<FMovieSceneDirectorBlueprintConditionCustomization> MakeInstance(UMovieScene* InMovieScene, TSharedPtr<IPropertyHandle> InPropertyHandle, TSharedPtr<IPropertyUtilities> InPropertyUtilities);
+	static UE_API TSharedRef<IPropertyTypeCustomization> MakeInstance();
+	static UE_API TSharedRef<IPropertyTypeCustomization> MakeInstance(UMovieScene* InMovieScene);
+	static UE_API TSharedRef<FMovieSceneDirectorBlueprintConditionCustomization> MakeInstance(UMovieScene* InMovieScene, TSharedPtr<IPropertyHandle> InPropertyHandle, TSharedPtr<IPropertyUtilities> InPropertyUtilities);
 
 protected:
 
-	virtual void GetPayloadVariables(UObject* EditObject, void* RawData, FPayloadVariableMap& OutPayloadVariables) const override;
-	virtual bool SetPayloadVariable(UObject* EditObject, void* RawData, FName FieldName, const FMovieSceneDirectorBlueprintVariableValue& NewVariableValue) override;
-	virtual UK2Node* FindEndpoint(UMovieSceneSequence* Sequence, UBlueprint* Blueprint, UObject* EditObject, void* RawData) const override;
-	virtual void GetWellKnownParameterPinNames(UObject* EditObject, void* RawData, TArray<FName>& OutWellKnownParameters) const override;
-	virtual void GetWellKnownParameterCandidates(UK2Node* Endpoint, TArray<FWellKnownParameterCandidates>& OutCandidates) const override;
-	virtual bool SetWellKnownParameterPinName(UObject* EditObject, void* RawData, int32 ParameterIndex, FName BoundPinName) override;
-	virtual FMovieSceneDirectorBlueprintEndpointDefinition GenerateEndpointDefinition(UMovieSceneSequence* Sequence) override;
-	virtual void OnCreateEndpoint(UMovieSceneSequence* Sequence, UBlueprint* Blueprint, const TArray<UObject*> EditObjects, const TArray<void*> RawData, const FMovieSceneDirectorBlueprintEndpointDefinition& EndpointDefinition, UK2Node* NewEndpoint) override;
-	virtual void OnSetEndpoint(UMovieSceneSequence* Sequence, UBlueprint* Blueprint, const TArray<UObject*> EditObjects, const TArray<void*> RawData, const FMovieSceneDirectorBlueprintEndpointDefinition& EndpointDefinition, UK2Node* NewEndpoint) override;
-	virtual void GetEditObjects(TArray<UObject*>& OutObjects) const override;
-	virtual void OnCollectQuickBindActions(UBlueprint* Blueprint, FBlueprintActionMenuBuilder& MenuBuilder) override;
+	UE_API virtual void GetPayloadVariables(UObject* EditObject, void* RawData, FPayloadVariableMap& OutPayloadVariables) const override;
+	UE_API virtual bool SetPayloadVariable(UObject* EditObject, void* RawData, FName FieldName, const FMovieSceneDirectorBlueprintVariableValue& NewVariableValue) override;
+	UE_API virtual UK2Node* FindEndpoint(UMovieSceneSequence* Sequence, UBlueprint* Blueprint, UObject* EditObject, void* RawData) const override;
+	UE_API virtual void GetWellKnownParameterPinNames(UObject* EditObject, void* RawData, TArray<FName>& OutWellKnownParameters) const override;
+	UE_API virtual void GetWellKnownParameterCandidates(UK2Node* Endpoint, TArray<FWellKnownParameterCandidates>& OutCandidates) const override;
+	UE_API virtual bool SetWellKnownParameterPinName(UObject* EditObject, void* RawData, int32 ParameterIndex, FName BoundPinName) override;
+	UE_API virtual FMovieSceneDirectorBlueprintEndpointDefinition GenerateEndpointDefinition(UMovieSceneSequence* Sequence) override;
+	UE_API virtual void OnCreateEndpoint(UMovieSceneSequence* Sequence, UBlueprint* Blueprint, const TArray<UObject*> EditObjects, const TArray<void*> RawData, const FMovieSceneDirectorBlueprintEndpointDefinition& EndpointDefinition, UK2Node* NewEndpoint) override;
+	UE_API virtual void OnSetEndpoint(UMovieSceneSequence* Sequence, UBlueprint* Blueprint, const TArray<UObject*> EditObjects, const TArray<void*> RawData, const FMovieSceneDirectorBlueprintEndpointDefinition& EndpointDefinition, UK2Node* NewEndpoint) override;
+	UE_API virtual void GetEditObjects(TArray<UObject*>& OutObjects) const override;
+	UE_API virtual void OnCollectQuickBindActions(UBlueprint* Blueprint, FBlueprintActionMenuBuilder& MenuBuilder) override;
 	virtual bool CreateNewCategoryForPayloadVariables() const override { return false; }
 private:
 
-	void SetEndpointImpl(UMovieScene* MovieScene, FMovieSceneDirectorBlueprintConditionData* DynamicBinding, UBlueprint* Blueprint, UK2Node* NewEndpoint);
-	void EnsureBlueprintExtensionCreated(UMovieScene* MovieScene, UBlueprint* Blueprint);
+	UE_API void SetEndpointImpl(UMovieScene* MovieScene, FMovieSceneDirectorBlueprintConditionData* DynamicBinding, UBlueprint* Blueprint, UK2Node* NewEndpoint);
+	UE_API void EnsureBlueprintExtensionCreated(UMovieScene* MovieScene, UBlueprint* Blueprint);
 
-	void CollectConditionBindActions(UBlueprint* Blueprint, FBlueprintActionMenuBuilder& MenuBuilder, bool bIsRebinding);
+	UE_API void CollectConditionBindActions(UBlueprint* Blueprint, FBlueprintActionMenuBuilder& MenuBuilder, bool bIsRebinding);
 
 private:
 
 	UMovieScene* EditedMovieScene;
 };
 
+#undef UE_API
diff --git a/Engine/Source/Editor/MovieSceneTools/Public/Conditions/MovieSceneDirectorBlueprintConditionUtils.h b/Engine/Source/Editor/MovieSceneTools/Public/Conditions/MovieSceneDirectorBlueprintConditionUtils.h
--- a/Engine/Source/Editor/MovieSceneTools/Public/Conditions/MovieSceneDirectorBlueprintConditionUtils.h
+++ b/Engine/Source/Editor/MovieSceneTools/Public/Conditions/MovieSceneDirectorBlueprintConditionUtils.h
@@ -13,22 +13,24 @@
 
 #include "MovieSceneDirectorBlueprintConditionUtils.generated.h"
 
+#define UE_API MOVIESCENETOOLS_API
+
 /**
  * A utility class for managing director blueprint condition endpoints.
  */
-struct MOVIESCENETOOLS_API FMovieSceneDirectorBlueprintConditionUtils
+struct FMovieSceneDirectorBlueprintConditionUtils
 {
 	using FSharedPlaybackState = UE::MovieScene::FSharedPlaybackState;
 
 	/**
 	 * Set an endpoint on the given blueprint condition
 	 */
-	static void SetEndpoint(UMovieScene* MovieScene, FMovieSceneDirectorBlueprintConditionData* DirectorBlueprintConditionData, UK2Node* NewEndpoint);
+	static UE_API void SetEndpoint(UMovieScene* MovieScene, FMovieSceneDirectorBlueprintConditionData* DirectorBlueprintConditionData, UK2Node* NewEndpoint);
 
 	/**
 	 * Ensures that the condition blueprint extension has been added to the given sequence's director blueprint.
 	 */
-	static void EnsureBlueprintExtensionCreated(UMovieSceneSequence* MovieSceneSequence, UBlueprint* Blueprint);
+	static UE_API void EnsureBlueprintExtensionCreated(UMovieSceneSequence* MovieSceneSequence, UBlueprint* Blueprint);
 
 	/**
 	 * Utility function for iterating all blueprint conditions in a sequence.
@@ -99,8 +101,8 @@ struct MOVIESCENETOOLS_API FMovieSceneDirectorBlueprintConditionUtils
 /**
  * Dummy class, used for easily getting a valid UFunction that helps prepare blueprint function graphs.
  */
-UCLASS()
-class MOVIESCENETOOLS_API UMovieSceneDirectorBlueprintConditionEndpointUtil : public UObject
+UCLASS(MinimalAPI)
+class UMovieSceneDirectorBlueprintConditionEndpointUtil : public UObject
 {
 	GENERATED_BODY()
 
@@ -108,3 +110,4 @@ class MOVIESCENETOOLS_API UMovieSceneDirectorBlueprintConditionEndpointUtil : pu
 	bool SampleDirectorBlueprintCondition(const FMovieSceneConditionContext& ConditionContext) const { return true; }
 };
 
+#undef UE_API
diff --git a/Engine/Source/Editor/MovieSceneTools/Public/Conditions/MovieScenePlatformConditionCustomization.h b/Engine/Source/Editor/MovieSceneTools/Public/Conditions/MovieScenePlatformConditionCustomization.h
--- a/Engine/Source/Editor/MovieSceneTools/Public/Conditions/MovieScenePlatformConditionCustomization.h
+++ b/Engine/Source/Editor/MovieSceneTools/Public/Conditions/MovieScenePlatformConditionCustomization.h
@@ -6,20 +6,24 @@
 #include "Templates/SharedPointer.h"
 #include "PropertyHandle.h"
 
+#define UE_API MOVIESCENETOOLS_API
+
 class IDetailLayoutBuilder;
 class SCheckBoxList;
 
-class MOVIESCENETOOLS_API FMovieScenePlatformConditionCustomization : public IDetailCustomization
+class FMovieScenePlatformConditionCustomization : public IDetailCustomization
 {
 public:
-	static TSharedRef<IDetailCustomization> MakeInstance();
-	virtual void CustomizeDetails(IDetailLayoutBuilder& DetailBuilder) override;
+	static UE_API TSharedRef<IDetailCustomization> MakeInstance();
+	UE_API virtual void CustomizeDetails(IDetailLayoutBuilder& DetailBuilder) override;
 
 private:
 
-	void OnPlatformCheckChanged(int32 Index);
-	TArray<FName> GetCurrentValidPlatformNames();
+	UE_API void OnPlatformCheckChanged(int32 Index);
+	UE_API TArray<FName> GetCurrentValidPlatformNames();
 
 	TSharedPtr<IPropertyHandle> ValidPlatformsPropertyHandle;
 	TSharedPtr<SCheckBoxList> CheckBoxList;
 };
+
+#undef UE_API
diff --git a/Engine/Source/Editor/MovieSceneTools/Public/Constraints/MovieSceneConstraintChannelHelper.h b/Engine/Source/Editor/MovieSceneTools/Public/Constraints/MovieSceneConstraintChannelHelper.h
--- a/Engine/Source/Editor/MovieSceneTools/Public/Constraints/MovieSceneConstraintChannelHelper.h
+++ b/Engine/Source/Editor/MovieSceneTools/Public/Constraints/MovieSceneConstraintChannelHelper.h
@@ -6,6 +6,8 @@
 #include "UObject/ObjectPtr.h"
 #include "MovieSceneSection.h"
 
+#define UE_API MOVIESCENETOOLS_API
+
 
 class IMovieSceneConstrainedSection;
 struct ITransformConstraintChannelInterface;
@@ -21,14 +23,14 @@ struct FFrameNumber;
 struct FConstraintAndActiveChannel;
 enum class EMovieSceneKeyInterpolation : uint8;
 
-struct MOVIESCENETOOLS_API FCompensationEvaluator
+struct FCompensationEvaluator
 {
 public:
 	TArray<FTransform> ChildLocals;
 	TArray<FTransform> ChildGlobals;
 	TArray<FTransform> SpaceGlobals;
 
-	FCompensationEvaluator(UTickableTransformConstraint* InConstraint);
+	UE_API FCompensationEvaluator(UTickableTransformConstraint* InConstraint);
 
 	struct FEvalParameters
 	{
@@ -48,22 +50,22 @@ public:
 		bool bKeepCurrent = false;
 	};
 	
-	void ComputeLocalTransforms(UWorld* InWorld, const FEvalParameters& InEvalParams);
-	void ComputeLocalTransformsBeforeDeletion(UWorld* InWorld, const FEvalParameters& InEvalParams);
-	void ComputeCompensation(UWorld* InWorld, const TSharedPtr<ISequencer>& InSequencer, const FFrameNumber& InTime);
-	void ComputeLocalTransformsForBaking(UWorld* InWorld, const FEvalParameters& InEvalParams);
-	void CacheTransforms(UWorld* InWorld, const FEvalParameters& InEvalParams);
-	void ComputeCurrentTransforms(UWorld* InWorld);
+	UE_API void ComputeLocalTransforms(UWorld* InWorld, const FEvalParameters& InEvalParams);
+	UE_API void ComputeLocalTransformsBeforeDeletion(UWorld* InWorld, const FEvalParameters& InEvalParams);
+	UE_API void ComputeCompensation(UWorld* InWorld, const TSharedPtr<ISequencer>& InSequencer, const FFrameNumber& InTime);
+	UE_API void ComputeLocalTransformsForBaking(UWorld* InWorld, const FEvalParameters& InEvalParams);
+	UE_API void CacheTransforms(UWorld* InWorld, const FEvalParameters& InEvalParams);
+	UE_API void ComputeCurrentTransforms(UWorld* InWorld);
 	
 private:
 
-	const TArray< TWeakObjectPtr<UTickableConstraint> > GetHandleTransformConstraints(UWorld* InWorld) const;
+	UE_API const TArray< TWeakObjectPtr<UTickableConstraint> > GetHandleTransformConstraints(UWorld* InWorld) const;
 
 	UTickableTransformConstraint* Constraint = nullptr;
 	UTransformableHandle* Handle = nullptr;
 };
 
-struct MOVIESCENETOOLS_API FConstraintSections
+struct FConstraintSections
 {
 	UMovieSceneSection* ConstraintSection = nullptr;
 	UMovieSceneSection* ChildTransformSection = nullptr;
@@ -71,62 +73,62 @@ struct MOVIESCENETOOLS_API FConstraintSections
 	FConstraintAndActiveChannel* ActiveChannel = nullptr;
 	ITransformConstraintChannelInterface* Interface = nullptr;
 };
-struct MOVIESCENETOOLS_API FMovieSceneConstraintChannelHelper
+struct FMovieSceneConstraintChannelHelper
 {
 public:
 
-	static bool AddConstraintToSequencer(
+	static UE_API bool AddConstraintToSequencer(
 		const TSharedPtr<ISequencer>& InSequencer,
 		UTickableTransformConstraint* InConstraint);
 	
 	/** Adds an active key if needed and does the compensation when switching. Will use the optional active and time if set. 
 	Will return true if key is actually set, may not be if the value is the same.*/
-	static bool SmartConstraintKey(
+	static UE_API bool SmartConstraintKey(
 		const TSharedPtr<ISequencer>& InSequencer,
 		UTickableTransformConstraint* InConstraint, 
 		const TOptional<bool>& InOptActive,
 		const TOptional<FFrameNumber>& InOptFrameTime);
 	
 	/** Compensate transform on handles when a constraint switches state. */
-	static void Compensate(
+	static UE_API void Compensate(
 		const TSharedPtr<ISequencer>& InSequencer,
 		const UTickableTransformConstraint* InConstraint,
 		const TOptional<FFrameNumber>& InOptTime,
 		const bool bCompPreviousTick);
 	
-	static void CompensateIfNeeded(
+	static UE_API void CompensateIfNeeded(
 		const TSharedPtr<ISequencer>& InSequencer,
 		IMovieSceneConstrainedSection* Section,
 		const TOptional<FFrameNumber>& OptionalTime,
 		const bool bCompPreviousTick,
 		const int32 InChildHash = INDEX_NONE);
 		
-	static void HandleConstraintRemoved(
+	static UE_API void HandleConstraintRemoved(
 		UTickableConstraint* InConstraint,
 		const FMovieSceneConstraintChannel* InChannel,
 		const TSharedPtr<ISequencer>& InSequencer,
 		UMovieSceneSection* InSection);
 
-	static void HandleConstraintKeyDeleted(
+	static UE_API void HandleConstraintKeyDeleted(
 		UTickableTransformConstraint* InConstraint,
 		const FMovieSceneConstraintChannel* InConstraintChannel,
 		const TSharedPtr<ISequencer>& InSequencer,
 		UMovieSceneSection* InSection,
 		const FFrameNumber& InTime);
 
-	static void HandleConstraintKeyMoved(
+	static UE_API void HandleConstraintKeyMoved(
 		const UTickableTransformConstraint* InConstraint,
 		const FMovieSceneConstraintChannel* InConstraintChannel,
 		UMovieSceneSection* InSection,
 		const FFrameNumber& InCurrentFrame, const FFrameNumber& InNextFrame);
 
 	/* Get the section and the channel for the given constraint, will be nullptr's if it doesn't exist in Sequencer*/
-	static FConstraintSections  GetConstraintSectionAndChannel(
+	static UE_API FConstraintSections  GetConstraintSectionAndChannel(
 		const UTickableTransformConstraint* InConstraint,
 		const TSharedPtr<ISequencer>& InSequencer);
 
 	/* For the given constraint get all of the transform keys for it's child and parent handles*/
-	static void GetTransformFramesForConstraintHandles(
+	static UE_API void GetTransformFramesForConstraintHandles(
 		const UTickableTransformConstraint* InConstraint,
 		const TSharedPtr<ISequencer>& InSequencer,
 		const FFrameNumber& StartFrame,
@@ -178,7 +180,7 @@ public:
 		const FFrameNumber& InTime,
 		EMovieSceneKeyInterpolation KeyInterpolation);
 
-	static void HandleConstraintPropertyChanged(
+	static UE_API void HandleConstraintPropertyChanged(
 		UTickableTransformConstraint* InConstraint,
 		const FMovieSceneConstraintChannel& InActiveChannel,
 		const FPropertyChangedEvent& InPropertyChangedEvent,
@@ -205,28 +207,30 @@ public:
 		const TArray<FFrameNumber>& Frames,
 		const TArray<FTransform>& Transforms);
 
-	static bool bDoNotCompensate;
+	static UE_API bool bDoNotCompensate;
 
 private:
 	/** Get the animatable interface for that handle if registered. */
-	static ITransformConstraintChannelInterface* GetHandleInterface(const UTransformableHandle* InHandle);
+	static UE_API ITransformConstraintChannelInterface* GetHandleInterface(const UTransformableHandle* InHandle);
 	
 	/** For the given handle create any movie scene binding for it based upon the current sequencer that's open*/
-	static void CreateBindingIDForHandle(const TSharedPtr<ISequencer>& InSequencer, UTransformableHandle* InHandle);
+	static UE_API void CreateBindingIDForHandle(const TSharedPtr<ISequencer>& InSequencer, UTransformableHandle* InHandle);
 
 	/** For the given handle, return true if a binding exists and if that binding represents a spawnable. */
-	static bool IsHandleSpawnable(const TSharedPtr<ISequencer>& InSequencer, const UTransformableHandle* InHandle);
+	static UE_API bool IsHandleSpawnable(const TSharedPtr<ISequencer>& InSequencer, const UTransformableHandle* InHandle);
 
 	/** Compensate scale keys when enabling/disabling scaling for parent constraints. */
-	static void CompensateScale(
+	static UE_API void CompensateScale(
 		UTickableParentConstraint* InParentConstraint,
 		const FMovieSceneConstraintChannel& InActiveChannel,
 		const TSharedPtr<ISequencer>& InSequencer,
 		UMovieSceneSection* InSection);
 
 	/** Handle offset modifications so that the child's transform channels are synced. */
-	static void HandleOffsetChanged(
+	static UE_API void HandleOffsetChanged(
 		UTickableTransformConstraint* InConstraint,
 		const FMovieSceneConstraintChannel& InActiveChannel,
 		const TSharedPtr<ISequencer>& InSequencer);
 };
+
+#undef UE_API
diff --git a/Engine/Source/Editor/MovieSceneTools/Public/FCPXML/FCPXMLMovieSceneTranslator.h b/Engine/Source/Editor/MovieSceneTools/Public/FCPXML/FCPXMLMovieSceneTranslator.h
--- a/Engine/Source/Editor/MovieSceneTools/Public/FCPXML/FCPXMLMovieSceneTranslator.h
+++ b/Engine/Source/Editor/MovieSceneTools/Public/FCPXML/FCPXMLMovieSceneTranslator.h
@@ -8,28 +8,30 @@
 #include "MovieScene.h"
 #include "MovieSceneTranslator.h"
 
+#define UE_API MOVIESCENETOOLS_API
+
 class FMovieSceneTranslatorContext;
 
 /** 
  * The FFCPXMLImporter class is the entry point for launching an import of data from an XML file into Sequencer. 
  */
-class MOVIESCENETOOLS_API FFCPXMLImporter : public FMovieSceneImporter
+class FFCPXMLImporter : public FMovieSceneImporter
 {
 public:
-	FFCPXMLImporter();
+	UE_API FFCPXMLImporter();
 
-	virtual ~FFCPXMLImporter();
+	UE_API virtual ~FFCPXMLImporter();
 
 	/** Format description. */
-	virtual FText GetFileTypeDescription() const;
+	UE_API virtual FText GetFileTypeDescription() const;
 	/** Import window title. */
-	virtual FText GetDialogTitle() const;
+	UE_API virtual FText GetDialogTitle() const;
 	/** Scoped transaction description. */
-	virtual FText GetTransactionDescription() const;
+	UE_API virtual FText GetTransactionDescription() const;
 	/** Message log window title. */
-	virtual FName GetMessageLogWindowTitle() const;
+	UE_API virtual FName GetMessageLogWindowTitle() const;
 	/** Message log list label. */
-	virtual FText GetMessageLogLabel() const;
+	UE_API virtual FText GetMessageLogLabel() const;
 
 public:
 	/*
@@ -41,34 +43,34 @@ public:
 	 * @param OutError The return error message
 	 * @return Whether the import was successful
 	 */
-	bool Import(UMovieScene* InMovieScene, FFrameRate InFrameRate, FString InFilename, TSharedRef<FMovieSceneTranslatorContext> InContext);
+	UE_API bool Import(UMovieScene* InMovieScene, FFrameRate InFrameRate, FString InFilename, TSharedRef<FMovieSceneTranslatorContext> InContext);
 };
 
 /** 
  * The FFCPXMLExporter class is the entry point for launching an export of data from Sequencer into an XML file. 
  */
-class MOVIESCENETOOLS_API FFCPXMLExporter : public FMovieSceneExporter
+class FFCPXMLExporter : public FMovieSceneExporter
 {
 public:
 	/** Constructor */
-	FFCPXMLExporter();
+	UE_API FFCPXMLExporter();
 	/** Destructor */
-	virtual ~FFCPXMLExporter();
+	UE_API virtual ~FFCPXMLExporter();
 
 	/** Format description. */
-	virtual FText GetFileTypeDescription() const;
+	UE_API virtual FText GetFileTypeDescription() const;
 	/** Export dialog window title. */
-	virtual FText GetDialogTitle() const;
+	UE_API virtual FText GetDialogTitle() const;
 	/** Default format file extension. */
-	virtual FText GetDefaultFileExtension() const;
+	UE_API virtual FText GetDefaultFileExtension() const;
 	/** Notification when export completes. */
-	virtual FText GetNotificationExportFinished() const;
+	UE_API virtual FText GetNotificationExportFinished() const;
 	/** Notification hyperlink to exported file path. */
-	virtual FText GetNotificationHyperlinkText() const;
+	UE_API virtual FText GetNotificationHyperlinkText() const;
 	/** Message log window title. */
-	virtual FName GetMessageLogWindowTitle() const;
+	UE_API virtual FName GetMessageLogWindowTitle() const;
 	/** Message log list label. */
-	virtual FText GetMessageLogLabel() const;
+	UE_API virtual FText GetMessageLogLabel() const;
 
 public:
 	/*
@@ -87,5 +89,7 @@ public:
 	 * @return Whether the export was successful
 	 */
 	
-	virtual bool Export(const UMovieScene* InMovieScene, FString InFilenameFormat, FFrameRate InFrameRate, uint32 InResX, uint32 InResY, int32 InHandleFrames, FString InSaveFilename, TSharedRef<FMovieSceneTranslatorContext> InContext, FString InMovieExtension, const FMovieSceneExportMetadata* InMetadata=nullptr);
+	UE_API virtual bool Export(const UMovieScene* InMovieScene, FString InFilenameFormat, FFrameRate InFrameRate, uint32 InResX, uint32 InResY, int32 InHandleFrames, FString InSaveFilename, TSharedRef<FMovieSceneTranslatorContext> InContext, FString InMovieExtension, const FMovieSceneExportMetadata* InMetadata=nullptr);
 };
+
+#undef UE_API
diff --git a/Engine/Source/Editor/MovieSceneTools/Public/FrameNumberDetailsCustomization.h b/Engine/Source/Editor/MovieSceneTools/Public/FrameNumberDetailsCustomization.h
--- a/Engine/Source/Editor/MovieSceneTools/Public/FrameNumberDetailsCustomization.h
+++ b/Engine/Source/Editor/MovieSceneTools/Public/FrameNumberDetailsCustomization.h
@@ -7,12 +7,14 @@
 #include "ISequencerNumericTypeInterface.h"
 #include "Widgets/Input/NumericTypeInterface.h"
 
+#define UE_API MOVIESCENETOOLS_API
+
 class IDetailLayoutBuilder;
 
 /**
  *  Customize the FFrameNumber to support conversion from seconds/frames/timecode formats.
  */
-class MOVIESCENETOOLS_API FFrameNumberDetailsCustomization : public IPropertyTypeCustomization
+class FFrameNumberDetailsCustomization : public IPropertyTypeCustomization
 {
 public:
 
@@ -37,12 +39,12 @@ public:
 
 	/** IPropertyTypeCustomization interface */
 	virtual void CustomizeHeader(TSharedRef<IPropertyHandle> PropertyHandle, FDetailWidgetRow& HeaderRow, IPropertyTypeCustomizationUtils& CustomizationUtils) override {}
-	virtual void CustomizeChildren(TSharedRef<IPropertyHandle> PropertyHandle, IDetailChildrenBuilder& ChildBuilder, IPropertyTypeCustomizationUtils& CustomizationUtils) override;
+	UE_API virtual void CustomizeChildren(TSharedRef<IPropertyHandle> PropertyHandle, IDetailChildrenBuilder& ChildBuilder, IPropertyTypeCustomizationUtils& CustomizationUtils) override;
 
 private:
-	FText OnGetTimeText() const;
-	FText OnGetTimeToolTipText() const;
-	void OnTimeTextCommitted(const FText& InText, ETextCommit::Type CommitInfo);
+	UE_API FText OnGetTimeText() const;
+	UE_API FText OnGetTimeToolTipText() const;
+	UE_API void OnTimeTextCommitted(const FText& InText, ETextCommit::Type CommitInfo);
 
 	TWeakPtr<ISequencer> WeakSequencer;
 
@@ -56,3 +58,5 @@ private:
 	/** If they've used the UIMax metadata on the FFrameNumber property, we store that for use via text box callbacks. */
 	int32 UIClampMax;
 };
+
+#undef UE_API
diff --git a/Engine/Source/Editor/MovieSceneTools/Public/LevelSequenceFBXInterop.h b/Engine/Source/Editor/MovieSceneTools/Public/LevelSequenceFBXInterop.h
--- a/Engine/Source/Editor/MovieSceneTools/Public/LevelSequenceFBXInterop.h
+++ b/Engine/Source/Editor/MovieSceneTools/Public/LevelSequenceFBXInterop.h
@@ -4,26 +4,29 @@
 
 #include "ISequencer.h"
 
-class MOVIESCENETOOLS_API FLevelSequenceFBXInterop
+#define UE_API MOVIESCENETOOLS_API
+
+class FLevelSequenceFBXInterop
 {
 public:
 
-	FLevelSequenceFBXInterop(TSharedRef<ISequencer> InSequencer);
+	UE_API FLevelSequenceFBXInterop(TSharedRef<ISequencer> InSequencer);
 
 	/** Imports the animation from an fbx file. */
-	void ImportFBX();
-	void ImportFBXOntoSelectedNodes();
+	UE_API void ImportFBX();
+	UE_API void ImportFBXOntoSelectedNodes();
 
 	/** Exports the animation to an fbx file. */
-	void ExportFBX();
+	UE_API void ExportFBX();
 
 private:
 
 	/** Exports sequence to a FBX file */
-	void ExportFBXInternal(const FString& ExportFilename, const TArray<FGuid>& Bindings, const TArray<UMovieSceneTrack*>& Tracks);
+	UE_API void ExportFBXInternal(const FString& ExportFilename, const TArray<FGuid>& Bindings, const TArray<UMovieSceneTrack*>& Tracks);
 
 private:
 
 	TWeakPtr<ISequencer> WeakSequencer;
 };
 
+#undef UE_API
diff --git a/Engine/Source/Editor/MovieSceneTools/Public/MVVM/ViewModels/BindingLifetimeTrackModel.h b/Engine/Source/Editor/MovieSceneTools/Public/MVVM/ViewModels/BindingLifetimeTrackModel.h
--- a/Engine/Source/Editor/MovieSceneTools/Public/MVVM/ViewModels/BindingLifetimeTrackModel.h
+++ b/Engine/Source/Editor/MovieSceneTools/Public/MVVM/ViewModels/BindingLifetimeTrackModel.h
@@ -8,6 +8,8 @@
 #include "Templates/SharedPointer.h"
 #include "MVVM/Extensions/IBindingLifetimeExtension.h"
 
+#define UE_API MOVIESCENETOOLS_API
+
 class UMovieSceneBindingLifetimeTrack;
 class UMovieSceneTrack;
 namespace UE::Sequencer { template <typename T> struct TAutoRegisterViewModelTypeID; }
@@ -17,30 +19,30 @@ namespace UE
 	namespace Sequencer
 	{
 
-		class MOVIESCENETOOLS_API FBindingLifetimeTrackModel
+		class FBindingLifetimeTrackModel
 			: public FTrackModel
 			, public IBindingLifetimeExtension
 		{
 		public:
 
-			UE_SEQUENCER_DECLARE_CASTABLE(FBindingLifetimeTrackModel, FTrackModel, IBindingLifetimeExtension);
+			UE_SEQUENCER_DECLARE_CASTABLE_API(UE_API, FBindingLifetimeTrackModel, FTrackModel, IBindingLifetimeExtension);
 
-			static TSharedPtr<FTrackModel> CreateTrackModel(UMovieSceneTrack* Track);
+			static UE_API TSharedPtr<FTrackModel> CreateTrackModel(UMovieSceneTrack* Track);
 
-			explicit FBindingLifetimeTrackModel(UMovieSceneBindingLifetimeTrack* Track);
+			UE_API explicit FBindingLifetimeTrackModel(UMovieSceneBindingLifetimeTrack* Track);
 
-			FSortingKey GetSortingKey() const override;
+			UE_API FSortingKey GetSortingKey() const override;
 
-			void OnDeferredModifyFlush() override;
+			UE_API void OnDeferredModifyFlush() override;
 
 			// IBindingLifetimeExtension
 			const TArray<FFrameNumberRange>& GetInverseLifetimeRange() const override { return InverseLifetimeRange; }
 		
 		protected:
-			virtual void OnConstruct() override;
+			UE_API virtual void OnConstruct() override;
 
 		private:
-			void RecalculateInverseLifetimeRange();
+			UE_API void RecalculateInverseLifetimeRange();
 
 			// The inverse of the range created by all of our binding lifetime sections
 			// In other words, the ranges where the object binding should be deactivated.
@@ -51,3 +53,4 @@ namespace UE
 	} // namespace Sequencer
 } // namespace UE
 
+#undef UE_API
diff --git a/Engine/Source/Editor/MovieSceneTools/Public/MovieSceneDirectorBlueprintEndpointCustomization.h b/Engine/Source/Editor/MovieSceneTools/Public/MovieSceneDirectorBlueprintEndpointCustomization.h
--- a/Engine/Source/Editor/MovieSceneTools/Public/MovieSceneDirectorBlueprintEndpointCustomization.h
+++ b/Engine/Source/Editor/MovieSceneTools/Public/MovieSceneDirectorBlueprintEndpointCustomization.h
@@ -8,6 +8,8 @@
 #include "MovieSceneDirectorBlueprintUtils.h"
 #include "Types/SlateEnums.h"
 
+#define UE_API MOVIESCENETOOLS_API
+
 class UK2Node;
 class UBlueprint;
 class UEdGraphPin;
@@ -43,13 +45,13 @@ DECLARE_DELEGATE_FourParams(FOnQuickBindActionSelected, const TArray<TSharedPtr<
 /**
  * Base class for details view customizations that operate on sequence director blueprint endpoints.
  */
-class MOVIESCENETOOLS_API FMovieSceneDirectorBlueprintEndpointCustomization : public IPropertyTypeCustomization
+class FMovieSceneDirectorBlueprintEndpointCustomization : public IPropertyTypeCustomization
 {
 public:
 
 	// IPropertyTypeCustomization interface
-	virtual void CustomizeHeader(TSharedRef<IPropertyHandle> InPropertyHandle, FDetailWidgetRow& HeaderRow, IPropertyTypeCustomizationUtils& CustomizationUtils) override;
-	virtual void CustomizeChildren(TSharedRef<IPropertyHandle> InPropertyHandle, IDetailChildrenBuilder& ChildBuilder, IPropertyTypeCustomizationUtils& CustomizationUtils) override;
+	UE_API virtual void CustomizeHeader(TSharedRef<IPropertyHandle> InPropertyHandle, FDetailWidgetRow& HeaderRow, IPropertyTypeCustomizationUtils& CustomizationUtils) override;
+	UE_API virtual void CustomizeChildren(TSharedRef<IPropertyHandle> InPropertyHandle, IDetailChildrenBuilder& ChildBuilder, IPropertyTypeCustomizationUtils& CustomizationUtils) override;
 
 
 	// These functions exposed below so this customization can act as a helper to another
@@ -57,20 +59,20 @@ public:
 	/**
 	 * Creates a single new endpoint represented by this property handle.
 	 */
-	void CreateEndpoint();
+	UE_API void CreateEndpoint();
 
 	/**
 	 * Generate the content of the quick bind sub-menu dropdown (shown if the endpoint is not already bound)
 	 */
-	void PopulateQuickBindSubMenu(FMenuBuilder& MenuBuilder, UMovieSceneSequence* Sequence, FOnQuickBindActionSelected InOnQuickBindActionSelected);
+	UE_API void PopulateQuickBindSubMenu(FMenuBuilder& MenuBuilder, UMovieSceneSequence* Sequence, FOnQuickBindActionSelected InOnQuickBindActionSelected);
 
 	/*
 	* Called when a quick bind action has been selected
 	*/
-	void HandleQuickBindActionSelected(const TArray<TSharedPtr<FEdGraphSchemaAction>>& SelectedAction, ESelectInfo::Type InSelectionType, UBlueprint* Blueprint, FMovieSceneDirectorBlueprintEndpointDefinition EndpointDefinition);
+	UE_API void HandleQuickBindActionSelected(const TArray<TSharedPtr<FEdGraphSchemaAction>>& SelectedAction, ESelectInfo::Type InSelectionType, UBlueprint* Blueprint, FMovieSceneDirectorBlueprintEndpointDefinition EndpointDefinition);
 
 	/* Set the property handle to use for customization. Used when using the customization class as a helper. */
-	void SetPropertyHandle(TSharedPtr<IPropertyHandle> InPropertyHandle);
+	UE_API void SetPropertyHandle(TSharedPtr<IPropertyHandle> InPropertyHandle);
 
 	/* Set the raw data to use for customization. Used when the customization class is used as a helper without a details view. */
 	void SetRawData(const TArray<void*>& InRawData) { PropertyRawData = InRawData; }
@@ -186,7 +188,7 @@ protected:
 	 * Get all the objects that the endpoints reside within.
 	 * By default, this gets the outer objects of the property handle.
      */
-	virtual void GetEditObjects(TArray<UObject*>& OutObjects) const;
+	UE_API virtual void GetEditObjects(TArray<UObject*>& OutObjects) const;
 
 	/**
 	 * Collect blueprint actions for quick-binding a non-connected endpoint.
@@ -206,38 +208,38 @@ protected:
 	 * Clear the endpoint represented by this property handle.
 	 * @note: Does not delete the endpoint in the blueprint itself.
 	 */
-	void ClearEndpoint();
+	UE_API void ClearEndpoint();
 
 
 	/**
 	 * Find the endpoint node in the director blueprint.
 	 */
-	UK2Node* FindEndpoint(UObject* EditObject, void* RawData) const;
+	UE_API UK2Node* FindEndpoint(UObject* EditObject, void* RawData) const;
 
 	/**
 	 * Assigns the specified function entry to this property handle.
 	 */
-	void SetEndpoint(const FMovieSceneDirectorBlueprintEndpointDefinition& EndpointDefinition, UK2Node* NewEndpoint, UK2Node* PayloadTemplate, EAutoCreatePayload AutoCreatePayload);
+	UE_API void SetEndpoint(const FMovieSceneDirectorBlueprintEndpointDefinition& EndpointDefinition, UK2Node* NewEndpoint, UK2Node* PayloadTemplate, EAutoCreatePayload AutoCreatePayload);
 
 	/**
 	 * Navigate to the definition of the endpoint represented by this property handle.
 	 */
-	void NavigateToDefinition();
+	UE_API void NavigateToDefinition();
 
 	/**
 	 * Generate the content of the main combo button menu dropdown
 	 */
-	TSharedRef<SWidget> GetMenuContent();
+	UE_API TSharedRef<SWidget> GetMenuContent();
 
 	/**
 	 * Generate the content of the rebind sub-menu dropdown (shown if the endpoint is already bound)
 	 */
-	void PopulateRebindSubMenu(FMenuBuilder& MenuBuilder, UMovieSceneSequence* Sequence);
+	UE_API void PopulateRebindSubMenu(FMenuBuilder& MenuBuilder, UMovieSceneSequence* Sequence);
 
 	/**
 	 * Called when a payload variable property value is changed on the details panel
 	 */
-	void OnPayloadVariableChanged(TSharedRef<FStructOnScope> InStructData, TSharedPtr<IPropertyHandle> LocalVariableProperty);
+	UE_API void OnPayloadVariableChanged(TSharedRef<FStructOnScope> InStructData, TSharedPtr<IPropertyHandle> LocalVariableProperty);
 
 protected:
 
@@ -252,7 +254,7 @@ protected:
 	/**
 	 * Gather data for the sequences being edited.
 	 */
-	void GatherSequenceData(FSequenceDataMap& AllSequenceData);
+	UE_API void GatherSequenceData(FSequenceDataMap& AllSequenceData);
 
 	/**
 	 * Get the property handle for this customization.
@@ -267,22 +269,22 @@ protected:
 	/**
 	 * Get the sequence that is common to all the values in our property handle, or nullptr if they are not all the same.
 	 */
-	UMovieSceneSequence* GetCommonSequence() const;
+	UE_API UMovieSceneSequence* GetCommonSequence() const;
 
 	/**
 	 * Get the endpoint that is common to all the values in our property handle, or nullptr if they are not all the same.
 	 */
-	UK2Node* GetCommonEndpoint() const;
+	UE_API UK2Node* GetCommonEndpoint() const;
 
 	/**
 	 * Get all the endpoints of all the values in our property handle.
 	 */
-	TArray<UK2Node*> GetAllValidEndpoints() const;
+	UE_API TArray<UK2Node*> GetAllValidEndpoints() const;
 
 	/**
 	 * Invoke a callback on all the endpoints.
 	 */
-	void IterateEndpoints(TFunctionRef<bool(UK2Node*)> Callback) const;
+	UE_API void IterateEndpoints(TFunctionRef<bool(UK2Node*)> Callback) const;
 
 
 	/** The property handle we're reflecting */
@@ -296,31 +298,33 @@ protected:
 
 private:
 
-	void CollectQuickBindActions(FGraphActionListBuilderBase& OutAllActions, UBlueprint* Blueprint, FMovieSceneDirectorBlueprintEndpointDefinition EndpointDefinition);
-	void CollectAllRebindActions(FGraphActionListBuilderBase& OutAllActions, UBlueprint* Blueprint, FMovieSceneDirectorBlueprintEndpointDefinition EndpointDefinition);
+	UE_API void CollectQuickBindActions(FGraphActionListBuilderBase& OutAllActions, UBlueprint* Blueprint, FMovieSceneDirectorBlueprintEndpointDefinition EndpointDefinition);
+	UE_API void CollectAllRebindActions(FGraphActionListBuilderBase& OutAllActions, UBlueprint* Blueprint, FMovieSceneDirectorBlueprintEndpointDefinition EndpointDefinition);
 
 private:
 
 	/** Get the name of the endpoint to display on the main combo button */
-	FText GetEndpointName() const;
+	UE_API FText GetEndpointName() const;
 	/** Get the icon of the endpoint to display on the main combo button */
-	const FSlateBrush* GetEndpointIcon() const;
+	UE_API const FSlateBrush* GetEndpointIcon() const;
 
-	void HandleRebindActionSelected(const TArray<TSharedPtr<FEdGraphSchemaAction>>& SelectedAction, ESelectInfo::Type InSelectionType, UBlueprint* Blueprint, FMovieSceneDirectorBlueprintEndpointDefinition EndpointDefinition);
+	UE_API void HandleRebindActionSelected(const TArray<TSharedPtr<FEdGraphSchemaAction>>& SelectedAction, ESelectInfo::Type InSelectionType, UBlueprint* Blueprint, FMovieSceneDirectorBlueprintEndpointDefinition EndpointDefinition);
 	
-	void OnBlueprintCompiled(UBlueprint*);
+	UE_API void OnBlueprintCompiled(UBlueprint*);
 
-	ECheckBoxState GetCallInEditorCheckState() const;
-	void OnSetCallInEditorCheckState(ECheckBoxState NewCheckBoxState);
+	UE_API ECheckBoxState GetCallInEditorCheckState() const;
+	UE_API void OnSetCallInEditorCheckState(ECheckBoxState NewCheckBoxState);
 
-	TSharedRef<SWidget> GetWellKnownParameterPinMenuContent(int32 ParameterIndex);
-	const FSlateBrush* GetWellKnownParameterPinIcon(int32 ParameterIndex) const;
-	FText GetWellKnownParameterPinText(int32 ParameterIndex) const;
-	bool CompareWellKnownParameterPinName(int32 ParameterIndex, FName InPinName) const;
-	void SetWellKnownParameterPinName(int32 ParameterIndex, FName InNewWellKnownParameterPinName);
+	UE_API TSharedRef<SWidget> GetWellKnownParameterPinMenuContent(int32 ParameterIndex);
+	UE_API const FSlateBrush* GetWellKnownParameterPinIcon(int32 ParameterIndex) const;
+	UE_API FText GetWellKnownParameterPinText(int32 ParameterIndex) const;
+	UE_API bool CompareWellKnownParameterPinName(int32 ParameterIndex, FName InPinName) const;
+	UE_API void SetWellKnownParameterPinName(int32 ParameterIndex, FName InNewWellKnownParameterPinName);
 
 private:
 
 	/** A cache of the common endpoint that is only used when the menu is open to avoid re-computing it every frame. */
 	TWeakObjectPtr<UK2Node_FunctionEntry> CachedCommonEndpoint;
 };
+
+#undef UE_API
diff --git a/Engine/Source/Editor/MovieSceneTools/Public/MovieSceneDirectorBlueprintUtils.h b/Engine/Source/Editor/MovieSceneTools/Public/MovieSceneDirectorBlueprintUtils.h
--- a/Engine/Source/Editor/MovieSceneTools/Public/MovieSceneDirectorBlueprintUtils.h
+++ b/Engine/Source/Editor/MovieSceneTools/Public/MovieSceneDirectorBlueprintUtils.h
@@ -12,6 +12,8 @@
 #include "UObject/WeakObjectPtr.h"
 #include "UObject/WeakObjectPtrTemplates.h"
 
+#define UE_API MOVIESCENETOOLS_API
+
 class FKismetCompilerContext;
 class UBlueprint;
 class UClass;
@@ -41,20 +43,20 @@ enum class EMovieSceneDirectorBlueprintEndpointType
 /**
  * Structure for creating extra parameter pins on an enpdoint
  */
-struct MOVIESCENETOOLS_API FMovieSceneDirectorBlueprintEndpointParameter
+struct FMovieSceneDirectorBlueprintEndpointParameter
 {
 	FString PinName;
 	EEdGraphPinDirection PinDirection;
 	FName PinTypeCategory;
 	TWeakObjectPtr<UObject> PinTypeClass;
 
-	static FString SanitizePinName(const FString& InPinName);
+	static UE_API FString SanitizePinName(const FString& InPinName);
 };
 
 /*
 * Structure for holding value of a blueprint variable 
 */
-struct MOVIESCENETOOLS_API FMovieSceneDirectorBlueprintVariableValue
+struct FMovieSceneDirectorBlueprintVariableValue
 {
 	/* In the case the variable is a UObject, reference to the UObject */
 	FSoftObjectPath ObjectValue;
@@ -66,7 +68,7 @@ struct MOVIESCENETOOLS_API FMovieSceneDirectorBlueprintVariableValue
 /**
  * Parameters for creating a blueprint endpoint
  */
-struct MOVIESCENETOOLS_API FMovieSceneDirectorBlueprintEndpointDefinition
+struct FMovieSceneDirectorBlueprintEndpointDefinition
 {
 	/** Type of endpoint */
 	EMovieSceneDirectorBlueprintEndpointType EndpointType;
@@ -112,7 +114,7 @@ struct MOVIESCENETOOLS_API FMovieSceneDirectorBlueprintEndpointDefinition
 };
 
 /** Parameters for endpoint call customization callback */
-struct MOVIESCENETOOLS_API FMovieSceneCustomizeDirectorBlueprintEndpointCallParams
+struct FMovieSceneCustomizeDirectorBlueprintEndpointCallParams
 {
 	UEdGraph* FunctionGraph = nullptr;
 	UK2Node_FunctionEntry* FunctionEntryNode = nullptr;
@@ -125,7 +127,7 @@ DECLARE_DELEGATE_OneParam(FMovieSceneCustomizeDirectorBlueprintEndpointCallDeleg
 /**
  * Parameters for creating an entry point, i.e. a call to an endpoint
  */
-struct MOVIESCENETOOLS_API FMovieSceneDirectorBlueprintEndpointCall
+struct FMovieSceneDirectorBlueprintEndpointCall
 {
 	/** The endpoint node to call */
 	UEdGraphNode* Endpoint = nullptr;
@@ -143,7 +145,7 @@ struct MOVIESCENETOOLS_API FMovieSceneDirectorBlueprintEndpointCall
 /**
  * Result structure for a created blueprint entry point, i.e. a call to an endpoint
  */
-struct MOVIESCENETOOLS_API FMovieSceneDirectorBlueprintEntrypointResult
+struct FMovieSceneDirectorBlueprintEntrypointResult
 {
 	/** The name of the function that calls the endpoint */
 	FName CompiledFunctionName;
@@ -166,7 +168,7 @@ struct MOVIESCENETOOLS_API FMovieSceneDirectorBlueprintEntrypointResult
 /**
  * Static utility library for dealing with movie-scene blueprint endpoints at edit/cook-time
  */
-struct MOVIESCENETOOLS_API FMovieSceneDirectorBlueprintUtils
+struct FMovieSceneDirectorBlueprintUtils
 {
 	/**
 	 * Create a new endpoint for the specified sequence.
@@ -175,7 +177,7 @@ struct MOVIESCENETOOLS_API FMovieSceneDirectorBlueprintUtils
 	 * @param EndpointDefinition  The end point definition
 	 * @return A valid pointer to the blueprint node for the endpoint
 	 */
-	static UK2Node* CreateEndpoint(UBlueprint* Blueprint, const FMovieSceneDirectorBlueprintEndpointDefinition& EndpointDefinition);
+	static UE_API UK2Node* CreateEndpoint(UBlueprint* Blueprint, const FMovieSceneDirectorBlueprintEndpointDefinition& EndpointDefinition);
 
 	/**
 	 * Create a new event endpoint for the specified sequence.
@@ -186,7 +188,7 @@ struct MOVIESCENETOOLS_API FMovieSceneDirectorBlueprintUtils
 	 * @param EndpointDefinition  The end point definition
 	 * @return A valid pointer to the blueprint node for the endpoint
 	 */
-	static UK2Node_CustomEvent* CreateEventEndpoint(UBlueprint* Blueprint, const FMovieSceneDirectorBlueprintEndpointDefinition& EndpointDefinition);
+	static UE_API UK2Node_CustomEvent* CreateEventEndpoint(UBlueprint* Blueprint, const FMovieSceneDirectorBlueprintEndpointDefinition& EndpointDefinition);
 	
 	/**
 	 * Create a new function endpoint for the specified sequence.
@@ -195,7 +197,7 @@ struct MOVIESCENETOOLS_API FMovieSceneDirectorBlueprintUtils
 	 * @param EndpointDefinition  The end point definition
 	 * @return A valid pointer to the blueprint node for the endpoint
 	 */
-	static UK2Node_FunctionEntry* CreateFunctionEndpoint(UBlueprint* Blueprint, const FMovieSceneDirectorBlueprintEndpointDefinition& EndpointDefinition);
+	static UE_API UK2Node_FunctionEntry* CreateFunctionEndpoint(UBlueprint* Blueprint, const FMovieSceneDirectorBlueprintEndpointDefinition& EndpointDefinition);
 
 	/**
 	 * Main generation function for creating the necessary function graph for calling an event endpoint with 
@@ -205,7 +207,7 @@ struct MOVIESCENETOOLS_API FMovieSceneDirectorBlueprintUtils
 	 * @param Compiler               (Required, non-null) The kismet compiler context for generation
 	 * @param Endpoint               (Required, non-null) The endpoint node that the event is bound to
 	 */
-	static FMovieSceneDirectorBlueprintEntrypointResult GenerateEntryPoint(const FMovieSceneDirectorBlueprintEndpointCall& EndpointCall, FKismetCompilerContext* Compiler);
+	static UE_API FMovieSceneDirectorBlueprintEntrypointResult GenerateEntryPoint(const FMovieSceneDirectorBlueprintEndpointCall& EndpointCall, FKismetCompilerContext* Compiler);
 
 	/**
 	 * Attempt to locate an output pin on the specified endpoint node that matches the specified pin class
@@ -214,14 +216,15 @@ struct MOVIESCENETOOLS_API FMovieSceneDirectorBlueprintUtils
 	 * @param BoundObjectPinClass    (Required, non-null) The class of the object pin to locate
 	 * @return An output pin that is of type PC_Object with the specified class as its PinSubCategoryObject, or nullptr if one was not found
 	 */
-	static UEdGraphPin* FindCallTargetPin(UK2Node* InEndpoint, UClass* CallTargetClass);
+	static UE_API UEdGraphPin* FindCallTargetPin(UK2Node* InEndpoint, UClass* CallTargetClass);
 
 private:
 
-	static UEdGraph* FindOrCreateEventGraph(UBlueprint* Blueprint, const FString& GraphName);
+	static UE_API UEdGraph* FindOrCreateEventGraph(UBlueprint* Blueprint, const FString& GraphName);
 
-	static bool GenerateEntryPointRawActorParameter(
+	static UE_API bool GenerateEntryPointRawActorParameter(
 			FKismetCompilerContext* Compiler, UEdGraph* Graph, UEdGraphNode* Endpoint, 
 			UK2Node* OriginNode, UEdGraphPin* DestPin, const FMovieSceneDirectorBlueprintVariableValue& PayloadValue);
 };
 
+#undef UE_API
diff --git a/Engine/Source/Editor/MovieSceneTools/Public/MovieSceneDynamicBindingCustomization.h b/Engine/Source/Editor/MovieSceneTools/Public/MovieSceneDynamicBindingCustomization.h
--- a/Engine/Source/Editor/MovieSceneTools/Public/MovieSceneDynamicBindingCustomization.h
+++ b/Engine/Source/Editor/MovieSceneTools/Public/MovieSceneDynamicBindingCustomization.h
@@ -7,6 +7,8 @@
 #include "IPropertyTypeCustomization.h"
 #include "Types/SlateEnums.h"
 
+#define UE_API MOVIESCENETOOLS_API
+
 class UK2Node;
 class UBlueprint;
 class UEdGraphPin;
@@ -21,35 +23,35 @@ enum class ECheckBoxState : uint8;
 /**
  * Customization for dynamic spawn/possession endpoint picker.
  */
-class MOVIESCENETOOLS_API FMovieSceneDynamicBindingCustomization : public FMovieSceneDirectorBlueprintEndpointCustomization
+class FMovieSceneDynamicBindingCustomization : public FMovieSceneDirectorBlueprintEndpointCustomization
 {
 public:
 
-	static TSharedRef<IPropertyTypeCustomization> MakeInstance();
-	static TSharedRef<IPropertyTypeCustomization> MakeInstance(UMovieScene* InMovieScene, FGuid InObjectBinding, int32 InBindingIndex=0);
+	static UE_API TSharedRef<IPropertyTypeCustomization> MakeInstance();
+	static UE_API TSharedRef<IPropertyTypeCustomization> MakeInstance(UMovieScene* InMovieScene, FGuid InObjectBinding, int32 InBindingIndex=0);
 
-	virtual void CustomizeChildren(TSharedRef<IPropertyHandle> InPropertyHandle, IDetailChildrenBuilder& ChildBuilder, IPropertyTypeCustomizationUtils& CustomizationUtils) override;
+	UE_API virtual void CustomizeChildren(TSharedRef<IPropertyHandle> InPropertyHandle, IDetailChildrenBuilder& ChildBuilder, IPropertyTypeCustomizationUtils& CustomizationUtils) override;
 
 protected:
 
-	virtual void GetPayloadVariables(UObject* EditObject, void* RawData, FPayloadVariableMap& OutPayloadVariables) const override;
-	virtual bool SetPayloadVariable(UObject* EditObject, void* RawData, FName FieldName, const FMovieSceneDirectorBlueprintVariableValue& NewVariableValue) override;
-	virtual UK2Node* FindEndpoint(UMovieSceneSequence* Sequence, UBlueprint* Blueprint, UObject* EditObject, void* RawData) const override;
-	virtual void GetWellKnownParameterPinNames(UObject* EditObject, void* RawData, TArray<FName>& OutWellKnownParameters) const override;
-	virtual void GetWellKnownParameterCandidates(UK2Node* Endpoint, TArray<FWellKnownParameterCandidates>& OutCandidates) const override;
-	virtual bool SetWellKnownParameterPinName(UObject* EditObject, void* RawData, int32 ParameterIndex, FName BoundPinName) override;
-	virtual FMovieSceneDirectorBlueprintEndpointDefinition GenerateEndpointDefinition(UMovieSceneSequence* Sequence) override;
-	virtual void OnCreateEndpoint(UMovieSceneSequence* Sequence, UBlueprint* Blueprint, const TArray<UObject*> EditObjects, const TArray<void*> RawData, const FMovieSceneDirectorBlueprintEndpointDefinition& EndpointDefinition, UK2Node* NewEndpoint) override;
-	virtual void OnSetEndpoint(UMovieSceneSequence* Sequence, UBlueprint* Blueprint, const TArray<UObject*> EditObjects, const TArray<void*> RawData, const FMovieSceneDirectorBlueprintEndpointDefinition& EndpointDefinition, UK2Node* NewEndpoint) override;
-	virtual void GetEditObjects(TArray<UObject*>& OutObjects) const override;
-	virtual void OnCollectQuickBindActions(UBlueprint* Blueprint, FBlueprintActionMenuBuilder& MenuBuilder) override;
+	UE_API virtual void GetPayloadVariables(UObject* EditObject, void* RawData, FPayloadVariableMap& OutPayloadVariables) const override;
+	UE_API virtual bool SetPayloadVariable(UObject* EditObject, void* RawData, FName FieldName, const FMovieSceneDirectorBlueprintVariableValue& NewVariableValue) override;
+	UE_API virtual UK2Node* FindEndpoint(UMovieSceneSequence* Sequence, UBlueprint* Blueprint, UObject* EditObject, void* RawData) const override;
+	UE_API virtual void GetWellKnownParameterPinNames(UObject* EditObject, void* RawData, TArray<FName>& OutWellKnownParameters) const override;
+	UE_API virtual void GetWellKnownParameterCandidates(UK2Node* Endpoint, TArray<FWellKnownParameterCandidates>& OutCandidates) const override;
+	UE_API virtual bool SetWellKnownParameterPinName(UObject* EditObject, void* RawData, int32 ParameterIndex, FName BoundPinName) override;
+	UE_API virtual FMovieSceneDirectorBlueprintEndpointDefinition GenerateEndpointDefinition(UMovieSceneSequence* Sequence) override;
+	UE_API virtual void OnCreateEndpoint(UMovieSceneSequence* Sequence, UBlueprint* Blueprint, const TArray<UObject*> EditObjects, const TArray<void*> RawData, const FMovieSceneDirectorBlueprintEndpointDefinition& EndpointDefinition, UK2Node* NewEndpoint) override;
+	UE_API virtual void OnSetEndpoint(UMovieSceneSequence* Sequence, UBlueprint* Blueprint, const TArray<UObject*> EditObjects, const TArray<void*> RawData, const FMovieSceneDirectorBlueprintEndpointDefinition& EndpointDefinition, UK2Node* NewEndpoint) override;
+	UE_API virtual void GetEditObjects(TArray<UObject*>& OutObjects) const override;
+	UE_API virtual void OnCollectQuickBindActions(UBlueprint* Blueprint, FBlueprintActionMenuBuilder& MenuBuilder) override;
 
 private:
 
-	void SetEndpointImpl(UMovieScene* MovieScene, FMovieSceneDynamicBinding* DynamicBinding, UBlueprint* Blueprint, UK2Node* NewEndpoint);
-	void EnsureBlueprintExtensionCreated(UMovieScene* MovieScene, UBlueprint* Blueprint);
+	UE_API void SetEndpointImpl(UMovieScene* MovieScene, FMovieSceneDynamicBinding* DynamicBinding, UBlueprint* Blueprint, UK2Node* NewEndpoint);
+	UE_API void EnsureBlueprintExtensionCreated(UMovieScene* MovieScene, UBlueprint* Blueprint);
 
-	void CollectResolverLibraryBindActions(UBlueprint* Blueprint, FBlueprintActionMenuBuilder& MenuBuilder, bool bIsRebinding);
+	UE_API void CollectResolverLibraryBindActions(UBlueprint* Blueprint, FBlueprintActionMenuBuilder& MenuBuilder, bool bIsRebinding);
 
 private:
 
@@ -58,3 +60,4 @@ private:
 	int32 BindingIndex = 0;
 };
 
+#undef UE_API
diff --git a/Engine/Source/Editor/MovieSceneTools/Public/MovieSceneDynamicBindingUtils.h b/Engine/Source/Editor/MovieSceneTools/Public/MovieSceneDynamicBindingUtils.h
--- a/Engine/Source/Editor/MovieSceneTools/Public/MovieSceneDynamicBindingUtils.h
+++ b/Engine/Source/Editor/MovieSceneTools/Public/MovieSceneDynamicBindingUtils.h
@@ -15,22 +15,24 @@
 
 #include "MovieSceneDynamicBindingUtils.generated.h"
 
+#define UE_API MOVIESCENETOOLS_API
+
 /**
  * A utility class for managing dynamic binding endpoints.
  */
-struct MOVIESCENETOOLS_API FMovieSceneDynamicBindingUtils
+struct FMovieSceneDynamicBindingUtils
 {
 	using FSharedPlaybackState = UE::MovieScene::FSharedPlaybackState;
 
 	/**
 	 * Set an endpoint on the given dynamic binding.
 	 */
-	static void SetEndpoint(UMovieScene* MovieScene, FMovieSceneDynamicBinding* DynamicBinding, UK2Node* NewEndpoint);
+	static UE_API void SetEndpoint(UMovieScene* MovieScene, FMovieSceneDynamicBinding* DynamicBinding, UK2Node* NewEndpoint);
 
 	/**
 	 * Ensures that the dynamic binding blueprint extension has been added to the given sequence's director blueprint.
 	 */
-	static void EnsureBlueprintExtensionCreated(UMovieSceneSequence* MovieSceneSequence, UBlueprint* Blueprint);
+	static UE_API void EnsureBlueprintExtensionCreated(UMovieSceneSequence* MovieSceneSequence, UBlueprint* Blueprint);
 
 	/**
 	 * Utility function for iterating all dynamic bindings in a sequence.
@@ -58,8 +60,8 @@ struct MOVIESCENETOOLS_API FMovieSceneDynamicBindingUtils
 /**
  * Dummy class, used for easily getting a valid UFunction that helps prepare blueprint function graphs.
  */
-UCLASS()
-class MOVIESCENETOOLS_API UMovieSceneDynamicBindingEndpointUtil : public UObject
+UCLASS(MinimalAPI)
+class UMovieSceneDynamicBindingEndpointUtil : public UObject
 {
 	GENERATED_BODY()
 
@@ -67,3 +69,4 @@ class MOVIESCENETOOLS_API UMovieSceneDynamicBindingEndpointUtil : public UObject
 	FMovieSceneDynamicBindingResolveResult SampleResolveBinding() { return FMovieSceneDynamicBindingResolveResult(); }
 };
 
+#undef UE_API
diff --git a/Engine/Source/Editor/MovieSceneTools/Public/MovieSceneEventUtils.h b/Engine/Source/Editor/MovieSceneTools/Public/MovieSceneEventUtils.h
--- a/Engine/Source/Editor/MovieSceneTools/Public/MovieSceneEventUtils.h
+++ b/Engine/Source/Editor/MovieSceneTools/Public/MovieSceneEventUtils.h
@@ -8,6 +8,8 @@
 #include "MovieSceneDirectorBlueprintUtils.h"
 #include "UObject/WeakObjectPtrTemplates.h"
 
+#define UE_API MOVIESCENETOOLS_API
+
 class FKismetCompilerContext;
 class UBlueprint;
 class UClass;
@@ -28,17 +30,17 @@ struct FMovieSceneEvent;
 /**
  * Static utility library for dealing with movie-scene events at edit/cook-time
  */
-struct MOVIESCENETOOLS_API FMovieSceneEventUtils
+struct FMovieSceneEventUtils
 {
 	/**
 	 * Generate event endpoint definition for the given track
 	 */
-	static FMovieSceneDirectorBlueprintEndpointDefinition GenerateEventDefinition(UMovieSceneTrack* Track);
+	static UE_API FMovieSceneDirectorBlueprintEndpointDefinition GenerateEventDefinition(UMovieSceneTrack* Track);
 
 	/**
 	 * Generate event endpoint definition for the given object binding
 	 */
-	static FMovieSceneDirectorBlueprintEndpointDefinition GenerateEventDefinition(UMovieScene* MovieScene, const FGuid& ObjectBindingID);
+	static UE_API FMovieSceneDirectorBlueprintEndpointDefinition GenerateEventDefinition(UMovieScene* MovieScene, const FGuid& ObjectBindingID);
 
 
 	/**
@@ -49,7 +51,7 @@ struct MOVIESCENETOOLS_API FMovieSceneEventUtils
 	 * @param Blueprint       (Required, non-null) The blueprint within which to create the new custom event
 	 * @return A valid pointer to the blueprint node for the event endpoint
 	 */
-	static UK2Node_CustomEvent* BindNewUserFacingEvent(FMovieSceneEvent* EntryPoint, UMovieSceneEventSectionBase* EventSection, UBlueprint* Blueprint);
+	static UE_API UK2Node_CustomEvent* BindNewUserFacingEvent(FMovieSceneEvent* EntryPoint, UMovieSceneEventSectionBase* EventSection, UBlueprint* Blueprint);
 
 	/**
 	 * Main generation function for creating the necessary function graph for calling an event endpoint with the parameters specified in the entry point payload
@@ -59,7 +61,7 @@ struct MOVIESCENETOOLS_API FMovieSceneEventUtils
 	 * @param Compiler               (Required, non-null) The kismet compiler context for generation
 	 * @param Endpoint               (Required, non-null) The endpoint node that the event is bound to
 	 */
-	static UK2Node_FunctionEntry* GenerateEntryPoint(FMovieSceneEvent* EntryPoint, FKismetCompilerContext* Compiler, UEdGraphNode* Endpoint);
+	static UE_API UK2Node_FunctionEntry* GenerateEntryPoint(FMovieSceneEvent* EntryPoint, FKismetCompilerContext* Compiler, UEdGraphNode* Endpoint);
 
 	/**
 	 * Create the necessary mapping between a director Blueprint and and event section in order for the event section to be compiled correctly.
@@ -67,7 +69,7 @@ struct MOVIESCENETOOLS_API FMovieSceneEventUtils
 	 * @param EventSection           (Required, non-null) The event section to bind to the director BP
 	 * @param DirectorBP             (Required, non-null) The owning blueprint that will compile the event section's entry points
 	 */
-	static void BindEventSectionToBlueprint(UMovieSceneEventSectionBase* EventSection, UBlueprint* DirectorBP);
+	static UE_API void BindEventSectionToBlueprint(UMovieSceneEventSectionBase* EventSection, UBlueprint* DirectorBP);
 
 	/**
 	 * Remove all event end points for the given event section
@@ -75,7 +77,7 @@ struct MOVIESCENETOOLS_API FMovieSceneEventUtils
 	 * @param EventSection           (Required, non-null) The event section to remove event endpoints from
 	 * @param DirectorBP             (Required, non-null) The owning blueprint from which to remove the event section's entry points
 	 */
-	static void RemoveEndpointsForEventSection(UMovieSceneEventSectionBase* EventSection, UBlueprint* DirectorBP);
+	static UE_API void RemoveEndpointsForEventSection(UMovieSceneEventSectionBase* EventSection, UBlueprint* DirectorBP);
 
 	/**
 	 * Remove unused custom events (that don't have corresponding event end points in the given event sections)
@@ -83,7 +85,7 @@ struct MOVIESCENETOOLS_API FMovieSceneEventUtils
 	 * @param EventSections          (Required, non-null) The event sections to search for event endpoints from
 	 * @param DirectorBP             (Required, non-null) The owning blueprint from which to remove the event section's entry points
 	 */
-	static void RemoveUnusedCustomEvents(const TArray<TWeakObjectPtr<UMovieSceneEventSectionBase>>& EventSections, UBlueprint* DirectorBP);
+	static UE_API void RemoveUnusedCustomEvents(const TArray<TWeakObjectPtr<UMovieSceneEventSectionBase>>& EventSections, UBlueprint* DirectorBP);
 
 	/**
 	 * Attempt to locate the blueprint node that relates to an event's end-point.
@@ -93,7 +95,7 @@ struct MOVIESCENETOOLS_API FMovieSceneEventUtils
 	 * @param OwnerBlueprint         (Required, non-null) The blueprint to search within
 	 * @return The K2Node that relates to the specified event's end point, or nullptr if one was not found
 	 */
-	static UK2Node* FindEndpoint(FMovieSceneEvent* EntryPoint, UMovieSceneEventSectionBase* EventSection, UBlueprint* OwnerBlueprint);
+	static UE_API UK2Node* FindEndpoint(FMovieSceneEvent* EntryPoint, UMovieSceneEventSectionBase* EventSection, UBlueprint* OwnerBlueprint);
 
 	/**
 	 * Set the specified EntryPoint to be bound to a new K2 node
@@ -103,6 +105,7 @@ struct MOVIESCENETOOLS_API FMovieSceneEventUtils
 	 * @param InNewEndpoint          (Optional) The new endpoint node to call from the event
 	 * @param BoundObjectPin         (Optional) The pin that should receive the bound object when invoking the event
 	 */
-	static void SetEndpoint(FMovieSceneEvent* EntryPoint, UMovieSceneEventSectionBase* EventSection, UK2Node* InNewEndpoint, UEdGraphPin* BoundObjectPin);
+	static UE_API void SetEndpoint(FMovieSceneEvent* EntryPoint, UMovieSceneEventSectionBase* EventSection, UK2Node* InNewEndpoint, UEdGraphPin* BoundObjectPin);
 };
 
+#undef UE_API
diff --git a/Engine/Source/Editor/MovieSceneTools/Public/MovieSceneObjectBindingIDCustomization.h b/Engine/Source/Editor/MovieSceneTools/Public/MovieSceneObjectBindingIDCustomization.h
--- a/Engine/Source/Editor/MovieSceneTools/Public/MovieSceneObjectBindingIDCustomization.h
+++ b/Engine/Source/Editor/MovieSceneTools/Public/MovieSceneObjectBindingIDCustomization.h
@@ -8,6 +8,8 @@
 #include "MovieSceneSequenceID.h"
 #include "Templates/SharedPointer.h"
 
+#define UE_API MOVIESCENETOOLS_API
+
 class FDetailWidgetRow;
 class FDragDropEvent;
 class FDragDropOperation;
@@ -27,7 +29,7 @@ namespace MovieScene
 
 
 
-class MOVIESCENETOOLS_API FMovieSceneObjectBindingIDCustomization
+class FMovieSceneObjectBindingIDCustomization
 	: public IPropertyTypeCustomization
 	, FMovieSceneObjectBindingIDPicker
 {
@@ -40,24 +42,26 @@ public:
 		: FMovieSceneObjectBindingIDPicker(InLocalSequenceID, InSequencer)
 	{}
 
-	static void BindTo(TSharedRef<ISequencer> InSequencer);
+	static UE_API void BindTo(TSharedRef<ISequencer> InSequencer);
 
-	virtual void CustomizeHeader(TSharedRef<IPropertyHandle> PropertyHandle, FDetailWidgetRow& HeaderRow, IPropertyTypeCustomizationUtils& CustomizationUtils) override;
+	UE_API virtual void CustomizeHeader(TSharedRef<IPropertyHandle> PropertyHandle, FDetailWidgetRow& HeaderRow, IPropertyTypeCustomizationUtils& CustomizationUtils) override;
 	virtual void CustomizeChildren(TSharedRef<IPropertyHandle> StructPropertyHandle, class IDetailChildrenBuilder& ChildBuilder, IPropertyTypeCustomizationUtils& StructCustomizationUtils) override {}
 
 private:
 
-	virtual UMovieSceneSequence* GetSequence() const override;
+	UE_API virtual UMovieSceneSequence* GetSequence() const override;
 
-	virtual void SetCurrentValue(const FMovieSceneObjectBindingID& InBindingId) override;
+	UE_API virtual void SetCurrentValue(const FMovieSceneObjectBindingID& InBindingId) override;
 
-	virtual FMovieSceneObjectBindingID GetCurrentValue() const override;
+	UE_API virtual FMovieSceneObjectBindingID GetCurrentValue() const override;
 
-	virtual bool HasMultipleValues() const override;
+	UE_API virtual bool HasMultipleValues() const override;
 
-	FReply OnDrop(const FGeometry& InGeometry, const FDragDropEvent& InDragDropEvent);
+	UE_API FReply OnDrop(const FGeometry& InGeometry, const FDragDropEvent& InDragDropEvent);
 
-	void OnResetToDefault();
+	UE_API void OnResetToDefault();
 
 	TSharedPtr<IPropertyHandle> StructProperty;
-};
\ No newline at end of file
+};
+
+#undef UE_API
diff --git a/Engine/Source/Editor/MovieSceneTools/Public/MovieSceneObjectBindingIDPicker.h b/Engine/Source/Editor/MovieSceneTools/Public/MovieSceneObjectBindingIDPicker.h
--- a/Engine/Source/Editor/MovieSceneTools/Public/MovieSceneObjectBindingIDPicker.h
+++ b/Engine/Source/Editor/MovieSceneTools/Public/MovieSceneObjectBindingIDPicker.h
@@ -8,6 +8,8 @@
 #include "Templates/SharedPointer.h"
 #include "Textures/SlateIcon.h"
 
+#define UE_API MOVIESCENETOOLS_API
+
 class FMenuBuilder;
 class FReply;
 class ISequencer;
@@ -34,7 +36,7 @@ struct FRelativeObjectBindingID;
 /**
  * Helper class that is used to pick object bindings for movie scene data
  */
-class MOVIESCENETOOLS_API FMovieSceneObjectBindingIDPicker
+class FMovieSceneObjectBindingIDPicker
 {
 public:
 
@@ -57,7 +59,7 @@ public:
 	/**
 	 * Check whether this picker actually has anything to pick
 	 */
-	bool IsEmpty() const;
+	UE_API bool IsEmpty() const;
 
 protected:
 	virtual ~FMovieSceneObjectBindingIDPicker() { }
@@ -77,35 +79,35 @@ protected:
 protected:
 
 	/** Initialize this class - rebuilds sequence hierarchy data and available IDs from the source sequence */
-	void Initialize();
+	UE_API void Initialize();
 
 	/** Access the text that relates to the currently selected binding ID */
-	FText GetCurrentText() const;
+	UE_API FText GetCurrentText() const;
 
 	/** Access the tooltip text that relates to the currently selected binding ID */
-	FText GetToolTipText() const;
+	UE_API FText GetToolTipText() const;
 
 	/** Get the icon that represents the currently assigned binding */
-	FSlateIcon GetCurrentIcon() const;
-	const FSlateBrush* GetCurrentIconBrush() const;
+	UE_API FSlateIcon GetCurrentIcon() const;
+	UE_API const FSlateBrush* GetCurrentIconBrush() const;
 
 	/** Get the visibility for the spawnable icon overlap */ 
-	EVisibility GetSpawnableIconOverlayVisibility() const;
+	UE_API EVisibility GetSpawnableIconOverlayVisibility() const;
 
 	/** Assign a new binding ID in response to user-input */
-	void SetBindingId(UE::MovieScene::FFixedObjectBindingID InBindingId);
+	UE_API void SetBindingId(UE::MovieScene::FFixedObjectBindingID InBindingId);
 
 	/** Build menu content that allows the user to choose a binding from inside the source sequence */
-	TSharedRef<SWidget> GetPickerMenu();
+	UE_API TSharedRef<SWidget> GetPickerMenu();
 
 	/** Build menu content that allows the user to choose a binding from inside the source sequence */
-	void GetPickerMenu(FMenuBuilder& MenuBuilder);
+	UE_API void GetPickerMenu(FMenuBuilder& MenuBuilder);
 
 	/** Get a widget that represents the currently chosen item */
-	TSharedRef<SWidget> GetCurrentItemWidget(TSharedRef<STextBlock> TextContent);
+	UE_API TSharedRef<SWidget> GetCurrentItemWidget(TSharedRef<STextBlock> TextContent);
 
 	/** Get a widget that represents a warning/fixup button for this binding */
-	TSharedRef<SWidget> GetWarningWidget();
+	UE_API TSharedRef<SWidget> GetWarningWidget();
 
 	/** Optional sequencer ptr */
 	TWeakPtr<ISequencer> WeakSequencer;
@@ -114,27 +116,27 @@ protected:
 	FMovieSceneSequenceID LocalSequenceID;
 
 	/** Update the cached text, tooltip and icon */
-	void UpdateCachedData();
+	UE_API void UpdateCachedData();
 
 private:
 
 	/** Get the currently set binding ID as a fixed binding ID from the root sequence */
-	UE::MovieScene::FFixedObjectBindingID GetCurrentValueAsFixed() const;
+	UE_API UE::MovieScene::FFixedObjectBindingID GetCurrentValueAsFixed() const;
 
 	/** Set the binding ID */
-	void SetCurrentValueFromFixed(UE::MovieScene::FFixedObjectBindingID InValue);
+	UE_API void SetCurrentValueFromFixed(UE::MovieScene::FFixedObjectBindingID InValue);
 
 	/** Called when the combo box has been clicked to populate its menu content */
-	void OnGetMenuContent(FMenuBuilder& MenuBuilder, TSharedPtr<FSequenceBindingNode> Node);
+	UE_API void OnGetMenuContent(FMenuBuilder& MenuBuilder, TSharedPtr<FSequenceBindingNode> Node);
 
 	/** Get the visibility for the warning/fixup button */
-	EVisibility GetFixedWarningVisibility() const;
+	UE_API EVisibility GetFixedWarningVisibility() const;
 
 	/** Get the visibility for the warning/fixup button */
-	FReply AttemptBindingFixup();
+	UE_API FReply AttemptBindingFixup();
 
 	/** Checks if a class is allowed by this binding */
-	bool IsClassAllowed(const TSharedPtr<FSequenceBindingNode>& Node) const;
+	UE_API bool IsClassAllowed(const TSharedPtr<FSequenceBindingNode>& Node) const;
 
 	/** Cached current text and tooltips */
 	FText CurrentText, ToolTipText;
@@ -164,3 +166,5 @@ protected:
 	/** Classes that can't be picked */
 	TArray<const UClass*> DisallowedClassFilters;
 };
+
+#undef UE_API
diff --git a/Engine/Source/Editor/MovieSceneTools/Public/MovieSceneSectionDetailsCustomization.h b/Engine/Source/Editor/MovieSceneTools/Public/MovieSceneSectionDetailsCustomization.h
--- a/Engine/Source/Editor/MovieSceneTools/Public/MovieSceneSectionDetailsCustomization.h
+++ b/Engine/Source/Editor/MovieSceneTools/Public/MovieSceneSectionDetailsCustomization.h
@@ -8,13 +8,15 @@
 #include "Widgets/Input/SCheckBox.h"
 #include "IDetailCustomization.h"
 
+#define UE_API MOVIESCENETOOLS_API
+
 class IDetailLayoutBuilder;
 class UMovieScene;
 
 /**
  *  Customizes FMovieSceneSection to expose the section bounds to the UI and allow changing their bounded states.
  */
-class MOVIESCENETOOLS_API FMovieSceneSectionDetailsCustomization : public IDetailCustomization
+class FMovieSceneSectionDetailsCustomization : public IDetailCustomization
 {
 public:
 	static TSharedRef<IDetailCustomization> MakeInstance(TSharedPtr<INumericTypeInterface<double>> InNumericTypeInterface, TWeakObjectPtr<UMovieScene> InParentMovieScene)
@@ -29,7 +31,7 @@ public:
 	}
 
 	/** IDetailCustomization interface */
-	virtual void CustomizeDetails(IDetailLayoutBuilder& DetailBuilder);
+	UE_API virtual void CustomizeDetails(IDetailLayoutBuilder& DetailBuilder);
 
 private:
 	enum class ERangeBoundValueType
@@ -41,51 +43,51 @@ private:
 
 private:
 	/** Get the range start value */
-	ERangeBoundValueType GetRangeStartValue(FFrameNumber& OutValue) const;
+	UE_API ERangeBoundValueType GetRangeStartValue(FFrameNumber& OutValue) const;
 
 	/** Convert the range start into an FText for display */
-	FText OnGetRangeStartText() const;
+	UE_API FText OnGetRangeStartText() const;
 	/** Convert the range start into an FText for tooltip display */
-	FText OnGetRangeStartToolTipText() const;
+	UE_API FText OnGetRangeStartToolTipText() const;
 	/** Convert the text into a new range start */
-	void OnRangeStartTextCommitted(const FText& InText, ETextCommit::Type CommitInfo);
+	UE_API void OnRangeStartTextCommitted(const FText& InText, ETextCommit::Type CommitInfo);
 	/** Should the textbox be editable? False if we have an infinite range.  */
-	bool IsRangeStartTextboxEnabled() const;
+	UE_API bool IsRangeStartTextboxEnabled() const;
 
 	/** Determines if the range is Open, Closed, or Undetermined which can happen in the case of multi-select.  */
-	ECheckBoxState GetRangeStartBoundedState() const;
+	UE_API ECheckBoxState GetRangeStartBoundedState() const;
 	/** Sets the range to have a fixed bound or convert to an open bound. */
-	void SetRangeStartBounded(bool InbIsBounded);
+	UE_API void SetRangeStartBounded(bool InbIsBounded);
 
 
 	/** Get the FText representing the appropriate Unicode icon for the toggle button. */
-	FText GetRangeStartButtonIcon() const;
+	UE_API FText GetRangeStartButtonIcon() const;
 	/** Called by the UI when the button is pressed to toggle the current state. */
-	FReply ToggleRangeStartBounded();
+	UE_API FReply ToggleRangeStartBounded();
 
 private:
 	/** Get the range end value */
-	ERangeBoundValueType GetRangeEndValue(FFrameNumber& OutValue) const;
+	UE_API ERangeBoundValueType GetRangeEndValue(FFrameNumber& OutValue) const;
 
 	/** Convert the range end into an FText for display */
-	FText OnGetRangeEndText() const;
+	UE_API FText OnGetRangeEndText() const;
 	/** Convert the range end into an FText for tooltip display */
-	FText OnGetRangeEndToolTipText() const;
+	UE_API FText OnGetRangeEndToolTipText() const;
 	/** Convert the text into a new range start */
-	void OnRangeEndTextCommitted(const FText& InText, ETextCommit::Type CommitInfo);
+	UE_API void OnRangeEndTextCommitted(const FText& InText, ETextCommit::Type CommitInfo);
 	/** Should the textbox be editable? False if we have an infinite range.  */
-	bool IsRangeEndTextboxEnabled() const;
+	UE_API bool IsRangeEndTextboxEnabled() const;
 
 	/** Determines if the range is Open, Closed, or Undetermined which can happen in the case of multi-select.  */
-	ECheckBoxState GetRangeEndBoundedState() const;
+	UE_API ECheckBoxState GetRangeEndBoundedState() const;
 	/** Sets the range to have a fixed bound or convert to an open bound. */
-	void SetRangeEndBounded(bool InbIsBounded);
+	UE_API void SetRangeEndBounded(bool InbIsBounded);
 
 
 	/** Get the FText representing the appropriate Unicode icon for the toggle button. */
-	FText GetRangeEndButtonIcon() const;
+	UE_API FText GetRangeEndButtonIcon() const;
 	/** Called by the UI when the button is pressed to toggle the current state. */
-	FReply ToggleRangeEndBounded();
+	UE_API FReply ToggleRangeEndBounded();
 private:
 
 	/** The Numeric Type interface used to convert between display formats and internal tick resolution. */
@@ -97,3 +99,5 @@ private:
 	/** The movie scene that owns the section we're customizing. Used to find out the overall bounds for changing a section bounds from infinite -> closed. */
 	TWeakObjectPtr<UMovieScene> ParentMovieScene;
 };
+
+#undef UE_API
