Updated headers using UnrealCodeFixup to make sure dllstorage is on methods/staticvars instead of types. Part 5/n

diff --git a/Engine/Source/Editor/MovieSceneTools/Public/MovieSceneToolHelpers.h b/Engine/Source/Editor/MovieSceneTools/Public/MovieSceneToolHelpers.h
--- a/Engine/Source/Editor/MovieSceneTools/Public/MovieSceneToolHelpers.h
+++ b/Engine/Source/Editor/MovieSceneTools/Public/MovieSceneToolHelpers.h
@@ -22,6 +22,8 @@
 #include "Channels/MovieSceneChannelProxy.h"
 #include "MovieSceneSection.h"
 
+#define UE_API MOVIESCENETOOLS_API
+
 class UMovieScene;
 class UMovieSceneSequence;
 struct FMovieSceneObjectBindingID;
@@ -84,14 +86,14 @@ DECLARE_DELEGATE_TwoParams(FTickAnimationCB, float, FFrameNumber);
 DECLARE_DELEGATE(FEndAnimationCB);
 
 //Skel Mesh Recorder to set up and restore various parameters on the skelmesh
-struct MOVIESCENETOOLS_API FSkelMeshRecorderState
+struct FSkelMeshRecorderState
 {
 public:
 	FSkelMeshRecorderState() {}
 	~FSkelMeshRecorderState() {}
 
-	void Init(USkeletalMeshComponent* InComponent);
-	void FinishRecording();
+	UE_API void Init(USkeletalMeshComponent* InComponent);
+	UE_API void FinishRecording();
 
 
 public:
@@ -117,7 +119,7 @@ enum class FChannelMergeAlgorithm : uint8
 	Override,
 };
 
-class MOVIESCENETOOLS_API MovieSceneToolHelpers
+class MovieSceneToolHelpers
 {
 public:
 	/**
@@ -128,9 +130,9 @@ public:
 	 * @param bTrimLeft Trim left or trim right
 	 * @param bDeleteKeys Delete keys outside the split ranges
 	 */
-	static void TrimSection(const TSet<UMovieSceneSection*>& Sections, FQualifiedFrameTime Time, bool bTrimLeft, bool bDeleteKeys);
-	static bool CanTrimSectionLeft(const TSet<UMovieSceneSection*>& Sections, FQualifiedFrameTime Time);
-	static bool CanTrimSectionRight(const TSet<UMovieSceneSection*>& Sections, FQualifiedFrameTime Time);
+	static UE_API void TrimSection(const TSet<UMovieSceneSection*>& Sections, FQualifiedFrameTime Time, bool bTrimLeft, bool bDeleteKeys);
+	static UE_API bool CanTrimSectionLeft(const TSet<UMovieSceneSection*>& Sections, FQualifiedFrameTime Time);
+	static UE_API bool CanTrimSectionRight(const TSet<UMovieSceneSection*>& Sections, FQualifiedFrameTime Time);
 
 	/**
 	 * Trim or extend section at the given time
@@ -141,7 +143,7 @@ public:
 	 * @param bTrimOrExtendleft Trim or extend left or right
 	 * @param bDeleteKeys Delete keys outside the split ranges
 	 */
-	static void TrimOrExtendSection(UMovieSceneTrack* Track, TOptional<int32> RowIndex, FQualifiedFrameTime Time, bool bTrimOrExtendLeft, bool bDeleteKeys);
+	static UE_API void TrimOrExtendSection(UMovieSceneTrack* Track, TOptional<int32> RowIndex, FQualifiedFrameTime Time, bool bTrimOrExtendLeft, bool bDeleteKeys);
 
 	/**
 	 * Splits sections at the given time
@@ -150,10 +152,10 @@ public:
 	 * @param Time	The time at which to split
 	 * @param bDeleteKeys Delete keys outside the split ranges
 	 */
-	static void SplitSection(const TSet<UMovieSceneSection*>& Sections, FQualifiedFrameTime Time, bool bDeleteKeys);
-	static bool CanSplitSection(const TSet<UMovieSceneSection*>& Sections, FQualifiedFrameTime Time);
+	static UE_API void SplitSection(const TSet<UMovieSceneSection*>& Sections, FQualifiedFrameTime Time, bool bDeleteKeys);
+	static UE_API bool CanSplitSection(const TSet<UMovieSceneSection*>& Sections, FQualifiedFrameTime Time);
 
-	static FTransform GetTransformOriginForFocusedSequence(TSharedPtr<ISequencer> InSequencer);
+	static UE_API FTransform GetTransformOriginForFocusedSequence(TSharedPtr<ISequencer> InSequencer);
 
 	/**
 	 * Parse a shot name into its components.
@@ -166,7 +168,7 @@ public:
 	 * @param TakeNumberDigits The number of digits to pad for the take number
 	 * @return Whether the shot name was parsed successfully
 	 */
-	static bool ParseShotName(const FString& ShotName, FString& ShotPrefix, uint32& ShotNumber, uint32& TakeNumber, uint32& ShotNumberDigits, uint32& TakeNumberDigits);
+	static UE_API bool ParseShotName(const FString& ShotName, FString& ShotPrefix, uint32& ShotNumber, uint32& TakeNumber, uint32& ShotNumberDigits, uint32& TakeNumberDigits);
 
 	/**
 	 * Compose a shot name given its components.
@@ -178,7 +180,7 @@ public:
 	 * @param TakeNumberDigits The number of digits to pad for the take number
 	 * @return The composed shot name
 	 */
-	static FString ComposeShotName(const FString& ShotPrefix, uint32 ShotNumber, uint32 TakeNumber, uint32 ShotNumberDigits, uint32 TakeNumberDigits);
+	static UE_API FString ComposeShotName(const FString& ShotPrefix, uint32 ShotNumber, uint32 TakeNumber, uint32 ShotNumberDigits, uint32 TakeNumberDigits);
 
 	/**
 	 * Generate a new subsequence package
@@ -188,10 +190,10 @@ public:
 	 * @param NewShotName The new shot name
 	 * @return The new subsequence path
 	 */
-	static FString GenerateNewSubsequencePath(UMovieScene* SequenceMovieScene, const FString& SubsequenceDirectory, FString& NewShotName);
+	static UE_API FString GenerateNewSubsequencePath(UMovieScene* SequenceMovieScene, const FString& SubsequenceDirectory, FString& NewShotName);
 	
 	UE_DEPRECATED(5.3, "GenerateNewShotPath has been deprecated in favor of GenerateNewSubsequencePath that takes a given directory")
-	static FString GenerateNewShotPath(UMovieScene* SequenceMovieScene, FString& NewShotName);
+	static UE_API FString GenerateNewShotPath(UMovieScene* SequenceMovieScene, FString& NewShotName);
 
 
 	/**
@@ -201,10 +203,10 @@ public:
 	 * @param Time The time to generate the new shot name at
 	 * @return The new shot name
 	 */
-	static FString GenerateNewSubsequenceName(const TArray<UMovieSceneSection*>& AllSections, const FString& SubsequencePrefix, FFrameNumber Time);
+	static UE_API FString GenerateNewSubsequenceName(const TArray<UMovieSceneSection*>& AllSections, const FString& SubsequencePrefix, FFrameNumber Time);
 
 	UE_DEPRECATED(5.3, "GenerateNewShotName has been deprecated in favor of GenerateNewSubsequenceName that takes a given prefix")
-	static FString GenerateNewShotName(const TArray<UMovieSceneSection*>& AllSections, FFrameNumber Time);
+	static UE_API FString GenerateNewShotName(const TArray<UMovieSceneSection*>& AllSections, FFrameNumber Time);
 
 	/*
 	 * Create sequence
@@ -214,7 +216,7 @@ public:
 	 * @param SectionToDuplicate The section to duplicate.
 	 * @return The new subsequence.
 	 */
-	static UMovieSceneSequence* CreateSequence(FString& NewSequenceName, FString& NewSequencePath, UMovieSceneSubSection* SectionToDuplicate = nullptr);
+	static UE_API UMovieSceneSequence* CreateSequence(FString& NewSequenceName, FString& NewSequencePath, UMovieSceneSubSection* SectionToDuplicate = nullptr);
 
 	/**
 	 * Gather takes - level sequence assets that have the same shot prefix and shot number in the same asset path (directory)
@@ -223,7 +225,7 @@ public:
 	 * @param AssetData The gathered asset take data
 	 * @param OutCurrentTakeNumber The current take number of the section
 	 */
-	static void GatherTakes(const UMovieSceneSection* Section, TArray<FAssetData>& AssetData, uint32& OutCurrentTakeNumber);
+	static UE_API void GatherTakes(const UMovieSceneSection* Section, TArray<FAssetData>& AssetData, uint32& OutCurrentTakeNumber);
 
 	/**
 	 * Get the take number for the given asset
@@ -233,7 +235,7 @@ public:
 	 * @param OutTakeNumber The take number for the given asset
 	 * @return Whether the take number was found
 	 */
-	static bool GetTakeNumber(const UMovieSceneSection* Section, FAssetData AssetData, uint32& OutTakeNumber);
+	static UE_API bool GetTakeNumber(const UMovieSceneSection* Section, FAssetData AssetData, uint32& OutTakeNumber);
 
 	/**
 	 * Set the take number for the given asset
@@ -242,7 +244,7 @@ public:
 	 * @param InTakeNumber The take number for the given asset
 	 * @return Whether the take number could be set
 	 */
-	static bool SetTakeNumber(const UMovieSceneSection* Section, uint32 InTakeNumber);
+	static UE_API bool SetTakeNumber(const UMovieSceneSection* Section, uint32 InTakeNumber);
 
 	/**
 	 * Get the next available row index for the section so that it doesn't overlap any other sections in time.
@@ -252,7 +254,7 @@ public:
 	 * @param SectionsToDisregard Disregard checking these sections
 	 * @return The next available row index
 	 */
-	static int32 FindAvailableRowIndex(UMovieSceneTrack* InTrack, UMovieSceneSection* InSection, const TArray<UMovieSceneSection*>& SectionsToDisregard = TArray<UMovieSceneSection*>());
+	static UE_API int32 FindAvailableRowIndex(UMovieSceneTrack* InTrack, UMovieSceneSection* InSection, const TArray<UMovieSceneSection*>& SectionsToDisregard = TArray<UMovieSceneSection*>());
 
 	/**
 	 * Does this section overlap any other track section?
@@ -262,7 +264,7 @@ public:
 	 * @param SectionsToDisregard Disregard checking these sections
 	 * @return Whether this section overlaps any other track section
 	 */
-	static bool OverlapsSection(UMovieSceneTrack* InTrack, UMovieSceneSection* InSection, const TArray<UMovieSceneSection*>& SectionsToDisregard = TArray<UMovieSceneSection*>());
+	static UE_API bool OverlapsSection(UMovieSceneTrack* InTrack, UMovieSceneSection* InSection, const TArray<UMovieSceneSection*>& SectionsToDisregard = TArray<UMovieSceneSection*>());
 
 	/**
 	 * Generate a combobox for editing enum values
@@ -272,7 +274,7 @@ public:
 	 * @param OnSelectionChanged Delegate fired when selection is changed
 	 * @return The new widget
 	 */
-	static TSharedRef<SWidget> MakeEnumComboBox(const UEnum* Enum, TAttribute<int32> CurrentValue, SEnumComboBox::FOnEnumSelectionChanged OnSelectionChanged);
+	static UE_API TSharedRef<SWidget> MakeEnumComboBox(const UEnum* Enum, TAttribute<int32> CurrentValue, SEnumComboBox::FOnEnumSelectionChanged OnSelectionChanged);
 
 
 	/**
@@ -283,7 +285,7 @@ public:
 	 * @param InOpenDirectory Optional directory path to open from. If none given, a dialog will pop up to prompt the user
 	 * @return Whether the import was successful
 	 */
-	static bool ShowImportEDLDialog(UMovieScene* InMovieScene, FFrameRate InFrameRate, FString InOpenDirectory = TEXT(""));
+	static UE_API bool ShowImportEDLDialog(UMovieScene* InMovieScene, FFrameRate InFrameRate, FString InOpenDirectory = TEXT(""));
 
 	/**
 	 * Show Export EDL Dialog
@@ -295,7 +297,7 @@ public:
 	 * @param MovieExtension The movie extension for the shot filenames (ie. .avi, .mov, .mp4)
 	 * @return Whether the export was successful
 	 */
-	static bool ShowExportEDLDialog(const UMovieScene* InMovieScene, FFrameRate InFrameRate, FString InSaveDirectory = TEXT(""), int32 InHandleFrames = 8, FString InMovieExtension = TEXT(".avi"));
+	static UE_API bool ShowExportEDLDialog(const UMovieScene* InMovieScene, FFrameRate InFrameRate, FString InSaveDirectory = TEXT(""), int32 InHandleFrames = 8, FString InMovieExtension = TEXT(".avi"));
 
 	/**
 	* Import movie scene formats
@@ -306,7 +308,7 @@ public:
 	* @param InOpenDirectory Optional directory path to open from. If none given, a dialog will pop up to prompt the user
 	* @return Whether the import was successful
 	*/
-	static bool MovieSceneTranslatorImport(FMovieSceneImporter* InImporter, UMovieScene* InMovieScene, FFrameRate InFrameRate, FString InOpenDirectory = TEXT(""));
+	static UE_API bool MovieSceneTranslatorImport(FMovieSceneImporter* InImporter, UMovieScene* InMovieScene, FFrameRate InFrameRate, FString InOpenDirectory = TEXT(""));
 
 	/**
 	* Export movie scene formats
@@ -318,7 +320,7 @@ public:
 	* @param InHandleFrames The number of handle frames to include for each shot.
 	* @return Whether the export was successful
 	*/
-	static bool MovieSceneTranslatorExport(FMovieSceneExporter* InExporter, const UMovieScene* InMovieScene, const FMovieSceneCaptureSettings& Settings);
+	static UE_API bool MovieSceneTranslatorExport(FMovieSceneExporter* InExporter, const UMovieScene* InMovieScene, const FMovieSceneCaptureSettings& Settings);
 
 	/** 
 	* Log messages and display error message window for MovieScene translators
@@ -327,7 +329,7 @@ public:
 	* @param InContext The context used to gather error, warning or info messages during import or export.
 	* @param bDisplayMessages Whether to open the message log window after adding the message.
 	*/
-	static void MovieSceneTranslatorLogMessages(FMovieSceneTranslator* InTranslator, TSharedRef<FMovieSceneTranslatorContext> InContext, bool bDisplayMessages);
+	static UE_API void MovieSceneTranslatorLogMessages(FMovieSceneTranslator* InTranslator, TSharedRef<FMovieSceneTranslatorContext> InContext, bool bDisplayMessages);
 
 	/**
 	* Log error output for MovieScene translators
@@ -335,7 +337,7 @@ public:
 	* @param InTranslator The movie scene importer or exporter.
 	* @param InContext The context used to gather error, warning or info messages during import or export.
 	*/
-	static void MovieSceneTranslatorLogOutput(FMovieSceneTranslator* InTranslator, TSharedRef<FMovieSceneTranslatorContext> InContext);
+	static UE_API void MovieSceneTranslatorLogOutput(FMovieSceneTranslator* InTranslator, TSharedRef<FMovieSceneTranslatorContext> InContext);
 
 	/**
 	* Export FBX
@@ -349,10 +351,10 @@ public:
 	* @param InFBXFileName the fbx file name.
 	* @return Whether the export was successful
 	*/
-	static bool ExportFBX(UWorld* World, const FAnimExportSequenceParameters& AnimExportSequenceParameters, const TArray<FGuid>& Bindings, const TArray<UMovieSceneTrack*>& Tracks, INodeNameAdapter& NodeNameAdapter, FMovieSceneSequenceIDRef& Template, const FString& InFBXFileName);
+	static UE_API bool ExportFBX(UWorld* World, const FAnimExportSequenceParameters& AnimExportSequenceParameters, const TArray<FGuid>& Bindings, const TArray<UMovieSceneTrack*>& Tracks, INodeNameAdapter& NodeNameAdapter, FMovieSceneSequenceIDRef& Template, const FString& InFBXFileName);
 
 	UE_DEPRECATED(5.5, "ExportFBX taking movie scene has been deprecated in favor of a new function that takes current and root movie scene sequences")
-	static bool ExportFBX(UWorld* World, UMovieScene* MovieScene, IMovieScenePlayer* Player, const TArray<FGuid>& Bindings, const TArray<UMovieSceneTrack*>& Tracks, INodeNameAdapter& NodeNameAdapter, FMovieSceneSequenceIDRef& Template,  const FString& InFBXFileName, FMovieSceneSequenceTransform& RootToLocalTransform);
+	static UE_API bool ExportFBX(UWorld* World, UMovieScene* MovieScene, IMovieScenePlayer* Player, const TArray<FGuid>& Bindings, const TArray<UMovieSceneTrack*>& Tracks, INodeNameAdapter& NodeNameAdapter, FMovieSceneSequenceIDRef& Template,  const FString& InFBXFileName, FMovieSceneSequenceTransform& RootToLocalTransform);
 
 	/**
 	* Import FBX with dialog
@@ -362,7 +364,7 @@ public:
 	* @param bCreateCameras Whether to allow creation of cameras if found in the fbx file.
 	* @return Whether the import was successful
 	*/
-	static bool ImportFBXWithDialog(UMovieSceneSequence* InSequence, ISequencer& InSequencer, const TMap<FGuid, FString>& InObjectBindingNameMap, TOptional<bool> bCreateCameras);
+	static UE_API bool ImportFBXWithDialog(UMovieSceneSequence* InSequence, ISequencer& InSequencer, const TMap<FGuid, FString>& InObjectBindingNameMap, TOptional<bool> bCreateCameras);
 
 	/**
 	* Get FBX Ready for Import. This make sure the passed in file make be imported. After calling this call ImportReadiedFbx. It returns out some parameters that we forcably change so we reset them later.
@@ -372,7 +374,7 @@ public:
 	* @param OutFBXParams Paremter to pass back to ImportReadiedFbx
 	* @return Whether the fbx file was ready and is ready to be import.
 	*/
-	static bool ReadyFBXForImport(const FString&  ImportFilename, UMovieSceneUserImportFBXSettings* ImportFBXSettings, FFBXInOutParameters& OutFBXParams);
+	static UE_API bool ReadyFBXForImport(const FString&  ImportFilename, UMovieSceneUserImportFBXSettings* ImportFBXSettings, FFBXInOutParameters& OutFBXParams);
 
 	/**
 	* Import into an FBX scene that has been readied already, via the ReadyFBXForImport call. We do this as two pass in case the client want's to do something, like create camera's, before actually
@@ -387,7 +389,7 @@ public:
 	* @return Whether the fbx file was ready and is ready to be import.
 	*/
 
-	static bool ImportFBXIfReady(UWorld* World, UMovieSceneSequence* InSequence, IMovieScenePlayer* Player, FMovieSceneSequenceIDRef TemplateID, TMap<FGuid, FString>& ObjectBindingMap, UMovieSceneUserImportFBXSettings* ImportFBXSettings,
+	static UE_API bool ImportFBXIfReady(UWorld* World, UMovieSceneSequence* InSequence, IMovieScenePlayer* Player, FMovieSceneSequenceIDRef TemplateID, TMap<FGuid, FString>& ObjectBindingMap, UMovieSceneUserImportFBXSettings* ImportFBXSettings,
 		const FFBXInOutParameters& InFBXParams, ISequencer* Sequencer = nullptr);
 
 	/**
@@ -403,7 +405,7 @@ public:
 	* @return Whether the import was successful
 	*/
 
-	static void ImportFBXCameraToExisting(UnFbx::FFbxImporter* FbxImporter, UMovieSceneSequence* InSequence, IMovieScenePlayer* Player, FMovieSceneSequenceIDRef TemplateID, TMap<FGuid, FString>& InObjectBindingMap, bool bMatchByNameOnly, bool bNotifySlate);
+	static UE_API void ImportFBXCameraToExisting(UnFbx::FFbxImporter* FbxImporter, UMovieSceneSequence* InSequence, IMovieScenePlayer* Player, FMovieSceneSequenceIDRef TemplateID, TMap<FGuid, FString>& InObjectBindingMap, bool bMatchByNameOnly, bool bNotifySlate);
 
 	/**
 	* Import FBX Node to existing actor/node
@@ -416,7 +418,7 @@ public:
 	* @param ObjectBinding Guid of the object we are importing onto.
 	* @return Whether the import was successful
 	*/
-	static bool ImportFBXNode(FString NodeName, UnFbx::FFbxCurvesAPI& CurveAPI, UMovieSceneSequence* InSequence, IMovieScenePlayer* Player, FMovieSceneSequenceIDRef TemplateID, FGuid ObjectBinding, ISequencer* Sequencer = nullptr);
+	static UE_API bool ImportFBXNode(FString NodeName, UnFbx::FFbxCurvesAPI& CurveAPI, UMovieSceneSequence* InSequence, IMovieScenePlayer* Player, FMovieSceneSequenceIDRef TemplateID, FGuid ObjectBinding, ISequencer* Sequencer = nullptr);
 
 	/**
 	 * Lock the given camera actor to the viewport
@@ -424,7 +426,7 @@ public:
 	 * @param Sequencer The Sequencer to set the camera cut enabled for
 	 * @param CameraActor The camera actor to lock
 	 */
-	static void LockCameraActorToViewport(const TSharedPtr<ISequencer>& Sequencer, ACameraActor* CameraActor);
+	static UE_API void LockCameraActorToViewport(const TSharedPtr<ISequencer>& Sequencer, ACameraActor* CameraActor);
 
 	/**
 	 * Create a new camera cut section for the given camera
@@ -433,7 +435,7 @@ public:
 	 * @param CameraGuid CameraGuid  Guid of the camera that was added.
 	 * @param FrameNumber FrameNumber it's added at.
 	 */
-	static void CreateCameraCutSectionForCamera(UMovieScene* MovieScene, FGuid CameraGuid, FFrameNumber FrameNumber);
+	static UE_API void CreateCameraCutSectionForCamera(UMovieScene* MovieScene, FGuid CameraGuid, FFrameNumber FrameNumber);
 
 	/**
 	* Import FBX Camera to existing camera's
@@ -441,7 +443,7 @@ public:
 	* @param CameraNode The Fbx camera
 	* @param InCameraActor UE actor
 	*/
-	static void CopyCameraProperties(fbxsdk::FbxCamera* CameraNode, AActor* InCameraActor);
+	static UE_API void CopyCameraProperties(fbxsdk::FbxCamera* CameraNode, AActor* InCameraActor);
 
 
 	/*
@@ -454,11 +456,11 @@ public:
 	 * @return Whether or not it succeeds
 
 	*/
-	static bool ExportToAnimSequence(UAnimSequence* AnimSequence, UAnimSeqExportOption* ExportOptions, const FAnimExportSequenceParameters& AnimExportSequenceParameters,
+	static UE_API bool ExportToAnimSequence(UAnimSequence* AnimSequence, UAnimSeqExportOption* ExportOptions, const FAnimExportSequenceParameters& AnimExportSequenceParameters,
 			USkeletalMeshComponent* SkelMesh);
 
 	UE_DEPRECATED(5.5, "ExportToAnimSequence taking movie scene has been deprecated in favor of a new function that takes current and root movie scene sequences")
-	static bool ExportToAnimSequence(UAnimSequence* AnimSequence, UAnimSeqExportOption* ExportOptions, UMovieScene* MovieScene, IMovieScenePlayer* Player,
+	static UE_API bool ExportToAnimSequence(UAnimSequence* AnimSequence, UAnimSeqExportOption* ExportOptions, UMovieScene* MovieScene, IMovieScenePlayer* Player,
 		USkeletalMeshComponent* SkelMesh, FMovieSceneSequenceIDRef& Template, FMovieSceneSequenceTransform& RootToLocalTransform);
 
 	/*
@@ -479,11 +481,11 @@ public:
 	
 	*/
 
-	static bool BakeToSkelMeshToCallbacks(const FAnimExportSequenceParameters& AnimExportSequenceParameters, USkeletalMeshComponent* SkelMesh, UAnimSeqExportOption* ExportOptions,
+	static UE_API bool BakeToSkelMeshToCallbacks(const FAnimExportSequenceParameters& AnimExportSequenceParameters, USkeletalMeshComponent* SkelMesh, UAnimSeqExportOption* ExportOptions,
 		FInitAnimationCB InitCallback, FStartAnimationCB StartCallback, FTickAnimationCB TickCallback, FEndAnimationCB EndCallback);
 
 	UE_DEPRECATED(5.5, "BakeToSkelMeshToCallbacks taking movie scene has been deprecated in favor of a new function that takes current and root movie scene sequences")
-	static bool BakeToSkelMeshToCallbacks(UMovieScene* MovieScene, IMovieScenePlayer* Player,
+	static UE_API bool BakeToSkelMeshToCallbacks(UMovieScene* MovieScene, IMovieScenePlayer* Player,
 		USkeletalMeshComponent* SkelMesh, FMovieSceneSequenceIDRef& Template, FMovieSceneSequenceTransform& RootToLocalTransform, UAnimSeqExportOption* ExportOptions,
 		FInitAnimationCB InitCallback, FStartAnimationCB StartCallback, FTickAnimationCB TickCallback, FEndAnimationCB EndCallback);
 
@@ -491,7 +493,7 @@ public:
 	/*
 	 * @return Whether this object class has hidden mobility and can't be animated
 	 */
-	static bool HasHiddenMobility(const UClass* ObjectClass);
+	static UE_API bool HasHiddenMobility(const UClass* ObjectClass);
 	
 	/*
 	* Get the Active EvaluationTrack for a given track. Will do a recompile if the track isn't valid
@@ -499,7 +501,7 @@ public:
 	*@aram Track The movie scene track whose evaluation counterpart we want
 	*@return Returns the evaluation track for the given movie scene track. May do a re-compile if needed.
 	*/
-	static const FMovieSceneEvaluationTrack* GetEvaluationTrack(ISequencer *Sequencer, const FGuid& TrackSignature);
+	static UE_API const FMovieSceneEvaluationTrack* GetEvaluationTrack(ISequencer *Sequencer, const FGuid& TrackSignature);
 
 
 	/*
@@ -511,13 +513,13 @@ public:
 	*@param KeyRot The rotation at this time
 	*@param Sequencer The Sequence that owns this track
 	*/
-	static void GetLocationAtTime(const FMovieSceneEvaluationTrack* Track, UObject* Object, FFrameTime KeyTime, FVector& KeyPos, FRotator& KeyRot, const TSharedPtr<ISequencer>& Sequencer);
+	static UE_API void GetLocationAtTime(const FMovieSceneEvaluationTrack* Track, UObject* Object, FFrameTime KeyTime, FVector& KeyPos, FRotator& KeyRot, const TSharedPtr<ISequencer>& Sequencer);
 	
 	/* Get the Parents (Scene/Actors) of this object.
 	* @param Parents Returned Parents
 	* @param InObject Object to find parents for
 	*/
-	static void GetParents(TArray<const UObject*>& Parents, const UObject* InObject);
+	static UE_API void GetParents(TArray<const UObject*>& Parents, const UObject* InObject);
 	
 	/* Return Reference Frame from the passed in paretns
 	* @param Sequencer The Sequence that's driving these parents.
@@ -525,7 +527,7 @@ public:
 	* @param KeyTime Time to Evaluate At
 	* @return Returns Reference Transform.
 	*/
-	static FTransform GetRefFrameFromParents(const TSharedPtr<ISequencer>& Sequencer, const TArray<const UObject*>& Parents, FFrameTime KeyTime);
+	static UE_API FTransform GetRefFrameFromParents(const TSharedPtr<ISequencer>& Sequencer, const TArray<const UObject*>& Parents, FFrameTime KeyTime);
 
 	/* Return Return ParentTm for current Parent Object
 	* @param CurrentRefTM Current Referemnce TM
@@ -534,38 +536,38 @@ public:
 	* @param KeyTime Time to Evaluate At
 	* @return Returns true if succesful in evaluating the parent in the sequencer and getting a transform.
 	*/
-	static bool GetParentTM(FTransform& CurrentRefTM, const TSharedPtr<ISequencer>& Sequencer, UObject* ParentObject, FFrameTime KeyTime);
+	static UE_API bool GetParentTM(FTransform& CurrentRefTM, const TSharedPtr<ISequencer>& Sequencer, UObject* ParentObject, FFrameTime KeyTime);
 
 	/*
 	 * Get the fbx cameras from the requested parent node
 	 */
-	static void GetCameras(fbxsdk::FbxNode* Parent, TArray<fbxsdk::FbxCamera*>& Cameras);
+	static UE_API void GetCameras(fbxsdk::FbxNode* Parent, TArray<fbxsdk::FbxCamera*>& Cameras);
 
 	/*
 	 * Get the fbx camera name
 	 */
-	static FString GetCameraName(fbxsdk::FbxCamera* InCamera);
+	static UE_API FString GetCameraName(fbxsdk::FbxCamera* InCamera);
 
 	/*
 	 * Import FBX into Control Rig Channels With Dialog
 	 */
-	static bool ImportFBXIntoControlRigChannelsWithDialog(const TSharedRef<ISequencer>& InSequencer, TArray<FRigControlFBXNodeAndChannels>* NodeAndChannels);
+	static UE_API bool ImportFBXIntoControlRigChannelsWithDialog(const TSharedRef<ISequencer>& InSequencer, TArray<FRigControlFBXNodeAndChannels>* NodeAndChannels);
 
 	/*
 	** Export FBX from Control Rig Channels With Dialog
 	*/
-	static bool ExportFBXFromControlRigChannelsWithDialog(const TSharedRef<ISequencer>& InSequencer, UMovieSceneTrack* Track);
+	static UE_API bool ExportFBXFromControlRigChannelsWithDialog(const TSharedRef<ISequencer>& InSequencer, UMovieSceneTrack* Track);
 
 	/*
 	* Import FBX into Control Rig Channels
 	*/	
-	static bool ImportFBXIntoControlRigChannels(UMovieScene* MovieScene, const FString& ImportFilename,  UMovieSceneUserImportFBXControlRigSettings *ControlRigSettings,
+	static UE_API bool ImportFBXIntoControlRigChannels(UMovieScene* MovieScene, const FString& ImportFilename,  UMovieSceneUserImportFBXControlRigSettings *ControlRigSettings,
 		TArray<FRigControlFBXNodeAndChannels>* NodeAndChannels, const TArray<FName>& SelectedControlNames, FFrameRate FrameRate);
 
 	/*
 	** Export FBX from Control Rig Channels
 	*/	
-	static bool ExportFBXFromControlRigChannels(const UMovieSceneSection* Section, const UMovieSceneUserExportFBXControlRigSettings* ExportFBXControlRigSettings,
+	static UE_API bool ExportFBXFromControlRigChannels(const UMovieSceneSection* Section, const UMovieSceneUserExportFBXControlRigSettings* ExportFBXControlRigSettings,
 	                                            const TArray<FName>& SelectedControlNames, const FMovieSceneSequenceTransform& RootToLocalTransform);
 
 	/*
@@ -573,14 +575,14 @@ public:
 	* @param BoundObject Object to get SkeletalMeshComponent from.If actor checks it's components, if component checks itself then child components.
 	* @return Returns the USkeletalMeshComponent if one is found
 	*/
-	static USkeletalMeshComponent* AcquireSkeletalMeshFromObject(UObject* BoundObject);
+	static UE_API USkeletalMeshComponent* AcquireSkeletalMeshFromObject(UObject* BoundObject);
 	
 	/*
 	* Get an actors and possible component parents.
 	* @param InActorAndComponent Actor and possible component to find parents for
 	* @param OutParentActors Returns an array of parents
 	*/
-	static void GetActorParents(const FActorForWorldTransforms& Actor,
+	static UE_API void GetActorParents(const FActorForWorldTransforms& Actor,
 		TArray<FActorForWorldTransforms>& OutParentActors);
 
 	/*
@@ -589,7 +591,7 @@ public:
 	* @param InActorAndComponent Actor and possible component to find parents for
 	* @param OutParentActors Returns an array of parents
 	*/
-	static void GetActorParentsWithAttachments(ISequencer* Sequencer, const FActorForWorldTransforms& Actor, TArray<FActorForWorldTransforms>& OutParentActors);
+	static UE_API void GetActorParentsWithAttachments(ISequencer* Sequencer, const FActorForWorldTransforms& Actor, TArray<FActorForWorldTransforms>& OutParentActors);
 
 	/*
 	*  Get an actors and it's parent key frames
@@ -599,7 +601,7 @@ public:
 	* @param EndFrame The last frame to stop looking for keys
 	* @param OutFrameMap Sorted map of the frame times found
 	*/
-	static void GetActorsAndParentsKeyFrames(ISequencer* Sequencer, const FActorForWorldTransforms& Actor,
+	static UE_API void GetActorsAndParentsKeyFrames(ISequencer* Sequencer, const FActorForWorldTransforms& Actor,
 		const FFrameNumber& StartFrame, const FFrameNumber& EndFrame, TSortedMap<FFrameNumber, FFrameNumber>& OutFrameMap);
 
 	/*
@@ -609,7 +611,7 @@ public:
 	* @param Frames The times we want to get the world transforms
 	* @param OutWorldTransforms The calculated world transforms, one for each specified frame.
 	*/
-	static void GetActorWorldTransforms(ISequencer* Sequencer, const FActorForWorldTransforms& Actors, const TArray<FFrameNumber>& Frames, TArray<FTransform>& OutWorldTransforms);
+	static UE_API void GetActorWorldTransforms(ISequencer* Sequencer, const FActorForWorldTransforms& Actors, const TArray<FFrameNumber>& Frames, TArray<FTransform>& OutWorldTransforms);
 
 	/* Set or add a key onto a float channel.
 	* @param ChannelData Channel to set or add
@@ -617,7 +619,7 @@ public:
 	* @param Value  Value to Set
 	* @param Interpolation Key type to set if added
 	*/
-	static void SetOrAddKey(TMovieSceneChannelData<FMovieSceneFloatValue>& ChannelData, FFrameNumber Time, float Value, const EMovieSceneKeyInterpolation Interpolation = EMovieSceneKeyInterpolation::Auto);
+	static UE_API void SetOrAddKey(TMovieSceneChannelData<FMovieSceneFloatValue>& ChannelData, FFrameNumber Time, float Value, const EMovieSceneKeyInterpolation Interpolation = EMovieSceneKeyInterpolation::Auto);
 
 	/* Set or add a key onto a double channel.
 	* @param ChannelData Channel to set or add
@@ -625,25 +627,25 @@ public:
 	* @param Value  Value to Set
 	* @param Interpolation Key type to set if added
 	*/
-	static void SetOrAddKey(TMovieSceneChannelData<FMovieSceneDoubleValue>& ChannelData, FFrameNumber Time, double Value, const EMovieSceneKeyInterpolation Interpolation = EMovieSceneKeyInterpolation::Auto);
+	static UE_API void SetOrAddKey(TMovieSceneChannelData<FMovieSceneDoubleValue>& ChannelData, FFrameNumber Time, double Value, const EMovieSceneKeyInterpolation Interpolation = EMovieSceneKeyInterpolation::Auto);
 
 
 	/*
 	* Set or add a key onto a float channel based on key value.
 	*/
-	static void SetOrAddKey(TMovieSceneChannelData<FMovieSceneFloatValue>& Curve, FFrameNumber Time, const FMovieSceneFloatValue& Value);
+	static UE_API void SetOrAddKey(TMovieSceneChannelData<FMovieSceneFloatValue>& Curve, FFrameNumber Time, const FMovieSceneFloatValue& Value);
 
 
 	/*
 	* Set or add a key onto a double channel based on key value.
 	*/
-	static void SetOrAddKey(TMovieSceneChannelData<FMovieSceneDoubleValue>& ChannelData, FFrameNumber Time, FMovieSceneDoubleValue Value);
+	static UE_API void SetOrAddKey(TMovieSceneChannelData<FMovieSceneDoubleValue>& ChannelData, FFrameNumber Time, FMovieSceneDoubleValue Value);
 
 
 	/* 
 	* Set or add a key onto a float channel based on rich curve data.
 	*/
-	static void SetOrAddKey(TMovieSceneChannelData<FMovieSceneFloatValue>& Curve, FFrameNumber Time, float Value, 
+	static UE_API void SetOrAddKey(TMovieSceneChannelData<FMovieSceneFloatValue>& Curve, FFrameNumber Time, float Value, 
 			float ArriveTangent, float LeaveTangent, ERichCurveInterpMode InterpMode, ERichCurveTangentMode TangentMode,
 			FFrameRate FrameRate, ERichCurveTangentWeightMode WeightedMode = RCTWM_WeightedNone, 
 			float ArriveTangentWeight = 0.0f, float LeaveTangentWeight = 0.0f);
@@ -651,7 +653,7 @@ public:
 	/*
 	* Set or add a key onto a double channel based on rich curve data.
 	*/
-	static void SetOrAddKey(TMovieSceneChannelData<FMovieSceneDoubleValue>& Curve, FFrameNumber Time, double Value, 
+	static UE_API void SetOrAddKey(TMovieSceneChannelData<FMovieSceneDoubleValue>& Curve, FFrameNumber Time, double Value, 
 			float ArriveTangent, float LeaveTangent, ERichCurveInterpMode InterpMode, ERichCurveTangentMode TangentMode,
 			FFrameRate FrameRate, ERichCurveTangentWeightMode WeightedMode = RCTWM_WeightedNone, 
 			float ArriveTangentWeight = 0.0f, float LeaveTangentWeight = 0.0f);
@@ -665,12 +667,12 @@ public:
 	* @param Frames The times we want to get the world transforms
 	* @param OutWorldTransforms The calculated world transforms, one for each specified frame.
 	*/
-	static void GetActorWorldTransforms(IMovieScenePlayer* Player, UMovieSceneSequence* InSequence, FMovieSceneSequenceIDRef Template,const FActorForWorldTransforms& Actors, const TArray<FFrameNumber>& Frames, TArray<FTransform>& OutWorldTransforms);
+	static UE_API void GetActorWorldTransforms(IMovieScenePlayer* Player, UMovieSceneSequence* InSequence, FMovieSceneSequenceIDRef Template,const FActorForWorldTransforms& Actors, const TArray<FFrameNumber>& Frames, TArray<FTransform>& OutWorldTransforms);
 
 	/*
 	 * Return whether this asset is valid for the given sequence
 	 */
-	static bool IsValidAsset(UMovieSceneSequence* Sequence, const FAssetData& InAssetData);
+	static UE_API bool IsValidAsset(UMovieSceneSequence* Sequence, const FAssetData& InAssetData);
 
 	/*
 	* Collapse all of the sections specified onto the first one
@@ -680,7 +682,7 @@ public:
 	* @param InSettings Baking settings ussed
 	* @return Return true if succeeds false if it doesn't
 	*/
-	static bool CollapseSection(TSharedPtr<ISequencer>& InSequencer, UMovieSceneTrack* InOwnerTrack, TArray<UMovieSceneSection*> InSections,
+	static UE_API bool CollapseSection(TSharedPtr<ISequencer>& InSequencer, UMovieSceneTrack* InOwnerTrack, TArray<UMovieSceneSection*> InSections,
 		const FBakingAnimationKeySettings& InSettings);
 
 	/*
@@ -689,7 +691,7 @@ public:
 	* @param OutSections All non-BlendType sections
 	* @param OutAbsoluteSections All BlendType sections
 	*/
-	static void SplitSectionsByBlendType(EMovieSceneBlendType BlendType, const TArray<UMovieSceneSection*>& InSections, TArray<UMovieSceneSection*>& OutSections, TArray<UMovieSceneSection*>& OutBlendTypeSections);
+	static UE_API void SplitSectionsByBlendType(EMovieSceneBlendType BlendType, const TArray<UMovieSceneSection*>& InSections, TArray<UMovieSceneSection*>& OutSections, TArray<UMovieSceneSection*>& OutBlendTypeSections);
 
 	/*
 	* Get the channel values at the specified time with the specified sections, will get
@@ -1151,10 +1153,10 @@ public:
 	}
 
 
-	static bool OptimizeSection(const FKeyDataOptimizationParams& InParams, UMovieSceneSection* InSection);
+	static UE_API bool OptimizeSection(const FKeyDataOptimizationParams& InParams, UMovieSceneSection* InSection);
 
 	/** Returns the frame numbers between start and end. */
-	static void CalculateFramesBetween(
+	static UE_API void CalculateFramesBetween(
 		const UMovieScene* MovieScene,
 		FFrameNumber StartFrame,
 		FFrameNumber EndFrame,
@@ -1162,33 +1164,33 @@ public:
 		TArray<FFrameNumber>& OutFrames);
 
 	/** Returns the transform section for that guid. */
-	static UMovieScene3DTransformSection* GetTransformSection(
+	static UE_API UMovieScene3DTransformSection* GetTransformSection(
 		const ISequencer* InSequencer,
 		const FGuid& InGuid,
 		const FTransform& InDefaultTransform = FTransform::Identity);
 
 	/** Adds transform keys to the section based on the channels filters. */
-	static bool AddTransformKeys(
+	static UE_API bool AddTransformKeys(
 		const UMovieScene3DTransformSection* InTransformSection,
 		const TArray<FFrameNumber>& Frames,
 		const TArray<FTransform>& InLocalTransforms,
 		const EMovieSceneTransformChannel& InChannels);
 
 	/** Import an animation sequence's root transforms into a transform section */
-	static void ImportAnimSequenceTransforms(const TSharedRef<ISequencer>& Sequencer, const FAssetData& Asset, UMovieScene3DTransformTrack* TransformTrack);
+	static UE_API void ImportAnimSequenceTransforms(const TSharedRef<ISequencer>& Sequencer, const FAssetData& Asset, UMovieScene3DTransformTrack* TransformTrack);
 
 	/** Import an animation sequence's root transforms into a transform section */
-	static void ImportAnimSequenceTransformsEnterPressed(const TSharedRef<ISequencer>& Sequencer, const TArray<FAssetData>& Asset, UMovieScene3DTransformTrack* TransformTrack);
+	static UE_API void ImportAnimSequenceTransformsEnterPressed(const TSharedRef<ISequencer>& Sequencer, const TArray<FAssetData>& Asset, UMovieScene3DTransformTrack* TransformTrack);
 };
 
 // Helper to make spawnables persist throughout the export process and then restore properly afterwards
-struct MOVIESCENETOOLS_API FSpawnableRestoreState
+struct FSpawnableRestoreState
 {
 	UE_DEPRECATED(5.5, "This constructor is deprecated in favor of passing in shared playback state")
-	FSpawnableRestoreState(UMovieScene* MovieScene);
+	UE_API FSpawnableRestoreState(UMovieScene* MovieScene);
 
-	FSpawnableRestoreState(UMovieScene* MovieScene, TSharedPtr<UE::MovieScene::FSharedPlaybackState> InSharedPlaybackState);
-	~FSpawnableRestoreState();
+	UE_API FSpawnableRestoreState(UMovieScene* MovieScene, TSharedPtr<UE::MovieScene::FSharedPlaybackState> InSharedPlaybackState);
+	UE_API ~FSpawnableRestoreState();
 
 	bool bWasChanged;
 	TMap<FGuid, ESpawnOwnership> SpawnOwnershipMap;
@@ -1223,3 +1225,4 @@ private:
 };
 
 
+#undef UE_API
diff --git a/Engine/Source/Editor/MovieSceneTools/Public/MovieSceneToolsModule.h b/Engine/Source/Editor/MovieSceneTools/Public/MovieSceneToolsModule.h
--- a/Engine/Source/Editor/MovieSceneTools/Public/MovieSceneToolsModule.h
+++ b/Engine/Source/Editor/MovieSceneTools/Public/MovieSceneToolsModule.h
@@ -10,6 +10,8 @@
 #include "Containers/UnrealString.h"
 #include "IMovieSceneTools.h"
 
+#define UE_API MOVIESCENETOOLS_API
+
 struct FAssetData;
 
 class UK2Node;
@@ -49,7 +51,7 @@ public:
 /**
 * Implements the MovieSceneTools module.
 */
-class MOVIESCENETOOLS_API FMovieSceneToolsModule
+class FMovieSceneToolsModule
 	: public IMovieSceneTools
 {
 public:
@@ -60,46 +62,46 @@ public:
 	}
 
 	// IModuleInterface interface
-	virtual void StartupModule() override;
-	virtual void ShutdownModule() override;
+	UE_API virtual void StartupModule() override;
+	UE_API virtual void ShutdownModule() override;
 
-	void RegisterAnimationBakeHelper(IMovieSceneToolsAnimationBakeHelper* BakeHelper);
-	void UnregisterAnimationBakeHelper(IMovieSceneToolsAnimationBakeHelper* BakeHelper);
+	UE_API void RegisterAnimationBakeHelper(IMovieSceneToolsAnimationBakeHelper* BakeHelper);
+	UE_API void UnregisterAnimationBakeHelper(IMovieSceneToolsAnimationBakeHelper* BakeHelper);
 	const TArray<IMovieSceneToolsAnimationBakeHelper*>& GetAnimationBakeHelpers() { return BakeHelpers; }
-	void RegisterTakeData(IMovieSceneToolsTakeData*);
-	void UnregisterTakeData(IMovieSceneToolsTakeData*);
+	UE_API void RegisterTakeData(IMovieSceneToolsTakeData*);
+	UE_API void UnregisterTakeData(IMovieSceneToolsTakeData*);
 
-	bool GatherTakes(const UMovieSceneSection* Section, TArray<FAssetData>& AssetData, uint32& OutCurrentTakeNumber);
-	bool GetTakeNumber(const UMovieSceneSection* Section, FAssetData AssetData, uint32& OutTakeNumber);
-	bool SetTakeNumber(const UMovieSceneSection* Section, uint32 InTakeNumber);
+	UE_API bool GatherTakes(const UMovieSceneSection* Section, TArray<FAssetData>& AssetData, uint32& OutCurrentTakeNumber);
+	UE_API bool GetTakeNumber(const UMovieSceneSection* Section, FAssetData AssetData, uint32& OutTakeNumber);
+	UE_API bool SetTakeNumber(const UMovieSceneSection* Section, uint32 InTakeNumber);
 
-	void RegisterTrackImporter(IMovieSceneToolsTrackImporter*);
-	void UnregisterTrackImporter(IMovieSceneToolsTrackImporter*);
+	UE_API void RegisterTrackImporter(IMovieSceneToolsTrackImporter*);
+	UE_API void UnregisterTrackImporter(IMovieSceneToolsTrackImporter*);
 
-	bool ImportAnimatedProperty(const FString& InPropertyName, const FRichCurve& InCurve, FGuid InBinding, UMovieScene* InMovieScene);
-	bool ImportStringProperty(const FString& InPropertyName, const FString& InPropertyValue, FGuid InBinding, UMovieScene* InMovieScene);
+	UE_API bool ImportAnimatedProperty(const FString& InPropertyName, const FRichCurve& InCurve, FGuid InBinding, UMovieScene* InMovieScene);
+	UE_API bool ImportStringProperty(const FString& InPropertyName, const FString& InPropertyValue, FGuid InBinding, UMovieScene* InMovieScene);
 
-	void RegisterKeyStructInstancedPropertyTypeCustomizer(IMovieSceneToolsKeyStructInstancedPropertyTypeCustomizer*);
-	void UnregisterKeyStructInstancedPropertyTypeCustomizer(IMovieSceneToolsKeyStructInstancedPropertyTypeCustomizer*);
+	UE_API void RegisterKeyStructInstancedPropertyTypeCustomizer(IMovieSceneToolsKeyStructInstancedPropertyTypeCustomizer*);
+	UE_API void UnregisterKeyStructInstancedPropertyTypeCustomizer(IMovieSceneToolsKeyStructInstancedPropertyTypeCustomizer*);
 
 	// Called By SKeyEditInterface to allow external modules to add key struct instanced property type customizations
-	void CustomizeKeyStructInstancedPropertyTypes(TSharedRef<IStructureDetailsView> StructureDetailsView, TWeakObjectPtr<UMovieSceneSection> Section);
+	UE_API void CustomizeKeyStructInstancedPropertyTypes(TSharedRef<IStructureDetailsView> StructureDetailsView, TWeakObjectPtr<UMovieSceneSection> Section);
 
 private:
 
-	void RegisterClipboardConversions();
+	UE_API void RegisterClipboardConversions();
 
-	static void FixupPayloadParameterNameForSection(UMovieSceneEventSectionBase* Section, UK2Node* InNode, FName OldPinName, FName NewPinName);
-	static void FixupPayloadParameterNameForDynamicBinding(UMovieScene* MovieScene, UK2Node* InNode, FName OldPinName, FName NewPinName);
-	static bool UpgradeLegacyEventEndpointForSection(UMovieSceneEventSectionBase* Section);
-	static void PostDuplicateEventSection(UMovieSceneEventSectionBase* Section);
-	static void RemoveForCookEventSection(UMovieSceneEventSectionBase* Section);
-	static bool IsTrackClassAllowed(UClass* InClass);
-	static bool IsCustomBindingClassAllowed(UClass* InClass);
-	static bool IsConditionClassAllowed(const UClass* InClass);
-	static void PostDuplicateEvent(ULevelSequence* LevelSequence);
-	static void FixupDynamicBindingsEvent(ULevelSequence* LevelSequence);
-	static void FixupPayloadParameterNameForDirectorBlueprintCondition(UMovieScene* MovieScene, UK2Node* InNode, FName OldPinName, FName NewPinName);
+	static UE_API void FixupPayloadParameterNameForSection(UMovieSceneEventSectionBase* Section, UK2Node* InNode, FName OldPinName, FName NewPinName);
+	static UE_API void FixupPayloadParameterNameForDynamicBinding(UMovieScene* MovieScene, UK2Node* InNode, FName OldPinName, FName NewPinName);
+	static UE_API bool UpgradeLegacyEventEndpointForSection(UMovieSceneEventSectionBase* Section);
+	static UE_API void PostDuplicateEventSection(UMovieSceneEventSectionBase* Section);
+	static UE_API void RemoveForCookEventSection(UMovieSceneEventSectionBase* Section);
+	static UE_API bool IsTrackClassAllowed(UClass* InClass);
+	static UE_API bool IsCustomBindingClassAllowed(UClass* InClass);
+	static UE_API bool IsConditionClassAllowed(const UClass* InClass);
+	static UE_API void PostDuplicateEvent(ULevelSequence* LevelSequence);
+	static UE_API void FixupDynamicBindingsEvent(ULevelSequence* LevelSequence);
+	static UE_API void FixupPayloadParameterNameForDirectorBlueprintCondition(UMovieScene* MovieScene, UK2Node* InNode, FName OldPinName, FName NewPinName);
 
 private:
 
@@ -166,3 +168,5 @@ private:
 	TArray<IMovieSceneToolsAnimationBakeHelper*> BakeHelpers;
 	TArray<IMovieSceneToolsKeyStructInstancedPropertyTypeCustomizer*> KeyStructInstancedPropertyTypeCustomizers;
 };
+
+#undef UE_API
diff --git a/Engine/Source/Editor/MovieSceneTools/Public/MovieSceneToolsProjectSettings.h b/Engine/Source/Editor/MovieSceneTools/Public/MovieSceneToolsProjectSettings.h
--- a/Engine/Source/Editor/MovieSceneTools/Public/MovieSceneToolsProjectSettings.h
+++ b/Engine/Source/Editor/MovieSceneTools/Public/MovieSceneToolsProjectSettings.h
@@ -12,6 +12,8 @@
 
 #include "MovieSceneToolsProjectSettings.generated.h"
 
+#define UE_API MOVIESCENETOOLS_API
+
 UENUM()
 enum class EMovieSceneToolsPropertyTrackType
 {
@@ -52,13 +54,13 @@ struct FMovieSceneToolsFbxSettings
 };
 
 // Settings for the level sequences
-UCLASS(config=EditorPerProjectUserSettings)
-class MOVIESCENETOOLS_API UMovieSceneToolsProjectSettings : public UObject
+UCLASS(MinimalAPI, config=EditorPerProjectUserSettings)
+class UMovieSceneToolsProjectSettings : public UObject
 {
 	GENERATED_BODY()
 
 public:
-	UMovieSceneToolsProjectSettings();
+	UE_API UMovieSceneToolsProjectSettings();
 
 	/** The default start time for new level sequences, in seconds. */
 	UPROPERTY(config, EditAnywhere, Category=Timeline, meta=(Units=s))
@@ -112,3 +114,5 @@ public:
 	UPROPERTY(config, EditAnywhere, Category=TrackSettings)
 	TArray<FMovieSceneToolsFbxSettings> FbxSettings;
 };
+
+#undef UE_API
diff --git a/Engine/Source/Editor/MovieSceneTools/Public/MovieSceneToolsUserSettings.h b/Engine/Source/Editor/MovieSceneTools/Public/MovieSceneToolsUserSettings.h
--- a/Engine/Source/Editor/MovieSceneTools/Public/MovieSceneToolsUserSettings.h
+++ b/Engine/Source/Editor/MovieSceneTools/Public/MovieSceneToolsUserSettings.h
@@ -12,6 +12,8 @@
 
 #include "MovieSceneToolsUserSettings.generated.h"
 
+#define UE_API MOVIESCENETOOLS_API
+
 UENUM()
 enum class EThumbnailQuality : uint8
 {
@@ -20,11 +22,11 @@ enum class EThumbnailQuality : uint8
 	Best,
 };
 
-UCLASS(config=EditorSettings)
-class MOVIESCENETOOLS_API UMovieSceneUserThumbnailSettings : public UObject
+UCLASS(MinimalAPI, config=EditorSettings)
+class UMovieSceneUserThumbnailSettings : public UObject
 {
 public:
-	UMovieSceneUserThumbnailSettings(const FObjectInitializer& Initializer);
+	UE_API UMovieSceneUserThumbnailSettings(const FObjectInitializer& Initializer);
 	
 	GENERATED_BODY()
 
@@ -51,19 +53,19 @@ public:
 	FOnForceRedraw& OnForceRedraw() { return OnForceRedrawEvent; }
 	void BroadcastRedrawThumbnails() const { OnForceRedrawEvent.Broadcast(); }
 
-	virtual void PostEditChangeProperty( struct FPropertyChangedEvent& PropertyChangedEvent) override;
+	UE_API virtual void PostEditChangeProperty( struct FPropertyChangedEvent& PropertyChangedEvent) override;
 
-	static void AddReferencedObjects(UObject* InThis, FReferenceCollector& InCollector);
+	static UE_API void AddReferencedObjects(UObject* InThis, FReferenceCollector& InCollector);
 
 private:
 	FOnForceRedraw OnForceRedrawEvent;
 };
 
-UCLASS(config=EditorSettings, BlueprintType)
-class MOVIESCENETOOLS_API UMovieSceneUserImportFBXSettings : public UObject
+UCLASS(MinimalAPI, config=EditorSettings, BlueprintType)
+class UMovieSceneUserImportFBXSettings : public UObject
 {
 public:
-	UMovieSceneUserImportFBXSettings(const FObjectInitializer& Initializer);
+	UE_API UMovieSceneUserImportFBXSettings(const FObjectInitializer& Initializer);
 	
 	GENERATED_BODY()
 
@@ -180,7 +182,7 @@ struct FControlToTransformMappings
 };
 
 USTRUCT(BlueprintType)
-struct MOVIESCENETOOLS_API FControlFindReplaceString
+struct FControlFindReplaceString
 {
 	GENERATED_BODY()
 
@@ -196,11 +198,11 @@ struct MOVIESCENETOOLS_API FControlFindReplaceString
 	FString Replace;
 };
 
-UCLASS(BlueprintType,config = EditorSettings)
-class MOVIESCENETOOLS_API UMovieSceneUserImportFBXControlRigSettings : public UObject
+UCLASS(MinimalAPI, BlueprintType,config = EditorSettings)
+class UMovieSceneUserImportFBXControlRigSettings : public UObject
 {
 public:
-	UMovieSceneUserImportFBXControlRigSettings(const FObjectInitializer& Initializer);
+	UE_API UMovieSceneUserImportFBXControlRigSettings(const FObjectInitializer& Initializer);
 
 	GENERATED_BODY()
 	/** Imported File Name */
@@ -274,14 +276,14 @@ public:
 public:
 	/** Load the default or metahuman preset into the current mappings */
 	UFUNCTION(BlueprintCallable, Category = "Control Rig")
-	void LoadControlMappingsFromPreset(bool bMetaHumanPreset);
+	UE_API void LoadControlMappingsFromPreset(bool bMetaHumanPreset);
 };
 
-UCLASS(BlueprintType, config = EditorSettings)
-class MOVIESCENETOOLS_API UMovieSceneUserExportFBXControlRigSettings : public UObject
+UCLASS(MinimalAPI, BlueprintType, config = EditorSettings)
+class UMovieSceneUserExportFBXControlRigSettings : public UObject
 {
 public:
-	UMovieSceneUserExportFBXControlRigSettings(const FObjectInitializer& Initializer);
+	UE_API UMovieSceneUserExportFBXControlRigSettings(const FObjectInitializer& Initializer);
 
 	GENERATED_BODY()
 	/** Imported File Name */
@@ -315,5 +317,7 @@ public:
 public:
 	/** Load the default or metahuman preset into the current mappings */
 	UFUNCTION(BlueprintCallable, Category = "Control Rig")
-	void LoadControlMappingsFromPreset(bool bMetaHumanPreset);
+	UE_API void LoadControlMappingsFromPreset(bool bMetaHumanPreset);
 };
+
+#undef UE_API
diff --git a/Engine/Source/Editor/MovieSceneTools/Public/MovieSceneTranslator.h b/Engine/Source/Editor/MovieSceneTools/Public/MovieSceneTranslator.h
--- a/Engine/Source/Editor/MovieSceneTools/Public/MovieSceneTranslator.h
+++ b/Engine/Source/Editor/MovieSceneTools/Public/MovieSceneTranslator.h
@@ -15,10 +15,12 @@
 #include "Sound/SoundBase.h"
 #include "MovieSceneExportMetadata.h"
 
+#define UE_API MOVIESCENETOOLS_API
+
 /**
 * MovieSceneTranslator context class.
 */
-class MOVIESCENETOOLS_API FMovieSceneTranslatorContext : public TSharedFromThis<FMovieSceneTranslatorContext>
+class FMovieSceneTranslatorContext : public TSharedFromThis<FMovieSceneTranslatorContext>
 {
 public:
 	/** Constructor */
@@ -27,19 +29,19 @@ public:
 	virtual ~FMovieSceneTranslatorContext() {}
 
 	/** Initialize the context */
-	void Init();
+	UE_API void Init();
 
 	/** Add message. */
-	void AddMessage(EMessageSeverity::Type MessageSeverity, FText ErrorMessage);
+	UE_API void AddMessage(EMessageSeverity::Type MessageSeverity, FText ErrorMessage);
 
 	/** Reset all messages. */
-	void ClearMessages();
+	UE_API void ClearMessages();
 
 	/** Returns true if specified type of message */
-	bool ContainsMessageType(EMessageSeverity::Type InSeverity) const;
+	UE_API bool ContainsMessageType(EMessageSeverity::Type InSeverity) const;
 
 	/** Get error messages. */
-	const TArray<TSharedRef<FTokenizedMessage>>& GetMessages() const;
+	UE_API const TArray<TSharedRef<FTokenizedMessage>>& GetMessages() const;
 
 private:
 
@@ -137,67 +139,67 @@ enum class EMovieSceneTranslatorSectionType : int32
 /** 
  * The FMovieSceneExportData class aggregates intermediate data from Sequencer classes to be used for timeline exports 
  */
-class MOVIESCENETOOLS_API FMovieSceneExportData : public TSharedFromThis<FMovieSceneExportData>
+class FMovieSceneExportData : public TSharedFromThis<FMovieSceneExportData>
 {
 public:
 	/** Constructor */
-	FMovieSceneExportData(const UMovieScene* InMovieScene, FFrameRate InFrameRate, uint32 InResX, uint32 InResY, int32 InHandleFrames, FString InSaveFilename, TSharedPtr<FMovieSceneTranslatorContext> InContext, FString InMovieExtension);
+	UE_API FMovieSceneExportData(const UMovieScene* InMovieScene, FFrameRate InFrameRate, uint32 InResX, uint32 InResY, int32 InHandleFrames, FString InSaveFilename, TSharedPtr<FMovieSceneTranslatorContext> InContext, FString InMovieExtension);
 	/** Default constructor, necessary for shared ref - should not be used */
-	FMovieSceneExportData();
+	UE_API FMovieSceneExportData();
 	/** Destructor */
-	~FMovieSceneExportData();
+	UE_API ~FMovieSceneExportData();
 
 	/** Gets export filename */
-	FString GetFilename() const;
+	UE_API FString GetFilename() const;
 	/** Gets export filename with full path */
-	FString GetFilenamePath() const;
+	UE_API FString GetFilenamePath() const;
 	/** Gets the shot movie extension */
-	FString GetMovieExtension() const;
+	UE_API FString GetMovieExtension() const;
 
 	/** Gets export frame rate */
-	FFrameRate GetFrameRate() const;
+	UE_API FFrameRate GetFrameRate() const;
 	/** Gets x resolution */
-	uint32 GetResX() const;
+	UE_API uint32 GetResX() const;
 	/** Gets y resolution */
-	uint32 GetResY() const;
+	UE_API uint32 GetResY() const;
 	/** Returns true if frame rate is a non-integral frame rate */
-	bool GetFrameRateIsNTSC() const;
+	UE_API bool GetFrameRateIsNTSC() const;
 	/** Returns the nearest integral frame rate */
-	uint32 GetNearestWholeFrameRate() const;
+	UE_API uint32 GetNearestWholeFrameRate() const;
 
 	/** Gets the frame handle */
-	int32 GetHandleFrames() const;
+	UE_API int32 GetHandleFrames() const;
 	/** Gets default audio sample rate */
-	int32 GetDefaultAudioSampleRate() const;
+	UE_API int32 GetDefaultAudioSampleRate() const;
 	/** Gets default audio depth */
-	int32 GetDefaultAudioDepth() const;
+	UE_API int32 GetDefaultAudioDepth() const;
 
 	/** True when the export data was successfully constructed. */
-	bool IsExportDataValid() const;
+	UE_API bool IsExportDataValid() const;
 
 	/** Find audio sections */
-	bool FindAudioSections(const FString& InSoundPathName, TArray<TSharedPtr<FMovieSceneExportAudioSectionData>>& OutFoundSections) const;
+	UE_API bool FindAudioSections(const FString& InSoundPathName, TArray<TSharedPtr<FMovieSceneExportAudioSectionData>>& OutFoundSections) const;
 
 private:
 
 	/** Entry point for creating the intermediate data to use when exporting. */
-	bool ConstructData(const UMovieScene* InMovieScene);
+	UE_API bool ConstructData(const UMovieScene* InMovieScene);
 	/** Loads intermediate movie scene data from Sequencer. */
-	bool ConstructMovieSceneData(const UMovieScene* InMovieScene);
+	UE_API bool ConstructMovieSceneData(const UMovieScene* InMovieScene);
 	/** Loads intermediate cinematic data from Sequencer. */
-	bool ConstructCinematicData(const UMovieScene* InMovieScene, const UMovieSceneCinematicShotTrack* InCinematicTrack);
+	UE_API bool ConstructCinematicData(const UMovieScene* InMovieScene, const UMovieSceneCinematicShotTrack* InCinematicTrack);
 	/** Loads intermediate cinematic data from Sequencer. */
-	bool ConstructCinematicTrackData(const UMovieScene* InMovieScene, TSharedPtr<FMovieSceneExportCinematicData> InCinematicData, int32 InRowIndex);
+	UE_API bool ConstructCinematicTrackData(const UMovieScene* InMovieScene, TSharedPtr<FMovieSceneExportCinematicData> InCinematicData, int32 InRowIndex);
 	/** Loads intermediate audio data from Sequencer. */
-	bool ConstructAudioData(const UMovieScene* InMovieScene, const UMovieSceneAudioTrack* InAudioTrack, TMap<int32, TSharedPtr<FMovieSceneExportAudioData>>& InAudioDataMap);
+	UE_API bool ConstructAudioData(const UMovieScene* InMovieScene, const UMovieSceneAudioTrack* InAudioTrack, TMap<int32, TSharedPtr<FMovieSceneExportAudioData>>& InAudioDataMap);
 	/** Loads intermediate audio data from Sequencer. */
-	bool ConstructAudioData(const UMovieScene* InMovieScene, TSharedPtr<FMovieSceneExportAudioData> InAudioData, int32 InRowIndex);
+	UE_API bool ConstructAudioData(const UMovieScene* InMovieScene, TSharedPtr<FMovieSceneExportAudioData> InAudioData, int32 InRowIndex);
 	/** Loads intermediate cinematic section data from Sequencer. */
-	bool ConstructCinematicSectionData(const UMovieScene* InMovieScene, TSharedPtr<FMovieSceneExportCinematicData> InCinematicData, const UMovieSceneCinematicShotSection* InCinematicSection);
+	UE_API bool ConstructCinematicSectionData(const UMovieScene* InMovieScene, TSharedPtr<FMovieSceneExportCinematicData> InCinematicData, const UMovieSceneCinematicShotSection* InCinematicSection);
 	/** Loads intermediate audio section data from Sequencer. */
-	bool ConstructAudioSectionData(const UMovieScene* InMovieScene, TSharedPtr<FMovieSceneExportAudioData> InAudioData, const UMovieSceneAudioSection* InAudioSection);
+	UE_API bool ConstructAudioSectionData(const UMovieScene* InMovieScene, TSharedPtr<FMovieSceneExportAudioData> InAudioData, const UMovieSceneAudioSection* InAudioSection);
 	/** Loads intermediate common section data from Sequencer. */
-	bool ConstructSectionData(const UMovieScene* InMovieScene, TSharedPtr<FMovieSceneExportSectionData> InSectionData, const UMovieSceneSection* InSection, EMovieSceneTranslatorSectionType InSectionType, const FString& InSectionDisplayName);
+	UE_API bool ConstructSectionData(const UMovieScene* InMovieScene, TSharedPtr<FMovieSceneExportSectionData> InSectionData, const UMovieSceneSection* InSection, EMovieSceneTranslatorSectionType InSectionType, const FString& InSectionDisplayName);
 
 	/** Context for messages */
 	TSharedPtr<FMovieSceneTranslatorContext> ExportContext;
@@ -343,7 +345,7 @@ public:
 };
 
 /** Abstract base class for importer/exporter */
-class MOVIESCENETOOLS_API FMovieSceneTranslator
+class FMovieSceneTranslator
 {
 public:
 	FMovieSceneTranslator() {}
@@ -356,7 +358,7 @@ public:
 };
 
 /** Abstract base class for movie scene importers */
-class MOVIESCENETOOLS_API FMovieSceneImporter : public FMovieSceneTranslator
+class FMovieSceneImporter : public FMovieSceneTranslator
 {
 public:
 	FMovieSceneImporter() : FMovieSceneTranslator() {}
@@ -384,7 +386,7 @@ public:
 };
 
 /** Abstract base class for movie scene exporters */
-class MOVIESCENETOOLS_API FMovieSceneExporter : public FMovieSceneTranslator
+class FMovieSceneExporter : public FMovieSceneTranslator
 {
 public:
 	FMovieSceneExporter() : FMovieSceneTranslator() {}
@@ -422,3 +424,4 @@ public:
 	virtual bool Export(const UMovieScene* InMovieScene, FString InFilenameFormat, FFrameRate InFrameRate, uint32 InResX, uint32 InResY, int32 InHandleFrames, FString InSaveFilename, TSharedRef<FMovieSceneTranslatorContext> InContext, FString InMovieExtension, const FMovieSceneExportMetadata* InMetadata=nullptr) = 0;
 };
 
+#undef UE_API
diff --git a/Engine/Source/Editor/MovieSceneTools/Public/MovieSceneTranslatorEDL.h b/Engine/Source/Editor/MovieSceneTools/Public/MovieSceneTranslatorEDL.h
--- a/Engine/Source/Editor/MovieSceneTools/Public/MovieSceneTranslatorEDL.h
+++ b/Engine/Source/Editor/MovieSceneTools/Public/MovieSceneTranslatorEDL.h
@@ -4,11 +4,13 @@
 
 #include "CoreMinimal.h"
 
+#define UE_API MOVIESCENETOOLS_API
+
 struct FFrameRate;
 struct FFrameNumber;
 class UMovieScene;
 
-class MOVIESCENETOOLS_API MovieSceneTranslatorEDL
+class MovieSceneTranslatorEDL
 {
 public:
 
@@ -20,7 +22,7 @@ public:
 	 * @param InFilename The filename to import
 	 * @return Whether the import was successful
 	 */
-	static bool ImportEDL(UMovieScene* InMovieScene, FFrameRate InFrameRate, FString InFilename);
+	static UE_API bool ImportEDL(UMovieScene* InMovieScene, FFrameRate InFrameRate, FString InFilename);
 
 	/**
 	 * Export EDL
@@ -32,5 +34,7 @@ public:
 	 * @param MovieExtension The movie extension for the shot filenames (ie. .avi, .mov, .mp4)
 	 * @return Whether the export was successful
 	 */
-	static bool ExportEDL(const UMovieScene* InMovieScene, FFrameRate InFrameRate, FString InSaveFilename, int32 InHandleFrames, FString InMovieExtension);
+	static UE_API bool ExportEDL(const UMovieScene* InMovieScene, FFrameRate InFrameRate, FString InSaveFilename, int32 InHandleFrames, FString InMovieExtension);
 };
+
+#undef UE_API
diff --git a/Engine/Source/Editor/MovieSceneTools/Public/SKeyEditInterface.h b/Engine/Source/Editor/MovieSceneTools/Public/SKeyEditInterface.h
--- a/Engine/Source/Editor/MovieSceneTools/Public/SKeyEditInterface.h
+++ b/Engine/Source/Editor/MovieSceneTools/Public/SKeyEditInterface.h
@@ -8,6 +8,8 @@
 #include "EditorUndoClient.h"
 #include "UObject/WeakObjectPtrTemplates.h"
 
+#define UE_API MOVIESCENETOOLS_API
+
 struct FPropertyChangedEvent;
 
 class ISequencer;
@@ -24,46 +26,48 @@ struct FKeyEditData
 /**
  * Widget that represents a details panel that refreshes on undo, and handles modification of the section on edit
  */
-class MOVIESCENETOOLS_API SKeyEditInterface : public SCompoundWidget, public FEditorUndoClient, public FNotifyHook
+class SKeyEditInterface : public SCompoundWidget, public FEditorUndoClient, public FNotifyHook
 {
 public:
 	SLATE_BEGIN_ARGS(SKeyEditInterface){}
 		SLATE_ATTRIBUTE(FKeyEditData, EditData)
 	SLATE_END_ARGS()
 
-	~SKeyEditInterface();
+	UE_API ~SKeyEditInterface();
 
-	void Construct(const FArguments& InArgs, TSharedRef<ISequencer> InSequencer);
+	UE_API void Construct(const FArguments& InArgs, TSharedRef<ISequencer> InSequencer);
 
 	/**
 	 * (Re)Initialize this widget's details panel
 	 */
-	void Initialize();
+	UE_API void Initialize();
 
 private:
 
 	/**
 	 * Create a binding ID customization for the details panel
 	 */
-	TSharedRef<IPropertyTypeCustomization> CreateBindingIDCustomization();
-	TSharedRef<IPropertyTypeCustomization> CreateFrameNumberCustomization();
-	TSharedRef<IPropertyTypeCustomization> CreateEventCustomization();
+	UE_API TSharedRef<IPropertyTypeCustomization> CreateBindingIDCustomization();
+	UE_API TSharedRef<IPropertyTypeCustomization> CreateFrameNumberCustomization();
+	UE_API TSharedRef<IPropertyTypeCustomization> CreateEventCustomization();
 
 	/**
 	 * Called when a property has been changed on the UI
 	 */
-	void OnFinishedChangingProperties(const FPropertyChangedEvent& ChangeEvent, TSharedPtr<FStructOnScope> KeyStruct);
+	UE_API void OnFinishedChangingProperties(const FPropertyChangedEvent& ChangeEvent, TSharedPtr<FStructOnScope> KeyStruct);
 
 private:
 
-	virtual FVector2D ComputeDesiredSize(float LayoutScaleMultiplier) const override;
+	UE_API virtual FVector2D ComputeDesiredSize(float LayoutScaleMultiplier) const override;
 
-	virtual void NotifyPreChange(FProperty* PropertyAboutToChange) override;
-	virtual void PostUndo(bool bSuccess) override;
-	virtual void PostRedo(bool bSuccess) override;
+	UE_API virtual void NotifyPreChange(FProperty* PropertyAboutToChange) override;
+	UE_API virtual void PostUndo(bool bSuccess) override;
+	UE_API virtual void PostRedo(bool bSuccess) override;
 
 private:
 	TAttribute<FKeyEditData> EditDataAttribute;
 	TWeakObjectPtr<UMovieSceneSection> WeakSection;
 	TWeakPtr<ISequencer> WeakSequencer;
-};
\ No newline at end of file
+};
+
+#undef UE_API
diff --git a/Engine/Source/Editor/MovieSceneTools/Public/Sections/ThumbnailSection.h b/Engine/Source/Editor/MovieSceneTools/Public/Sections/ThumbnailSection.h
--- a/Engine/Source/Editor/MovieSceneTools/Public/Sections/ThumbnailSection.h
+++ b/Engine/Source/Editor/MovieSceneTools/Public/Sections/ThumbnailSection.h
@@ -9,6 +9,8 @@
 #include "ISequencerSection.h"
 #include "TrackEditorThumbnail/TrackEditorThumbnail.h"
 
+#define UE_API MOVIESCENETOOLS_API
+
 class FLevelEditorViewportClient;
 class FMenuBuilder;
 class FSceneViewport;
@@ -26,18 +28,18 @@ struct FThumbnailCameraSettings
 /**
  * Thumbnail section, which paints and ticks the appropriate section.
  */
-class MOVIESCENETOOLS_API FThumbnailSection
+class FThumbnailSection
 	: public ISequencerSection
 	, public TSharedFromThis<FThumbnailSection>
 {
 public:
 
 	/** Create and initialize a new instance. */
-	FThumbnailSection(TSharedPtr<ISequencer> InSequencer, TSharedPtr<FTrackEditorThumbnailPool> InThumbnailPool, IViewportThumbnailClient* InViewportThumbanilClient, UMovieSceneSection& InSection);
-	FThumbnailSection(TSharedPtr<ISequencer> InSequencer, TSharedPtr<FTrackEditorThumbnailPool> InThumbnailPool, ICustomThumbnailClient* InCustomThumbnailClient, UMovieSceneSection& InSection);
+	UE_API FThumbnailSection(TSharedPtr<ISequencer> InSequencer, TSharedPtr<FTrackEditorThumbnailPool> InThumbnailPool, IViewportThumbnailClient* InViewportThumbanilClient, UMovieSceneSection& InSection);
+	UE_API FThumbnailSection(TSharedPtr<ISequencer> InSequencer, TSharedPtr<FTrackEditorThumbnailPool> InThumbnailPool, ICustomThumbnailClient* InCustomThumbnailClient, UMovieSceneSection& InSection);
 
 	/** Destructor. */
-	~FThumbnailSection();
+	UE_API ~FThumbnailSection();
 
 public:
 
@@ -48,7 +50,7 @@ public:
 	}
 
 	/** Enter rename mode for the section */
-	void EnterRename();
+	UE_API void EnterRename();
 
 	/** Get whether the text is renameable */
 	virtual bool CanRename() const { return false; }
@@ -68,31 +70,31 @@ public:
 
 	//~ ISequencerSection interface
 
-	virtual void BuildSectionContextMenu(FMenuBuilder& MenuBuilder, const FGuid& ObjectBinding) override;
-	virtual void BuildSectionSidebarMenu(FMenuBuilder& MenuBuilder, const FGuid& ObjectBinding) override;
-	virtual TSharedRef<SWidget> GenerateSectionWidget() override;
-	virtual float GetSectionGripSize() const override;
-	virtual float GetSectionHeight(const UE::Sequencer::FViewDensityInfo& ViewDensity) const override;
-	virtual UMovieSceneSection* GetSectionObject() override;
-	virtual FText GetSectionTitle() const override;
-	virtual int32 OnPaintSection( FSequencerSectionPainter& InPainter ) const override;
-	virtual void Tick( const FGeometry& AllottedGeometry, const FGeometry& ParentGeometry, const double InCurrentTime, const float InDeltaTime ) override;
+	UE_API virtual void BuildSectionContextMenu(FMenuBuilder& MenuBuilder, const FGuid& ObjectBinding) override;
+	UE_API virtual void BuildSectionSidebarMenu(FMenuBuilder& MenuBuilder, const FGuid& ObjectBinding) override;
+	UE_API virtual TSharedRef<SWidget> GenerateSectionWidget() override;
+	UE_API virtual float GetSectionGripSize() const override;
+	UE_API virtual float GetSectionHeight(const UE::Sequencer::FViewDensityInfo& ViewDensity) const override;
+	UE_API virtual UMovieSceneSection* GetSectionObject() override;
+	UE_API virtual FText GetSectionTitle() const override;
+	UE_API virtual int32 OnPaintSection( FSequencerSectionPainter& InPainter ) const override;
+	UE_API virtual void Tick( const FGeometry& AllottedGeometry, const FGeometry& ParentGeometry, const double InCurrentTime, const float InDeltaTime ) override;
 
 protected:
 
 	/** Called to force a redraw of this section's thumbnails */
-	void RedrawThumbnails();
+	UE_API void RedrawThumbnails();
 
 	/** Get the range that is currently visible in the section's time space */
-	TRange<double> GetVisibleRange() const;
+	UE_API TRange<double> GetVisibleRange() const;
 
 	/** Get the total range that thumbnails are to be generated for in the section's time space */
-	TRange<double> GetTotalRange() const;
+	UE_API TRange<double> GetTotalRange() const;
 
 	/** Get rename visibility */
-	EVisibility GetRenameVisibility() const;
+	UE_API EVisibility GetRenameVisibility() const;
 
-	void BuildThumbnailsMenu(FMenuBuilder& InMenuBuilder);
+	UE_API void BuildThumbnailsMenu(FMenuBuilder& InMenuBuilder);
 
 protected:
 
@@ -132,16 +134,18 @@ protected:
 /**
  * Thumbnail section, which paints and ticks the appropriate section.
  */
-class MOVIESCENETOOLS_API FViewportThumbnailSection
+class FViewportThumbnailSection
 	: public FThumbnailSection
 	, public IViewportThumbnailClient
 {
 public:
 
 	/** Create and initialize a new instance. */
-	FViewportThumbnailSection(TSharedPtr<ISequencer> InSequencer, TSharedPtr<FTrackEditorThumbnailPool> InThumbnailPool, UMovieSceneSection& InSection);
+	UE_API FViewportThumbnailSection(TSharedPtr<ISequencer> InSequencer, TSharedPtr<FTrackEditorThumbnailPool> InThumbnailPool, UMovieSceneSection& InSection);
 
 	//~ IViewportThumbnailClient interface
-	virtual void PreDraw(FTrackEditorThumbnail& Thumbnail) override;
-	virtual void PostDraw(FTrackEditorThumbnail& Thumbnail) override;
+	UE_API virtual void PreDraw(FTrackEditorThumbnail& Thumbnail) override;
+	UE_API virtual void PostDraw(FTrackEditorThumbnail& Thumbnail) override;
 };
+
+#undef UE_API
diff --git a/Engine/Source/Editor/MovieSceneTools/Public/SequenceBindingTree.h b/Engine/Source/Editor/MovieSceneTools/Public/SequenceBindingTree.h
--- a/Engine/Source/Editor/MovieSceneTools/Public/SequenceBindingTree.h
+++ b/Engine/Source/Editor/MovieSceneTools/Public/SequenceBindingTree.h
@@ -14,6 +14,8 @@
 #include "Textures/SlateIcon.h"
 #include "UObject/ObjectKey.h"
 
+#define UE_API MOVIESCENETOOLS_API
+
 class UMovieScene;
 class UMovieSceneSequence;
 
@@ -48,7 +50,7 @@ struct FSequenceBindingNode
 };
 
 /** Data structure used internally to represent the bindings of a sequence recursively */
-struct MOVIESCENETOOLS_API FSequenceBindingTree
+struct FSequenceBindingTree
 {
 	/**
 	 * Conditionally reconstruct the tree structure from the specified sequence if any of the sequences have changed since the last population
@@ -58,7 +60,7 @@ struct MOVIESCENETOOLS_API FSequenceBindingTree
 	 * @param InActiveSequenceID	The sequence ID for the above sequence within the root context
 	 * @retrun True if this tree was rebuilt, false otherwise
 	 */
-	bool ConditionalRebuild(UMovieSceneSequence* InSequence, FObjectKey InActiveSequence, FMovieSceneSequenceID InActiveSequenceID);
+	UE_API bool ConditionalRebuild(UMovieSceneSequence* InSequence, FObjectKey InActiveSequence, FMovieSceneSequenceID InActiveSequenceID);
 
 	/**
 	 * Construct the tree structure from the specified sequence.
@@ -67,7 +69,7 @@ struct MOVIESCENETOOLS_API FSequenceBindingTree
 	 * @param InActiveSequence		A sequence at which point we can start to generate localally resolving IDs
 	 * @param InActiveSequenceID	The sequence ID for the above sequence within the root context
 	 */
-	void ForceRebuild(UMovieSceneSequence* InSequence, FObjectKey InActiveSequence, FMovieSceneSequenceID InActiveSequenceID);
+	UE_API void ForceRebuild(UMovieSceneSequence* InSequence, FObjectKey InActiveSequence, FMovieSceneSequenceID InActiveSequenceID);
 
 	/** Get the root of the tree */
 	TSharedRef<FSequenceBindingNode> GetRootNode() const
@@ -91,13 +93,13 @@ private:
 	struct FSequenceIDStack;
 
 	/** Recursive sort helper for a sequence binding node */
-	static void Sort(TSharedRef<FSequenceBindingNode> Node);
+	static UE_API void Sort(TSharedRef<FSequenceBindingNode> Node);
 
 	/** Recursive builder function that iterates into sub sequences */
-	void Build(UMovieSceneSequence* InSequence, FSequenceIDStack& SequenceIDStack);
+	UE_API void Build(UMovieSceneSequence* InSequence, FSequenceIDStack& SequenceIDStack);
 
 	/** Ensure that a parent node exists for the specified object */
-	TSharedRef<FSequenceBindingNode> EnsureParent(const FGuid& InParentGuid, UMovieScene* InMovieScene, FMovieSceneSequenceID SequenceID);
+	UE_API TSharedRef<FSequenceBindingNode> EnsureParent(const FGuid& InParentGuid, UMovieScene* InMovieScene, FMovieSceneSequenceID SequenceID);
 
 private:
 
@@ -115,4 +117,6 @@ private:
 	TMap<FObjectKey, FGuid> CachedSequenceSignatures;
 	/** Whether the tree is considered empty */
 	bool bIsEmpty;
-};
\ No newline at end of file
+};
+
+#undef UE_API
diff --git a/Engine/Source/Editor/MovieSceneTools/Public/TrackEditorThumbnail/TrackEditorThumbnail.h b/Engine/Source/Editor/MovieSceneTools/Public/TrackEditorThumbnail/TrackEditorThumbnail.h
--- a/Engine/Source/Editor/MovieSceneTools/Public/TrackEditorThumbnail/TrackEditorThumbnail.h
+++ b/Engine/Source/Editor/MovieSceneTools/Public/TrackEditorThumbnail/TrackEditorThumbnail.h
@@ -12,6 +12,8 @@
 #include "Misc/FrameTime.h"
 #include "Slate/SlateTextures.h"
 
+#define UE_API MOVIESCENETOOLS_API
+
 struct FMinimalViewInfo;
 
 class FLevelEditorViewportClient;
@@ -30,29 +32,29 @@ DECLARE_DELEGATE_OneParam(FOnThumbnailDraw, FTrackEditorThumbnail&);
 /**
  * Track Editor Thumbnail, which keeps a Texture to be displayed by a viewport.
  */
-class MOVIESCENETOOLS_API FTrackEditorThumbnail
+class FTrackEditorThumbnail
 	: public ISlateViewport
 	, public TSharedFromThis<FTrackEditorThumbnail>
 {
 public:
 
 	/** Create and initialize a new instance. */
-	FTrackEditorThumbnail(const FOnThumbnailDraw& InOnDraw, const FIntPoint& InDesiredSize, TRange<double> InTimeRange, double InPosition);
+	UE_API FTrackEditorThumbnail(const FOnThumbnailDraw& InOnDraw, const FIntPoint& InDesiredSize, TRange<double> InTimeRange, double InPosition);
 
 	/** Virtual destructor. */
-	virtual ~FTrackEditorThumbnail();
+	UE_API virtual ~FTrackEditorThumbnail();
 
 public:
 
 	/**
 	 * Assign this thumbnail from a slate texture.
 	 */
-	void AssignFrom(TSharedRef<FSlateTextureData, ESPMode::ThreadSafe> InTextureData);
+	UE_API void AssignFrom(TSharedRef<FSlateTextureData, ESPMode::ThreadSafe> InTextureData);
 
 	/**
 	 * Ensure that this thumbnail has a render target of the specified size
 	 */
-	void ResizeRenderTarget(const FIntPoint& InSize);
+	UE_API void ResizeRenderTarget(const FIntPoint& InSize);
 
 	/**
 	 * Access the (potentially null) render target to be used for rendering onto this thumbnail
@@ -71,14 +73,14 @@ public:
 	}
 
 	/** Renders the thumbnail to the texture. */
-	void DrawThumbnail();
+	UE_API void DrawThumbnail();
 
 	/** Prompt this thumbnail to fade in */
-	void SetupFade(const TSharedRef<SWidget>& InWidget);
-	void PlayFade();
+	UE_API void SetupFade(const TSharedRef<SWidget>& InWidget);
+	UE_API void PlayFade();
 
 	/** Gets the curve for fading in the thumbnail. */
-	float GetFadeInCurve() const;
+	UE_API float GetFadeInCurve() const;
 
 	/** Get the full time-range that this thumbnail occupies */
 	const TRange<double> GetTimeRange() const { return TimeRange; }
@@ -90,12 +92,12 @@ public:
 
 	// ISlateViewport interface
 
-	virtual FIntPoint GetSize() const override;
-	virtual FSlateShaderResource* GetViewportRenderTargetTexture() const override;
-	virtual bool RequiresVsync() const override;
+	UE_API virtual FIntPoint GetSize() const override;
+	UE_API virtual FSlateShaderResource* GetViewportRenderTargetTexture() const override;
+	UE_API virtual bool RequiresVsync() const override;
 
 	/** Destroy the texture on the next tick */
-	void DestroyTexture();
+	UE_API void DestroyTexture();
 
 public:
 
@@ -187,22 +189,22 @@ struct FThumbnailCacheData
 	TOptional<double> SingleReferenceFrame;
 };
 
-class MOVIESCENETOOLS_API FTrackEditorThumbnailCache
+class FTrackEditorThumbnailCache
 {
 public:
-	FTrackEditorThumbnailCache(const TSharedPtr<FTrackEditorThumbnailPool>& ThumbnailPool, IViewportThumbnailClient* InViewportThumbnailClient);
-	FTrackEditorThumbnailCache(const TSharedPtr<FTrackEditorThumbnailPool>& ThumbnailPool, ICustomThumbnailClient* InCustomThumbnailClient);
+	UE_API FTrackEditorThumbnailCache(const TSharedPtr<FTrackEditorThumbnailPool>& ThumbnailPool, IViewportThumbnailClient* InViewportThumbnailClient);
+	UE_API FTrackEditorThumbnailCache(const TSharedPtr<FTrackEditorThumbnailPool>& ThumbnailPool, ICustomThumbnailClient* InCustomThumbnailClient);
 	
-	~FTrackEditorThumbnailCache();
+	UE_API ~FTrackEditorThumbnailCache();
 
 	void ForceRedraw() { bForceRedraw = true; }
 
-	void SetSingleReferenceFrame(TOptional<double> InReferenceFrame);
+	UE_API void SetSingleReferenceFrame(TOptional<double> InReferenceFrame);
 	TOptional<double> GetSingleReferenceFrame() const { return CurrentCache.SingleReferenceFrame; }
 
-	void Update(const TRange<double>& NewRange, const TRange<double>& VisibleRange, const FIntPoint& AllottedSize, const FIntPoint& InDesiredSize, EThumbnailQuality InQuality, double InCurrentTime);
+	UE_API void Update(const TRange<double>& NewRange, const TRange<double>& VisibleRange, const FIntPoint& AllottedSize, const FIntPoint& InDesiredSize, EThumbnailQuality InQuality, double InCurrentTime);
 
-	void Revalidate(double InCurrentTime);
+	UE_API void Revalidate(double InCurrentTime);
 
 	const TArray<TSharedPtr<FTrackEditorThumbnail>>& GetThumbnails() const
 	{
@@ -211,27 +213,27 @@ public:
 
 protected:
 
-	void DrawThumbnail(FTrackEditorThumbnail& TrackEditorThumbnail);
+	UE_API void DrawThumbnail(FTrackEditorThumbnail& TrackEditorThumbnail);
 
-	void DrawViewportThumbnail(FTrackEditorThumbnail& TrackEditorThumbnail) const;
+	UE_API void DrawViewportThumbnail(FTrackEditorThumbnail& TrackEditorThumbnail) const;
 
-	void ComputeNewThumbnails();
+	UE_API void ComputeNewThumbnails();
 
-	void Setup();
+	UE_API void Setup();
 
-	bool ShouldRegenerateEverything() const;
+	UE_API bool ShouldRegenerateEverything() const;
 
-	FIntPoint CalculateTextureSize(const FMinimalViewInfo& ViewInfo) const;
+	UE_API FIntPoint CalculateTextureSize(const FMinimalViewInfo& ViewInfo) const;
 
-	void UpdateSingleThumbnail();
+	UE_API void UpdateSingleThumbnail();
 
-	void UpdateFilledThumbnails();
+	UE_API void UpdateFilledThumbnails();
 
-	void GenerateFront(const TRange<double>& Boundary);
+	UE_API void GenerateFront(const TRange<double>& Boundary);
 
-	void GenerateBack(const TRange<double>& Boundary);
+	UE_API void GenerateBack(const TRange<double>& Boundary);
 
-	void SetupViewportEngineFlags();
+	UE_API void SetupViewportEngineFlags();
 
 protected:
 
@@ -255,3 +257,5 @@ protected:
 	/** Whether to force a redraw or not */
 	bool bForceRedraw;
 };
+
+#undef UE_API
diff --git a/Engine/Source/Editor/MovieSceneTools/Public/TrackEditorThumbnail/TrackEditorThumbnailPool.h b/Engine/Source/Editor/MovieSceneTools/Public/TrackEditorThumbnail/TrackEditorThumbnailPool.h
--- a/Engine/Source/Editor/MovieSceneTools/Public/TrackEditorThumbnail/TrackEditorThumbnailPool.h
+++ b/Engine/Source/Editor/MovieSceneTools/Public/TrackEditorThumbnail/TrackEditorThumbnailPool.h
@@ -4,6 +4,8 @@
 
 #include "CoreMinimal.h"
 
+#define UE_API MOVIESCENETOOLS_API
+
 class FTrackEditorThumbnail;
 class ISequencer;
 
@@ -11,23 +13,23 @@ class ISequencer;
  * Track Editor Thumbnail pool, which keeps a list of thumbnails that
  * need to be drawn and draws them incrementally.
  */
-class MOVIESCENETOOLS_API FTrackEditorThumbnailPool
+class FTrackEditorThumbnailPool
 {
 public:
 
-	FTrackEditorThumbnailPool(TSharedPtr<ISequencer> InSequencer);
+	UE_API FTrackEditorThumbnailPool(TSharedPtr<ISequencer> InSequencer);
 
 public:
 
 	/** Requests that the passed in thumbnails need to be drawn */
-	void AddThumbnailsNeedingRedraw(const TArray<TSharedPtr<FTrackEditorThumbnail>>& InThumbnails);
+	UE_API void AddThumbnailsNeedingRedraw(const TArray<TSharedPtr<FTrackEditorThumbnail>>& InThumbnails);
 
 	/** Draws a small number of thumbnails that are enqueued for drawing */
 	/* @return Whether thumbnails were drawn */
-	bool DrawThumbnails();
+	UE_API bool DrawThumbnails();
 
 	/** Informs the pool that the thumbnails passed in no longer need to be drawn */
-	void RemoveThumbnailsNeedingRedraw(const TArray< TSharedPtr<FTrackEditorThumbnail>>& InThumbnails);
+	UE_API void RemoveThumbnailsNeedingRedraw(const TArray< TSharedPtr<FTrackEditorThumbnail>>& InThumbnails);
 
 private:
 
@@ -46,3 +48,5 @@ private:
 	/** Whether we need to sort */
 	bool bNeedsSort;
 };
+
+#undef UE_API
diff --git a/Engine/Source/Editor/MovieSceneTools/Public/TrackEditors/AudioTrackEditor.h b/Engine/Source/Editor/MovieSceneTools/Public/TrackEditors/AudioTrackEditor.h
--- a/Engine/Source/Editor/MovieSceneTools/Public/TrackEditors/AudioTrackEditor.h
+++ b/Engine/Source/Editor/MovieSceneTools/Public/TrackEditors/AudioTrackEditor.h
@@ -15,6 +15,8 @@
 #include "Containers/Map.h"
 #include "TimeToPixel.h"
 
+#define UE_API MOVIESCENETOOLS_API
+
 struct FAssetData;
 class FAudioThumbnail;
 class FDelegateHandle;
@@ -26,7 +28,7 @@ class UMovieSceneAudioTrack;
 /**
  * Tools for audio tracks
  */
-class MOVIESCENETOOLS_API FAudioTrackEditor
+class FAudioTrackEditor
 	: public FMovieSceneTrackEditor
 {
 public:
@@ -36,10 +38,10 @@ public:
 	 *
 	 * @param InSequencer The sequencer instance to be used by this tool
 	 */
-	FAudioTrackEditor(TSharedRef<ISequencer> InSequencer);
+	UE_API FAudioTrackEditor(TSharedRef<ISequencer> InSequencer);
 
 	/** Virtual destructor. */
-	virtual ~FAudioTrackEditor();
+	UE_API virtual ~FAudioTrackEditor();
 
 	/**
 	 * Creates an instance of this class.  Called by a sequencer 
@@ -47,76 +49,76 @@ public:
 	 * @param OwningSequencer The sequencer instance to be used by this tool
 	 * @return The new instance of this class
 	 */
-	static TSharedRef<ISequencerTrackEditor> CreateTrackEditor(TSharedRef<ISequencer> OwningSequencer);
+	static UE_API TSharedRef<ISequencerTrackEditor> CreateTrackEditor(TSharedRef<ISequencer> OwningSequencer);
 
 public:
 
 	// ISequencerTrackEditor interface
 
-	virtual FText GetDisplayName() const override;
-	virtual void OnInitialize() override;
-	virtual void OnRelease() override;
-	virtual void BuildAddTrackMenu(FMenuBuilder& MenuBuilder) override;
-	virtual void BuildObjectBindingTrackMenu(FMenuBuilder& MenuBuilder, const TArray<FGuid>& ObjectBindings, const UClass* ObjectClass) override;
-	virtual TSharedPtr<SWidget> BuildOutlinerColumnWidget(const FBuildColumnWidgetParams& Params, const FName& ColumnName) override;
-	virtual bool HandleAssetAdded(UObject* Asset, const FGuid& TargetObjectGuid) override;
-	virtual TSharedRef<ISequencerSection> MakeSectionInterface(UMovieSceneSection& SectionObject, UMovieSceneTrack& Track, FGuid ObjectBinding) override;
-	virtual bool SupportsType(TSubclassOf<UMovieSceneTrack> Type) const override;
-	virtual bool SupportsSequence(UMovieSceneSequence* InSequence) const override;
-	virtual const FSlateBrush* GetIconBrush() const override;
-	virtual bool IsResizable(UMovieSceneTrack* InTrack) const override;
-	virtual void Resize(float NewSize, UMovieSceneTrack* InTrack) override;
-	virtual bool OnAllowDrop(const FDragDropEvent& DragDropEvent, FSequencerDragDropParams& DragDropParams) override;
-	virtual FReply OnDrop(const FDragDropEvent& DragDropEvent, const FSequencerDragDropParams& DragDropParams) override;
+	UE_API virtual FText GetDisplayName() const override;
+	UE_API virtual void OnInitialize() override;
+	UE_API virtual void OnRelease() override;
+	UE_API virtual void BuildAddTrackMenu(FMenuBuilder& MenuBuilder) override;
+	UE_API virtual void BuildObjectBindingTrackMenu(FMenuBuilder& MenuBuilder, const TArray<FGuid>& ObjectBindings, const UClass* ObjectClass) override;
+	UE_API virtual TSharedPtr<SWidget> BuildOutlinerColumnWidget(const FBuildColumnWidgetParams& Params, const FName& ColumnName) override;
+	UE_API virtual bool HandleAssetAdded(UObject* Asset, const FGuid& TargetObjectGuid) override;
+	UE_API virtual TSharedRef<ISequencerSection> MakeSectionInterface(UMovieSceneSection& SectionObject, UMovieSceneTrack& Track, FGuid ObjectBinding) override;
+	UE_API virtual bool SupportsType(TSubclassOf<UMovieSceneTrack> Type) const override;
+	UE_API virtual bool SupportsSequence(UMovieSceneSequence* InSequence) const override;
+	UE_API virtual const FSlateBrush* GetIconBrush() const override;
+	UE_API virtual bool IsResizable(UMovieSceneTrack* InTrack) const override;
+	UE_API virtual void Resize(float NewSize, UMovieSceneTrack* InTrack) override;
+	UE_API virtual bool OnAllowDrop(const FDragDropEvent& DragDropEvent, FSequencerDragDropParams& DragDropParams) override;
+	UE_API virtual FReply OnDrop(const FDragDropEvent& DragDropEvent, const FSequencerDragDropParams& DragDropParams) override;
 	
 protected:
 
 	/** Delegate for AnimatablePropertyChanged in HandleAssetAdded for sounds */
-	FKeyPropertyResult AddNewSound(FFrameNumber KeyTime, class USoundBase* Sound, UMovieSceneAudioTrack* Track, int32 RowIndex);
+	UE_API FKeyPropertyResult AddNewSound(FFrameNumber KeyTime, class USoundBase* Sound, UMovieSceneAudioTrack* Track, int32 RowIndex);
 
 	/** Delegate for AnimatablePropertyChanged in HandleAssetAdded for attached sounds */
-	FKeyPropertyResult AddNewAttachedSound(FFrameNumber KeyTime, class USoundBase* Sound, UMovieSceneAudioTrack* Track, TArray<TWeakObjectPtr<UObject>> ObjectsToAttachTo);
+	UE_API FKeyPropertyResult AddNewAttachedSound(FFrameNumber KeyTime, class USoundBase* Sound, UMovieSceneAudioTrack* Track, TArray<TWeakObjectPtr<UObject>> ObjectsToAttachTo);
 
 private:
 
 	/** Callback for executing the "Add Audio Track" menu entry. */
-	void HandleAddAudioTrackMenuEntryExecute();
+	UE_API void HandleAddAudioTrackMenuEntryExecute();
 
 	/** Callback for executing the "Add Audio Track" menu entry on an actor */
-	void HandleAddAttachedAudioTrackMenuEntryExecute(FMenuBuilder& MenuBuilder, TArray<FGuid> ObjectBindings);
+	UE_API void HandleAddAttachedAudioTrackMenuEntryExecute(FMenuBuilder& MenuBuilder, TArray<FGuid> ObjectBindings);
 
 	/** Audio sub menu */
-	TSharedRef<SWidget> BuildAudioSubMenu(FOnAssetSelected OnAssetSelected, FOnAssetEnterPressed OnAssetEnterPressed);
+	UE_API TSharedRef<SWidget> BuildAudioSubMenu(FOnAssetSelected OnAssetSelected, FOnAssetEnterPressed OnAssetEnterPressed);
 
 	/** Audio asset selected */
-	void OnAudioAssetSelected(const FAssetData& AssetData, UMovieSceneTrack* Track);
+	UE_API void OnAudioAssetSelected(const FAssetData& AssetData, UMovieSceneTrack* Track);
 
 	/** Audio asset enter pressed */
-	void OnAudioAssetEnterPressed(const TArray<FAssetData>& AssetData, UMovieSceneTrack* Track);
+	UE_API void OnAudioAssetEnterPressed(const TArray<FAssetData>& AssetData, UMovieSceneTrack* Track);
 
 	/** Attached audio asset selected */
-	void OnAttachedAudioAssetSelected(const FAssetData& AssetData, TArray<FGuid> ObjectBindings);
+	UE_API void OnAttachedAudioAssetSelected(const FAssetData& AssetData, TArray<FGuid> ObjectBindings);
 
 	/** Attached audio asset enter pressed */
-	void OnAttachedAudioEnterPressed(const TArray<FAssetData>& AssetData, TArray<FGuid> ObjectBindings);
+	UE_API void OnAttachedAudioEnterPressed(const TArray<FAssetData>& AssetData, TArray<FGuid> ObjectBindings);
 
 	/** Registers a delegate with the sequencer for monitoring edits */
-	void RegisterMovieSceneChangedDelegate();
+	UE_API void RegisterMovieSceneChangedDelegate();
 
 	/** Movie Scene Data Changed Delegate */
-	void OnMovieSceneDataChanged(EMovieSceneDataChangeType InChangeType);
+	UE_API void OnMovieSceneDataChanged(EMovieSceneDataChangeType InChangeType);
 
 	/** Returns true if the given Sequence or any subsequence contains an audio track */
-	bool SequenceContainsAudioTrack(const UMovieSceneSequence* InSequence);
+	UE_API bool SequenceContainsAudioTrack(const UMovieSceneSequence* InSequence);
 
 	/** Will return true if a sequence contains an audio track and the user was notified about the potential clock source issue */
-	bool CheckSequenceClockSource();
+	UE_API bool CheckSequenceClockSource();
 
 	/** Prompts user and potentially modifies settings pref USequencerSettings::bAutoSelectAudioClockSource */
-	void PromptUserForClockSource();
+	UE_API void PromptUserForClockSource();
 
 	/** Sets the clock source for the given sequence to use the audio clock */
-	void SetClockSoureToAudioClock();
+	UE_API void SetClockSoureToAudioClock();
 
 private:
 
@@ -127,38 +129,38 @@ private:
 /**
  * Class for audio sections, handles drawing of all waveform previews.
  */
-class MOVIESCENETOOLS_API FAudioSection
+class FAudioSection
 	: public ISequencerSection
 	, public TSharedFromThis<FAudioSection>
 {
 public:
 
 	/** Constructor. */
-	FAudioSection(UMovieSceneSection& InSection, TWeakPtr<ISequencer> InSequencer);
+	UE_API FAudioSection(UMovieSceneSection& InSection, TWeakPtr<ISequencer> InSequencer);
 
 	/** Virtual destructor. */
-	virtual ~FAudioSection();
+	UE_API virtual ~FAudioSection();
 
 public:
 
 	// ISequencerSection interface
 
-	virtual UMovieSceneSection* GetSectionObject() override;
-	virtual FText GetSectionTitle() const override;
-	virtual FText GetSectionToolTip() const override;
-	virtual float GetSectionHeight(const UE::Sequencer::FViewDensityInfo& ViewDensity) const override;
-	virtual int32 OnPaintSection(FSequencerSectionPainter& Painter) const override;
-	virtual void Tick(const FGeometry& AllottedGeometry, const FGeometry& ParentGeometry, const double InCurrentTime, const float InDeltaTime) override;
-	virtual void BeginResizeSection() override;
-	virtual void ResizeSection(ESequencerSectionResizeMode ResizeMode, FFrameNumber ResizeTime) override;
-	virtual void BeginSlipSection() override;
-	virtual void SlipSection(FFrameNumber SlipTime) override;
-	virtual TOptional<FFrameTime> GetSectionTime(FSequencerSectionPainter& InPainter) const override;
+	UE_API virtual UMovieSceneSection* GetSectionObject() override;
+	UE_API virtual FText GetSectionTitle() const override;
+	UE_API virtual FText GetSectionToolTip() const override;
+	UE_API virtual float GetSectionHeight(const UE::Sequencer::FViewDensityInfo& ViewDensity) const override;
+	UE_API virtual int32 OnPaintSection(FSequencerSectionPainter& Painter) const override;
+	UE_API virtual void Tick(const FGeometry& AllottedGeometry, const FGeometry& ParentGeometry, const double InCurrentTime, const float InDeltaTime) override;
+	UE_API virtual void BeginResizeSection() override;
+	UE_API virtual void ResizeSection(ESequencerSectionResizeMode ResizeMode, FFrameNumber ResizeTime) override;
+	UE_API virtual void BeginSlipSection() override;
+	UE_API virtual void SlipSection(FFrameNumber SlipTime) override;
+	UE_API virtual TOptional<FFrameTime> GetSectionTime(FSequencerSectionPainter& InPainter) const override;
 	
 private:
 
 	/* Re-creates the texture used to preview the waveform. */
-	void RegenerateWaveforms(TRange<float> DrawRange, int32 XOffset, int32 XSize, const FColor& ColorTint, float DisplayScale);
+	UE_API void RegenerateWaveforms(TRange<float> DrawRange, int32 XOffset, int32 XSize, const FColor& ColorTint, float DisplayScale);
 
 private:
 
@@ -190,3 +192,5 @@ private:
 	/** Cached start time valid only during resize */
 	FFrameNumber InitialStartTimeDuringResize;
 };
+
+#undef UE_API
diff --git a/Engine/Source/Editor/MovieSceneTools/Public/TrackEditors/BindingLifetimeTrackEditor.h b/Engine/Source/Editor/MovieSceneTools/Public/TrackEditors/BindingLifetimeTrackEditor.h
--- a/Engine/Source/Editor/MovieSceneTools/Public/TrackEditors/BindingLifetimeTrackEditor.h
+++ b/Engine/Source/Editor/MovieSceneTools/Public/TrackEditors/BindingLifetimeTrackEditor.h
@@ -13,6 +13,8 @@
 #include "MovieSceneTrackEditor.h"
 #include "UObject/NameTypes.h"
 
+#define UE_API MOVIESCENETOOLS_API
+
 class FMenuBuilder;
 class ISequencer;
 class SWidget;
@@ -27,7 +29,7 @@ struct FGuid;
 /**
  * A track editor for controlling the lifetime of an object binding
  */
-class MOVIESCENETOOLS_API FBindingLifetimeTrackEditor
+class FBindingLifetimeTrackEditor
 	: public FMovieSceneTrackEditor
 {
 public:
@@ -38,7 +40,7 @@ public:
 	 * @param InSequencer The sequencer instance to be used by this tool.
 	 * @return The new instance of this class.
 	 */
-	static TSharedRef<ISequencerTrackEditor> CreateTrackEditor(TSharedRef<ISequencer> InSequencer);
+	static UE_API TSharedRef<ISequencerTrackEditor> CreateTrackEditor(TSharedRef<ISequencer> InSequencer);
 
 public:
 
@@ -47,30 +49,32 @@ public:
 	 *
 	 * @param InSequencer The sequencer instance to be used by this tool.
 	 */
-	FBindingLifetimeTrackEditor(TSharedRef<ISequencer> InSequencer);
+	UE_API FBindingLifetimeTrackEditor(TSharedRef<ISequencer> InSequencer);
 
-	void CreateNewSection(UMovieSceneTrack* Track, bool bSelect);
+	UE_API void CreateNewSection(UMovieSceneTrack* Track, bool bSelect);
 
 public:
 
 	// ISequencerTrackEditor interface
 
-	virtual FText GetDisplayName() const override;
+	UE_API virtual FText GetDisplayName() const override;
 
-	virtual UMovieSceneTrack* AddTrack(UMovieScene* FocusedMovieScene, const FGuid& ObjectHandle, TSubclassOf<class UMovieSceneTrack> TrackClass, FName UniqueTypeName) override;
-	virtual void BuildObjectBindingTrackMenu(FMenuBuilder& MenuBuilder, const TArray<FGuid>& ObjectBindings, const UClass* ObjectClass) override;
+	UE_API virtual UMovieSceneTrack* AddTrack(UMovieScene* FocusedMovieScene, const FGuid& ObjectHandle, TSubclassOf<class UMovieSceneTrack> TrackClass, FName UniqueTypeName) override;
+	UE_API virtual void BuildObjectBindingTrackMenu(FMenuBuilder& MenuBuilder, const TArray<FGuid>& ObjectBindings, const UClass* ObjectClass) override;
 	
-	virtual TSharedPtr<SWidget> BuildOutlinerEditWidget(const FGuid& ObjectBinding, UMovieSceneTrack* Track, const FBuildEditWidgetParams& Params) override;
+	UE_API virtual TSharedPtr<SWidget> BuildOutlinerEditWidget(const FGuid& ObjectBinding, UMovieSceneTrack* Track, const FBuildEditWidgetParams& Params) override;
 	
 	virtual bool HandleAssetAdded(UObject* Asset, const FGuid& TargetObjectGuid) override { return false; }
-	virtual bool SupportsType(TSubclassOf<UMovieSceneTrack> Type) const override;
-	virtual bool SupportsSequence(UMovieSceneSequence* InSequence) const override;
+	UE_API virtual bool SupportsType(TSubclassOf<UMovieSceneTrack> Type) const override;
+	UE_API virtual bool SupportsSequence(UMovieSceneSequence* InSequence) const override;
 
-	virtual TSharedRef<ISequencerSection> MakeSectionInterface(UMovieSceneSection& SectionObject, UMovieSceneTrack& Track, FGuid ObjectBinding) override;
+	UE_API virtual TSharedRef<ISequencerSection> MakeSectionInterface(UMovieSceneSection& SectionObject, UMovieSceneTrack& Track, FGuid ObjectBinding) override;
 
 private:
 
 	/** Callback for executing the "Add Binding Lifetime Track" menu entry. */
-	void HandleAddBindingLifetimeTrackMenuEntryExecute(TArray<FGuid> ObjectBindings);
-	bool CanAddBindingLifetimeTrack(FGuid ObjectBinding) const;
+	UE_API void HandleAddBindingLifetimeTrackMenuEntryExecute(TArray<FGuid> ObjectBindings);
+	UE_API bool CanAddBindingLifetimeTrack(FGuid ObjectBinding) const;
 };
+
+#undef UE_API
diff --git a/Engine/Source/Editor/MovieSceneTools/Public/TrackEditors/CinematicShotTrackEditor.h b/Engine/Source/Editor/MovieSceneTools/Public/TrackEditors/CinematicShotTrackEditor.h
--- a/Engine/Source/Editor/MovieSceneTools/Public/TrackEditors/CinematicShotTrackEditor.h
+++ b/Engine/Source/Editor/MovieSceneTools/Public/TrackEditors/CinematicShotTrackEditor.h
@@ -12,6 +12,8 @@
 #include "SubTrackEditor.h"
 #include "ISequencerTrackEditor.h"
 
+#define UE_API MOVIESCENETOOLS_API
+
 class AActor;
 class FMenuBuilder;
 class FTrackEditorThumbnailPool;
@@ -22,7 +24,7 @@ class UMovieSceneSubSection;
 /**
  * Tools for cinematic shots.
  */
-class MOVIESCENETOOLS_API FCinematicShotTrackEditor : public FSubTrackEditor
+class FCinematicShotTrackEditor : public FSubTrackEditor
 {
 public:
 
@@ -31,7 +33,7 @@ public:
 	 *
 	 * @param InSequencer The sequencer instance to be used by this tool.
 	 */
-	FCinematicShotTrackEditor(TSharedRef<ISequencer> InSequencer);
+	UE_API FCinematicShotTrackEditor(TSharedRef<ISequencer> InSequencer);
 
 	/** Virtual destructor. */
 	virtual ~FCinematicShotTrackEditor() { }
@@ -42,69 +44,71 @@ public:
 	 * @param OwningSequencer The sequencer instance to be used by this tool.
 	 * @return The new instance of this class.
 	 */
-	static TSharedRef<ISequencerTrackEditor> CreateTrackEditor(TSharedRef<ISequencer> OwningSequencer);
+	static UE_API TSharedRef<ISequencerTrackEditor> CreateTrackEditor(TSharedRef<ISequencer> OwningSequencer);
 
 	UE_DEPRECATED(5.5, "Use FCameraCutPlaybackCapability::LastViewTargetCamera instead.")
-	TWeakObjectPtr<AActor> GetCinematicShotCamera() const;
+	UE_API TWeakObjectPtr<AActor> GetCinematicShotCamera() const;
 
 public:
 
 	// ISequencerTrackEditor interface
-	virtual FText GetDisplayName() const override;
-	virtual TSharedPtr<SWidget> BuildOutlinerColumnWidget(const FBuildColumnWidgetParams& Params, const FName& ColumnName) override;
-	virtual TSharedRef<ISequencerSection> MakeSectionInterface(UMovieSceneSection& SectionObject, UMovieSceneTrack& Track, FGuid ObjectBinding) override;
-	virtual bool SupportsSequence(UMovieSceneSequence* InSequence) const override;
-	virtual void Tick(float DeltaTime) override;
-	virtual void BuildTrackContextMenu(FMenuBuilder& MenuBuilder, UMovieSceneTrack* Track) override;
+	UE_API virtual FText GetDisplayName() const override;
+	UE_API virtual TSharedPtr<SWidget> BuildOutlinerColumnWidget(const FBuildColumnWidgetParams& Params, const FName& ColumnName) override;
+	UE_API virtual TSharedRef<ISequencerSection> MakeSectionInterface(UMovieSceneSection& SectionObject, UMovieSceneTrack& Track, FGuid ObjectBinding) override;
+	UE_API virtual bool SupportsSequence(UMovieSceneSequence* InSequence) const override;
+	UE_API virtual void Tick(float DeltaTime) override;
+	UE_API virtual void BuildTrackContextMenu(FMenuBuilder& MenuBuilder, UMovieSceneTrack* Track) override;
 
 	/*
 	 * Render shots. 
 	 *
 	 * @param Sections The sections to render
 	 */
-	void RenderShots(const TArray<UMovieSceneCinematicShotSection*>& Sections);
+	UE_API void RenderShots(const TArray<UMovieSceneCinematicShotSection*>& Sections);
 
 public:
 
 	// FSubTrackEditor interface
-	virtual FText GetSubTrackName() const override;
-	virtual FText GetSubTrackToolTip() const override;
-	virtual FName GetSubTrackBrushName() const override;
-	virtual FString GetSubSectionDisplayName(const UMovieSceneSubSection* Section) const override;
-	virtual FString GetDefaultSubsequenceName() const override;
-	virtual FString GetDefaultSubsequenceDirectory() const override;
-	virtual TSubclassOf<UMovieSceneSubTrack> GetSubTrackClass() const;
+	UE_API virtual FText GetSubTrackName() const override;
+	UE_API virtual FText GetSubTrackToolTip() const override;
+	UE_API virtual FName GetSubTrackBrushName() const override;
+	UE_API virtual FString GetSubSectionDisplayName(const UMovieSceneSubSection* Section) const override;
+	UE_API virtual FString GetDefaultSubsequenceName() const override;
+	UE_API virtual FString GetDefaultSubsequenceDirectory() const override;
+	UE_API virtual TSubclassOf<UMovieSceneSubTrack> GetSubTrackClass() const;
 
 protected:
 
-	virtual bool HandleAddSubTrackMenuEntryCanExecute() const override;
-	virtual bool CanHandleAssetAdded(UMovieSceneSequence* Sequence) const override;
+	UE_API virtual bool HandleAddSubTrackMenuEntryCanExecute() const override;
+	UE_API virtual bool CanHandleAssetAdded(UMovieSceneSequence* Sequence) const override;
 
 private:
 
 	/** Delegate for shots button lock state */
-	ECheckBoxState AreShotsLocked() const;
+	UE_API ECheckBoxState AreShotsLocked() const;
 
 	/** Delegate for locked shots button */
-	void OnLockShotsClicked(ECheckBoxState CheckBoxState);
+	UE_API void OnLockShotsClicked(ECheckBoxState CheckBoxState);
 	
 	/** Delegate for shots button lock tooltip */
-	FText GetLockShotsToolTip() const;
+	UE_API FText GetLockShotsToolTip() const;
 
 	/** Callback for ImportEDL. */
-	void ImportEDL();
+	UE_API void ImportEDL();
 	
 	/** Callback for ExportEDL. */
-	void ExportEDL();
+	UE_API void ExportEDL();
 
 	/** Callback for ImportFCPXML. */
-	void ImportFCPXML();
+	UE_API void ImportFCPXML();
 
 	/** Callback for ExportFCPXML. */
-	void ExportFCPXML();
+	UE_API void ExportFCPXML();
 
 private:
 
 	/** The Thumbnail pool which draws all the viewport thumbnails for the shot track. */
 	TSharedPtr<FTrackEditorThumbnailPool> ThumbnailPool;
 };
+
+#undef UE_API
