Update transactions when editing subtrack transform origins.
Use one transaction, and start and end it when dragging the gizmo begins and ends. Remove transactions from individual updates, and when resetting on a sequence scrub.

This leads to better undo support, and prevents extra transactions from being created.

diff --git a/Engine/Source/Editor/MovieSceneTools/Private/EditModes/SubTrackEditorMode.cpp b/Engine/Source/Editor/MovieSceneTools/Private/EditModes/SubTrackEditorMode.cpp
--- a/Engine/Source/Editor/MovieSceneTools/Private/EditModes/SubTrackEditorMode.cpp
+++ b/Engine/Source/Editor/MovieSceneTools/Private/EditModes/SubTrackEditorMode.cpp
@@ -20,6 +20,8 @@
 #include "Tracks/IMovieSceneTransformOrigin.h"
 #include "Tracks/MovieSceneSubTrack.h"
 
+#define LOCTEXT_NAMESPACE "FSubTrackEditorMode"
+
 FName FSubTrackEditorMode::ModeName("EditMode.SubTrackEditMode");
 
 FSubTrackEditorMode::FSubTrackEditorMode()
@@ -113,6 +115,7 @@ bool FSubTrackEditorMode::HandleBeginTransform()
 		{
 			PreviewLocation = GetAverageLocationOfBindingsInSubSection(SubSection);
 		}
+		GEditor->BeginTransaction(LOCTEXT("EditSubseqeunceTransformOriginTransaction", "Edit Subsequence Transform Origin"));
 		return true;
 	}
 
@@ -127,6 +130,7 @@ bool FSubTrackEditorMode::HandleEndTransform()
 		// Only reset the preview rotation. Resetting the preview location here could interfere with multiple rotation edits in a row.
 		// The location would otherwise change after the user let go of the mouse, and would have to move it to the new location to continue rotating.
 		PreviewCoordinateSpaceRotation.Reset();
+		GEditor->EndTransaction();
 		return true;
 	}
 
@@ -538,4 +542,6 @@ UMovieSceneSubSection* FSubTrackEditorMode::GetSectionToEdit() const
 	}
 
 	return nullptr;
-}
\ No newline at end of file
+}
+
+#undef LOCTEXT_NAMESPACE
\ No newline at end of file
diff --git a/Engine/Source/Editor/MovieSceneTools/Private/TrackEditors/SubTrackEditor.cpp b/Engine/Source/Editor/MovieSceneTools/Private/TrackEditors/SubTrackEditor.cpp
--- a/Engine/Source/Editor/MovieSceneTools/Private/TrackEditors/SubTrackEditor.cpp
+++ b/Engine/Source/Editor/MovieSceneTools/Private/TrackEditors/SubTrackEditor.cpp
@@ -1354,10 +1354,6 @@ void FSubTrackEditor::UpdateOrigin(FVector InPosition, FRotator InRotation)
 	
 	// @todo Sequencer - The sequencer probably should have taken care of this
 	MovieSceneSequence->SetFlags(RF_Transactional);
-	
-	// Create a transaction record because we are about to add keys
-	const bool bShouldActuallyTransact = !GIsTransacting;		// Don't transact if we're recording in a PIE world.  That type of keyframe capture cannot be undone.
-	FScopedTransaction AutoKeyTransaction( NSLOCTEXT("AnimatablePropertyTool", "PropertyChanged", "Animatable Property Changed"), bShouldActuallyTransact);
 
 	TArray<UMovieSceneSubSection*> SectionsToKey;
 
@@ -1400,23 +1396,20 @@ void FSubTrackEditor::UpdateOrigin(FVector InPosition, FRotator InRotation)
 void FSubTrackEditor::ResetSectionPreviews()
 {
 	bool bSectionReverted = false;
-
-	// Create a transaction record because we are about to add keys
-	const bool bShouldActuallyTransact = !GIsTransacting;		// Don't transact if we're recording in a PIE world.  That type of keyframe capture cannot be undone.
-	FScopedTransaction OriginPreviewResetTransaction( NSLOCTEXT("AnimatablePropertyTool", "PropertyChanged", "Animatable Property Changed"), bShouldActuallyTransact);
 	
-	for(UMovieSceneSubSection* SubSection : SectionsWithPreviews)
+	for (UMovieSceneSubSection* SubSection : SectionsWithPreviews)
 	{
-		if(!SubSection)
+		if (SubSection && (SubSection->GetKeyPreviewPosition().IsSet() || SubSection->GetKeyPreviewRotation().IsSet()))
 		{
-			continue;
+			SubSection->ResetKeyPreviewRotationAndLocation();
+			SubSection->Modify(false);
+			SubSection->MarkAsChanged();
+			bSectionReverted = true;
 		}
-		SubSection->ResetKeyPreviewRotationAndLocation();
-		SubSection->Modify(false);
-		SubSection->MarkAsChanged();
-		bSectionReverted = true;
 	}
 
+	SectionsWithPreviews.Empty();
+
 	FSubTrackEditorMode* EditorMode = static_cast<FSubTrackEditorMode*>(GLevelEditorModeTools().GetActiveMode(FSubTrackEditorMode::ModeName));
 
 	if (EditorMode)
@@ -1424,7 +1417,7 @@ void FSubTrackEditor::ResetSectionPreviews()
 		EditorMode->ClearCachedCoordinates();
 	}
 	
-	if(bSectionReverted)
+	if (bSectionReverted)
 	{
 		GetSequencer()->NotifyMovieSceneDataChanged(EMovieSceneDataChangeType::TrackValueChanged);		
 	}
