Curve Editor: Fix tangents are not updated when tangent visibility changes

diff --git a/Engine/Source/Editor/MovieSceneTools/Private/Cache/MovieSceneCachedCurve.h b/Engine/Source/Editor/MovieSceneTools/Private/Cache/MovieSceneCachedCurve.h
--- a/Engine/Source/Editor/MovieSceneTools/Private/Cache/MovieSceneCachedCurve.h
+++ b/Engine/Source/Editor/MovieSceneTools/Private/Cache/MovieSceneCachedCurve.h
@@ -471,7 +471,12 @@ namespace UE::MovieSceneTools
 	template <typename ChannelType>
 	void FMovieSceneCachedCurve<ChannelType>::ConditionallyUpdateKeyDrawInfos(const FMovieSceneUpdateCachedCurveData<ChannelType>& UpdateData)
 	{
-		if (EnumHasAnyFlags(Flags.load(), EMovieSceneCurveCacheChangeFlags::ChangedCurveData | EMovieSceneCurveCacheChangeFlags::ChangedSelection))
+		const bool bNeedsUpdate = EnumHasAnyFlags(Flags.load(),
+			EMovieSceneCurveCacheChangeFlags::ChangedCurveData |
+			EMovieSceneCurveCacheChangeFlags::ChangedSelection |
+			EMovieSceneCurveCacheChangeFlags::ChangedTangentVisibility);
+
+		if (bNeedsUpdate)
 		{
 			UpdateData.CurveModel.GetKeyDrawInfo(ECurvePointType::ArriveTangent, FKeyHandle::Invalid(), CachedDrawParams.ArriveTangentDrawInfo);
 			UpdateData.CurveModel.GetKeyDrawInfo(ECurvePointType::LeaveTangent, FKeyHandle::Invalid(), CachedDrawParams.LeaveTangentDrawInfo);
