Sequencer: Clear events if they don't have the same director

If you copy an event track from one sequence to another, the events in the pasted event track will point back to the blueprint director in the copied sequence. This is no good, so those events need to be cleared.

diff --git a/Engine/Source/Editor/MovieSceneTools/Private/MovieSceneEventUtils.cpp b/Engine/Source/Editor/MovieSceneTools/Private/MovieSceneEventUtils.cpp
--- a/Engine/Source/Editor/MovieSceneTools/Private/MovieSceneEventUtils.cpp
+++ b/Engine/Source/Editor/MovieSceneTools/Private/MovieSceneEventUtils.cpp
@@ -41,6 +41,7 @@
 #include "MovieSceneEventBlueprintExtension.h"
 #include "MovieSceneFwd.h"
 #include "MovieScenePossessable.h"
+#include "MovieSceneSequence.h"
 #include "MovieSceneSpawnable.h"
 #include "ScopedTransaction.h"
 #include "Sections/MovieSceneEventSectionBase.h"
@@ -306,6 +307,59 @@ void FMovieSceneEventUtils::RemoveEndpointsForEventSection(UMovieSceneEventSecti
 	}
 }
 
+void FMovieSceneEventUtils::ClearInvalidEndpointsForEventSection(UMovieSceneEventSectionBase* EventSection)
+{
+	check(EventSection);
+
+	UMovieSceneSequence* Sequence = EventSection->GetTypedOuter<UMovieSceneSequence>();
+	if (Sequence)
+	{
+		TArray<TTuple<FMovieSceneEvent*, UObject*, UMovieSceneSequence*>> InvalidEntryPoints;
+
+		for (FMovieSceneEvent& EntryPoint : EventSection->GetAllEntryPoints())
+		{
+			TWeakObjectPtr<UObject> EndPoint = EntryPoint.WeakEndpoint;
+			if (!EndPoint.IsValid())
+			{
+				continue;
+			}
+
+			UBlueprint* FunctionBlueprint = EndPoint->GetTypedOuter<UBlueprint>();
+			if (!FunctionBlueprint)
+			{
+				continue;
+			}
+
+			UMovieSceneSequence* FunctionBlueprintSequence = Cast<UMovieSceneSequence>(FunctionBlueprint->GetOuter());
+			if (FunctionBlueprintSequence != nullptr && FunctionBlueprintSequence != Sequence)
+			{
+				InvalidEntryPoints.Emplace(&EntryPoint, EndPoint.Get(), FunctionBlueprintSequence);
+			}
+		}
+
+		if (InvalidEntryPoints.Num())
+		{
+			FScopedTransaction ClearInvalidEndpoints(LOCTEXT("ClearInvalidEndpoint", "Clear Invalid Endpoints"));
+
+			for (TTuple<FMovieSceneEvent*, UObject*, UMovieSceneSequence*> InvalidEntryPoint : InvalidEntryPoints)
+			{
+				FMovieSceneEvent* EntryPoint = InvalidEntryPoint.Get<0>();
+				UObject* EndPoint = InvalidEntryPoint.Get<1>();
+				UMovieSceneSequence* FunctionBlueprintSequence = InvalidEntryPoint.Get<2>();
+
+				FNotificationInfo Info(FText::Format(LOCTEXT("ClearInvalidEndpointNotify", "Clearing event {0} from {1} because it refers to a different sequence {2}"), FText::FromString(GetNameSafe(EndPoint)), FText::FromString(GetNameSafe(Sequence)), FText::FromString(GetNameSafe(FunctionBlueprintSequence))));
+				Info.ExpireDuration = 3.f;
+				FSlateNotificationManager::Get().AddNotification(Info);
+
+				UE_LOG(LogMovieScene, Display, TEXT("Clearing event %s from %s because it refers to a different sequence %s"), *GetNameSafe(EndPoint), *GetNameSafe(Sequence), *GetNameSafe(FunctionBlueprintSequence));
+
+				// Clear this entry point
+				*EntryPoint = FMovieSceneEvent();
+			}
+		}
+	}
+}
+
 void FMovieSceneEventUtils::RemoveUnusedCustomEvents(const TArray<TWeakObjectPtr<UMovieSceneEventSectionBase>>& EventSections, UBlueprint* DirectorBP)
 {
 	check(DirectorBP);
diff --git a/Engine/Source/Editor/MovieSceneTools/Private/MovieSceneToolsModule.cpp b/Engine/Source/Editor/MovieSceneTools/Private/MovieSceneToolsModule.cpp
--- a/Engine/Source/Editor/MovieSceneTools/Private/MovieSceneToolsModule.cpp
+++ b/Engine/Source/Editor/MovieSceneTools/Private/MovieSceneToolsModule.cpp
@@ -427,6 +427,8 @@ void FMovieSceneToolsModule::PostDuplicateEventSection(UMovieSceneEventSectionBa
 	FMovieSceneSequenceEditor* SequenceEditor     = FMovieSceneSequenceEditor::Find(Sequence);
 	UBlueprint*                SequenceDirectorBP = SequenceEditor ? SequenceEditor->FindDirectorBlueprint(Sequence) : nullptr;
 
+	FMovieSceneEventUtils::ClearInvalidEndpointsForEventSection(Section);
+
 	if (SequenceDirectorBP)
 	{
 		// Always bind the event section onto the blueprint to ensure that we get another chance to upgrade when the BP compiles if this try wasn't successful
diff --git a/Engine/Source/Editor/MovieSceneTools/Public/MovieSceneEventUtils.h b/Engine/Source/Editor/MovieSceneTools/Public/MovieSceneEventUtils.h
--- a/Engine/Source/Editor/MovieSceneTools/Public/MovieSceneEventUtils.h
+++ b/Engine/Source/Editor/MovieSceneTools/Public/MovieSceneEventUtils.h
@@ -79,6 +79,13 @@ struct FMovieSceneEventUtils
 	 */
 	static UE_API void RemoveEndpointsForEventSection(UMovieSceneEventSectionBase* EventSection, UBlueprint* DirectorBP);
 
+	/**
+	 * Clear invalid endpoints (end points to that point to a blueprint in a different movie scene, possibly from copy/paste)
+	 *
+	 * @param EventSection           (Required, non-null) The event section to search for invalid event endpoints from
+	 */
+	static UE_API void ClearInvalidEndpointsForEventSection(UMovieSceneEventSectionBase* EventSection);
+
 	/**
 	 * Remove unused custom events (that don't have corresponding event end points in the given event sections)
 	 *
