Sequencer: Add IsNodeAllowed to FMovieSceneObjectBindingIDPicker so that derived classes can determine whether a binding is allowable or not.

FProxyObjectBindingIDPicker disallows bindings in the currently focused sequence

diff --git a/Engine/Plugins/MovieScene/LevelSequenceEditor/Source/LevelSequenceEditor/Private/LevelSequenceEditorSubsystem.cpp b/Engine/Plugins/MovieScene/LevelSequenceEditor/Source/LevelSequenceEditor/Private/LevelSequenceEditorSubsystem.cpp
--- a/Engine/Plugins/MovieScene/LevelSequenceEditor/Source/LevelSequenceEditor/Private/LevelSequenceEditorSubsystem.cpp
+++ b/Engine/Plugins/MovieScene/LevelSequenceEditor/Source/LevelSequenceEditor/Private/LevelSequenceEditorSubsystem.cpp
@@ -11,6 +11,7 @@
 #include "ISceneOutliner.h"
 #include "LevelSequenceEditorCommands.h"
 #include "MovieScenePossessable.h"
+#include "SequenceBindingTree.h"
 #include "SequencerSettings.h"
 #include "MovieScene.h"
 #include "MovieSceneSpawnable.h"
@@ -142,6 +143,18 @@ namespace UE::Sequencer
 			return FMovieSceneObjectBindingID();
 		}
 
+		bool IsNodeAllowed(const TSharedPtr<FSequenceBindingNode>& Node) const override
+		{
+			// Disallow picking a binding from this current sequence
+			const UE::MovieScene::FFixedObjectBindingID &BindingID = Node->BindingID;
+			if (BindingID.SequenceID == LocalSequenceID)
+			{
+				return false;
+			}
+
+			return FMovieSceneObjectBindingIDPicker::IsNodeAllowed(Node);
+		}
+
 	private:
 		FGuid ObjectBindingID;
 		TFunction<void(FMovieSceneObjectBindingID)> PreOnPicked;
diff --git a/Engine/Source/Editor/MovieSceneTools/Private/MovieSceneObjectBindingIDPicker.cpp b/Engine/Source/Editor/MovieSceneTools/Private/MovieSceneObjectBindingIDPicker.cpp
--- a/Engine/Source/Editor/MovieSceneTools/Private/MovieSceneObjectBindingIDPicker.cpp
+++ b/Engine/Source/Editor/MovieSceneTools/Private/MovieSceneObjectBindingIDPicker.cpp
@@ -71,7 +71,7 @@ void FMovieSceneObjectBindingIDPicker::OnGetMenuContent(FMenuBuilder& MenuBuilde
 
 	bool bHadAnyEntries = false;
 
-	if (Node->BindingID.Guid.IsValid() && IsClassAllowed(Node))
+	if (Node->BindingID.Guid.IsValid() && IsNodeAllowed(Node))
 	{
 		bHadAnyEntries = true;
 		MenuBuilder.AddMenuEntry(
@@ -103,7 +103,7 @@ void FMovieSceneObjectBindingIDPicker::OnGetMenuContent(FMenuBuilder& MenuBuilde
 					);
 			}
 		}
-		else if (IsClassAllowed(Child))
+		else if (IsNodeAllowed(Child))
 		{
 			bHadAnyEntries = true;
 			MenuBuilder.AddMenuEntry(
@@ -213,10 +213,15 @@ FReply FMovieSceneObjectBindingIDPicker::AttemptBindingFixup()
 	return FReply::Handled();
 }
 
+bool FMovieSceneObjectBindingIDPicker::IsNodeAllowed(const TSharedPtr<FSequenceBindingNode>& Node) const
+{
+	return IsClassAllowed(Node);
+}
+
 bool FMovieSceneObjectBindingIDPicker::IsClassAllowed(const TSharedPtr<FSequenceBindingNode>& Node) const
 {
 	const UClass* InClass = nullptr;
-	auto BindingID = Node->BindingID;
+	const UE::MovieScene::FFixedObjectBindingID& BindingID = Node->BindingID;
 
 	TSharedPtr<ISequencer> Sequencer = WeakSequencer.Pin();
 	UMovieSceneSequence* BindSequence = Sequencer.IsValid() ? Sequencer->GetEvaluationState()->FindSequence(BindingID.SequenceID) : nullptr;
diff --git a/Engine/Source/Editor/MovieSceneTools/Public/MovieSceneObjectBindingIDPicker.h b/Engine/Source/Editor/MovieSceneTools/Public/MovieSceneObjectBindingIDPicker.h
--- a/Engine/Source/Editor/MovieSceneTools/Public/MovieSceneObjectBindingIDPicker.h
+++ b/Engine/Source/Editor/MovieSceneTools/Public/MovieSceneObjectBindingIDPicker.h
@@ -76,6 +76,9 @@ protected:
 	/** Whether there are multiple values */
 	virtual bool HasMultipleValues() const { return false; }
 
+	/** Checks if node is allowed by this binding */
+	UE_API virtual bool IsNodeAllowed(const TSharedPtr<FSequenceBindingNode>& Node) const;
+
 protected:
 
 	/** Initialize this class - rebuilds sequence hierarchy data and available IDs from the source sequence */
