Curve Editor: Set key details box to fill and center/space the frame number appropriately

#rnx

diff --git a/Engine/Source/Editor/CurveEditor/Private/SCurveEditorPanel.cpp b/Engine/Source/Editor/CurveEditor/Private/SCurveEditorPanel.cpp
--- a/Engine/Source/Editor/CurveEditor/Private/SCurveEditorPanel.cpp
+++ b/Engine/Source/Editor/CurveEditor/Private/SCurveEditorPanel.cpp
@@ -1433,6 +1433,9 @@ void SCurveEditorPanel::BuildToolbar(FToolBarBuilder& InToolBarBuilder, TSharedP
 
 	InToolBarBuilder.BeginSection("Key Details");
 	{
+		FMenuEntryStyleParams StyleParams;
+		StyleParams.VerticalAlignment = VAlign_Fill;
+
 		InToolBarBuilder.AddWidget(
 			SNew(SBox)
 			.HAlign(HAlign_Center)
@@ -1440,7 +1443,7 @@ void SCurveEditorPanel::BuildToolbar(FToolBarBuilder& InToolBarBuilder, TSharedP
 			[
 				KeyDetailsView.ToSharedRef()				
 			]
-			, {}, NAME_None, true, {}, {}, MakeResizeParams(TEXT("KeyDetails")));
+			, StyleParams, NAME_None, true, {}, {}, MakeResizeParams(TEXT("KeyDetails")));
 	}
 	InToolBarBuilder.EndSection();
 
diff --git a/Engine/Source/Editor/CurveEditor/Private/SCurveKeyDetailPanel.cpp b/Engine/Source/Editor/CurveEditor/Private/SCurveKeyDetailPanel.cpp
--- a/Engine/Source/Editor/CurveEditor/Private/SCurveKeyDetailPanel.cpp
+++ b/Engine/Source/Editor/CurveEditor/Private/SCurveKeyDetailPanel.cpp
@@ -185,14 +185,15 @@ void SCurveKeyDetailPanel::ConstructChildLayout(TSharedPtr<SWidget> TimeWidget,
 	ChildSlot
 	[
 		SNew(SBox)
+		.VAlign(VAlign_Fill)
 		.MaxDesiredWidth(FAppStyle::Get().GetFloat("CurveEditor.KeyDetailWidth"))
 		[
 			SNew(SHorizontalBox)
 
 			// "Time" Edit box
 			+ SHorizontalBox::Slot()
-			.VAlign(VAlign_Center)
-			.Padding(4.f, 0.f, 0.f, 2.f)
+			.VAlign(VAlign_Fill)
+			.Padding(4.f, 0.f, 0.f, 0.f)
 			.FillWidth(0.5f)
 			[
 				TimeWidget.ToSharedRef()
@@ -201,9 +202,9 @@ void SCurveKeyDetailPanel::ConstructChildLayout(TSharedPtr<SWidget> TimeWidget,
 
 			// "Value" Edit box
 			+ SHorizontalBox::Slot()
-			.VAlign(VAlign_Center)
+			.VAlign(VAlign_Fill)
+			.Padding(4.f, 0.f, 0.f, 0.f)
 			.FillWidth(0.5f)
-			.Padding(4.f, 0.f, 0.f, 2.f)
 			[
 				ValueWidget.ToSharedRef()
 			]
diff --git a/Engine/Source/Editor/MovieSceneTools/Private/FrameNumberDetailsCustomization.cpp b/Engine/Source/Editor/MovieSceneTools/Private/FrameNumberDetailsCustomization.cpp
--- a/Engine/Source/Editor/MovieSceneTools/Private/FrameNumberDetailsCustomization.cpp
+++ b/Engine/Source/Editor/MovieSceneTools/Private/FrameNumberDetailsCustomization.cpp
@@ -74,14 +74,19 @@ void FFrameNumberDetailsCustomization::CustomizeChildren(TSharedRef<IPropertyHan
 		]
 		.ValueContent()
 		[
-			SNew(SEditableTextBox)
-			.Text(this, &FFrameNumberDetailsCustomization::OnGetTimeText)
-			.ToolTipText(this, &FFrameNumberDetailsCustomization::OnGetTimeToolTipText)
-			.OnTextCommitted(this, &FFrameNumberDetailsCustomization::OnTimeTextCommitted)
-			.SelectAllTextWhenFocused(true)
-			.ClearKeyboardFocusOnCommit(false)
-			.RevertTextOnEscape(true)
-			.Font(IDetailLayoutBuilder::GetDetailFont())
+			SNew(SBox)
+			.VAlign(VAlign_Fill)
+			[
+				SNew(SEditableTextBox)
+				.Text(this, &FFrameNumberDetailsCustomization::OnGetTimeText)
+				.ToolTipText(this, &FFrameNumberDetailsCustomization::OnGetTimeToolTipText)
+				.OnTextCommitted(this, &FFrameNumberDetailsCustomization::OnTimeTextCommitted)
+				.SelectAllTextWhenFocused(true)
+				.ClearKeyboardFocusOnCommit(false)
+				.RevertTextOnEscape(true)
+				.Font(IDetailLayoutBuilder::GetDetailFont())
+				.Padding(FMargin(8.0f, 4.0f, 8.0f, 4.0f))
+			]
 		];
 }
 
