Sequencer: Tweak focus behavior again

Users would like the value to be committed when focus is lost, but only if it hasn't changed.

diff --git a/Engine/Source/Editor/MovieSceneTools/Private/CurveKeyEditors/SNumericKeyEditor.h b/Engine/Source/Editor/MovieSceneTools/Private/CurveKeyEditors/SNumericKeyEditor.h
--- a/Engine/Source/Editor/MovieSceneTools/Private/CurveKeyEditors/SNumericKeyEditor.h
+++ b/Engine/Source/Editor/MovieSceneTools/Private/CurveKeyEditors/SNumericKeyEditor.h
@@ -26,6 +26,20 @@ struct SNonThrottledSpinBox : SSpinBox<T>
 		}
 		return Reply;
 	}
+
+	virtual FReply OnFocusReceived(const FGeometry& MyGeometry, const FFocusEvent& InFocusEvent) override
+	{
+		// Store the value when focus is received so that it can be compared in OnValueChanged which is sent after focus is lost
+		InitialValueOnFocused = SSpinBox<T>::GetValue();
+		return SSpinBox<T>::OnFocusReceived(MyGeometry, InFocusEvent);
+	}
+
+	virtual void OnFocusLost(const FFocusEvent& InFocusEvent) override
+	{
+		// Don't clear InitialValueOnFocused since OnValueChanged is sent after focus is lost and the value is needed to compare whether it has actually changed
+	}
+
+	TOptional<T> InitialValueOnFocused;
 };
 
 /**
@@ -78,7 +92,7 @@ public:
 
 		ChildSlot
 		[
-			SNew(SNonThrottledSpinBox<NumericType>)
+			SAssignNew(SpinBox, SNonThrottledSpinBox<NumericType>)
 			.Style(&FAppStyle::GetWidgetStyle<FSpinBoxStyle>("Sequencer.HyperlinkSpinBox"))
 			.Font(FAppStyle::GetFontStyle("Sequencer.AnimationOutliner.RegularFont"))
 			.MinValue(NumericPropertyParams.MinValue)
@@ -131,12 +145,19 @@ private:
 
 	void OnValueChanged(NumericType Value)
 	{
-		// Update the key editor only when spinning/sliding the value since the value hasn't been committed yet.
-		// This should only happen in this case because OnValueChanged is called aggressively (ie. on focus changing)
+		// Update the key editor when spinning/sliding the value since the value hasn't been committed yet
 		if (bSliding)
 		{
 			KeyEditor->SetValueWithNotify(Value, EMovieSceneDataChangeType::TrackValueChanged);
 		}
+		// Otherwise, ValueChanged will also be emitted on focus changing. Set the value only if is different from the initial focused value
+		else if (!bCommitted && SpinBox->InitialValueOnFocused.IsSet() && SpinBox->InitialValueOnFocused.GetValue() != Value)
+		{
+			const FScopedTransaction Transaction(LOCTEXT("SetNumericKey", "Set Key Value"));
+			KeyEditor->SetValueWithNotify(Value, EMovieSceneDataChangeType::TrackValueChangedRefreshImmediately);
+		}
+
+		bCommitted = false;
 	}
 
 	void OnValueCommitted(NumericType Value, ETextCommit::Type CommitInfo)
@@ -145,13 +166,17 @@ private:
 		{
 			const FScopedTransaction Transaction( LOCTEXT("SetNumericKey", "Set Key Value") );
 			KeyEditor->SetValueWithNotify(Value, EMovieSceneDataChangeType::TrackValueChangedRefreshImmediately);
+			bCommitted = true;
 		}
 	}
 
 private:
 
+	TSharedPtr<SNonThrottledSpinBox<NumericType>> SpinBox;
+
 	TSharedPtr<ISequencerKeyEditor<NumericType>> KeyEditor;
 	bool bSliding = false;
+	bool bCommitted = false;
 };
 
 template<typename ChannelType, typename NumericType>
