Sequencer: Fix ongoing focus issues in the key details

Edited values are usually committed when the user moves focus. This caused problems ages ago because if you focus on a key at time X and then change time to time Y, focus will be lost and the key value will be commited at time Y. The initial fix was to capture the focused widget when scrubbing time and then restore it. That was a bad hack and was therefore removed and decided that edited values should no longer by committed when the uesr moves focus. However, users didn't like that.

This change now commits values when the user changes focus. It addresses 2 problems:

- A key would be created when clicking on the key editor and then moving focus without changing the value. The fix is to not add a key if the value already exists at that time.

- A key will be modified at a different time when focus changes.  The fix introduces an "editing time" for the key editor. The editing time is set when the widget receives focus and is used whenever SetValue is called. The editing time is then cleared when the value is committed.

diff --git a/Engine/Source/Editor/MovieSceneTools/Private/CurveKeyEditors/SNumericKeyEditor.h b/Engine/Source/Editor/MovieSceneTools/Private/CurveKeyEditors/SNumericKeyEditor.h
--- a/Engine/Source/Editor/MovieSceneTools/Private/CurveKeyEditors/SNumericKeyEditor.h
+++ b/Engine/Source/Editor/MovieSceneTools/Private/CurveKeyEditors/SNumericKeyEditor.h
@@ -26,6 +26,21 @@ struct SNonThrottledSpinBox : SSpinBox<T>
 		}
 		return Reply;
 	}
+
+	virtual FReply OnFocusReceived(const FGeometry& MyGeometry, const FFocusEvent& InFocusEvent) override
+	{
+		KeyEditor->BeginEditing(KeyEditor->GetCurrentTime());
+		return SSpinBox<T>::OnFocusReceived(MyGeometry, InFocusEvent);
+	}
+
+	void SetKeyEditor(TSharedPtr<ISequencerKeyEditor<T>> InKeyEditor)
+	{
+		KeyEditor = InKeyEditor;
+	}
+
+private:
+
+	TSharedPtr<ISequencerKeyEditor<T>> KeyEditor;
 };
 
 /**
@@ -78,7 +93,7 @@ public:
 
 		ChildSlot
 		[
-			SNew(SNonThrottledSpinBox<NumericType>)
+			SAssignNew(SpinBox, SNonThrottledSpinBox<NumericType>)
 			.Style(&FAppStyle::GetWidgetStyle<FSpinBoxStyle>("Sequencer.HyperlinkSpinBox"))
 			.Font(FAppStyle::GetFontStyle("Sequencer.AnimationOutliner.RegularFont"))
 			.MinValue(NumericPropertyParams.MinValue)
@@ -97,6 +112,8 @@ public:
 			.OnBeginSliderMovement(this, &SNumericKeyEditorWidget::OnBeginSliderMovement)
 			.OnEndSliderMovement(this, &SNumericKeyEditorWidget::OnEndSliderMovement)
 		];
+
+		SpinBox->SetKeyEditor(KeyEditor);
 	}
 
 private:
@@ -141,16 +158,19 @@ private:
 
 	void OnValueCommitted(NumericType Value, ETextCommit::Type CommitInfo)
 	{
-		if (CommitInfo == ETextCommit::OnEnter)
+		if (CommitInfo == ETextCommit::OnEnter || CommitInfo == ETextCommit::OnUserMovedFocus)
 		{
 			const FScopedTransaction Transaction( LOCTEXT("SetNumericKey", "Set Key Value") );
 			KeyEditor->SetValueWithNotify(Value, EMovieSceneDataChangeType::TrackValueChangedRefreshImmediately);
 		}
+
+		KeyEditor->EndEditing();
 	}
 
 private:
 
 	TSharedPtr<ISequencerKeyEditor<NumericType>> KeyEditor;
+	TSharedPtr<SNonThrottledSpinBox<NumericType>> SpinBox;
 	bool bSliding = false;
 };
 
diff --git a/Engine/Source/Editor/MovieSceneTools/Private/CurveKeyEditors/SequencerKeyEditor.h b/Engine/Source/Editor/MovieSceneTools/Private/CurveKeyEditors/SequencerKeyEditor.h
--- a/Engine/Source/Editor/MovieSceneTools/Private/CurveKeyEditors/SequencerKeyEditor.h
+++ b/Engine/Source/Editor/MovieSceneTools/Private/CurveKeyEditors/SequencerKeyEditor.h
@@ -33,6 +33,9 @@ struct ISequencerKeyEditor
 	virtual FTrackInstancePropertyBindings* GetPropertyBindings() const = 0;
 	virtual FString GetMetaData(const FName& Key) const = 0;
 	virtual bool GetEditingKeySelection() const = 0;
+	virtual FFrameTime GetCurrentTime() const = 0;
+	virtual void BeginEditing(FFrameTime) = 0;
+	virtual void EndEditing() = 0;
 };
 
 template<typename ChannelType, typename ValueType>
@@ -246,7 +249,11 @@ struct TSequencerKeyEditor
 				{
 					// When auto setting track defaults are disabled, add a key even when it's empty so that the changed
 					// value is saved and is propagated to the property.
-					AddKeyToChannel(Channel, CurrentTime, InValue, Interpolation);
+					if (!ValueExistsAtTime(Channel, CurrentTime, InValue))
+					{
+						AddKeyToChannel(Channel, CurrentTime, InValue, Interpolation);
+					}
+					
 					bHasAnyKeys = Channel->GetNumKeys() != 0;
 				}
 
@@ -329,6 +336,11 @@ struct TSequencerKeyEditor
 
 	FFrameTime GetCurrentTime() const
 	{
+		if (EditingTime.IsSet())
+		{
+			return EditingTime.GetValue();
+		}
+
 		ISequencer* Sequencer = GetSequencer();
 		if (Sequencer)
 		{
@@ -339,6 +351,16 @@ struct TSequencerKeyEditor
 		return 0;
 	}
 
+	void BeginEditing(FFrameTime InEditingTime)
+	{
+		EditingTime = InEditingTime;
+	}
+
+	void EndEditing()
+	{
+		EditingTime.Reset();
+	}
+
 private:
 
 	FGuid ObjectBindingID;
@@ -350,6 +372,11 @@ private:
 	TFunction<TOptional<ValueType>(UObject&, FTrackInstancePropertyBindings*)> OnGetExternalValue;
 	TSharedPtr<INumericTypeInterface<ValueType>> NumericTypeInterface;
 	bool bApplyInUnwarpedLocalSpace = false;
+
+	// The time at which to edit the key value. This can be enforced when receiving focus on the spin box 
+	// so that when focus is lost and the value is commited, the value is committed at the editing time 
+	// rather than the Sequencer time which could have changed.
+	TOptional<FFrameTime> EditingTime;
 };
 
 
@@ -402,6 +429,18 @@ struct TSequencerKeyEditorWrapper : ISequencerKeyEditor<ValueType>
 	{
 		return Impl.GetEditingKeySelection();
 	}
+	FFrameTime GetCurrentTime() const override
+	{
+		return Impl.GetCurrentTime();
+	}
+	void BeginEditing(FFrameTime InFrameTime) override
+	{
+		Impl.BeginEditing(InFrameTime);
+	}
+	void EndEditing() override
+	{
+		Impl.EndEditing();
+	}
 
 private:
 
