Sequencer: Preserve scale ratio in key details

Add bCanPreserveRatio to channel meta which enables the key editor to proportionally scale all of the channels with the same group name (ie. when adjusting Scale.X, the other channels with the same group, Scale.Y and Scale.Z will also be set proportionally)

Add ReciprocalValue and TransformValue to MovieSceneChannelTraits so that channels can generically transform values, which is needed for preserving ratios when editing

The scale properties of 3DTransformSection 2DTransformSection  have this capability enabled

Add UncheckedImage to SOutlinerColumnButton

diff --git a/Engine/Source/Editor/MovieSceneTools/Private/CurveKeyEditors/SequencerKeyEditor.h b/Engine/Source/Editor/MovieSceneTools/Private/CurveKeyEditors/SequencerKeyEditor.h
--- a/Engine/Source/Editor/MovieSceneTools/Private/CurveKeyEditors/SequencerKeyEditor.h
+++ b/Engine/Source/Editor/MovieSceneTools/Private/CurveKeyEditors/SequencerKeyEditor.h
@@ -139,6 +139,51 @@ struct TSequencerKeyEditor
 		return Result;
 	}
 
+	void InitializeChannelGroups() 
+	{
+		ChannelNameToInitialValue.Reset();
+
+		using namespace UE::MovieScene;
+
+		ChannelType* Channel = ChannelHandle.Get();
+		ISequencer* Sequencer = WeakSequencer.Pin().Get();
+		UMovieSceneSection* OwningSection = WeakSection.Get();
+		const FMovieSceneChannelMetaData* ChannelMetaData = ChannelHandle.GetMetaData();
+
+		if (Channel && ChannelMetaData && Sequencer && OwningSection)
+		{
+			FFrameTime LocalTime = GetCurrentTime();
+			const FFrameTime CurrentTime = UE::MovieScene::ClampToDiscreteRange(LocalTime, OwningSection->GetRange()) - ChannelMetaData->GetOffsetTime(OwningSection);
+
+			// Store the initial value of all the grouped channels
+			const FMovieSceneChannelProxy& ChannelProxy = OwningSection->GetChannelProxy();
+			TArray<FMovieSceneChannelHandle> GroupedChannels = GetGroupedChannels(ChannelProxy, ChannelMetaData->Group);
+			for (FMovieSceneChannelHandle GroupChannelHandle : GroupedChannels)
+			{
+				const FMovieSceneChannelMetaData* GroupChannelMetaData = GroupChannelHandle.GetMetaData();
+				if (!GroupChannelMetaData)
+				{
+					continue;
+				}
+
+				if (ChannelType* GroupChannel = reinterpret_cast<ChannelType*>(GroupChannelHandle.Get()))
+				{
+					ValueType Result{};
+
+					if (EvaluateChannel(OwningSection, GroupChannel, CurrentTime, Result))
+					{
+						if (GroupChannelMetaData->bInvertValue)
+						{
+							InvertValue(Result);
+						}
+
+						ChannelNameToInitialValue.Add(GroupChannelMetaData->Name, Result);
+					}
+				}
+			}
+		}
+	}
+
 	bool GetEditingKeySelection() const
 	{
 		using namespace UE::MovieScene;
@@ -170,6 +215,22 @@ struct TSequencerKeyEditor
 		return false;
 	}
 
+	TArray<FMovieSceneChannelHandle> GetGroupedChannels(const FMovieSceneChannelProxy& ChannelProxy, const FText& Group)
+	{
+		TArray<FMovieSceneChannelHandle> GroupedChannelHandles;
+
+		for (const FMovieSceneChannelMetaData& EachChannelMetaData : ChannelProxy.GetMetaData<ChannelType>())
+		{
+			if (EachChannelMetaData.Group.IdenticalTo(Group))
+			{
+				FMovieSceneChannelHandle GroupChannelHandle = ChannelProxy.GetChannelByName(EachChannelMetaData.Name);
+				GroupedChannelHandles.Add(GroupChannelHandle);
+			}
+		}
+
+		return GroupedChannelHandles;
+	}
+
 	void SetValue(const ValueType& InValue)
 	{
 		using namespace UE::MovieScene;
@@ -200,7 +261,9 @@ struct TSequencerKeyEditor
 		Owner->Modify();
 		Owner->SetFlags(RF_Transactional);
 
-		const bool  bAutoSetTrackDefaults = Sequencer->GetAutoSetTrackDefaults();
+		const bool bAutoSetTrackDefaults = Sequencer->GetAutoSetTrackDefaults();
+		const bool bPreserveRatio = ChannelMetaData ? ChannelMetaData->bPreserveRatio : false;
+		const bool bCanPreserveRatio = ChannelMetaData ? ChannelMetaData->bCanPreserveRatio : false;
 
 		const FKeySelection& KeySelection = Sequencer->GetViewModel()->GetSelection()->KeySelection;
 
@@ -210,6 +273,9 @@ struct TSequencerKeyEditor
 			InvertValue(NewValue);
 		}
 
+		FFrameTime LocalTime = GetCurrentTime();
+		const FFrameNumber CurrentTime = LocalTime.RoundToFrame() - ChannelMetaData->GetOffsetTime(OwningSection);
+
 		const bool bEditingKeySelection = GetEditingKeySelection();
 		if (bEditingKeySelection)
 		{
@@ -227,11 +293,91 @@ struct TSequencerKeyEditor
 				}
 			}
 		}
-		else
+		else if (bPreserveRatio && bCanPreserveRatio)
 		{
-			FFrameTime LocalTime = GetCurrentTime();
-			const FFrameNumber CurrentTime = LocalTime.RoundToFrame() - ChannelMetaData->GetOffsetTime(OwningSection);
+			const FMovieSceneChannelProxy& ChannelProxy = OwningSection->GetChannelProxy();
 
+			TArray<FMovieSceneChannelHandle> GroupedChannels = GetGroupedChannels(ChannelProxy, ChannelMetaData->Group);
+
+			// First determine if there are any keys on any of the grouped channels. If any of the channels have keys, all channels will need to be keyed
+			bool bHasAnyKeys = false;
+			for (FMovieSceneChannelHandle GroupChannelHandle : GroupedChannels)
+			{
+				if (ChannelType* GroupChannel = reinterpret_cast<ChannelType*>(GroupChannelHandle.Get()))
+				{
+					if (GroupChannel->GetNumKeys() != 0)
+					{
+						bHasAnyKeys = true;
+						break;
+					}
+				}
+			}
+
+			ValueType Ratio = NewValue;
+			if (ChannelNameToInitialValue.Contains(ChannelMetaData->Name))
+			{
+				FMovieSceneChannelTraitsTransform<ValueType> Transform;
+				Transform.Scale = ChannelNameToInitialValue[ChannelMetaData->Name];
+				ReciprocalValue(Transform.Scale);
+				TransformValue(Ratio, Transform);
+			}
+
+			for (FMovieSceneChannelHandle GroupChannelHandle : GroupedChannels)
+			{
+				if (ChannelType* GroupChannel = reinterpret_cast<ChannelType*>(GroupChannelHandle.Get()))
+				{
+					const FMovieSceneChannelMetaData* GroupChannelMetaData = GroupChannelHandle.GetMetaData();
+					if (!GroupChannelMetaData)
+					{
+						continue;
+					}
+
+					ValueType NewPreserveRatioValue = NewValue;
+					if (ChannelNameToInitialValue.Contains(GroupChannelMetaData->Name))
+					{
+						NewPreserveRatioValue = ChannelNameToInitialValue[GroupChannelMetaData->Name];
+						FMovieSceneChannelTraitsTransform<ValueType> Transform;
+						Transform.Scale = Ratio;
+						TransformValue(NewPreserveRatioValue, Transform);
+					}
+
+					if (bHasAnyKeys)
+					{
+						TArray<FKeyHandle> GroupKeysAtCurrentTime;
+						GroupChannel->GetKeys(TRange<FFrameNumber>(CurrentTime), nullptr, &GroupKeysAtCurrentTime);
+
+						if (GroupKeysAtCurrentTime.Num() > 0)
+						{
+							AssignValue(GroupChannel, GroupKeysAtCurrentTime[0], NewPreserveRatioValue);
+						}
+						else
+						{
+							EMovieSceneKeyInterpolation Interpolation = GetInterpolationMode(GroupChannel, CurrentTime, Sequencer->GetKeyInterpolation());
+							if (!ValueExistsAtTime(GroupChannel, CurrentTime, NewPreserveRatioValue))
+							{
+								AddKeyToChannel(GroupChannel, CurrentTime, NewPreserveRatioValue, Interpolation);
+							}
+							
+							OwningSection->ExpandToFrame(LocalTime.RoundToFrame());
+						}
+					}
+
+					// Always update the default value when auto-set default values is enabled so that the last changes
+					// are always saved to the track.
+					if (bAutoSetTrackDefaults)
+					{
+						SetChannelDefault(GroupChannel, NewPreserveRatioValue);
+					}
+
+					//need to tell channel change happened (float will call AutoSetTangents())
+					GroupChannel->PostEditChange();
+
+					Sequencer->OnChannelChanged().Broadcast(GroupChannelMetaData, OwningSection);
+				}
+			}
+		}
+		else
+		{
 			EMovieSceneKeyInterpolation Interpolation = GetInterpolationMode(Channel, CurrentTime, Sequencer->GetKeyInterpolation());
 
 			TArray<FKeyHandle> KeysAtCurrentTime;
@@ -249,9 +395,9 @@ struct TSequencerKeyEditor
 				{
 					// When auto setting track defaults are disabled, add a key even when it's empty so that the changed
 					// value is saved and is propagated to the property.
-					if (!ValueExistsAtTime(Channel, CurrentTime, InValue))
+					if (!ValueExistsAtTime(Channel, CurrentTime, NewValue))
 					{
-						AddKeyToChannel(Channel, CurrentTime, InValue, Interpolation);
+						AddKeyToChannel(Channel, CurrentTime, NewValue, Interpolation);
 					}
 					
 					bHasAnyKeys = Channel->GetNumKeys() != 0;
@@ -262,21 +408,20 @@ struct TSequencerKeyEditor
 					OwningSection->ExpandToFrame(LocalTime.RoundToFrame());
 				}
 			}
-		}
 
-		// Always update the default value when auto-set default values is enabled so that the last changes
-		// are always saved to the track.
-		if (bAutoSetTrackDefaults)
-		{
-			SetChannelDefault(Channel, NewValue);
-		}
+			// Always update the default value when auto-set default values is enabled so that the last changes
+			// are always saved to the track.
+			if (bAutoSetTrackDefaults)
+			{
+				SetChannelDefault(Channel, NewValue);
+			}
 		 
-		//need to tell channel change happened (float will call AutoSetTangents())
-		Channel->PostEditChange();
-
-		const FMovieSceneChannelMetaData* MetaData = ChannelHandle.GetMetaData();
-		Sequencer->OnChannelChanged().Broadcast(MetaData, OwningSection);
+			//need to tell channel change happened (float will call AutoSetTangents())
+			Channel->PostEditChange();
 
+			const FMovieSceneChannelMetaData* MetaData = ChannelHandle.GetMetaData();
+			Sequencer->OnChannelChanged().Broadcast(MetaData, OwningSection);
+		}
 	}
 
 	void SetValueWithNotify(const ValueType& InValue, EMovieSceneDataChangeType NotifyType = EMovieSceneDataChangeType::TrackValueChanged)
@@ -354,6 +499,8 @@ struct TSequencerKeyEditor
 	void BeginEditing(FFrameTime InEditingTime)
 	{
 		EditingTime = InEditingTime;
+
+		InitializeChannelGroups();
 	}
 
 	void EndEditing()
@@ -377,6 +524,7 @@ private:
 	// so that when focus is lost and the value is commited, the value is committed at the editing time 
 	// rather than the Sequencer time which could have changed.
 	TOptional<FFrameTime> EditingTime;
+	TMap<FName, ValueType> ChannelNameToInitialValue;
 };
 
 
diff --git a/Engine/Source/Editor/Sequencer/Private/MVVM/ViewModels/CategoryModel.cpp b/Engine/Source/Editor/Sequencer/Private/MVVM/ViewModels/CategoryModel.cpp
--- a/Engine/Source/Editor/Sequencer/Private/MVVM/ViewModels/CategoryModel.cpp
+++ b/Engine/Source/Editor/Sequencer/Private/MVVM/ViewModels/CategoryModel.cpp
@@ -15,6 +15,7 @@
 #include "MVVM/Views/SOutlinerItemViewBase.h"
 #include "MVVM/Views/SSequencerKeyNavigationButtons.h"
 #include "MVVM/Views/SOutlinerTrackColorPicker.h"
+#include "MVVM/Views/SOutlinerTrackPreserveRatio.h"
 #include "MVVM/ViewModels/OutlinerColumns/OutlinerColumnTypes.h"
 #include "Math/UnrealMathUtility.h"
 #include "Misc/AssertionMacros.h"
@@ -222,7 +223,7 @@ TSharedPtr<SWidget> FCategoryGroupModel::CreateOutlinerViewForColumn(const FCrea
 
 	if (InColumnName == FCommonOutlinerNames::Edit)
 	{
-
+		return SNew(SOutlinerTrackPreserveRatio, SharedThis(this), InParams.Editor);
 	}
 
 	if (InColumnName == FCommonOutlinerNames::KeyFrame)
diff --git a/Engine/Source/Editor/Sequencer/Private/MVVM/Views/SOutlinerTrackPreserveRatio.cpp b/Engine/Source/Editor/Sequencer/Private/MVVM/Views/SOutlinerTrackPreserveRatio.cpp
new file mode 100644
--- /dev/null
+++ b/Engine/Source/Editor/Sequencer/Private/MVVM/Views/SOutlinerTrackPreserveRatio.cpp
@@ -0,0 +1,157 @@
+// Copyright Epic Games, Inc. All Rights Reserved.
+
+#include "MVVM/Views/SOutlinerTrackPreserveRatio.h"
+#include "MVVM/Extensions/ITrackExtension.h"
+#include "MVVM/ViewModels/ChannelModel.h"
+#include "MVVM/ViewModels/ViewModel.h"
+#include "MVVM/ViewModels/ViewModelIterators.h"
+#include "MVVM/ViewModels/EditorViewModel.h"
+#include "MVVM/Views/OutlinerColumns/SOutlinerColumnButton.h"
+
+#include "Channels/MovieSceneChannelHandle.h"
+#include "Channels/MovieSceneChannelEditorData.h"
+#include "IKeyArea.h"
+
+#include "Widgets/DeclarativeSyntaxSupport.h"
+#include "Widgets/Layout/SBox.h"
+
+#include "CoreGlobals.h"
+#include "Engine/Engine.h"
+#include "Misc/ConfigCacheIni.h"
+#include "ScopedTransaction.h"
+#include "Styling/AppStyle.h"
+
+#define LOCTEXT_NAMESPACE "SOutlinerTrackPreserveRatio"
+
+namespace UE::Sequencer
+{
+
+void SOutlinerTrackPreserveRatio::Construct(const FArguments& InArgs, TWeakViewModelPtr<IOutlinerExtension> InWeakOutlinerExtension, const TSharedPtr<FEditorViewModel>& EditorViewModel)
+{
+	WeakOutlinerExtension = InWeakOutlinerExtension;
+
+	bool bPreserveRatio = false;
+	GConfig->GetBool(TEXT("SelectionDetails"), TEXT("PreserveScaleRatio"), bPreserveRatio, GEditorPerProjectIni);
+
+	SetPreserveRatio(bPreserveRatio);
+
+	TAttribute<bool> IsEnabled = MakeAttributeLambda(
+		[WeakEditor = TWeakPtr<FEditorViewModel>(EditorViewModel)]
+		{
+			TSharedPtr<FEditorViewModel> EditorPinned = WeakEditor.Pin();
+			return EditorPinned && !EditorPinned->IsReadOnly();
+		}
+	);
+
+	ChildSlot
+	[
+		SNew(SBox)
+		.HAlign(HAlign_Left)
+		.Visibility(CanPreserveRatio() ? EVisibility::Visible : EVisibility::Hidden)
+		[
+			SNew(SOutlinerColumnButton)
+			.IsFocusable(false)
+			.IsEnabled(IsEnabled)
+			.IsChecked(this, &SOutlinerTrackPreserveRatio::OnGetPreserveRatio)
+			.ToolTipText(LOCTEXT("PreserveRatioTooltip", "When enabled, all axis values scale together so the object maintains its proportions in all directions."))
+			.Image(FAppStyle::GetBrush("Icons.Lock"))
+			.UncheckedImage(FAppStyle::GetBrush("Icons.Unlock"))
+			.OnClicked(this, &SOutlinerTrackPreserveRatio::OnSetPreserveRatio)
+		]
+	];
+}
+
+bool SOutlinerTrackPreserveRatio::OnGetPreserveRatio() const
+{
+	bool bPreserveRatio = true;
+
+	int32 Count = 0;
+	for (const FMovieSceneChannelMetaData* ChannelMetaData : GetMetaData())
+	{
+		Count++;
+		if (ChannelMetaData && ChannelMetaData->bPreserveRatio != bPreserveRatio)
+		{
+			bPreserveRatio = false;
+			break;
+		}
+	}
+
+	if (Count > 0)
+	{
+		return bPreserveRatio;
+	}
+	return false;
+}
+
+FReply SOutlinerTrackPreserveRatio::OnSetPreserveRatio()
+{
+	const bool bPreserveRatio = !OnGetPreserveRatio();
+	GConfig->SetBool(TEXT("SelectionDetails"), TEXT("PreserveScaleRatio"), bPreserveRatio, GEditorPerProjectIni);
+
+	SetPreserveRatio(bPreserveRatio);
+
+	return FReply::Handled();
+}
+
+TArray<const FMovieSceneChannelMetaData*> SOutlinerTrackPreserveRatio::GetMetaData() const
+{
+	TArray<const FMovieSceneChannelMetaData*> AllChannelMetaData;
+
+	TViewModelPtr<IOutlinerExtension> Outliner = WeakOutlinerExtension.Pin();
+	if (!Outliner)
+	{
+		return AllChannelMetaData;
+	}
+
+	TViewModelPtr<ITrackAreaExtension> TrackArea = Outliner.ImplicitCast();
+	if (!TrackArea)
+	{
+		return AllChannelMetaData;
+	}
+
+	for (const FViewModelPtr& TrackAreaModel : TrackArea->GetTrackAreaModelList())
+	{
+		for (TSharedPtr<FChannelModel> ChannelModel : TrackAreaModel->GetDescendantsOfType<FChannelModel>())
+		{
+			if (TSharedPtr<IKeyArea> KeyArea = ChannelModel->GetKeyArea())
+			{
+				const FMovieSceneChannelHandle& ChannelHandle = KeyArea->GetChannel();
+				const FMovieSceneChannelMetaData* ChannelMetaData = ChannelHandle.GetMetaData();
+				if (ChannelMetaData)
+				{
+					AllChannelMetaData.Add(ChannelMetaData);
+				}
+			}
+		}
+	}
+
+	return AllChannelMetaData;
+}
+
+void SOutlinerTrackPreserveRatio::SetPreserveRatio(bool bInPreserveRatio)
+{
+	for (const FMovieSceneChannelMetaData* ChannelMetaData : GetMetaData())
+	{
+		if (ChannelMetaData)
+		{
+			ChannelMetaData->bPreserveRatio = bInPreserveRatio;
+		}
+	}
+}
+
+bool SOutlinerTrackPreserveRatio::CanPreserveRatio() const
+{
+	for (const FMovieSceneChannelMetaData* ChannelMetaData : GetMetaData())
+	{
+		if (ChannelMetaData && ChannelMetaData->bCanPreserveRatio)
+		{
+			return true;
+		}
+	}
+	return false;
+}
+
+} // namespace UE::Sequencer
+
+#undef LOCTEXT_NAMESPACE
+
diff --git a/Engine/Source/Editor/Sequencer/Private/MVVM/Views/SOutlinerTrackPreserveRatio.h b/Engine/Source/Editor/Sequencer/Private/MVVM/Views/SOutlinerTrackPreserveRatio.h
new file mode 100644
--- /dev/null
+++ b/Engine/Source/Editor/Sequencer/Private/MVVM/Views/SOutlinerTrackPreserveRatio.h
@@ -0,0 +1,48 @@
+// Copyright Epic Games, Inc. All Rights Reserved.
+
+#pragma once
+
+#include "SequencerCoreFwd.h"
+#include "Widgets/SCompoundWidget.h"
+#include "MVVM/ViewModelPtr.h"
+
+struct FMovieSceneChannelMetaData;
+
+namespace UE::Sequencer
+{
+
+class IOutlinerExtension;
+class FEditorViewModel;
+
+class SOutlinerTrackPreserveRatio
+	: public SCompoundWidget
+{
+public:
+	SLATE_BEGIN_ARGS(SOutlinerTrackPreserveRatio){}
+	SLATE_END_ARGS()
+
+	SEQUENCER_API void Construct(const FArguments& InArgs, TWeakViewModelPtr<IOutlinerExtension> InWeakOutlinerExtension, const TSharedPtr<FEditorViewModel>& EditorViewModel);
+
+private:
+
+	/** Whether this button is seto to preserve ratio */
+	bool OnGetPreserveRatio() const;
+
+	/** Called when the user clicks on the preserve ratio button */
+	FReply OnSetPreserveRatio();
+
+	/** Retrieve the channel metadata for this channel group */
+	TArray<const FMovieSceneChannelMetaData*> GetMetaData() const;
+
+	/** Set the channel meta data to enable/disable preserve ratio */
+	void SetPreserveRatio(bool bInPreserveRatio);
+
+	/** Gets whether any of the channels can preserve ratio */
+	bool CanPreserveRatio() const;
+
+private:
+
+	TWeakViewModelPtr<IOutlinerExtension> WeakOutlinerExtension;
+};
+
+} // namespace UE::Sequencer
diff --git a/Engine/Source/Editor/SequencerCore/Private/MVVM/Views/OutlinerColumns/SOutlinerColumnButton.cpp b/Engine/Source/Editor/SequencerCore/Private/MVVM/Views/OutlinerColumns/SOutlinerColumnButton.cpp
--- a/Engine/Source/Editor/SequencerCore/Private/MVVM/Views/OutlinerColumns/SOutlinerColumnButton.cpp
+++ b/Engine/Source/Editor/SequencerCore/Private/MVVM/Views/OutlinerColumns/SOutlinerColumnButton.cpp
@@ -12,6 +12,9 @@ namespace UE::Sequencer
 void SOutlinerColumnButton::Construct(const FArguments& InArgs)
 {
 	IsRowHovered = InArgs._IsRowHovered;
+	IsChecked = InArgs._IsChecked;
+	Image = InArgs._Image;
+	UncheckedImage = InArgs._UncheckedImage;
 
 	TSharedRef<SButton> Button = SNew(SButton)
 		.IsFocusable(InArgs._IsFocusable)
@@ -23,7 +26,7 @@ void SOutlinerColumnButton::Construct(const FArguments& InArgs)
 		[
 			SNew(SImage)
 			.ColorAndOpacity(this, &SOutlinerColumnButton::GetColorAndOpacity)
-			.Image(InArgs._Image)
+			.Image(this, &SOutlinerColumnButton::GetImage)
 		];
 
 	if (InArgs._OnGetMenuContent.IsBound())
@@ -49,6 +52,15 @@ void SOutlinerColumnButton::Construct(const FArguments& InArgs)
 	}
 }
 
+const FSlateBrush* SOutlinerColumnButton::GetImage() const
+{
+	if (IsChecked.IsBound())
+	{
+		return IsChecked.Get() ? Image.Get() : UncheckedImage.Get();
+	}
+	return Image.Get();
+}
+
 FSlateColor SOutlinerColumnButton::GetColorAndOpacity() const
 {
 	const uint8 IntIsRowHovered    = uint8(IsRowHovered.Get(false));   // 0 or 1
diff --git a/Engine/Source/Editor/SequencerCore/Public/MVVM/Views/OutlinerColumns/SOutlinerColumnButton.h b/Engine/Source/Editor/SequencerCore/Public/MVVM/Views/OutlinerColumns/SOutlinerColumnButton.h
--- a/Engine/Source/Editor/SequencerCore/Public/MVVM/Views/OutlinerColumns/SOutlinerColumnButton.h
+++ b/Engine/Source/Editor/SequencerCore/Public/MVVM/Views/OutlinerColumns/SOutlinerColumnButton.h
@@ -18,12 +18,20 @@ class SOutlinerColumnButton : public SCompoundWidget
 {
 public:
 
-	SLATE_BEGIN_ARGS(SOutlinerColumnButton) : _IsFocusable(false) {}
+	SLATE_BEGIN_ARGS(SOutlinerColumnButton) 
+		: _Image(nullptr)
+		, _UncheckedImage(nullptr)
+ 		, _IsFocusable(false)
+		{}
 
 		SLATE_ATTRIBUTE(const FSlateBrush*, Image)
 
+		SLATE_ATTRIBUTE(const FSlateBrush*, UncheckedImage)
+
 		SLATE_ATTRIBUTE(bool, IsRowHovered)
 
+		SLATE_ATTRIBUTE(bool, IsChecked)
+
 		SLATE_ARGUMENT(bool, IsFocusable)
 
 		SLATE_EVENT(FOnClicked, OnClicked)
@@ -38,12 +46,17 @@ private:
 
 	UE_API FSlateColor GetColorAndOpacity() const;
 
+	UE_API const FSlateBrush* GetImage() const;
+
 	UE_API FReply ToggleMenu();
 
 private:
 
 	TSharedPtr<SMenuAnchor> MenuAnchor;
 	TAttribute<bool> IsRowHovered;
+	TAttribute<bool> IsChecked;
+	TAttribute<const FSlateBrush*> Image;
+	TAttribute<const FSlateBrush*> UncheckedImage;
 };
 
 } // namespace UE::Sequencer
diff --git a/Engine/Source/Runtime/MovieScene/Private/Channels/MovieSceneChannelEditorData.cpp b/Engine/Source/Runtime/MovieScene/Private/Channels/MovieSceneChannelEditorData.cpp
--- a/Engine/Source/Runtime/MovieScene/Private/Channels/MovieSceneChannelEditorData.cpp
+++ b/Engine/Source/Runtime/MovieScene/Private/Channels/MovieSceneChannelEditorData.cpp
@@ -29,6 +29,8 @@ FMovieSceneChannelMetaData::FMovieSceneChannelMetaData()
 	, SortOrder(0)
 	, bSortEmptyGroupsLast(true)
 	, bInvertValue(false)
+	, bCanPreserveRatio(false)
+	, bPreserveRatio(false)
 	, Name(NAME_None)
 	, KeyOffset(0)
 {}
@@ -40,6 +42,8 @@ FMovieSceneChannelMetaData::FMovieSceneChannelMetaData(FName InName, FText InDis
 	, SortOrder(0)
 	, bSortEmptyGroupsLast(true)
 	, bInvertValue(false)
+	, bCanPreserveRatio(false)
+	, bPreserveRatio(false)
 	, Name(InName)
 	, DisplayText(InDisplayText)
 	, Group(InGroup)
diff --git a/Engine/Source/Runtime/MovieScene/Private/Channels/MovieSceneDoubleChannel.cpp b/Engine/Source/Runtime/MovieScene/Private/Channels/MovieSceneDoubleChannel.cpp
--- a/Engine/Source/Runtime/MovieScene/Private/Channels/MovieSceneDoubleChannel.cpp
+++ b/Engine/Source/Runtime/MovieScene/Private/Channels/MovieSceneDoubleChannel.cpp
@@ -703,9 +703,25 @@ void InvertValue(double& InOutValue)
 	InOutValue = -InOutValue;
 }
 
-void InvertValue(float& InOutValue)
+void ReciprocalValue(double& InOutValue)
 {
-	InOutValue = -InOutValue;
+	double OriginalValue = InOutValue;
+
+	if (InOutValue != 0)
+	{
+		InOutValue = 1.0 / InOutValue;
+	}
+}
+
+void TransformValue(double& InOutValue, const FMovieSceneChannelTraitsTransform<double>& Transform)
+{
+	double OriginalValue = InOutValue;
+
+	if (Transform.Scale != 0.0)
+	{
+		InOutValue = (InOutValue - Transform.Pivot) * Transform.Scale + Transform.Pivot;
+	}
+	InOutValue += Transform.Offset;
 }
 
 void FMovieSceneDoubleChannel::AddKeys(const TArray<FFrameNumber>& InTimes, const TArray<FMovieSceneDoubleValue>& InValues)
diff --git a/Engine/Source/Runtime/MovieScene/Private/Channels/MovieSceneFloatChannel.cpp b/Engine/Source/Runtime/MovieScene/Private/Channels/MovieSceneFloatChannel.cpp
--- a/Engine/Source/Runtime/MovieScene/Private/Channels/MovieSceneFloatChannel.cpp
+++ b/Engine/Source/Runtime/MovieScene/Private/Channels/MovieSceneFloatChannel.cpp
@@ -317,6 +317,28 @@ void AssignValue(FMovieSceneFloatChannel* InChannel, FKeyHandle InKeyHandle, flo
 	return TMovieSceneCurveChannelImpl<FMovieSceneFloatChannel>::AssignValue(InChannel, InKeyHandle, InValue);
 }
 
+void InvertValue(float& InOutValue)
+{
+	InOutValue = -InOutValue;
+}
+
+void ReciprocalValue(float& InOutValue)
+{
+	if (InOutValue != 0.f)
+	{
+		InOutValue = 1.0f / InOutValue;
+	}
+}
+
+void TransformValue(float& InOutValue, const FMovieSceneChannelTraitsTransform<float>& Transform)
+{
+	if (Transform.Scale != 0.f)
+	{
+		InOutValue = (InOutValue - Transform.Pivot) * Transform.Scale + Transform.Pivot;
+	}
+	InOutValue += Transform.Offset;
+}
+
 void FMovieSceneFloatChannel::AddKeys(const TArray<FFrameNumber>& InTimes, const TArray<FMovieSceneFloatValue>& InValues)
 {
 	check(InTimes.Num() == InValues.Num());
diff --git a/Engine/Source/Runtime/MovieScene/Public/Channels/MovieSceneChannelEditorData.h b/Engine/Source/Runtime/MovieScene/Public/Channels/MovieSceneChannelEditorData.h
--- a/Engine/Source/Runtime/MovieScene/Public/Channels/MovieSceneChannelEditorData.h
+++ b/Engine/Source/Runtime/MovieScene/Public/Channels/MovieSceneChannelEditorData.h
@@ -78,6 +78,10 @@ struct FMovieSceneChannelMetaData
 	bool bSortEmptyGroupsLast;
 	/** Invert the value for display purposes */
 	bool bInvertValue;
+	/** Whether this channel can preserve the ratio to its grouped channels when editing */
+	bool bCanPreserveRatio;
+	/** Determines whether preserve ratio is enabled or not, not saved */
+	mutable bool bPreserveRatio;
 
 	/** This channel's unique name */
 	FName Name;
diff --git a/Engine/Source/Runtime/MovieScene/Public/Channels/MovieSceneChannelTraits.h b/Engine/Source/Runtime/MovieScene/Public/Channels/MovieSceneChannelTraits.h
--- a/Engine/Source/Runtime/MovieScene/Public/Channels/MovieSceneChannelTraits.h
+++ b/Engine/Source/Runtime/MovieScene/Public/Channels/MovieSceneChannelTraits.h
@@ -47,6 +47,14 @@ struct TMovieSceneChannelTraitsBase
 #endif
 };
 
+template<typename ValueType>
+struct FMovieSceneChannelTraitsTransform
+{
+	ValueType Offset = {};
+	ValueType Scale = {};
+	ValueType Pivot = {};
+};
+
 /**
  * Traits structure to be specialized for any channel type passed to FMovieSceneChannelProxy
  */
@@ -277,5 +285,11 @@ namespace MovieScene
 	template<typename ValueType>
 	void InvertValue(ValueType& InOutValue){}
 
+	template<typename ValueType>
+	void ReciprocalValue(ValueType& InOutValue) {}
+
+	template<typename ValueType>
+	void TransformValue(ValueType& InOutValue, const FMovieSceneChannelTraitsTransform<ValueType>& Transform) {}
+
 } // namespace MovieScene
 } // namespace UE
diff --git a/Engine/Source/Runtime/MovieScene/Public/Channels/MovieSceneDoubleChannel.h b/Engine/Source/Runtime/MovieScene/Public/Channels/MovieSceneDoubleChannel.h
--- a/Engine/Source/Runtime/MovieScene/Public/Channels/MovieSceneDoubleChannel.h
+++ b/Engine/Source/Runtime/MovieScene/Public/Channels/MovieSceneDoubleChannel.h
@@ -455,6 +455,7 @@ MOVIESCENE_API bool ValueExistsAtTime(const FMovieSceneDoubleChannel* InChannel,
 MOVIESCENE_API void AssignValue(FMovieSceneDoubleChannel* InChannel, FKeyHandle InKeyHandle, double InValue);
 MOVIESCENE_API void AssignValue(FMovieSceneDoubleChannel* InChannel, FKeyHandle InKeyHandle, float InValue);
 MOVIESCENE_API void InvertValue(double& InOutValue);
-MOVIESCENE_API void InvertValue(float& InOutValue);
+MOVIESCENE_API void ReciprocalValue(double& InOutValue);
+MOVIESCENE_API void TransformValue(double& InOutValue, const FMovieSceneChannelTraitsTransform<double>& Transform);
 
 
diff --git a/Engine/Source/Runtime/MovieScene/Public/Channels/MovieSceneFloatChannel.h b/Engine/Source/Runtime/MovieScene/Public/Channels/MovieSceneFloatChannel.h
--- a/Engine/Source/Runtime/MovieScene/Public/Channels/MovieSceneFloatChannel.h
+++ b/Engine/Source/Runtime/MovieScene/Public/Channels/MovieSceneFloatChannel.h
@@ -396,4 +396,6 @@ MOVIESCENE_API void Dilate(FMovieSceneFloatChannel* InChannel, FFrameNumber Orig
 MOVIESCENE_API bool ValueExistsAtTime(const FMovieSceneFloatChannel* InChannel, FFrameNumber InFrameNumber, float InValue);
 MOVIESCENE_API bool ValueExistsAtTime(const FMovieSceneFloatChannel* InChannel, FFrameNumber InFrameNumber, const FMovieSceneFloatValue& InValue);
 MOVIESCENE_API void AssignValue(FMovieSceneFloatChannel* InChannel, FKeyHandle InKeyHandle, float InValue);
-
+MOVIESCENE_API void InvertValue(float& InOutValue);
+MOVIESCENE_API void ReciprocalValue(float& InOutValue);
+MOVIESCENE_API void TransformValue(float& InOutValue, const FMovieSceneChannelTraitsTransform<float>& Transform);
diff --git a/Engine/Source/Runtime/MovieSceneTracks/Private/Sections/MovieScene3DTransformSection.cpp b/Engine/Source/Runtime/MovieSceneTracks/Private/Sections/MovieScene3DTransformSection.cpp
--- a/Engine/Source/Runtime/MovieSceneTracks/Private/Sections/MovieScene3DTransformSection.cpp
+++ b/Engine/Source/Runtime/MovieSceneTracks/Private/Sections/MovieScene3DTransformSection.cpp
@@ -155,6 +155,7 @@ struct F3DTransformChannelEditorData
 			MetaData[6].Color = AxisDisplayInfo::GetAxisColor(XAxis);
 			MetaData[6].SortOrder = SortOrderStart + ScaleOrderOffset + ReverseSwizzle[0];
 			MetaData[6].bCanCollapseToTrack = false;
+			MetaData[6].bCanPreserveRatio = true;
 			for (const FName& PropertyMetaDataKey : PropertyMetaDataKeys)
 			{
 				MetaData[6].PropertyMetaData.Add(PropertyMetaDataKey, RelativeScale3DProperty->GetMetaData(PropertyMetaDataKey));
@@ -166,6 +167,7 @@ struct F3DTransformChannelEditorData
 			MetaData[7].Color = AxisDisplayInfo::GetAxisColor(YAxis);
 			MetaData[7].SortOrder = SortOrderStart + ScaleOrderOffset + ReverseSwizzle[1];
 			MetaData[7].bCanCollapseToTrack = false;
+			MetaData[7].bCanPreserveRatio = true;
 			for (const FName& PropertyMetaDataKey : PropertyMetaDataKeys)
 			{
 				MetaData[7].PropertyMetaData.Add(PropertyMetaDataKey, RelativeScale3DProperty->GetMetaData(PropertyMetaDataKey));
@@ -177,6 +179,7 @@ struct F3DTransformChannelEditorData
 			MetaData[8].Color = AxisDisplayInfo::GetAxisColor(ZAxis);
 			MetaData[8].SortOrder = SortOrderStart + ScaleOrderOffset + ReverseSwizzle[2];
 			MetaData[8].bCanCollapseToTrack = false;
+			MetaData[8].bCanPreserveRatio = true;
 			for (const FName& PropertyMetaDataKey : PropertyMetaDataKeys)
 			{
 				MetaData[8].PropertyMetaData.Add(PropertyMetaDataKey, RelativeScale3DProperty->GetMetaData(PropertyMetaDataKey));
diff --git a/Engine/Source/Runtime/UMG/Private/Animation/MovieScene2DTransformSection.cpp b/Engine/Source/Runtime/UMG/Private/Animation/MovieScene2DTransformSection.cpp
--- a/Engine/Source/Runtime/UMG/Private/Animation/MovieScene2DTransformSection.cpp
+++ b/Engine/Source/Runtime/UMG/Private/Animation/MovieScene2DTransformSection.cpp
@@ -54,12 +54,14 @@ struct F2DTransformSectionEditorData
 		MetaData[3].bEnabled = EnumHasAllFlags(Mask, EMovieScene2DTransformChannel::ScaleX);
 		MetaData[3].SortOrder = 3;
 		MetaData[3].bCanCollapseToTrack = false;
+		MetaData[3].bCanPreserveRatio = true;
 
 		MetaData[4].SetIdentifiers("Scale.Y", FCommonChannelData::ChannelY, ScaleGroup);
 		MetaData[4].SubPropertyPath = MetaData[4].Name;
 		MetaData[4].bEnabled = EnumHasAllFlags(Mask, EMovieScene2DTransformChannel::ScaleY);
 		MetaData[4].SortOrder = 4;
 		MetaData[4].bCanCollapseToTrack = false;
+		MetaData[4].bCanPreserveRatio = true;
 
 		MetaData[5].SetIdentifiers("Shear.X", FCommonChannelData::ChannelX, ShearGroup);
 		MetaData[5].SubPropertyPath = MetaData[5].Name;
