Curve Editor: Introduce FCurveModel::ReplaceKeyHandles and implement it in Sequencer.

This is required for a follow up change for a new command-based undo system. For example, when undoing a "remove keys" operation, we want to add back the keys with the exact same FKeyHandle.

diff --git a/Engine/Source/Editor/CurveEditor/Private/CurveModel.cpp b/Engine/Source/Editor/CurveEditor/Private/CurveModel.cpp
--- a/Engine/Source/Editor/CurveEditor/Private/CurveModel.cpp
+++ b/Engine/Source/Editor/CurveEditor/Private/CurveModel.cpp
@@ -7,6 +7,33 @@
 #include "Curves/KeyHandle.h"
 #include "HAL/PlatformCrt.h"
 
+void FCurveModel::PutKeys(
+	TConstArrayView<FKeyHandle> InKeys, TConstArrayView<FKeyPosition> InKeyPositions, TConstArrayView<FKeyAttributes> InAttributes
+	)
+{
+	const int32 Num = InKeys.Num();
+	if (ensure(InKeyPositions.Num() == Num && InAttributes.Num() == Num))
+	{
+		RemoveKeys(InKeys, 0.0);
+	
+		TArray<TOptional<FKeyHandle>> OptionalAddedHandles;
+		OptionalAddedHandles.SetNum(InKeys.Num());
+		TArrayView<TOptional<FKeyHandle>> NewKeyHandlesView = MakeArrayView(OptionalAddedHandles);
+		AddKeys(InKeyPositions, InAttributes, &NewKeyHandlesView);
+
+		TArray<FKeyHandle> AddedKeyHandles;
+		Algo::TransformIf(OptionalAddedHandles, AddedKeyHandles,
+			[](const TOptional<FKeyHandle>& Handle){ return Handle.IsSet(); },
+			[](const TOptional<FKeyHandle>& Handle){ return *Handle; }
+			);
+
+		if (ensure(AddedKeyHandles.Num() == Num))
+		{
+			ReplaceKeyHandles(AddedKeyHandles, InKeys);
+		}
+	}
+}
+
 void FCurveModel::GetClosestKeysTo(double InTime, TOptional<FKeyHandle>& OutPreviousKeyHandle, TOptional<FKeyHandle>& OutNextKeyHandle) const
 {
 	// Default implementation does linear search. Subclasses can override this with a more efficient implementations.
@@ -61,15 +88,11 @@ void FCurveModel::SetKeyAttributes(TArrayView<const FKeyHandle> InKeys, const FK
 	}
 
 	SetKeyAttributes(InKeys, ExpandedAttributes);
-	Modify();
 }
 
 TOptional<FKeyHandle> FCurveModel::AddKey(const FKeyPosition& NewKeyPosition, const FKeyAttributes& InAttributes)
 {
-	Modify();
-
 	TOptional<FKeyHandle> Handle;
-
 	TArrayView<TOptional<FKeyHandle>> Handles = MakeArrayView(&Handle, 1);
 	AddKeys(TArrayView<const FKeyPosition>(&NewKeyPosition, 1), TArrayView<const FKeyAttributes>(&InAttributes, 1), &Handles);
 	return Handle;
diff --git a/Engine/Source/Editor/CurveEditor/Public/CurveModel.h b/Engine/Source/Editor/CurveEditor/Public/CurveModel.h
--- a/Engine/Source/Editor/CurveEditor/Public/CurveModel.h
+++ b/Engine/Source/Editor/CurveEditor/Public/CurveModel.h
@@ -139,6 +139,28 @@ public:
 	 */
 	virtual void RemoveKeys(TArrayView<const FKeyHandle> InKeys, double InCurrentTime) = 0;
 
+	/** Adds the given keys if they are not yet present, or overwrites their current data if already present. */
+	UE_API virtual void PutKeys(TConstArrayView<FKeyHandle> InKeys, TConstArrayView<FKeyPosition> InKeyPositions, TConstArrayView<FKeyAttributes> InAttributes);
+
+	/**
+	 * Remaps InNewHandles to reference the keys that InCurrentHandles reference.
+	 * After this function concludes, the FKeyHandle InCurrentHandles no longer reference any keys in this curve.
+	 *
+	 * This operation is needed if you have to add keys that are mapped to specific FKeyHandle instances.
+	 * A use case for this is undo / redo. Suppose the following user action sequence: 1. Move FKeyHandle(42), 2. Delete FKeyHandle(42), 3. Undo.
+	 * In step 3, we must add the key back. Calling FCurveEditor::AddKeys generates a new FKeyHandle, but we need to add it back as FKeyHandle(42);
+	 * otherwise step 1's reference to FKeyHandle(42) breaks and undoing the movement action would not work.
+	 *
+	 * FKeyHandle in InNewHandles that are already mapped to a key are skipped..
+	 * The corresponding InCurrentHandles is left unchanged and continues to reference the same key as before the call.
+	 * 
+	 * FKeyHandles in InCurrentHandles that reference no keys are skipped.
+	 *
+	 * @param InCurrentHandles Key handles that are currentl in the curve model
+	 * @param InNewHandles The key handles to remap to. Each index corresponds to the same index in InCurrentHandles.
+	 */
+	virtual void ReplaceKeyHandles(TConstArrayView<FKeyHandle> InCurrentHandles, TConstArrayView<FKeyHandle> InNewHandles) {}
+
 	/**
 	 * Retrieve all key positions that pertain to the specified input key handles
 	 *
@@ -154,7 +176,7 @@ public:
 	 * @param InKeyPositions         Array of desired key positions to be applied to each of the corresponding key handles
 	 */
 	virtual void SetKeyPositions(TArrayView<const FKeyHandle> InKeys, TArrayView<const FKeyPosition> InKeyPositions, EPropertyChangeType::Type ChangeType = EPropertyChangeType::Unspecified) = 0;
-
+	
 	/**
 	 * Populate the specified draw info structure with data describing how to draw the specified point type
 	 *
diff --git a/Engine/Source/Editor/MovieSceneTools/Private/Channels/ChannelCurveModel.cpp b/Engine/Source/Editor/MovieSceneTools/Private/Channels/ChannelCurveModel.cpp
--- a/Engine/Source/Editor/MovieSceneTools/Private/Channels/ChannelCurveModel.cpp
+++ b/Engine/Source/Editor/MovieSceneTools/Private/Channels/ChannelCurveModel.cpp
@@ -274,10 +274,6 @@ void FChannelCurveModel<ChannelType, ChannelValue, KeyType>::SetKeyPositions(TAr
 		FFrameRate   TickResolution = Section->GetTypedOuter<UMovieScene>()->GetTickResolution();
 		FFrameNumber KeyOffset      = 0;
 
-		const int32 LastIndex = Times.Num()-1;
-		const FFrameNumber FirstTime = Times[0];
-		const FFrameNumber LastTime = Times[LastIndex];
-
 		// Expand to frame first, then offset
 		if (InKeyPositions.Num() > 0)
 		{
@@ -319,13 +315,7 @@ void FChannelCurveModel<ChannelType, ChannelValue, KeyType>::SetKeyPositions(TAr
 			}
 		}
 		
-		Channel->PostEditChange();
-		if(WeakSequencer.IsValid())
-		{ 
-			WeakSequencer.Pin()->OnChannelChanged().Broadcast(MetaData, Section);
-		}
-		CurveModifiedDelegate.Broadcast();
-
+		InvokePostChangeCallbacks(Channel, Section);
 		SignedOwner->MarkAsChanged();
 	}
 }
@@ -474,8 +464,6 @@ void FChannelCurveModel<ChannelType, ChannelValue, KeyType>::AddKeys(TArrayView<
 
 	if (Channel && MetaData && Section && !IsReadOnly())
 	{
-		SignedOwner->Modify();
-
 		TMovieSceneChannelData<ChannelValue> ChannelData = Channel->GetData();
 		FFrameRate TickResolution = Section->GetTypedOuter<UMovieScene>()->GetTickResolution();
 
@@ -515,12 +503,7 @@ void FChannelCurveModel<ChannelType, ChannelValue, KeyType>::AddKeys(TArrayView<
 		// We reuse SetKeyAttributes here as there is complex logic determining which parts of the attributes are valid to set.
 		// For now we need to duplicate the new key handle array due to API mismatch. This will auto calculate tangents if needed.
 		SetKeyAttributes(NewKeyHandles, InKeyAttributes);
-		Channel->PostEditChange();
-		if (WeakSequencer.IsValid())
-		{
-			WeakSequencer.Pin()->OnChannelChanged().Broadcast(MetaData, Section);
-		}
-		CurveModifiedDelegate.Broadcast();
+		InvokePostChangeCallbacks(Channel, Section);
 	}
 }
 
@@ -538,8 +521,6 @@ void FChannelCurveModel<ChannelType, ChannelValue, KeyType>::RemoveKeys(TArrayVi
 
 	if (Channel && Section && !IsReadOnly())
 	{
-		SignedOwner->Modify();
-
 		TMovieSceneChannelData<ChannelValue> ChannelData = Channel->GetData();		
 		double CurrentValue = 0.0;
 		Evaluate(InCurrentTime, CurrentValue);
@@ -559,13 +540,41 @@ void FChannelCurveModel<ChannelType, ChannelValue, KeyType>::RemoveKeys(TArrayVi
 			Channel->SetDefault(ValueAsCurveValue);
 		}
 
-		Channel->PostEditChange();
-		if (WeakSequencer.IsValid())
+		InvokePostChangeCallbacks(Channel, Section);
+	}
+}
+
+template <class ChannelType, class ChannelValue, class KeyType>
+void FChannelCurveModel<ChannelType, ChannelValue, KeyType>::ReplaceKeyHandles(
+	TConstArrayView<FKeyHandle> InCurrentHandles, TConstArrayView<FKeyHandle> InNewHandles
+	)
+{
+	if (!ensure(InCurrentHandles.Num() == InNewHandles.Num()))
+	{
+		return;
+	}
+
+	ChannelType* Channel = ChannelHandle.Get();
+	UMovieSceneSection* Section = WeakSection.Get();
+	const FMovieSceneChannelMetaData* MetaData = ChannelHandle.GetMetaData();
+
+	if (Channel && MetaData && Section && !IsReadOnly())
+	{
+		TMovieSceneChannelData<ChannelValue> ChannelData = Channel->GetData();
+		for (int32 Index = 0; Index < InCurrentHandles.Num(); ++Index)
 		{
-			const FMovieSceneChannelMetaData* MetaData = ChannelHandle.GetMetaData();
-			WeakSequencer.Pin()->OnChannelChanged().Broadcast(MetaData, Section);
+			const int32 KeyIndex = ChannelData.GetIndex(InCurrentHandles[Index]);
+
+			// Skip remapping to keys that are already in use.
+			const bool bCannotReplace = KeyIndex == INDEX_NONE || ChannelData.GetIndex(InNewHandles[Index]) != INDEX_NONE;
+			if (bCannotReplace)
+			{
+				continue;
+			}
+
+			ChannelData.ReplaceKeyHandle(KeyIndex, InNewHandles[Index]);
+			InvokePostChangeCallbacks(Channel, Section);
 		}
-		CurveModifiedDelegate.Broadcast();
 	}
 }
 
@@ -626,6 +635,22 @@ void FChannelCurveModel<ChannelType, ChannelValue, KeyType>::GetCurveColorObject
 	}
 }
 
+template <class ChannelType, class ChannelValue, class KeyType>
+void FChannelCurveModel<ChannelType, ChannelValue, KeyType>::InvokePostChangeCallbacks(ChannelType* Channel, UMovieSceneSection* Section)
+{
+	LastSignature.Invalidate();
+	
+	Channel->PostEditChange();
+	
+	if (WeakSequencer.IsValid())
+	{
+		const FMovieSceneChannelMetaData* MetaData = ChannelHandle.GetMetaData();
+		WeakSequencer.Pin()->OnChannelChanged().Broadcast(MetaData, Section);
+	}
+	
+	CurveModifiedDelegate.Broadcast();
+}
+
 
 // Explicit template instantiation
 template class FChannelCurveModel<FMovieSceneDoubleChannel, FMovieSceneDoubleValue, double>;
diff --git a/Engine/Source/Editor/MovieSceneTools/Public/Channels/ChannelCurveModel.h b/Engine/Source/Editor/MovieSceneTools/Public/Channels/ChannelCurveModel.h
--- a/Engine/Source/Editor/MovieSceneTools/Public/Channels/ChannelCurveModel.h
+++ b/Engine/Source/Editor/MovieSceneTools/Public/Channels/ChannelCurveModel.h
@@ -61,9 +61,12 @@ public:
 	virtual int32 GetNumKeys() const override;
 	virtual void GetNeighboringKeys(const FKeyHandle InKeyHandle, TOptional<FKeyHandle>& OutPreviousKeyHandle, TOptional<FKeyHandle>& OutNextKeyHandle) const override;
 	virtual bool Evaluate(double ProspectiveTime, double& OutValue) const override;
+	
 	virtual void AddKeys(TArrayView<const FKeyPosition> InKeyPositions, TArrayView<const FKeyAttributes> InAttributes, TArrayView<TOptional<FKeyHandle>>* OutKeyHandles) override;
 	virtual void RemoveKeys(TArrayView<const FKeyHandle> InKeys, double InCurrentTime) override;
 
+	virtual void ReplaceKeyHandles(TConstArrayView<FKeyHandle> InCurrentHandles, TConstArrayView<FKeyHandle> InNewHandles) override;
+
 	virtual FTransform2d GetCurveTransform() const override;
 	virtual bool IsReadOnly() const override;
 	virtual UObject* GetOwningObject() const override
@@ -101,6 +104,9 @@ protected:
 	virtual double GetKeyValue(TArrayView<const ChannelValue> Values, int32 Index) const = 0;
 	virtual void SetKeyValue(int32 Index, double KeyValue) const = 0;
 
+	/** Calls PostEditChange, OnCurvesModified, notifies Sequencer's OnChannelChanged(), and ensures that LastSignature is invalidated. */
+	virtual void InvokePostChangeCallbacks(ChannelType* Channel, UMovieSceneSection* Section);
+
 public:
 
 	const TMovieSceneChannelHandle<ChannelType>& GetChannelHandle() const  { return ChannelHandle; } 
diff --git a/Engine/Source/Runtime/Engine/Classes/Curves/KeyHandle.h b/Engine/Source/Runtime/Engine/Classes/Curves/KeyHandle.h
--- a/Engine/Source/Runtime/Engine/Classes/Curves/KeyHandle.h
+++ b/Engine/Source/Runtime/Engine/Classes/Curves/KeyHandle.h
@@ -177,6 +177,14 @@ public:
 	 */
 	ENGINE_API void DeallocateHandle(int32 Index);
 
+	/**
+	 *  Replaces the key handle at Index with HandleToReplace.
+	 *  
+	 * @param Index				The index at which to replace the handle
+	 * @param HandleToReplace	The key handle to set at the index. Fails if this handle already is in the table.
+	 */
+	ENGINE_API void ReplaceHandle(int32 Index, FKeyHandle HandleToReplace);
+
 	/**
 	 * Reset this lookup table, forgetting all key handles and indices
 	 */
diff --git a/Engine/Source/Runtime/Engine/Private/Curves/KeyHandle.cpp b/Engine/Source/Runtime/Engine/Private/Curves/KeyHandle.cpp
--- a/Engine/Source/Runtime/Engine/Private/Curves/KeyHandle.cpp
+++ b/Engine/Source/Runtime/Engine/Private/Curves/KeyHandle.cpp
@@ -253,6 +253,16 @@ void FKeyHandleLookupTable::DeallocateHandle(int32 Index)
 	}
 }
 
+void FKeyHandleLookupTable::ReplaceHandle(int32 Index, FKeyHandle HandleToReplace)
+{
+	if (KeyHandles.IsValidIndex(Index) && !KeyHandlesToIndices.Contains(HandleToReplace))
+	{
+		KeyHandlesToIndices.Remove(KeyHandles[Index]);
+		KeyHandlesToIndices.Add(HandleToReplace, Index);
+		KeyHandles[Index] = HandleToReplace;
+	}
+}
+
 void FKeyHandleLookupTable::Reset()
 {
 	KeyHandles.Reset();
diff --git a/Engine/Source/Runtime/MovieScene/Private/Channels/MovieSceneChannelData.cpp b/Engine/Source/Runtime/MovieScene/Private/Channels/MovieSceneChannelData.cpp
--- a/Engine/Source/Runtime/MovieScene/Private/Channels/MovieSceneChannelData.cpp
+++ b/Engine/Source/Runtime/MovieScene/Private/Channels/MovieSceneChannelData.cpp
@@ -142,6 +142,18 @@ int32 FMovieSceneChannelData::AddKeyInternal(FFrameNumber InTime)
 	return InsertIndex;
 }
 
+void FMovieSceneChannelData::ReplaceKeyHandlesInternal(TConstArrayView<int32> KeyIndices, TConstArrayView<FKeyHandle> NewHandles)
+{
+	if (ensureMsgf(KeyHandles, TEXT("This channel does not contain key handles"))
+		&& ensure(KeyIndices.Num() == NewHandles.Num()))
+	{
+		for (int32 Index = 0; Index < KeyIndices.Num(); ++Index)
+		{
+			KeyHandles->ReplaceHandle(KeyIndices[Index], NewHandles[Index]);
+		}
+	}
+}
+
 int32 FMovieSceneChannelData::MoveKeyInternal(int32 KeyIndex, FFrameNumber InNewTime)
 {
 	check(Times->IsValidIndex(KeyIndex));
diff --git a/Engine/Source/Runtime/MovieScene/Public/Channels/MovieSceneChannelData.h b/Engine/Source/Runtime/MovieScene/Public/Channels/MovieSceneChannelData.h
--- a/Engine/Source/Runtime/MovieScene/Public/Channels/MovieSceneChannelData.h
+++ b/Engine/Source/Runtime/MovieScene/Public/Channels/MovieSceneChannelData.h
@@ -227,6 +227,13 @@ protected:
 	 */
 	MOVIESCENE_API int32 AddKeyInternal(FFrameNumber InTime);
 
+	/**
+	 * Replaces the key handle at the given Index with NewHandle. Fails if the handle is already in use.
+	 *
+	 * @param KeyIndices	The index of the key to replace
+	 * @param NewHandles	The key handle instance that should be at the given index.
+	 */
+	MOVIESCENE_API void ReplaceKeyHandlesInternal(TConstArrayView<int32> KeyIndices, TConstArrayView<FKeyHandle> NewHandles);
 
 protected:
 
@@ -462,6 +469,27 @@ struct TMovieSceneChannelData : FMovieSceneChannelData
 		}
 	}
 
+	/**
+	 * Replaces the key handle at the given Index with NewHandle. Fails if the handle is already in use.
+	 *
+	 * @param KeyIndex	The index of the key to replace
+	 * @param NewHandle	The key handle instance that should be at the given index.
+	 */
+	void ReplaceKeyHandle(int32 KeyIndex, FKeyHandle NewHandle)
+	{
+		ReplaceKeyHandles(TConstArrayView<int32>(&KeyIndex, 1), TConstArrayView<FKeyHandle>(&NewHandle, 1));
+	}
+	/**
+	 * Replaces the key handles at the given KeyIndices with NewHandles. Handles already in use are skipped.
+	 *
+	 * @param KeyIndices	The indidces of the keys to replace
+	 * @param NewHandles	The key handle instances that should be at the given indices.
+	 */
+	void ReplaceKeyHandles(TConstArrayView<int32> KeyIndices, TConstArrayView<FKeyHandle> NewHandles)
+	{
+		ReplaceKeyHandlesInternal(KeyIndices, NewHandles);
+	}
+
 	/**
 	 * Remove the key at a given index
 	 *
