Sequencer: Fix OnPayloadVariableChanged so that it doesn't treat all strings as uobjects

Retested UE-165499 to make sure redirector is processed

diff --git a/Engine/Source/Editor/MovieSceneTools/Private/MovieSceneDirectorBlueprintEndpointCustomization.cpp b/Engine/Source/Editor/MovieSceneTools/Private/MovieSceneDirectorBlueprintEndpointCustomization.cpp
--- a/Engine/Source/Editor/MovieSceneTools/Private/MovieSceneDirectorBlueprintEndpointCustomization.cpp
+++ b/Engine/Source/Editor/MovieSceneTools/Private/MovieSceneDirectorBlueprintEndpointCustomization.cpp
@@ -412,7 +412,11 @@ void FMovieSceneDirectorBlueprintEndpointCustomization::OnPayloadVariableChanged
 		return;
 	}
 
-	TObjectPtr<UObject> NewValueObject = FindObject<UObject>(nullptr, *NewValueString);
+	TObjectPtr<UObject> NewValueObject;
+	if (FPackageName::IsValidPath(NewValueString))
+	{
+		NewValueObject = FindObject<UObject>(nullptr, *NewValueString);
+	}
 
 #if WITH_EDITORONLY_DATA
 	// Fixup redirectors
