Sequencer: Fix the text not being focused in the rename widget by showing/hiding upon entering/existing renaming.

#rnx

diff --git a/Engine/Source/Editor/MovieSceneTools/Private/Sections/CinematicShotSection.h b/Engine/Source/Editor/MovieSceneTools/Private/Sections/CinematicShotSection.h
--- a/Engine/Source/Editor/MovieSceneTools/Private/Sections/CinematicShotSection.h
+++ b/Engine/Source/Editor/MovieSceneTools/Private/Sections/CinematicShotSection.h
@@ -43,6 +43,7 @@ public:
 
 	// FThumbnail interface
 	virtual void SetSingleTime(double GlobalTime) override;
+	virtual bool CanRename() const override { return true; }
 	virtual FText HandleThumbnailTextBlockText() const override;
 	virtual void HandleThumbnailTextBlockTextCommitted(const FText& NewThumbnailName, ETextCommit::Type CommitType) override;
 	virtual UCameraComponent* GetViewCamera() override;
diff --git a/Engine/Source/Editor/MovieSceneTools/Private/Sections/ThumbnailSection.cpp b/Engine/Source/Editor/MovieSceneTools/Private/Sections/ThumbnailSection.cpp
--- a/Engine/Source/Editor/MovieSceneTools/Private/Sections/ThumbnailSection.cpp
+++ b/Engine/Source/Editor/MovieSceneTools/Private/Sections/ThumbnailSection.cpp
@@ -89,7 +89,9 @@ TSharedRef<SWidget> FThumbnailSection::GenerateSectionWidget()
 				.ShadowOffset(FVector2D(1,1))
 				.OnTextCommitted(this, &FThumbnailSection::HandleThumbnailTextBlockTextCommitted)
 				.IsReadOnly(!CanRename())
-				.Visibility(this, &FThumbnailSection::GetRenameVisibility)
+				.Visibility(EVisibility::Hidden) // Starts hidden and is made visible when renaming is invoked
+				.OnEnterEditingMode(this, &FThumbnailSection::OnEnterEditingMode)
+				.OnExitEditingMode(this, &FThumbnailSection::OnExitEditingMode)
 		];
 }
 
@@ -102,13 +104,27 @@ EVisibility FThumbnailSection::GetRenameVisibility() const
 	return EVisibility::Hidden;
 }
 
+void FThumbnailSection::OnEnterEditingMode()
+{
+	if (NameWidget.IsValid())
+	{
+		NameWidget->SetVisibility(EVisibility::Visible);
+	}
+}
+
+void FThumbnailSection::OnExitEditingMode()
+{
+	if (NameWidget.IsValid())
+	{
+		NameWidget->SetVisibility(EVisibility::Hidden);
+	}
+}
+
 void FThumbnailSection::EnterRename()
 {
 	if (NameWidget.IsValid())
 	{
-		NameWidget->SetReadOnly(false);
 		NameWidget->EnterEditingMode();
-		NameWidget->SetReadOnly(!CanRename());
 	}
 }
 
diff --git a/Engine/Source/Editor/MovieSceneTools/Public/Sections/ThumbnailSection.h b/Engine/Source/Editor/MovieSceneTools/Public/Sections/ThumbnailSection.h
--- a/Engine/Source/Editor/MovieSceneTools/Public/Sections/ThumbnailSection.h
+++ b/Engine/Source/Editor/MovieSceneTools/Public/Sections/ThumbnailSection.h
@@ -94,6 +94,12 @@ protected:
 	/** Get rename visibility */
 	UE_API EVisibility GetRenameVisibility() const;
 
+	/** When renaming is started, show the rename widget */
+	UE_API void OnEnterEditingMode();
+
+	/** When renaming is completed, hide the rename widget */
+	UE_API void OnExitEditingMode();
+
 	UE_API void BuildThumbnailsMenu(FMenuBuilder& InMenuBuilder);
 
 protected:
