Fixed weird case where if the first component you add to a prefab is a mesh, the transform would not be saved.
Fixed crash in the TransformTrackEditor if there is not active member in the changed event
Enabled IDO earlier in the Unittests and fixed some broken tests following this change.

#rnx

diff --git a/Engine/Source/Editor/MovieSceneTools/Private/TrackEditors/TransformTrackEditor.cpp b/Engine/Source/Editor/MovieSceneTools/Private/TrackEditors/TransformTrackEditor.cpp
--- a/Engine/Source/Editor/MovieSceneTools/Private/TrackEditors/TransformTrackEditor.cpp
+++ b/Engine/Source/Editor/MovieSceneTools/Private/TrackEditors/TransformTrackEditor.cpp
@@ -395,7 +395,7 @@ void F3DTransformTrackEditor::OnTransformChanged( UObject& InObject )
 
 void F3DTransformTrackEditor::OnPrePropertyChanged(UObject* InObject, const FEditPropertyChain& InPropertyChain)
 {
-	FProperty* PropertyAboutToChange = InPropertyChain.GetActiveMemberNode()->GetValue();
+	FProperty* PropertyAboutToChange = InPropertyChain.GetActiveMemberNode() ? InPropertyChain.GetActiveMemberNode()->GetValue() : nullptr;
 	const FName MemberPropertyName = PropertyAboutToChange != nullptr ? PropertyAboutToChange->GetFName() : NAME_None;
 	const bool bTransformationToChange =
 		(MemberPropertyName == USceneComponent::GetRelativeLocationPropertyName() ||
