Curve Editor: Introduce FCurveModel::GetAllKeys. For cases were API user wants all keys this is signifcantly faster.

With 110.000 keys, I profiled that it takes about 130 ms for GetValueRange to determine min max values. Before, the only way to get all keys was to use GetKeys, which wanted min/max x/x. The work for GetValueRange is effectively unnecessary.

diff --git a/Engine/Source/Editor/CurveEditor/Private/CurveModel.cpp b/Engine/Source/Editor/CurveEditor/Private/CurveModel.cpp
--- a/Engine/Source/Editor/CurveEditor/Private/CurveModel.cpp
+++ b/Engine/Source/Editor/CurveEditor/Private/CurveModel.cpp
@@ -7,6 +7,14 @@
 #include "Curves/KeyHandle.h"
 #include "HAL/PlatformCrt.h"
 
+void FCurveModel::GetAllKeys(TArray<FKeyHandle>& OutKeyHandles) const
+{
+	double MinX, MaxX, MinY, MaxY;
+	GetTimeRange(MinX, MaxX);
+	GetValueRange(MinY, MaxY);
+	GetKeys(MinX, MaxX, MinY, MaxY, OutKeyHandles);
+}
+
 void FCurveModel::PutKeys(
 	TConstArrayView<FKeyHandle> InKeys, TConstArrayView<FKeyPosition> InKeyPositions, TConstArrayView<FKeyAttributes> InAttributes
 	)
diff --git a/Engine/Source/Editor/CurveEditor/Public/CurveModel.h b/Engine/Source/Editor/CurveEditor/Public/CurveModel.h
--- a/Engine/Source/Editor/CurveEditor/Public/CurveModel.h
+++ b/Engine/Source/Editor/CurveEditor/Public/CurveModel.h
@@ -117,6 +117,17 @@ public:
 	 */
 	virtual void GetKeys(double MinTime, double MaxTime, double MinValue, double MaxValue, TArray<FKeyHandle>& OutKeyHandles) const = 0;
 
+	/**
+	 * Gets all keys in the curve.
+	 * Prefer this function over GetKeys because it has the potential to be much faster (GetValueRange evaluation can be slow for bezier curves).
+	 */
+	UE_API virtual void GetAllKeys(TArray<FKeyHandle>& OutKeyHandles) const;
+	/**
+	 * Gets all keys in the curve.
+	 * Prefer this function over GetKeys because it has the potential to be much faster (GetValueRange evaluation can be slow for bezier curves).
+	 */
+	TArray<FKeyHandle> GetAllKeys() const { TArray<FKeyHandle> Result; GetAllKeys(Result); return Result; }
+
 	/**
 	 * Add keys to this curve
 	 *
diff --git a/Engine/Source/Editor/MovieSceneTools/Private/Channels/ChannelCurveModel.cpp b/Engine/Source/Editor/MovieSceneTools/Private/Channels/ChannelCurveModel.cpp
--- a/Engine/Source/Editor/MovieSceneTools/Private/Channels/ChannelCurveModel.cpp
+++ b/Engine/Source/Editor/MovieSceneTools/Private/Channels/ChannelCurveModel.cpp
@@ -212,6 +212,24 @@ void FChannelCurveModel<ChannelType, ChannelValue, KeyType>::GetKeys(double MinT
 	}
 }
 
+template <class ChannelType, class ChannelValue, class KeyType>
+void FChannelCurveModel<ChannelType, ChannelValue, KeyType>::GetAllKeys(TArray<FKeyHandle>& OutKeyHandles) const
+{
+	ChannelType* Channel = ChannelHandle.Get();
+	UMovieSceneSection* Section = WeakSection.Get();
+	if (!Channel || !Section)
+	{
+		return;
+	}
+	
+	TMovieSceneChannelData<ChannelValue> ChannelData = Channel->GetData();
+	const int32 Num = ChannelData.GetTimes().Num();
+	for (int32 KeyIndex = 0; KeyIndex < Num; ++KeyIndex)
+	{
+		OutKeyHandles.Add(ChannelData.GetHandle(KeyIndex));
+	}
+}
+
 template <class ChannelType, class ChannelValue, class KeyType>
 void FChannelCurveModel<ChannelType, ChannelValue, KeyType>::GetKeyDrawInfo(ECurvePointType PointType, const FKeyHandle InKeyHandle, FKeyDrawInfo& OutDrawInfo) const
 {
diff --git a/Engine/Source/Editor/MovieSceneTools/Public/Channels/ChannelCurveModel.h b/Engine/Source/Editor/MovieSceneTools/Public/Channels/ChannelCurveModel.h
--- a/Engine/Source/Editor/MovieSceneTools/Public/Channels/ChannelCurveModel.h
+++ b/Engine/Source/Editor/MovieSceneTools/Public/Channels/ChannelCurveModel.h
@@ -50,6 +50,7 @@ public:
 
 	virtual void DrawCurve(const FCurveEditor& CurveEditor, const FCurveEditorScreenSpace& ScreenSpace, TArray<TTuple<double, double>>& OutInterpolatingPoints) const override;
 	virtual void GetKeys(double MinTime, double MaxTime, double MinValue, double MaxValue, TArray<FKeyHandle>& OutKeyHandles) const override;
+	virtual void GetAllKeys(TArray<FKeyHandle>& OutKeyHandles) const override;
 	virtual void GetKeyDrawInfo(ECurvePointType PointType, const FKeyHandle InKeyHandle, FKeyDrawInfo& OutDrawInfo) const override;
 
 	virtual void GetKeyPositions(TArrayView<const FKeyHandle> InKeys, TArrayView<FKeyPosition> OutKeyPositions) const override;
