Curve Editor: Improve performance of FBezierChannelCurveModel::GetValueRange for large number of keys. The improvement is to use the extreme points of the curve sections using derivatives instead of via regression method. Not implemented for weighted tangents because ComputeExtents is not implemented for it.

On ControlRigExample -> PossessableBindings_CopyPaste -> Possessable_FKCR_CopyPaste, open Curve Editor.
In SCurveEditorViewNormalized::UpdateViewToTransformCurves usage of GetValueRange totalling 110.000 keys accross all curves, we saved 135ms.

We can also probably use the min/max values if the tangents themselves are less than the min/max values since it wouldn't be possible for the value to exceed those bounds if the tangents are within those limits.

diff --git a/Engine/Source/Editor/MovieSceneTools/Private/Channels/BezierChannelCurveModel.cpp b/Engine/Source/Editor/MovieSceneTools/Private/Channels/BezierChannelCurveModel.cpp
--- a/Engine/Source/Editor/MovieSceneTools/Private/Channels/BezierChannelCurveModel.cpp
+++ b/Engine/Source/Editor/MovieSceneTools/Private/Channels/BezierChannelCurveModel.cpp
@@ -584,14 +584,13 @@ template<typename ChannelType, typename ChannelValue, typename KeyType>
 void FBezierChannelCurveModel<ChannelType, ChannelValue, KeyType>::GetValueRange(double InMinTime, double InMaxTime, double& MinValue, double& MaxValue) const
 {
 	ChannelType* Channel = this->GetChannelHandle().Get();
-
 	UMovieSceneSignedObject* SignedOwner = this->template GetOwningObjectOrOuter<UMovieSceneSignedObject>();
 	UMovieScene*             MovieScene  = SignedOwner ? SignedOwner->GetTypedOuter<UMovieScene>() : nullptr;
 
 	if (Channel && SignedOwner && MovieScene && !this->IsReadOnly())
 	{
-		TArrayView<const FFrameNumber> Times = Channel->GetData().GetTimes();
-		TArrayView<const ChannelValue> Values = Channel->GetData().GetValues();
+		const TConstArrayView<FFrameNumber> Times = Channel->GetData().GetTimes();
+		const TConstArrayView<ChannelValue> Values = Channel->GetData().GetValues();
 
 		if (Times.Num() == 0)
 		{
@@ -600,33 +599,46 @@ void FBezierChannelCurveModel<ChannelType, ChannelValue, KeyType>::GetValueRange
 		}
 		else
 		{
-			FFrameRate TickResolution = MovieScene->GetTickResolution();
-			double ToTime = TickResolution.AsInterval();
-			int32 LastKeyIndex = Values.Num() - 1;
+			const FFrameRate TickResolution = MovieScene->GetTickResolution();
+			const double ToTime = TickResolution.AsInterval();
+			const int32 LastKeyIndex = Values.Num() - 1;
 			MinValue = TNumericLimits<double>::Max();
 			MaxValue = TNumericLimits<double>::Lowest();
 
 			for (int32 i = 0; i < Values.Num(); i++)
 			{
-				double KeyTime = static_cast<double>(Times[i].Value) * ToTime;
+				const double KeyTime = static_cast<double>(Times[i].Value) * ToTime;
 				if (KeyTime < InMinTime)
 				{
 					continue;
 				}
-				else if (KeyTime > InMaxTime)
+				if (KeyTime > InMaxTime)
 				{
 					break;
 				}
+				
 				const ChannelValue& Key = Values[i];
+				MinValue = FMath::Min(MinValue, static_cast<double>(Key.Value));
+				MaxValue = FMath::Max(MaxValue, static_cast<double>(Key.Value));
 
-				MinValue = FMath::Min(MinValue, (double)Key.Value);
-				MaxValue = FMath::Max(MaxValue, (double)Key.Value);
-
-				if (Key.InterpMode == RCIM_Cubic && i != LastKeyIndex)
+				const bool bNeedsCubicEstimation = Key.InterpMode == RCIM_Cubic && i != LastKeyIndex;
+				const bool bIsWeighted = Key.Tangent.TangentWeightMode != RCTWM_WeightedNone;
+				if (bNeedsCubicEstimation && !bIsWeighted)
+				{
+					// This method is faster than FeaturePointMethod's regression method because the extrema are computed by solving the derivatives.
+					// 17th of July, 2025: On ControlRigExample -> PossessableBindings_CopyPaste -> Possessable_FKCR_CopyPaste, open Curve Editor.
+					// In SCurveEditorViewNormalized::UpdateViewToTransformCurves usage of GetValueRange totalling 110.000 keys accross all curves,
+					// we saved 135ms. 
+					const UE::MovieScene::Interpolation::FInterpolationExtents Extents = Channel->ComputeExtents(Times[i].Value, Times[i + 1].Value);
+					MinValue = FMath::Min(MinValue, Extents.MinValue);
+					MaxValue = FMath::Max(MaxValue, Extents.MaxValue);
+				}
+				else if (bNeedsCubicEstimation && bIsWeighted)
 				{
-					const ChannelValue& NextKey = Values[i + 1];
-					double NextTime = static_cast<double>(Times[i + 1].Value) * ToTime;
-					double TimeStep = (NextTime - KeyTime) * 0.2f;
+					// UE-305634: FCachedInterpolation::ComputeExtents is not yet implemented correctly for weighted tangents.
+					// So we need to fall back to the slow regression method. 
+					const double NextTime = static_cast<double>(Times[i + 1].Value) * ToTime;
+					const double TimeStep = (NextTime - KeyTime) * 0.2f;
 					FeaturePointMethod(KeyTime, NextTime, Key.Value, TimeStep, 0, 3, MaxValue, MinValue);
 				}
 			}
diff --git a/Engine/Source/Runtime/MovieScene/Private/Channels/MovieSceneFloatChannel.cpp b/Engine/Source/Runtime/MovieScene/Private/Channels/MovieSceneFloatChannel.cpp
--- a/Engine/Source/Runtime/MovieScene/Private/Channels/MovieSceneFloatChannel.cpp
+++ b/Engine/Source/Runtime/MovieScene/Private/Channels/MovieSceneFloatChannel.cpp
@@ -4,6 +4,7 @@
 #include "Channels/MovieSceneChannelProxy.h"
 #include "Channels/MovieSceneCurveChannelImpl.h"
 #include "Channels/MovieSceneInterpolation.h"
+#include "Channels/MovieScenePiecewiseCurveUtils.inl"
 #include "MovieSceneFrameMigration.h"
 #include "MovieSceneFwd.h"
 #include "UObject/FortniteMainBranchObjectVersion.h"
@@ -52,6 +53,77 @@ namespace UE::MovieScene
 	}
 } // namespace UE::MovieScene
 
+namespace UE::MovieScene::Interpolation
+{
+	struct FFloatChannelPiecewiseData
+	{
+		const FMovieSceneFloatChannel* Channel;
+
+		bool HasDefaultValue() const
+		{
+			return Channel->GetDefault().IsSet();
+		}
+		double GetDefaultValue() const
+		{
+			return Channel->GetDefault().Get(0.0);
+		}
+		double PreExtrapolate(const FFrameTime& InTime) const
+		{
+			float Result = 0.0;
+			TMovieSceneCurveChannelImpl<FMovieSceneFloatChannel>::Evaluate(Channel, InTime, Result);
+			return Result;
+		}
+		double PostExtrapolate(const FFrameTime& InTime) const
+		{
+			float Result = 0.0;
+			TMovieSceneCurveChannelImpl<FMovieSceneFloatChannel>::Evaluate(Channel, InTime, Result);
+			return Result;
+		}
+		int32 NumPieces() const
+		{
+			return FMath::Max(0, Channel->GetData().GetValues().Num() - 1);
+		}
+		int32 GetIndexOfPieceByTime(const FFrameTime& Time) const
+		{
+			TArrayView<const FFrameNumber> Times = Channel->GetData().GetTimes();
+			return FMath::Max(Algo::UpperBound(Times, Time)-1, 0);
+		}
+		Interpolation::FCachedInterpolation GetPieceByIndex(int32 Index) const
+		{
+			return TMovieSceneCurveChannelImpl<FMovieSceneFloatChannel>::GetInterpolationForKey(Channel, Index);
+		}
+		Interpolation::FCachedInterpolation GetPieceByTime(const FFrameTime& Time) const
+		{
+			return TMovieSceneCurveChannelImpl<FMovieSceneFloatChannel>::GetInterpolationForTime(Channel, Time);
+		}
+		FFrameNumber GetFiniteStart() const
+		{
+			return Channel->GetData().GetTimes()[0];
+		}
+		FFrameNumber GetFiniteEnd() const
+		{
+			return Channel->GetData().GetTimes().Last();
+		}
+		ERichCurveExtrapolation GetPreExtrapolation() const
+		{
+			return Channel->PreInfinityExtrap;
+		}
+		ERichCurveExtrapolation GetPostExtrapolation() const
+		{
+			return Channel->PostInfinityExtrap;
+		}
+		double GetStartingValue() const
+		{
+			return Channel->GetData().GetValues()[0].Value;
+		}
+		double GetEndingValue() const
+		{
+			return Channel->GetData().GetValues().Last().Value;
+		}
+	};
+
+} // namespace UE::MovieScene::Interpolation
+
 bool FMovieSceneFloatValue::Serialize(FArchive& Ar)
 {
 	return TMovieSceneCurveChannelImpl<FMovieSceneFloatChannel>::SerializeChannelValue(*this, Ar);
@@ -80,6 +152,14 @@ FMovieSceneFloatChannel::FMovieSceneFloatChannel()
 
 FMovieSceneFloatChannel::~FMovieSceneFloatChannel() = default;
 
+UE::MovieScene::Interpolation::FInterpolationExtents FMovieSceneFloatChannel::ComputeExtents(FFrameTime StartTime, FFrameTime EndTime) const
+{
+	using namespace UE::MovieScene;
+	using namespace UE::MovieScene::Interpolation;
+
+	return ComputePiecewiseExtents(FFloatChannelPiecewiseData{ this }, StartTime, EndTime);
+}
+
 UE::MovieScene::FPiecewiseCurve FMovieSceneFloatChannel::AsPiecewiseCurve(bool bWithPreAndPostInfinityExtrap) const
 {
 	using namespace UE::MovieScene;
diff --git a/Engine/Source/Runtime/MovieScene/Public/Channels/MovieSceneFloatChannel.h b/Engine/Source/Runtime/MovieScene/Public/Channels/MovieSceneFloatChannel.h
--- a/Engine/Source/Runtime/MovieScene/Public/Channels/MovieSceneFloatChannel.h
+++ b/Engine/Source/Runtime/MovieScene/Public/Channels/MovieSceneFloatChannel.h
@@ -129,6 +129,15 @@ struct FMovieSceneFloatChannel : public FMovieSceneChannel
 
 	MOVIESCENE_API virtual ~FMovieSceneFloatChannel();
 
+	/**
+	 * Compute the value extents of this curve within the specified limits
+	 * 
+	 * @param StartTime     The start of the range within which to compute extents
+	 * @param EndTime       The end of the range within which to compute extents
+	 * @return The extents of this channel
+	 */
+	MOVIESCENE_API UE::MovieScene::Interpolation::FInterpolationExtents ComputeExtents(FFrameTime StartTime, FFrameTime EndTime) const;
+	
 	MOVIESCENE_API UE::MovieScene::FPiecewiseCurve AsPiecewiseCurve(bool bWithPreAndPostInfinityExtrap = true) const;
 
 	/**
