Curve Editor: Fix an issue where pre-/post-infinity was not painted correctly when the first or last key's tangent mode was linear and pre-/post-infinity interpolation was linear

diff --git a/Engine/Source/Editor/MovieSceneTools/Private/Cache/MovieSceneCachedCurve.h b/Engine/Source/Editor/MovieSceneTools/Private/Cache/MovieSceneCachedCurve.h
--- a/Engine/Source/Editor/MovieSceneTools/Private/Cache/MovieSceneCachedCurve.h
+++ b/Engine/Source/Editor/MovieSceneTools/Private/Cache/MovieSceneCachedCurve.h
@@ -897,14 +897,21 @@ namespace UE::MovieSceneTools
 		else
 		{
 			const bool bPreinfinityVisible = Times[0].Value > ScreenSpace.GetInputMin() * TickResolution;
-			const bool bLinearWithoutArriveTangent =
+			if (!bPreinfinityVisible)
+			{
+				return true;
+			}
+			
+			const bool bHasArriveTangent = !KeyAttributes.IsEmpty() && !KeyAttributes[0].HasArriveTangent();
+			const bool bFirstKeyIsLinearInterpMode = !KeyAttributes.IsEmpty() && KeyAttributes[0].GetInterpMode() == ERichCurveInterpMode::RCIM_Linear;
+
+			const bool bUseConstantForLinear =
 				FullRangeInterpolation.PreInfinityExtrapolation == ERichCurveExtrapolation::RCCE_Linear &&
-				(KeyAttributes.IsEmpty() || !KeyAttributes[0].HasArriveTangent());
+				(!bHasArriveTangent || !bFirstKeyIsLinearInterpMode);
 
-			if (bPreinfinityVisible &&
-				(FullRangeInterpolation.PreInfinityExtrapolation == ERichCurveExtrapolation::RCCE_None ||
+			if ((FullRangeInterpolation.PreInfinityExtrapolation == ERichCurveExtrapolation::RCCE_None ||
 				FullRangeInterpolation.PreInfinityExtrapolation == ERichCurveExtrapolation::RCCE_Constant ||
-				bLinearWithoutArriveTangent))
+				bUseConstantForLinear))
 			{
 				const double ExtentValue = Sign * Values[0].Value;
 				CachedDrawParams.InterpolatingPoints.Emplace(
@@ -913,13 +920,44 @@ namespace UE::MovieSceneTools
 
 				return true;
 			}
-			else if (bPreinfinityVisible &&
-				FullRangeInterpolation.PreInfinityExtrapolation == ERichCurveExtrapolation::RCCE_Linear)
+			else if (FullRangeInterpolation.PreInfinityExtrapolation == ERichCurveExtrapolation::RCCE_Linear)
 			{
-				const double ExtentValue = Sign * Values[0].Value + KeyAttributes[0].GetArriveTangent() * (ScreenSpace.GetInputMin() - Times[0] / TickResolution);
-				CachedDrawParams.InterpolatingPoints.Emplace(
-					ScreenSpace.SecondsToScreen(ScreenSpace.GetInputMin()),
-					ScreenSpace.ValueToScreen(ExtentValue));
+				const double MinTime = ScreenSpace.GetInputMin();
+
+				if (KeyAttributes[0].GetInterpMode() == ERichCurveInterpMode::RCIM_Linear)
+				{
+					// Draw a line in direction of the first two keys
+					const double X1 = Times[1] / TickResolution;
+					const double Y1 = Values[1].Value;
+
+					const double X2 = Times[0] / TickResolution;
+					const double Y2 = Values[0].Value;
+
+					if (X2 == X1) 
+					{
+						const double ExtentValue = Y1 > Y2 ? ScreenSpace.GetOutputMin() : ScreenSpace.GetOutputMax();
+
+						CachedDrawParams.InterpolatingPoints.Emplace(
+							ScreenSpace.SecondsToScreen(X2),
+							ScreenSpace.ValueToScreen(ExtentValue));
+					}
+					else
+					{
+						const double ExtentValue = Y1 + (Y2 - Y1) * (MinTime - X1) / (X2 - X1);
+
+						CachedDrawParams.InterpolatingPoints.Emplace(
+							ScreenSpace.SecondsToScreen(MinTime),
+							ScreenSpace.ValueToScreen(ExtentValue));
+					}
+				}
+				else
+				{
+					// Draw a line in direction of the tangent
+					const double ExtentValue = Sign * Values[0].Value + KeyAttributes[0].GetArriveTangent() * (MinTime - Times[0] / TickResolution);
+					CachedDrawParams.InterpolatingPoints.Emplace(
+						ScreenSpace.SecondsToScreen(MinTime),
+						ScreenSpace.ValueToScreen(ExtentValue));
+				}
 
 				return true;
 			}
@@ -953,14 +991,21 @@ namespace UE::MovieSceneTools
 		else
 		{
 			const bool bPostinfinityVisible = Times.Last().Value < ScreenSpace.GetInputMax() * TickResolution;
-			const bool bLinearWithoutLeaveTangent =
+			if (!bPostinfinityVisible)
+			{
+				return true;
+			}
+
+			const bool bHasLeaveTangent = !KeyAttributes.IsEmpty() && !KeyAttributes.Last().HasLeaveTangent();
+			const bool bLastKeyIsLinearInterpMode = !KeyAttributes.IsEmpty() && KeyAttributes.Last().GetInterpMode() == ERichCurveInterpMode::RCIM_Linear;
+
+			const bool bUseConstantForLinear =
 				FullRangeInterpolation.PostInfinityExtrapolation == ERichCurveExtrapolation::RCCE_Linear &&
-				(KeyAttributes.IsEmpty() || !KeyAttributes.Last().HasLeaveTangent());
+				(!bHasLeaveTangent || !bLastKeyIsLinearInterpMode);
 
-			if (bPostinfinityVisible &&
-				(FullRangeInterpolation.PostInfinityExtrapolation == ERichCurveExtrapolation::RCCE_None ||
+			if (FullRangeInterpolation.PostInfinityExtrapolation == ERichCurveExtrapolation::RCCE_None ||
 				FullRangeInterpolation.PostInfinityExtrapolation == ERichCurveExtrapolation::RCCE_Constant ||
-				bLinearWithoutLeaveTangent))
+				bUseConstantForLinear)
 			{
 				const double ExtentValue = Sign * Values.Last().Value;
 				CachedDrawParams.InterpolatingPoints.Emplace(
@@ -969,13 +1014,46 @@ namespace UE::MovieSceneTools
 
 				return true;
 			}
-			else if (bPostinfinityVisible &&
-				FullRangeInterpolation.PostInfinityExtrapolation == ERichCurveExtrapolation::RCCE_Linear)
+			else if (FullRangeInterpolation.PostInfinityExtrapolation == ERichCurveExtrapolation::RCCE_Linear)
 			{
-				const double ExtentValue = Sign * Values.Last().Value + KeyAttributes.Last().GetLeaveTangent() * (ScreenSpace.GetInputMax() - Times.Last() / TickResolution);
-				CachedDrawParams.InterpolatingPoints.Emplace(
-					ScreenSpace.SecondsToScreen(ScreenSpace.GetInputMax()),
-					ScreenSpace.ValueToScreen(ExtentValue));
+				const double MaxTime = ScreenSpace.GetInputMax();
+
+				if (KeyAttributes.Last().GetInterpMode() == ERichCurveInterpMode::RCIM_Linear)
+				{					
+					// Draw a line in direction of the last two keys
+					const int32 NumKeys = Times.Num();
+
+					const double X1 = Times[NumKeys - 2] / TickResolution;
+					const double Y1 = Values[NumKeys - 2].Value;
+
+					const double X2 = Times[NumKeys - 1] / TickResolution;
+					const double Y2 = Values[NumKeys - 1].Value;
+
+					if (X2 == X1)
+					{
+						const double ExtentValue = Y1 > Y2 ? ScreenSpace.GetOutputMin() : ScreenSpace.GetOutputMax();
+
+						CachedDrawParams.InterpolatingPoints.Emplace(
+							ScreenSpace.SecondsToScreen(X2),
+							ScreenSpace.ValueToScreen(ExtentValue));
+					}
+					else
+					{
+						const double ExtentValue = Y1 + (Y2 - Y1) * (MaxTime - X1) / (X2 - X1);
+
+						CachedDrawParams.InterpolatingPoints.Emplace(
+							ScreenSpace.SecondsToScreen(MaxTime),
+							ScreenSpace.ValueToScreen(ExtentValue));
+					}
+				}
+				else
+				{
+					// Draw a line in direction of the tangent
+					const double ExtentValue = Sign * Values.Last().Value + KeyAttributes.Last().GetLeaveTangent() * (ScreenSpace.GetInputMax() - Times.Last() / TickResolution);
+					CachedDrawParams.InterpolatingPoints.Emplace(
+						ScreenSpace.SecondsToScreen(MaxTime),
+						ScreenSpace.ValueToScreen(ExtentValue));
+				}
 
 				return true;
 			}
