Curve editor: Fix possible check being hit when drawing pre-/post-infinity extrapolation

#rnx

diff --git a/Engine/Source/Editor/MovieSceneTools/Private/Cache/MovieSceneCachedCurve.h b/Engine/Source/Editor/MovieSceneTools/Private/Cache/MovieSceneCachedCurve.h
--- a/Engine/Source/Editor/MovieSceneTools/Private/Cache/MovieSceneCachedCurve.h
+++ b/Engine/Source/Editor/MovieSceneTools/Private/Cache/MovieSceneCachedCurve.h
@@ -903,7 +903,10 @@ namespace UE::MovieSceneTools
 			}
 			
 			const bool bHasArriveTangent = !KeyAttributes.IsEmpty() && !KeyAttributes[0].HasArriveTangent();
-			const bool bFirstKeyIsLinearInterpMode = !KeyAttributes.IsEmpty() && KeyAttributes[0].GetInterpMode() == ERichCurveInterpMode::RCIM_Linear;
+			const bool bFirstKeyIsLinearInterpMode = 
+				!KeyAttributes.IsEmpty() && 
+				KeyAttributes[0].HasInterpMode() && 
+				KeyAttributes[0].GetInterpMode() == ERichCurveInterpMode::RCIM_Linear;
 
 			const bool bUseConstantForLinear =
 				FullRangeInterpolation.PreInfinityExtrapolation == ERichCurveExtrapolation::RCCE_Linear &&
@@ -997,7 +1000,10 @@ namespace UE::MovieSceneTools
 			}
 
 			const bool bHasLeaveTangent = !KeyAttributes.IsEmpty() && !KeyAttributes.Last().HasLeaveTangent();
-			const bool bLastKeyIsLinearInterpMode = !KeyAttributes.IsEmpty() && KeyAttributes.Last().GetInterpMode() == ERichCurveInterpMode::RCIM_Linear;
+			const bool bLastKeyIsLinearInterpMode = 
+				!KeyAttributes.IsEmpty() && 
+				KeyAttributes[0].HasInterpMode() && 
+				KeyAttributes.Last().GetInterpMode() == ERichCurveInterpMode::RCIM_Linear;
 
 			const bool bUseConstantForLinear =
 				FullRangeInterpolation.PostInfinityExtrapolation == ERichCurveExtrapolation::RCCE_Linear &&
