Curve Editor: Misc tests for new command-based undo system.

diff --git a/Engine/Source/Editor/CurveEditor/Public/CurveDataAbstraction.h b/Engine/Source/Editor/CurveEditor/Public/CurveDataAbstraction.h
--- a/Engine/Source/Editor/CurveEditor/Public/CurveDataAbstraction.h
+++ b/Engine/Source/Editor/CurveEditor/Public/CurveDataAbstraction.h
@@ -40,6 +40,9 @@ struct FKeyPosition
 	/** The key's output (t-axis) position (i.e. it's value) */
 	UPROPERTY()
 	double OutputValue;
+	
+	bool operator==(const FKeyPosition& Other) const = default;
+	bool operator!=(const FKeyPosition& Other) const = default;
 };
 
 /**
diff --git a/Engine/Source/Editor/MovieSceneTools/Private/TrackEditors/SubTrackEditor.cpp b/Engine/Source/Editor/MovieSceneTools/Private/TrackEditors/SubTrackEditor.cpp
--- a/Engine/Source/Editor/MovieSceneTools/Private/TrackEditors/SubTrackEditor.cpp
+++ b/Engine/Source/Editor/MovieSceneTools/Private/TrackEditors/SubTrackEditor.cpp
@@ -730,6 +730,12 @@ void FSubTrackEditor::OnRelease()
 	{
 		GLevelEditorModeTools().DeactivateMode(FSubTrackEditorMode::ModeName);
 	}
+
+	if (!GetSequencer())
+	{
+		return;
+	}
+	
 	GetSequencer()->GetViewModel()->GetSelection()->TrackArea.OnChanged.RemoveAll(this);
 
 	SectionsWithPreviews.Empty();
