Curve Editor: Add MarkAsChanged toFChannelCurveModel when removing keys. May need to revisit this to make sure all of the removed Modify() calls are workin.
#trivial

diff --git a/Engine/Source/Editor/MovieSceneTools/Private/Channels/ChannelCurveModel.cpp b/Engine/Source/Editor/MovieSceneTools/Private/Channels/ChannelCurveModel.cpp
--- a/Engine/Source/Editor/MovieSceneTools/Private/Channels/ChannelCurveModel.cpp
+++ b/Engine/Source/Editor/MovieSceneTools/Private/Channels/ChannelCurveModel.cpp
@@ -539,6 +539,13 @@ void FChannelCurveModel<ChannelType, ChannelValue, KeyType>::RemoveKeys(TArrayVi
 
 	if (Channel && Section && !IsReadOnly())
 	{
+		//Dominik should we put this into a function to replace Modify()? 
+		{
+			Section->MarkAsChanged();
+			LastSignature.Invalidate();
+		}
+		//
+
 		TMovieSceneChannelData<ChannelValue> ChannelData = Channel->GetData();		
 		double CurrentValue = 0.0;
 		Evaluate(InCurrentTime, CurrentValue);
