Movie Scene Tools: Call MarkAsChanged & MarkPackageDirty after changing curve using Curve Editor.
- This fixes a user not being prompted to save the asset after changing keys using curve editor, possibly losing changes.
- When you undo an action, this also clears the dirty flag correctly if the asset was not dirty before you made the change. For this purpose, FScopedTransaction was introduced to FScopedCurveChange.

diff --git a/Engine/Source/Editor/CurveEditor/Private/Modification/Utils/ScopedCurveChange.cpp b/Engine/Source/Editor/CurveEditor/Private/Modification/Utils/ScopedCurveChange.cpp
--- a/Engine/Source/Editor/CurveEditor/Private/Modification/Utils/ScopedCurveChange.cpp
+++ b/Engine/Source/Editor/CurveEditor/Private/Modification/Utils/ScopedCurveChange.cpp
@@ -27,6 +27,7 @@ FScopedCurveChange::FScopedCurveChange(
 	TWeakPtr<FCurveEditor> InCurveEditor, FCurveChangeDiff InKeyChangeDiff, const FText& InDescription, EScopedKeyChangeFlags InFlags
 	)
 	: FScopedChangeBase(MoveTemp(InCurveEditor), InDescription.IsEmpty() ? NSLOCTEXT("CurveEditor", "ChangeKeys", "Change Keys") : InDescription)
+	, Transaction(Description)
 	, ChangeDiff(MoveTemp(InKeyChangeDiff))
 	, Flags(InFlags)
 {}
@@ -49,9 +50,13 @@ FScopedCurveChange::~FScopedCurveChange()
 		return;
 	}
 
-	// Early out without expensive call to ComputeDiff.
 	const bool bRevertOnCancelled = EnumHasAnyFlags(Flags, EScopedKeyChangeFlags::RevertOnCancel);
 	const bool bHasBeenCancelled = IsCancelled();
+	if (bHasBeenCancelled)
+	{
+		Transaction.Cancel();
+	}
+	// Early out without expensive call to ComputeDiff.
 	if (bHasBeenCancelled && !bRevertOnCancelled)
 	{
 		return;
@@ -74,14 +79,13 @@ FScopedCurveChange::~FScopedCurveChange()
 		return;
 	}
 	
+	LogPerfData(); // Important to do this before MoveTemp
 	if (!bHasBeenCancelled)
 	{
-		LogPerfData(); // Important to do this before MoveTemp
 		TryAppendCommand(MakeUnique<FGenericCurveChangeCommand>(MoveTemp(DeltaChange)));
 	}
 	else
 	{
-		LogPerfData();
 		GenericCurveChange::RevertChange(*CurveEditorPin, DeltaChange);
 	}
 }
diff --git a/Engine/Source/Editor/CurveEditor/Public/Modification/Utils/ScopedCurveChange.h b/Engine/Source/Editor/CurveEditor/Public/Modification/Utils/ScopedCurveChange.h
--- a/Engine/Source/Editor/CurveEditor/Public/Modification/Utils/ScopedCurveChange.h
+++ b/Engine/Source/Editor/CurveEditor/Public/Modification/Utils/ScopedCurveChange.h
@@ -5,6 +5,7 @@
 #include "Misc/EnumClassFlags.h"
 #include "Modification/Keys/CurveChangeDiff.h"
 #include "ScopedChangeBase.h"
+#include "ScopedTransaction.h"
 #include "Templates/SharedPointer.h"
 
 struct FCurveModelID;
@@ -26,7 +27,7 @@ ENUM_CLASS_FLAGS(EScopedKeyChangeFlags);
 enum class ETrackedCurves : uint8 { All, Selected };
 	
 /**
- * Detects changes made to the curves specified on construction and adds it as an undo-able action.
+ * Starts a transaction. Detects changes made to the curves specified on construction and adds it as an undo-able action.
  * When the scope ends, checks whether keys were added, removed, moved, or its attributes (FKeyAttributes or FCurveAttributes) were edited.
  *
  * ========== Performance ==========
@@ -90,6 +91,14 @@ public:
 	CURVEEDITOR_API ~FScopedCurveChange();
 
 private:
+
+	/**
+	 * The transaction active during the change.
+	 * 
+	 * This is required because FCurveModel implementations mark the owning package dirty when keys are modified.
+	 * This achieves that undoing the change will unmark the package as dirty, if it was not dirty before the change.
+	 * */
+	FScopedTransaction Transaction;
 	
 	/** Checks for changes. */
 	const FCurveChangeDiff ChangeDiff;
diff --git a/Engine/Source/Editor/MovieSceneTools/Private/Channels/BezierChannelCurveModel.cpp b/Engine/Source/Editor/MovieSceneTools/Private/Channels/BezierChannelCurveModel.cpp
--- a/Engine/Source/Editor/MovieSceneTools/Private/Channels/BezierChannelCurveModel.cpp
+++ b/Engine/Source/Editor/MovieSceneTools/Private/Channels/BezierChannelCurveModel.cpp
@@ -503,9 +503,8 @@ void FBezierChannelCurveModel<ChannelType, ChannelValue, KeyType>::SetKeyAttribu
 		{
 			Channel->AutoSetTangents();
 		}
-
-		SignedOwner->MarkAsChanged();
-		this->CurveModifiedDelegate.Broadcast();
+		
+		OnPostOwnerChanged(SignedOwner);
 	}
 }
 
@@ -538,11 +537,19 @@ void FBezierChannelCurveModel<ChannelType, ChannelValue, KeyType>::SetCurveAttri
 			Channel->PostInfinityExtrap = InCurveAttributes.GetPostExtrapolation();
 		}
 
-		SignedOwner->MarkAsChanged();
-		this->CurveModifiedDelegate.Broadcast();
+		OnPostOwnerChanged(SignedOwner);
 	}
 }
 
+template <typename ChannelType, typename ChannelValue, typename KeyType>
+void FBezierChannelCurveModel<ChannelType, ChannelValue, KeyType>::OnPostOwnerChanged(TNotNull<UMovieSceneSignedObject*> InSignedOwner)
+{
+	InSignedOwner->MarkAsChanged();
+	InSignedOwner->MarkPackageDirty();
+	
+	this->CurveModifiedDelegate.Broadcast();
+}
+
 /*	 Finds min/max for cubic curves:
 Looks for feature points in the signal(determined by change in direction of local tangent), these locations are then re-examined in closer detail recursively
 Similar to function in RichCurve but usees the Channel::Evaluate function, instead of CurveModel::Eval*/
diff --git a/Engine/Source/Editor/MovieSceneTools/Private/Channels/ChannelCurveModel.cpp b/Engine/Source/Editor/MovieSceneTools/Private/Channels/ChannelCurveModel.cpp
--- a/Engine/Source/Editor/MovieSceneTools/Private/Channels/ChannelCurveModel.cpp
+++ b/Engine/Source/Editor/MovieSceneTools/Private/Channels/ChannelCurveModel.cpp
@@ -272,12 +272,6 @@ void FChannelCurveModel<ChannelType, ChannelValue, KeyType>::SetKeyPositions(TAr
 	UMovieSceneSection* Section = WeakSection.Get();
 	const FMovieSceneChannelMetaData* MetaData = ChannelHandle.GetMetaData();
 
-	UMovieSceneSignedObject* SignedOwner = Cast<UMovieSceneSignedObject>(WeakOwningObject.Get());
-	if (!SignedOwner)
-	{
-		SignedOwner = WeakSection.Get();
-	}
-
 	if (Channel && MetaData && Section && !IsReadOnly())
 	{
 		TMovieSceneChannelData<ChannelValue> ChannelData = Channel->GetData();
@@ -287,8 +281,6 @@ void FChannelCurveModel<ChannelType, ChannelValue, KeyType>::SetKeyPositions(TAr
 			return;
 		}
 
-		Section->MarkAsChanged();
-
 		FFrameRate   TickResolution = Section->GetTypedOuter<UMovieScene>()->GetTickResolution();
 		FFrameNumber KeyOffset      = 0;
 
@@ -333,8 +325,7 @@ void FChannelCurveModel<ChannelType, ChannelValue, KeyType>::SetKeyPositions(TAr
 			}
 		}
 		
-		InvokePostChangeCallbacks(Channel, Section);
-		SignedOwner->MarkAsChanged();
+		OnPostChannelChanged(Channel, Section);
 	}
 }
 
@@ -474,12 +465,6 @@ void FChannelCurveModel<ChannelType, ChannelValue, KeyType>::AddKeys(TArrayView<
 	UMovieSceneSection* Section = WeakSection.Get();
 	const FMovieSceneChannelMetaData* MetaData = ChannelHandle.GetMetaData();
 
-	UMovieSceneSignedObject* SignedOwner = Cast<UMovieSceneSignedObject>(WeakOwningObject.Get());
-	if (!SignedOwner)
-	{
-		SignedOwner = WeakSection.Get();
-	}
-
 	if (Channel && MetaData && Section && !IsReadOnly())
 	{
 		TMovieSceneChannelData<ChannelValue> ChannelData = Channel->GetData();
@@ -521,7 +506,7 @@ void FChannelCurveModel<ChannelType, ChannelValue, KeyType>::AddKeys(TArrayView<
 		// We reuse SetKeyAttributes here as there is complex logic determining which parts of the attributes are valid to set.
 		// For now we need to duplicate the new key handle array due to API mismatch. This will auto calculate tangents if needed.
 		SetKeyAttributes(NewKeyHandles, InKeyAttributes);
-		InvokePostChangeCallbacks(Channel, Section);
+		OnPostChannelChanged(Channel, Section);
 	}
 }
 
@@ -530,22 +515,8 @@ void FChannelCurveModel<ChannelType, ChannelValue, KeyType>::RemoveKeys(TArrayVi
 {
 	ChannelType* Channel = ChannelHandle.Get();
 	UMovieSceneSection* Section = WeakSection.Get();
-	UMovieSceneSignedObject* SignedOwner = Cast<UMovieSceneSignedObject>(WeakOwningObject.Get());
-	if (!SignedOwner)
-	{
-		SignedOwner = WeakSection.Get();
-	}
-
-
 	if (Channel && Section && !IsReadOnly())
 	{
-		//Dominik should we put this into a function to replace Modify()? 
-		{
-			Section->MarkAsChanged();
-			LastSignature.Invalidate();
-		}
-		//
-
 		TMovieSceneChannelData<ChannelValue> ChannelData = Channel->GetData();		
 		double CurrentValue = 0.0;
 		Evaluate(InCurrentTime, CurrentValue);
@@ -565,7 +536,7 @@ void FChannelCurveModel<ChannelType, ChannelValue, KeyType>::RemoveKeys(TArrayVi
 			Channel->SetDefault(ValueAsCurveValue);
 		}
 
-		InvokePostChangeCallbacks(Channel, Section);
+		OnPostChannelChanged(Channel, Section);
 	}
 }
 
@@ -598,7 +569,7 @@ void FChannelCurveModel<ChannelType, ChannelValue, KeyType>::ReplaceKeyHandles(
 			}
 
 			ChannelData.ReplaceKeyHandle(KeyIndex, InNewHandles[Index]);
-			InvokePostChangeCallbacks(Channel, Section);
+			OnPostChannelChanged(Channel, Section);
 		}
 	}
 }
@@ -661,10 +632,11 @@ void FChannelCurveModel<ChannelType, ChannelValue, KeyType>::GetCurveColorObject
 }
 
 template <class ChannelType, class ChannelValue, class KeyType>
-void FChannelCurveModel<ChannelType, ChannelValue, KeyType>::InvokePostChangeCallbacks(ChannelType* Channel, UMovieSceneSection* Section)
+void FChannelCurveModel<ChannelType, ChannelValue, KeyType>::OnPostChannelChanged(TNotNull<ChannelType*> Channel, TNotNull<UMovieSceneSection*> Section)
 {
 	LastSignature.Invalidate();
 	
+	// Allow channel to make any follow up changes to the keys before executing any other delegates
 	Channel->PostEditChange();
 	
 	if (WeakSequencer.IsValid())
@@ -672,11 +644,19 @@ void FChannelCurveModel<ChannelType, ChannelValue, KeyType>::InvokePostChangeCal
 		const FMovieSceneChannelMetaData* MetaData = ChannelHandle.GetMetaData();
 		WeakSequencer.Pin()->OnChannelChanged().Broadcast(MetaData, Section);
 	}
+
+	Section->MarkAsChanged();
+	Section->MarkPackageDirty();
+	
+	if (UMovieSceneSignedObject* SignedOwner = Cast<UMovieSceneSignedObject>(WeakOwningObject.Get()))
+	{
+		SignedOwner->MarkAsChanged();
+		SignedOwner->MarkPackageDirty();
+	}
 	
 	CurveModifiedDelegate.Broadcast();
 }
 
-
 // Explicit template instantiation
 template class FChannelCurveModel<FMovieSceneDoubleChannel, FMovieSceneDoubleValue, double>;
 template class FChannelCurveModel<FMovieSceneFloatChannel, FMovieSceneFloatValue, float>;
diff --git a/Engine/Source/Editor/MovieSceneTools/Public/Channels/BezierChannelCurveModel.h b/Engine/Source/Editor/MovieSceneTools/Public/Channels/BezierChannelCurveModel.h
--- a/Engine/Source/Editor/MovieSceneTools/Public/Channels/BezierChannelCurveModel.h
+++ b/Engine/Source/Editor/MovieSceneTools/Public/Channels/BezierChannelCurveModel.h
@@ -13,6 +13,7 @@
 #include "CurveEditorTypes.h"
 #include "Curves/KeyHandle.h"
 #include "IBufferedCurveModel.h"
+#include "Misc/NotNull.h"
 #include "MovieSceneSection.h"
 #include "Templates/SharedPointer.h"
 #include "Templates/Tuple.h"
@@ -85,6 +86,12 @@ protected:
 
 	TSharedPtr<UE::MovieSceneTools::FMovieSceneCachedCurve<ChannelType>> CachedCurve;
 
+	/**
+	 * Updates all dependent systems after a change has been made to owning object.
+	 * Calls MarkAsChanged, dirties the package, and FCurveModel::OnCurvesModified.
+	 */
+	void OnPostOwnerChanged(TNotNull<UMovieSceneSignedObject*> InSignedOwner);
+
 private:
 	void FeaturePointMethod(double StartTime, double EndTime, double StartValue, double Mu, int Depth, int MaxDepth, double& MaxV, double& MinVal) const;
 };
diff --git a/Engine/Source/Editor/MovieSceneTools/Public/Channels/ChannelCurveModel.h b/Engine/Source/Editor/MovieSceneTools/Public/Channels/ChannelCurveModel.h
--- a/Engine/Source/Editor/MovieSceneTools/Public/Channels/ChannelCurveModel.h
+++ b/Engine/Source/Editor/MovieSceneTools/Public/Channels/ChannelCurveModel.h
@@ -15,6 +15,7 @@
 #include "Curves/KeyHandle.h"
 #include "Delegates/IDelegateInstance.h"
 #include "IBufferedCurveModel.h"
+#include "Misc/NotNull.h"
 #include "Misc/OptionalFwd.h"
 #include "MovieSceneSection.h"
 #include "Templates/SharedPointer.h"
@@ -105,8 +106,17 @@ protected:
 	virtual double GetKeyValue(TArrayView<const ChannelValue> Values, int32 Index) const = 0;
 	virtual void SetKeyValue(int32 Index, double KeyValue) const = 0;
 
-	/** Calls PostEditChange, OnCurvesModified, notifies Sequencer's OnChannelChanged(), and ensures that LastSignature is invalidated. */
-	virtual void InvokePostChangeCallbacks(ChannelType* Channel, UMovieSceneSection* Section);
+	/**
+	 * Updates all dependent systems after a change has been made to the channel, namely:
+	 *	- Invalidate LastSignature,
+	 *	- ChannelType::PostEditChange,
+	 *	- ChannelType::OnChannelChanged(),
+	 *	- UMovieSceneSection::MarkAsChanged,
+	 *	- MarkPackageDirty: this is required so the user is prompted to save the package after editing keys.
+	 *	If currently recording a transaction, the generated transaction will undo this flag when the user performs an undo op.
+	 *	- FCurveModel::OnCurvesModified
+	 */
+	void OnPostChannelChanged(TNotNull<ChannelType*> Channel, TNotNull<UMovieSceneSection*> Section);
 
 public:
 
