Constraints: fix compensation & deletion when updating a key

#rnx

diff --git a/Engine/Plugins/Animation/ControlRig/Source/ControlRigEditor/Private/Sequencer/ControlConstraintChannelInterface.cpp b/Engine/Plugins/Animation/ControlRig/Source/ControlRigEditor/Private/Sequencer/ControlConstraintChannelInterface.cpp
--- a/Engine/Plugins/Animation/ControlRig/Source/ControlRigEditor/Private/Sequencer/ControlConstraintChannelInterface.cpp
+++ b/Engine/Plugins/Animation/ControlRig/Source/ControlRigEditor/Private/Sequencer/ControlConstraintChannelInterface.cpp
@@ -194,7 +194,7 @@ bool FControlConstraintChannelInterface::SmartConstraintKey(
 			// add active key
 			{
 				TMovieSceneChannelData<bool> ChannelData = Channel->ActiveChannel.GetData();
-				ChannelData.AddKey(InTime, ActiveValueToBeSet);
+				ChannelData.UpdateOrAddKey(InTime, ActiveValueToBeSet);
 			}
 
 			// compensate
@@ -216,7 +216,7 @@ bool FControlConstraintChannelInterface::SmartConstraintKey(
 			// cleanup next keys that have the same value as they are useless 
 			{
 				TArray<FFrameNumber> TimesRemoved;
-				CleanDuplicates(Channel->ActiveChannel, InTime, ActiveValueToBeSet, TimesRemoved);
+				CleanDuplicates(Channel->ActiveChannel, TimesRemoved);
 
 				// as compensation has been disabled earlier (FMovieSceneConstraintChannelHelper::bDoNotCompensate set to true),
 				// previous transform compensation keys must be explicitly removed as HandleConstraintKeyDeleted won't do it
diff --git a/Engine/Source/Editor/MovieSceneTools/Private/Constraints/ComponentConstraintChannelInterface.cpp b/Engine/Source/Editor/MovieSceneTools/Private/Constraints/ComponentConstraintChannelInterface.cpp
--- a/Engine/Source/Editor/MovieSceneTools/Private/Constraints/ComponentConstraintChannelInterface.cpp
+++ b/Engine/Source/Editor/MovieSceneTools/Private/Constraints/ComponentConstraintChannelInterface.cpp
@@ -172,7 +172,7 @@ bool FComponentConstraintChannelInterface::SmartConstraintKey(
 			// add active key
 			{
 				TMovieSceneChannelData<bool> ChannelData = Channel->ActiveChannel.GetData();
-				ChannelData.AddKey(InTime, ActiveValueToBeSet);
+				ChannelData.UpdateOrAddKey(InTime, ActiveValueToBeSet);
 			}
 
 			// compensate
@@ -198,7 +198,7 @@ bool FComponentConstraintChannelInterface::SmartConstraintKey(
 			// cleanup next keys that have the same value as they are useless 
 			{
 				TArray<FFrameNumber> TimesRemoved;
-				CleanDuplicates(Channel->ActiveChannel, InTime, ActiveValueToBeSet, TimesRemoved);
+				CleanDuplicates(Channel->ActiveChannel, TimesRemoved);
 
 				// as compensation has been disabled earlier (FMovieSceneConstraintChannelHelper::bDoNotCompensate set to true),
 				// previous transform compensation keys must be explicitly removed as HandleConstraintKeyDeleted won't do it
diff --git a/Engine/Source/Editor/MovieSceneTools/Private/Constraints/TransformConstraintChannelInterface.cpp b/Engine/Source/Editor/MovieSceneTools/Private/Constraints/TransformConstraintChannelInterface.cpp
--- a/Engine/Source/Editor/MovieSceneTools/Private/Constraints/TransformConstraintChannelInterface.cpp
+++ b/Engine/Source/Editor/MovieSceneTools/Private/Constraints/TransformConstraintChannelInterface.cpp
@@ -71,31 +71,37 @@ bool ITransformConstraintChannelInterface::CanAddKey(const FMovieSceneConstraint
 	return true;
 }
 
-void ITransformConstraintChannelInterface::CleanDuplicates(FMovieSceneConstraintChannel& InOutActiveChannel,
-	const FFrameNumber& InTime, const bool InValueToClean, TArray<FFrameNumber>& OutTimesRemoved)
+void ITransformConstraintChannelInterface::CleanDuplicates(FMovieSceneConstraintChannel& InOutActiveChannel, TArray<FFrameNumber>& OutTimesRemoved)
 {
 	TMovieSceneChannelData<bool> ChannelData = InOutActiveChannel.GetData();
 	const TArrayView<const FFrameNumber> Times = ChannelData.GetTimes();
-	if (Times.IsEmpty())
+	const TArrayView<const bool> ActiveValues = ChannelData.GetValues();
+
+	const int32 NumTimes = Times.Num();
+	if (NumTimes == 0 || ActiveValues.Num() != NumTimes)
 	{
 		return;
 	}
-	
-	const int32 Index = Times.Find(InTime);
-	const int32 NextTimeIndex = Algo::UpperBound(Times, InTime);
-	
-	if (NextTimeIndex != Index && Times.IsValidIndex(NextTimeIndex))
+
+	if (NumTimes > 1)
 	{
-		bool NextValue = false;
-		InOutActiveChannel.Evaluate(Times[NextTimeIndex], NextValue);
-		if (NextValue == InValueToClean)
+		for (int32 Index = NumTimes-1; Index >= 1; --Index)
 		{
-			OutTimesRemoved.Add(Times[NextTimeIndex]);
-
-			//same value so delete
-			ChannelData.RemoveKey(NextTimeIndex);
+			// no need to keep consecutive values if they are similar
+			if (ActiveValues[Index] == ActiveValues[Index-1])
+			{
+				OutTimesRemoved.Add(Times[Index]);
+				ChannelData.RemoveKey(Index);
+			}
 		}
 	}
+
+	// no need to keep single inactive key
+	if (!ActiveValues[0])
+	{
+		OutTimesRemoved.Add(Times[0]);
+		ChannelData.RemoveKey(0);
+	}
 }
 
 void ITransformConstraintChannelInterface::OnSequencerClosed(TSharedRef<ISequencer> InSequencer)
diff --git a/Engine/Source/Editor/MovieSceneTools/Public/Constraints/TransformConstraintChannelInterface.h b/Engine/Source/Editor/MovieSceneTools/Public/Constraints/TransformConstraintChannelInterface.h
--- a/Engine/Source/Editor/MovieSceneTools/Public/Constraints/TransformConstraintChannelInterface.h
+++ b/Engine/Source/Editor/MovieSceneTools/Public/Constraints/TransformConstraintChannelInterface.h
@@ -54,8 +54,7 @@ protected:
 	MOVIESCENETOOLS_API ITransformConstraintChannelInterface();
 
 	static MOVIESCENETOOLS_API bool CanAddKey(const FMovieSceneConstraintChannel& InActiveChannel, const FFrameNumber& InTime, bool& ActiveValue);
-	static MOVIESCENETOOLS_API void CleanDuplicates(FMovieSceneConstraintChannel& InOutActiveChannel,
-		const FFrameNumber& InTime, const bool InValueToClean, TArray<FFrameNumber>& OutTimesRemoved);
+	static MOVIESCENETOOLS_API void CleanDuplicates(FMovieSceneConstraintChannel& InOutActiveChannel, TArray<FFrameNumber>& OutTimesRemoved);
 
 	void Initialize();
 	void Shutdown();
