Sequencer: Move struct property change propagation from OnFinishedChangingProperties to NotifyPostChange so that it can run before PostTransacted

Some systems like USD, rely on PostTransacted having the changed property value, so this ensures that the property change is propagated before that.

Tested USD, changes to generic properties and combined struct properties (ie. Location, Rotation, Scale)

diff --git a/Engine/Source/Editor/MovieSceneTools/Private/SKeyEditInterface.cpp b/Engine/Source/Editor/MovieSceneTools/Private/SKeyEditInterface.cpp
--- a/Engine/Source/Editor/MovieSceneTools/Private/SKeyEditInterface.cpp
+++ b/Engine/Source/Editor/MovieSceneTools/Private/SKeyEditInterface.cpp
@@ -16,6 +16,7 @@
 #include "Misc/FrameNumber.h"
 #include "Modules/ModuleManager.h"
 #include "IDetailCustomization.h"
+#include "IStructureDataProvider.h"
 #include "MovieSceneToolsModule.h"
 
 SKeyEditInterface::~SKeyEditInterface()
@@ -74,7 +75,7 @@ void SKeyEditInterface::Initialize()
 		StructureViewArgs.bShowInterfaces = false;
 	}
 
-	TSharedRef<IStructureDetailsView> StructureDetailsView = FModuleManager::GetModuleChecked<FPropertyEditorModule>("PropertyEditor")
+	StructureDetailsView = FModuleManager::GetModuleChecked<FPropertyEditorModule>("PropertyEditor")
 		.CreateStructureDetailView(DetailsViewArgs, StructureViewArgs, nullptr);
 
 	// register details customizations for this instance
@@ -82,7 +83,7 @@ void SKeyEditInterface::Initialize()
 	StructureDetailsView->GetDetailsView()->RegisterInstancedCustomPropertyTypeLayout("MovieSceneEvent", FOnGetPropertyTypeCustomizationInstance::CreateSP(this, &SKeyEditInterface::CreateEventCustomization));
 	StructureDetailsView->GetDetailsView()->RegisterInstancedCustomPropertyTypeLayout("FrameNumber", FOnGetPropertyTypeCustomizationInstance::CreateSP(this, &SKeyEditInterface::CreateFrameNumberCustomization));
 
-	FMovieSceneToolsModule::Get().CustomizeKeyStructInstancedPropertyTypes(StructureDetailsView, WeakSection);
+	FMovieSceneToolsModule::Get().CustomizeKeyStructInstancedPropertyTypes(StructureDetailsView.ToSharedRef(), WeakSection);
 
 	StructureDetailsView->SetStructureData(NewEditData.KeyStruct);
 	StructureDetailsView->GetOnFinishedChangingPropertiesDelegate().AddSP(this, &SKeyEditInterface::OnFinishedChangingProperties, NewEditData.KeyStruct);
@@ -121,29 +122,6 @@ TSharedRef<IPropertyTypeCustomization> SKeyEditInterface::CreateEventCustomizati
 
 void SKeyEditInterface::OnFinishedChangingProperties(const FPropertyChangedEvent& ChangeEvent, TSharedPtr<FStructOnScope> KeyStruct)
 {
-	if (KeyStruct->GetStruct()->IsChildOf(FMovieSceneKeyStruct::StaticStruct()))
-	{
-		if (UMovieSceneSection* Section = WeakSection.Get())
-		{
-			Section->Modify();
-		}
-
-		((FMovieSceneKeyStruct*)KeyStruct->GetStructMemory())->PropagateChanges(ChangeEvent);
-	}
-	else if (KeyStruct->GetStruct()->IsChildOf(FGeneratedMovieSceneKeyStruct::StaticStruct()))
-	{
-		if (UMovieSceneSection* Section = WeakSection.Get())
-		{
-			Section->Modify();
-		}
-
-		FGeneratedMovieSceneKeyStruct* GeneratedStruct = reinterpret_cast<FGeneratedMovieSceneKeyStruct*>(KeyStruct->GetStructMemory());
-		if (GeneratedStruct->OnPropertyChangedEvent)
-		{
-			GeneratedStruct->OnPropertyChangedEvent(ChangeEvent);
-		}
-	}
-
 	TSharedPtr<ISequencer> Sequencer = WeakSequencer.Pin();
 	if (Sequencer.IsValid())
 	{
@@ -159,6 +137,54 @@ void SKeyEditInterface::NotifyPreChange( FProperty* PropertyAboutToChange )
 	}
 }
 
+void SKeyEditInterface::NotifyPostChange(const FPropertyChangedEvent& ChangeEvent, class FEditPropertyChain* PropertyThatChanged)
+{
+	TSharedPtr<const class IStructureDataProvider> StructureProvider = StructureDetailsView->GetStructureProvider();
+	if (!StructureProvider)
+	{
+		return;
+	}
+
+	TArray<TSharedPtr<FStructOnScope>> KeyStructInstances;
+	StructureProvider->GetInstances(KeyStructInstances, FMovieSceneKeyStruct::StaticStruct());
+
+	for (const TSharedPtr<FStructOnScope>& KeyStruct : KeyStructInstances)
+	{
+		if (KeyStruct.IsValid() && KeyStruct->IsValid())
+		{
+			if (KeyStruct->GetStruct()->IsChildOf(FMovieSceneKeyStruct::StaticStruct()))
+			{
+				if (UMovieSceneSection* Section = WeakSection.Get())
+				{
+					Section->Modify();
+				}
+
+				((FMovieSceneKeyStruct*)KeyStruct->GetStructMemory())->PropagateChanges(ChangeEvent);
+			}
+		}
+	}
+
+	TArray<TSharedPtr<FStructOnScope>> GeneratedKeyStructInstances;
+	StructureProvider->GetInstances(GeneratedKeyStructInstances, FGeneratedMovieSceneKeyStruct::StaticStruct());
+
+	for (const TSharedPtr<FStructOnScope>& KeyStruct : GeneratedKeyStructInstances)
+	{
+		if (KeyStruct.IsValid() && KeyStruct->IsValid())
+		{
+			if (UMovieSceneSection* Section = WeakSection.Get())
+			{
+				Section->Modify();
+			}
+
+			FGeneratedMovieSceneKeyStruct* GeneratedStruct = reinterpret_cast<FGeneratedMovieSceneKeyStruct*>(KeyStruct->GetStructMemory());
+			if (GeneratedStruct->OnPropertyChangedEvent)
+			{
+				GeneratedStruct->OnPropertyChangedEvent(ChangeEvent);
+			}
+		}
+	}
+}
+
 void SKeyEditInterface::PostUndo(bool bSuccess)
 {
 	TSharedPtr<ISequencer> Sequencer = WeakSequencer.Pin();
diff --git a/Engine/Source/Editor/MovieSceneTools/Public/SKeyEditInterface.h b/Engine/Source/Editor/MovieSceneTools/Public/SKeyEditInterface.h
--- a/Engine/Source/Editor/MovieSceneTools/Public/SKeyEditInterface.h
+++ b/Engine/Source/Editor/MovieSceneTools/Public/SKeyEditInterface.h
@@ -16,6 +16,7 @@ class ISequencer;
 class FStructOnScope;
 class UMovieSceneSection;
 class IPropertyTypeCustomization;
+class IStructureDetailsView;
 
 struct FKeyEditData
 {
@@ -61,6 +62,7 @@ private:
 	UE_API virtual FVector2D ComputeDesiredSize(float LayoutScaleMultiplier) const override;
 
 	UE_API virtual void NotifyPreChange(FProperty* PropertyAboutToChange) override;
+	UE_API virtual void NotifyPostChange(const FPropertyChangedEvent& PropertyChangedEvent, class FEditPropertyChain* PropertyThatChanged) override;
 	UE_API virtual void PostUndo(bool bSuccess) override;
 	UE_API virtual void PostRedo(bool bSuccess) override;
 
@@ -68,6 +70,7 @@ private:
 	TAttribute<FKeyEditData> EditDataAttribute;
 	TWeakObjectPtr<UMovieSceneSection> WeakSection;
 	TWeakPtr<ISequencer> WeakSequencer;
+	TSharedPtr<IStructureDetailsView> StructureDetailsView;
 };
 
 #undef UE_API
