Curve Editor: Fix issues where bezier curves were not always painted correctly

* Fix an issue where curves were not painted when none of its keys were visible
* Fix an issue where linear pre-/post-infinity extrapolation was not always painted correctly

diff --git a/Engine/Source/Editor/MovieSceneTools/Private/Cache/MovieSceneCachedCurve.h b/Engine/Source/Editor/MovieSceneTools/Private/Cache/MovieSceneCachedCurve.h
--- a/Engine/Source/Editor/MovieSceneTools/Private/Cache/MovieSceneCachedCurve.h
+++ b/Engine/Source/Editor/MovieSceneTools/Private/Cache/MovieSceneCachedCurve.h
@@ -893,25 +893,33 @@ namespace UE::MovieSceneTools
 		}
 		else
 		{
-			const bool bPreinfinityVisible = Times[0].Value > ScreenSpace.GetInputMin() * TickResolution;
+			const double FirstKeyTime = Times[0] / TickResolution;
+			const bool bPreinfinityVisible = FirstKeyTime > ScreenSpace.GetInputMin();
 			if (!bPreinfinityVisible)
 			{
 				return true;
 			}
+
+			const bool bOnlyPreInfinityVisible = FirstKeyTime > ScreenSpace.GetInputMax();
+			if (bOnlyPreInfinityVisible)
+			{
+				// Handle cases where the last key is not visible by adding it to draw params.
+				// This will form a line between the first key and the pre-infinity extent.
+				const double FirstKeyValue = Sign * Values[0].Value;
+				CachedDrawParams.InterpolatingPoints.Emplace(
+					ScreenSpace.SecondsToScreen(FirstKeyTime),
+					ScreenSpace.ValueToScreen(FirstKeyValue));
+			}
 			
-			const bool bHasArriveTangent = !KeyAttributes.IsEmpty() && !KeyAttributes[0].HasArriveTangent();
-			const bool bFirstKeyIsLinearInterpMode = 
+			const bool bFirstKeyCanUseLinearInterpMode = 
 				!KeyAttributes.IsEmpty() && 
 				KeyAttributes[0].HasInterpMode() && 
-				KeyAttributes[0].GetInterpMode() == ERichCurveInterpMode::RCIM_Linear;
-
-			const bool bUseConstantForLinear =
-				FullRangeInterpolation.PreInfinityExtrapolation == ERichCurveExtrapolation::RCCE_Linear &&
-				(!bHasArriveTangent || !bFirstKeyIsLinearInterpMode);
+				KeyAttributes[0].GetInterpMode() != ERichCurveInterpMode::RCIM_Constant &&
+				KeyAttributes[0].GetInterpMode() != ERichCurveInterpMode::RCIM_None;
 
 			if ((FullRangeInterpolation.PreInfinityExtrapolation == ERichCurveExtrapolation::RCCE_None ||
 				FullRangeInterpolation.PreInfinityExtrapolation == ERichCurveExtrapolation::RCCE_Constant ||
-				bUseConstantForLinear))
+				!bFirstKeyCanUseLinearInterpMode))
 			{
 				const double ExtentValue = Sign * Values[0].Value;
 				CachedDrawParams.InterpolatingPoints.Emplace(
@@ -924,7 +932,8 @@ namespace UE::MovieSceneTools
 			{
 				const double MinTime = ScreenSpace.GetInputMin();
 
-				if (KeyAttributes[0].GetInterpMode() == ERichCurveInterpMode::RCIM_Linear)
+				if (KeyAttributes[0].HasInterpMode() &&
+					KeyAttributes[0].GetInterpMode() == ERichCurveInterpMode::RCIM_Linear)
 				{
 					// Draw a line in direction of the first two keys
 					const double X1 = Times[1] / TickResolution;
@@ -990,25 +999,33 @@ namespace UE::MovieSceneTools
 		}
 		else
 		{
-			const bool bPostinfinityVisible = Times.Last().Value < ScreenSpace.GetInputMax() * TickResolution;
+			const double LastKeyTime = Times.Last() / TickResolution;
+			const bool bPostinfinityVisible = LastKeyTime < ScreenSpace.GetInputMax();
 			if (!bPostinfinityVisible)
 			{
 				return true;
 			}
 
-			const bool bHasLeaveTangent = !KeyAttributes.IsEmpty() && !KeyAttributes.Last().HasLeaveTangent();
-			const bool bLastKeyIsLinearInterpMode = 
+			const bool bOnlyPostInfinityVisible = LastKeyTime < ScreenSpace.GetInputMin();
+			if (bOnlyPostInfinityVisible)
+			{
+				// Handle cases where the last key is not visible by adding it to draw params.
+				// This will form a line between the last key and the post-infinity extent.
+				const double LastKeyValue = Sign * Values.Last().Value;
+				CachedDrawParams.InterpolatingPoints.Emplace(
+					ScreenSpace.SecondsToScreen(LastKeyTime),
+					ScreenSpace.ValueToScreen(LastKeyValue));
+			}
+
+			const bool bLastKeyCanUseLinearInterpMode = 
 				!KeyAttributes.IsEmpty() && 
 				KeyAttributes[0].HasInterpMode() && 
-				KeyAttributes.Last().GetInterpMode() == ERichCurveInterpMode::RCIM_Linear;
-
-			const bool bUseConstantForLinear =
-				FullRangeInterpolation.PostInfinityExtrapolation == ERichCurveExtrapolation::RCCE_Linear &&
-				(!bHasLeaveTangent || !bLastKeyIsLinearInterpMode);
+				KeyAttributes.Last().GetInterpMode() != ERichCurveInterpMode::RCIM_Constant &&
+				KeyAttributes.Last().GetInterpMode() != ERichCurveInterpMode::RCIM_None;
 
 			if (FullRangeInterpolation.PostInfinityExtrapolation == ERichCurveExtrapolation::RCCE_None ||
 				FullRangeInterpolation.PostInfinityExtrapolation == ERichCurveExtrapolation::RCCE_Constant ||
-				bUseConstantForLinear)
+				!bLastKeyCanUseLinearInterpMode)
 			{
 				const double ExtentValue = Sign * Values.Last().Value;
 				CachedDrawParams.InterpolatingPoints.Emplace(
@@ -1021,7 +1038,8 @@ namespace UE::MovieSceneTools
 			{
 				const double MaxTime = ScreenSpace.GetInputMax();
 
-				if (KeyAttributes.Last().GetInterpMode() == ERichCurveInterpMode::RCIM_Linear)
+				if (KeyAttributes.Last().HasInterpMode() &&
+					KeyAttributes.Last().GetInterpMode() == ERichCurveInterpMode::RCIM_Linear)
 				{					
 					// Draw a line in direction of the last two keys
 					const int32 NumKeys = Times.Num();
