Sequencer: Attach sections with Preserve None now default to infinite.

If there is an existing section on the track at the insert time, the existing section will be clamped to the start time of the new attach section.
If there is an existing section after the insert time, the new section will be clamped to the existing section's start time.

diff --git a/Engine/Source/Editor/MovieSceneTools/Private/TrackEditors/AttachTrackEditor.cpp b/Engine/Source/Editor/MovieSceneTools/Private/TrackEditors/AttachTrackEditor.cpp
--- a/Engine/Source/Editor/MovieSceneTools/Private/TrackEditors/AttachTrackEditor.cpp
+++ b/Engine/Source/Editor/MovieSceneTools/Private/TrackEditors/AttachTrackEditor.cpp
@@ -1223,6 +1223,7 @@ FKeyPropertyResult F3DAttachTrackEditor::AddKeyInternal( FFrameNumber KeyTime, c
 
 	FMovieSceneSequenceID CurrentSequenceID = GetSequencer()->GetFocusedTemplateID();
 	UMovieScene* MovieScene = GetSequencer()->GetFocusedMovieSceneSequence()->GetMovieScene();
+	TRange<FFrameNumber> PlaybackRange = MovieScene->GetPlaybackRange();
 
 	// It's possible that the objects bound to this parent binding ID are null, in which case there will be no compensation
 	UObject* ParentObject = GetConstraintObject(GetSequencer(), ConstraintBindingID);	
@@ -1270,7 +1271,7 @@ FKeyPropertyResult F3DAttachTrackEditor::AddKeyInternal( FFrameNumber KeyTime, c
 		}
 
 		// Clamp to next attach section's start time or the end of the current movie scene range
-		FFrameNumber AttachEndTime = MovieScene->GetPlaybackRange().GetUpperBoundValue();
+		FFrameNumber AttachEndTime = PlaybackRange.GetUpperBoundValue();
 		for (UMovieSceneSection* Section : Track->GetAllSections())
 		{
 			FFrameNumber StartTime = Section->HasStartFrame() ? Section->GetInclusiveStartFrame() : 0;
@@ -1289,8 +1290,36 @@ FKeyPropertyResult F3DAttachTrackEditor::AddKeyInternal( FFrameNumber KeyTime, c
 		if (PreserveType == ETransformPreserveType::None)
 		{
 			Track->Modify();
+
+			UMovieSceneSection* ExistingSection = MovieSceneHelpers::FindSectionAtTime(Track->GetAllSections(), KeyTime);
+			UMovieSceneSection* NextSection = MovieSceneHelpers::FindNextSection(Track->GetAllSections(), KeyTime);
+
+			// If there's an existing section, clamp it to the start time of the new section
+			if (ExistingSection)
+			{
+				ExistingSection->Modify();
+				ExistingSection->SetEndFrame(TRangeBound<FFrameNumber>::Exclusive(KeyTime));
+			}
+			
 			UMovieSceneSection* NewSection = Cast<UMovieScene3DAttachTrack>(Track)->AddConstraint(KeyTime, Duration, SocketName, ComponentName, ConstraintBindingID);
 			
+			// If there isn't an existing section, make the start frame of the new section infinite
+			if (!ExistingSection)
+			{
+				NewSection->SectionRange.Value.SetLowerBound(TRangeBound<FFrameNumber>::Open());
+			}
+
+			// If there's a next section, make the end frame of the new section the start frame of the next section
+			if (NextSection)
+			{
+				NewSection->SectionRange.Value.SetUpperBound(NextSection->GetInclusiveStartFrame());
+			}
+			// Otherwise, make the end frame of the new section infinite
+			else
+			{
+				NewSection->SectionRange.Value.SetUpperBound(TRangeBound<FFrameNumber>::Open());
+			}
+
 			KeyPropertyResult.bTrackModified = true;
 			KeyPropertyResult.SectionsCreated.Add(NewSection);
 
@@ -1370,10 +1399,13 @@ FKeyPropertyResult F3DAttachTrackEditor::AddKeyInternal( FFrameNumber KeyTime, c
 		if (IntersectingSection)
 		{
 			// Calculate range to revert
-			TRange<FFrameNumber> RevertRange = TRange<FFrameNumber>(KeyTime, FMath::Min(AttachEndTime, IntersectingSection.GetValue()->GetExclusiveEndFrame()));
+			FFrameNumber IntersectionSectionStart = IntersectingSection.GetValue()->HasStartFrame() ? IntersectingSection.GetValue()->GetInclusiveStartFrame() : PlaybackRange.GetLowerBoundValue();
+			FFrameNumber IntersectionSectionEnd = IntersectingSection.GetValue()->HasEndFrame() ? IntersectingSection.GetValue()->GetExclusiveEndFrame() : PlaybackRange.GetUpperBoundValue();
+
+			TRange<FFrameNumber> RevertRange = TRange<FFrameNumber>(KeyTime, FMath::Min(AttachEndTime, IntersectionSectionEnd));
 
 			// If the intersecting section starts at the same time as the new section, remove it
-			if (IntersectingSection.GetValue()->GetInclusiveStartFrame() == KeyTime)
+			if (IntersectionSectionStart == KeyTime)
 			{
 				Track->RemoveSection(*IntersectingSection.GetValue());
 			}
@@ -1384,7 +1416,8 @@ FKeyPropertyResult F3DAttachTrackEditor::AddKeyInternal( FFrameNumber KeyTime, c
 				{
 					continue;
 				}
-				IntersectingSection.GetValue()->SetEndFrame(KeyTime - 1);
+				
+				IntersectingSection.GetValue()->SetEndFrame(TRangeBound<FFrameNumber>::Exclusive(KeyTime));
 			}
 
 			UMovieScene3DAttachSection* IntersectingAttachSection = Cast<UMovieScene3DAttachSection>(IntersectingSection.GetValue());
