Movie Scene: migrate text track to movie scene tracks

diff --git a/Engine/Config/BaseEngine.ini b/Engine/Config/BaseEngine.ini
--- a/Engine/Config/BaseEngine.ini
+++ b/Engine/Config/BaseEngine.ini
@@ -1364,6 +1364,14 @@ ClassName=/Script/OnlineSubsystemUtils.OnlineEngineInterfaceImpl
 ; 5.7
 +PropertyRedirects=(OldName="MeshNaniteSettings.bPreserveArea",NewName="ShapePreservation")
 
+; Movie Scene Text Track move to Movie Scene Tracks
++StructRedirects=(OldName="/Script/MovieSceneTextTrack.MovieSceneTextChannel",NewName="/Script/MovieSceneTracks.MovieSceneTextChannel")
++ClassRedirects=(OldName="/Script/MovieSceneTextTrack.MovieSceneTextPropertySystem",NewName="/Script/MovieSceneTracks.MovieSceneTextPropertySystem")
++ClassRedirects=(OldName="/Script/MovieSceneTextTrack.MovieSceneTextSection",NewName="/Script/MovieSceneTracks.MovieSceneTextSection")
++ClassRedirects=(OldName="/Script/MovieSceneTextTrack.MovieSceneTextTrack",NewName="/Script/MovieSceneTracks.MovieSceneTextTrack")
++ClassRedirects=(OldName="/Script/MovieSceneTextTrack.TextChannelEvaluatorSystem",NewName="/Script/MovieSceneTracks.TextChannelEvaluatorSystem")
++ClassRedirects=(OldName="/Script/MovieSceneTextTrackEditor.MovieSceneTextKeyStruct",NewName="/Script/MovieSceneTools.MovieSceneTextKeyStruct")
+
 ; End of releases redirects
 
 ; Fixing the output of tags for blueprints to be reset.
diff --git a/Engine/Plugins/MovieScene/MovieSceneTextTrack/MovieSceneTextTrack.uplugin b/Engine/Plugins/MovieScene/MovieSceneTextTrack/MovieSceneTextTrack.uplugin
--- a/Engine/Plugins/MovieScene/MovieSceneTextTrack/MovieSceneTextTrack.uplugin
+++ b/Engine/Plugins/MovieScene/MovieSceneTextTrack/MovieSceneTextTrack.uplugin
@@ -3,7 +3,7 @@
 	"Version": 1,
 	"VersionName": "1.0",
 	"FriendlyName": "MovieSceneTextTrack",
-	"Description": "Adds support to key Text Properties in Movie Scene Sequences",
+	"Description": "Deprecated plugin. Text support moved to Movie Scene Tracks (built-in).",
 	"Category": "Text",
 	"CreatedBy": "Epic Games, Inc.",
 	"CreatedByURL": "https://epicgames.com",
@@ -12,7 +12,7 @@
 	"SupportURL": "",
 	"CanContainContent": false,
 	"IsBetaVersion": false,
-	"IsExperimentalVersion": false,
+	"IsExperimentalVersion": true,
 	"Installed": false,
 	"Modules": [
 		{
diff --git a/Engine/Plugins/MovieScene/MovieSceneTextTrack/Source/MovieSceneTextTrack/Private/MovieSceneTextChannel.cpp b/Engine/Plugins/MovieScene/MovieSceneTextTrack/Source/MovieSceneTextTrack/Private/MovieSceneTextChannel.cpp
deleted file mode 100644
diff --git a/Engine/Plugins/MovieScene/MovieSceneTextTrack/Source/MovieSceneTextTrack/Private/MovieSceneTextPropertySystem.cpp b/Engine/Plugins/MovieScene/MovieSceneTextTrack/Source/MovieSceneTextTrack/Private/MovieSceneTextPropertySystem.cpp
deleted file mode 100644
diff --git a/Engine/Plugins/MovieScene/MovieSceneTextTrack/Source/MovieSceneTextTrack/Private/MovieSceneTextSection.cpp b/Engine/Plugins/MovieScene/MovieSceneTextTrack/Source/MovieSceneTextTrack/Private/MovieSceneTextSection.cpp
deleted file mode 100644
diff --git a/Engine/Plugins/MovieScene/MovieSceneTextTrack/Source/MovieSceneTextTrack/Private/MovieSceneTextTrack.cpp b/Engine/Plugins/MovieScene/MovieSceneTextTrack/Source/MovieSceneTextTrack/Private/MovieSceneTextTrack.cpp
deleted file mode 100644
diff --git a/Engine/Plugins/MovieScene/MovieSceneTextTrack/Source/MovieSceneTextTrack/Private/MovieSceneTextTrackModule.cpp b/Engine/Plugins/MovieScene/MovieSceneTextTrack/Source/MovieSceneTextTrack/Private/MovieSceneTextTrackModule.cpp
--- a/Engine/Plugins/MovieScene/MovieSceneTextTrack/Source/MovieSceneTextTrack/Private/MovieSceneTextTrackModule.cpp
+++ b/Engine/Plugins/MovieScene/MovieSceneTextTrack/Source/MovieSceneTextTrack/Private/MovieSceneTextTrackModule.cpp
@@ -1,24 +1,5 @@
 // Copyright Epic Games, Inc. All Rights Reserved.
 
-#include "Modules/ModuleInterface.h"
 #include "Modules/ModuleManager.h"
-#include "TextComponentTypes.h"
 
-class FMovieSceneTextTrackModule : public IModuleInterface
-{
-public:
-	virtual void StartupModule() override;
-	virtual void ShutdownModule() override;
-};
-
-IMPLEMENT_MODULE(FMovieSceneTextTrackModule, MovieSceneTextTrack)
-
-void FMovieSceneTextTrackModule::StartupModule()
-{
-	UE::MovieScene::FTextComponentTypes::Get();
-}
-
-void FMovieSceneTextTrackModule::ShutdownModule()
-{
-	UE::MovieScene::FTextComponentTypes::Destroy();
-}
+IMPLEMENT_MODULE(FDefaultModuleImpl, MovieSceneTextTrack)
diff --git a/Engine/Plugins/MovieScene/MovieSceneTextTrack/Source/MovieSceneTextTrack/Private/TextChannelEvaluatorSystem.cpp b/Engine/Plugins/MovieScene/MovieSceneTextTrack/Source/MovieSceneTextTrack/Private/TextChannelEvaluatorSystem.cpp
deleted file mode 100644
diff --git a/Engine/Plugins/MovieScene/MovieSceneTextTrack/Source/MovieSceneTextTrack/Private/TextComponentTypes.cpp b/Engine/Plugins/MovieScene/MovieSceneTextTrack/Source/MovieSceneTextTrack/Private/TextComponentTypes.cpp
--- a/Engine/Plugins/MovieScene/MovieSceneTextTrack/Source/MovieSceneTextTrack/Private/TextComponentTypes.cpp
+++ b/Engine/Plugins/MovieScene/MovieSceneTextTrack/Source/MovieSceneTextTrack/Private/TextComponentTypes.cpp
@@ -1,18 +1,17 @@
 // Copyright Epic Games, Inc. All Rights Reserved.
 
 #include "TextComponentTypes.h"
-#include "EntitySystem/BuiltInComponentTypes.h"
-#include "EntitySystem/MovieSceneComponentRegistry.h"
-#include "EntitySystem/MovieSceneEntityFactoryTemplates.h"
-#include "EntitySystem/MovieSceneEntitySystemLinker.h"
-#include "EntitySystem/MovieScenePropertyComponentHandler.h"
 
 namespace UE::MovieScene
 {
 
+PRAGMA_DISABLE_DEPRECATION_WARNINGS
 static TUniquePtr<FTextComponentTypes> GTextComponentTypes;
+PRAGMA_ENABLE_DEPRECATION_WARNINGS
+
 static bool GTextComponentTypesDestroyed = false;
 
+PRAGMA_DISABLE_DEPRECATION_WARNINGS
 FTextComponentTypes* FTextComponentTypes::Get()
 {
 	if (!GTextComponentTypes.IsValid())
@@ -22,6 +21,7 @@ FTextComponentTypes* FTextComponentTypes::Get()
 	}
 	return GTextComponentTypes.Get();
 }
+PRAGMA_ENABLE_DEPRECATION_WARNINGS
 
 void FTextComponentTypes::Destroy()
 {
@@ -29,25 +29,4 @@ void FTextComponentTypes::Destroy()
 	GTextComponentTypesDestroyed = true;
 }
 
-FTextComponentTypes::FTextComponentTypes()
-{
-	FBuiltInComponentTypes* BuiltInTypes  = FBuiltInComponentTypes::Get();
-	FComponentRegistry* ComponentRegistry = UMovieSceneEntitySystemLinker::GetComponents();
-
-	ComponentRegistry->NewPropertyType(Text, TEXT("FText"));
-
-	ComponentRegistry->NewComponentType(&TextChannel, TEXT("Text Channel"));
-	ComponentRegistry->NewComponentType(&TextResult , TEXT("Text Result"));
-
-	ComponentRegistry->Factories.DuplicateChildComponent(TextChannel);
-	ComponentRegistry->Factories.DuplicateChildComponent(TextResult);
-
-	ComponentRegistry->Factories.DefineMutuallyInclusiveComponent(TextChannel, BuiltInTypes->EvalTime);
-	ComponentRegistry->Factories.DefineMutuallyInclusiveComponent(TextChannel, TextResult);
-
-	BuiltInTypes->PropertyRegistry.DefineProperty(Text, TEXT("Apply Text Properties"))
-		.AddSoleChannel(TextResult)
-		.Commit();
-}
-
 }
diff --git a/Engine/Plugins/MovieScene/MovieSceneTextTrack/Source/MovieSceneTextTrack/Public/MovieSceneTextChannel.h b/Engine/Plugins/MovieScene/MovieSceneTextTrack/Source/MovieSceneTextTrack/Public/MovieSceneTextChannel.h
--- a/Engine/Plugins/MovieScene/MovieSceneTextTrack/Source/MovieSceneTextTrack/Public/MovieSceneTextChannel.h
+++ b/Engine/Plugins/MovieScene/MovieSceneTextTrack/Source/MovieSceneTextTrack/Public/MovieSceneTextChannel.h
@@ -2,171 +2,8 @@
 
 #pragma once
 
-#include "Channels/MovieSceneChannel.h"
-#include "Channels/MovieSceneChannelData.h"
-#include "Channels/MovieSceneChannelTraits.h"
+#include "Misc/CoreMiscDefines.h"
 
-#if WITH_EDITOR
-#include "UObject/WeakObjectPtr.h"
-#endif
-
-#include "MovieSceneTextChannel.generated.h"
-
-#if WITH_EDITOR
-class UPackage;
-#endif
-
-USTRUCT()
-struct MOVIESCENETEXTTRACK_API FMovieSceneTextChannel : public FMovieSceneChannel
-{
-	GENERATED_BODY()
-
-	typedef FText CurveValueType;
-
-#if WITH_EDITOR
-	void SetPackage(UPackage* InPackage);
-
-	UPackage* GetPackage() const;
-#endif
-
-	/**
-	 * Access a mutable interface for this channel's data
-	 * @return An object that is able to manipulate this channel's data
-	 */
-	inline TMovieSceneChannelData<FText> GetData()
-	{
-		return TMovieSceneChannelData<FText>(&Times, &Values, this, &KeyHandles);
-	}
-
-	/**
-	 * Access a constant interface for this channel's data
-	 * @return An object that is able to interrogate this channel's data
-	 */
-	inline TMovieSceneChannelData<const FText> GetData() const
-	{
-		return TMovieSceneChannelData<const FText>(&Times, &Values);
-	}
-
-	inline TConstArrayView<FFrameNumber> GetTimes() const
-	{
-		return Times;
-	}
-
-	/** Returns whether this channel has any values */
-	inline bool HasAnyData() const
-	{
-		return Times.Num() != 0 || bHasDefaultValue == true;
-	}
-
-	/**
-	 * Evaluate this channel
-	 * @param InTime     The time to evaluate at
-	 * @return A pointer to the string, or nullptr
-	 */
-	const FText* Evaluate(FFrameTime InTime) const;
-
-	/**
-	 * Add keys with these times to channel. The number of elements in both arrays much match or nothing is added.
-	 * @param InTimes Times to add
-	 * @param InValues Values to add
-	 */
-	inline void AddKeys(const TArray<FFrameNumber>& InTimes, const TArray<FText>& InValues)
-	{
-		check(InTimes.Num() == InValues.Num());
-		Times.Append(InTimes);
-		Values.Append(InValues);
-	}
-
-	/**
-	 * Set this channel's default value that should be used when no keys are present
-	 * @param InDefaultValue The desired default value
-	 */
-	inline void SetDefault(FText InDefaultValue)
-	{
-		bHasDefaultValue = true;
-		DefaultValue = InDefaultValue;
-	}
-
-	/**
-	 * Get this channel's default value that will be used when no keys are present
-	 * @return (Optional) The channel's default value
-	 */
-	inline TOptional<FText> GetDefault() const
-	{
-		return bHasDefaultValue ? TOptional<FText>(DefaultValue) : TOptional<FText>();
-	}
-
-	/** Remove this channel's default value causing the channel to have no effect where no keys are present */
-	inline void RemoveDefault()
-	{
-		bHasDefaultValue = false;
-	}
-
-	//~ Begin FMovieSceneChannel
-	virtual void GetKeys(const TRange<FFrameNumber>& WithinRange, TArray<FFrameNumber>* OutKeyTimes, TArray<FKeyHandle>* OutKeyHandles) override;
-	virtual void GetKeyTimes(TArrayView<const FKeyHandle> InHandles, TArrayView<FFrameNumber> OutKeyTimes) override;
-	virtual void SetKeyTimes(TArrayView<const FKeyHandle> InHandles, TArrayView<const FFrameNumber> InKeyTimes) override;
-	virtual void DuplicateKeys(TArrayView<const FKeyHandle> InHandles, TArrayView<FKeyHandle> OutNewHandles) override;
-	virtual void DeleteKeys(TArrayView<const FKeyHandle> InHandles) override;
-	virtual void DeleteKeysFrom(FFrameNumber InTime, bool bDeleteKeysBefore) override;
-	virtual void RemapTimes(const UE::MovieScene::IRetimingInterface& Retimer) override;
-	virtual TRange<FFrameNumber> ComputeEffectiveRange() const override;
-	virtual int32 GetNumKeys() const override;
-	virtual void Reset() override;
-	virtual void Offset(FFrameNumber DeltaPosition) override;
-	virtual void Optimize(const FKeyDataOptimizationParams& InParameters) override;
-	virtual void ClearDefault() override;
-	//~ End FMovieSceneChannel
-
-private:
-	UPROPERTY(meta=(KeyTimes))
-	TArray<FFrameNumber> Times;
-
-	/** Array of values that correspond to each key time */
-	UPROPERTY(meta=(KeyValues))
-	TArray<FText> Values;
-
-	/** Default value used when there are no keys */
-	UPROPERTY()
-	FText DefaultValue;
-
-	UPROPERTY()
-	bool bHasDefaultValue = false;
-
-	/** This needs to be a UPROPERTY so it gets saved into editor transactions but transient so it does not get saved into assets. */
-	UPROPERTY(Transient)
-	FMovieSceneKeyHandleMap KeyHandles;
-
-#if WITH_EDITOR
-	/** Package to use for Text Localization. Always set to the Package of the Section holding this Channel. */
-	TWeakObjectPtr<UPackage> Package;
-#endif
-};
-
-template<>
-struct TMovieSceneChannelTraits<FMovieSceneTextChannel> : TMovieSceneChannelTraitsBase<FMovieSceneTextChannel>
-{
-#if WITH_EDITOR
-	/** Text channels can have external values (ie, they can get their values from external objects for UI purposes) */
-	typedef TMovieSceneExternalValue<FText> ExtendedEditorDataType;
-#endif
-};
-
-inline bool EvaluateChannel(const FMovieSceneTextChannel* InChannel, FFrameTime InTime, FText& OutValue)
-{
-	if (const FText* Result = InChannel->Evaluate(InTime))
-	{
-		OutValue = *Result;
-		return true;
-	}
-	return false;
-}
-
-inline bool ValueExistsAtTime(const FMovieSceneTextChannel* Channel, FFrameNumber InFrameNumber, const FText& Value)
-{
-	const FFrameTime FrameTime(InFrameNumber);
-
-	const FText* ExistingValue = Channel->Evaluate(FrameTime);
-
-	return ExistingValue && Value.EqualTo(*ExistingValue);
-}
+/** Moved to MovieSceneTracks module */
+UE_DEPRECATED_HEADER(5.7, "Use Channels/MovieSceneTextChannel.h (MovieSceneTracks module) instead")
+#include "Channels/MovieSceneTextChannel.h"
diff --git a/Engine/Plugins/MovieScene/MovieSceneTextTrack/Source/MovieSceneTextTrack/Public/MovieSceneTextPropertySystem.h b/Engine/Plugins/MovieScene/MovieSceneTextTrack/Source/MovieSceneTextTrack/Public/MovieSceneTextPropertySystem.h
--- a/Engine/Plugins/MovieScene/MovieSceneTextTrack/Source/MovieSceneTextTrack/Public/MovieSceneTextPropertySystem.h
+++ b/Engine/Plugins/MovieScene/MovieSceneTextTrack/Source/MovieSceneTextTrack/Public/MovieSceneTextPropertySystem.h
@@ -2,18 +2,8 @@
 
 #pragma once
 
-#include "Systems/MovieScenePropertySystem.h"
-#include "MovieSceneTextPropertySystem.generated.h"
+#include "Misc/CoreMiscDefines.h"
 
-UCLASS(MinimalAPI)
-class UMovieSceneTextPropertySystem : public UMovieScenePropertySystem
-{
-	GENERATED_BODY()
-
-public:
-	UMovieSceneTextPropertySystem(const FObjectInitializer& InObjectInitializer);
-
-	//~ Begin UMovieSceneEntitySystem
-	virtual void OnRun(FSystemTaskPrerequisites& InPrerequisites, FSystemSubsequentTasks& Subsequents) override;
-	//~ End UMovieSceneEntitySystem
-};
+/** Moved to MovieSceneTracks module */
+UE_DEPRECATED_HEADER(5.7, "Use Systems/MovieSceneTextPropertySystem.h (MovieSceneTracks module) instead")
+#include "Systems/MovieSceneTextPropertySystem.h"
diff --git a/Engine/Plugins/MovieScene/MovieSceneTextTrack/Source/MovieSceneTextTrack/Public/MovieSceneTextSection.h b/Engine/Plugins/MovieScene/MovieSceneTextTrack/Source/MovieSceneTextTrack/Public/MovieSceneTextSection.h
--- a/Engine/Plugins/MovieScene/MovieSceneTextTrack/Source/MovieSceneTextTrack/Public/MovieSceneTextSection.h
+++ b/Engine/Plugins/MovieScene/MovieSceneTextTrack/Source/MovieSceneTextTrack/Public/MovieSceneTextSection.h
@@ -2,30 +2,8 @@
 
 #pragma once
 
-#include "EntitySystem/IMovieSceneEntityProvider.h"
-#include "MovieSceneSection.h"
-#include "MovieSceneTextChannel.h"
-#include "MovieSceneTextSection.generated.h"
+#include "Misc/CoreMiscDefines.h"
 
-UCLASS(MinimalAPI)
-class UMovieSceneTextSection
-	: public UMovieSceneSection
-	, public IMovieSceneEntityProvider
-{
-	GENERATED_BODY()
-
-public:
-	UMovieSceneTextSection(const FObjectInitializer& InObjectInitializer);
-
-	/** Public access to this section's internal data function */
-	const FMovieSceneTextChannel& GetChannel() const { return TextChannel; }
-
-	//~ Begin IMovieSceneEntityProvider
-	virtual bool PopulateEvaluationFieldImpl(const TRange<FFrameNumber>& EffectiveRange, const FMovieSceneEvaluationFieldEntityMetaData& InMetaData, FMovieSceneEntityComponentFieldBuilder* OutFieldBuilder) override;
-	virtual void ImportEntityImpl(UMovieSceneEntitySystemLinker* EntityLinker, const FEntityImportParams& Params, FImportedEntity* OutImportedEntity) override;
-	//~ End IMovieSceneEntityProvider
-
-private:
-	UPROPERTY()
-	FMovieSceneTextChannel TextChannel;
-};
+/** Moved to MovieSceneTracks module */
+UE_DEPRECATED_HEADER(5.7, "Use Sections/MovieSceneTextSection.h (MovieSceneTracks module) instead")
+#include "Sections/MovieSceneTextSection.h"
diff --git a/Engine/Plugins/MovieScene/MovieSceneTextTrack/Source/MovieSceneTextTrack/Public/MovieSceneTextTrack.h b/Engine/Plugins/MovieScene/MovieSceneTextTrack/Source/MovieSceneTextTrack/Public/MovieSceneTextTrack.h
--- a/Engine/Plugins/MovieScene/MovieSceneTextTrack/Source/MovieSceneTextTrack/Public/MovieSceneTextTrack.h
+++ b/Engine/Plugins/MovieScene/MovieSceneTextTrack/Source/MovieSceneTextTrack/Public/MovieSceneTextTrack.h
@@ -2,27 +2,8 @@
 
 #pragma once
 
-#include "Tracks/MovieScenePropertyTrack.h"
-#include "MovieSceneTextTrack.generated.h"
+#include "Misc/CoreMiscDefines.h"
 
-/** Movie Scene Track that holds a series of texts. */
-UCLASS(MinimalAPI)
-class UMovieSceneTextTrack : public UMovieScenePropertyTrack
-{
-	GENERATED_BODY()
-
-public:
-	UMovieSceneTextTrack();
-
-	//~ Begin UMovieSceneTrack
-	virtual bool SupportsType(TSubclassOf<UMovieSceneSection> SectionClass) const override;
-	virtual UMovieSceneSection* CreateNewSection() override;
-	virtual void AddSection(UMovieSceneSection& Section) override;
-	virtual const TArray<UMovieSceneSection*>& GetAllSections() const override;
-	virtual bool HasSection(const UMovieSceneSection& Section) const override;
-	virtual bool IsEmpty() const override;
-	virtual void RemoveAllAnimationData() override;
-	virtual void RemoveSection(UMovieSceneSection& Section) override;
-	virtual void RemoveSectionAt(int32 SectionIndex) override;
-	//~ End UMovieSceneTrack
-};
+/** Moved to MovieSceneTracks module */
+UE_DEPRECATED_HEADER(5.7, "Use Tracks/MovieSceneTextTrack.h (MovieSceneTracks module) instead")
+#include "Tracks/MovieSceneTextTrack.h"
diff --git a/Engine/Plugins/MovieScene/MovieSceneTextTrack/Source/MovieSceneTextTrack/Public/TextChannelEvaluatorSystem.h b/Engine/Plugins/MovieScene/MovieSceneTextTrack/Source/MovieSceneTextTrack/Public/TextChannelEvaluatorSystem.h
--- a/Engine/Plugins/MovieScene/MovieSceneTextTrack/Source/MovieSceneTextTrack/Public/TextChannelEvaluatorSystem.h
+++ b/Engine/Plugins/MovieScene/MovieSceneTextTrack/Source/MovieSceneTextTrack/Public/TextChannelEvaluatorSystem.h
@@ -2,20 +2,8 @@
 
 #pragma once
 
-#include "EntitySystem/MovieSceneEntitySystem.h"
-#include "TextChannelEvaluatorSystem.generated.h"
+#include "Misc/CoreMiscDefines.h"
 
-/** System responsible for evaluating Text channels */
-UCLASS(MinimalAPI)
-class UTextChannelEvaluatorSystem : public UMovieSceneEntitySystem
-{
-	GENERATED_BODY()
-
-public:
-	UTextChannelEvaluatorSystem(const FObjectInitializer& InObjectInitializer);
-
-	//~ Begin UMovieSceneEntitySystem
-	virtual void OnSchedulePersistentTasks(UE::MovieScene::IEntitySystemScheduler* TaskScheduler) override;
-	virtual void OnRun(FSystemTaskPrerequisites& InPrerequisites, FSystemSubsequentTasks& Subsequents) override;
-	//~ End UMovieSceneEntitySystem
-};
+/** Moved to MovieSceneTracks module */
+UE_DEPRECATED_HEADER(5.7, "Use Systems/TextChannelEvaluatorSystem.h (MovieSceneTracks module) instead")
+#include "Systems/TextChannelEvaluatorSystem.h"
diff --git a/Engine/Plugins/MovieScene/MovieSceneTextTrack/Source/MovieSceneTextTrack/Public/TextComponentTypes.h b/Engine/Plugins/MovieScene/MovieSceneTextTrack/Source/MovieSceneTextTrack/Public/TextComponentTypes.h
--- a/Engine/Plugins/MovieScene/MovieSceneTextTrack/Source/MovieSceneTextTrack/Public/TextComponentTypes.h
+++ b/Engine/Plugins/MovieScene/MovieSceneTextTrack/Source/MovieSceneTextTrack/Public/TextComponentTypes.h
@@ -2,30 +2,19 @@
 
 #pragma once
 
-#include "EntitySystem/MovieSceneEntityIDs.h"
-#include "EntitySystem/MovieScenePropertyMetaDataTraits.h"
-#include "EntitySystem/MovieScenePropertySystemTypes.h"
-#include "EntitySystem/MovieScenePropertyTraits.h"
-#include "Internationalization/Text.h"
-
-struct FMovieSceneTextChannel;
+#include "EntitySystem/BuiltInComponentTypes.h"
+#include "Misc/CoreMiscDefines.h"
+#include "MovieSceneTracksComponentTypes.h"
 
 namespace UE::MovieScene
 {
 
-/** The component data for evaluating a Text channel */
-struct FSourceTextChannel
-{
-	FSourceTextChannel() : Source(nullptr) {}
-	FSourceTextChannel(const FMovieSceneTextChannel* InSource) : Source(InSource) {}
-	const FMovieSceneTextChannel* Source;
-};
-
-using FTextPropertyTraits = TDirectPropertyTraits<FText>;
-
-struct FTextComponentTypes
+struct UE_DEPRECATED(5.7, "FTextComponentTypes deprecated. Use FBuiltInComponentTypes and FMovieSceneTracksComponentTypes instead") FTextComponentTypes
 {
+PRAGMA_DISABLE_DEPRECATION_WARNINGS
 	MOVIESCENETEXTTRACK_API static FTextComponentTypes* Get();
+PRAGMA_ENABLE_DEPRECATION_WARNINGS
+
 	MOVIESCENETEXTTRACK_API static void Destroy();
 
 	// An FMovieSceneTextChannel
@@ -37,7 +26,13 @@ struct FTextComponentTypes
 	TPropertyComponents<FTextPropertyTraits> Text;
 
 private:
-	FTextComponentTypes();
+	PRAGMA_DISABLE_DEPRECATION_WARNINGS
+	FTextComponentTypes() = default;
+	FTextComponentTypes(const FTextComponentTypes&) = default;
+	FTextComponentTypes(FTextComponentTypes&&) = default;
+	FTextComponentTypes& operator=(const FTextComponentTypes&) = default;
+	FTextComponentTypes& operator=(FTextComponentTypes&&) = default;
+	PRAGMA_ENABLE_DEPRECATION_WARNINGS
 };
 
-}
\ No newline at end of file
+}
diff --git a/Engine/Plugins/MovieScene/MovieSceneTextTrack/Source/MovieSceneTextTrackEditor/MovieSceneTextTrackEditor.Build.cs b/Engine/Plugins/MovieScene/MovieSceneTextTrack/Source/MovieSceneTextTrackEditor/MovieSceneTextTrackEditor.Build.cs
--- a/Engine/Plugins/MovieScene/MovieSceneTextTrack/Source/MovieSceneTextTrackEditor/MovieSceneTextTrackEditor.Build.cs
+++ b/Engine/Plugins/MovieScene/MovieSceneTextTrack/Source/MovieSceneTextTrackEditor/MovieSceneTextTrackEditor.Build.cs
@@ -12,17 +12,6 @@ public class MovieSceneTextTrackEditor : ModuleRules
 			new string[]
 			{
 				"Core",
-				"CoreUObject",
-				"Engine",
-				"MovieScene",
-				"MovieSceneTextTrack",
-				"MovieSceneTools",
-				"MovieSceneTracks",
-				"Sequencer",
-				"SequencerCore",
-				"Slate",
-				"SlateCore",
-				"UnrealEd",
 			}
 		);
 	}
diff --git a/Engine/Plugins/MovieScene/MovieSceneTextTrack/Source/MovieSceneTextTrackEditor/Private/MovieSceneTextKeyStruct.h b/Engine/Plugins/MovieScene/MovieSceneTextTrack/Source/MovieSceneTextTrackEditor/Private/MovieSceneTextKeyStruct.h
deleted file mode 100644
diff --git a/Engine/Plugins/MovieScene/MovieSceneTextTrack/Source/MovieSceneTextTrackEditor/Private/MovieSceneTextTrackEditorModule.cpp b/Engine/Plugins/MovieScene/MovieSceneTextTrack/Source/MovieSceneTextTrackEditor/Private/MovieSceneTextTrackEditorModule.cpp
--- a/Engine/Plugins/MovieScene/MovieSceneTextTrack/Source/MovieSceneTextTrackEditor/Private/MovieSceneTextTrackEditorModule.cpp
+++ b/Engine/Plugins/MovieScene/MovieSceneTextTrack/Source/MovieSceneTextTrackEditor/Private/MovieSceneTextTrackEditorModule.cpp
@@ -1,47 +1,5 @@
 // Copyright Epic Games, Inc. All Rights Reserved.
 
-#include "ISequencerModule.h"
-#include "Modules/ModuleInterface.h"
 #include "Modules/ModuleManager.h"
-#include "MovieSceneTextChannel.h"
-#include "SequencerChannelInterface.h"
-#include "TextClipboardTypes.h"
-#include "TextPropertyTrackEditor.h"
 
-class FMovieSceneTextTrackEditorModule : public IModuleInterface
-{
-public:
-	//~ Begin IModuleInterface
-	virtual void StartupModule() override;
-	virtual void ShutdownModule() override;
-	//~ End IModuleInterface
-
-private:
-	FDelegateHandle TextPropertyTrackCreateEditorHandle;
-};
-
-IMPLEMENT_MODULE(FMovieSceneTextTrackEditorModule, MovieSceneTextTrackEditor)
-
-void FMovieSceneTextTrackEditorModule::StartupModule()
-{
-	if (!GIsEditor)
-	{
-		return;
-	}
-
-	ISequencerModule& SequencerModule = FModuleManager::Get().LoadModuleChecked<ISequencerModule>("Sequencer");
-	TextPropertyTrackCreateEditorHandle = SequencerModule.RegisterPropertyTrackEditor<FTextPropertyTrackEditor>();
-	SequencerModule.RegisterChannelInterface<FMovieSceneTextChannel>();
-}
-
-void FMovieSceneTextTrackEditorModule::ShutdownModule()
-{
-	ISequencerModule* SequencerModule = FModuleManager::Get().GetModulePtr<ISequencerModule>("Sequencer");
-	if (!SequencerModule)
-	{
-		return;
-	}
-
-	SequencerModule->UnRegisterTrackEditor(TextPropertyTrackCreateEditorHandle);
-	TextPropertyTrackCreateEditorHandle.Reset();
-}
+IMPLEMENT_MODULE(FDefaultModuleImpl, MovieSceneTextTrackEditor)
diff --git a/Engine/Plugins/MovieScene/MovieSceneTextTrack/Source/MovieSceneTextTrackEditor/Private/TextClipboardTypes.h b/Engine/Plugins/MovieScene/MovieSceneTextTrack/Source/MovieSceneTextTrackEditor/Private/TextClipboardTypes.h
deleted file mode 100644
diff --git a/Engine/Plugins/MovieScene/MovieSceneTextTrack/Source/MovieSceneTextTrackEditor/Private/TextPropertyTrackEditor.cpp b/Engine/Plugins/MovieScene/MovieSceneTextTrack/Source/MovieSceneTextTrackEditor/Private/TextPropertyTrackEditor.cpp
deleted file mode 100644
diff --git a/Engine/Plugins/MovieScene/MovieSceneTextTrack/Source/MovieSceneTextTrackEditor/Private/TextPropertyTrackEditor.h b/Engine/Plugins/MovieScene/MovieSceneTextTrack/Source/MovieSceneTextTrackEditor/Private/TextPropertyTrackEditor.h
deleted file mode 100644
diff --git a/Engine/Plugins/MovieScene/MovieSceneTextTrack/Source/MovieSceneTextTrackEditor/Private/Widgets/STextKeyEditor.cpp b/Engine/Plugins/MovieScene/MovieSceneTextTrack/Source/MovieSceneTextTrackEditor/Private/Widgets/STextKeyEditor.cpp
deleted file mode 100644
diff --git a/Engine/Plugins/MovieScene/MovieSceneTextTrack/Source/MovieSceneTextTrackEditor/Private/Widgets/STextKeyEditor.h b/Engine/Plugins/MovieScene/MovieSceneTextTrack/Source/MovieSceneTextTrackEditor/Private/Widgets/STextKeyEditor.h
deleted file mode 100644
diff --git a/Engine/Source/Editor/MovieSceneTools/Private/Channels/BuiltInChannelEditors.cpp b/Engine/Source/Editor/MovieSceneTools/Private/Channels/BuiltInChannelEditors.cpp
--- a/Engine/Source/Editor/MovieSceneTools/Private/Channels/BuiltInChannelEditors.cpp
+++ b/Engine/Source/Editor/MovieSceneTools/Private/Channels/BuiltInChannelEditors.cpp
@@ -22,6 +22,7 @@
 #include "CurveKeyEditors/SNumericKeyEditor.h"
 #include "CurveKeyEditors/SBoolCurveKeyEditor.h"
 #include "CurveKeyEditors/SStringCurveKeyEditor.h"
+#include "CurveKeyEditors/STextCurveKeyEditor.h"
 #include "CurveKeyEditors/SEnumKeyEditor.h"
 #include "UObject/StructOnScope.h"
 #include "MVVM/Views/KeyDrawParams.h"
@@ -65,6 +66,7 @@
 #include "ISequencerModule.h"
 #include "MovieScene.h"
 #include "MovieSceneTracksComponentTypes.h"
+#include "KeyStructTypes/MovieSceneTextKeyStruct.h"
 
 #define LOCTEXT_NAMESPACE "BuiltInChannelEditors"
 
@@ -296,6 +298,10 @@ bool CanCreateKeyEditor(const FMovieSceneStringChannel*  Channel)
 {
 	return true;
 }
+bool CanCreateKeyEditor(const FMovieSceneTextChannel* Channel)
+{
+	return true;
+}
 bool CanCreateKeyEditor(const FMovieSceneObjectPathChannel* Channel)
 {
 	return true;
@@ -397,6 +403,22 @@ TSharedRef<SWidget> CreateKeyEditor(const TMovieSceneChannelHandle<FMovieSceneSt
 	return SNew(SStringCurveKeyEditor, KeyEditor);
 }
 
+TSharedRef<SWidget> CreateKeyEditor(const TMovieSceneChannelHandle<FMovieSceneTextChannel>& Channel, const UE::Sequencer::FCreateKeyEditorParams& Params)
+{
+	const TMovieSceneExternalValue<FText>* ExternalValue = Channel.GetExtendedEditorData();
+	if (!ExternalValue)
+	{
+		return SNullWidget::NullWidget;
+	}
+
+	TSequencerKeyEditor<FMovieSceneTextChannel, FText> KeyEditor(
+		Params.ObjectBindingID, Channel,
+		Params.OwningSection, Params.Sequencer, Params.PropertyBindings, ExternalValue->OnGetExternalValue
+		);
+	KeyEditor.SetOwningObject(Cast<UMovieSceneSignedObject>(Params.OwningObject));
+
+	return SNew(STextCurveKeyEditor, KeyEditor);
+}
 
 TSharedRef<SWidget> CreateKeyEditor(const TMovieSceneChannelHandle<FMovieSceneByteChannel>&    Channel, const UE::Sequencer::FCreateKeyEditorParams& Params)
 {
@@ -882,6 +904,58 @@ UMovieSceneKeyStructType* InstanceGeneratedStruct(FMovieSceneTimeWarpChannel* Ch
 	return NewStruct;
 }
 
+UMovieSceneKeyStructType* InstanceGeneratedStruct(FMovieSceneTextChannel* Channel, FSequencerKeyStructGenerator* Generator)
+{
+	UScriptStruct* ChannelType = FMovieSceneTextChannel::StaticStruct();
+	check(ChannelType);
+
+	UPackage* const Package = Channel->GetPackage();
+	if (!ensureMsgf(Package, TEXT("No valid package could be found for Text Channel")))
+	{
+		return nullptr;
+	}
+
+	// Generated struct to be unique per Package
+	const FName InstanceStructName = *(ChannelType->GetName() + Package->GetPersistentGuid().ToString());
+
+	UMovieSceneKeyStructType* Existing = Generator->FindGeneratedStruct(InstanceStructName);
+	if (Existing)
+	{
+		return Existing;
+	}
+
+	static const FName TimesMetaDataTag("KeyTimes");
+	static const FName ValuesMetaDataTag("KeyValues");
+
+	FArrayProperty* SourceTimes  = FSequencerKeyStructGenerator::FindArrayPropertyWithTag(ChannelType, TimesMetaDataTag);
+	FArrayProperty* SourceValues = FSequencerKeyStructGenerator::FindArrayPropertyWithTag(ChannelType, ValuesMetaDataTag);
+
+	if (!ensureAlwaysMsgf(SourceTimes && SourceValues
+		, TEXT("FMovieSceneTextChannel does not have KeyTimes & KeyValues meta data tags in their respective array properties. Did something change ?")))
+	{
+		return nullptr;
+	}
+
+	constexpr EObjectFlags ObjectFlags = RF_Public | RF_Standalone | RF_Transient;
+
+	UMovieSceneTextKeyStruct* NewStruct = NewObject<UMovieSceneTextKeyStruct>(Package, NAME_None, ObjectFlags);
+	NewStruct->SourceTimesProperty  = SourceTimes;
+	NewStruct->SourceValuesProperty = SourceValues;
+
+	FTextProperty* NewValueProperty = new FTextProperty(NewStruct, "Value", RF_NoFlags);
+	NewValueProperty->SetPropertyFlags(CPF_Edit);
+	NewValueProperty->SetMetaData("Category", TEXT("Key"));
+	NewValueProperty->ArrayDim = 1;
+
+	NewStruct->AddCppProperty(NewValueProperty);
+	NewStruct->DestValueProperty = NewValueProperty;
+
+	FSequencerKeyStructGenerator::FinalizeNewKeyStruct(NewStruct);
+
+	Generator->AddGeneratedStruct(InstanceStructName, NewStruct);
+	return NewStruct;
+}
+
 void PostConstructKeyInstance(const TMovieSceneChannelHandle<FMovieSceneTimeWarpChannel>& ChannelHandle, FKeyHandle InHandle, FStructOnScope* Struct)
 {
 	const UMovieSceneKeyStructType* GeneratedStructType = CastChecked<const UMovieSceneKeyStructType>(Struct->GetStruct());
@@ -1055,6 +1129,32 @@ void PostConstructKeyInstance(const TMovieSceneChannelHandle<FMovieSceneDoubleCh
 	KeyStruct->OnPropertyChangedEvent = CopyInstanceToKeyLambda;
 }
 
+void PostConstructKeyInstance(const TMovieSceneChannelHandle<FMovieSceneTextChannel>& ChannelHandle, FKeyHandle InHandle, FStructOnScope* Struct)
+{
+	FMovieSceneTextChannel* Channel = ChannelHandle.Get();
+	if (!Channel)
+	{
+		return;
+	}
+
+	// Set Package that the Text Property will need for Localization
+	Struct->SetPackage(ChannelHandle.Get()->GetPackage());
+
+	const UMovieSceneKeyStructType* GeneratedStructType = CastChecked<const UMovieSceneKeyStructType>(Struct->GetStruct());
+
+	const int32 InitialKeyIndex = Channel->GetData().GetIndex(InHandle);
+
+	// Copy the initial value into the struct
+	if (InitialKeyIndex != INDEX_NONE)
+	{
+		const uint8* SrcValueData  = GeneratedStructType->SourceValuesProperty->ContainerPtrToValuePtr<uint8>(Channel);
+		uint8* DestValueData = GeneratedStructType->DestValueProperty->ContainerPtrToValuePtr<uint8>(Struct->GetStructMemory());
+
+		FScriptArrayHelper SourceValuesArray(GeneratedStructType->SourceValuesProperty.Get(), SrcValueData);
+		GeneratedStructType->SourceValuesProperty->Inner->CopyCompleteValue(DestValueData, SourceValuesArray.GetRawPtr(InitialKeyIndex));
+	}
+}
+
 template<typename ChannelType>
 void DrawKeysImpl(ChannelType* Channel, TArrayView<const FKeyHandle> InKeyHandles, const UMovieSceneSection* InOwner, TArrayView<FKeyDrawParams> OutKeyDrawParams)
 {
diff --git a/Engine/Source/Editor/MovieSceneTools/Private/CurveKeyEditors/STextCurveKeyEditor.cpp b/Engine/Source/Editor/MovieSceneTools/Private/CurveKeyEditors/STextCurveKeyEditor.cpp
new file mode 100644
--- /dev/null
+++ b/Engine/Source/Editor/MovieSceneTools/Private/CurveKeyEditors/STextCurveKeyEditor.cpp
@@ -0,0 +1,34 @@
+// Copyright Epic Games, Inc. All Rights Reserved.
+
+#include "STextCurveKeyEditor.h"
+#include "ScopedTransaction.h"
+#include "Widgets/Input/SEditableTextBox.h"
+
+#define LOCTEXT_NAMESPACE "STextCurveKeyEditor"
+
+void STextCurveKeyEditor::Construct(const FArguments& InArgs, const TSequencerKeyEditor<FMovieSceneTextChannel, FText>& InKeyEditor)
+{
+	KeyEditor = InKeyEditor;
+
+	ChildSlot
+	[
+		SNew(SEditableTextBox)
+		.MinDesiredWidth(10.f)
+		.SelectAllTextWhenFocused(true)
+		.Text(this, &STextCurveKeyEditor::GetText)
+		.OnTextCommitted(this, &STextCurveKeyEditor::OnTextCommitted)
+	];
+}
+
+FText STextCurveKeyEditor::GetText() const
+{
+	return KeyEditor.GetCurrentValue();
+}
+
+void STextCurveKeyEditor::OnTextCommitted(const FText& InText, ETextCommit::Type InCommitType)
+{
+	FScopedTransaction Transaction(LOCTEXT("SetTextKey", "Set Text Key Value"));
+	KeyEditor.SetValueWithNotify(InText, EMovieSceneDataChangeType::TrackValueChangedRefreshImmediately);
+}
+
+#undef LOCTEXT_NAMESPACE
diff --git a/Engine/Source/Editor/MovieSceneTools/Private/CurveKeyEditors/STextCurveKeyEditor.h b/Engine/Source/Editor/MovieSceneTools/Private/CurveKeyEditors/STextCurveKeyEditor.h
new file mode 100644
--- /dev/null
+++ b/Engine/Source/Editor/MovieSceneTools/Private/CurveKeyEditors/STextCurveKeyEditor.h
@@ -0,0 +1,31 @@
+// Copyright Epic Games, Inc. All Rights Reserved.
+
+#pragma once
+
+#include "SequencerKeyEditor.h"
+#include "Channels/MovieSceneChannelHandle.h"
+#include "Channels/MovieSceneTextChannel.h"
+#include "Widgets/SCompoundWidget.h"
+
+class ISequencer;
+
+/**
+ * A widget for editing a curve representing text keys.
+ */
+class STextCurveKeyEditor : public SCompoundWidget
+{
+public:
+	SLATE_BEGIN_ARGS(STextCurveKeyEditor){}
+	SLATE_END_ARGS()
+
+	void Construct(const FArguments& InArgs, const TSequencerKeyEditor<FMovieSceneTextChannel, FText>& InKeyEditor);
+
+private:
+	/** Gets the evaluated Text value at the current sequencer time */
+	FText GetText() const;
+
+	/** Updates the CachedText result to the committed text and adds/updates a key to the text channel at the current sequencer time */
+	void OnTextCommitted(const FText& InText, ETextCommit::Type InCommitType);
+
+	TSequencerKeyEditor<FMovieSceneTextChannel, FText> KeyEditor;
+};
diff --git a/Engine/Source/Editor/MovieSceneTools/Private/MovieSceneToolsModule.cpp b/Engine/Source/Editor/MovieSceneTools/Private/MovieSceneToolsModule.cpp
--- a/Engine/Source/Editor/MovieSceneTools/Private/MovieSceneToolsModule.cpp
+++ b/Engine/Source/Editor/MovieSceneTools/Private/MovieSceneToolsModule.cpp
@@ -31,6 +31,7 @@
 #include "TrackEditors/PropertyTrackEditors/VisibilityPropertyTrackEditor.h"
 #include "TrackEditors/PropertyTrackEditors/ActorReferencePropertyTrackEditor.h"
 #include "TrackEditors/PropertyTrackEditors/StringPropertyTrackEditor.h"
+#include "TrackEditors/PropertyTrackEditors/TextPropertyTrackEditor.h"
 
 #include "TrackEditors/TransformTrackEditor.h"
 #include "TrackEditors/CameraCutTrackEditor.h"
@@ -174,6 +175,7 @@ void FMovieSceneToolsModule::StartupModule()
 		VisibilityPropertyTrackCreateEditorHandle = SequencerModule.RegisterPropertyTrackEditor<FVisibilityPropertyTrackEditor>();
 		ActorReferencePropertyTrackCreateEditorHandle = SequencerModule.RegisterPropertyTrackEditor<FActorReferencePropertyTrackEditor>();
 		StringPropertyTrackCreateEditorHandle = SequencerModule.RegisterPropertyTrackEditor<FStringPropertyTrackEditor>();
+		TextPropertyTrackCreateEditorHandle = SequencerModule.RegisterPropertyTrackEditor<FTextPropertyTrackEditor>();
 		ObjectTrackCreateEditorHandle = SequencerModule.RegisterPropertyTrackEditor<FObjectPropertyTrackEditor>();
 
 		// register specialty track editors
@@ -234,6 +236,7 @@ void FMovieSceneToolsModule::StartupModule()
 		SequencerModule.RegisterChannelInterface<FMovieSceneFloatChannel>();
 		SequencerModule.RegisterChannelInterface<FMovieSceneDoubleChannel>();
 		SequencerModule.RegisterChannelInterface<FMovieSceneStringChannel>();
+		SequencerModule.RegisterChannelInterface<FMovieSceneTextChannel>();
 		SequencerModule.RegisterChannelInterface<FMovieSceneParticleChannel>();
 		SequencerModule.RegisterChannelInterface<FMovieSceneActorReferenceData>();
 		SequencerModule.RegisterChannelInterface<FMovieSceneEventSectionData>();
@@ -367,6 +370,7 @@ void FMovieSceneToolsModule::ShutdownModule()
 	SequencerModule.UnRegisterTrackEditor( VisibilityPropertyTrackCreateEditorHandle );
 	SequencerModule.UnRegisterTrackEditor( ActorReferencePropertyTrackCreateEditorHandle );
 	SequencerModule.UnRegisterTrackEditor( StringPropertyTrackCreateEditorHandle );
+	SequencerModule.UnRegisterTrackEditor( TextPropertyTrackCreateEditorHandle );
 	SequencerModule.UnRegisterTrackEditor( CameraShakeSourceShakeCreateEditorHandle );
 
 	// unregister specialty track editors
diff --git a/Engine/Source/Editor/MovieSceneTools/Private/TrackEditors/PropertyTrackEditors/TextPropertyTrackEditor.cpp b/Engine/Source/Editor/MovieSceneTools/Private/TrackEditors/PropertyTrackEditors/TextPropertyTrackEditor.cpp
new file mode 100644
--- /dev/null
+++ b/Engine/Source/Editor/MovieSceneTools/Private/TrackEditors/PropertyTrackEditors/TextPropertyTrackEditor.cpp
@@ -0,0 +1,26 @@
+// Copyright Epic Games, Inc. All Rights Reserved.
+
+#include "TextPropertyTrackEditor.h"
+#include "Channels/MovieSceneTextChannel.h"
+#include "ISequencerChannelInterface.h"
+#include "SequencerKeyStructGenerator.h"
+#include "UObject/TextProperty.h"
+
+TSharedRef<ISequencerTrackEditor> FTextPropertyTrackEditor::CreateTrackEditor(TSharedRef<ISequencer> OwningSequencer)
+{
+	return MakeShared<FTextPropertyTrackEditor>(OwningSequencer);
+}
+
+void FTextPropertyTrackEditor::GenerateKeysFromPropertyChanged(const FPropertyChangedParams& PropertyChangedParams
+	, UMovieSceneSection* SectionToKey
+	, FGeneratedTrackKeys& OutGeneratedKeys)
+{
+	const FTextProperty* TextProperty = CastField<const FTextProperty>(PropertyChangedParams.PropertyPath.GetLeafMostProperty().Property.Get());
+	if (!TextProperty)
+	{
+		return;
+	}
+
+	FText TextPropertyValue = PropertyChangedParams.GetPropertyValue<FText>();
+	OutGeneratedKeys.Add(FMovieSceneChannelValueSetter::Create<FMovieSceneTextChannel>(0, MoveTemp(TextPropertyValue), true));
+}
diff --git a/Engine/Source/Editor/MovieSceneTools/Private/TrackEditors/PropertyTrackEditors/TextPropertyTrackEditor.h b/Engine/Source/Editor/MovieSceneTools/Private/TrackEditors/PropertyTrackEditors/TextPropertyTrackEditor.h
new file mode 100644
--- /dev/null
+++ b/Engine/Source/Editor/MovieSceneTools/Private/TrackEditors/PropertyTrackEditors/TextPropertyTrackEditor.h
@@ -0,0 +1,38 @@
+// Copyright Epic Games, Inc. All Rights Reserved.
+
+#pragma once
+
+#include "PropertyTrackEditor.h"
+#include "Tracks/MovieSceneTextTrack.h"
+
+class UMovieSceneKeyStructType;
+class FSequencerKeyStructGenerator;
+struct FMovieSceneTextChannel;
+
+/** Track Editor for Text Property */
+class FTextPropertyTrackEditor : public FPropertyTrackEditor<UMovieSceneTextTrack>
+{
+public:
+	FTextPropertyTrackEditor(TSharedRef<ISequencer> InSequencer)
+		: FPropertyTrackEditor(InSequencer, GetAnimatedPropertyTypes())
+	{
+	}
+
+	/** Retrieve a list of all property types that this track editor animates  */
+	static TArray<FAnimatedPropertyKey, TInlineAllocator<1>> GetAnimatedPropertyTypes()
+	{
+		return TArray<FAnimatedPropertyKey, TInlineAllocator<1>>({ FAnimatedPropertyKey::FromPropertyTypeName(NAME_TextProperty) });
+	}
+
+	/**
+	 * Creates an instance of this class (called by a sequencer).
+	 * @param OwningSequencer The sequencer instance to be used by this tool
+	 * @return The new instance of this class
+	 */
+	static TSharedRef<ISequencerTrackEditor> CreateTrackEditor(TSharedRef<ISequencer> OwningSequencer);
+
+protected:
+	//~ Begin FPropertyTrackEditor
+	virtual void GenerateKeysFromPropertyChanged(const FPropertyChangedParams& PropertyChangedParams, UMovieSceneSection* SectionToKey, FGeneratedTrackKeys& OutGeneratedKeys) override;
+	//~ End FPropertyTrackEditor
+};
diff --git a/Engine/Source/Editor/MovieSceneTools/Public/Channels/BuiltInChannelEditors.h b/Engine/Source/Editor/MovieSceneTools/Public/Channels/BuiltInChannelEditors.h
--- a/Engine/Source/Editor/MovieSceneTools/Public/Channels/BuiltInChannelEditors.h
+++ b/Engine/Source/Editor/MovieSceneTools/Public/Channels/BuiltInChannelEditors.h
@@ -22,6 +22,7 @@
 #include "Channels/MovieSceneObjectPathChannel.h"
 #include "Channels/MovieSceneTimeWarpChannel.h"
 #include "Sections/MovieSceneStringSection.h"
+#include "Channels/MovieSceneTextChannel.h"
 #include "Sections/MovieSceneParticleSection.h"
 #include "Sections/MovieSceneActorReferenceSection.h"
 
@@ -47,6 +48,7 @@ MOVIESCENETOOLS_API bool CanCreateKeyEditor(const FMovieSceneIntegerChannel*
 MOVIESCENETOOLS_API bool CanCreateKeyEditor(const FMovieSceneDoubleChannel*      Channel);
 MOVIESCENETOOLS_API bool CanCreateKeyEditor(const FMovieSceneFloatChannel*      Channel);
 MOVIESCENETOOLS_API bool CanCreateKeyEditor(const FMovieSceneStringChannel*     Channel);
+MOVIESCENETOOLS_API bool CanCreateKeyEditor(const FMovieSceneTextChannel*       Channel);
 MOVIESCENETOOLS_API bool CanCreateKeyEditor(const FMovieSceneObjectPathChannel* Channel);
 MOVIESCENETOOLS_API bool CanCreateKeyEditor(const FMovieSceneActorReferenceData* Channel);
 
@@ -56,18 +58,20 @@ MOVIESCENETOOLS_API TSharedRef<SWidget> CreateKeyEditor(const TMovieSceneChannel
 MOVIESCENETOOLS_API TSharedRef<SWidget> CreateKeyEditor(const TMovieSceneChannelHandle<FMovieSceneDoubleChannel>&      Channel, const UE::Sequencer::FCreateKeyEditorParams& Params);
 MOVIESCENETOOLS_API TSharedRef<SWidget> CreateKeyEditor(const TMovieSceneChannelHandle<FMovieSceneFloatChannel>&       Channel, const UE::Sequencer::FCreateKeyEditorParams& Params);
 MOVIESCENETOOLS_API TSharedRef<SWidget> CreateKeyEditor(const TMovieSceneChannelHandle<FMovieSceneStringChannel>&      Channel, const UE::Sequencer::FCreateKeyEditorParams& Params);
+MOVIESCENETOOLS_API TSharedRef<SWidget> CreateKeyEditor(const TMovieSceneChannelHandle<FMovieSceneTextChannel>&        Channel, const UE::Sequencer::FCreateKeyEditorParams& Params);
 MOVIESCENETOOLS_API TSharedRef<SWidget> CreateKeyEditor(const TMovieSceneChannelHandle<FMovieSceneObjectPathChannel>&  Channel, const UE::Sequencer::FCreateKeyEditorParams& Params);
 MOVIESCENETOOLS_API TSharedRef<SWidget> CreateKeyEditor(const TMovieSceneChannelHandle<FMovieSceneActorReferenceData>& Channel, const UE::Sequencer::FCreateKeyEditorParams& Params);
 MOVIESCENETOOLS_API TSharedRef<SWidget> CreateKeyEditor(const TMovieSceneChannelHandle<FMovieSceneTimeWarpChannel>&    Channel, const UE::Sequencer::FCreateKeyEditorParams& Params);
 
 MOVIESCENETOOLS_API UMovieSceneKeyStructType* InstanceGeneratedStruct(FMovieSceneByteChannel* Channel,       FSequencerKeyStructGenerator* Generator);
 MOVIESCENETOOLS_API UMovieSceneKeyStructType* InstanceGeneratedStruct(FMovieSceneTimeWarpChannel* Channel,   FSequencerKeyStructGenerator* Generator);
+MOVIESCENETOOLS_API UMovieSceneKeyStructType* InstanceGeneratedStruct(FMovieSceneTextChannel* Channel,       FSequencerKeyStructGenerator* Generator);
 MOVIESCENETOOLS_API UMovieSceneKeyStructType* InstanceGeneratedStruct(FMovieSceneObjectPathChannel* Channel, FSequencerKeyStructGenerator* Generator);
 
 MOVIESCENETOOLS_API void PostConstructKeyInstance(const TMovieSceneChannelHandle<FMovieSceneTimeWarpChannel>& ChannelHandle, FKeyHandle InHandle, FStructOnScope* Struct);
 MOVIESCENETOOLS_API void PostConstructKeyInstance(const TMovieSceneChannelHandle<FMovieSceneObjectPathChannel>& ChannelHandle, FKeyHandle InHandle, FStructOnScope* Struct);
 MOVIESCENETOOLS_API void PostConstructKeyInstance(const TMovieSceneChannelHandle<FMovieSceneDoubleChannel>& ChannelHandle, FKeyHandle InHandle, FStructOnScope* Struct);
-
+MOVIESCENETOOLS_API void PostConstructKeyInstance(const TMovieSceneChannelHandle<FMovieSceneTextChannel>& ChannelHandle, FKeyHandle InHandle, FStructOnScope* Struct);
 
 /** Key drawing overrides */
 MOVIESCENETOOLS_API void DrawKeys(FMovieSceneDoubleChannel*   Channel, TArrayView<const FKeyHandle> InKeyHandles, const UMovieSceneSection* InOwner, TArrayView<FKeyDrawParams> OutKeyDrawParams);
diff --git a/Engine/Source/Editor/MovieSceneTools/Public/ClipboardTypes.h b/Engine/Source/Editor/MovieSceneTools/Public/ClipboardTypes.h
--- a/Engine/Source/Editor/MovieSceneTools/Public/ClipboardTypes.h
+++ b/Engine/Source/Editor/MovieSceneTools/Public/ClipboardTypes.h
@@ -58,6 +58,11 @@ namespace MovieSceneClipboard
 		static FName Name("String");
 		return Name;
 	}
+	template<> inline FName GetKeyTypeName<FText>()
+	{
+		static FName Name("Text");
+		return Name;
+	}
 	template<> inline FName GetKeyTypeName<FMovieSceneActorReferenceKey>()
 	{
 		static FName Name("MovieSceneActorReferenceKey");
diff --git a/Engine/Source/Editor/MovieSceneTools/Public/KeyStructTypes/MovieSceneTextKeyStruct.h b/Engine/Source/Editor/MovieSceneTools/Public/KeyStructTypes/MovieSceneTextKeyStruct.h
new file mode 100644
--- /dev/null
+++ b/Engine/Source/Editor/MovieSceneTools/Public/KeyStructTypes/MovieSceneTextKeyStruct.h
@@ -0,0 +1,22 @@
+// Copyright Epic Games, Inc. All Rights Reserved.
+
+#pragma once
+
+#include "SequencerKeyStructGenerator.h"
+#include "MovieSceneTextKeyStruct.generated.h"
+
+/**
+ * Derived only to explicitly specify that it's Editor Only
+ * to be ignored in Cooked Builds as these KeyStructs are instanced under a non-editor package for text localization
+ * @see InstanceGeneratedStruct overload for FMovieSceneTextChannel
+ */
+UCLASS(MinimalAPI)
+class UMovieSceneTextKeyStruct : public UMovieSceneKeyStructType
+{
+	GENERATED_BODY()
+
+public:
+	//~ Begin UObject
+	virtual bool IsEditorOnly() const override { return true; }
+	//~ End UObject
+};
diff --git a/Engine/Source/Editor/MovieSceneTools/Public/MovieSceneToolsModule.h b/Engine/Source/Editor/MovieSceneTools/Public/MovieSceneToolsModule.h
--- a/Engine/Source/Editor/MovieSceneTools/Public/MovieSceneToolsModule.h
+++ b/Engine/Source/Editor/MovieSceneTools/Public/MovieSceneToolsModule.h
@@ -120,6 +120,7 @@ private:
 	FDelegateHandle VisibilityPropertyTrackCreateEditorHandle;
 	FDelegateHandle ActorReferencePropertyTrackCreateEditorHandle;
 	FDelegateHandle StringPropertyTrackCreateEditorHandle;
+	FDelegateHandle TextPropertyTrackCreateEditorHandle;
 	FDelegateHandle ObjectTrackCreateEditorHandle;
 
 	FDelegateHandle AnimationTrackCreateEditorHandle;
diff --git a/Engine/Source/Runtime/MovieScene/Private/EntitySystem/BuiltInComponentTypes.cpp b/Engine/Source/Runtime/MovieScene/Private/EntitySystem/BuiltInComponentTypes.cpp
--- a/Engine/Source/Runtime/MovieScene/Private/EntitySystem/BuiltInComponentTypes.cpp
+++ b/Engine/Source/Runtime/MovieScene/Private/EntitySystem/BuiltInComponentTypes.cpp
@@ -325,6 +325,7 @@ FBuiltInComponentTypes::FBuiltInComponentTypes()
 	ComponentRegistry->NewComponentType(&DoubleChannel[8],        TEXT("Double Channel 8"));
 	ComponentRegistry->NewComponentType(&WeightChannel,           TEXT("Weight Channel"));
 	ComponentRegistry->NewComponentType(&StringChannel,           TEXT("String Channel"));
+	ComponentRegistry->NewComponentType(&TextChannel,             TEXT("Text Channel"));
 	ComponentRegistry->NewComponentType(&ObjectPathChannel,       TEXT("Object Path Channel"));
 
 	ComponentRegistry->NewComponentType(&CachedInterpolation[0],   TEXT("Cached Interpolation [0]"));
@@ -367,6 +368,7 @@ FBuiltInComponentTypes::FBuiltInComponentTypes()
 	ComponentRegistry->NewComponentType(&DoubleResult[7],       TEXT("Double Result 7"));
 	ComponentRegistry->NewComponentType(&DoubleResult[8],       TEXT("Double Result 8"));
 	ComponentRegistry->NewComponentType(&StringResult,          TEXT("String Result"));
+	ComponentRegistry->NewComponentType(&TextResult,            TEXT("Text Result"));
 	ComponentRegistry->NewComponentType(&ObjectResult,          TEXT("Object Result"));
 
 	ComponentRegistry->NewComponentType(&BaseByte,			    TEXT("Base Byte"));
@@ -613,6 +615,15 @@ FBuiltInComponentTypes::FBuiltInComponentTypes()
 		ComponentRegistry->Factories.DefineMutuallyInclusiveComponent(StringChannel, StringResult);
 	}
 
+	// Text channel relationships
+	{
+		ComponentRegistry->Factories.DuplicateChildComponent(TextChannel);
+		ComponentRegistry->Factories.DuplicateChildComponent(TextResult);
+
+		ComponentRegistry->Factories.DefineMutuallyInclusiveComponent(TextChannel, EvalTime);
+		ComponentRegistry->Factories.DefineMutuallyInclusiveComponent(TextChannel, TextResult);
+	}
+
 	// Object path channel relationships
 	{
 		// Object channel components should be duplicated to children
diff --git a/Engine/Source/Runtime/MovieScene/Public/EntitySystem/BuiltInComponentTypes.h b/Engine/Source/Runtime/MovieScene/Public/EntitySystem/BuiltInComponentTypes.h
--- a/Engine/Source/Runtime/MovieScene/Public/EntitySystem/BuiltInComponentTypes.h
+++ b/Engine/Source/Runtime/MovieScene/Public/EntitySystem/BuiltInComponentTypes.h
@@ -54,6 +54,7 @@ struct FMovieSceneFloatChannel;
 struct FMovieSceneIntegerChannel;
 struct FMovieSceneObjectPathChannel;
 struct FMovieSceneStringChannel;
+struct FMovieSceneTextChannel;
 struct FMovieScenePropertyBinding;
 
 
@@ -215,6 +216,22 @@ struct FSourceStringChannel
 	const FMovieSceneStringChannel* Source;
 };
 
+/**
+ * The component data for evaluating a text channel
+ */
+struct FSourceTextChannel
+{
+	FSourceTextChannel()
+		: Source(nullptr)
+	{}
+
+	FSourceTextChannel(const FMovieSceneTextChannel* InSource)
+		: Source(InSource)
+	{}
+
+	const FMovieSceneTextChannel* Source;
+};
+
 struct FEvaluationHookFlags
 {
 	bool bHasBegun = false;
@@ -658,6 +675,9 @@ public:
 	// An FMovieSceneStringChannel
 	TComponentTypeID<FSourceStringChannel> StringChannel;
 
+	// An FMovieSceneTextChannel
+	TComponentTypeID<FSourceTextChannel> TextChannel;
+
 	// A cached interpolation structure relating to either float channels or double channels
 	TComponentTypeID<Interpolation::FCachedInterpolation> CachedInterpolation[9];
 
@@ -685,6 +705,9 @@ public:
 	// The result of an evaluated FMovieSceneStringChannel
 	TComponentTypeID<FString> StringResult;
 
+	// Result of an evaluated FMovieSceneTextChannel
+	TComponentTypeID<FText> TextResult;
+
 	// The result of an evaluated FMovieSceneObjectPathChannel
 	TComponentTypeID<FObjectComponent> ObjectResult;
 
diff --git a/Engine/Source/Runtime/MovieSceneTracks/Private/Channels/MovieSceneTextChannel.cpp b/Engine/Source/Runtime/MovieSceneTracks/Private/Channels/MovieSceneTextChannel.cpp
new file mode 100644
--- /dev/null
+++ b/Engine/Source/Runtime/MovieSceneTracks/Private/Channels/MovieSceneTextChannel.cpp
@@ -0,0 +1,108 @@
+// Copyright Epic Games, Inc. All Rights Reserved.
+
+#include "Channels/MovieSceneTextChannel.h"
+
+#if WITH_EDITOR
+#include "UObject/Package.h"
+#endif
+
+#if WITH_EDITOR
+void FMovieSceneTextChannel::SetPackage(UPackage* InPackage)
+{
+	Package = InPackage;
+}
+
+UPackage* FMovieSceneTextChannel::GetPackage() const
+{
+	return Package.Get();
+}
+#endif
+
+const FText* FMovieSceneTextChannel::Evaluate(FFrameTime InTime) const
+{
+	if (Times.Num())
+	{
+		const int32 Index = FMath::Max(0, Algo::UpperBound(Times, InTime.FrameNumber)-1);
+		return &Values[Index];
+	}
+	return bHasDefaultValue ? &DefaultValue : nullptr;
+}
+
+void FMovieSceneTextChannel::GetKeys(const TRange<FFrameNumber>& WithinRange
+	, TArray<FFrameNumber>* OutKeyTimes
+	, TArray<FKeyHandle>* OutKeyHandles)
+{
+	GetData().GetKeys(WithinRange, OutKeyTimes, OutKeyHandles);
+}
+
+void FMovieSceneTextChannel::GetKeyTimes(TArrayView<const FKeyHandle> InHandles, TArrayView<FFrameNumber> OutKeyTimes)
+{
+	GetData().GetKeyTimes(InHandles, OutKeyTimes);
+}
+
+void FMovieSceneTextChannel::SetKeyTimes(TArrayView<const FKeyHandle> InHandles, TArrayView<const FFrameNumber> InKeyTimes)
+{
+	GetData().SetKeyTimes(InHandles, InKeyTimes);
+}
+
+void FMovieSceneTextChannel::DuplicateKeys(TArrayView<const FKeyHandle> InHandles, TArrayView<FKeyHandle> OutNewHandles)
+{
+	GetData().DuplicateKeys(InHandles, OutNewHandles);
+}
+
+void FMovieSceneTextChannel::DeleteKeys(TArrayView<const FKeyHandle> InHandles)
+{
+	GetData().DeleteKeys(InHandles);
+}
+
+void FMovieSceneTextChannel::DeleteKeysFrom(FFrameNumber InTime, bool bDeleteKeysBefore)
+{
+	// Insert a key at the current time to maintain evaluation
+	if (GetData().GetTimes().Num() > 0)
+	{
+		if (const FText* CurrentValue = Evaluate(InTime))
+		{
+			FText Value(*CurrentValue);
+			GetData().UpdateOrAddKey(InTime, Value);
+		}
+	}
+	GetData().DeleteKeysFrom(InTime, bDeleteKeysBefore);
+}
+
+void FMovieSceneTextChannel::RemapTimes(const UE::MovieScene::IRetimingInterface& Retimer)
+{
+	GetData().RemapTimes(Retimer);
+}
+
+TRange<FFrameNumber> FMovieSceneTextChannel::ComputeEffectiveRange() const
+{
+	return GetData().GetTotalRange();
+}
+
+int32 FMovieSceneTextChannel::GetNumKeys() const
+{
+	return Times.Num();
+}
+
+void FMovieSceneTextChannel::Reset()
+{
+	Times.Reset();
+	Values.Reset();
+	KeyHandles.Reset();
+	bHasDefaultValue = false;
+}
+
+void FMovieSceneTextChannel::Offset(FFrameNumber DeltaPosition)
+{
+	GetData().Offset(DeltaPosition);
+}
+
+void FMovieSceneTextChannel::Optimize(const FKeyDataOptimizationParams& InParameters)
+{
+	UE::MovieScene::Optimize(this, InParameters);
+}
+
+void FMovieSceneTextChannel::ClearDefault()
+{
+	bHasDefaultValue = false;
+}
diff --git a/Engine/Source/Runtime/MovieSceneTracks/Private/MovieSceneTracksComponentTypes.cpp b/Engine/Source/Runtime/MovieSceneTracks/Private/MovieSceneTracksComponentTypes.cpp
--- a/Engine/Source/Runtime/MovieSceneTracks/Private/MovieSceneTracksComponentTypes.cpp
+++ b/Engine/Source/Runtime/MovieSceneTracks/Private/MovieSceneTracksComponentTypes.cpp
@@ -562,6 +562,7 @@ FMovieSceneTracksComponentTypes::FMovieSceneTracksComponentTypes()
 	ComponentRegistry->NewPropertyType(FloatVector, TEXT("float vector"));
 	ComponentRegistry->NewPropertyType(DoubleVector, TEXT("double vector"));
 	ComponentRegistry->NewPropertyType(String, TEXT("FString"));
+	ComponentRegistry->NewPropertyType(Text, TEXT("FText"));
 	ComponentRegistry->NewPropertyType(Object, TEXT("Object"));
 
 	ComponentRegistry->NewPropertyType(Transform, TEXT("FTransform"));
@@ -741,6 +742,12 @@ FMovieSceneTracksComponentTypes::FMovieSceneTracksComponentTypes()
 	.AddSoleChannel(BuiltInComponents->StringResult)
 	.Commit();
 
+	// --------------------------------------------------------------------------------------------
+	// Set up text properties
+	BuiltInComponents->PropertyRegistry.DefineProperty(Text, TEXT("Apply Text Properties"))
+	.AddSoleChannel(BuiltInComponents->TextResult)
+	.Commit();
+
 	// --------------------------------------------------------------------------------------------
 	// Set up Object properties
 	BuiltInComponents->PropertyRegistry.DefineProperty(Object, TEXT("Apply Object Properties"))
diff --git a/Engine/Source/Runtime/MovieSceneTracks/Private/Sections/MovieSceneTextSection.cpp b/Engine/Source/Runtime/MovieSceneTracks/Private/Sections/MovieSceneTextSection.cpp
new file mode 100644
--- /dev/null
+++ b/Engine/Source/Runtime/MovieSceneTracks/Private/Sections/MovieSceneTextSection.cpp
@@ -0,0 +1,46 @@
+// Copyright Epic Games, Inc. All Rights Reserved.
+
+#include "Sections/MovieSceneTextSection.h"
+#include "Channels/MovieSceneTextChannel.h"
+#include "Tracks/MovieScenePropertyTrack.h"
+
+UMovieSceneTextSection::UMovieSceneTextSection(const FObjectInitializer& InObjectInitializer)
+	: Super(InObjectInitializer)
+{
+	bSupportsInfiniteRange = true;
+	EvalOptions.EnableAndSetCompletionMode(EMovieSceneCompletionMode::ProjectDefault);
+
+#if WITH_EDITOR
+	ChannelProxy = MakeShared<FMovieSceneChannelProxy>(TextChannel, FMovieSceneChannelMetaData(), TMovieSceneExternalValue<FText>::Make());
+	TextChannel.SetPackage(GetPackage());
+#else
+	ChannelProxy = MakeShared<FMovieSceneChannelProxy>(TextChannel);
+#endif
+}
+
+bool UMovieSceneTextSection::PopulateEvaluationFieldImpl(const TRange<FFrameNumber>& EffectiveRange
+	, const FMovieSceneEvaluationFieldEntityMetaData& InMetaData
+	, FMovieSceneEntityComponentFieldBuilder* OutFieldBuilder)
+{
+	FMovieScenePropertyTrackEntityImportHelper::PopulateEvaluationField(*this, EffectiveRange, InMetaData, OutFieldBuilder);
+	return true;
+}
+
+void UMovieSceneTextSection::ImportEntityImpl(UMovieSceneEntitySystemLinker* EntityLinker
+	, const FEntityImportParams& Params
+	, FImportedEntity* OutImportedEntity)
+{
+	using namespace UE::MovieScene;
+
+	if (!TextChannel.HasAnyData())
+	{
+		return;
+	}
+
+	const FBuiltInComponentTypes* Components = FBuiltInComponentTypes::Get();
+	const FMovieSceneTracksComponentTypes* TracksComponents = FMovieSceneTracksComponentTypes::Get();
+
+	FPropertyTrackEntityImportHelper(TracksComponents->Text)
+		.Add(Components->TextChannel, &TextChannel)
+		.Commit(this, Params, OutImportedEntity);
+}
diff --git a/Engine/Source/Runtime/MovieSceneTracks/Private/Systems/MovieSceneTextPropertySystem.cpp b/Engine/Source/Runtime/MovieSceneTracks/Private/Systems/MovieSceneTextPropertySystem.cpp
new file mode 100644
--- /dev/null
+++ b/Engine/Source/Runtime/MovieSceneTracks/Private/Systems/MovieSceneTextPropertySystem.cpp
@@ -0,0 +1,20 @@
+// Copyright Epic Games, Inc. All Rights Reserved.
+
+#include "Systems/MovieSceneTextPropertySystem.h"
+#include "MovieSceneTracksComponentTypes.h"
+#include "Systems/TextChannelEvaluatorSystem.h"
+
+UMovieSceneTextPropertySystem::UMovieSceneTextPropertySystem(const FObjectInitializer& InObjectInitializer)
+	: Super(InObjectInitializer)
+{
+	BindToProperty(UE::MovieScene::FMovieSceneTracksComponentTypes::Get()->Text);
+	if (HasAnyFlags(RF_ClassDefaultObject))
+	{
+		DefineImplicitPrerequisite(UTextChannelEvaluatorSystem::StaticClass(), GetClass());
+	}
+}
+
+void UMovieSceneTextPropertySystem::OnRun(FSystemTaskPrerequisites& InPrerequisites, FSystemSubsequentTasks& Subsequents)
+{
+	Super::OnRun(InPrerequisites, Subsequents);
+}
diff --git a/Engine/Source/Runtime/MovieSceneTracks/Private/Systems/TextChannelEvaluatorSystem.cpp b/Engine/Source/Runtime/MovieSceneTracks/Private/Systems/TextChannelEvaluatorSystem.cpp
new file mode 100644
--- /dev/null
+++ b/Engine/Source/Runtime/MovieSceneTracks/Private/Systems/TextChannelEvaluatorSystem.cpp
@@ -0,0 +1,73 @@
+// Copyright Epic Games, Inc. All Rights Reserved.
+
+#include "Systems/TextChannelEvaluatorSystem.h"
+#include "Channels/MovieSceneTextChannel.h"
+#include "EntitySystem/BuiltInComponentTypes.h"
+#include "EntitySystem/MovieSceneEntitySystemLinker.h"
+#include "EntitySystem/MovieSceneEntitySystemTask.h"
+#include "EntitySystem/MovieSceneEvalTimeSystem.h"
+
+DECLARE_CYCLE_STAT(TEXT("MovieScene: Evaluate text channels"), MovieSceneEval_EvaluateTextChannelTask, STATGROUP_MovieSceneECS);
+
+namespace UE::MovieScene::Private
+{
+	struct FEvaluateTextChannels
+	{
+		static void ForEachEntity(FSourceTextChannel TextChannel, FFrameTime FrameTime, FText& OutResult)
+		{
+			if (const FText* Value = TextChannel.Source->Evaluate(FrameTime))
+			{
+				OutResult = *Value;
+			}
+			else
+			{
+				OutResult = FText::GetEmpty();
+			}
+		}
+	};
+}
+
+UTextChannelEvaluatorSystem::UTextChannelEvaluatorSystem(const FObjectInitializer& InObjectInitializer)
+	: Super(InObjectInitializer)
+{
+	using namespace UE::MovieScene;
+
+	SystemCategories  = EEntitySystemCategory::ChannelEvaluators;
+	RelevantComponent = FBuiltInComponentTypes::Get()->TextChannel;
+	Phase             = ESystemPhase::Scheduling;
+
+	if (HasAnyFlags(RF_ClassDefaultObject))
+	{
+		DefineImplicitPrerequisite(UMovieSceneEvalTimeSystem::StaticClass(), GetClass());
+	}
+}
+
+void UTextChannelEvaluatorSystem::OnSchedulePersistentTasks(UE::MovieScene::IEntitySystemScheduler* TaskScheduler)
+{
+	using namespace UE::MovieScene;
+
+	const FBuiltInComponentTypes* BuiltInComponents = FBuiltInComponentTypes::Get();
+
+	FEntityTaskBuilder()
+		.Read(BuiltInComponents->TextChannel)
+		.Read(BuiltInComponents->EvalTime)
+		.Write(BuiltInComponents->TextResult)
+		.FilterNone({ BuiltInComponents->Tags.Ignored })
+		.SetStat(GET_STATID(MovieSceneEval_EvaluateTextChannelTask))
+		.Fork_PerEntity<Private::FEvaluateTextChannels>(&Linker->EntityManager, TaskScheduler);
+}
+
+void UTextChannelEvaluatorSystem::OnRun(FSystemTaskPrerequisites& InPrerequisites, FSystemSubsequentTasks& Subsequents)
+{
+	using namespace UE::MovieScene;
+
+	const FBuiltInComponentTypes* BuiltInComponents = FBuiltInComponentTypes::Get();
+
+	FEntityTaskBuilder()
+		.Read(BuiltInComponents->TextChannel)
+		.Read(BuiltInComponents->EvalTime)
+		.Write(BuiltInComponents->TextResult)
+		.FilterNone({ BuiltInComponents->Tags.Ignored })
+		.SetStat(GET_STATID(MovieSceneEval_EvaluateTextChannelTask))
+		.Dispatch_PerEntity<Private::FEvaluateTextChannels>(&Linker->EntityManager, InPrerequisites, &Subsequents);
+}
diff --git a/Engine/Source/Runtime/MovieSceneTracks/Private/Tracks/MovieSceneTextTrack.cpp b/Engine/Source/Runtime/MovieSceneTracks/Private/Tracks/MovieSceneTextTrack.cpp
new file mode 100644
--- /dev/null
+++ b/Engine/Source/Runtime/MovieSceneTracks/Private/Tracks/MovieSceneTextTrack.cpp
@@ -0,0 +1,53 @@
+// Copyright Epic Games, Inc. All Rights Reserved.
+
+#include "Tracks/MovieSceneTextTrack.h"
+#include "Sections/MovieSceneTextSection.h"
+
+UMovieSceneTextTrack::UMovieSceneTextTrack()
+{
+}
+
+bool UMovieSceneTextTrack::SupportsType(TSubclassOf<UMovieSceneSection> SectionClass) const
+{
+	return SectionClass == UMovieSceneTextSection::StaticClass();
+}
+
+UMovieSceneSection* UMovieSceneTextTrack::CreateNewSection()
+{
+	return NewObject<UMovieSceneTextSection>(this, NAME_None, RF_Transactional);
+}
+
+void UMovieSceneTextTrack::AddSection(UMovieSceneSection& Section)
+{
+	Sections.Add(&Section);
+}
+
+const TArray<UMovieSceneSection*>& UMovieSceneTextTrack::GetAllSections() const
+{
+	return Sections;
+}
+
+bool UMovieSceneTextTrack::HasSection(const UMovieSceneSection& Section) const
+{
+	return Sections.Contains(&Section);
+}
+
+bool UMovieSceneTextTrack::IsEmpty() const
+{
+	return Sections.IsEmpty();
+}
+
+void UMovieSceneTextTrack::RemoveAllAnimationData()
+{
+	Sections.Empty();
+}
+
+void UMovieSceneTextTrack::RemoveSection(UMovieSceneSection& Section)
+{
+	Sections.Remove(&Section);
+}
+
+void UMovieSceneTextTrack::RemoveSectionAt(int32 SectionIndex)
+{
+	Sections.RemoveAt(SectionIndex);
+}
diff --git a/Engine/Source/Runtime/MovieSceneTracks/Public/Channels/MovieSceneTextChannel.h b/Engine/Source/Runtime/MovieSceneTracks/Public/Channels/MovieSceneTextChannel.h
new file mode 100644
--- /dev/null
+++ b/Engine/Source/Runtime/MovieSceneTracks/Public/Channels/MovieSceneTextChannel.h
@@ -0,0 +1,172 @@
+// Copyright Epic Games, Inc. All Rights Reserved.
+
+#pragma once
+
+#include "Channels/MovieSceneChannel.h"
+#include "Channels/MovieSceneChannelData.h"
+#include "Channels/MovieSceneChannelTraits.h"
+
+#if WITH_EDITOR
+#include "UObject/WeakObjectPtr.h"
+#endif
+
+#include "MovieSceneTextChannel.generated.h"
+
+#if WITH_EDITOR
+class UPackage;
+#endif
+
+USTRUCT()
+struct FMovieSceneTextChannel : public FMovieSceneChannel
+{
+	GENERATED_BODY()
+
+	typedef FText CurveValueType;
+
+#if WITH_EDITOR
+	MOVIESCENETRACKS_API void SetPackage(UPackage* InPackage);
+
+	MOVIESCENETRACKS_API UPackage* GetPackage() const;
+#endif
+
+	/**
+	 * Access a mutable interface for this channel's data
+	 * @return An object that is able to manipulate this channel's data
+	 */
+	inline TMovieSceneChannelData<FText> GetData()
+	{
+		return TMovieSceneChannelData<FText>(&Times, &Values, this, &KeyHandles);
+	}
+
+	/**
+	 * Access a constant interface for this channel's data
+	 * @return An object that is able to interrogate this channel's data
+	 */
+	inline TMovieSceneChannelData<const FText> GetData() const
+	{
+		return TMovieSceneChannelData<const FText>(&Times, &Values);
+	}
+
+	inline TConstArrayView<FFrameNumber> GetTimes() const
+	{
+		return Times;
+	}
+
+	/** Returns whether this channel has any values */
+	inline bool HasAnyData() const
+	{
+		return Times.Num() != 0 || bHasDefaultValue == true;
+	}
+
+	/**
+	 * Evaluate this channel
+	 * @param InTime     The time to evaluate at
+	 * @return A pointer to the string, or nullptr
+	 */
+	MOVIESCENETRACKS_API const FText* Evaluate(FFrameTime InTime) const;
+
+	/**
+	 * Add keys with these times to channel. The number of elements in both arrays much match or nothing is added.
+	 * @param InTimes Times to add
+	 * @param InValues Values to add
+	 */
+	inline void AddKeys(const TArray<FFrameNumber>& InTimes, const TArray<FText>& InValues)
+	{
+		check(InTimes.Num() == InValues.Num());
+		Times.Append(InTimes);
+		Values.Append(InValues);
+	}
+
+	/**
+	 * Set this channel's default value that should be used when no keys are present
+	 * @param InDefaultValue The desired default value
+	 */
+	inline void SetDefault(FText InDefaultValue)
+	{
+		bHasDefaultValue = true;
+		DefaultValue = InDefaultValue;
+	}
+
+	/**
+	 * Get this channel's default value that will be used when no keys are present
+	 * @return (Optional) The channel's default value
+	 */
+	inline TOptional<FText> GetDefault() const
+	{
+		return bHasDefaultValue ? TOptional<FText>(DefaultValue) : TOptional<FText>();
+	}
+
+	/** Remove this channel's default value causing the channel to have no effect where no keys are present */
+	inline void RemoveDefault()
+	{
+		bHasDefaultValue = false;
+	}
+
+	//~ Begin FMovieSceneChannel
+	MOVIESCENETRACKS_API virtual void GetKeys(const TRange<FFrameNumber>& WithinRange, TArray<FFrameNumber>* OutKeyTimes, TArray<FKeyHandle>* OutKeyHandles) override;
+	MOVIESCENETRACKS_API virtual void GetKeyTimes(TArrayView<const FKeyHandle> InHandles, TArrayView<FFrameNumber> OutKeyTimes) override;
+	MOVIESCENETRACKS_API virtual void SetKeyTimes(TArrayView<const FKeyHandle> InHandles, TArrayView<const FFrameNumber> InKeyTimes) override;
+	MOVIESCENETRACKS_API virtual void DuplicateKeys(TArrayView<const FKeyHandle> InHandles, TArrayView<FKeyHandle> OutNewHandles) override;
+	MOVIESCENETRACKS_API virtual void DeleteKeys(TArrayView<const FKeyHandle> InHandles) override;
+	MOVIESCENETRACKS_API virtual void DeleteKeysFrom(FFrameNumber InTime, bool bDeleteKeysBefore) override;
+	MOVIESCENETRACKS_API virtual void RemapTimes(const UE::MovieScene::IRetimingInterface& Retimer) override;
+	MOVIESCENETRACKS_API virtual TRange<FFrameNumber> ComputeEffectiveRange() const override;
+	MOVIESCENETRACKS_API virtual int32 GetNumKeys() const override;
+	MOVIESCENETRACKS_API virtual void Reset() override;
+	MOVIESCENETRACKS_API virtual void Offset(FFrameNumber DeltaPosition) override;
+	MOVIESCENETRACKS_API virtual void Optimize(const FKeyDataOptimizationParams& InParameters) override;
+	MOVIESCENETRACKS_API virtual void ClearDefault() override;
+	//~ End FMovieSceneChannel
+
+private:
+	UPROPERTY(meta=(KeyTimes))
+	TArray<FFrameNumber> Times;
+
+	/** Array of values that correspond to each key time */
+	UPROPERTY(meta=(KeyValues))
+	TArray<FText> Values;
+
+	/** Default value used when there are no keys */
+	UPROPERTY()
+	FText DefaultValue;
+
+	UPROPERTY()
+	bool bHasDefaultValue = false;
+
+	/** This needs to be a UPROPERTY so it gets saved into editor transactions but transient so it does not get saved into assets. */
+	UPROPERTY(Transient)
+	FMovieSceneKeyHandleMap KeyHandles;
+
+#if WITH_EDITOR
+	/** Package to use for Text Localization. Always set to the Package of the Section holding this Channel. */
+	TWeakObjectPtr<UPackage> Package;
+#endif
+};
+
+template<>
+struct TMovieSceneChannelTraits<FMovieSceneTextChannel> : TMovieSceneChannelTraitsBase<FMovieSceneTextChannel>
+{
+#if WITH_EDITOR
+	/** Text channels can have external values (ie, they can get their values from external objects for UI purposes) */
+	typedef TMovieSceneExternalValue<FText> ExtendedEditorDataType;
+#endif
+};
+
+inline bool EvaluateChannel(const FMovieSceneTextChannel* InChannel, FFrameTime InTime, FText& OutValue)
+{
+	if (const FText* Result = InChannel->Evaluate(InTime))
+	{
+		OutValue = *Result;
+		return true;
+	}
+	return false;
+}
+
+inline bool ValueExistsAtTime(const FMovieSceneTextChannel* Channel, FFrameNumber InFrameNumber, const FText& Value)
+{
+	const FFrameTime FrameTime(InFrameNumber);
+
+	const FText* ExistingValue = Channel->Evaluate(FrameTime);
+
+	return ExistingValue && Value.EqualTo(*ExistingValue);
+}
diff --git a/Engine/Source/Runtime/MovieSceneTracks/Public/MovieSceneTracksComponentTypes.h b/Engine/Source/Runtime/MovieSceneTracks/Public/MovieSceneTracksComponentTypes.h
--- a/Engine/Source/Runtime/MovieSceneTracks/Public/MovieSceneTracksComponentTypes.h
+++ b/Engine/Source/Runtime/MovieSceneTracks/Public/MovieSceneTracksComponentTypes.h
@@ -657,6 +657,7 @@ using FEulerTransformPropertyTraits     = TIndirectPropertyTraits<FEulerTransfor
 using FComponentTransformPropertyTraits = TDirectPropertyTraits<FIntermediate3DTransform>;
 using FRotatorPropertyTraits            = TDirectPropertyTraits<FRotator>;
 using FStringPropertyTraits             = TDirectPropertyTraits<FString>;
+using FTextPropertyTraits               = TDirectPropertyTraits<FText>;
 
 using FBoolParameterTraits              = TDirectPropertyTraits<bool, false>;
 using FByteParameterTraits              = TDirectPropertyTraits<uint8, false>;
@@ -685,6 +686,7 @@ struct FMovieSceneTracksComponentTypes
 	TPropertyComponents<FComponentTransformPropertyTraits> ComponentTransform;
 	TPropertyComponents<FRotatorPropertyTraits> Rotator;
 	TPropertyComponents<FStringPropertyTraits> String;
+	TPropertyComponents<FTextPropertyTraits> Text;
 	TPropertyComponents<FObjectPropertyTraits> Object;
 
 	TPropertyComponents<FFloatParameterTraits> FloatParameter;
diff --git a/Engine/Source/Runtime/MovieSceneTracks/Public/Sections/MovieSceneTextSection.h b/Engine/Source/Runtime/MovieSceneTracks/Public/Sections/MovieSceneTextSection.h
new file mode 100644
--- /dev/null
+++ b/Engine/Source/Runtime/MovieSceneTracks/Public/Sections/MovieSceneTextSection.h
@@ -0,0 +1,34 @@
+// Copyright Epic Games, Inc. All Rights Reserved.
+
+#pragma once
+
+#include "Channels/MovieSceneTextChannel.h"
+#include "EntitySystem/IMovieSceneEntityProvider.h"
+#include "MovieSceneSection.h"
+#include "MovieSceneTextSection.generated.h"
+
+UCLASS(MinimalAPI)
+class UMovieSceneTextSection
+	: public UMovieSceneSection
+	, public IMovieSceneEntityProvider
+{
+	GENERATED_BODY()
+
+public:
+	UMovieSceneTextSection(const FObjectInitializer& InObjectInitializer);
+
+	/** Public access to this section's internal data function */
+	const FMovieSceneTextChannel& GetChannel() const
+	{
+		return TextChannel;
+	}
+
+	//~ Begin IMovieSceneEntityProvider
+	virtual bool PopulateEvaluationFieldImpl(const TRange<FFrameNumber>& EffectiveRange, const FMovieSceneEvaluationFieldEntityMetaData& InMetaData, FMovieSceneEntityComponentFieldBuilder* OutFieldBuilder) override;
+	virtual void ImportEntityImpl(UMovieSceneEntitySystemLinker* EntityLinker, const FEntityImportParams& Params, FImportedEntity* OutImportedEntity) override;
+	//~ End IMovieSceneEntityProvider
+
+private:
+	UPROPERTY()
+	FMovieSceneTextChannel TextChannel;
+};
diff --git a/Engine/Source/Runtime/MovieSceneTracks/Public/Systems/MovieSceneTextPropertySystem.h b/Engine/Source/Runtime/MovieSceneTracks/Public/Systems/MovieSceneTextPropertySystem.h
new file mode 100644
--- /dev/null
+++ b/Engine/Source/Runtime/MovieSceneTracks/Public/Systems/MovieSceneTextPropertySystem.h
@@ -0,0 +1,19 @@
+// Copyright Epic Games, Inc. All Rights Reserved.
+
+#pragma once
+
+#include "Systems/MovieScenePropertySystem.h"
+#include "MovieSceneTextPropertySystem.generated.h"
+
+UCLASS(MinimalAPI)
+class UMovieSceneTextPropertySystem : public UMovieScenePropertySystem
+{
+	GENERATED_BODY()
+
+public:
+	UMovieSceneTextPropertySystem(const FObjectInitializer& InObjectInitializer);
+
+	//~ Begin UMovieSceneEntitySystem
+	MOVIESCENETRACKS_API virtual void OnRun(FSystemTaskPrerequisites& InPrerequisites, FSystemSubsequentTasks& Subsequents) override;
+	//~ End UMovieSceneEntitySystem
+};
diff --git a/Engine/Source/Runtime/MovieSceneTracks/Public/Systems/TextChannelEvaluatorSystem.h b/Engine/Source/Runtime/MovieSceneTracks/Public/Systems/TextChannelEvaluatorSystem.h
new file mode 100644
--- /dev/null
+++ b/Engine/Source/Runtime/MovieSceneTracks/Public/Systems/TextChannelEvaluatorSystem.h
@@ -0,0 +1,21 @@
+// Copyright Epic Games, Inc. All Rights Reserved.
+
+#pragma once
+
+#include "EntitySystem/MovieSceneEntitySystem.h"
+#include "TextChannelEvaluatorSystem.generated.h"
+
+/** System responsible for evaluating Text channels */
+UCLASS(MinimalAPI)
+class UTextChannelEvaluatorSystem : public UMovieSceneEntitySystem
+{
+	GENERATED_BODY()
+
+public:
+	UTextChannelEvaluatorSystem(const FObjectInitializer& InObjectInitializer);
+
+	//~ Begin UMovieSceneEntitySystem
+	virtual void OnSchedulePersistentTasks(UE::MovieScene::IEntitySystemScheduler* TaskScheduler) override;
+	virtual void OnRun(FSystemTaskPrerequisites& InPrerequisites, FSystemSubsequentTasks& Subsequents) override;
+	//~ End UMovieSceneEntitySystem
+};
diff --git a/Engine/Source/Runtime/MovieSceneTracks/Public/Tracks/MovieSceneTextTrack.h b/Engine/Source/Runtime/MovieSceneTracks/Public/Tracks/MovieSceneTextTrack.h
new file mode 100644
--- /dev/null
+++ b/Engine/Source/Runtime/MovieSceneTracks/Public/Tracks/MovieSceneTextTrack.h
@@ -0,0 +1,28 @@
+// Copyright Epic Games, Inc. All Rights Reserved.
+
+#pragma once
+
+#include "Tracks/MovieScenePropertyTrack.h"
+#include "MovieSceneTextTrack.generated.h"
+
+/** Movie Scene Track that holds a series of texts. */
+UCLASS(MinimalAPI)
+class UMovieSceneTextTrack : public UMovieScenePropertyTrack
+{
+	GENERATED_BODY()
+
+public:
+	UMovieSceneTextTrack();
+
+	//~ Begin UMovieSceneTrack
+	virtual bool SupportsType(TSubclassOf<UMovieSceneSection> SectionClass) const override;
+	virtual UMovieSceneSection* CreateNewSection() override;
+	virtual void AddSection(UMovieSceneSection& Section) override;
+	virtual const TArray<UMovieSceneSection*>& GetAllSections() const override;
+	virtual bool HasSection(const UMovieSceneSection& Section) const override;
+	virtual bool IsEmpty() const override;
+	virtual void RemoveAllAnimationData() override;
+	virtual void RemoveSection(UMovieSceneSection& Section) override;
+	virtual void RemoveSectionAt(int32 SectionIndex) override;
+	//~ End UMovieSceneTrack
+};
