Sequencer: Fix event key endpoint not being propagated

In 44520263, the struct property change propagation was moved from OnFinishedChangingProperties to NotifyPostChange so that it can run before PostTransacted. A few customizations didn't follow the pattern of calling NotifyPreChange/NotifyPostChange/>NotifyFinishedChangingProperties so the property change propagation was missing for these.

Tested that creating a new endpoint and assigning an endpoint through the key editor persists.

#rnx

diff --git a/Engine/Source/Editor/MovieSceneTools/Private/Conditions/MovieSceneConditionCustomization.cpp b/Engine/Source/Editor/MovieSceneTools/Private/Conditions/MovieSceneConditionCustomization.cpp
--- a/Engine/Source/Editor/MovieSceneTools/Private/Conditions/MovieSceneConditionCustomization.cpp
+++ b/Engine/Source/Editor/MovieSceneTools/Private/Conditions/MovieSceneConditionCustomization.cpp
@@ -329,11 +329,13 @@ void FMovieSceneConditionCustomization::FillDirectorBlueprintConditionSubMenu(FM
 						const FScopedTransaction Transaction(LOCTEXT("CreateNewConditionEndpoint", "Create New Condition Endpoint"));
 
 						ThisSequence->Modify();
+						SharedThis->ConditionPropertyHandle->NotifyPreChange();
 						// Create a new director blueprint condition and set it in the details view. Use 'interactive change' so we don't early fire the property finished changing event and reset the details view mid-change
 						PropertyCustomizationHelpers::CreateNewInstanceOfEditInlineObjectClass(SharedThis->ConditionPropertyHandle.ToSharedRef(), UMovieSceneDirectorBlueprintCondition::StaticClass(), EPropertyValueSetFlags::InteractiveChange);
 						TSharedPtr<IPropertyHandle> DirectorBlueprintConditionHandle = SharedThis->ConditionPropertyHandle->GetChildHandle(TEXT("DirectorBlueprintConditionData"));
 						TSharedPtr<FMovieSceneDirectorBlueprintConditionCustomization> BlueprintConditionCustomization = FMovieSceneDirectorBlueprintConditionCustomization::MakeInstance(ThisSequence->GetMovieScene(), DirectorBlueprintConditionHandle, SharedThis->PropertyUtilities);
 						BlueprintConditionCustomization->CreateEndpoint();
+						SharedThis->ConditionPropertyHandle->NotifyPostChange(EPropertyChangeType::ValueSet);
 						SharedThis->PropertyUtilities->NotifyFinishedChangingProperties(FPropertyChangedEvent(SharedThis->ConditionPropertyHandle->GetProperty()));
 						// Extra end transaction because we use 'Interactive Change' in the CreateNewInstance call
 						GEditor->EndTransaction();
@@ -407,7 +409,9 @@ void FMovieSceneConditionCustomization::OnUseSelected()
 				const FScopedTransaction Transaction(LOCTEXT("SetConditionClass", "Set Condition Class"));
 
 				Sequence->Modify();
+				ConditionPropertyHandle->NotifyPreChange();
 				PropertyCustomizationHelpers::CreateNewInstanceOfEditInlineObjectClass(ConditionPropertyHandle.ToSharedRef(), SelectedBlueprint->GeneratedClass, EPropertyValueSetFlags::InteractiveChange);
+				ConditionPropertyHandle->NotifyPostChange(EPropertyChangeType::ValueSet);
 				PropertyUtilities->NotifyFinishedChangingProperties(FPropertyChangedEvent(ConditionPropertyHandle->GetProperty()));
 				// Extra end transaction because we use 'Interactive Change' in the CreateNewInstance call
 				GEditor->EndTransaction();
diff --git a/Engine/Source/Editor/MovieSceneTools/Private/MovieSceneDirectorBlueprintEndpointCustomization.cpp b/Engine/Source/Editor/MovieSceneTools/Private/MovieSceneDirectorBlueprintEndpointCustomization.cpp
--- a/Engine/Source/Editor/MovieSceneTools/Private/MovieSceneDirectorBlueprintEndpointCustomization.cpp
+++ b/Engine/Source/Editor/MovieSceneTools/Private/MovieSceneDirectorBlueprintEndpointCustomization.cpp
@@ -747,6 +747,11 @@ void FMovieSceneDirectorBlueprintEndpointCustomization::SetWellKnownParameterPin
 {
 	FScopedTransaction Transaction(LOCTEXT("SetWellKnownParameterPinTransaction", "Set Bound Object Pin"));
 
+	if (PropertyHandle.IsValid())
+	{
+		PropertyHandle->NotifyPreChange();
+	}
+
 	TArray<UObject*> EditObjects;
 	GetEditObjects(EditObjects);
 
@@ -764,6 +769,7 @@ void FMovieSceneDirectorBlueprintEndpointCustomization::SetWellKnownParameterPin
 	if (PropertyHandle.IsValid())
 	{
 		// Ensure that anything listening for property changed notifications are notified of the new binding
+		PropertyHandle->NotifyPostChange(EPropertyChangeType::ValueSet);
 		PropertyHandle->NotifyFinishedChangingProperties();
 	}
 
@@ -1349,6 +1355,11 @@ void FMovieSceneDirectorBlueprintEndpointCustomization::SetEndpoint(const FMovie
 		return;
 	}
 
+	if (PropertyHandle.IsValid())
+	{
+		PropertyHandle->NotifyPreChange();
+	}
+
 	UEdGraphPin* CallTargetPin = FMovieSceneDirectorBlueprintUtils::FindCallTargetPin(NewEndpoint, EndpointDefinition.PossibleCallTargetClass);
 
 	// Map of the Payload Variable Names to their Default Values as Strings
@@ -1449,6 +1460,7 @@ void FMovieSceneDirectorBlueprintEndpointCustomization::SetEndpoint(const FMovie
 	if (PropertyHandle.IsValid())
 	{
 		// Ensure that anything listening for property changed notifications are notified of the new binding
+		PropertyHandle->NotifyPostChange(EPropertyChangeType::ValueSet);
 		PropertyHandle->NotifyFinishedChangingProperties();
 	}
 
@@ -1498,6 +1510,11 @@ void FMovieSceneDirectorBlueprintEndpointCustomization::CreateEndpoint()
 	// Create user facing endpoint
 	FScopedTransaction Transaction(LOCTEXT("CreateEndpoint", "Create Endpoint"));
 
+	if (PropertyHandle.IsValid())
+	{
+		PropertyHandle->NotifyPreChange();
+	}
+
 	UK2Node* LastNewEndpoint = nullptr;
 	TArray<UBlueprint*> BlueprintsToRecompile;
 
@@ -1536,6 +1553,7 @@ void FMovieSceneDirectorBlueprintEndpointCustomization::CreateEndpoint()
 	if (PropertyHandle.IsValid())
 	{
 		// Ensure that anything listening for property changed notifications are notified of the new binding
+		PropertyHandle->NotifyPostChange(EPropertyChangeType::ValueSet);
 		PropertyHandle->NotifyFinishedChangingProperties();
 	}
 
