Sequencer: Add option to bake with time warp applied

This option gets the unwarped local time during baking and sets keys with those unwarped times.

diff --git a/Engine/Source/Editor/MovieSceneTools/Public/BakingAnimationKeySettings.h b/Engine/Source/Editor/MovieSceneTools/Public/BakingAnimationKeySettings.h
--- a/Engine/Source/Editor/MovieSceneTools/Public/BakingAnimationKeySettings.h
+++ b/Engine/Source/Editor/MovieSceneTools/Public/BakingAnimationKeySettings.h
@@ -25,6 +25,7 @@ struct FBakingAnimationKeySettings
 		FrameIncrement = 1;
 		bReduceKeys = false;
 		Tolerance = 0.001f;
+		bTimeWarp = false;
 	}
 
 	UPROPERTY(BlueprintReadWrite, EditAnywhere, Category = "Bake")
@@ -33,15 +34,23 @@ struct FBakingAnimationKeySettings
 	UPROPERTY(BlueprintReadWrite, EditAnywhere, Category = "Bake")
 	FFrameNumber EndFrame;
 
+	/** Bake on keyed frames only or bake all frames between start and end */
 	UPROPERTY(BlueprintReadWrite, EditAnywhere, Category = "Bake")
 	EBakingKeySettings BakingKeySettings;
 
+	/** Frames to increment when baking */
 	UPROPERTY(BlueprintReadWrite, EditAnywhere, Category = "Bake", meta = (ClampMin = "1", UIMin = "1", EditCondition = "BakingKeySettings == EBakingKeySettings::AllFrames"))
 	int32 FrameIncrement;
 
+	/** Reduce keys after baking */
 	UPROPERTY(BlueprintReadWrite, EditAnywhere, Category = "Bake", meta = (EditCondition = "BakingKeySettings == EBakingKeySettings::AllFrames"))
 	bool bReduceKeys;
 
+	/** Tolerance to use when reducing keys */
 	UPROPERTY(BlueprintReadWrite, EditAnywhere, Category = "Bake", meta = (EditCondition = "BakingKeySettings == EBakingKeySettings::AllFrames || bReduceKeys"))
 	float Tolerance;
+
+	/** Bake with time warp applied. If there is time warp on the sequence, you will need to disable it manually after baking, otherwise times will be warped twice. */
+	UPROPERTY(BlueprintReadWrite, EditAnywhere, Category = "Bake")
+	bool bTimeWarp;
 };
diff --git a/Engine/Source/Editor/Sequencer/Private/SequencerUtilities.cpp b/Engine/Source/Editor/Sequencer/Private/SequencerUtilities.cpp
--- a/Engine/Source/Editor/Sequencer/Private/SequencerUtilities.cpp
+++ b/Engine/Source/Editor/Sequencer/Private/SequencerUtilities.cpp
@@ -4402,6 +4402,7 @@ bool FSequencerUtilities::BakeTransform(TSharedRef<ISequencer> Sequencer, const
 	ensure(PlaybackContext);
 	const FConstraintsManagerController& Controller = FConstraintsManagerController::Get(PlaybackContext);
 
+	TArray<FFrameNumber> RemappedTimes;
 	for (FFrameNumber KeyTime : AllFrames)
 	{
 		TOptional<FFrameTime> NewGlobalTime = LocalToRootTransform.TryTransformTime(KeyTime);
@@ -4413,6 +4414,16 @@ bool FSequencerUtilities::BakeTransform(TSharedRef<ISequencer> Sequencer, const
 		Sequencer->SetGlobalTime(NewGlobalTime.GetValue());
 		Controller.EvaluateAllConstraints();
 
+		if (InSettings.bTimeWarp)
+		{
+			FQualifiedFrameTime RemappedTime = Sequencer->GetUnwarpedLocalTime();
+			RemappedTimes.Add(RemappedTime.Time.FrameNumber);
+		}
+		else
+		{
+			RemappedTimes.Add(KeyTime);
+		}
+
 		for (const FMovieSceneBindingProxy& ObjectBinding : ObjectBindings)
 		{
 			FGuid Guid = ObjectBinding.BindingID;
@@ -4501,8 +4512,8 @@ bool FSequencerUtilities::BakeTransform(TSharedRef<ISequencer> Sequencer, const
 	for (TPair<FGuid, FBakeData>& BakeData : BakeDataMap)
 	{
 		FGuid Guid = BakeData.Key;
-		TArray<FFrameNumber> KeyTimes;
-		BakeData.Value.KeyTimes.GenerateKeyArray(KeyTimes);
+		TArray<FFrameNumber> KeyTimes = RemappedTimes;
+		
 		// Disable or delete any constraint (attach/path) tracks
 		AActor* AttachParentActor = nullptr;
 		for (UMovieSceneTrack* Track : FocusedMovieScene->FindTracks(UMovieScene3DConstraintTrack::StaticClass(), Guid))
