Sequencer: Add support for indirect manipulation of sub track offsets

The experimental gizmos support indirect manipulation (middle mouse with a highlighted axis or ctrl-LMB, ctrl-MMB, ctrl-RMB for specific axes). This modifies FSubTrackEditorMode::InputDelta to allow for those manipulations

diff --git a/Engine/Source/Editor/MovieSceneTools/Private/EditModes/SubTrackEditorMode.cpp b/Engine/Source/Editor/MovieSceneTools/Private/EditModes/SubTrackEditorMode.cpp
--- a/Engine/Source/Editor/MovieSceneTools/Private/EditModes/SubTrackEditorMode.cpp
+++ b/Engine/Source/Editor/MovieSceneTools/Private/EditModes/SubTrackEditorMode.cpp
@@ -51,12 +51,20 @@ bool FSubTrackEditorMode::InputDelta(FEditorViewportClient* InViewportClient, FV
 	const bool bCtrlDown = InViewport->KeyState(EKeys::LeftControl) || InViewport->KeyState(EKeys::RightControl);
 	const bool bShiftDown = InViewport->KeyState(EKeys::LeftShift) || InViewport->KeyState(EKeys::RightShift);
 	const bool bAltDown = InViewport->KeyState(EKeys::LeftAlt) || InViewport->KeyState(EKeys::RightAlt);
-	const bool bMouseButtonDown = InViewport->KeyState(EKeys::LeftMouseButton);
+	const bool bLeftMouseButtonDown = InViewport->KeyState(EKeys::LeftMouseButton);
+	const bool bMiddleMouseButtonDown = InViewport->KeyState(EKeys::MiddleMouseButton);
 	const bool bAnyModifiers = bAltDown || bCtrlDown || bShiftDown;
 	
 	const EAxisList::Type CurrentAxis = InViewportClient->GetCurrentWidgetAxis();
 
-	if(bMouseButtonDown && !bAnyModifiers && CurrentAxis != EAxisList::None)
+	const bool bIndirectManipulate = bMiddleMouseButtonDown && CurrentAxis != EAxisList::None;
+	const bool bAxisManipulate = bCtrlDown && bLeftMouseButtonDown && CurrentAxis != EAxisList::None;
+	const bool bDoManipulate = 
+		(bLeftMouseButtonDown && !bAnyModifiers && CurrentAxis != EAxisList::None) ||
+		bIndirectManipulate ||
+		bAxisManipulate;
+
+	if(bDoManipulate)
 	{
 		const FTransform TransformOriginFocusedSequence = GetTransformOriginForSequence(GetFocusedSequenceID());
 
