Sequencer: Revert unintentional change

diff --git a/Engine/Source/Editor/MovieSceneTools/Public/KeyframeTrackEditor.h b/Engine/Source/Editor/MovieSceneTools/Public/KeyframeTrackEditor.h
--- a/Engine/Source/Editor/MovieSceneTools/Public/KeyframeTrackEditor.h
+++ b/Engine/Source/Editor/MovieSceneTools/Public/KeyframeTrackEditor.h
@@ -421,17 +421,10 @@ struct FMovieSceneChannelValueSetter
 	template<typename ChannelType, typename ValueType>
 	static FMovieSceneChannelValueSetter Create(int32 ChannelIndex, ValueType&& InNewValue, bool bAddKey)
 	{
-		//FMovieSceneChannelValueSetter NewValue;
-		using namespace UE::MovieScene;
-
-		FUnpackedChannelValue Value(TIndexedChannelValue<ChannelType>(Forward<ValueType>(InNewValue), ChannelIndex), NAME_None);
-
 		FMovieSceneChannelValueSetter NewValue;
-		NewValue.Impl = FUnpackedChannelValueSetter(MoveTemp(Value), bAddKey);
-		return MoveTemp(NewValue);
-		//NewValue.Impl = TAddKeyImpl<ChannelType, typename TDecay<ValueType>::Type>(ChannelIndex, bAddKey, Forward<ValueType>(InNewValue));
+		NewValue.Impl = TAddKeyImpl<ChannelType, typename TDecay<ValueType>::Type>(ChannelIndex, bAddKey, Forward<ValueType>(InNewValue));
 		
-		//return MoveTemp(NewValue);
+		return MoveTemp(NewValue);
 	}
 
 	/** Templated construction function that can add a key (and potentially also set a default) for the specified channel and value */
