Added missing file from CL 45046896

diff --git a/Engine/Source/Editor/MovieSceneTools/Private/MovieSceneToolsModule.cpp b/Engine/Source/Editor/MovieSceneTools/Private/MovieSceneToolsModule.cpp
--- a/Engine/Source/Editor/MovieSceneTools/Private/MovieSceneToolsModule.cpp
+++ b/Engine/Source/Editor/MovieSceneTools/Private/MovieSceneToolsModule.cpp
@@ -182,7 +182,7 @@ void FMovieSceneToolsModule::StartupModule()
 		// register property track editors
 		BytePropertyTrackCreateEditorHandle = SequencerModule.RegisterPropertyTrackEditor<FBytePropertyTrackEditor>();
 		RotatorPropertyTrackCreateEditorHandle = SequencerModule.RegisterPropertyTrackEditor<FRotatorPropertyTrackEditor>();
-		VisibilityPropertyTrackCreateEditorHandle = SequencerModule.RegisterPropertyTrackEditor<FVisibilityPropertyTrackEditor>();
+		VisibilityPropertyTrackCreateEditorHandle = SequencerModule.RegisterTrackEditor(FOnCreateTrackEditor::CreateStatic(FVisibilityPropertyTrackEditor::CreateTrackEditor));
 		ActorReferencePropertyTrackCreateEditorHandle = SequencerModule.RegisterPropertyTrackEditor<FActorReferencePropertyTrackEditor>();
 		StringPropertyTrackCreateEditorHandle = SequencerModule.RegisterPropertyTrackEditor<FStringPropertyTrackEditor>();
 		TextPropertyTrackCreateEditorHandle = SequencerModule.RegisterPropertyTrackEditor<FTextPropertyTrackEditor>();
