Sequencer: Do not optimize channels to key when root motion offset is present

diff --git a/Engine/Source/Editor/MovieSceneTools/Private/TrackEditors/TransformTrackEditor.cpp b/Engine/Source/Editor/MovieSceneTools/Private/TrackEditors/TransformTrackEditor.cpp
--- a/Engine/Source/Editor/MovieSceneTools/Private/TrackEditors/TransformTrackEditor.cpp
+++ b/Engine/Source/Editor/MovieSceneTools/Private/TrackEditors/TransformTrackEditor.cpp
@@ -857,7 +857,7 @@ void F3DTransformTrackEditor::GetTransformKeys( const TOptional<FTransformData>&
 
 	USceneComponent* SceneComponent = MovieSceneHelpers::SceneComponentFromRuntimeObject(Object);
 	USceneComponent* AttachParent = SceneComponent ? SceneComponent->GetAttachParent() : nullptr;
-	bool bHasNeutralRotationSpace = IsValid(AttachParent) || GetTransformOrigin().GetRotation() == FQuat::Identity;
+	bool bHasNeutralRotationSpace = IsValid(AttachParent) || (GetTransformOrigin().GetRotation() == FQuat::Identity && GetRootMotionOffset(Object).Equals(FTransform::Identity));
 
 	// Set translation keys/defaults
 	{
