Adjustment  for an issue for capturing thumbnails for sequences. Render at a square aspect ratio and clean up some advanced rendering features to create the cleanest image possible like materials.

diff --git a/Engine/Source/Editor/MovieSceneTools/Private/TrackEditorThumbnail/TrackThumbnailUtils.cpp b/Engine/Source/Editor/MovieSceneTools/Private/TrackEditorThumbnail/TrackThumbnailUtils.cpp
--- a/Engine/Source/Editor/MovieSceneTools/Private/TrackEditorThumbnail/TrackThumbnailUtils.cpp
+++ b/Engine/Source/Editor/MovieSceneTools/Private/TrackEditorThumbnail/TrackThumbnailUtils.cpp
@@ -103,12 +103,18 @@ namespace UE::MoveSceneTools
 			// Auto-exposure is a temporal effect.
 			// Manual causes the AutoExposureBias to be used which will prevent the thumbnail from being mostly white in very bright scenes.
 			// This will not work in all cases but should in most cases (since it depends on what value is set for AutoExposureBias).
-			ProcessSettings.AutoExposureMethod = AEM_Manual;
+			ProcessSettings.AutoExposureMethod = AEM_Basic;
 			ProcessSettings.bOverride_AutoExposureMethod = true;
+			ProcessSettings.AutoExposureBias = 1.0f;
+			ProcessSettings.bOverride_AutoExposureBias = true;
 			
 			NewView->OverridePostProcessSettings(ProcessSettings, 1.f);
 		}
 
+		ViewFamily.EngineShowFlags.DisableAdvancedFeatures();
+		ViewFamily.EngineShowFlags.MotionBlur = 0;
+		ViewFamily.EngineShowFlags.AntiAliasing = 0;
+
 		const float GlobalResolutionFraction = 1.f;
 		ViewFamily.SetScreenPercentageInterface(new FLegacyScreenPercentageDriver(ViewFamily, GlobalResolutionFraction));
 
diff --git a/Engine/Source/Editor/Sequencer/Private/Misc/Thumbnail/ThumbnailCaptureUtils_CameraCuts.cpp b/Engine/Source/Editor/Sequencer/Private/Misc/Thumbnail/ThumbnailCaptureUtils_CameraCuts.cpp
--- a/Engine/Source/Editor/Sequencer/Private/Misc/Thumbnail/ThumbnailCaptureUtils_CameraCuts.cpp
+++ b/Engine/Source/Editor/Sequencer/Private/Misc/Thumbnail/ThumbnailCaptureUtils_CameraCuts.cpp
@@ -50,6 +50,8 @@ namespace UE::Sequencer::Private
 		// Important to get camera view after PreDrawThumbnailSetupSequencer since it may have set its transform!
 		FMinimalViewInfo ViewInfo;
 		Component->GetCameraView(FApp::GetDeltaTime(), ViewInfo);
+		// We're rendering square, so lets make sure we set the aspect ratio here.
+		ViewInfo.AspectRatio = 1.0f;
 		// It's important to capture with sRGB so the thumbnail has the right brightness
 		UTextureRenderTarget2D* RenderTarget2D = UKismetRenderingLibrary::CreateRenderTarget2D(
 			Component, ThumbnailSize, ThumbnailSize, RTF_RGBA8_SRGB
