Small comment fix that talks about manual autoexposure when it's talking about basic that wasn't adjusted while testing some changes from previous submission.

diff --git a/Engine/Source/Editor/MovieSceneTools/Private/TrackEditorThumbnail/TrackThumbnailUtils.cpp b/Engine/Source/Editor/MovieSceneTools/Private/TrackEditorThumbnail/TrackThumbnailUtils.cpp
--- a/Engine/Source/Editor/MovieSceneTools/Private/TrackEditorThumbnail/TrackThumbnailUtils.cpp
+++ b/Engine/Source/Editor/MovieSceneTools/Private/TrackEditorThumbnail/TrackThumbnailUtils.cpp
@@ -101,7 +101,7 @@ namespace UE::MoveSceneTools
 			
 			// Temporal effects need time to warm up, which we don't do for thumbnail rendering for performance reasons.
 			// Auto-exposure is a temporal effect.
-			// Manual causes the AutoExposureBias to be used which will prevent the thumbnail from being mostly white in very bright scenes.
+			// Basic causes the AutoExposureBias to be used which will prevent the thumbnail from being mostly white in very bright scenes.
 			// This will not work in all cases but should in most cases (since it depends on what value is set for AutoExposureBias).
 			ProcessSettings.AutoExposureMethod = AEM_Basic;
 			ProcessSettings.bOverride_AutoExposureMethod = true;
