Sequencer: Indirect manipulation with RMB

#rnx

diff --git a/Engine/Source/Editor/MovieSceneTools/Private/EditModes/SubTrackEditorMode.cpp b/Engine/Source/Editor/MovieSceneTools/Private/EditModes/SubTrackEditorMode.cpp
--- a/Engine/Source/Editor/MovieSceneTools/Private/EditModes/SubTrackEditorMode.cpp
+++ b/Engine/Source/Editor/MovieSceneTools/Private/EditModes/SubTrackEditorMode.cpp
@@ -53,11 +53,12 @@ bool FSubTrackEditorMode::InputDelta(FEditorViewportClient* InViewportClient, FV
 	const bool bAltDown = InViewport->KeyState(EKeys::LeftAlt) || InViewport->KeyState(EKeys::RightAlt);
 	const bool bLeftMouseButtonDown = InViewport->KeyState(EKeys::LeftMouseButton);
 	const bool bMiddleMouseButtonDown = InViewport->KeyState(EKeys::MiddleMouseButton);
+	const bool bRightMouseButtonDown = InViewport->KeyState(EKeys::RightMouseButton);
 	const bool bAnyModifiers = bAltDown || bCtrlDown || bShiftDown;
 	
 	const EAxisList::Type CurrentAxis = InViewportClient->GetCurrentWidgetAxis();
 
-	const bool bIndirectManipulate = bMiddleMouseButtonDown && CurrentAxis != EAxisList::None;
+	const bool bIndirectManipulate = (bMiddleMouseButtonDown || bRightMouseButtonDown) && CurrentAxis != EAxisList::None;
 	const bool bAxisManipulate = bCtrlDown && bLeftMouseButtonDown && CurrentAxis != EAxisList::None;
 	const bool bDoManipulate = 
 		(bLeftMouseButtonDown && !bAnyModifiers && CurrentAxis != EAxisList::None) ||
