Constraints: ability to create constraints from the start of the level sequence or for the whole sequence
- this adds the notion of infinite constraints that are always active

diff --git a/Engine/Plugins/Animation/ControlRig/Source/ControlRigEditor/Private/Editor/Constraints/ConstraintCreationOptions.h b/Engine/Plugins/Animation/ControlRig/Source/ControlRigEditor/Private/Editor/Constraints/ConstraintCreationOptions.h
new file mode 100644
--- /dev/null
+++ b/Engine/Plugins/Animation/ControlRig/Source/ControlRigEditor/Private/Editor/Constraints/ConstraintCreationOptions.h
@@ -0,0 +1,24 @@
+﻿// Copyright Epic Games, Inc. All Rights Reserved.
+
+#pragma once
+
+#include "Constraints/MovieSceneConstraintChannelHelper.h"
+#include "UObject/Object.h"
+
+#include "ConstraintCreationOptions.generated.h"
+
+/**
+ * UConstraintCreationOptions is used to pass options for creating constraints.
+ */
+
+UCLASS(Blueprintable)
+class UConstraintCreationOptions : public UObject
+{
+	GENERATED_BODY()
+
+public:
+
+	/* Creation options related to sequencer. */
+	UPROPERTY(EditAnywhere, Category=Options)
+	FSequencerCreationOptions SequencerOptions;
+};
diff --git a/Engine/Plugins/Animation/ControlRig/Source/ControlRigEditor/Private/Editor/Constraints/SConstraintsWidget.cpp b/Engine/Plugins/Animation/ControlRig/Source/ControlRigEditor/Private/Editor/Constraints/SConstraintsWidget.cpp
--- a/Engine/Plugins/Animation/ControlRig/Source/ControlRigEditor/Private/Editor/Constraints/SConstraintsWidget.cpp
+++ b/Engine/Plugins/Animation/ControlRig/Source/ControlRigEditor/Private/Editor/Constraints/SConstraintsWidget.cpp
@@ -3,6 +3,7 @@
 #include "SConstraintsWidget.h"
 
 #include "ActorPickerMode.h"
+#include "ConstraintCreationOptions.h"
 #include "Constraints/MovieSceneConstraintChannelHelper.h"
 #include "ConstraintsActor.h"
 #include "ControlRigEditorStyle.h"
@@ -280,11 +281,8 @@ namespace UE::Private
 
 					if (WeakSequencer.IsValid())
 					{
-						if (bUseDefault)
-						{
-							Constraint->Evaluate();
-						}
-						FControlRigParameterTrackEditor::AddConstraintToSequencer(WeakSequencer.Pin(), Constraint);
+						const UConstraintCreationOptions* Options = bUseDefault ? GetDefault<UConstraintCreationOptions>() : nullptr;
+						FControlRigParameterTrackEditor::AddConstraintToSequencer(WeakSequencer.Pin(), Constraint, Options);
 					}
 					else
 					{
@@ -433,6 +431,7 @@ TSharedRef<SWidget> SConstraintMenuEntry::GenerateConstraintDefaultWidget()
 		DetailsViewArgs.bShowModifiedPropertiesOption = false;
 		DetailsViewArgs.ColumnWidth = 0.45f;
 		DetailsViewArgs.NotifyHook = this;
+		DetailsViewArgs.bAllowMultipleTopLevelObjects = true;
 	}
 	
 	TSharedRef<IDetailsView> DetailsView = FModuleManager::GetModuleChecked<FPropertyEditorModule>("PropertyEditor").CreateDetailView(DetailsViewArgs);
@@ -464,11 +463,23 @@ TSharedRef<SWidget> SConstraintMenuEntry::GenerateConstraintDefaultWidget()
 				}
 			}
 
+			static const FName SequencerOptionsPropName = GET_MEMBER_NAME_CHECKED(UConstraintCreationOptions, SequencerOptions);
+			if (PropertyName == SequencerOptionsPropName)
+			{
+				TWeakPtr<ISequencer> WeakSequencer = UE::Private::GetSequencerChecked();
+				return UE::Private::GetSequencerChecked().IsValid();
+			}
+				
 			return true;
 		});
 	DetailsView->SetIsPropertyVisibleDelegate(PropertyVisibility);
-		
-	DetailsView->SetObject(ConfigurableConstraint);
+	
+	TArray<UObject*> Objects(
+	{
+		ConfigurableConstraint,
+		GetMutableDefault<UConstraintCreationOptions>()
+	});
+	DetailsView->SetObjects(Objects);
 
 	return SNew(SBorder)
 		.Visibility(EVisibility::Visible)
@@ -1884,8 +1895,14 @@ TSharedPtr<SWidget> SConstraintsEditionWidget::CreateContextMenu()
 	FIsPropertyVisible PropertyVisibility = FIsPropertyVisible::CreateLambda(
 		[bIsLookAtConstraint](const FPropertyAndParent& InPropertyAndParent)
 		{
+			const FName PropertyName = InPropertyAndParent.Property.GetFName();
+				
 			if (!bIsLookAtConstraint)
 			{
+				if (PropertyName == GET_MEMBER_NAME_CHECKED(UTickableTransformConstraint, bUseCurrentOffset))
+				{
+					return false;
+				}
 				return true;
 			}
 
@@ -1893,7 +1910,6 @@ TSharedPtr<SWidget> SConstraintsEditionWidget::CreateContextMenu()
 			static const FName MaintainOffsetPropName = GET_MEMBER_NAME_CHECKED(UTickableTransformConstraint, bMaintainOffset);
 			static const FName DynamicOffsetPropName = GET_MEMBER_NAME_CHECKED(UTickableTransformConstraint, bDynamicOffset);
 						
-			const FName PropertyName = InPropertyAndParent.Property.GetFName();
 			if (PropertyName == MaintainOffsetPropName || PropertyName == DynamicOffsetPropName)
 			{
 				return false;
@@ -1962,18 +1978,15 @@ TSharedPtr<SWidget> SConstraintsEditionWidget::CreateContextMenu()
 			FText::Format(LOCTEXT("BakeConstraintDoItTooltip", "Bake {0} transforms."), ConstraintLabel),
 			FSlateIcon(),
 			FUIAction(FExecuteAction::CreateLambda([this, Constraint, World]()
+			{
+				const TSharedPtr<ISequencer>& Sequencer = WeakSequencer.IsValid() ? WeakSequencer.Pin() : nullptr;
+				UTickableTransformConstraint* TransformConstraint = Cast<UTickableTransformConstraint>(Constraint);
+				if (TransformConstraint && Sequencer)
 				{
-					if (UTickableTransformConstraint* TransformConstraint = Cast<UTickableTransformConstraint>(Constraint))
-					{
-						if (!WeakSequencer.IsValid())
-						{
-							return;
-						}
-						FScopedTransaction Transaction(LOCTEXT("BakeConstraint", "Bake Constraint"));
-						FConstraintBaker::Bake(World, TransformConstraint, WeakSequencer.Pin(), TOptional< FBakingAnimationKeySettings>(), TOptional<TArray<FFrameNumber>>());
-
-					}
-				})),
+					FScopedTransaction Transaction(LOCTEXT("BakeConstraint", "Bake Constraint"));
+					FConstraintBaker::Bake(World, TransformConstraint, WeakSequencer.Pin(), TOptional< FBakingAnimationKeySettings>(), TOptional<TArray<FFrameNumber>>());
+				}
+			})),
 			NAME_None,
 			EUserInterfaceActionType::Button);
 		}
diff --git a/Engine/Plugins/Animation/ControlRig/Source/ControlRigEditor/Private/Editor/Constraints/SConstraintsWidget.h b/Engine/Plugins/Animation/ControlRig/Source/ControlRigEditor/Private/Editor/Constraints/SConstraintsWidget.h
--- a/Engine/Plugins/Animation/ControlRig/Source/ControlRigEditor/Private/Editor/Constraints/SConstraintsWidget.h
+++ b/Engine/Plugins/Animation/ControlRig/Source/ControlRigEditor/Private/Editor/Constraints/SConstraintsWidget.h
@@ -93,6 +93,7 @@ public:
 
 		SLATE_EVENT(FOnConstraintCreated, OnConstraintCreated)
 		SLATE_EVENT(FOnGetConstrainable, OnGetParent)
+		SLATE_ARGUMENT(TWeakPtr<ISequencer>, WeakSequencer)
 	
 	SLATE_END_ARGS()
 
diff --git a/Engine/Plugins/Animation/ControlRig/Source/ControlRigEditor/Private/Sequencer/ControlConstraintChannelInterface.cpp b/Engine/Plugins/Animation/ControlRig/Source/ControlRigEditor/Private/Sequencer/ControlConstraintChannelInterface.cpp
--- a/Engine/Plugins/Animation/ControlRig/Source/ControlRigEditor/Private/Sequencer/ControlConstraintChannelInterface.cpp
+++ b/Engine/Plugins/Animation/ControlRig/Source/ControlRigEditor/Private/Sequencer/ControlConstraintChannelInterface.cpp
@@ -69,13 +69,17 @@ UWorld* FControlConstraintChannelInterface::GetHandleWorld(UTransformableHandle*
 
 bool FControlConstraintChannelInterface::SmartConstraintKey(
 	UTickableTransformConstraint* InConstraint, const TOptional<bool>& InOptActive,
-	const FFrameNumber& InTime, const TSharedPtr<ISequencer>& InSequencer)
+	const TRange<FFrameNumber>& InTimeRange, const TSharedPtr<ISequencer>& InSequencer)
 {
 	if (!InConstraint)
 	{
 		return false;
 	}
 
+	// TODO handle range
+	const bool bInfiniteInput = !InTimeRange.HasLowerBound();
+	const TOptional<FFrameNumber> Time = bInfiniteInput ? TOptional<FFrameNumber>() : InTimeRange.GetLowerBoundValue();
+
 	const UTransformableControlHandle* ControlHandle = Cast<UTransformableControlHandle>(InConstraint->ChildTRSHandle);
 	if (!ControlHandle)
 	{
@@ -89,7 +93,7 @@ bool FControlConstraintChannelInterface::SmartConstraintKey(
 		return false;
 	}
 	
-	FScopedTransaction Transaction(LOCTEXT("KeyConstraintaKehy", "Key Constraint Key"));
+	FScopedTransaction Transaction(LOCTEXT("KeyControlConstraintKey", "Key Control Constraint Key"));
 	ConstraintSection->Modify();
 	TransformSection->Modify();
 
@@ -109,12 +113,15 @@ bool FControlConstraintChannelInterface::SmartConstraintKey(
 	// add key if needed
 	if (FConstraintAndActiveChannel* Channel = ConstraintSection->GetConstraintChannel(InConstraint->ConstraintID))
 	{
-		bool ActiveValueToBeSet = false;
+		if (bInfiniteInput)
+		{
+			ensure(Channel->ActiveChannel.GetNumKeys() == 0);
+		}
+		
+		bool ActiveValueToBeSet = bInfiniteInput;
 		//add key if we can and make sure the key we are setting is what we want
-		if (CanAddKey(Channel->ActiveChannel, InTime, ActiveValueToBeSet) && (InOptActive.IsSet() == false || InOptActive.GetValue() == ActiveValueToBeSet))
+		if (bInfiniteInput || (CanAddKey(Channel->ActiveChannel, *Time, ActiveValueToBeSet) && (InOptActive.IsSet() == false || InOptActive.GetValue() == ActiveValueToBeSet)))
 		{
-			const FFrameRate TickResolution = InSequencer->GetFocusedTickResolution();
-
 			const bool bNeedsCompensation = InConstraint->NeedsCompensation();
 				
 			TGuardValue<bool> CompensateGuard(FMovieSceneConstraintChannelHelper::bDoNotCompensate, true);
@@ -128,25 +135,58 @@ bool FControlConstraintChannelInterface::SmartConstraintKey(
 			TArray<FFrameNumber> FramesToCompensate;
 			if (bNeedsCompensation)
 			{
-				FMovieSceneConstraintChannelHelper::GetFramesToCompensate<FMovieSceneFloatChannel>(Channel->ActiveChannel, ActiveValueToBeSet, InTime, Channels, FramesToCompensate);
+				FMovieSceneConstraintChannelHelper::GetFramesToCompensate<FMovieSceneFloatChannel>(Channel->ActiveChannel, ActiveValueToBeSet, Time, Channels, FramesToCompensate);
 			}
 			else
 			{
-				FramesToCompensate.Add(InTime);
+				if (!bInfiniteInput)
+				{
+					FramesToCompensate.Add(*Time);
+				}
 			}
 
+			const bool bIsActiveByDefault = Channel->ActiveChannel.GetDefault() && *Channel->ActiveChannel.GetDefault();
+			const bool bWasInfinite = bIsActiveByDefault && Channel->ActiveChannel.GetNumKeys() == 0;
+			
+			const FFrameRate TickResolution = InSequencer->GetFocusedTickResolution();
+			const FFrameTime CurrentTime = InSequencer->GetLocalTime().ConvertTo(TickResolution);
+			
+			// check if the first active key will be added before the current frame
+			// if that's the case then the current frame must be added and the compensation evaluator notified
+			TOptional<FFrameNumber> OptCurrentTime;
+			bool bPropagateRelative = true;
+			if (!bInfiniteInput && !bWasInfinite && !bHadConstraintChannel && ActiveValueToBeSet)
+			{
+				Tie(OptCurrentTime, bPropagateRelative) = IsActiveKeyAddedBeforeCurrentTime(CurrentTime.GetFrame(), *Time, Channels, FramesToCompensate);
+			}
+			
 			FCompensationEvaluator::FEvalParameters EvalParams(InSequencer, FramesToCompensate);
 			EvalParams.bToActive = ActiveValueToBeSet;
-			if (bHadConstraintChannel && ActiveValueToBeSet)
+			EvalParams.OptCurrentFrame = OptCurrentTime;
+			EvalParams.bPropagateRelative = bPropagateRelative;
+
+			if (bWasInfinite)
+			{
+				EvalParams.bToActive = false;
+			}
+			else if (bIsActiveByDefault && Time)
 			{
 				const TMovieSceneChannelData<const bool> ChannelData = Channel->ActiveChannel.GetData();
-				if (ChannelData.GetTimes().IsEmpty())
+				const TArrayView<const FFrameNumber> Times = ChannelData.GetTimes();
+				const TArrayView<const bool> Values = ChannelData.GetValues();
+				
+				const bool bInfiniteAtLeft = *Time < Times[0] && !Values[0];
+				EvalParams.bInfiniteAtLeft = bInfiniteAtLeft;
+			}
+			
+			if (bHadConstraintChannel && ActiveValueToBeSet && !bWasInfinite)
+			{
+				const TMovieSceneChannelData<const bool> ChannelData = Channel->ActiveChannel.GetData();
+				if (ChannelData.GetTimes().IsEmpty() && !bInfiniteInput)
 				{
 					// this means that the channel already existed but there was no existing key
 					// e.g., when the user deletes keys directly in sequencer without deleting the constraint
-					const FFrameTime FrameTime = InSequencer->GetLocalTime().ConvertTo(TickResolution);
-					const FFrameNumber Time = FrameTime.GetFrame();
-					EvalParams.bKeepCurrent = !FramesToCompensate.IsEmpty() && FramesToCompensate[0] == Time;
+					EvalParams.bKeepCurrent = !FramesToCompensate.IsEmpty() && FramesToCompensate[0] == CurrentTime.GetFrame();
 				}
 			}
 			
@@ -167,19 +207,19 @@ bool FControlConstraintChannelInterface::SmartConstraintKey(
 			{
 				// get the number of float channels to treat
 				NumChannels = FControlRigSpaceChannelHelpers::GetNumFloatChannels(ControlElement->Settings.ControlType);
-				if (bNeedsCompensation && NumChannels > 0)
+				if (bNeedsCompensation && !bInfiniteInput && NumChannels > 0)
 				{
 					ChannelIndex = pChannelIndex->ChannelIndex;
-					EvaluateTangentAtThisTime<FMovieSceneFloatChannel>(ChannelIndex, NumChannels, TransformSection, InTime, Tangents);
+					EvaluateTangentAtThisTime<FMovieSceneFloatChannel>(ChannelIndex, NumChannels, TransformSection, *Time, Tangents);
 				}
 			}
 		
 			const EMovieSceneTransformChannel ChannelsToKey =InConstraint->GetChannelsToKey();
 			
 			// add child's transform key at Time-1 to keep animation
-			if (bNeedsCompensation)
+			if (bNeedsCompensation && !bInfiniteInput)
 			{
-				const FFrameNumber TimeMinusOne(InTime - 1);
+				const FFrameNumber TimeMinusOne(*Time - 1);
 
 				ControlHandle->AddTransformKeys({ TimeMinusOne },
 					{ ChildLocals[0] }, ChannelsToKey, TickResolution, nullptr,true);
@@ -192,28 +232,47 @@ bool FControlConstraintChannelInterface::SmartConstraintKey(
 			}
 
 			// add active key
+			if (bInfiniteInput)
+			{
+				constexpr bool bActiveByDefault = true;
+				Channel->ActiveChannel.SetDefault(bActiveByDefault);
+			}
+			else
 			{
 				TMovieSceneChannelData<bool> ChannelData = Channel->ActiveChannel.GetData();
-				ChannelData.UpdateOrAddKey(InTime, ActiveValueToBeSet);
+
+				const TArrayView<const FFrameNumber> FrameNumbers = ChannelData.GetTimes();
+				const bool bRemoveInfinite = bIsActiveByDefault && ActiveValueToBeSet && (FrameNumbers.IsEmpty() || *Time < FrameNumbers[0]); 
+				
+				ChannelData.UpdateOrAddKey(*Time, ActiveValueToBeSet);
+
+				if (bRemoveInfinite)
+				{
+					Channel->ActiveChannel.RemoveDefault();
+				}
 			}
 
 			// compensate
 			{
 				// we need to remove the first transforms as we store NumFrames+1 transforms
-				ChildLocals.RemoveAt(0);
+				if (!ChildLocals.IsEmpty())
+				{
+					ChildLocals.RemoveAt(0);
+				}
 
 				// add keys
 				ControlHandle->AddTransformKeys(FramesToCompensate,
 					ChildLocals, ChannelsToKey, TickResolution, nullptr,true);
 
 				// set tangents at Time
-				if (bNeedsCompensation && NumChannels > 0)
+				if (bNeedsCompensation && !bInfiniteInput && NumChannels > 0)
 				{
-					SetTangentsAtThisTime<FMovieSceneFloatChannel>(ChannelIndex, NumChannels, TransformSection, InTime, Tangents);
+					SetTangentsAtThisTime<FMovieSceneFloatChannel>(ChannelIndex, NumChannels, TransformSection, *Time, Tangents);
 				}
 			}
 
-			// cleanup next keys that have the same value as they are useless 
+			// cleanup next keys that have the same value as they are useless
+			if (!bInfiniteInput)
 			{
 				TArray<FFrameNumber> TimesRemoved;
 				CleanDuplicates(Channel->ActiveChannel, TimesRemoved);
diff --git a/Engine/Plugins/Animation/ControlRig/Source/ControlRigEditor/Private/Sequencer/ControlRigParameterTrackEditor.cpp b/Engine/Plugins/Animation/ControlRig/Source/ControlRigEditor/Private/Sequencer/ControlRigParameterTrackEditor.cpp
--- a/Engine/Plugins/Animation/ControlRig/Source/ControlRigEditor/Private/Sequencer/ControlRigParameterTrackEditor.cpp
+++ b/Engine/Plugins/Animation/ControlRig/Source/ControlRigEditor/Private/Sequencer/ControlRigParameterTrackEditor.cpp
@@ -116,6 +116,7 @@
 #include "Editor/RigVMEditorTools.h"
 #include "Sequencer/AnimLayers/AnimLayers.h"
 #include "Decorations/MovieSceneMuteSoloDecoration.h"
+#include "Editor/Constraints/ConstraintCreationOptions.h"
 
 #define LOCTEXT_NAMESPACE "FControlRigParameterTrackEditor"
 
@@ -2700,10 +2701,19 @@ void FControlRigParameterTrackEditor::OnChannelChanged(const FMovieSceneChannelM
 	}
 }
 
-void FControlRigParameterTrackEditor::AddConstraintToSequencer(const TSharedPtr<ISequencer>& InSequencer, UTickableTransformConstraint* InConstraint)
+void FControlRigParameterTrackEditor::AddConstraintToSequencer(const TSharedPtr<ISequencer>& InSequencer, UTickableTransformConstraint* InConstraint,
+	const UConstraintCreationOptions* InSequencerOptions)
 {
 	TGuardValue<bool> DisableTrackCreation(bAutoGenerateControlRigTrack, false);
-	FMovieSceneConstraintChannelHelper::AddConstraintToSequencer(InSequencer, InConstraint);
+	if (InSequencerOptions)
+	{
+		InConstraint->Evaluate();
+	}
+	
+	static const FSequencerCreationOptions DefaultOptions;
+	const FSequencerCreationOptions& Options = InSequencerOptions ? InSequencerOptions->SequencerOptions : DefaultOptions;
+	
+	FMovieSceneConstraintChannelHelper::AddConstraintToSequencer(InSequencer, InConstraint, Options);
 }
 
 void FControlRigParameterTrackEditor::AddTrackForComponent(USceneComponent* InComponent, FGuid InBinding) 
@@ -4321,7 +4331,7 @@ void FControlRigParameterTrackEditor::GetControlRigKeys(
 	}
 	URigHierarchy* Hierarchy = InControlRig->GetHierarchy();
 
-	//Need seperate index fo bools,ints and enums and floats since there are seperate entries for each later when they are accessed by the set key stuff.
+	//Need separate index for bools, ints and enums and floats since there are separate entries for each later when they are accessed by the set key stuff.
 	int32 SpaceChannelIndex = 0;
 	for (int32 LocalControlIndex = 0; LocalControlIndex < Controls.Num(); ++LocalControlIndex)
 	{
diff --git a/Engine/Plugins/Animation/ControlRig/Source/ControlRigEditor/Private/Sequencer/ControlRigParameterTrackEditor.h b/Engine/Plugins/Animation/ControlRig/Source/ControlRigEditor/Private/Sequencer/ControlRigParameterTrackEditor.h
--- a/Engine/Plugins/Animation/ControlRig/Source/ControlRigEditor/Private/Sequencer/ControlRigParameterTrackEditor.h
+++ b/Engine/Plugins/Animation/ControlRig/Source/ControlRigEditor/Private/Sequencer/ControlRigParameterTrackEditor.h
@@ -39,6 +39,7 @@ struct FBakingAnimationKeySettings;
 class ISequencer;
 class UMovieSceneTrack;
 class IStructureDetailsView;
+class UConstraintCreationOptions;
 
 namespace UE::MovieScene
 {
@@ -352,7 +353,8 @@ public:
 	static void SmartReduce(const TSharedPtr<ISequencer>& InSequencer, const FSmartReduceParams& InParams, UMovieSceneSection* InSection);
 
 	/** Adds the constraint to sequencer and create the bindings if needed. Note that ControlRig Track auto-generation is disabled in this function. */
-	static void AddConstraintToSequencer(const TSharedPtr<ISequencer>& InSequencer, UTickableTransformConstraint* InConstraint);
+	static void AddConstraintToSequencer(const TSharedPtr<ISequencer>& InSequencer, UTickableTransformConstraint* InConstraint,
+		const UConstraintCreationOptions* InSequencerOptions = nullptr);
 
 private:
 	FDelegateHandle SelectionChangedHandle;
diff --git a/Engine/Plugins/Animation/ControlRig/Source/ControlRigEditor/Private/Tools/ConstraintBaker.cpp b/Engine/Plugins/Animation/ControlRig/Source/ControlRigEditor/Private/Tools/ConstraintBaker.cpp
--- a/Engine/Plugins/Animation/ControlRig/Source/ControlRigEditor/Private/Tools/ConstraintBaker.cpp
+++ b/Engine/Plugins/Animation/ControlRig/Source/ControlRigEditor/Private/Tools/ConstraintBaker.cpp
@@ -34,26 +34,32 @@ void FConstraintBaker::GetMinimalFramesToBake(
 	{
 		return;
 	}
+	
 	IMovieSceneConstrainedSection* ConstrainedSection = Cast<IMovieSceneConstrainedSection>(InSection);
 
-	// get constraint channel data if any
-	TArrayView<const FFrameNumber> ConstraintFrames;
-	TArrayView<const bool> ConstraintValues;
-
 	// note that we might want to bake a constraint which is not animated
 	FConstraintAndActiveChannel* ThisActiveChannel = ConstrainedSection->GetConstraintChannel(InConstraint->ConstraintID);
-	if (ThisActiveChannel)
+	if (!ThisActiveChannel)
 	{
-		const TMovieSceneChannelData<const bool> ConstraintChannelData = ThisActiveChannel->ActiveChannel.GetData();
-		ConstraintFrames = ConstraintChannelData.GetTimes();
-		ConstraintValues = ConstraintChannelData.GetValues();
+		return;
 	}
+	
+	const TMovieSceneChannelData<const bool> ConstraintChannelData = ThisActiveChannel->ActiveChannel.GetData();
+	// get constraint channel data if any
+	TArrayView<const FFrameNumber> ConstraintFrames = ConstraintChannelData.GetTimes();
+	TArrayView<const bool> ConstraintValues = ConstraintChannelData.GetValues();
 
+	const TOptional<bool> bDefault = ThisActiveChannel->ActiveChannel.GetDefault();
+	const bool bInfinite = bDefault && *bDefault;
+	
 	// there's no channel or an empty one and the constraint is inactive so no need to do anything if the constraint
 	// is not active
 	if (ConstraintFrames.IsEmpty() && !InConstraint->IsFullyActive())
 	{
-		return;
+		if (!bInfinite)
+		{
+			return;
+		}
 	}
 
 	// default init bounds to the scene bounds
@@ -75,7 +81,7 @@ void FConstraintBaker::GetMinimalFramesToBake(
 		}
 
 		// set end to the last key if inactive
-		const bool bIsLastKeyInactive = (ConstraintValues.Last() == false);
+		const bool bIsLastKeyInactive = ConstraintValues.IsEmpty() ? false : (ConstraintValues.Last() == false);
 		if (bIsLastKeyInactive)
 		{
 			LastBakingFrame = ConstraintFrames.Last();
diff --git a/Engine/Plugins/Animation/ControlRig/Source/ControlRigEditor/Public/ControlConstraintChannelInterface.h b/Engine/Plugins/Animation/ControlRig/Source/ControlRigEditor/Public/ControlConstraintChannelInterface.h
--- a/Engine/Plugins/Animation/ControlRig/Source/ControlRigEditor/Public/ControlConstraintChannelInterface.h
+++ b/Engine/Plugins/Animation/ControlRig/Source/ControlRigEditor/Public/ControlConstraintChannelInterface.h
@@ -36,7 +36,7 @@ struct FControlConstraintChannelInterface : ITransformConstraintChannelInterface
 	/** Add an active/inactive key to the constraint channel if needed and does the transform compensation on the control's transform channels. */
 	virtual bool SmartConstraintKey(
 		UTickableTransformConstraint* InConstraint, const TOptional<bool>& InOptActive,
-		const FFrameNumber& InTime, const TSharedPtr<ISequencer>& InSequencer) override;
+		const TRange<FFrameNumber>& InTimeRange, const TSharedPtr<ISequencer>& InSequencer) override;
 	
 	/** Add keys on the controls's transform channels. */
 	virtual void AddHandleTransformKeys(
diff --git a/Engine/Source/Editor/MovieSceneTools/Private/Channels/ConstraintChannelEditor.cpp b/Engine/Source/Editor/MovieSceneTools/Private/Channels/ConstraintChannelEditor.cpp
--- a/Engine/Source/Editor/MovieSceneTools/Private/Channels/ConstraintChannelEditor.cpp
+++ b/Engine/Source/Editor/MovieSceneTools/Private/Channels/ConstraintChannelEditor.cpp
@@ -91,23 +91,49 @@ namespace
 
 		const TMovieSceneChannelData<const bool> ChannelData = Channel->GetData();
 		const TArrayView<const bool> Values = ChannelData.GetValues();
+
+		const bool bIsActiveByDefault = Channel->GetDefault() && *Channel->GetDefault();
+		if (Values.Num() == 0 && bIsActiveByDefault)
+		{
+			Ranges.Emplace(INDEX_NONE);
+			return Ranges;
+		}
+		
 		if (Values.Num() == 1)
 		{
 			if (Values[0] == true)
 			{
 				Ranges.Emplace(0);
 			}
+			else
+			{
+				if (bIsActiveByDefault)
+				{
+					Ranges.Emplace(INDEX_NONE, 0);
+				}
+			}
 			return Ranges;
 		}
 		
 		int32 Offset = 0;
+
+		if (bIsActiveByDefault)
+		{
+			const int32 FirstInactive = Values.Find(false);
+			if (FirstInactive != INDEX_NONE)
+			{
+				Ranges.Emplace(INDEX_NONE, FirstInactive);
+				Offset = FirstInactive+1;
+			}
+		}
+		
 		while (Values.IsValidIndex(Offset))
 		{
 			TRange<int32> Range = GetRange(Values, Offset);
 			if (!Range.IsEmpty())
 			{
-				Ranges.Emplace(Range);
 				Offset = Range.GetUpperBoundValue()+1;
+				Ranges.Emplace(MoveTemp(Range));
 			}
 			else
 			{
@@ -213,27 +239,51 @@ TArray<FKeyBarCurveModel::FBarRange> FConstraintChannelEditor::GetBarRanges(
 		return Ranges;
 	}
 
-	// convert to bar range
-	const TMovieSceneChannelData<const bool> ChannelData = Channel->GetData();
-	const TArrayView<const FFrameNumber> Times = ChannelData.GetTimes();
 	const FConstraintChannelLabelBuilder LabelBuilder(Owner, Channel);
-	const FText& Label = LabelBuilder.Label;
-	for (const IndexRange& ActiveRange : IndexRanges)
+	auto AddNewRange = [&Ranges, LabelStr = LabelBuilder.Label.ToString()]() -> FKeyBarCurveModel::FBarRange&
 	{
-		FKeyBarCurveModel::FBarRange BarRange;
-		FFrameRate TickResolution = Owner->GetTypedOuter<UMovieScene>()->GetTickResolution();
-		double LowerValue = Times[ActiveRange.GetLowerBoundValue()] / TickResolution;
-		double UpperValue = Times[ActiveRange.GetUpperBoundValue()] / TickResolution;
-
-		BarRange.Range.SetLowerBound(TRangeBound<double>(LowerValue));
-		BarRange.Range.SetUpperBound(TRangeBound<double>(UpperValue));
-
-		BarRange.Name = FName(*Label.ToString());
+		FKeyBarCurveModel::FBarRange& BarRange = Ranges.Emplace_GetRef();
+		
+		BarRange.Name = FName(*LabelStr);
 		BarRange.Color = FLinearColor(.2, .5, .1);
 		static FLinearColor ZebraTint = FLinearColor::White.CopyWithNewOpacity(0.01f);
 		BarRange.Color = BarRange.Color * (1.f - ZebraTint.A) + ZebraTint * ZebraTint.A;
-		BarRange.bRangeIsInfinite = false;
-		Ranges.Add(BarRange);
+
+		return BarRange;
+	};
+
+	// infinite range
+	if (IndexRanges[0].HasLowerBound() && IndexRanges[0].GetLowerBoundValue() == INDEX_NONE &&
+		IndexRanges[0].HasUpperBound() && IndexRanges[0].GetUpperBoundValue() == INDEX_NONE)
+	{
+		FKeyBarCurveModel::FBarRange& BarRange = AddNewRange();
+		BarRange.bRangeIsInfinite = true;
+	}
+	else
+	{
+		FFrameRate TickResolution = Owner->GetTypedOuter<UMovieScene>()->GetTickResolution();
+		
+		// convert to bar range
+		const TMovieSceneChannelData<const bool> ChannelData = Channel->GetData();
+		const TArrayView<const FFrameNumber> Times = ChannelData.GetTimes();
+		for (const IndexRange& ActiveRange : IndexRanges)
+		{
+			FKeyBarCurveModel::FBarRange& BarRange = AddNewRange();
+			
+			if (ActiveRange.HasLowerBound() && ActiveRange.GetLowerBoundValue() > INDEX_NONE)
+			{
+				double LowerValue = Times[ActiveRange.GetLowerBoundValue()] / TickResolution;
+				BarRange.Range.SetLowerBound(TRangeBound<double>(LowerValue));
+			}
+
+			if (ActiveRange.HasUpperBound() && ActiveRange.GetUpperBoundValue() > INDEX_NONE)
+			{
+				double UpperValue = Times[ActiveRange.GetUpperBoundValue()] / TickResolution;
+				BarRange.Range.SetUpperBound(TRangeBound<double>(UpperValue));
+			}
+			
+			BarRange.bRangeIsInfinite = false;
+		}		
 	}
 
 	return Ranges;
@@ -280,13 +330,23 @@ int32 DrawExtra(
 		const FKeyBarCurveModel::FBarRange& Range = Ranges[Index];
 
 		// draw bar
-		double LowerSeconds = Range.Range.GetLowerBoundValue();
-		double UpperSeconds = Range.Range.GetUpperBoundValue();
-		if (UpperSeconds == Range.Range.GetLowerBoundValue())
+		double LowerSeconds = InputMin;
+		double UpperSeconds = InputMax;
+		
+		if (Range.Range.HasLowerBound())
 		{
-			UpperSeconds = InputMax;
+			LowerSeconds = Range.Range.GetLowerBoundValue();
 		}
 
+		if (Range.Range.HasUpperBound())
+		{
+			UpperSeconds = Range.Range.GetUpperBoundValue();
+			if (Range.Range.HasLowerBound() && UpperSeconds == Range.Range.GetLowerBoundValue())
+			{
+				UpperSeconds = InputMax;
+			}
+		}			
+
 		const double BoxStart = PaintArgs.TimeToPixel.SecondsToPixel(LowerSeconds);
 		const double BoxEnd = PaintArgs.TimeToPixel.SecondsToPixel(UpperSeconds);
 		const double BoxSize = BoxEnd - BoxStart;
@@ -361,13 +421,10 @@ TUniquePtr<FCurveModel> CreateCurveEditorModel(const TMovieSceneChannelHandle<FM
 {
 	if (FMovieSceneConstraintChannel* Channel = ChannelHandle.Get())
 	{
-		if (Channel->GetNumKeys() > 0)
+		const UCurveEditorSettings* Settings = GetDefault<UCurveEditorSettings>();
+		if (Settings && Settings->GetShowBars())
 		{
-			const UCurveEditorSettings* Settings = GetDefault<UCurveEditorSettings>();
-			if (Settings == nullptr || Settings->GetShowBars())
-			{
-				return MakeUnique<FConstraintChannelCurveModel>(ChannelHandle, Params.OwningSection, Params.Sequencer);
-			}
+			return MakeUnique<FConstraintChannelCurveModel>(ChannelHandle, Params.OwningSection, Params.Sequencer);
 		}
 	}
 	return nullptr;
diff --git a/Engine/Source/Editor/MovieSceneTools/Private/Channels/SCurveEditorKeyBarView.cpp b/Engine/Source/Editor/MovieSceneTools/Private/Channels/SCurveEditorKeyBarView.cpp
--- a/Engine/Source/Editor/MovieSceneTools/Private/Channels/SCurveEditorKeyBarView.cpp
+++ b/Engine/Source/Editor/MovieSceneTools/Private/Channels/SCurveEditorKeyBarView.cpp
@@ -140,12 +140,15 @@ void SCurveEditorKeyBarView::DrawLabels(const FGeometry& AllottedGeometry, FSlat
 			for (int32 Index = 0;Index < Ranges.Num(); ++Index)
 			{
 				const FKeyBarCurveModel::FBarRange& Range = Ranges[Index];
-				const double LowerSeconds = Range.Range.GetLowerBoundValue();
-				double UpperSeconds = Range.Range.GetUpperBoundValue();
-				if (LowerSeconds == UpperSeconds)
+
+				double LowerSeconds = Range.Range.HasLowerBound() ? Range.Range.GetLowerBoundValue() : InputMin;
+				double UpperSeconds = Range.Range.HasUpperBound() ? Range.Range.GetUpperBoundValue() : InputMax;
+				
+				if (UpperSeconds == LowerSeconds)
 				{
 					UpperSeconds = InputMax;
 				}
+
 				const bool bOutsideUpper = (Index != Ranges.Num() - 1) && UpperSeconds < InputMin;
 				if (LowerSeconds > InputMax || bOutsideUpper) //out of range
 				{
diff --git a/Engine/Source/Editor/MovieSceneTools/Private/Constraints/ComponentConstraintChannelInterface.cpp b/Engine/Source/Editor/MovieSceneTools/Private/Constraints/ComponentConstraintChannelInterface.cpp
--- a/Engine/Source/Editor/MovieSceneTools/Private/Constraints/ComponentConstraintChannelInterface.cpp
+++ b/Engine/Source/Editor/MovieSceneTools/Private/Constraints/ComponentConstraintChannelInterface.cpp
@@ -56,10 +56,10 @@ UWorld* FComponentConstraintChannelInterface::GetHandleWorld(UTransformableHandl
 bool FComponentConstraintChannelInterface::SmartConstraintKey(
 	UTickableTransformConstraint* InConstraint,
 	const TOptional<bool>& InOptActive,
-	const FFrameNumber& InTime,
+	const TRange<FFrameNumber>& InTimeRange,
 	const TSharedPtr<ISequencer>& InSequencer)
 {
-	if (!IsValid(InConstraint->ChildTRSHandle) || !InConstraint->ChildTRSHandle->IsValid())
+	if (!IsValid(InConstraint->ChildTRSHandle) || !InConstraint->ChildTRSHandle->IsValid() || !InSequencer.IsValid())
 	{
 		return false;
 	}
@@ -73,7 +73,11 @@ bool FComponentConstraintChannelInterface::SmartConstraintKey(
 		return false;
 	}
 
-	FScopedTransaction Transaction(NSLOCTEXT("Constraints", "KeyConstraintaKehy", "Key Constraint Key"));
+	// TODO handle range
+	const bool bInfiniteInput = !InTimeRange.HasLowerBound();
+	const TOptional<FFrameNumber> Time = bInfiniteInput ? TOptional<FFrameNumber>() : InTimeRange.GetLowerBoundValue();
+
+	FScopedTransaction Transaction(NSLOCTEXT("Constraints", "KeyConstraintKey", "Key Constraint Key"));
 	TransformSection->Modify();
 	ConstraintSection->Modify();
 
@@ -90,23 +94,24 @@ bool FComponentConstraintChannelInterface::SmartConstraintKey(
 			Manager->RemoveStaticConstraint(InConstraint);
 		}
 		ConstraintSection->AddConstraintChannel(InConstraint);
-		if (InSequencer.IsValid())
-		{
-			InSequencer->RecreateCurveEditor();
-		}
+		InSequencer->RecreateCurveEditor();
 	}
 
 	// add key if needed
 	if (FConstraintAndActiveChannel* Channel = ConstraintSection->GetConstraintChannel(InConstraint->ConstraintID))
 	{
-		bool ActiveValueToBeSet = false;
+		if (bInfiniteInput)
+		{
+			ensure(Channel->ActiveChannel.GetNumKeys() == 0);
+		}
+
+		bool ActiveValueToBeSet = bInfiniteInput;
 		//add key if we can and make sure the key we are setting is what we want
-		if (CanAddKey(Channel->ActiveChannel, InTime, ActiveValueToBeSet) && (InOptActive.IsSet() == false || InOptActive.GetValue() == ActiveValueToBeSet))
+		if (bInfiniteInput || (CanAddKey(Channel->ActiveChannel, *Time, ActiveValueToBeSet) && (!InOptActive.IsSet() || InOptActive.GetValue() == ActiveValueToBeSet)))
 		{
 			const bool bNeedsCompensation = InConstraint->NeedsCompensation();
 
 			//new for compensation
-
 			TGuardValue<bool> CompensateGuard(FMovieSceneConstraintChannelHelper::bDoNotCompensate, true);
 			TGuardValue<bool> RemoveConstraintGuard(FConstraintsManagerController::bDoNotRemoveConstraint, true);
 
@@ -115,26 +120,58 @@ bool FComponentConstraintChannelInterface::SmartConstraintKey(
 			TArray<FFrameNumber> FramesToCompensate;
 			if (bNeedsCompensation)
 			{
-				FMovieSceneConstraintChannelHelper::GetFramesToCompensate<FMovieSceneDoubleChannel>(Channel->ActiveChannel, ActiveValueToBeSet, InTime, Channels, FramesToCompensate);
+				FMovieSceneConstraintChannelHelper::GetFramesToCompensate<FMovieSceneDoubleChannel>(Channel->ActiveChannel, ActiveValueToBeSet, Time, Channels, FramesToCompensate);
 			}
 			else
 			{
-				FramesToCompensate.Add(InTime);
+				if (!bInfiniteInput)
+				{
+					FramesToCompensate.Add(*Time);
+				}
+			}
+
+			const bool bIsActiveByDefault = Channel->ActiveChannel.GetDefault() && *Channel->ActiveChannel.GetDefault();
+			const bool bWasInfinite = bIsActiveByDefault && Channel->ActiveChannel.GetNumKeys() == 0;
+			
+			const FFrameRate TickResolution = InSequencer->GetFocusedTickResolution();
+			const FFrameTime CurrentTime = InSequencer->GetLocalTime().ConvertTo(TickResolution);
+			
+			// check if the first active key will be added before the current frame
+			// if that's the case then the current frame must be added and the compensation evaluator notified
+			TOptional<FFrameNumber> OptCurrentTime;
+			bool bPropagateRelative = true;
+			if (!bInfiniteInput && !bWasInfinite && !bHadConstraintChannel && ActiveValueToBeSet)
+			{
+				Tie(OptCurrentTime, bPropagateRelative) = IsActiveKeyAddedBeforeCurrentTime(CurrentTime.GetFrame(), *Time, Channels, FramesToCompensate);
 			}
 
 			FCompensationEvaluator::FEvalParameters EvalParams(InSequencer, FramesToCompensate);
 			EvalParams.bToActive = ActiveValueToBeSet;
-			if (bHadConstraintChannel && ActiveValueToBeSet)
+			EvalParams.OptCurrentFrame = OptCurrentTime;
+			EvalParams.bPropagateRelative = bPropagateRelative;
+
+			if (bWasInfinite)
+			{
+				EvalParams.bToActive = false;
+			}
+			else if (bIsActiveByDefault && Time)
 			{
 				const TMovieSceneChannelData<const bool> ChannelData = Channel->ActiveChannel.GetData();
-				if (ChannelData.GetTimes().IsEmpty())
+				const TArrayView<const FFrameNumber> Times = ChannelData.GetTimes();
+				const TArrayView<const bool> Values = ChannelData.GetValues();
+				
+				const bool bInfiniteAtLeft = *Time < Times[0] && !Values[0];
+				EvalParams.bInfiniteAtLeft = bInfiniteAtLeft;
+			}
+			
+			if (bHadConstraintChannel && ActiveValueToBeSet && !bWasInfinite)
+			{
+				const TMovieSceneChannelData<const bool> ChannelData = Channel->ActiveChannel.GetData();
+				if (ChannelData.GetTimes().IsEmpty() && !bInfiniteInput)
 				{
 					// this means that the channel already existed but there was no existing key
 					// e.g., when the user deletes keys directly in sequencer without deleting the constraint
-					const FFrameRate TickResolution = InSequencer->GetFocusedTickResolution();
-					const FFrameTime FrameTime = InSequencer->GetLocalTime().ConvertTo(TickResolution);
-					const FFrameNumber Time = FrameTime.GetFrame();
-					EvalParams.bKeepCurrent = !FramesToCompensate.IsEmpty() && FramesToCompensate[0] == Time;
+					EvalParams.bKeepCurrent = !FramesToCompensate.IsEmpty() && FramesToCompensate[0] == CurrentTime.GetFrame();
 				}
 			}
 			
@@ -147,17 +184,17 @@ bool FComponentConstraintChannelInterface::SmartConstraintKey(
 			TArray<FMovieSceneTangentData> Tangents;
 			static constexpr int32 ChannelIndex = 0, NumChannels = 9;
 			
-			if (bNeedsCompensation)
+			if (bNeedsCompensation && !bInfiniteInput)
 			{
-				EvaluateTangentAtThisTime<FMovieSceneDoubleChannel>(ChannelIndex, NumChannels, TransformSection, InTime, Tangents);
+				EvaluateTangentAtThisTime<FMovieSceneDoubleChannel>(ChannelIndex, NumChannels, TransformSection, *Time, Tangents);
 			}
 
 			const EMovieSceneTransformChannel ChannelsToKey = InConstraint->GetChannelsToKey();
 
 			// add child's transform key at Time-1 to keep animation
-			if (bNeedsCompensation)
+			if (bNeedsCompensation && !bInfiniteInput)
 			{
-				const FFrameNumber TimeMinusOne(InTime - 1);
+				const FFrameNumber TimeMinusOne(*Time - 1);
 
 				TArray<FTransform> TransformMinusOne({ChildLocals[0]} );
 				if (ConstraintSection != TransformSection)
@@ -170,32 +207,51 @@ bool FComponentConstraintChannelInterface::SmartConstraintKey(
 			}
 
 			// add active key
+			if (bInfiniteInput)
+			{
+				constexpr bool bActiveByDefault = true;
+				Channel->ActiveChannel.SetDefault(bActiveByDefault);
+			}
+			else
 			{
 				TMovieSceneChannelData<bool> ChannelData = Channel->ActiveChannel.GetData();
-				ChannelData.UpdateOrAddKey(InTime, ActiveValueToBeSet);
+
+				const TArrayView<const FFrameNumber> FrameNumbers = ChannelData.GetTimes();
+				const bool bRemoveInfinite = bIsActiveByDefault && ActiveValueToBeSet && (FrameNumbers.IsEmpty() || *Time < FrameNumbers[0]); 
+				
+				ChannelData.UpdateOrAddKey(*Time, ActiveValueToBeSet);
+
+				if (bRemoveInfinite)
+				{
+					Channel->ActiveChannel.RemoveDefault();
+				}
 			}
 
 			// compensate
 			{
 				// we need to remove the first transforms as we store NumFrames+1 transforms
-				ChildLocals.RemoveAt(0);
+				if (!ChildLocals.IsEmpty())
+				{
+					ChildLocals.RemoveAt(0);
+				}
 
 				if (ConstraintSection != TransformSection)
 				{
-					RecomposeTransforms(InSequencer, ComponentHandle->Component.Get(), TransformSection, FramesToCompensate,ChildLocals);
+					RecomposeTransforms(InSequencer, ComponentHandle->Component.Get(), TransformSection, FramesToCompensate, ChildLocals);
 				}
 
 				// add keys
 				MovieSceneToolHelpers::AddTransformKeys(TransformSection, FramesToCompensate, ChildLocals, ChannelsToKey);
 
 				// set tangents at Time
-				if (bNeedsCompensation)
+				if (bNeedsCompensation && !bInfiniteInput)
 				{
-					SetTangentsAtThisTime<FMovieSceneDoubleChannel>(ChannelIndex,NumChannels, TransformSection, InTime, Tangents);
+					SetTangentsAtThisTime<FMovieSceneDoubleChannel>(ChannelIndex, NumChannels, TransformSection, *Time, Tangents);
 				}
 			}
 
-			// cleanup next keys that have the same value as they are useless 
+			// cleanup next keys that have the same value as they are useless
+			if (!bInfiniteInput)
 			{
 				TArray<FFrameNumber> TimesRemoved;
 				CleanDuplicates(Channel->ActiveChannel, TimesRemoved);
diff --git a/Engine/Source/Editor/MovieSceneTools/Private/Constraints/ComponentConstraintChannelInterface.h b/Engine/Source/Editor/MovieSceneTools/Private/Constraints/ComponentConstraintChannelInterface.h
--- a/Engine/Source/Editor/MovieSceneTools/Private/Constraints/ComponentConstraintChannelInterface.h
+++ b/Engine/Source/Editor/MovieSceneTools/Private/Constraints/ComponentConstraintChannelInterface.h
@@ -26,7 +26,7 @@ struct FComponentConstraintChannelInterface : ITransformConstraintChannelInterfa
 	/** Add an active/inactive key to the constraint channel if needed and does the transform compensation on the component's transform channels. */
 	virtual bool SmartConstraintKey(
 		UTickableTransformConstraint* InConstraint, const TOptional<bool>& InOptActive,
-		const FFrameNumber& InTime, const TSharedPtr<ISequencer>& InSequencer) override;
+		const TRange<FFrameNumber>& InTimeRange, const TSharedPtr<ISequencer>& InSequencer) override;
 
 	/** Add keys on the component's transform channels. */
 	virtual void AddHandleTransformKeys(
diff --git a/Engine/Source/Editor/MovieSceneTools/Private/Constraints/MovieSceneConstraintChannelHelper.cpp b/Engine/Source/Editor/MovieSceneTools/Private/Constraints/MovieSceneConstraintChannelHelper.cpp
--- a/Engine/Source/Editor/MovieSceneTools/Private/Constraints/MovieSceneConstraintChannelHelper.cpp
+++ b/Engine/Source/Editor/MovieSceneTools/Private/Constraints/MovieSceneConstraintChannelHelper.cpp
@@ -2,7 +2,6 @@
 
 #include "Constraints/MovieSceneConstraintChannelHelper.inl"
 
-#include "ISequencer.h"
 #include "Evaluation/MovieSceneEvaluationTemplateInstance.h"
 
 #include "LevelSequence.h"
@@ -25,7 +24,6 @@
 #include "Framework/Notifications/NotificationManager.h"
 #include "Widgets/Notifications/SNotificationList.h"
 
-
 /*
 *
 *  FCompensationEvaluator
@@ -39,6 +37,8 @@ FCompensationEvaluator::FCompensationEvaluator(UTickableTransformConstraint* InC
 
 void FCompensationEvaluator::ComputeLocalTransforms(UWorld* InWorld, const FEvalParameters& InEvalParams)
 {
+	using namespace UE::TransformConstraintUtil;
+	
 	if (!InEvalParams.IsValid())
 	{
 		return;
@@ -61,7 +61,7 @@ void FCompensationEvaluator::ComputeLocalTransforms(UWorld* InWorld, const FEval
 	auto GetLastActiveConstraint = [ConstraintsMinusThis]()
 	{
 		// find last active constraint in the list that is different than the one we want to compensate for
-		const int32 LastActiveIndex = UE::TransformConstraintUtil::GetLastActiveConstraintIndex(ConstraintsMinusThis);
+		const int32 LastActiveIndex = GetLastActiveConstraintIndex(ConstraintsMinusThis);
 
 		// if found, return its parent global transform
 		return LastActiveIndex > INDEX_NONE ? Cast<UTickableTransformConstraint>(ConstraintsMinusThis[LastActiveIndex]) : nullptr;
@@ -75,8 +75,11 @@ void FCompensationEvaluator::ComputeLocalTransforms(UWorld* InWorld, const FEval
 	const TConstArrayView<FFrameNumber> Frames(InEvalParams.Frames);
 	const int32 NumFrames = Frames.Num();
 
+	// store current state
 	const FTransform CurrentLocal = Handle->GetLocalTransform();
 	const FTransform CurrentGlobal = Handle->GetGlobalTransform();
+	const FTransform CurrentSpaceGlobal = Constraint->GetParentGlobalTransform();
+	const FTransform CurrentRelative = ComputeRelativeTransform(CurrentLocal, CurrentGlobal, CurrentSpaceGlobal, Constraint);
 
 	// resize arrays to num frames + 1 as we also evaluate at InFrames[0]-1
 	ChildLocals.SetNum(NumFrames + 1);
@@ -85,16 +88,10 @@ void FCompensationEvaluator::ComputeLocalTransforms(UWorld* InWorld, const FEval
 
 	// avoid transacting when evaluating sequencer
 	TGuardValue<ITransaction*> TransactionGuard(GUndo, nullptr);
-	
-	const TArray<IMovieSceneToolsAnimationBakeHelper*>& BakeHelpers = FMovieSceneToolsModule::Get().GetAnimationBakeHelpers();
-	for (IMovieSceneToolsAnimationBakeHelper* BakeHelper : BakeHelpers)
-	{
-		if (BakeHelper)
-		{
-			BakeHelper->StartBaking(MovieScene);
-		}
-	}
 
+	const FMovieSceneToolsAnimationBakingWrapper BakingWrapper(MovieScene);
+	BakingWrapper.Start();
+	
 	// get all constraints to evaluate
 	const FConstraintsManagerController& Controller = FConstraintsManagerController::Get(InWorld);
 	static constexpr bool bSorted = true, bTickHandles = true;
@@ -102,6 +99,9 @@ void FCompensationEvaluator::ComputeLocalTransforms(UWorld* InWorld, const FEval
 
 	FMovieSceneInverseSequenceTransform LocalToRootTransform = Sequencer->GetFocusedMovieSceneSequenceTransform().Inverse();
 
+	const bool bCreatedBeforeCurrentFrame = InEvalParams.OptCurrentFrame.IsSet();
+	const bool bPropagateRelative = InEvalParams.bPropagateRelative;
+	
 	for (int32 Index = 0; Index < NumFrames + 1; ++Index)
 	{
 		FFrameNumber FrameNumber = (Index == 0) ? Frames[0] - 1 : Frames[Index - 1];
@@ -114,33 +114,24 @@ void FCompensationEvaluator::ComputeLocalTransforms(UWorld* InWorld, const FEval
 		FrameNumber = RootTime->GetFrame();
 
 		// evaluate animation
-		const FMovieSceneEvaluationRange EvaluationRange = FMovieSceneEvaluationRange(FFrameTime(FrameNumber), TickResolution);
-		const FMovieSceneContext Context = FMovieSceneContext(EvaluationRange, PlaybackStatus).SetHasJumped(true);
-
-		for (IMovieSceneToolsAnimationBakeHelper* BakeHelper : BakeHelpers)
 		{
-			if (BakeHelper)
-			{
-				BakeHelper->PreEvaluation(MovieScene, FrameNumber);
-			}
-		}
-		Sequencer->GetEvaluationTemplate().EvaluateSynchronousBlocking(Context);
+			const FMovieSceneEvaluationRange EvaluationRange = FMovieSceneEvaluationRange(FFrameTime(FrameNumber), TickResolution);
+			const FMovieSceneContext Context = FMovieSceneContext(EvaluationRange, PlaybackStatus).SetHasJumped(true);
 
-		// evaluate constraints
-		for (const TWeakObjectPtr<UTickableConstraint>& InConstraint : AllConstraints)
-		{
-			if (InConstraint.IsValid())
-			{
-				InConstraint->Evaluate(bTickHandles);
-			}
-		}
+			BakingWrapper.PreEvaluate(FrameNumber);
+		
+			Sequencer->GetEvaluationTemplate().EvaluateSynchronousBlocking(Context);
 
-		for (IMovieSceneToolsAnimationBakeHelper* BakeHelper : BakeHelpers)
-		{
-			if (BakeHelper)
+			// evaluate constraints
+			for (const TWeakObjectPtr<UTickableConstraint>& InConstraint : AllConstraints)
 			{
-				BakeHelper->PostEvaluation(MovieScene, FrameNumber);
+				if (InConstraint.IsValid())
+				{
+					InConstraint->Evaluate(bTickHandles);
+				}
 			}
+		
+			BakingWrapper.PostEvaluate(FrameNumber);
 		}
 
 		FTransform& ChildLocal = ChildLocals[Index];
@@ -168,12 +159,31 @@ void FCompensationEvaluator::ComputeLocalTransforms(UWorld* InWorld, const FEval
 			// and constraint space for all other times
 			if (Index == 0)
 			{
-				if (LastConstraint)
+				if (bCreatedBeforeCurrentFrame)
+				{
+					const FFrameNumber CurrentFrame = *InEvalParams.OptCurrentFrame;
+					if ((bPropagateRelative && FrameNumber < CurrentFrame) || FrameNumber == CurrentFrame)
+					{
+						SpaceGlobal = Constraint->GetParentGlobalTransform();
+						Handle->SetGlobalTransform(CurrentRelative * SpaceGlobal);
+						Handle->PreEvaluate();
+						ChildLocal = Handle->GetLocalTransform();
+						ChildGlobal = Handle->GetGlobalTransform();
+						
+						if (LastConstraint)
+						{
+							SpaceGlobal = LastConstraint->GetParentGlobalTransform();
+							if (TOptional<FTransform> Relative = GetConstraintsRelativeTransform(ConstraintsMinusThis, ChildLocal, ChildGlobal))
+							{
+								ChildLocal = *Relative;
+							}
+						}
+					}
+				}
+				else if (LastConstraint)
 				{
 					SpaceGlobal = LastConstraint->GetParentGlobalTransform();
-					TOptional<FTransform> Relative =
-						UE::TransformConstraintUtil::GetConstraintsRelativeTransform(ConstraintsMinusThis, ChildLocal, ChildGlobal);
-					if (Relative)
+					if (TOptional<FTransform> Relative = GetConstraintsRelativeTransform(ConstraintsMinusThis, ChildLocal, ChildGlobal))
 					{
 						ChildLocal = *Relative;
 					}
@@ -185,9 +195,43 @@ void FCompensationEvaluator::ComputeLocalTransforms(UWorld* InWorld, const FEval
 			}
 			else
 			{
-				SpaceGlobal = Constraint->GetParentGlobalTransform();
-				ChildLocal = UE::TransformConstraintUtil::ComputeRelativeTransform(
-					ChildLocal, ChildGlobal, SpaceGlobal, Constraint);
+				if (InEvalParams.bInfiniteAtLeft)
+				{
+					if (Index == 1)
+					{
+						SpaceGlobal = Constraint->GetParentGlobalTransform();
+						ChildLocal = ComputeRelativeTransform(ChildLocal, ChildGlobal, SpaceGlobal, Constraint);
+					}
+					else
+					{
+						if (LastConstraint)
+						{
+							SpaceGlobal = LastConstraint->GetParentGlobalTransform();
+							if (TOptional<FTransform> Relative = GetConstraintsRelativeTransform(ConstraintsMinusThis, ChildLocal, ChildGlobal))
+							{
+								ChildLocal = *Relative;
+							}
+						}
+					}
+				}
+				else if (bCreatedBeforeCurrentFrame)
+				{
+					const FFrameNumber CurrentFrame = *InEvalParams.OptCurrentFrame;
+					if ((bPropagateRelative && FrameNumber < CurrentFrame) || FrameNumber == CurrentFrame)
+					{
+					 	ChildLocal = CurrentRelative;
+					}
+					else
+					{
+						SpaceGlobal = Constraint->GetParentGlobalTransform();
+						ChildLocal = ComputeRelativeTransform(ChildLocal, ChildGlobal, SpaceGlobal, Constraint);						
+					}
+				}
+				else
+				{
+					SpaceGlobal = Constraint->GetParentGlobalTransform();
+					ChildLocal = ComputeRelativeTransform(ChildLocal, ChildGlobal, SpaceGlobal, Constraint);
+				}
 			}
 		}
 		else
@@ -197,17 +241,14 @@ void FCompensationEvaluator::ComputeLocalTransforms(UWorld* InWorld, const FEval
 			if (Index == 0)
 			{
 				SpaceGlobal = Constraint->GetParentGlobalTransform();
-				ChildLocal = UE::TransformConstraintUtil::ComputeRelativeTransform(
-					ChildLocal, ChildGlobal, SpaceGlobal, Constraint);
+				ChildLocal = ComputeRelativeTransform(ChildLocal, ChildGlobal, SpaceGlobal, Constraint);
 			}
 			else
 			{
 				if (LastConstraint)
 				{
 					SpaceGlobal = LastConstraint->GetParentGlobalTransform();
-					TOptional<FTransform> Relative =
-						UE::TransformConstraintUtil::GetConstraintsRelativeTransform(ConstraintsMinusThis, ChildLocal, ChildGlobal);
-					if (Relative)
+					if (TOptional<FTransform> Relative = GetConstraintsRelativeTransform(ConstraintsMinusThis, ChildLocal, ChildGlobal))
 					{
 						ChildLocal = *Relative;
 					}
@@ -215,13 +256,9 @@ void FCompensationEvaluator::ComputeLocalTransforms(UWorld* InWorld, const FEval
 			}
 		}
 	}
-	for (IMovieSceneToolsAnimationBakeHelper* BakeHelper : BakeHelpers)
-	{
-		if (BakeHelper)
-		{
-			BakeHelper->StopBaking(MovieScene);
-		}
-	}
+	
+	BakingWrapper.Stop();
+	
 	//get back to where we are at, should also make sure things are active
 	Sequencer->ForceEvaluate();
 }
@@ -276,14 +313,9 @@ void FCompensationEvaluator::ComputeLocalTransformsForBaking(UWorld* InWorld, co
 	// avoid transacting when evaluating sequencer
 	TGuardValue<ITransaction*> TransactionGuard(GUndo, nullptr);
 	
-	const TArray<IMovieSceneToolsAnimationBakeHelper*>& BakeHelpers = FMovieSceneToolsModule::Get().GetAnimationBakeHelpers();
-	for (IMovieSceneToolsAnimationBakeHelper* BakeHelper : BakeHelpers)
-	{
-		if (BakeHelper)
-		{
-			BakeHelper->StartBaking(MovieScene);
-		}
-	}
+	const FMovieSceneToolsAnimationBakingWrapper BakingWrapper(MovieScene);
+	BakingWrapper.Start();
+	
 	FMovieSceneInverseSequenceTransform LocalToRootTransform = Sequencer->GetFocusedMovieSceneSequenceTransform().Inverse();
 
 	for (int32 Index = 0; Index < NumFrames; ++Index)
@@ -300,13 +332,8 @@ void FCompensationEvaluator::ComputeLocalTransformsForBaking(UWorld* InWorld, co
 		const FMovieSceneEvaluationRange EvaluationRange = FMovieSceneEvaluationRange(FFrameTime(FrameNumber), TickResolution);
 		const FMovieSceneContext Context = FMovieSceneContext(EvaluationRange, PlaybackStatus).SetHasJumped(true);
 
-		for (IMovieSceneToolsAnimationBakeHelper* BakeHelper : BakeHelpers)
-		{
-			if (BakeHelper)
-			{
-				BakeHelper->PreEvaluation(MovieScene, FrameNumber);
-			}
-		}
+		BakingWrapper.PreEvaluate(FrameNumber);
+
 		Sequencer->GetEvaluationTemplate().EvaluateSynchronousBlocking(Context);
 
 		// evaluate constraints
@@ -318,16 +345,7 @@ void FCompensationEvaluator::ComputeLocalTransformsForBaking(UWorld* InWorld, co
 			}
 		}
 
-		for (IMovieSceneToolsAnimationBakeHelper* BakeHelper : BakeHelpers)
-		{
-			if (BakeHelper)
-			{
-				BakeHelper->PostEvaluation(MovieScene, FrameNumber);
-			}
-		}
-
-		// evaluate ControlRig?
-		// ControlRig->Evaluate_AnyThread();
+		BakingWrapper.PostEvaluate(FrameNumber);
 
 		FTransform& ChildLocal = ChildLocals[Index];
 		FTransform& ChildGlobal = ChildGlobals[Index];
@@ -349,13 +367,8 @@ void FCompensationEvaluator::ComputeLocalTransformsForBaking(UWorld* InWorld, co
 			}
 		}
 	}
-	for (IMovieSceneToolsAnimationBakeHelper* BakeHelper : BakeHelpers)
-	{
-		if (BakeHelper)
-		{
-			BakeHelper->StopBaking(MovieScene);
-		}
-	}
+
+	BakingWrapper.Stop();
 
 	const bool bIsValidAfterBaking = IsValid(Handle) && Handle->IsValid();
 	if (!bIsValidAfterBaking)
@@ -410,14 +423,9 @@ void FCompensationEvaluator::ComputeLocalTransformsBeforeDeletion(UWorld* InWorl
 	// avoid transacting when evaluating sequencer
 	TGuardValue<ITransaction*> TransactionGuard(GUndo, nullptr);
 	
-	const TArray<IMovieSceneToolsAnimationBakeHelper*>& BakeHelpers = FMovieSceneToolsModule::Get().GetAnimationBakeHelpers();
-	for (IMovieSceneToolsAnimationBakeHelper* BakeHelper : BakeHelpers)
-	{
-		if (BakeHelper)
-		{
-			BakeHelper->StartBaking(MovieScene);
-		}
-	}
+	const FMovieSceneToolsAnimationBakingWrapper BakingWrapper(MovieScene);
+	BakingWrapper.Start();
+	
 	FMovieSceneInverseSequenceTransform LocalToRootTransform = Sequencer->GetFocusedMovieSceneSequenceTransform().Inverse();
 
 	for (int32 Index = 0; Index < NumFrames; ++Index)
@@ -434,13 +442,8 @@ void FCompensationEvaluator::ComputeLocalTransformsBeforeDeletion(UWorld* InWorl
 		const FMovieSceneEvaluationRange EvaluationRange = FMovieSceneEvaluationRange(FFrameTime(FrameNumber), TickResolution);
 		const FMovieSceneContext Context = FMovieSceneContext(EvaluationRange, PlaybackStatus).SetHasJumped(true);
 
-		for (IMovieSceneToolsAnimationBakeHelper* BakeHelper : BakeHelpers)
-		{
-			if (BakeHelper)
-			{
-				BakeHelper->PreEvaluation(MovieScene, FrameNumber);
-			}
-		}
+		BakingWrapper.PreEvaluate(FrameNumber);
+		
 		Sequencer->GetEvaluationTemplate().EvaluateSynchronousBlocking(Context);
 
 		// evaluate constraints
@@ -452,16 +455,8 @@ void FCompensationEvaluator::ComputeLocalTransformsBeforeDeletion(UWorld* InWorl
 			}
 		}
 		
-		for (IMovieSceneToolsAnimationBakeHelper* BakeHelper : BakeHelpers)
-		{
-			if (BakeHelper)
-			{
-				BakeHelper->PostEvaluation(MovieScene, FrameNumber);
-			}
-		}
-		// evaluate ControlRig?
-		// ControlRig->Evaluate_AnyThread();
-
+		BakingWrapper.PostEvaluate(FrameNumber);
+		
 		FTransform& ChildLocal = ChildLocals[Index];
 		FTransform& ChildGlobal = ChildGlobals[Index];
 		FTransform& SpaceGlobal = SpaceGlobals[Index];
@@ -482,13 +477,8 @@ void FCompensationEvaluator::ComputeLocalTransformsBeforeDeletion(UWorld* InWorl
 			}
 		}
 	}
-	for (IMovieSceneToolsAnimationBakeHelper* BakeHelper : BakeHelpers)
-	{
-		if (BakeHelper)
-		{
-			BakeHelper->StopBaking(MovieScene);
-		}
-	}
+
+	BakingWrapper.Stop();
 }
 
 void FCompensationEvaluator::ComputeCompensation(UWorld* InWorld, const TSharedPtr<ISequencer>& InSequencer, const FFrameNumber& InTime)
@@ -519,18 +509,13 @@ void FCompensationEvaluator::ComputeCompensation(UWorld* InWorld, const TSharedP
 	TGuardValue<ITransaction*> TransactionGuard(GUndo, nullptr);
 	
 	UMovieScene* MovieScene = InSequencer->GetFocusedMovieSceneSequence()->GetMovieScene();
-	const TArray<IMovieSceneToolsAnimationBakeHelper*>& BakeHelpers = FMovieSceneToolsModule::Get().GetAnimationBakeHelpers();
-	for (IMovieSceneToolsAnimationBakeHelper* BakeHelper : BakeHelpers)
-	{
-		if (BakeHelper)
-		{
-			BakeHelper->StartBaking(MovieScene);
-		}
-	}
+	
+	const FMovieSceneToolsAnimationBakingWrapper BakingWrapper(MovieScene);
+	BakingWrapper.Start();
 
 	FMovieSceneInverseSequenceTransform LocalToRootTransform = InSequencer->GetFocusedMovieSceneSequenceTransform().Inverse();
 
-	auto EvaluateAt = [Handle = this->Handle, InSequencer, &AllConstraints, &BakeHelpers, &LocalToRootTransform](FFrameNumber InFrame)
+	auto EvaluateAt = [Handle = this->Handle, InSequencer, &AllConstraints, &BakingWrapper, &LocalToRootTransform](FFrameNumber InFrame)
 	{
 		TOptional<FFrameTime> RootTime = LocalToRootTransform.TryTransformTime(InFrame);
 		if (!RootTime)
@@ -547,13 +532,8 @@ void FCompensationEvaluator::ComputeCompensation(UWorld* InWorld, const TSharedP
 		const FMovieSceneEvaluationRange EvaluationRange0 = FMovieSceneEvaluationRange(FFrameTime(InFrame), TickResolution);
 		const FMovieSceneContext Context0 = FMovieSceneContext(EvaluationRange0, PlaybackStatus).SetHasJumped(true);
 
-		for (IMovieSceneToolsAnimationBakeHelper* BakeHelper : BakeHelpers)
-		{
-			if (BakeHelper)
-			{
-				BakeHelper->PreEvaluation(MovieScene, InFrame);
-			}
-		}
+		BakingWrapper.PreEvaluate(InFrame);
+			
 		InSequencer->GetEvaluationTemplate().EvaluateSynchronousBlocking(Context0);
 
 		for (const TWeakObjectPtr<UTickableConstraint>& InConstraint : AllConstraints)
@@ -564,13 +544,7 @@ void FCompensationEvaluator::ComputeCompensation(UWorld* InWorld, const TSharedP
 			}
 		}
 
-		for (IMovieSceneToolsAnimationBakeHelper* BakeHelper : BakeHelpers)
-		{
-			if (BakeHelper)
-			{
-				BakeHelper->PostEvaluation(MovieScene, InFrame);
-			}
-		}
+		BakingWrapper.PostEvaluate(InFrame);
 
 		if (Handle)
 		{
@@ -593,13 +567,7 @@ void FCompensationEvaluator::ComputeCompensation(UWorld* InWorld, const TSharedP
 	EvaluateAt(InTime - 1);
 	ChildLocals[0] = Handle->GetLocalTransform();
 
-	for (IMovieSceneToolsAnimationBakeHelper* BakeHelper : BakeHelpers)
-	{
-		if (BakeHelper)
-		{
-			BakeHelper->StopBaking(MovieScene);
-		}
-	}
+	BakingWrapper.Stop();
 	
 	// if constraint at T-1 then switch to its space
 	if (const UTickableTransformConstraint* LastConstraint = GetLastActiveConstraint())
@@ -649,7 +617,7 @@ void FCompensationEvaluator::CacheTransforms(UWorld* InWorld, const FEvalParamet
 	ChildGlobals.SetNum(NumFrames);
 	SpaceGlobals.SetNum(NumFrames);
 
-	const TArray<IMovieSceneToolsAnimationBakeHelper*>& BakeHelpers = FMovieSceneToolsModule::Get().GetAnimationBakeHelpers();
+	const FMovieSceneToolsAnimationBakingWrapper BakingWrapper(MovieScene);
 
 	FMovieSceneInverseSequenceTransform LocalToRootTransform = Sequencer->GetFocusedMovieSceneSequenceTransform().Inverse();
 
@@ -664,13 +632,8 @@ void FCompensationEvaluator::CacheTransforms(UWorld* InWorld, const FEvalParamet
 		const FMovieSceneEvaluationRange EvaluationRange = FMovieSceneEvaluationRange(RootTime.GetValue(), TickResolution);
 		const FMovieSceneContext Context = FMovieSceneContext(EvaluationRange, PlaybackStatus).SetHasJumped(true);
 
-		for (IMovieSceneToolsAnimationBakeHelper* BakeHelper : BakeHelpers)
-		{
-			if (BakeHelper)
-			{
-				BakeHelper->PreEvaluation(MovieScene, InFrame);
-			}
-		}
+		BakingWrapper.PreEvaluate(InFrame);
+		
 		Sequencer->GetEvaluationTemplate().EvaluateSynchronousBlocking(Context);
 
 		// evaluate constraints
@@ -682,25 +645,13 @@ void FCompensationEvaluator::CacheTransforms(UWorld* InWorld, const FEvalParamet
 			}
 		}
 		
-		for (IMovieSceneToolsAnimationBakeHelper* BakeHelper : BakeHelpers)
-		{
-			if (BakeHelper)
-			{
-				BakeHelper->PostEvaluation(MovieScene, InFrame);
-			}
-		}
+		BakingWrapper.PostEvaluate(InFrame);
 	};
 
 	// avoid transacting when evaluating sequencer
 	TGuardValue<ITransaction*> TransactionGuard(GUndo, nullptr);
 	
-	for (IMovieSceneToolsAnimationBakeHelper* BakeHelper : BakeHelpers)
-	{
-		if (BakeHelper)
-		{
-			BakeHelper->StartBaking(MovieScene);
-		}
-	}
+	BakingWrapper.Start();
 	
 	for (int32 Index = 0; Index < NumFrames; ++Index)
 	{
@@ -713,13 +664,7 @@ void FCompensationEvaluator::CacheTransforms(UWorld* InWorld, const FEvalParamet
 		SpaceGlobals[Index] = Constraint->GetParentGlobalTransform();
 	}
 	
-	for (IMovieSceneToolsAnimationBakeHelper* BakeHelper : BakeHelpers)
-	{
-		if (BakeHelper)
-		{
-			BakeHelper->StopBaking(MovieScene);
-		}
-	}
+	BakingWrapper.Stop();
 }
 
 void FCompensationEvaluator::ComputeCurrentTransforms(UWorld* InWorld)
@@ -1065,9 +1010,11 @@ static void ShowNonSpawnableWarning(
 
 bool FMovieSceneConstraintChannelHelper::AddConstraintToSequencer(
 	const TSharedPtr<ISequencer>& InSequencer,
-	UTickableTransformConstraint* InConstraint)
+	UTickableTransformConstraint* InConstraint,
+	const FSequencerCreationOptions& InOptions)
 {
-	if (!InSequencer.IsValid() || !InSequencer->GetFocusedMovieSceneSequence())
+	UMovieSceneSequence* MovieSceneSequence = InSequencer.IsValid() ? InSequencer->GetFocusedMovieSceneSequence() : nullptr;
+	if (!MovieSceneSequence)
 	{
 		return false;
 	}
@@ -1112,9 +1059,29 @@ bool FMovieSceneConstraintChannelHelper::AddConstraintToSequencer(
 
 	const FFrameRate TickResolution = InSequencer->GetFocusedTickResolution();
 	const FFrameTime FrameTime = InSequencer->GetLocalTime().ConvertTo(TickResolution);
-	const FFrameNumber Time = FrameTime.GetFrame();
+
+	TRange<FFrameNumber> Range;
 	
-	return Interface->SmartConstraintKey(InConstraint, TOptional<bool>(), Time, InSequencer);
+	switch (InOptions.CreationTime)
+	{
+	case ECreationTime::CurrentFrame:
+		Range = TRange<FFrameNumber>(FrameTime.GetFrame());
+		break;
+	case ECreationTime::FromStart:
+	{
+		const UMovieScene* MovieScene = MovieSceneSequence->GetMovieScene();
+		if (ensure(MovieScene))
+		{
+			Range = TRange<FFrameNumber>(MovieScene->GetPlaybackRange().GetLowerBoundValue());
+		}
+		break;
+	}
+	case ECreationTime::Infinite:
+	default:
+		break;
+	}
+
+	return Interface->SmartConstraintKey(InConstraint, TOptional<bool>(), Range, InSequencer);
 }
 
 bool FMovieSceneConstraintChannelHelper::SmartConstraintKey(
@@ -1150,7 +1117,7 @@ bool FMovieSceneConstraintChannelHelper::SmartConstraintKey(
 	CreateBindingIDForHandle(InSequencer, InConstraint->ChildTRSHandle);
 	CreateBindingIDForHandle(InSequencer, InConstraint->ParentTRSHandle);
 
-	const bool bSucceeded = Interface->SmartConstraintKey(InConstraint, InOptActive, Time, InSequencer);
+	const bool bSucceeded = Interface->SmartConstraintKey(InConstraint, InOptActive, TRange<FFrameNumber>(Time), InSequencer);
 
 	return bSucceeded;
 }
diff --git a/Engine/Source/Editor/MovieSceneTools/Private/Constraints/TransformConstraintChannelInterface.cpp b/Engine/Source/Editor/MovieSceneTools/Private/Constraints/TransformConstraintChannelInterface.cpp
--- a/Engine/Source/Editor/MovieSceneTools/Private/Constraints/TransformConstraintChannelInterface.cpp
+++ b/Engine/Source/Editor/MovieSceneTools/Private/Constraints/TransformConstraintChannelInterface.cpp
@@ -61,9 +61,19 @@ bool ITransformConstraintChannelInterface::CanAddKey(const FMovieSceneConstraint
 	const TArrayView<const FFrameNumber> Times = ChannelData.GetTimes();
 	if (Times.IsEmpty())
 	{
-		ActiveValue = true;
+		ActiveValue = !InActiveChannel.IsInfinite();
 		return true;
 	}
+
+	if (const TOptional<bool> OptDefault = InActiveChannel.GetDefault())
+	{
+		const TArrayView<const bool> Values = ChannelData.GetValues();
+		if (InTime < Times[0] && !Values[0])
+		{
+			ActiveValue = OptDefault && *OptDefault;
+			return true;
+		}
+	}
 	
 	InActiveChannel.Evaluate(InTime, ActiveValue);
 	ActiveValue = !ActiveValue;
@@ -97,7 +107,9 @@ void ITransformConstraintChannelInterface::CleanDuplicates(FMovieSceneConstraint
 	}
 
 	// no need to keep single inactive key
-	if (!ActiveValues[0])
+	const TOptional<bool> OptDefault = InOutActiveChannel.GetDefault();
+	const bool bIsDefaultTrue = OptDefault && *OptDefault;
+	if (!ActiveValues[0] && !bIsDefaultTrue)
 	{
 		OutTimesRemoved.Add(Times[0]);
 		ChannelData.RemoveKey(0);
diff --git a/Engine/Source/Editor/MovieSceneTools/Private/MovieSceneToolsModule.cpp b/Engine/Source/Editor/MovieSceneTools/Private/MovieSceneToolsModule.cpp
--- a/Engine/Source/Editor/MovieSceneTools/Private/MovieSceneToolsModule.cpp
+++ b/Engine/Source/Editor/MovieSceneTools/Private/MovieSceneToolsModule.cpp
@@ -905,6 +905,69 @@ void FMovieSceneToolsModule::CustomizeKeyStructInstancedPropertyTypes(TSharedRef
 	}
 }
 
+FMovieSceneToolsAnimationBakingWrapper::FMovieSceneToolsAnimationBakingWrapper(UMovieScene* InMovieScene)
+	: MovieScene(InMovieScene)
+	, BakeHelpers(FMovieSceneToolsModule::Get().GetAnimationBakeHelpers())
+{
+	ensure(InMovieScene);
+}
+
+void FMovieSceneToolsAnimationBakingWrapper::Start() const
+{
+	if (MovieScene)
+	{
+		for (IMovieSceneToolsAnimationBakeHelper* BakeHelper : BakeHelpers)
+		{
+			if (BakeHelper)
+			{
+				BakeHelper->StartBaking(MovieScene);
+			}
+		}
+	}			
+}
+
+void FMovieSceneToolsAnimationBakingWrapper::PreEvaluate(const FFrameNumber& Frame) const
+{
+	if (MovieScene)
+	{
+		for (IMovieSceneToolsAnimationBakeHelper* BakeHelper : BakeHelpers)
+		{
+			if (BakeHelper)
+			{
+				BakeHelper->PreEvaluation(MovieScene, Frame);
+			}
+		}
+	}
+}
+
+void FMovieSceneToolsAnimationBakingWrapper::PostEvaluate(const FFrameNumber& Frame) const
+{
+	if (MovieScene)
+	{
+		for (IMovieSceneToolsAnimationBakeHelper* BakeHelper : BakeHelpers)
+		{
+			if (BakeHelper)
+			{
+				BakeHelper->PostEvaluation(MovieScene, Frame);
+			}
+		}
+	}
+}
+
+void FMovieSceneToolsAnimationBakingWrapper::Stop() const
+{
+	if (MovieScene)
+	{
+		for (IMovieSceneToolsAnimationBakeHelper* BakeHelper : BakeHelpers)
+		{
+			if (BakeHelper)
+			{
+				BakeHelper->StopBaking(MovieScene);
+			}
+		}
+	}
+}
+
 IMPLEMENT_MODULE( FMovieSceneToolsModule, MovieSceneTools );
 
 #undef LOCTEXT_NAMESPACE
diff --git a/Engine/Source/Editor/MovieSceneTools/Public/Constraints/MovieSceneConstraintChannelHelper.h b/Engine/Source/Editor/MovieSceneTools/Public/Constraints/MovieSceneConstraintChannelHelper.h
--- a/Engine/Source/Editor/MovieSceneTools/Public/Constraints/MovieSceneConstraintChannelHelper.h
+++ b/Engine/Source/Editor/MovieSceneTools/Public/Constraints/MovieSceneConstraintChannelHelper.h
@@ -6,8 +6,9 @@
 #include "UObject/ObjectPtr.h"
 #include "MovieSceneSection.h"
 
-#define UE_API MOVIESCENETOOLS_API
+#include "MovieSceneConstraintChannelHelper.generated.h"
 
+#define UE_API MOVIESCENETOOLS_API
 
 class IMovieSceneConstrainedSection;
 struct ITransformConstraintChannelInterface;
@@ -23,6 +24,25 @@ struct FFrameNumber;
 struct FConstraintAndActiveChannel;
 enum class EMovieSceneKeyInterpolation : uint8;
 
+UENUM()
+enum class ECreationTime : uint8
+{
+	CurrentFrame,
+	FromStart,
+	Infinite,
+	// FromRange
+};
+
+USTRUCT()
+struct MOVIESCENETOOLS_API FSequencerCreationOptions
+{
+	GENERATED_BODY()
+
+	/** Whether to create the constraint at the current frame, at the beginning of the sequence, or infinitely. */
+	UPROPERTY(EditAnywhere, Category=Options)
+	ECreationTime CreationTime = ECreationTime::CurrentFrame;
+};
+
 struct FCompensationEvaluator
 {
 public:
@@ -48,6 +68,9 @@ public:
 		const TConstArrayView<const FFrameNumber> Frames;
 		bool bToActive = false;
 		bool bKeepCurrent = false;
+		bool bInfiniteAtLeft = false;
+		TOptional<FFrameNumber> OptCurrentFrame;
+		bool bPropagateRelative = true;
 	};
 	
 	UE_API void ComputeLocalTransforms(UWorld* InWorld, const FEvalParameters& InEvalParams);
@@ -79,7 +102,8 @@ public:
 
 	static UE_API bool AddConstraintToSequencer(
 		const TSharedPtr<ISequencer>& InSequencer,
-		UTickableTransformConstraint* InConstraint);
+		UTickableTransformConstraint* InConstraint,
+		const FSequencerCreationOptions& InOptions = FSequencerCreationOptions());
 	
 	/** Adds an active key if needed and does the compensation when switching. Will use the optional active and time if set. 
 	Will return true if key is actually set, may not be if the value is the same.*/
@@ -141,9 +165,9 @@ public:
 	static void GetFramesToCompensate(
 		const FMovieSceneConstraintChannel& InActiveChannel,
 		const bool InActiveValueToBeSet,
-		const FFrameNumber& InTime,
+		const TOptional<FFrameNumber>& InTime,
 		const TArrayView<ChannelType*>& InChannels,
-		TArray<FFrameNumber>& OutFramesAfter);
+		TArray<FFrameNumber>& OutFrames);
 
 	/** @todo documentation. */
 	template< typename ChannelType >
@@ -207,6 +231,14 @@ public:
 		const TArray<FFrameNumber>& Frames,
 		const TArray<FTransform>& Transforms);
 
+	/** Checks whether the key to be added is prior to the current frame and updates the frames to be compensated accordingly. */
+	template<typename ChannelType>
+	static TPair<TOptional<FFrameNumber>, bool> IsActiveKeyAddedBeforeCurrentTime(
+		const FFrameNumber& InCurrentTime,
+		const FFrameNumber& InTime,
+		const TArrayView<ChannelType*>& InChannels,
+		TArray<FFrameNumber>& InOutFramesToCompensate);
+	
 	static UE_API bool bDoNotCompensate;
 
 private:
diff --git a/Engine/Source/Editor/MovieSceneTools/Public/Constraints/MovieSceneConstraintChannelHelper.inl b/Engine/Source/Editor/MovieSceneTools/Public/Constraints/MovieSceneConstraintChannelHelper.inl
--- a/Engine/Source/Editor/MovieSceneTools/Public/Constraints/MovieSceneConstraintChannelHelper.inl
+++ b/Engine/Source/Editor/MovieSceneTools/Public/Constraints/MovieSceneConstraintChannelHelper.inl
@@ -16,54 +16,125 @@ template<typename ChannelType>
 void FMovieSceneConstraintChannelHelper::GetFramesToCompensate(
 	const FMovieSceneConstraintChannel& InActiveChannel,
 	const bool InActiveValueToBeSet,
-	const FFrameNumber& InTime,
+	const TOptional<FFrameNumber>& InTime,
 	const TArrayView<ChannelType*>& InChannels,
-	TArray<FFrameNumber>& OutFramesAfter)
+	TArray<FFrameNumber>& OutFrames)
 {
 	using ChannelValueType = typename ChannelType::ChannelValueType;
 
 	const bool bHasKeys = (InActiveChannel.GetNumKeys() > 0);
+	const bool bInfinite = InActiveChannel.IsInfinite(); 
 	
-	OutFramesAfter.Reset();
-
-	// add the current frame
-	OutFramesAfter.Add(InTime);
+	OutFrames.Reset();
 
-	// add the next frames that need transform compensation 
-	for (const ChannelType* InChannel: InChannels)
+	if (InTime)
 	{
-		const TMovieSceneChannelData<const ChannelValueType> ChannelData = InChannel->GetData();
-		const TArrayView<const FFrameNumber> Times = ChannelData.GetTimes();
-		if (!Times.IsEmpty())
+		if (bInfinite)
 		{
-			// look for the first next key frame for this channel 
-			const int32 NextTimeIndex = Algo::UpperBound(Times, InTime);
-			if (Times.IsValidIndex(NextTimeIndex))
+			OutFrames.Add(*InTime);
+			
+			// add all the frames that need transform compensation
+			for (const ChannelType* InChannel: InChannels)
 			{
-				// store the time while the state is different
-				for (int32 Index = NextTimeIndex; Index < Times.Num(); ++Index)
+				const TMovieSceneChannelData<const ChannelValueType> ChannelData = InChannel->GetData();
+				const TArrayView<const FFrameNumber> Times = ChannelData.GetTimes();
+				if (!Times.IsEmpty())
 				{
-					if (!bHasKeys)
+					int32 NextTimeIndex = Algo::UpperBound(Times, *InTime);
+					while (Times.IsValidIndex(NextTimeIndex))
 					{
-						OutFramesAfter.Add(Times[Index]);
+						OutFrames.Add(Times[NextTimeIndex++]);
 					}
-					else
+				}
+			}
+			
+			// uniqueness
+			OutFrames.Sort();
+			OutFrames.SetNum(Algo::Unique(OutFrames));
+			return;
+		}
+
+		const bool bIsActiveByDefault = InActiveChannel.GetDefault() && *InActiveChannel.GetDefault();
+		if (bHasKeys && bIsActiveByDefault)
+		{
+			const TArrayView<const FFrameNumber> ConstraintTimes = InActiveChannel.GetTimes();
+			const TArrayView<const bool> Values = InActiveChannel.GetValues();
+			if (*InTime <= ConstraintTimes[0] && !Values[0])
+			{
+				// add all the frames that need transform compensation
+				for (const ChannelType* InChannel: InChannels)
+				{
+					const TMovieSceneChannelData<const ChannelValueType> ChannelData = InChannel->GetData();
+					const TArrayView<const FFrameNumber> Times = ChannelData.GetTimes();
+					if (!Times.IsEmpty())
 					{
-						bool NextValue = false; InActiveChannel.Evaluate(Times[Index], NextValue);
-						if (NextValue == InActiveValueToBeSet)
+						int32 Index = 0;
+						while (Index < Times.Num() && Times[Index] < *InTime)
 						{
-							break;
+							OutFrames.Add(Times[Index++]);
 						}
-						OutFramesAfter.Add(Times[Index]);
 					}
 				}
+				
+				// uniqueness
+				OutFrames.Sort();
+				OutFrames.SetNum(Algo::Unique(OutFrames));
+				
+				// insert time upfront after sorting even if greater when dealing with the left side of an infinite constraint
+				OutFrames.Insert(*InTime, 0);
+				return;
 			}
 		}
+		
+		{
+			// add the current frame
+			OutFrames.Add(*InTime);
+			
+			// add the next frames that need transform compensation 
+			for (const ChannelType* InChannel: InChannels)
+			{
+				const TMovieSceneChannelData<const ChannelValueType> ChannelData = InChannel->GetData();
+				const TArrayView<const FFrameNumber> Times = ChannelData.GetTimes();
+				if (!Times.IsEmpty())
+				{
+					// look for the first next key frame for this channel 
+					const int32 NextTimeIndex = Algo::UpperBound(Times, *InTime);
+					if (Times.IsValidIndex(NextTimeIndex))
+					{
+						// store the time while the state is different
+						for (int32 Index = NextTimeIndex; Index < Times.Num(); ++Index)
+						{
+							if (!bHasKeys)
+							{
+								OutFrames.Add(Times[Index]);
+							}
+							else
+							{
+								bool NextValue = false; InActiveChannel.Evaluate(Times[Index], NextValue);
+								if (NextValue == InActiveValueToBeSet)
+								{
+									break;
+								}
+								OutFrames.Add(Times[Index]);
+							}
+						}
+					}
+				}
+			}	
+		}
+	}
+	else
+	{
+		for (const ChannelType* InChannel: InChannels)
+		{
+			const TMovieSceneChannelData<const ChannelValueType> ChannelData = InChannel->GetData();
+			OutFrames.Append(ChannelData.GetTimes());
+		}
 	}
 
 	// uniqueness
-	OutFramesAfter.Sort();
-	OutFramesAfter.SetNum(Algo::Unique(OutFramesAfter));
+	OutFrames.Sort();
+	OutFrames.SetNum(Algo::Unique(OutFrames));
 }
 
 template< typename ChannelType >
@@ -149,46 +220,58 @@ void FMovieSceneConstraintChannelHelper::GetFramesWithinActiveState(
 	
 	const TMovieSceneChannelData<const bool> ConstraintChannelData = InActiveChannel.GetData();
 	const TArrayView<const FFrameNumber>& ActiveTimes = ConstraintChannelData.GetTimes();
-	if (ActiveTimes.IsEmpty())
+	const bool bInfinite = InActiveChannel.IsInfinite();
+	if (ActiveTimes.IsEmpty() && !bInfinite)
 	{
 		return;
 	}
 
-	const FFrameNumber& FirstTime = ActiveTimes[0];
-	const FFrameNumber& LastTime = ActiveTimes.Last();
-	
-	// add active times
-	for (int32 Index = 0; Index < ActiveTimes.Num(); ++Index)
+	if (bInfinite)
 	{
-		OutFrames.Add(ActiveTimes[Index]);
+		for (const ChannelType* InChannel: InChannels)
+		{
+			const TMovieSceneChannelData<const ChannelValueType> ChannelData = InChannel->GetData();
+			OutFrames.Append(ChannelData.GetTimes());
+		}
 	}
+	else
+	{
+		const FFrameNumber& FirstTime = ActiveTimes[0];
+		const FFrameNumber& LastTime = ActiveTimes.Last();
+		
+		// add active times
+		for (int32 Index = 0; Index < ActiveTimes.Num(); ++Index)
+		{
+			OutFrames.Add(ActiveTimes[Index]);
+		}
 
-	const bool bIsLastStateInactive = ConstraintChannelData.GetValues().Last() == false; 
+		const bool bIsLastStateInactive = ConstraintChannelData.GetValues().Last() == false; 
 
-	// add frames where the constraint is active 
-	for (const ChannelType* InChannel: InChannels)
-	{
-		const TMovieSceneChannelData<const ChannelValueType> ChannelData = InChannel->GetData();
-		const TArrayView<const FFrameNumber> Times = ChannelData.GetTimes();
-		if (!Times.IsEmpty())
+		// add frames where the constraint is active 
+		for (const ChannelType* InChannel: InChannels)
 		{
-			// look for the first next key frame for this channel 
-			const int32 NextTimeIndex = Algo::UpperBound(Times,FirstTime);
-			if (Times.IsValidIndex(NextTimeIndex))
+			const TMovieSceneChannelData<const ChannelValueType> ChannelData = InChannel->GetData();
+			const TArrayView<const FFrameNumber> Times = ChannelData.GetTimes();
+			if (!Times.IsEmpty())
 			{
-				// store the time is the state is active
-				for (int32 Index = NextTimeIndex; Index < Times.Num(); ++Index)
+				// look for the first next key frame for this channel 
+				const int32 NextTimeIndex = Algo::UpperBound(Times,FirstTime);
+				if (Times.IsValidIndex(NextTimeIndex))
 				{
-					bool bIsActive = false;
-					InActiveChannel.Evaluate(Times[Index], bIsActive);
-					if (bIsActive)
+					// store the time is the state is active
+					for (int32 Index = NextTimeIndex; Index < Times.Num(); ++Index)
 					{
-						OutFrames.Add(Times[Index]);
-					}
+						bool bIsActive = false;
+						InActiveChannel.Evaluate(Times[Index], bIsActive);
+						if (bIsActive)
+						{
+							OutFrames.Add(Times[Index]);
+						}
 
-					if (bIsLastStateInactive && Times[Index] > LastTime)
-					{
-						break;
+						if (bIsLastStateInactive && Times[Index] > LastTime)
+						{
+							break;
+						}
 					}
 				}
 			}
@@ -499,4 +582,49 @@ void SetTangentsAtThisTime(int32 ChannelIndex,
 			Value.TangentMode = RCTM_Break;
 		}
 	}
+}
+
+template<typename ChannelType>
+TPair<TOptional<FFrameNumber>, bool> IsActiveKeyAddedBeforeCurrentTime(
+	const FFrameNumber& InCurrentTime,
+	const FFrameNumber& InTime,
+	const TArrayView<ChannelType*>& InChannels,
+	TArray<FFrameNumber>& InOutFramesToCompensate)
+{
+	TPair< TOptional<FFrameNumber>, bool > Result;
+
+	// this means that we're adding the constraint before the current frame
+	if (InTime < InCurrentTime)
+	{
+		Result.Key = InCurrentTime;
+		
+		// add the current frame if missing
+		const int32 Index = Algo::LowerBound(InOutFramesToCompensate, InCurrentTime);	
+		if (Index >= InOutFramesToCompensate.Num() || InOutFramesToCompensate[Index] != InCurrentTime)
+		{
+			InOutFramesToCompensate.Insert(InCurrentTime, Index);
+		}
+
+		// check if there's a transform key before the current time
+		// if that's not the case then we need to propagate the current offset
+		const bool bPropagateRelative = [InChannels, InCurrentTime]()
+		{
+			for (const ChannelType* TransformChannel : InChannels)
+			{
+				const TArrayView<const FFrameNumber> FrameNumbers = TransformChannel->GetTimes();
+				if (!FrameNumbers.IsEmpty())
+				{
+					if (FrameNumbers[0] < InCurrentTime)
+					{
+						return false;
+					}
+				}
+			}
+			return true;
+		}();
+		
+		Result.Value = bPropagateRelative;
+	}
+
+	return Result;
 }
\ No newline at end of file
diff --git a/Engine/Source/Editor/MovieSceneTools/Public/Constraints/TransformConstraintChannelInterface.h b/Engine/Source/Editor/MovieSceneTools/Public/Constraints/TransformConstraintChannelInterface.h
--- a/Engine/Source/Editor/MovieSceneTools/Public/Constraints/TransformConstraintChannelInterface.h
+++ b/Engine/Source/Editor/MovieSceneTools/Public/Constraints/TransformConstraintChannelInterface.h
@@ -39,7 +39,7 @@ struct ITransformConstraintChannelInterface
 	/** Add an active/inactive key to the constraint channel if needed and does the transform compensation on the transform channels. */
 	virtual bool SmartConstraintKey(
 		UTickableTransformConstraint* InConstraint, const TOptional<bool>& InOptActive,
-		const FFrameNumber& InTime, const TSharedPtr<ISequencer>& InSequencer) = 0;
+		const TRange<FFrameNumber>& InTimeRange, const TSharedPtr<ISequencer>& InSequencer) = 0;
 
 	/** Add keys on the transform channels of the object that the handle wraps. */
 	virtual void AddHandleTransformKeys(
diff --git a/Engine/Source/Editor/MovieSceneTools/Public/MovieSceneToolsModule.h b/Engine/Source/Editor/MovieSceneTools/Public/MovieSceneToolsModule.h
--- a/Engine/Source/Editor/MovieSceneTools/Public/MovieSceneToolsModule.h
+++ b/Engine/Source/Editor/MovieSceneTools/Public/MovieSceneToolsModule.h
@@ -41,6 +41,20 @@ public:
 	virtual void StopBaking(UMovieScene* MovieScene) {};
 };
 
+struct FMovieSceneToolsAnimationBakingWrapper
+{
+	FMovieSceneToolsAnimationBakingWrapper(UMovieScene* InMovieScene);
+
+	void Start() const;
+	void PreEvaluate(const FFrameNumber& Frame) const;
+	void PostEvaluate(const FFrameNumber& Frame) const;
+	void Stop() const;
+	
+private:
+	UMovieScene* MovieScene = nullptr;
+	const TArray<IMovieSceneToolsAnimationBakeHelper*>& BakeHelpers;
+};
+
 // Interface to allow external modules to register additional key struct instanced property type customizations
 class IMovieSceneToolsKeyStructInstancedPropertyTypeCustomizer
 {
diff --git a/Engine/Source/Runtime/Experimental/Animation/Constraints/Private/ConstraintChannel.cpp b/Engine/Source/Runtime/Experimental/Animation/Constraints/Private/ConstraintChannel.cpp
--- a/Engine/Source/Runtime/Experimental/Animation/Constraints/Private/ConstraintChannel.cpp
+++ b/Engine/Source/Runtime/Experimental/Animation/Constraints/Private/ConstraintChannel.cpp
@@ -8,14 +8,28 @@ bool FMovieSceneConstraintChannel::Evaluate(FFrameTime InTime, bool& OutValue) c
 {
 	if (Times.IsEmpty())
 	{
+		if (bHasDefaultValue)
+		{
+			OutValue = DefaultValue;
+			return true;
+		}
 		return false;
 	}
 
 	if (InTime.FrameNumber < Times[0])
 	{
+		if (bHasDefaultValue && !Values.IsEmpty() && !Values[0])
+		{
+			OutValue = DefaultValue;
+			return true;
+		}
 		return false;
 	}
 
 	return FMovieSceneBoolChannel::Evaluate(InTime, OutValue);
 }
 
+bool FMovieSceneConstraintChannel::IsInfinite() const
+{
+	return Times.IsEmpty() && bHasDefaultValue && DefaultValue;
+}
\ No newline at end of file
diff --git a/Engine/Source/Runtime/Experimental/Animation/Constraints/Private/Transform/TransformConstraint.cpp b/Engine/Source/Runtime/Experimental/Animation/Constraints/Private/Transform/TransformConstraint.cpp
--- a/Engine/Source/Runtime/Experimental/Animation/Constraints/Private/Transform/TransformConstraint.cpp
+++ b/Engine/Source/Runtime/Experimental/Animation/Constraints/Private/Transform/TransformConstraint.cpp
@@ -196,6 +196,26 @@ FString UTickableTransformConstraint::GetTypeLabel() const
 	return Super::GetTypeLabel();
 }
 
+bool UTickableTransformConstraint::CanEditChange(const FProperty* InProperty) const
+{
+	if (InProperty)
+	{
+		const FName MemberPropertyName = InProperty->GetFName();
+		if (MemberPropertyName == GET_MEMBER_NAME_CHECKED(UTickableTranslationConstraint, OffsetTranslation) ||
+			MemberPropertyName == GET_MEMBER_NAME_CHECKED(UTickableRotationConstraint, OffsetRotation) ||
+			MemberPropertyName == GET_MEMBER_NAME_CHECKED(UTickableScaleConstraint, OffsetScale) ||
+			MemberPropertyName == GET_MEMBER_NAME_CHECKED(UTickableParentConstraint, OffsetTransform))
+		{
+			if (GetOuter() == GetTransientPackage())
+			{
+				return bMaintainOffset && !bUseCurrentOffset;
+			}
+		}
+	}
+	
+	return Super::CanEditChange(InProperty);
+}
+
 void UTickableTransformConstraint::PostEditChangeProperty(FPropertyChangedEvent& PropertyChangedEvent)
 {
 	Super::PostEditChangeProperty(PropertyChangedEvent);
diff --git a/Engine/Source/Runtime/Experimental/Animation/Constraints/Private/Transform/TransformConstraintUtil.cpp b/Engine/Source/Runtime/Experimental/Animation/Constraints/Private/Transform/TransformConstraintUtil.cpp
--- a/Engine/Source/Runtime/Experimental/Animation/Constraints/Private/Transform/TransformConstraintUtil.cpp
+++ b/Engine/Source/Runtime/Experimental/Animation/Constraints/Private/Transform/TransformConstraintUtil.cpp
@@ -402,7 +402,16 @@ bool AddConstraint(UWorld* InWorld, UTransformableHandle* InParentHandle, UTrans
 		return false;
 	}
 
-	if (!bUseDefault)
+	bool bSetup = !bUseDefault; 
+	if (bUseDefault)
+	{
+		if (UTickableTransformConstraint* CDO = InNewConstraint->GetClass()->GetDefaultObject<UTickableTransformConstraint>())
+		{
+			bSetup = CDO->bUseCurrentOffset;
+		}
+	}
+
+	if (bSetup)
 	{
 		InNewConstraint->bMaintainOffset = bMaintainOffset;
 		InNewConstraint->Setup();
diff --git a/Engine/Source/Runtime/Experimental/Animation/Constraints/Public/ConstraintChannel.h b/Engine/Source/Runtime/Experimental/Animation/Constraints/Public/ConstraintChannel.h
--- a/Engine/Source/Runtime/Experimental/Animation/Constraints/Public/ConstraintChannel.h
+++ b/Engine/Source/Runtime/Experimental/Animation/Constraints/Public/ConstraintChannel.h
@@ -54,6 +54,12 @@ struct FMovieSceneConstraintChannel : public FMovieSceneBoolChannel
 	 */
 	CONSTRAINTS_API virtual bool Evaluate(FFrameTime InTime, bool& OutValue) const override;
 
+	/**
+	 * Checks whether this constraint channel is infinite
+	 * @return true if there are no keys and the default value is set to true.
+	 */
+	CONSTRAINTS_API bool IsInfinite() const;
+	
 #if WITH_EDITOR
 	using ExtraLabelFunction = TFunction< FString() >;
 	ExtraLabelFunction ExtraLabel;
@@ -70,7 +76,7 @@ struct TStructOpsTypeTraits<FMovieSceneConstraintChannel> : public TStructOpsTyp
 template<>
 struct TMovieSceneChannelTraits<FMovieSceneConstraintChannel> : TMovieSceneChannelTraitsBase<FMovieSceneConstraintChannel>
 {
-	enum { SupportsDefaults = false };
+	enum { SupportsDefaults = true };
 
 #if WITH_EDITOR
 
diff --git a/Engine/Source/Runtime/Experimental/Animation/Constraints/Public/Transform/TransformConstraint.h b/Engine/Source/Runtime/Experimental/Animation/Constraints/Public/Transform/TransformConstraint.h
--- a/Engine/Source/Runtime/Experimental/Animation/Constraints/Public/Transform/TransformConstraint.h
+++ b/Engine/Source/Runtime/Experimental/Animation/Constraints/Public/Transform/TransformConstraint.h
@@ -167,6 +167,10 @@ public:
 	CONSTRAINTS_API virtual void TeardownConstraint(UWorld * InWorld) override;
 	CONSTRAINTS_API virtual void AddedToWorld(UWorld* InWorld) override;
 
+	/** Used at creation to define whether the current relative offset should be used instead of the provided one. */
+	UPROPERTY(EditAnywhere, Transient, Category="Offset", meta = (DisplayPriority=0))
+	bool bUseCurrentOffset = true;
+	
 #if WITH_EDITOR
 public:
 	/** Returns the constraint's label used for UI. */
@@ -177,6 +181,7 @@ public:
 	CONSTRAINTS_API virtual FString GetTypeLabel() const override;
 
 	// UObject interface
+	CONSTRAINTS_API virtual bool CanEditChange(const FProperty* InProperty) const override;
 	CONSTRAINTS_API virtual void PostEditChangeProperty(FPropertyChangedEvent& PropertyChangedEvent) override;
 	CONSTRAINTS_API virtual void PostEditUndo() override;
 	// End of UObject interface
