Sequencer: Fix SetKeyPositions accidentally removing keys that will be set later.

Example: InKeys { 42, 54 } with current positions { (0.5, OldY), (0.533, OldY2) } and InKeyPositions = { (0.533, NewY1), (0.9, NewY2) }.
When we get to key handle 42, MoveKey will detect that key handle 52 already has time 0.533 and attempt to remove it. However, because we'll later move 52 from 0.533 to 0.9, we should skip it.

diff --git a/Engine/Source/Editor/MovieSceneTools/Private/Channels/ChannelCurveModel.cpp b/Engine/Source/Editor/MovieSceneTools/Private/Channels/ChannelCurveModel.cpp
--- a/Engine/Source/Editor/MovieSceneTools/Private/Channels/ChannelCurveModel.cpp
+++ b/Engine/Source/Editor/MovieSceneTools/Private/Channels/ChannelCurveModel.cpp
@@ -27,6 +27,7 @@
 #include "MovieScene.h"
 #include "MovieSceneNameableTrack.h"
 #include "MovieSceneSection.h"
+#include "Algo/IndexOf.h"
 #include "Styling/AppStyle.h"
 #include "Styling/ISlateStyle.h"
 #include "Templates/UnrealTemplate.h"
@@ -263,6 +264,22 @@ void FChannelCurveModel<ChannelType, ChannelValue, KeyType>::GetKeyPositions(TAr
 	}
 }
 
+namespace UE::MovieSceneTools::Private
+{
+/** Finds the index of InKeyHandle if it is positioned after AfterIndex. */
+static int32 IndexOfKeyAfter(const TConstArrayView<FKeyHandle>& InHandles, const FKeyHandle& InKeyHandle, int32 AfterIndex)
+{
+	for (int32 LaterIndex = AfterIndex + 1; LaterIndex < InHandles.Num(); ++LaterIndex)
+	{
+		if (InHandles[LaterIndex] == InKeyHandle)
+		{
+			return LaterIndex;
+		}
+	}
+	return INDEX_NONE;
+}
+}
+
 template <class ChannelType, class ChannelValue, class KeyType>
 void FChannelCurveModel<ChannelType, ChannelValue, KeyType>::SetKeyPositions(TArrayView<const FKeyHandle> InKeys, TArrayView<const FKeyPosition> InKeyPositions, EPropertyChangeType::Type ChangeType)
 {
@@ -312,15 +329,36 @@ void FChannelCurveModel<ChannelType, ChannelValue, KeyType>::SetKeyPositions(TAr
 			}
 		}
 
+		const auto ComputeFrameTime = [&InKeyPositions, &TickResolution, &KeyOffset](int32 Index)
+		{
+			return (InKeyPositions[Index].InputValue * TickResolution).RoundToFrame() - KeyOffset;
+		};
 		for (int32 Index = 0; Index < InKeys.Num(); ++Index)
 		{
 			int32 KeyIndex = ChannelData.GetIndex(InKeys[Index]);
 			if (KeyIndex != INDEX_NONE)
 			{
-				FFrameNumber NewTime = (InKeyPositions[Index].InputValue * TickResolution).RoundToFrame() - KeyOffset;
+				const FFrameNumber NewTime = ComputeFrameTime(Index);
 
 				const bool bRemoveDuplicates = ChangeType == EPropertyChangeType::ValueSet;
-				KeyIndex = ChannelData.MoveKey(KeyIndex, NewTime, bRemoveDuplicates);
+				// Skip removing a duplicate key handle if it appears in InKeys after our current Index and the x-position will be changed.
+				// Example: InKeys { 42, 54 } with current positions { (0.5, OldY), (0.533, OldY2) } and InKeyPositions = { (0.533, NewY1), (0.9, NewY2) }.
+				// When we get to key handle 42, MoveKey will detect that key handle 52 already has time 0.533 and attempt to remove it. However, because
+				// we'll later move 52 from 0.533 to 0.9, we should skip it.
+				const auto CanRemoveDuplicateHandle = [Index, &InKeys, &ChannelData, &ComputeFrameTime](const FKeyHandle& DuplicateKeyHandle)
+				{
+					const int32 LaterKeyPositionIndex =  UE::MovieSceneTools::Private::IndexOfKeyAfter(InKeys, DuplicateKeyHandle, Index);
+					if (LaterKeyPositionIndex != INDEX_NONE)
+					{
+						const int32 InternalIndex = ChannelData.GetIndex(DuplicateKeyHandle);
+						const FFrameNumber& CurrentTime = ChannelData.GetTimes()[InternalIndex];
+						const FFrameNumber NewTime = ComputeFrameTime(LaterKeyPositionIndex);
+						return CurrentTime != NewTime;
+					}
+					return false;
+				};
+				
+				KeyIndex = ChannelData.MoveKey(KeyIndex, NewTime, bRemoveDuplicates, CanRemoveDuplicateHandle);
 				SetKeyValue(KeyIndex, InKeyPositions[Index].OutputValue);
 			}
 		}
diff --git a/Engine/Source/Runtime/MovieScene/Public/Channels/MovieSceneChannelData.h b/Engine/Source/Runtime/MovieScene/Public/Channels/MovieSceneChannelData.h
--- a/Engine/Source/Runtime/MovieScene/Public/Channels/MovieSceneChannelData.h
+++ b/Engine/Source/Runtime/MovieScene/Public/Channels/MovieSceneChannelData.h
@@ -17,6 +17,7 @@
 #include "UObject/Class.h"
 #include "UObject/ObjectMacros.h"
 #include "IMovieSceneRetimingInterface.h"
+#include "Algo/RemoveIf.h"
 
 #include "MovieSceneChannelData.generated.h"
 
@@ -344,6 +345,12 @@ struct TMovieSceneChannelData : FMovieSceneChannelData
 	 */
 
 	int32 MoveKey(int32 KeyIndex, FFrameNumber NewTime, bool bRemoveDuplicateKeys = false)
+	{
+		return MoveKey(KeyIndex, NewTime, bRemoveDuplicateKeys, [](const FKeyHandle& KeyHandle){ return true; });
+	}
+
+	template<typename TCanRemove> requires std::is_invocable_r_v<bool, TCanRemove, const FKeyHandle&>
+	int32 MoveKey(int32 KeyIndex, FFrameNumber NewTime, bool bRemoveDuplicateKeys, TCanRemove&& CanRemove)
 	{
 		if (bRemoveDuplicateKeys)
 		{
@@ -353,6 +360,7 @@ struct TMovieSceneChannelData : FMovieSceneChannelData
 			TArray<FFrameNumber> DummyTimes;
 			GetKeys(TRange<FFrameNumber>(NewTime, NewTime), &DummyTimes, &KeysToRemove);
 			KeysToRemove.Remove(KeyHandle);
+			KeysToRemove.SetNum(Algo::RemoveIf(KeysToRemove, CanRemove));
 
 			if (KeysToRemove.Num() > 0)
 			{
