Movie Scene: Convert check to ensure and add more debug logging to fix crash without reliable repro.

diff --git a/Engine/Source/Editor/MovieSceneTools/Private/Channels/DoubleChannelCurveModel.cpp b/Engine/Source/Editor/MovieSceneTools/Private/Channels/DoubleChannelCurveModel.cpp
--- a/Engine/Source/Editor/MovieSceneTools/Private/Channels/DoubleChannelCurveModel.cpp
+++ b/Engine/Source/Editor/MovieSceneTools/Private/Channels/DoubleChannelCurveModel.cpp
@@ -112,6 +112,17 @@ void FDoubleChannelCurveModel::GetValueRange(double& MinValue, double& MaxValue)
 	}
 }
 
+namespace UE::MovieSceneTools
+{
+static FString LexToString(TArrayView<const FFrameNumber> Times)
+{
+	return FString::JoinBy(Times, TEXT(", "), [](const FFrameNumber& FrameNumber)
+	{
+		return LexToString (FrameNumber);
+	});
+}
+}
+
 void FDoubleChannelCurveModel::GetValueRange(double InMinTime, double InMaxTime, double& MinValue, double& MaxValue) const
 {
 	FMovieSceneDoubleChannel* Channel = GetChannelHandle().Get();
@@ -119,12 +130,20 @@ void FDoubleChannelCurveModel::GetValueRange(double InMinTime, double InMaxTime,
 
 	if (Channel && Section)
 	{
-		TArrayView<const FFrameNumber> Times = Channel->GetData().GetTimes();
-
-		FFrameRate TickResolution = Section->GetTypedOuter<UMovieScene>()->GetTickResolution();
-
+		const TArrayView<const FFrameNumber> Times = Channel->GetData().GetTimes();
+
+		const FFrameRate TickResolution = Section->GetTypedOuter<UMovieScene>()->GetTickResolution();
+		const FFrameTime StartTime = InMinTime * TickResolution;
+		const FFrameTime EndTime = InMaxTime * TickResolution;
+
+		// DP 5.7: This exists to debug issue UE-326363. We think that the Times array may be out of order (broken invariant). Remove once the issue is fixed.
+		using namespace UE::MovieSceneTools;
+		UE_CLOG(StartTime > EndTime, LogMovieScene, Warning, TEXT("StartTime > EndTime: InMinTime: %f, InMaxTime: %f, TickResolution: %s, StartTime (%s), EndTime (%s): Array is %s"),
+			InMinTime, InMaxTime, *TickResolution.ToPrettyText().ToString(), *LexToString(StartTime), *LexToString(EndTime), *LexToString(Times)
+			);
+		
 		UE::MovieScene::Interpolation::FInterpolationExtents Extents =
-			Channel->ComputeExtents(InMinTime * TickResolution, InMaxTime * TickResolution);
+			Channel->ComputeExtents(StartTime, EndTime);
 
 		MinValue = Extents.MinValue;
 		MaxValue = Extents.MaxValue;
diff --git a/Engine/Source/Runtime/MovieScene/Public/Channels/MovieScenePiecewiseCurveUtils.inl b/Engine/Source/Runtime/MovieScene/Public/Channels/MovieScenePiecewiseCurveUtils.inl
--- a/Engine/Source/Runtime/MovieScene/Public/Channels/MovieScenePiecewiseCurveUtils.inl
+++ b/Engine/Source/Runtime/MovieScene/Public/Channels/MovieScenePiecewiseCurveUtils.inl
@@ -57,6 +57,11 @@ UE::MovieScene::Interpolation::FInterpolationExtents ComputePiecewiseExtents(con
 	using namespace UE::MovieScene;
 	using namespace UE::MovieScene::Interpolation;
 
+	if (!ensure(StartTime <= EndTime))
+	{
+		return FInterpolationExtents{ 0, 0, 0, 0 };
+	}
+	
 	check(StartTime <= EndTime);
 
 	const int32 NumPieces = PiecewiseData.NumPieces();
