Sequencer: Fix time not getting set

This regression come from 42417964 which added a PostConstructKeyInstance for DoubleChannel to invert the channel value for LUF. It only set the value and not the time. This follows the implementation of PostConstrucKeyInstance for FMovieSceneTimeWarpChannel and FMovieSceneObjectPathChannel

diff --git a/Engine/Source/Editor/MovieSceneTools/Private/Channels/BuiltInChannelEditors.cpp b/Engine/Source/Editor/MovieSceneTools/Private/Channels/BuiltInChannelEditors.cpp
--- a/Engine/Source/Editor/MovieSceneTools/Private/Channels/BuiltInChannelEditors.cpp
+++ b/Engine/Source/Editor/MovieSceneTools/Private/Channels/BuiltInChannelEditors.cpp
@@ -1093,6 +1093,9 @@ void PostConstructKeyInstance(const TMovieSceneChannelHandle<FMovieSceneDoubleCh
 	const UMovieSceneKeyStructType* GeneratedStructType = CastChecked<const UMovieSceneKeyStructType>(Struct->GetStruct());
 	void* StructMemory = Struct->GetStructMemory();
 	FStructProperty* ValueProperty = CastFieldChecked<FStructProperty>(GeneratedStructType->DestValueProperty.Get());
+	FStructProperty* TimeProperty = CastFieldChecked<FStructProperty>(GeneratedStructType->DestTimeProperty.Get());
+
+	const FFrameNumber* TimeAddress = TimeProperty->ContainerPtrToValuePtr<FFrameNumber>(StructMemory);
 
 	FMovieSceneDoubleChannel* Channel = ChannelHandle.Get();
 	if (Channel)
@@ -1113,7 +1116,7 @@ void PostConstructKeyInstance(const TMovieSceneChannelHandle<FMovieSceneDoubleCh
 	}
 
 	// It is safe to capture the property and address in this lambda because the lambda is owned by the struct itself, so cannot be invoked if the struct has been destroyed
-	auto CopyInstanceToKeyLambda = [ChannelHandle, InHandle, StructMemory, ValueProperty, bInvertValue](const FPropertyChangedEvent&)
+	auto CopyInstanceToKeyLambda = [ChannelHandle, InHandle, StructMemory, ValueProperty, bInvertValue, TimeAddress](const FPropertyChangedEvent&)
 	{
 		if (FMovieSceneDoubleChannel* DestinationChannel = ChannelHandle.Get())
 		{
@@ -1122,6 +1125,9 @@ void PostConstructKeyInstance(const TMovieSceneChannelHandle<FMovieSceneDoubleCh
 			{
 				double Value = *ValueProperty->ContainerPtrToValuePtr<double>(StructMemory);
 				DestinationChannel->GetData().GetValues()[KeyIndex].Value = bInvertValue ? -Value : Value;
+
+				// Set the new key time
+				DestinationChannel->SetKeyTime(InHandle, *TimeAddress);
 			}
 		}
 	};
