Curve Editor: Fix moving keys not being interactive, i.e. not updating the control rigs in realtime.

diff --git a/Engine/Source/Editor/MovieSceneTools/Private/Channels/ChannelCurveModel.cpp b/Engine/Source/Editor/MovieSceneTools/Private/Channels/ChannelCurveModel.cpp
--- a/Engine/Source/Editor/MovieSceneTools/Private/Channels/ChannelCurveModel.cpp
+++ b/Engine/Source/Editor/MovieSceneTools/Private/Channels/ChannelCurveModel.cpp
@@ -690,6 +690,15 @@ void FChannelCurveModel<ChannelType, ChannelValue, KeyType>::PostProcessChanges(
 	// Allow channel to make any follow up changes to the keys before executing any other delegates
 	Channel->PostEditChange();
 
+	// These cause interactive changes, such as evaluating sequencer as the keys are moved. Do not defer even when in a change scope.
+	Section->MarkAsChanged();
+	Section->MarkPackageDirty();
+	if (UMovieSceneSignedObject* SignedOwner = Cast<UMovieSceneSignedObject>(WeakOwningObject.Get()))
+	{
+		SignedOwner->MarkAsChanged();
+		SignedOwner->MarkPackageDirty();
+	}
+
 	InvokePostChannelChangedEvents(Section);
 }
 
@@ -711,15 +720,6 @@ void FChannelCurveModel<ChannelType, ChannelValue, KeyType>::InvokePostChannelCh
 		WeakSequencer.Pin()->OnChannelChanged().Broadcast(MetaData, Section);
 	}
 
-	Section->MarkAsChanged();
-	Section->MarkPackageDirty();
-	
-	if (UMovieSceneSignedObject* SignedOwner = Cast<UMovieSceneSignedObject>(WeakOwningObject.Get()))
-	{
-		SignedOwner->MarkAsChanged();
-		SignedOwner->MarkPackageDirty();
-	}
-	
 	CurveModifiedDelegate.Broadcast();
 }
 
