Sequencer: In RecomposeTransform, initialize the CurrentTransform with the current transform values. Otherwise, channels that don't yet have a track/channel will be 0.

diff --git a/Engine/Source/Editor/MovieSceneTools/Private/TrackEditors/TransformTrackEditor.cpp b/Engine/Source/Editor/MovieSceneTools/Private/TrackEditors/TransformTrackEditor.cpp
--- a/Engine/Source/Editor/MovieSceneTools/Private/TrackEditors/TransformTrackEditor.cpp
+++ b/Engine/Source/Editor/MovieSceneTools/Private/TrackEditors/TransformTrackEditor.cpp
@@ -1164,6 +1164,16 @@ FTransformData F3DTransformTrackEditor::RecomposeTransform(const FTransformData&
 			TRecompositionResult<FIntermediate3DTransform> TransformData = System->RecomposeBlendOperational(TrackComponents->ComponentTransform, Query, CurrentValue);
 
 			double CurrentTransformChannels[9] = { 0, 0, 0, 0, 0, 0, 0, 0, 0 };
+			CurrentTransformChannels[0] = CurrentValue.GetTranslation()[0];
+			CurrentTransformChannels[1] = CurrentValue.GetTranslation()[1];
+			CurrentTransformChannels[2] = CurrentValue.GetTranslation()[2];
+			CurrentTransformChannels[3] = CurrentValue.GetRotation().Roll;
+			CurrentTransformChannels[4] = CurrentValue.GetRotation().Pitch;
+			CurrentTransformChannels[5] = CurrentValue.GetRotation().Yaw;
+			CurrentTransformChannels[6] = CurrentValue.GetScale()[0];
+			CurrentTransformChannels[7] = CurrentValue.GetScale()[1];
+			CurrentTransformChannels[8] = CurrentValue.GetScale()[2];
+
 			EMovieSceneTransformChannel ChannelsObtained(EMovieSceneTransformChannel::None);
 			check(EntityIDs.Num() == TransformData.Values.Num());
 			for (int32 EntityIndex = 0; EntityIndex < TransformData.Values.Num(); ++EntityIndex)
