Anim In Engine: Add new FAllSpawnableRestoreState that disables and restores all spawn tracks in the sequence hierarchy

diff --git a/Engine/Plugins/Animation/ControlRig/Source/ControlRigEditor/Private/ControlRigSequencerEditorLibrary.cpp b/Engine/Plugins/Animation/ControlRig/Source/ControlRigEditor/Private/ControlRigSequencerEditorLibrary.cpp
--- a/Engine/Plugins/Animation/ControlRig/Source/ControlRigEditor/Private/ControlRigSequencerEditorLibrary.cpp
+++ b/Engine/Plugins/Animation/ControlRig/Source/ControlRigEditor/Private/ControlRigSequencerEditorLibrary.cpp
@@ -1251,7 +1251,7 @@ bool UControlRigSequencerEditorLibrary::BakeToControlRig(UWorld* World, ULevelSe
 	bool bResult = false;
 	const FScopedTransaction Transaction(LOCTEXT("BakeToControlRig_Transaction", "Bake To Control Rig"));
 	{
-		FSpawnableRestoreState SpawnableRestoreState(MovieScene, Player->GetSharedPlaybackState().ToSharedPtr());
+		FAllSpawnableRestoreState SpawnableRestoreState(Player->GetSharedPlaybackState().ToSharedPtr());
 
 		if (LevelPlayer && SpawnableRestoreState.bWasChanged)
 		{
diff --git a/Engine/Plugins/Animation/ControlRig/Source/ControlRigEditor/Private/Tools/ControlRigSnapper.cpp b/Engine/Plugins/Animation/ControlRig/Source/ControlRigEditor/Private/Tools/ControlRigSnapper.cpp
--- a/Engine/Plugins/Animation/ControlRig/Source/ControlRigEditor/Private/Tools/ControlRigSnapper.cpp
+++ b/Engine/Plugins/Animation/ControlRig/Source/ControlRigEditor/Private/Tools/ControlRigSnapper.cpp
@@ -551,7 +551,7 @@ bool FControlRigSnapper::SnapIt(FFrameNumber StartFrame, FFrameNumber EndFrame,c
 		//disabling this for nowUE::Sequencer::FSequencerBakingSetupRestore RestoreBaking(Sequencer);
 		FScopedTransaction ScopedTransaction(LOCTEXT("SnapAnimation", "Snap Animation"));
 
-		FSpawnableRestoreState SpawnableRestoreState(MovieScene, Sequencer->GetSharedPlaybackState().ToSharedPtr());
+		FAllSpawnableRestoreState SpawnableRestoreState(Sequencer->GetSharedPlaybackState().ToSharedPtr());
 
 		TArray<FTransform> WorldTransformToSnap;
 		bool bSnapToFirstFrameNotParents = !CalculateWorldTransformsFromParents(Sequencer.Get(), ParentToSnap, Frames, WorldTransformToSnap);
diff --git a/Engine/Plugins/Animation/ControlRig/Source/ControlRigEditor/Private/Tools/EvaluateSequencerTools.cpp b/Engine/Plugins/Animation/ControlRig/Source/ControlRigEditor/Private/Tools/EvaluateSequencerTools.cpp
--- a/Engine/Plugins/Animation/ControlRig/Source/ControlRigEditor/Private/Tools/EvaluateSequencerTools.cpp
+++ b/Engine/Plugins/Animation/ControlRig/Source/ControlRigEditor/Private/Tools/EvaluateSequencerTools.cpp
@@ -208,9 +208,8 @@ namespace AIE
 		{
 			return false;
 		}
-		UMovieScene* MovieScene = Sequencer->GetFocusedMovieSceneSequence()->GetMovieScene();
 		
-		FSpawnableRestoreState SpawnableRestoreState(MovieScene, Sequencer->GetSharedPlaybackState().ToSharedPtr());
+		FAllSpawnableRestoreState SpawnableRestoreState(Sequencer->GetSharedPlaybackState().ToSharedPtr());
 
 		const FQualifiedFrameTime QFrameTime = Sequencer->GetGlobalTime();
 		FMovieSceneContext Context = FMovieSceneContext(FMovieSceneEvaluationRange(QFrameTime.Time, QFrameTime.Rate), Sequencer->GetPlaybackStatus()).SetHasJumped(true);
@@ -254,13 +253,25 @@ namespace AIE
 		FFrameRate TickResolution = MovieScene->GetTickResolution();
 		FFrameRate DisplayRate = MovieScene->GetDisplayRate();
 
-		FSpawnableRestoreState SpawnableRestoreState(MovieScene, Sequencer->GetSharedPlaybackState().ToSharedPtr());
+		TOptional<TRange<FFrameNumber>> OptionalRange = Sequencer->GetSubSequenceRange();
+		const FFrameNumber StartFrame = OptionalRange.IsSet() ? OptionalRange.GetValue().GetLowerBoundValue() : MovieScene->GetPlaybackRange().GetLowerBoundValue();
+		const FFrameNumber EndFrame = OptionalRange.IsSet() ? OptionalRange.GetValue().GetUpperBoundValue() : MovieScene->GetPlaybackRange().GetUpperBoundValue();
+
+		FAllSpawnableRestoreState SpawnableRestoreState(Sequencer->GetSharedPlaybackState().ToSharedPtr());
 
 		TSet<UControlRig*> ControlsToRevalaute;
 		FControlRigEditMode::FTurnOffPosePoseUpdate  TurnOff; //stop flashing
 		for (int32 Index : Indices)
 		{
-			const FFrameNumber FrameNumber = FrameTimeByIndex.CalculateFrame(Index);
+			FFrameNumber FrameNumber = FrameTimeByIndex.CalculateFrame(Index);
+			if (FrameNumber < StartFrame || FrameNumber > EndFrame)
+			{
+				continue; //if out of range bail
+			}
+			if (FrameNumber == EndFrame && FrameNumber != StartFrame) //if on end boundary subtract one tick to handle subsequence eval issues
+			{
+				--FrameNumber;
+			}
 			FFrameTime GlobalTime = LocalToRootTransform.TryTransformTime(FrameNumber).Get(FrameNumber); //player evals in root time so need to go back to it.
 
 			FMovieSceneContext Context = FMovieSceneContext(FMovieSceneEvaluationRange(GlobalTime, TickResolution), Sequencer->GetPlaybackStatus()).SetHasJumped(true);
diff --git a/Engine/Plugins/Importers/USDImporter/Source/USDStage/Private/USDLevelSequenceHelper.cpp b/Engine/Plugins/Importers/USDImporter/Source/USDStage/Private/USDLevelSequenceHelper.cpp
--- a/Engine/Plugins/Importers/USDImporter/Source/USDStage/Private/USDLevelSequenceHelper.cpp
+++ b/Engine/Plugins/Importers/USDImporter/Source/USDStage/Private/USDLevelSequenceHelper.cpp
@@ -4847,7 +4847,7 @@ void FUsdLevelSequenceHelperImpl::HandleControlRigSectionChange(UMovieSceneContr
 			PinnedSequencer->EnterSilentMode();
 		}
 
-		FSpawnableRestoreState SpawnableRestoreState(MovieScene, Player->GetSharedPlaybackState().ToSharedPtr());
+		FAllSpawnableRestoreState SpawnableRestoreState(Player->GetSharedPlaybackState().ToSharedPtr());
 		if (LevelPlayer && SpawnableRestoreState.bWasChanged)
 		{
 			// Evaluate at the beginning of the subscene time to ensure that spawnables are created before export
diff --git a/Engine/Plugins/MovieScene/SequencerScripting/Source/SequencerScriptingEditor/Private/SequencerTools.cpp b/Engine/Plugins/MovieScene/SequencerScripting/Source/SequencerScriptingEditor/Private/SequencerTools.cpp
--- a/Engine/Plugins/MovieScene/SequencerScripting/Source/SequencerScriptingEditor/Private/SequencerTools.cpp
+++ b/Engine/Plugins/MovieScene/SequencerScripting/Source/SequencerScriptingEditor/Private/SequencerTools.cpp
@@ -304,7 +304,7 @@ bool ExportFBXInternal(const FSequencerExportFBXParams& InParams, UMovieSceneSeq
 	bool bDidExport = false;
 	
 	{
-		FSpawnableRestoreState SpawnableRestoreState(MovieScene, Player->GetSharedPlaybackState().ToSharedPtr());
+		FAllSpawnableRestoreState SpawnableRestoreState(Player->GetSharedPlaybackState().ToSharedPtr());
 
 		if (SpawnableRestoreState.bWasChanged)
 		{
@@ -402,7 +402,7 @@ bool USequencerToolsFunctionLibrary::ExportAnimSequenceWaitForDelegate(UWorld* W
 
 	bool bResult = false;
 
-	FSpawnableRestoreState SpawnableRestoreState(MovieScene, Player->GetSharedPlaybackState().ToSharedPtr());
+	FAllSpawnableRestoreState SpawnableRestoreState(Player->GetSharedPlaybackState().ToSharedPtr());
 
 	if (SpawnableRestoreState.bWasChanged)
 	{
@@ -788,7 +788,7 @@ bool ImportFBXInternal(UWorld* World, UMovieSceneSequence* Sequence, const TArra
 	bool bResult = false;
 	FScopedTransaction ImportFBXTransaction(NSLOCTEXT("Sequencer", "ImportFBX", "Import FBX"));
 	{
-		FSpawnableRestoreState SpawnableRestoreState(MovieScene, Player->GetSharedPlaybackState().ToSharedPtr());
+		FAllSpawnableRestoreState SpawnableRestoreState(Player->GetSharedPlaybackState().ToSharedPtr());
  
 		if (SpawnableRestoreState.bWasChanged)
 		{
diff --git a/Engine/Source/Editor/MovieSceneTools/Private/LevelSequenceFBXInterop.cpp b/Engine/Source/Editor/MovieSceneTools/Private/LevelSequenceFBXInterop.cpp
--- a/Engine/Source/Editor/MovieSceneTools/Private/LevelSequenceFBXInterop.cpp
+++ b/Engine/Source/Editor/MovieSceneTools/Private/LevelSequenceFBXInterop.cpp
@@ -242,7 +242,7 @@ void FLevelSequenceFBXInterop::ExportFBXInternal(const FString& ExportFilename,
 		UnFbx::FFbxExporter::FLevelSequenceNodeNameAdapter NodeNameAdapter(MovieScene, Sequencer.Get(), Template);
 
 		{
-			FSpawnableRestoreState SpawnableRestoreState(MovieScene, Sequencer->GetSharedPlaybackState().ToSharedPtr());
+			FAllSpawnableRestoreState SpawnableRestoreState(Sequencer->GetSharedPlaybackState().ToSharedPtr());
 			if (SpawnableRestoreState.bWasChanged)
 			{
 				// Evaluate at the beginning of the subscene time to ensure that spawnables are created before export
diff --git a/Engine/Source/Editor/MovieSceneTools/Private/MovieSceneToolHelpers.cpp b/Engine/Source/Editor/MovieSceneTools/Private/MovieSceneToolHelpers.cpp
--- a/Engine/Source/Editor/MovieSceneTools/Private/MovieSceneToolHelpers.cpp
+++ b/Engine/Source/Editor/MovieSceneTools/Private/MovieSceneToolHelpers.cpp
@@ -4535,6 +4535,7 @@ bool MovieSceneToolHelpers::ExportToAnimSequence(UAnimSequence* AnimSequence, UA
 	return AnimSequence->GetDataModel()->HasBeenPopulated();
 }
 
+PRAGMA_DISABLE_DEPRECATION_WARNINGS  // Use FAllSpawnableRestoreState
 
 FSpawnableRestoreState::FSpawnableRestoreState(UMovieScene* MovieScene)
 	: FSpawnableRestoreState(MovieScene, nullptr)
@@ -4547,57 +4548,11 @@ FSpawnableRestoreState::FSpawnableRestoreState(UMovieScene* MovieScene, TSharedP
 	, WeakMovieScene(MovieScene)
 	, SharedPlaybackState(InSharedPlaybackState)
 {
-	auto SaveStateForSpawnable = [this](FGuid BindingGuid, ESpawnOwnership PreviousSpawnOwnership, TFunction<void(ESpawnOwnership)> SpawnOwnershipSetter)
-	{
-		// Spawnable could be in a subscene, so temporarily override it to persist throughout
-		SpawnOwnershipMap.Add(BindingGuid, PreviousSpawnOwnership);
-
-		SpawnOwnershipSetter(ESpawnOwnership::RootSequence);
-
-		UMovieSceneSpawnTrack* SpawnTrack = WeakMovieScene->FindTrack<UMovieSceneSpawnTrack>(BindingGuid);
-
-		if (SpawnTrack && SpawnTrack->GetAllSections().Num() > 0)
-		{
-			// Start a transaction that will be undone later for the modifications to the spawn track
-			if (!bWasChanged)
-			{
-				GEditor->BeginTransaction(NSLOCTEXT("MovieSceneToolHelpers", "SpwanableRestoreState", "SpawnableRestoreState"));
-			}
-
-			bWasChanged = true;
-
-			UMovieSceneSpawnSection* SpawnSection = Cast<UMovieSceneSpawnSection>(SpawnTrack->GetAllSections()[0]);
-			SpawnSection->Modify();
-			SpawnSection->GetChannel().Reset();
-			SpawnSection->GetChannel().SetDefault(true);
-		}
-
-		UMovieSceneBindingLifetimeTrack* BindingLifetimeTrack = WeakMovieScene->FindTrack<UMovieSceneBindingLifetimeTrack>(BindingGuid);
-		if (BindingLifetimeTrack && BindingLifetimeTrack->GetAllSections().Num() > 0)
-		{
-			// Start a transaction that will be undone later for the modifications to the binding lifetime track
-			if (!bWasChanged)
-			{
-				GEditor->BeginTransaction(NSLOCTEXT("MovieSceneToolHelpers", "SpwanableRestoreState", "SpawnableRestoreState"));
-			}
-
-			bWasChanged = true;
-
-			BindingLifetimeTrack->Modify();
-			BindingLifetimeTrack->RemoveAllAnimationData();
-			// Add a single infinite section
-			UMovieSceneSection* NewSection = NewObject<UMovieSceneSection>(BindingLifetimeTrack, UMovieSceneBindingLifetimeSection::StaticClass(), NAME_None, RF_Transactional);
-			check(NewSection);
-			NewSection->SetRange(TRange<FFrameNumber>::All());
-			BindingLifetimeTrack->AddSection(*NewSection);
-		}
-	};
-	
 	for (int32 SpawnableIndex = 0; SpawnableIndex < WeakMovieScene->GetSpawnableCount(); ++SpawnableIndex)
 	{
 		FMovieSceneSpawnable& Spawnable = WeakMovieScene->GetSpawnable(SpawnableIndex);
 
-		SaveStateForSpawnable(Spawnable.GetGuid(), Spawnable.GetSpawnOwnership(), [&Spawnable](ESpawnOwnership SpawnOwnership) { Spawnable.SetSpawnOwnership(SpawnOwnership);});
+		FAllSpawnableRestoreState::SaveStateForSpawnable(bWasChanged, SpawnOwnershipMap, WeakMovieScene, Spawnable.GetGuid(), Spawnable.GetSpawnOwnership(), [&Spawnable](ESpawnOwnership SpawnOwnership) { Spawnable.SetSpawnOwnership(SpawnOwnership);});
 	}
 
 	TSharedPtr<UE::MovieScene::FSharedPlaybackState> SharedPlaybackStateToUse = SharedPlaybackState;
@@ -4617,7 +4572,7 @@ FSpawnableRestoreState::FSpawnableRestoreState(UMovieScene* MovieScene, TSharedP
 				{
 					if (UMovieSceneSpawnableBindingBase* SpawnableBinding = BindingReference.CustomBinding->AsSpawnable(SharedPlaybackStateToUse.ToSharedRef()))
 					{
-						SaveStateForSpawnable(BindingReference.ID, SpawnableBinding->SpawnOwnership, [SpawnableBinding](ESpawnOwnership SpawnOwnership) { SpawnableBinding->SpawnOwnership = SpawnOwnership; });
+						FAllSpawnableRestoreState::SaveStateForSpawnable(bWasChanged, SpawnOwnershipMap, WeakMovieScene, BindingReference.ID, SpawnableBinding->SpawnOwnership, [SpawnableBinding](ESpawnOwnership SpawnOwnership) { SpawnableBinding->SpawnOwnership = SpawnOwnership; });
 					}								
 				}
 			}
@@ -4696,6 +4651,203 @@ FSpawnableRestoreState::~FSpawnableRestoreState()
 	GEditor->bSuspendBroadcastPostUndoRedo = bOrigSuspendBroadcastPostUndoRedo;
 }
 
+PRAGMA_ENABLE_DEPRECATION_WARNINGS
+
+void FAllSpawnableRestoreState::SaveStateForSpawnable(bool& bInOutWasChanged, TMap<FGuid, ESpawnOwnership>& InOutSpawnOwnershipMap, TWeakObjectPtr<UMovieScene>& WeakMovieScene, FGuid BindingGuid, ESpawnOwnership PreviousSpawnOwnership, TFunction<void(ESpawnOwnership)> SpawnOwnershipSetter)
+{
+	if (WeakMovieScene.IsValid() == false)
+	{
+		return;
+	}
+	// Spawnable could be in a subscene, so temporarily override it to persist throughout
+	InOutSpawnOwnershipMap.Add(BindingGuid, PreviousSpawnOwnership);
+
+	SpawnOwnershipSetter(ESpawnOwnership::RootSequence);
+
+	UMovieSceneSpawnTrack* SpawnTrack = WeakMovieScene->FindTrack<UMovieSceneSpawnTrack>(BindingGuid);
+
+	if (SpawnTrack && SpawnTrack->GetAllSections().Num() > 0)
+	{
+		// Start a transaction that will be undone later for the modifications to the spawn track
+		if (!bInOutWasChanged)
+		{
+			GEditor->BeginTransaction(NSLOCTEXT("MovieSceneToolHelpers", "SpwanableRestoreState", "SpawnableRestoreState"));
+		}
+
+		bInOutWasChanged = true;
+
+		UMovieSceneSpawnSection* SpawnSection = Cast<UMovieSceneSpawnSection>(SpawnTrack->GetAllSections()[0]);
+		SpawnSection->Modify();
+		SpawnSection->GetChannel().Reset();
+		SpawnSection->GetChannel().SetDefault(true);
+	}
+
+	UMovieSceneBindingLifetimeTrack* BindingLifetimeTrack = WeakMovieScene->FindTrack<UMovieSceneBindingLifetimeTrack>(BindingGuid);
+	if (BindingLifetimeTrack && BindingLifetimeTrack->GetAllSections().Num() > 0)
+	{
+		// Start a transaction that will be undone later for the modifications to the binding lifetime track
+		if (!bInOutWasChanged)
+		{
+			GEditor->BeginTransaction(NSLOCTEXT("MovieSceneToolHelpers", "SpwanableRestoreState", "SpawnableRestoreState"));
+		}
+
+		bInOutWasChanged = true;
+
+		BindingLifetimeTrack->Modify();
+		BindingLifetimeTrack->RemoveAllAnimationData();
+		// Add a single infinite section
+		UMovieSceneSection* NewSection = NewObject<UMovieSceneSection>(BindingLifetimeTrack, UMovieSceneBindingLifetimeSection::StaticClass(), NAME_None, RF_Transactional);
+		check(NewSection);
+		NewSection->SetRange(TRange<FFrameNumber>::All());
+		BindingLifetimeTrack->AddSection(*NewSection);
+	}
+};
+
+FAllSpawnableRestoreState::FAllSpawnableRestoreState(TSharedPtr<UE::MovieScene::FSharedPlaybackState> InSharedPlaybackState)
+	: SharedPlaybackState(InSharedPlaybackState)
+{
+
+	TSharedPtr<UE::MovieScene::FSharedPlaybackState> SharedPlaybackStateToUse = SharedPlaybackState;
+	if (!SharedPlaybackStateToUse)
+	{
+		UE_LOG(LogMovieScene, Warning, TEXT("FAllSpawnableRestoreState requires a valid FSharedPlaybackState"));
+		return;
+	}
+	const FMovieSceneSequenceHierarchy* Hierarchy = SharedPlaybackStateToUse->GetHierarchy();
+	if (!Hierarchy)
+	{
+		UE_LOG(LogMovieScene, Warning, TEXT("FAllSpawnableRestoreState requires a valid FMovieSceneSequenceHierarchy from FSharedPlaybackState"));
+		return;
+	}
+	TArray<FMovieSceneSequenceID> SequenceIDs;
+	Hierarchy->AllSubSequenceIDs(SequenceIDs);
+	TArray<TWeakObjectPtr<UMovieSceneSequence>> Sequences;
+	for (const FMovieSceneSequenceID& SequenceID : SequenceIDs)
+	{
+		if (UMovieSceneSequence* Sequence = SharedPlaybackStateToUse->GetSequence(SequenceID))
+		{
+			if (UMovieScene* MovieScene = Sequence->GetMovieScene())
+			{
+				FPerMovieScene PerMovieScene;
+				PerMovieScene.WeakMovieScene = MovieScene;
+				PerMovieScene.SequenceID = SequenceID;
+				OwnershipMaps.Add(MovieScene, PerMovieScene);
+			}
+		}
+	}
+
+	for (TPair <TWeakObjectPtr<UMovieScene>, FPerMovieScene>& OwnershipMap : OwnershipMaps)
+	{
+		if (OwnershipMap.Key.IsValid() == false)
+		{
+			continue;
+		}
+		for (int32 SpawnableIndex = 0; SpawnableIndex < OwnershipMap.Key->GetSpawnableCount(); ++SpawnableIndex)
+		{
+			FMovieSceneSpawnable& Spawnable = OwnershipMap.Key->GetSpawnable(SpawnableIndex);
+
+			FAllSpawnableRestoreState::SaveStateForSpawnable(bWasChanged, OwnershipMap.Value.SpawnOwnershipMap, OwnershipMap.Key, Spawnable.GetGuid(), Spawnable.GetSpawnOwnership(), [&Spawnable](ESpawnOwnership SpawnOwnership) { Spawnable.SetSpawnOwnership(SpawnOwnership); });
+		}
+
+		ensure(SharedPlaybackStateToUse.IsValid());
+		if (UMovieSceneSequence* Sequence = OwnershipMap.Key->GetTypedOuter<UMovieSceneSequence>())
+		{
+			if (FMovieSceneBindingReferences* BindingReferences = Sequence->GetBindingReferences())
+			{
+				for (FMovieSceneBindingReference& BindingReference : BindingReferences->GetAllReferences())
+				{
+					if (BindingReference.CustomBinding)
+					{
+						if (UMovieSceneSpawnableBindingBase* SpawnableBinding = BindingReference.CustomBinding->AsSpawnable(SharedPlaybackStateToUse.ToSharedRef()))
+						{
+							FAllSpawnableRestoreState::SaveStateForSpawnable(bWasChanged, OwnershipMap.Value.SpawnOwnershipMap, OwnershipMap.Key, BindingReference.ID, SpawnableBinding->SpawnOwnership, [SpawnableBinding](ESpawnOwnership SpawnOwnership) { SpawnableBinding->SpawnOwnership = SpawnOwnership; });
+						}
+					}
+				}
+			}
+		}
+	}
+}
+
+FAllSpawnableRestoreState::~FAllSpawnableRestoreState()
+{
+	if (bWasChanged)
+	{
+		GEditor->EndTransaction();
+	}
+
+	if (!bWasChanged || OwnershipMaps.IsEmpty())
+	{
+		return;
+	}
+
+	TSharedPtr<UE::MovieScene::FSharedPlaybackState> SharedPlaybackStateToUse = SharedPlaybackState;
+
+	for (TPair <TWeakObjectPtr<UMovieScene>, FPerMovieScene>& OwnershipMap : OwnershipMaps)
+	{
+		if (OwnershipMap.Key.IsValid() == false)
+		{
+			continue;
+		}
+		// Restore spawnable owners
+		for (int32 SpawnableIndex = 0; SpawnableIndex < OwnershipMap.Key->GetSpawnableCount(); ++SpawnableIndex)
+		{
+			FMovieSceneSpawnable& Spawnable = OwnershipMap.Key->GetSpawnable(SpawnableIndex);
+			if (ESpawnOwnership* SpawnOwnership = OwnershipMap.Value.SpawnOwnershipMap.Find(Spawnable.GetGuid()))
+			{
+				Spawnable.SetSpawnOwnership(*SpawnOwnership);
+			}
+		}
+
+		if (!SharedPlaybackStateToUse)
+		{
+			// Supporting deprecated path where we have no SharedPlaybackState
+			UE::MovieScene::FSharedPlaybackStateCreateParams CreateParams;
+			CreateParams.PlaybackContext = GEditor->GetEditorWorldContext().World();
+			UMovieSceneSequence* ThisSequence = OwnershipMap.Key->GetTypedOuter<UMovieSceneSequence>();
+			SharedPlaybackStateToUse = MakeShared<UE::MovieScene::FSharedPlaybackState>(*ThisSequence, CreateParams);
+
+			FMovieSceneEvaluationState State;
+			SharedPlaybackStateToUse->AddCapabilityRaw(&State);
+			State.AssignSequence(MovieSceneSequenceID::Root, *ThisSequence, SharedPlaybackStateToUse.ToSharedRef());
+		}
+		ensure(SharedPlaybackStateToUse.IsValid());
+
+		if (UMovieSceneSequence* Sequence = OwnershipMap.Key->GetTypedOuter<UMovieSceneSequence>())
+		{
+			if (FMovieSceneBindingReferences* BindingReferences = Sequence->GetBindingReferences())
+			{
+				for (FMovieSceneBindingReference& BindingReference : BindingReferences->GetAllReferences())
+				{
+					if (BindingReference.CustomBinding)
+					{
+						if (UMovieSceneSpawnableBindingBase* SpawnableBinding = BindingReference.CustomBinding->AsSpawnable(SharedPlaybackStateToUse.ToSharedRef()))
+						{
+							if (ESpawnOwnership* SpawnOwnership = OwnershipMap.Value.SpawnOwnershipMap.Find(BindingReference.ID))
+							{
+								SpawnableBinding->SpawnOwnership = *SpawnOwnership;
+							}
+						}
+					}
+				}
+			}
+		}
+	}
+
+	// Restore modified spawned sections
+	const bool bOrigSquelchTransactionNotification = GEditor->bSquelchTransactionNotification;
+	GEditor->bSquelchTransactionNotification = true;
+	const bool bOrigIgnoreSelectionChange = GEditor->bIgnoreSelectionChange;
+	GEditor->bIgnoreSelectionChange = true;
+	const bool bOrigSuspendBroadcastPostUndoRedo = GEditor->bSuspendBroadcastPostUndoRedo;
+	GEditor->bSuspendBroadcastPostUndoRedo = true;
+
+	GEditor->UndoTransaction(false);
+
+	GEditor->bSquelchTransactionNotification = bOrigSquelchTransactionNotification;
+	GEditor->bIgnoreSelectionChange = bOrigIgnoreSelectionChange;
+	GEditor->bSuspendBroadcastPostUndoRedo = bOrigSuspendBroadcastPostUndoRedo;
+}
 
 void MovieSceneToolHelpers::GetParents(TArray<const UObject*>& Parents, const UObject* InObject)
 {
diff --git a/Engine/Source/Editor/MovieSceneTools/Public/MovieSceneToolHelpers.h b/Engine/Source/Editor/MovieSceneTools/Public/MovieSceneToolHelpers.h
--- a/Engine/Source/Editor/MovieSceneTools/Public/MovieSceneToolHelpers.h
+++ b/Engine/Source/Editor/MovieSceneTools/Public/MovieSceneToolHelpers.h
@@ -1223,19 +1223,43 @@ public:
 	static UE_API void ImportAnimSequenceTransformsEnterPressed(const TSharedRef<ISequencer>& Sequencer, const TArray<FAssetData>& Asset, UMovieScene3DTransformTrack* TransformTrack);
 };
 
-// Helper to make spawnables persist throughout the export process and then restore properly afterwards
 struct FSpawnableRestoreState
 {
 	UE_DEPRECATED(5.5, "This constructor is deprecated in favor of passing in shared playback state")
 	UE_API FSpawnableRestoreState(UMovieScene* MovieScene);
 
+	UE_DEPRECATED(5.7, "FSpawnableRestoreState has been deprecated in favor of FAllSpawnableRestoreState that turns off and restores spawnables for the whole sequence")
 	UE_API FSpawnableRestoreState(UMovieScene* MovieScene, TSharedPtr<UE::MovieScene::FSharedPlaybackState> InSharedPlaybackState);
+	UE_DEPRECATED(5.7, "FSpawnableRestoreState has been deprecated in favor of FAllSpawnableRestoreState that turns off and restores spawnables for the whole sequence")
 	UE_API ~FSpawnableRestoreState();
 
 	bool bWasChanged;
 	TMap<FGuid, ESpawnOwnership> SpawnOwnershipMap;
 	TWeakObjectPtr<UMovieScene> WeakMovieScene;
 	TSharedPtr<UE::MovieScene::FSharedPlaybackState> SharedPlaybackState;
+
+};
+
+// Helper to make all spawnables in all nested sub sequences persist throughout the export process and then restore properly afterwards
+// You should use this when you may have multiple nested sequences, in particularly when using sequencer as the player
+struct FAllSpawnableRestoreState
+{
+	UE_API FAllSpawnableRestoreState(TSharedPtr<UE::MovieScene::FSharedPlaybackState> InSharedPlaybackState);
+	UE_API ~FAllSpawnableRestoreState();
+	bool bWasChanged = false;
+
+	//helper function to set up the SpawnOwnershipMap given the binding, will update if the spawn track was updated
+	static void SaveStateForSpawnable(bool& bInOutWasChanged, TMap<FGuid, ESpawnOwnership>& InOutSpawnOwnershipMap, TWeakObjectPtr<UMovieScene>& WeakMovieScene, FGuid BindingGuid, ESpawnOwnership PreviousSpawnOwnership, TFunction<void(ESpawnOwnership)> SpawnOwnershipSetter);
+
+private:
+	struct FPerMovieScene
+	{
+		TMap<FGuid, ESpawnOwnership> SpawnOwnershipMap;
+		TWeakObjectPtr<UMovieScene> WeakMovieScene;
+		FMovieSceneSequenceID SequenceID;
+	};
+	TMap <TWeakObjectPtr<UMovieScene>, FPerMovieScene> OwnershipMaps;
+	TSharedPtr<UE::MovieScene::FSharedPlaybackState> SharedPlaybackState;
 };
 
 class FTrackEditorBindingIDPicker : public FMovieSceneObjectBindingIDPicker
