Anim In Engine: Fix broken linked anim sequences.

diff --git a/Engine/Plugins/Animation/ControlRig/Source/ControlRigEditor/Private/ControlRigSequencerEditorLibrary.cpp b/Engine/Plugins/Animation/ControlRig/Source/ControlRigEditor/Private/ControlRigSequencerEditorLibrary.cpp
--- a/Engine/Plugins/Animation/ControlRig/Source/ControlRigEditor/Private/ControlRigSequencerEditorLibrary.cpp
+++ b/Engine/Plugins/Animation/ControlRig/Source/ControlRigEditor/Private/ControlRigSequencerEditorLibrary.cpp
@@ -1235,6 +1235,7 @@ bool UControlRigSequencerEditorLibrary::BakeToControlRig(UWorld* World, ULevelSe
 	FMovieSceneSequenceTransform RootToLocalTransform;
 	IMovieScenePlayer* Player = nullptr;
 	ULevelSequencePlayer* LevelPlayer = nullptr;
+	UMovieScene* MovieSceneForSpawnableRestore = nullptr;
 	if (WeakSequencer.IsValid())
 	{
 		Player = WeakSequencer.Pin().Get();
@@ -1242,6 +1243,7 @@ bool UControlRigSequencerEditorLibrary::BakeToControlRig(UWorld* World, ULevelSe
 	else
 	{
 		Player = LevelPlayer = ULevelSequencePlayer::CreateLevelSequencePlayer(World, LevelSequence, Settings, OutActor);
+		MovieSceneForSpawnableRestore = MovieScene;
 	}
 	if (Player == nullptr)
 	{
@@ -1251,7 +1253,7 @@ bool UControlRigSequencerEditorLibrary::BakeToControlRig(UWorld* World, ULevelSe
 	bool bResult = false;
 	const FScopedTransaction Transaction(LOCTEXT("BakeToControlRig_Transaction", "Bake To Control Rig"));
 	{
-		FAllSpawnableRestoreState SpawnableRestoreState(Player->GetSharedPlaybackState().ToSharedPtr());
+		FAllSpawnableRestoreState SpawnableRestoreState(Player->GetSharedPlaybackState().ToSharedPtr(), MovieSceneForSpawnableRestore);
 
 		if (LevelPlayer && SpawnableRestoreState.bWasChanged)
 		{
diff --git a/Engine/Plugins/Importers/USDImporter/Source/USDStage/Private/USDLevelSequenceHelper.cpp b/Engine/Plugins/Importers/USDImporter/Source/USDStage/Private/USDLevelSequenceHelper.cpp
--- a/Engine/Plugins/Importers/USDImporter/Source/USDStage/Private/USDLevelSequenceHelper.cpp
+++ b/Engine/Plugins/Importers/USDImporter/Source/USDStage/Private/USDLevelSequenceHelper.cpp
@@ -4818,6 +4818,7 @@ void FUsdLevelSequenceHelperImpl::HandleControlRigSectionChange(UMovieSceneContr
 	// where the section was changed after all), but its possible to have a fallback too
 	IMovieScenePlayer* Player = nullptr;
 	ULevelSequencePlayer* LevelPlayer = nullptr;
+	UMovieScene* MovieSceneForSpawnableRestore = nullptr;
 	{
 		if (PinnedSequencer)
 		{
@@ -4828,6 +4829,7 @@ void FUsdLevelSequenceHelperImpl::HandleControlRigSectionChange(UMovieSceneContr
 			ALevelSequenceActor* OutActor = nullptr;
 			FMovieSceneSequencePlaybackSettings Settings;
 			Player = LevelPlayer = ULevelSequencePlayer::CreateLevelSequencePlayer(World, LevelSequence, Settings, OutActor);
+			MovieSceneForSpawnableRestore = MovieScene;
 		}
 
 		if (!Player)
@@ -4847,7 +4849,7 @@ void FUsdLevelSequenceHelperImpl::HandleControlRigSectionChange(UMovieSceneContr
 			PinnedSequencer->EnterSilentMode();
 		}
 
-		FAllSpawnableRestoreState SpawnableRestoreState(Player->GetSharedPlaybackState().ToSharedPtr());
+		FAllSpawnableRestoreState SpawnableRestoreState(Player->GetSharedPlaybackState().ToSharedPtr(), MovieSceneForSpawnableRestore);
 		if (LevelPlayer && SpawnableRestoreState.bWasChanged)
 		{
 			// Evaluate at the beginning of the subscene time to ensure that spawnables are created before export
diff --git a/Engine/Plugins/MovieScene/SequencerScripting/Source/SequencerScriptingEditor/Private/SequencerTools.cpp b/Engine/Plugins/MovieScene/SequencerScripting/Source/SequencerScriptingEditor/Private/SequencerTools.cpp
--- a/Engine/Plugins/MovieScene/SequencerScripting/Source/SequencerScriptingEditor/Private/SequencerTools.cpp
+++ b/Engine/Plugins/MovieScene/SequencerScripting/Source/SequencerScriptingEditor/Private/SequencerTools.cpp
@@ -304,7 +304,7 @@ bool ExportFBXInternal(const FSequencerExportFBXParams& InParams, UMovieSceneSeq
 	bool bDidExport = false;
 	
 	{
-		FAllSpawnableRestoreState SpawnableRestoreState(Player->GetSharedPlaybackState().ToSharedPtr());
+		FAllSpawnableRestoreState SpawnableRestoreState(Player->GetSharedPlaybackState().ToSharedPtr(), MovieScene);
 
 		if (SpawnableRestoreState.bWasChanged)
 		{
@@ -402,7 +402,7 @@ bool USequencerToolsFunctionLibrary::ExportAnimSequenceWaitForDelegate(UWorld* W
 
 	bool bResult = false;
 
-	FAllSpawnableRestoreState SpawnableRestoreState(Player->GetSharedPlaybackState().ToSharedPtr());
+	FAllSpawnableRestoreState SpawnableRestoreState(Player->GetSharedPlaybackState().ToSharedPtr(), MovieScene);
 
 	if (SpawnableRestoreState.bWasChanged)
 	{
@@ -788,7 +788,7 @@ bool ImportFBXInternal(UWorld* World, UMovieSceneSequence* Sequence, const TArra
 	bool bResult = false;
 	FScopedTransaction ImportFBXTransaction(NSLOCTEXT("Sequencer", "ImportFBX", "Import FBX"));
 	{
-		FAllSpawnableRestoreState SpawnableRestoreState(Player->GetSharedPlaybackState().ToSharedPtr());
+		FAllSpawnableRestoreState SpawnableRestoreState(Player->GetSharedPlaybackState().ToSharedPtr(), MovieScene);
  
 		if (SpawnableRestoreState.bWasChanged)
 		{
diff --git a/Engine/Source/Editor/MovieSceneTools/Private/MovieSceneToolHelpers.cpp b/Engine/Source/Editor/MovieSceneTools/Private/MovieSceneToolHelpers.cpp
--- a/Engine/Source/Editor/MovieSceneTools/Private/MovieSceneToolHelpers.cpp
+++ b/Engine/Source/Editor/MovieSceneTools/Private/MovieSceneToolHelpers.cpp
@@ -4703,7 +4703,7 @@ void FAllSpawnableRestoreState::SaveStateForSpawnable(bool& bInOutWasChanged, TM
 	}
 };
 
-FAllSpawnableRestoreState::FAllSpawnableRestoreState(TSharedPtr<UE::MovieScene::FSharedPlaybackState> InSharedPlaybackState)
+FAllSpawnableRestoreState::FAllSpawnableRestoreState(TSharedPtr<UE::MovieScene::FSharedPlaybackState> InSharedPlaybackState, UMovieScene* InMovieScene)
 	: SharedPlaybackState(InSharedPlaybackState)
 {
 
@@ -4713,25 +4713,35 @@ FAllSpawnableRestoreState::FAllSpawnableRestoreState(TSharedPtr<UE::MovieScene::
 		UE_LOG(LogMovieScene, Warning, TEXT("FAllSpawnableRestoreState requires a valid FSharedPlaybackState"));
 		return;
 	}
-	const FMovieSceneSequenceHierarchy* Hierarchy = SharedPlaybackStateToUse->GetHierarchy();
-	if (!Hierarchy)
+	if (InMovieScene)
 	{
-		UE_LOG(LogMovieScene, Warning, TEXT("FAllSpawnableRestoreState requires a valid FMovieSceneSequenceHierarchy from FSharedPlaybackState"));
-		return;
+		FPerMovieScene PerMovieScene;
+		PerMovieScene.WeakMovieScene = InMovieScene;
+		PerMovieScene.SequenceID = MovieSceneSequenceID::Root;
+		OwnershipMaps.Add(InMovieScene, PerMovieScene);
 	}
-	TArray<FMovieSceneSequenceID> SequenceIDs;
-	Hierarchy->AllSubSequenceIDs(SequenceIDs);
-	TArray<TWeakObjectPtr<UMovieSceneSequence>> Sequences;
-	for (const FMovieSceneSequenceID& SequenceID : SequenceIDs)
+	else
 	{
-		if (UMovieSceneSequence* Sequence = SharedPlaybackStateToUse->GetSequence(SequenceID))
+		const FMovieSceneSequenceHierarchy* Hierarchy = SharedPlaybackStateToUse->GetHierarchy();
+		if (!Hierarchy)
 		{
-			if (UMovieScene* MovieScene = Sequence->GetMovieScene())
+			UE_LOG(LogMovieScene, Warning, TEXT("FAllSpawnableRestoreState requires a valid FMovieSceneSequenceHierarchy from FSharedPlaybackState"));
+			return;
+		}
+		TArray<FMovieSceneSequenceID> SequenceIDs;
+		Hierarchy->AllSubSequenceIDs(SequenceIDs);
+		TArray<TWeakObjectPtr<UMovieSceneSequence>> Sequences;
+		for (const FMovieSceneSequenceID& SequenceID : SequenceIDs)
+		{
+			if (UMovieSceneSequence* Sequence = SharedPlaybackStateToUse->GetSequence(SequenceID))
 			{
-				FPerMovieScene PerMovieScene;
-				PerMovieScene.WeakMovieScene = MovieScene;
-				PerMovieScene.SequenceID = SequenceID;
-				OwnershipMaps.Add(MovieScene, PerMovieScene);
+				if (UMovieScene* MovieScene = Sequence->GetMovieScene())
+				{
+					FPerMovieScene PerMovieScene;
+					PerMovieScene.WeakMovieScene = MovieScene;
+					PerMovieScene.SequenceID = SequenceID;
+					OwnershipMaps.Add(MovieScene, PerMovieScene);
+				}
 			}
 		}
 	}
diff --git a/Engine/Source/Editor/MovieSceneTools/Public/MovieSceneToolHelpers.h b/Engine/Source/Editor/MovieSceneTools/Public/MovieSceneToolHelpers.h
--- a/Engine/Source/Editor/MovieSceneTools/Public/MovieSceneToolHelpers.h
+++ b/Engine/Source/Editor/MovieSceneTools/Public/MovieSceneToolHelpers.h
@@ -1242,9 +1242,10 @@ struct FSpawnableRestoreState
 
 // Helper to make all spawnables in all nested sub sequences persist throughout the export process and then restore properly afterwards
 // You should use this when you may have multiple nested sequences, in particularly when using sequencer as the player
+// If you create a player you need to pass in a MovieScene to use
 struct FAllSpawnableRestoreState
 {
-	UE_API FAllSpawnableRestoreState(TSharedPtr<UE::MovieScene::FSharedPlaybackState> InSharedPlaybackState);
+	UE_API FAllSpawnableRestoreState(TSharedPtr<UE::MovieScene::FSharedPlaybackState> InSharedPlaybackState, UMovieScene* OptionalMovieScene = nullptr);
 	UE_API ~FAllSpawnableRestoreState();
 	bool bWasChanged = false;
 
diff --git a/Engine/Source/Runtime/MovieSceneTracks/Private/Sections/MovieSceneSkeletalAnimationSection.cpp b/Engine/Source/Runtime/MovieSceneTracks/Private/Sections/MovieSceneSkeletalAnimationSection.cpp
--- a/Engine/Source/Runtime/MovieSceneTracks/Private/Sections/MovieSceneSkeletalAnimationSection.cpp
+++ b/Engine/Source/Runtime/MovieSceneTracks/Private/Sections/MovieSceneSkeletalAnimationSection.cpp
@@ -836,18 +836,21 @@ bool UMovieSceneSkeletalAnimationSection::GetRootMotionTransform(FAnimationPoseD
 	return true;
 }
 
+//this is actually multiply out on previous clips, let's refactor with new anim mixer
 void UMovieSceneSkeletalAnimationSection::MultiplyOutInverseOnNextClips(FVector PreviousMatchedLocationOffset, FRotator PreviousMatchedRotationOffset)
 {
 	UMovieSceneCommonAnimationTrack* Track = GetTypedOuter<UMovieSceneCommonAnimationTrack>();
 	if (Track)
 	{
+		
 		bool bMultiplyOutInverse = false;
 		//calculate the diff here....
 		FTransform Previous(PreviousMatchedRotationOffset.Quaternion(), PreviousMatchedLocationOffset);
 		FTransform Matched(MatchedRotationOffset.Quaternion(), MatchedLocationOffset);
 		FTransform Inverse = Previous.GetRelativeTransformReverse(Matched);
-		for (UMovieSceneSection* Section : Track->AnimationSections)
-		{
+		for(int32 Index = Track->AnimationSections.Num() - 1; Index >=0; -- Index)
+		{ 
+			UMovieSceneSection* Section  = Track->AnimationSections[Index];
 			if (bMultiplyOutInverse)
 			{
 				UMovieSceneSkeletalAnimationSection* AnimSection = Cast<UMovieSceneSkeletalAnimationSection>(Section);
