Constraints: fix offset changed delegate not doing anything with infinite constraints

#rnx

diff --git a/Engine/Source/Editor/MovieSceneTools/Private/Constraints/MovieSceneConstraintChannelHelper.cpp b/Engine/Source/Editor/MovieSceneTools/Private/Constraints/MovieSceneConstraintChannelHelper.cpp
--- a/Engine/Source/Editor/MovieSceneTools/Private/Constraints/MovieSceneConstraintChannelHelper.cpp
+++ b/Engine/Source/Editor/MovieSceneTools/Private/Constraints/MovieSceneConstraintChannelHelper.cpp
@@ -1586,7 +1586,7 @@ void FMovieSceneConstraintChannelHelper::HandleOffsetChanged(
 	}
 
 	const TArrayView<const FFrameNumber> Times = InActiveChannel.GetTimes();
-	if (Times.IsEmpty())
+	if (Times.IsEmpty() && !InActiveChannel.IsInfinite())
 	{
 		return;
 	}
