Sequencer: Fixed keying operations on tracks where there was no overlapping section

The code inside FRegisteredPropertyTrackEditor::ProcessKeyOperation was previously reliant on there being an entity alive in order to run the recomposition code. This is erroneous in the case where a keying operation is being performed on a track at a time where it does not have any section (which often happens in UMG due to its finite-by-default sections). The expected behavior here is that the keying operation should expand the relevant section and add a key for the external value.

I've changed the function so that the recomposition code is optional, and only run if there is at least a single valid entity to recompose otherwise the incoming value is used directly.

diff --git a/Engine/Source/Editor/MovieSceneTools/Private/TrackEditors/PropertyTrackEditors/RegisteredPropertyTrackEditor.cpp b/Engine/Source/Editor/MovieSceneTools/Private/TrackEditors/PropertyTrackEditors/RegisteredPropertyTrackEditor.cpp
--- a/Engine/Source/Editor/MovieSceneTools/Private/TrackEditors/PropertyTrackEditors/RegisteredPropertyTrackEditor.cpp
+++ b/Engine/Source/Editor/MovieSceneTools/Private/TrackEditors/PropertyTrackEditors/RegisteredPropertyTrackEditor.cpp
@@ -320,15 +320,15 @@ void FRegisteredPropertyTrackEditor::ProcessKeyOperation(UObject* Object, TArray
 	TArray<FMovieSceneEntityID> Entities;
 	for (const FKeySectionOperation& Operation : SectionsToKey)
 	{
-		FMovieSceneEntityID EntityID = EvaluationTemplate.FindEntityFromOwner(Operation.Section->GetSectionObject(), 0, SequenceID);
-		if (!EntityID)
+		if (!Track)
 		{
-			continue;
+			Track = CastChecked<UMovieScenePropertyTrack>(Operation.Section->GetSectionObject()->GetOuter());
 		}
 
-		if (!Track)
+		FMovieSceneEntityID EntityID = EvaluationTemplate.FindEntityFromOwner(Operation.Section->GetSectionObject(), 0, SequenceID);
+		if (!EntityID)
 		{
-			Track = CastChecked<UMovieScenePropertyTrack>(Operation.Section->GetSectionObject()->GetOuter());
+			continue;
 		}
 
 		// Find an entity that matches the bound object
@@ -364,80 +364,81 @@ void FRegisteredPropertyTrackEditor::ProcessKeyOperation(UObject* Object, TArray
 	UMovieSceneEntitySystemLinker*         EntityLinker = EvaluationTemplate.GetEntitySystemLinker();
 	UMovieScenePropertyInstantiatorSystem* System       = EntityLinker ? EntityLinker->FindSystem<UMovieScenePropertyInstantiatorSystem>() : nullptr;
 
+	// Coerce the value to its operating type
+	FIntermediatePropertyValue CoercedValue = PropertyDefinition->Handler->CoercePropertyValue(*PropertyDefinition, *PropertyAndValue->Key, PropertyAndValue->Value);
+
+	// Recompose the value if possible
 	if (System && Entities.Num() != 0)
 	{
 		FDecompositionQuery Query;
 		Query.Entities = Entities;
-		Query.Object   = Object;
+		Query.Object = Object;
 		Query.bConvertFromSourceEntityIDs = false;
 
 		TGuardValue<FEntityManager*> DebugVizGuard(GEntityManagerForDebuggingVisualizers, &EntityLinker->EntityManager);
 
-		// Coerce the value to its operating type
-		FIntermediatePropertyValue CoercedValue = PropertyDefinition->Handler->CoercePropertyValue(*PropertyDefinition, *PropertyAndValue->Key, PropertyAndValue->Value);
-
 		// Recompose the value from any blending that might be present
 		FPropertyRecomposerImpl Recomposer;
 		Recomposer.OnGetPropertyInfo = FOnGetPropertyRecomposerPropertyInfo::CreateUObject(
-					System, &UMovieScenePropertyInstantiatorSystem::FindPropertyFromSource);
+			System, &UMovieScenePropertyInstantiatorSystem::FindPropertyFromSource);
 
-		FIntermediatePropertyValue RecomposedValue = MoveTemp(Recomposer.RecomposeBlendOperational(*PropertyDefinition, Query, CoercedValue).Values[0]);
+		CoercedValue = MoveTemp(Recomposer.RecomposeBlendOperational(*PropertyDefinition, Query, CoercedValue).Values[0]);
+	}
 
-		// Unpack the resulting value to its constituent channels
-		FUnpackedChannelValues UnpackedChannelValues;
-		PropertyDefinition->Handler->UnpackChannels(*PropertyDefinition, *PropertyAndValue->Key, RecomposedValue, UnpackedChannelValues);
+	// Unpack the resulting value to its constituent channels
+	FUnpackedChannelValues UnpackedChannelValues;
+	PropertyDefinition->Handler->UnpackChannels(*PropertyDefinition, *PropertyAndValue->Key, CoercedValue, UnpackedChannelValues);
+
+	if (UnpackedChannelValues.GetValues().Num() == 0)
+	{
+		return;
+	}
 
-		if (UnpackedChannelValues.GetValues().Num() == 0)
+	for (const FKeySectionOperation& Operation : SectionsToKey)
+	{
+		UMovieSceneSection* SectionObject = Operation.Section->GetSectionObject();
+		if (!SectionObject)
 		{
-			return;
+			continue;
 		}
 
-		for (const FKeySectionOperation& Operation : SectionsToKey)
+		FMovieSceneChannelProxy& ChannelProxy = SectionObject->GetChannelProxy();
+
+		// Figure out which channels relate to which indices in the unpacked values
+		TMap<FMovieSceneChannel*, int32> ChannelToIndex;
+		for (int32 Index = 0; Index < UnpackedChannelValues.Num(); ++Index)
 		{
-			UMovieSceneSection* SectionObject = Operation.Section->GetSectionObject();
-			if (!SectionObject)
+			if (FMovieSceneChannel* Channel = UnpackedChannelValues[Index]->RetrieveChannel(ChannelProxy))
 			{
-				continue;
+				ChannelToIndex.Add(Channel, Index);
 			}
+		}
 
-			FMovieSceneChannelProxy& ChannelProxy = SectionObject->GetChannelProxy();
-
-			// Figure out which channels relate to which indices in the unpacked values
-			TMap<FMovieSceneChannel*, int32> ChannelToIndex;
-			for (int32 Index = 0; Index < UnpackedChannelValues.Num(); ++Index)
+		for (int32 Index = 0; Index < Operation.KeyAreas.Num(); ++Index)
+		{
+			TSharedPtr<IKeyArea> KeyArea = Operation.KeyAreas[Index];
+			FMovieSceneChannel* Channel = KeyArea->ResolveChannel();
+			if (!Channel)
 			{
-				if (FMovieSceneChannel* Channel = UnpackedChannelValues[Index]->RetrieveChannel(ChannelProxy))
-				{
-					ChannelToIndex.Add(Channel, Index);
-				}
+				continue;
 			}
 
-			for (int32 Index = 0; Index < Operation.KeyAreas.Num(); ++Index)
+			const int32* UnpackedIndex = ChannelToIndex.Find(Channel);
+			if (UnpackedIndex == nullptr || !SectionObject->TryModify())
 			{
-				TSharedPtr<IKeyArea> KeyArea = Operation.KeyAreas[Index];
-				FMovieSceneChannel* Channel = KeyArea->ResolveChannel();
-				if (!Channel)
-				{
-					continue;
-				}
-
-				const int32* UnpackedIndex = ChannelToIndex.Find(Channel);
-				if (UnpackedIndex == nullptr || !SectionObject->TryModify())
-				{
-					continue;
-				}
+				continue;
+			}
 
-				EMovieSceneKeyInterpolation InterpolationMode = KeyArea->GetInterpolationMode(KeyTime, GetSequencer()->GetKeyInterpolation());
-				UnpackedChannelValues[*UnpackedIndex]->AddKey(Channel, KeyTime, InterpolationMode);
+			EMovieSceneKeyInterpolation InterpolationMode = KeyArea->GetInterpolationMode(KeyTime, GetSequencer()->GetKeyInterpolation());
+			UnpackedChannelValues[*UnpackedIndex]->AddKey(Channel, KeyTime, InterpolationMode);
 
-				if (OutResults)
+			if (OutResults)
+			{
+				TArray<FKeyHandle> Handles;
+				KeyArea->GetKeyHandles(Handles, TRange<FFrameNumber>(KeyTime));
+				for (FKeyHandle Handle : Handles)
 				{
-					TArray<FKeyHandle> Handles;
-					KeyArea->GetKeyHandles(Handles, TRange<FFrameNumber>(KeyTime));
-					for (FKeyHandle Handle : Handles)
-					{
-						OutResults->Add({ KeyArea, Handle });
-					}
+					OutResults->Add({ KeyArea, Handle });
 				}
 			}
 		}
