Constraints: baking with new evaluation & physics
- ensure that the animation is updated with the proper delta time when baking

#rnx

diff --git a/Engine/Plugins/Animation/ControlRig/Source/ControlRigEditor/Private/Tools/ControlRigSnapper.cpp b/Engine/Plugins/Animation/ControlRig/Source/ControlRigEditor/Private/Tools/ControlRigSnapper.cpp
--- a/Engine/Plugins/Animation/ControlRig/Source/ControlRigEditor/Private/Tools/ControlRigSnapper.cpp
+++ b/Engine/Plugins/Animation/ControlRig/Source/ControlRigEditor/Private/Tools/ControlRigSnapper.cpp
@@ -36,6 +36,7 @@
 #include "Transform/TransformConstraint.h"
 #include "Transform/TransformConstraintUtil.h"
 #include "Constraints/MovieSceneConstraintChannelHelper.inl"
+#include "Transform/AnimationEvaluation.h"
 
 
 #include UE_INLINE_GENERATED_CPP_BY_NAME(ControlRigSnapper)
@@ -132,10 +133,12 @@ static bool LocalGetControlRigControlTransforms(IMovieScenePlayer* Player, const
 		UWorld* World = ControlRig->GetWorld();
 		const FConstraintsManagerController& Controller = FConstraintsManagerController::Get(World);
 		FFrameRate TickResolution = MovieScene->GetTickResolution();
-		FFrameRate DisplayRate = MovieScene->GetDisplayRate();
 
 		FMovieSceneInverseSequenceTransform LocalToRootTransform = RootToLocalTransform.Inverse();
 
+		const float DisplayRate = static_cast<float>(MovieScene->GetDisplayRate().AsDecimal());
+		const UE::Anim::FEvaluationForCachingScope CachingScope(DisplayRate > 0.f ? 1.f / DisplayRate : 1.f / 30.f);
+		
 		OutTransforms.SetNum(Frames.Num());
 		for (int32 Index = 0; Index < Frames.Num(); ++Index)
 		{
@@ -273,12 +276,14 @@ struct FGuidAndActor
 		{
 			UMovieScene* MovieScene = Sequencer->GetFocusedMovieSceneSequence()->GetMovieScene();
 			FFrameRate TickResolution = MovieScene->GetTickResolution();
-			FFrameRate DisplayRate = MovieScene->GetDisplayRate();
 			FMovieSceneInverseSequenceTransform LocalToRootTransform = Sequencer->GetFocusedMovieSceneSequenceTransform().Inverse();
 
 			//adjust keys for constraint
 			const FConstraintsManagerController& Controller = FConstraintsManagerController::Get(Actor->GetWorld());
 
+			const float DisplayRate = static_cast<float>(MovieScene->GetDisplayRate().AsDecimal());
+			const UE::Anim::FEvaluationForCachingScope CachingScope(DisplayRate > 0.f ? 1.f / DisplayRate : 1.f / 30.f);
+			
 			for (int32 Index = 0; Index < LocalTransforms.Num(); ++Index) 
 			{
 				const FFrameNumber Frame = Frames[Index];
diff --git a/Engine/Plugins/Animation/ControlRig/Source/ControlRigEditor/Private/Tools/EvaluateSequencerTools.cpp b/Engine/Plugins/Animation/ControlRig/Source/ControlRigEditor/Private/Tools/EvaluateSequencerTools.cpp
--- a/Engine/Plugins/Animation/ControlRig/Source/ControlRigEditor/Private/Tools/EvaluateSequencerTools.cpp
+++ b/Engine/Plugins/Animation/ControlRig/Source/ControlRigEditor/Private/Tools/EvaluateSequencerTools.cpp
@@ -251,7 +251,6 @@ namespace AIE
 		UMovieScene* MovieScene = Sequencer->GetFocusedMovieSceneSequence()->GetMovieScene();
 		const FConstraintsManagerController& Controller = FConstraintsManagerController::Get(World);
 		FFrameRate TickResolution = MovieScene->GetTickResolution();
-		FFrameRate DisplayRate = MovieScene->GetDisplayRate();
 
 		TOptional<TRange<FFrameNumber>> OptionalRange = Sequencer->GetSubSequenceRange();
 		const FFrameNumber StartFrame = OptionalRange.IsSet() ? OptionalRange.GetValue().GetLowerBoundValue() : MovieScene->GetPlaybackRange().GetLowerBoundValue();
@@ -261,6 +260,10 @@ namespace AIE
 
 		TSet<UControlRig*> ControlsToRevalaute;
 		FControlRigEditMode::FTurnOffPosePoseUpdate  TurnOff; //stop flashing
+
+		const float DisplayRate = static_cast<float>(MovieScene->GetDisplayRate().AsDecimal());
+		const Anim::FEvaluationForCachingScope CachingScope(DisplayRate > 0.f ? 1.f / DisplayRate : 1.f / 30.f);
+		
 		for (int32 Index : Indices)
 		{
 			FFrameNumber FrameNumber = FrameTimeByIndex.CalculateFrame(Index);
diff --git a/Engine/Source/Editor/MovieSceneTools/Private/Constraints/MovieSceneConstraintChannelHelper.cpp b/Engine/Source/Editor/MovieSceneTools/Private/Constraints/MovieSceneConstraintChannelHelper.cpp
--- a/Engine/Source/Editor/MovieSceneTools/Private/Constraints/MovieSceneConstraintChannelHelper.cpp
+++ b/Engine/Source/Editor/MovieSceneTools/Private/Constraints/MovieSceneConstraintChannelHelper.cpp
@@ -347,58 +347,63 @@ void FCompensationEvaluator::ComputeLocalTransformsForBaking(UWorld* InWorld, co
 	
 	FMovieSceneInverseSequenceTransform LocalToRootTransform = Sequencer->GetFocusedMovieSceneSequenceTransform().Inverse();
 
-	for (int32 Index = 0; Index < NumFrames; ++Index)
 	{
-		TOptional<FFrameTime> RootTime = LocalToRootTransform.TryTransformTime(Frames[Index]);
-		if (!RootTime)
+		const float DisplayRate = static_cast<float>(MovieScene->GetDisplayRate().AsDecimal());
+		const UE::Anim::FEvaluationForCachingScope CachingScope(DisplayRate > 0.f ? 1.f / DisplayRate : 1.f / 30.f);
+	
+		for (int32 Index = 0; Index < NumFrames; ++Index)
 		{
-			continue;
-		}
+			TOptional<FFrameTime> RootTime = LocalToRootTransform.TryTransformTime(Frames[Index]);
+			if (!RootTime)
+			{
+				continue;
+			}
 
-		FFrameNumber FrameNumber = RootTime->GetFrame();
+			FFrameNumber FrameNumber = RootTime->GetFrame();
 
-		// evaluate animation
-		const FMovieSceneEvaluationRange EvaluationRange = FMovieSceneEvaluationRange(FFrameTime(FrameNumber), TickResolution);
-		const FMovieSceneContext Context = FMovieSceneContext(EvaluationRange, PlaybackStatus).SetHasJumped(true);
+			// evaluate animation
+			const FMovieSceneEvaluationRange EvaluationRange = FMovieSceneEvaluationRange(FFrameTime(FrameNumber), TickResolution);
+			const FMovieSceneContext Context = FMovieSceneContext(EvaluationRange, PlaybackStatus).SetHasJumped(true);
 
-		BakingWrapper.PreEvaluate(FrameNumber);
+			BakingWrapper.PreEvaluate(FrameNumber);
 
-		Sequencer->GetEvaluationTemplate().EvaluateSynchronousBlocking(Context);
+			Sequencer->GetEvaluationTemplate().EvaluateSynchronousBlocking(Context);
 
-		// evaluate constraints
-		for (const TWeakObjectPtr<UTickableConstraint>& InConstraint : AllConstraints)
-		{
-			if (InConstraint.IsValid())
+			// evaluate constraints
+			for (const TWeakObjectPtr<UTickableConstraint>& InConstraint : AllConstraints)
 			{
-				InConstraint->Evaluate(true);
+				if (InConstraint.IsValid())
+				{
+					InConstraint->Evaluate(true);
+				}
 			}
-		}
 
-		BakingWrapper.PostEvaluate(FrameNumber);
+			BakingWrapper.PostEvaluate(FrameNumber);
 
-		FTransform& ChildLocal = ChildLocals[Index];
-		FTransform& ChildGlobal = ChildGlobals[Index];
-		FTransform& SpaceGlobal = SpaceGlobals[Index];
+			FTransform& ChildLocal = ChildLocals[Index];
+			FTransform& ChildGlobal = ChildGlobals[Index];
+			FTransform& SpaceGlobal = SpaceGlobals[Index];
 
-		// store child transforms        	
-		ChildLocal = Handle->GetLocalTransform();
-		ChildGlobal = Handle->GetGlobalTransform();
+			// store child transforms
+			ChildLocal = Handle->GetLocalTransform();
+			ChildGlobal = Handle->GetGlobalTransform();
 
-		// store constraint/parent space global transform
-		if (const UTickableTransformConstraint* LastConstraint = GetLastActiveConstraint())
-		{
-			SpaceGlobal = LastConstraint->GetParentGlobalTransform();
-			TOptional<FTransform> Relative =
-				UE::TransformConstraintUtil::GetConstraintsRelativeTransform(ConstraintsMinusThis, ChildLocal, ChildGlobal);
-			if (Relative)
+			// store constraint/parent space global transform
+			if (const UTickableTransformConstraint* LastConstraint = GetLastActiveConstraint())
 			{
-				ChildLocal = *Relative;
+				SpaceGlobal = LastConstraint->GetParentGlobalTransform();
+				TOptional<FTransform> Relative =
+					UE::TransformConstraintUtil::GetConstraintsRelativeTransform(ConstraintsMinusThis, ChildLocal, ChildGlobal);
+				if (Relative)
+				{
+					ChildLocal = *Relative;
+				}
 			}
 		}
-	}
-
-	BakingWrapper.Stop();
 
+		BakingWrapper.Stop();
+	}
+	
 	const bool bIsValidAfterBaking = IsValid(Handle) && Handle->IsValid();
 	if (!bIsValidAfterBaking)
 	{
diff --git a/Engine/Source/Editor/MovieSceneTools/Private/MovieSceneToolHelpers.cpp b/Engine/Source/Editor/MovieSceneTools/Private/MovieSceneToolHelpers.cpp
--- a/Engine/Source/Editor/MovieSceneTools/Private/MovieSceneToolHelpers.cpp
+++ b/Engine/Source/Editor/MovieSceneTools/Private/MovieSceneToolHelpers.cpp
@@ -98,6 +98,7 @@
 #include "GameFramework/Character.h"
 #include "Systems/MovieSceneTransformOriginSystem.h"
 #include "Tracks/IMovieSceneTransformOrigin.h"
+#include "Transform/AnimationEvaluation.h"
 
 /* FSkelMeshRecorder
  ***********/
@@ -4148,6 +4149,7 @@ static void TickFrame(const FFrameNumber& FrameNumber,float DeltaTime, UMovieSce
 	UWorld* World = GCurrentLevelEditingViewportClient ? GCurrentLevelEditingViewportClient->GetWorld() : nullptr;
 	if (World)
 	{
+		const UE::Anim::FEvaluationForCachingScope CachingScope(DeltaTime);
 		const FConstraintsManagerController& Controller = FConstraintsManagerController::Get(World);
 		Controller.EvaluateAllConstraints();
 	}
@@ -5082,7 +5084,6 @@ static void GetSequencerActorWorldTransforms(IMovieScenePlayer* Player, FMovieSc
 			OutTransforms.SetNum(Frames.Num());
 
 			FFrameRate TickResolution = MovieScene->GetTickResolution();
-			FFrameRate DisplayRate = MovieScene->GetDisplayRate();
 			const TArray<IMovieSceneToolsAnimationBakeHelper*>& BakeHelpers = FMovieSceneToolsModule::Get().GetAnimationBakeHelpers();
 
 			for (IMovieSceneToolsAnimationBakeHelper* BakeHelper : BakeHelpers)
@@ -5093,6 +5094,7 @@ static void GetSequencerActorWorldTransforms(IMovieScenePlayer* Player, FMovieSc
 				}
 			}
 
+			const float DisplayRate = static_cast<float>(MovieScene->GetDisplayRate().AsDecimal());
 			for (int32 Index = 0; Index < Frames.Num(); ++Index)
 			{
 				FFrameNumber FrameNumber = Frames[Index];
@@ -5114,8 +5116,11 @@ static void GetSequencerActorWorldTransforms(IMovieScenePlayer* Player, FMovieSc
 				Player->GetEvaluationTemplate().EvaluateSynchronousBlocking(Context);
 
 
-				const FConstraintsManagerController& Controller = FConstraintsManagerController::Get(Actor->GetWorld());
-				Controller.EvaluateAllConstraints();
+				{
+					const UE::Anim::FEvaluationForCachingScope CachingScope(DisplayRate > 0.f ? 1.f / DisplayRate : 1.f / 30.f);
+					const FConstraintsManagerController& Controller = FConstraintsManagerController::Get(Actor->GetWorld());
+					Controller.EvaluateAllConstraints();
+				}
 					
 				for (IMovieSceneToolsAnimationBakeHelper* BakeHelper : BakeHelpers)
 				{
diff --git a/Engine/Source/Runtime/Experimental/Animation/Constraints/Private/Transform/AnimationEvaluation.cpp b/Engine/Source/Runtime/Experimental/Animation/Constraints/Private/Transform/AnimationEvaluation.cpp
--- a/Engine/Source/Runtime/Experimental/Animation/Constraints/Private/Transform/AnimationEvaluation.cpp
+++ b/Engine/Source/Runtime/Experimental/Animation/Constraints/Private/Transform/AnimationEvaluation.cpp
@@ -48,6 +48,35 @@ FAnimationEvaluator::FAnimationEvaluator(USkeletalMeshComponent* InSkeletalMeshC
 	}
 }
 
+static float CachingDelta = 0.f;
+static int32 CachingCount = 0;
+
+bool IsCaching()
+{
+	return CachingCount > 0;
+}
+	
+FEvaluationForCachingScope::FEvaluationForCachingScope(const float InDeltaTime)
+{
+	if (ensure(CachingCount == 0))
+	{
+		CachingDelta = InDeltaTime;
+	}
+	++CachingCount;
+}
+
+FEvaluationForCachingScope::~FEvaluationForCachingScope()
+{
+	ensure(CachingCount > 0);
+	
+	--CachingCount;
+	if (CachingCount <= 0)
+	{
+		CachingDelta = 0.f;
+		CachingCount = 0;
+	}
+}
+	
 FAnimationEvaluator::~FAnimationEvaluator()
 {
 	if (OnBoneTransformsFinalizedHandle.IsValid())
@@ -276,12 +305,12 @@ void FAnimationEvaluator::EvaluateAnimation()
 			bShouldTickAnimation |= !Context.PostProcessAnimInstance->GetUpdateCounter().HasEverBeenUpdated();
 		}
 
-		if (bShouldTickAnimation)
+		if (bShouldTickAnimation || IsCaching())
 		{
 			// We bypass TickPose() and call TickAnimation directly, so URO doesn't intercept us.
 			static constexpr float DeltaTime = 0.0f;
 			static constexpr bool bNeedsValidRootMotion = false;
-			RawSkeletalMeshComponent->TickAnimation(DeltaTime, bNeedsValidRootMotion);
+			RawSkeletalMeshComponent->TickAnimation(IsCaching() ? CachingDelta : DeltaTime, bNeedsValidRootMotion);
 		}
 	}
 
diff --git a/Engine/Source/Runtime/Experimental/Animation/Constraints/Public/Transform/AnimationEvaluation.h b/Engine/Source/Runtime/Experimental/Animation/Constraints/Public/Transform/AnimationEvaluation.h
--- a/Engine/Source/Runtime/Experimental/Animation/Constraints/Public/Transform/AnimationEvaluation.h
+++ b/Engine/Source/Runtime/Experimental/Animation/Constraints/Public/Transform/AnimationEvaluation.h
@@ -123,4 +123,17 @@ private:
 	FDelegateHandle ConstraintsNotificationHandle;
 };
 
+/**
+ * FEvaluationForCachingScope provides a way to set the delta time when updating the animations.
+ * This class is currently not thread safe and won't update the delta time on re-entrant calls.
+*/
+
+struct FEvaluationForCachingScope : public FNoncopyable
+{
+	CONSTRAINTS_API FEvaluationForCachingScope(const float InDeltaTime = 1.f / 30.f);
+	CONSTRAINTS_API ~FEvaluationForCachingScope();
+};
+
+CONSTRAINTS_API bool IsCaching();
+	
 }
\ No newline at end of file
