Sequencer: StableTextId calls GetObjectsWithOuter which now requires a package

This fixes a crash in event structs that have a Text property.

diff --git a/Engine/Source/Editor/MovieSceneTools/Private/MovieSceneDirectorBlueprintEndpointCustomization.cpp b/Engine/Source/Editor/MovieSceneTools/Private/MovieSceneDirectorBlueprintEndpointCustomization.cpp
--- a/Engine/Source/Editor/MovieSceneTools/Private/MovieSceneDirectorBlueprintEndpointCustomization.cpp
+++ b/Engine/Source/Editor/MovieSceneTools/Private/MovieSceneDirectorBlueprintEndpointCustomization.cpp
@@ -289,6 +289,8 @@ void FMovieSceneDirectorBlueprintEndpointCustomization::CustomizeChildren(TShare
 
 			TSharedPtr<FStructOnScope> StructData = MakeShared<FStructOnScope>(CommonFunction);
 
+			StructData->SetPackage(Blueprint->GetPackage());
+
 			for (FProperty* Field : TFieldRange<FProperty>(CommonFunction))
 			{
 				// Ignore parameters we can't use as inputs.
