Sequencer: Fix crash when keying non-object properties in ObjectPropertyTrackEditor

Adds a check to make sure the property is actually an FObjectPropertyBase before
calling GetPropertyValue<UObject*>(). This prevents a crash when a non-object property is incorrectly routed to this track editor, as the type size
assertion in GetPropertyValue would fail.

[RN] fix animation/sequencer

diff --git a/Engine/Source/Editor/MovieSceneTools/Private/TrackEditors/ObjectPropertyTrackEditor.cpp b/Engine/Source/Editor/MovieSceneTools/Private/TrackEditors/ObjectPropertyTrackEditor.cpp
--- a/Engine/Source/Editor/MovieSceneTools/Private/TrackEditors/ObjectPropertyTrackEditor.cpp
+++ b/Engine/Source/Editor/MovieSceneTools/Private/TrackEditors/ObjectPropertyTrackEditor.cpp
@@ -10,6 +10,7 @@
 #include "KeyPropertyParams.h"
 #include "Math/Color.h"
 #include "Modules/ModuleManager.h"
+#include "MovieSceneFwd.h"
 #include "PropertyPath.h"
 #include "Tracks/MovieSceneObjectPropertyTrack.h"
 #include "UObject/Class.h"
@@ -42,6 +43,16 @@ FText FObjectPropertyTrackEditor::GetDisplayName() const
 
 void FObjectPropertyTrackEditor::GenerateKeysFromPropertyChanged(const FPropertyChangedParams& PropertyChangedParams, UMovieSceneSection* SectionToKey, FGeneratedTrackKeys& OutGeneratedKeys)
 {
+	// Verify the property is actually an object property before attempting to get the value.
+	// This prevents a crash when a non-object property (e.g., FInterfaceProperty) is incorrectly
+	// routed to this track editor, as the type size assertion in GetPropertyValue would fail.
+	const FObjectPropertyBase* ObjectProperty = CastField<FObjectPropertyBase>(PropertyChangedParams.PropertyPath.GetLeafMostProperty().Property.Get());
+	if (!ObjectProperty)
+	{
+		UE_LOG(LogMovieScene, Warning, TEXT("FObjectPropertyTrackEditor::GenerateKeysFromPropertyChanged called with non-object property type"));
+		return;
+	}
+
 	UObject* Existing = PropertyChangedParams.GetPropertyValue<UObject*>();
 	OutGeneratedKeys.Add(FMovieSceneChannelValueSetter::Create<FMovieSceneObjectPathChannel>(0, Existing, true));
 }
