- Fix using Remove vs RemoveCurrent while iterating over a map

#rnx

diff --git a/Engine/Plugins/Compositing/OpenColorIO/Source/OpenColorIO/Private/OpenColorIOShaderCompilationManager.cpp b/Engine/Plugins/Compositing/OpenColorIO/Source/OpenColorIO/Private/OpenColorIOShaderCompilationManager.cpp
--- a/Engine/Plugins/Compositing/OpenColorIO/Source/OpenColorIO/Private/OpenColorIOShaderCompilationManager.cpp
+++ b/Engine/Plugins/Compositing/OpenColorIO/Source/OpenColorIO/Private/OpenColorIOShaderCompilationManager.cpp
@@ -128,7 +128,7 @@ void FOpenColorIOShaderCompilationManager::ProcessAsyncResults()
 		if (Results.FinishedJobs.Num() == Results.NumJobsQueued)
 		{
 			PendingFinalizeOpenColorIOShaderMaps.Add(It.Key(), FOpenColorIOShaderMapCompileResults(Results));
-			OpenColorIOShaderMapJobs.Remove(It.Key());
+			It.RemoveCurrent();
 		}
 	}
 
diff --git a/Engine/Plugins/Experimental/NNERuntimeIREE/Source/NNERuntimeIREEShader/Private/NNERuntimeIREEShaderCompilationManager.cpp b/Engine/Plugins/Experimental/NNERuntimeIREE/Source/NNERuntimeIREEShader/Private/NNERuntimeIREEShaderCompilationManager.cpp
--- a/Engine/Plugins/Experimental/NNERuntimeIREE/Source/NNERuntimeIREEShader/Private/NNERuntimeIREEShaderCompilationManager.cpp
+++ b/Engine/Plugins/Experimental/NNERuntimeIREE/Source/NNERuntimeIREEShader/Private/NNERuntimeIREEShaderCompilationManager.cpp
@@ -96,7 +96,7 @@ void FNNERuntimeIREEShaderCompilationManager::ProcessAsyncResults()
 		if (Results.FinishedJobs.Num() == Results.NumJobsQueued)
 		{
 			PendingFinalizeNNERuntimeIREEShaderMaps.Add(It.Key(), FNNERuntimeIREEShaderMapCompileResults(Results));
-			NNERuntimeIREEShaderMapJobs.Remove(It.Key());
+			It.RemoveCurrent();
 		}
 	}
 
diff --git a/Engine/Plugins/Experimental/UAF/UAFAnimGraph/Source/UAFAnimGraph/Private/TraitCore/TraitRegistry.cpp b/Engine/Plugins/Experimental/UAF/UAFAnimGraph/Source/UAFAnimGraph/Private/TraitCore/TraitRegistry.cpp
--- a/Engine/Plugins/Experimental/UAF/UAFAnimGraph/Source/UAFAnimGraph/Private/TraitCore/TraitRegistry.cpp
+++ b/Engine/Plugins/Experimental/UAF/UAFAnimGraph/Source/UAFAnimGraph/Private/TraitCore/TraitRegistry.cpp
@@ -222,7 +222,7 @@ namespace UE::UAF
 					}
 				}
 
-				TraitUIDToEntryMap.Remove(It.Key());
+				It.RemoveCurrent();
 
 				break;
 			}
diff --git a/Engine/Source/Editor/MovieSceneTools/Private/MovieSceneToolHelpers.cpp b/Engine/Source/Editor/MovieSceneTools/Private/MovieSceneToolHelpers.cpp
--- a/Engine/Source/Editor/MovieSceneTools/Private/MovieSceneToolHelpers.cpp
+++ b/Engine/Source/Editor/MovieSceneTools/Private/MovieSceneToolHelpers.cpp
@@ -3871,13 +3871,13 @@ bool MovieSceneToolHelpers::ImportFBXIfReady(UWorld* World, UMovieSceneSequence*
 		}
 
 		bool bFoundMatch = false;
-		for (auto It = ObjectBindingMap.CreateConstIterator(); It; ++It)
+		for (auto It = ObjectBindingMap.CreateIterator(); It; ++It)
 		{
 			if (FCString::Strcmp(*It.Value().ToUpper(), *NodeName.ToUpper()) == 0)
 			{
 				MovieSceneToolHelpers::ImportFBXNode(NodeName, CurveAPI, Sequence, Player, TemplateID, It.Key(), Sequencer);
 
-				ObjectBindingMap.Remove(It.Key());
+				It.RemoveCurrent();
 				AllNodeNames.RemoveAt(NodeIndex);
 
 				bFoundMatch = true;
diff --git a/Engine/Source/Runtime/Core/Private/Misc/RemoteConfigIni.cpp b/Engine/Source/Runtime/Core/Private/Misc/RemoteConfigIni.cpp
--- a/Engine/Source/Runtime/Core/Private/Misc/RemoteConfigIni.cpp
+++ b/Engine/Source/Runtime/Core/Private/Misc/RemoteConfigIni.cpp
@@ -207,7 +207,7 @@ bool FRemoteConfigAsyncTaskManager::AreAllTasksFinished(bool bDoRemoval)
 			if (AsyncTask->IsDone() && bDoRemoval)
 			{
 				// Remove completed tasks from the queue if so desired
-				PendingTasks.Remove(It.Key());
+				It.RemoveCurrent();
 			}
 			else
 			{
diff --git a/Engine/Source/Runtime/InputCore/Private/InputCoreTypes.cpp b/Engine/Source/Runtime/InputCore/Private/InputCoreTypes.cpp
--- a/Engine/Source/Runtime/InputCore/Private/InputCoreTypes.cpp
+++ b/Engine/Source/Runtime/InputCore/Private/InputCoreTypes.cpp
@@ -1032,12 +1032,12 @@ void EKeys::GetAllKeys(TArray<FKey>& OutKeys)
 
 void EKeys::RemoveKeysWithCategory(const FName InCategory)
 {
-	for (TMap<FKey, TSharedPtr<FKeyDetails> >::TConstIterator It(InputKeys); It; ++It)
+	for (TMap<FKey, TSharedPtr<FKeyDetails> >::TIterator It(InputKeys); It; ++It)
 	{
 		FName KeyCategory = It.Key().GetMenuCategory();
 		if (KeyCategory.IsEqual(InCategory))
 		{
-			InputKeys.Remove(It.Key());
+			It.RemoveCurrent();
 		}
 	}
 }
