Sequencer: Add support for creating subsequences of custom LevelSequence subclasses

#FailedValidations Preflight

diff --git a/Engine/Source/Editor/MovieSceneTools/Private/MovieSceneToolHelpers.cpp b/Engine/Source/Editor/MovieSceneTools/Private/MovieSceneToolHelpers.cpp
--- a/Engine/Source/Editor/MovieSceneTools/Private/MovieSceneToolHelpers.cpp
+++ b/Engine/Source/Editor/MovieSceneTools/Private/MovieSceneToolHelpers.cpp
@@ -720,11 +720,14 @@ FString MovieSceneToolHelpers::GenerateNewSubsequenceName(const TArray<UMovieSce
 	return ComposeShotName(SubsequencePrefix, ProjectSettings->FirstShotNumber, ProjectSettings->FirstTakeNumber, ProjectSettings->ShotNumDigits, ProjectSettings->TakeNumDigits);
 }
 
-UMovieSceneSequence* MovieSceneToolHelpers::CreateSequence(FString& NewSequenceName, FString& NewSequencePath, UMovieSceneSubSection* SectionToDuplicate)
+UMovieSceneSequence* MovieSceneToolHelpers::CreateSequence(FString& NewSequenceName, FString& NewSequencePath, UMovieSceneSubSection* SectionToDuplicate, TSubclassOf<UMovieSceneSequence> SequenceClass)
 {
 	// Create a new level sequence asset with the appropriate name
 	IAssetTools& AssetTools = FModuleManager::GetModuleChecked<FAssetToolsModule>("AssetTools").Get();
-
+	if (!SequenceClass)
+	{
+		SequenceClass = ULevelSequence::StaticClass();
+	}
 	UObject* NewAsset = nullptr;
 	for (TObjectIterator<UClass> It; It; ++It)
 	{
@@ -732,7 +735,7 @@ UMovieSceneSequence* MovieSceneToolHelpers::CreateSequence(FString& NewSequenceN
 		if (CurrentClass->IsChildOf(UFactory::StaticClass()) && !(CurrentClass->HasAnyClassFlags(CLASS_Abstract)))
 		{
 			UFactory* Factory = Cast<UFactory>(CurrentClass->GetDefaultObject());
-			if (Factory->CanCreateNew() && Factory->ImportPriority >= 0 && Factory->SupportedClass == ULevelSequence::StaticClass())
+			if (Factory->CanCreateNew() && Factory->ImportPriority >= 0 && Factory->SupportedClass == SequenceClass)
 			{
 				if (SectionToDuplicate != nullptr)
 				{
@@ -740,7 +743,7 @@ UMovieSceneSequence* MovieSceneToolHelpers::CreateSequence(FString& NewSequenceN
 				}
 				else
 				{
-					NewAsset = AssetTools.CreateAssetWithDialog(NewSequenceName, NewSequencePath, ULevelSequence::StaticClass(), Factory);
+					NewAsset = AssetTools.CreateAssetWithDialog(NewSequenceName, NewSequencePath, SequenceClass, Factory);
 				}
 				break;
 			}
diff --git a/Engine/Source/Editor/MovieSceneTools/Private/TrackEditors/SubTrackEditor.cpp b/Engine/Source/Editor/MovieSceneTools/Private/TrackEditors/SubTrackEditor.cpp
--- a/Engine/Source/Editor/MovieSceneTools/Private/TrackEditors/SubTrackEditor.cpp
+++ b/Engine/Source/Editor/MovieSceneTools/Private/TrackEditors/SubTrackEditor.cpp
@@ -960,6 +960,12 @@ void FSubTrackEditor::InsertSection(UMovieSceneTrack* Track)
 
 	UMovieSceneSubTrack* SubTrack = FindOrCreateSubTrack(MovieScene, Track);
 
+	TSubclassOf<UMovieSceneSequence> OwningSequenceClass = nullptr;
+	if (UMovieSceneSequence* OwningSequence = SubTrack->GetTypedOuter<UMovieSceneSequence>())
+	{
+		OwningSequenceClass = OwningSequence->GetClass();
+	}
+
 	FString NewSequenceName = MovieSceneToolHelpers::GenerateNewSubsequenceName(SubTrack->GetAllSections(), GetDefaultSubsequenceName(), NewSectionStartTime.FrameNumber);
 	FString NewSequencePath = MovieSceneToolHelpers::GenerateNewSubsequencePath(GetSequencer()->GetFocusedMovieSceneSequence()->GetMovieScene(), GetDefaultSubsequenceDirectory(), NewSequenceName);
 
diff --git a/Engine/Source/Editor/MovieSceneTools/Public/MovieSceneToolHelpers.h b/Engine/Source/Editor/MovieSceneTools/Public/MovieSceneToolHelpers.h
--- a/Engine/Source/Editor/MovieSceneTools/Public/MovieSceneToolHelpers.h
+++ b/Engine/Source/Editor/MovieSceneTools/Public/MovieSceneToolHelpers.h
@@ -214,9 +214,10 @@ public:
 	 * @param NewSequenceName The new sequence name.
 	 * @param NewSequencePath The new sequence path. 
 	 * @param SectionToDuplicate The section to duplicate.
+	 * @param SequenceClass The class of sequence. Used only if SectionToDuplicate is not provided.
 	 * @return The new subsequence.
 	 */
-	static UE_API UMovieSceneSequence* CreateSequence(FString& NewSequenceName, FString& NewSequencePath, UMovieSceneSubSection* SectionToDuplicate = nullptr);
+	static UE_API UMovieSceneSequence* CreateSequence(FString& NewSequenceName, FString& NewSequencePath, UMovieSceneSubSection* SectionToDuplicate = nullptr, TSubclassOf<UMovieSceneSequence> SequenceClass = nullptr);
 
 	/**
 	 * Gather takes - level sequence assets that have the same shot prefix and shot number in the same asset path (directory)
