Constraints: fix baking and deletion of infinite constraints
- update the default channel values when removing a constraint for deletion
- compensate keys out of the baking range to maintain animation

#rnx

diff --git a/Engine/Plugins/Animation/ControlRig/Source/ControlRig/Private/Constraints/ControlRigTransformableHandle.cpp b/Engine/Plugins/Animation/ControlRig/Source/ControlRig/Private/Constraints/ControlRigTransformableHandle.cpp
--- a/Engine/Plugins/Animation/ControlRig/Source/ControlRig/Private/Constraints/ControlRigTransformableHandle.cpp
+++ b/Engine/Plugins/Animation/ControlRig/Source/ControlRig/Private/Constraints/ControlRigTransformableHandle.cpp
@@ -862,15 +862,32 @@ bool UTransformableControlHandle::AddTransformKeys(const TArray<FFrameNumber>& I
 	UMovieSceneSection*,
 	const bool bLocal) const
 {
-	if (!ControlRig.IsValid() || ControlName == NAME_None || InFrames.IsEmpty() || InFrames.Num() != InTransforms.Num() || !ControlRig->GetHierarchy())
+	if (!ControlRig.IsValid() || ControlName == NAME_None || !ControlRig->GetHierarchy())
 	{
 		return false;
 	}
+
 	FRigControlElement* ControlElement = ControlRig.Get()->FindControl(ControlName);
 	if(!ControlElement)
 	{
 		return false;
 	}
+	
+	const int32 NumFrames = InFrames.Num();
+	const int32 NumTransforms = InTransforms.Num();
+	const bool bUpdateDefault = NumFrames == 0 && NumTransforms == 1;
+	if (NumFrames == 0)
+	{
+		if (!bUpdateDefault)
+		{
+			return false;
+		}
+	}
+	else if (NumTransforms != NumFrames)
+	{
+		return false;
+	}
+
 	auto KeyframeFunc = [this, bLocal](const FTransform& InTransform, const FRigControlModifiedContext& InKeyframeContext, FRigControlElement* InControlElement)
 	{
 		UControlRig* InControlRig = ControlRig.Get();
@@ -900,15 +917,24 @@ bool UTransformableControlHandle::AddTransformKeys(const TArray<FFrameNumber>& I
 	};
 
 	FRigControlModifiedContext KeyframeContext;
-	KeyframeContext.SetKey = EControlRigSetKey::Always;
 	KeyframeContext.KeyMask = static_cast<uint32>(InChannels);
-
-	for (int32 Index = 0; Index < InFrames.Num(); ++Index)
+	
+	if (bUpdateDefault)
+	{
+		KeyframeFunc(InTransforms[0], KeyframeContext, ControlElement);
+	}
+	else
 	{
-		const FFrameNumber& Frame = InFrames[Index];
-		KeyframeContext.LocalTime = InTickResolution.AsSeconds(FFrameTime(Frame));
+		KeyframeContext.SetKey = EControlRigSetKey::Always;
 
-		KeyframeFunc(InTransforms[Index], KeyframeContext, ControlElement);
+		for (int32 Index = 0; Index < InFrames.Num(); ++Index)
+		{
+			const FFrameNumber& Frame = InFrames[Index];
+			KeyframeContext.LocalTime = InTickResolution.AsSeconds(FFrameTime(Frame));
+
+			KeyframeFunc(InTransforms[Index], KeyframeContext, ControlElement);
+		}
+		
 	}
 
 	return true;
diff --git a/Engine/Plugins/Animation/ControlRig/Source/ControlRigEditor/Private/Tools/ConstraintBaker.cpp b/Engine/Plugins/Animation/ControlRig/Source/ControlRigEditor/Private/Tools/ConstraintBaker.cpp
--- a/Engine/Plugins/Animation/ControlRig/Source/ControlRigEditor/Private/Tools/ConstraintBaker.cpp
+++ b/Engine/Plugins/Animation/ControlRig/Source/ControlRigEditor/Private/Tools/ConstraintBaker.cpp
@@ -297,6 +297,49 @@ void FConstraintBaker::Bake(UWorld* InWorld,
 	{
 		return;
 	}
+
+	// get transform channels
+	const TArrayView<FMovieSceneFloatChannel*> FloatTransformChannels = Handle->GetFloatChannels(ConstraintSections.ChildTransformSection);
+	const TArrayView<FMovieSceneDoubleChannel*> DoubleTransformChannels = Handle->GetDoubleChannels(ConstraintSections.ChildTransformSection);
+
+	const FMovieSceneConstraintChannel& ConstraintChannel = ConstraintSections.ActiveChannel->ActiveChannel;
+	const bool bIsInfinite = ConstraintChannel.IsInfinite();
+
+	TArray<FFrameNumber> FramesToCompensate;
+	TArray<FTransform> TransformsToCompensate;
+	if (bIsInfinite)
+	{
+		if (FloatTransformChannels.Num() > 0)
+		{
+			FMovieSceneConstraintChannelHelper::GetFramesWithinActiveState(ConstraintChannel, FloatTransformChannels, FramesToCompensate);
+		}
+		else
+		{
+			FMovieSceneConstraintChannelHelper::GetFramesWithinActiveState(ConstraintChannel, DoubleTransformChannels, FramesToCompensate);
+		}
+
+		// remove frames to bake
+		if (!FramesToCompensate.IsEmpty())
+		{
+			const TSet<FFrameNumber> SetFramesToBake(FramesToBake);
+			for (int32 Index = FramesToCompensate.Num()-1; Index >= 0; --Index)
+			{
+				if (SetFramesToBake.Contains(FramesToCompensate[Index]))
+				{
+					FramesToCompensate.RemoveAt(Index);
+				}
+			}
+		}
+
+		// compute compensation
+		if (!FramesToCompensate.IsEmpty())
+		{
+			FCompensationEvaluator EvaluatorForCompensation(InConstraint);
+			const FCompensationEvaluator::FEvalParameters EvalParamsForCompensation(InSequencer, FramesToCompensate);
+			EvaluatorForCompensation.ComputeLocalTransformsBeforeDeletion(InWorld, EvalParamsForCompensation);
+			TransformsToCompensate = MoveTemp(EvaluatorForCompensation.ChildLocals);
+		}
+	}
 	
 	ConstraintSections.ConstraintSection->Modify();
 
@@ -305,10 +348,6 @@ void FConstraintBaker::Bake(UWorld* InWorld,
 	const TArrayView<const FFrameNumber> ConstraintFrames = ConstraintChannelData.GetTimes();
 	const TArray<bool> PreviousConstraintValues(ConstraintChannelData.GetValues());
 	
-	// get transform channels
-	const TArrayView<FMovieSceneFloatChannel*> FloatTransformChannels = Handle->GetFloatChannels(ConstraintSections.ChildTransformSection);
-	const TArrayView<FMovieSceneDoubleChannel*> DoubleTransformChannels = Handle->GetDoubleChannels(ConstraintSections.ChildTransformSection);
-
 	EMovieSceneKeyInterpolation KeyType = InSequencer->GetKeyInterpolation();
 
 	for (int32 Index = 0; Index < ConstraintFrames.Num(); ++Index)
@@ -373,6 +412,12 @@ void FConstraintBaker::Bake(UWorld* InWorld,
 		MovieSceneToolHelpers::OptimizeSection(Param, ConstraintSections.ConstraintSection);
 	}
 
+	// add compensation keys outside the baking range when infinite
+	if (bIsInfinite && !FramesToCompensate.IsEmpty() && FramesToCompensate.Num() == TransformsToCompensate.Num())
+	{
+		ConstraintSections.Interface->AddHandleTransformKeys(InSequencer, Handle, FramesToCompensate, TransformsToCompensate, Channels);
+	}
+
 	// remove useless constraint keys
 	CleanupConstraintKeys(*ConstraintSections.ActiveChannel);
 	
diff --git a/Engine/Source/Editor/MovieSceneTools/Private/Constraints/MovieSceneConstraintChannelHelper.cpp b/Engine/Source/Editor/MovieSceneTools/Private/Constraints/MovieSceneConstraintChannelHelper.cpp
--- a/Engine/Source/Editor/MovieSceneTools/Private/Constraints/MovieSceneConstraintChannelHelper.cpp
+++ b/Engine/Source/Editor/MovieSceneTools/Private/Constraints/MovieSceneConstraintChannelHelper.cpp
@@ -414,7 +414,14 @@ void FCompensationEvaluator::ComputeLocalTransformsForBaking(UWorld* InWorld, co
 }
 void FCompensationEvaluator::ComputeLocalTransformsBeforeDeletion(UWorld* InWorld, const FEvalParameters& InEvalParams)
 {
-	if (!InEvalParams.IsValid())
+	if (!InEvalParams.Sequencer)
+	{
+		return;
+	}
+
+	const TConstArrayView<FFrameNumber> Frames(InEvalParams.Frames);
+	const bool bFromCurrentFrame = Frames.IsEmpty() && InEvalParams.OptCurrentFrame.IsSet();
+	if (Frames.IsEmpty() && !bFromCurrentFrame)
 	{
 		return;
 	}
@@ -437,17 +444,37 @@ void FCompensationEvaluator::ComputeLocalTransformsBeforeDeletion(UWorld* InWorl
 		return LastActiveIndex > INDEX_NONE ? Cast<UTickableTransformConstraint>(ConstraintsMinusThis[LastActiveIndex]) : nullptr;
 	};
 
+	if (bFromCurrentFrame)
+	{
+		// store child transforms        	
+		FTransform& ChildLocal = ChildLocals.Emplace_GetRef(Handle->GetLocalTransform());
+		FTransform& ChildGlobal = ChildGlobals.Emplace_GetRef(Handle->GetGlobalTransform());
+		FTransform& SpaceGlobal = SpaceGlobals.Emplace_GetRef(FTransform::Identity);
+
+		// store constraint/parent space global transform
+		if (const UTickableTransformConstraint* LastConstraint = GetLastActiveConstraint())
+		{
+			SpaceGlobal = LastConstraint->GetParentGlobalTransform();
+			TOptional<FTransform> Relative =
+				UE::TransformConstraintUtil::GetConstraintsRelativeTransform(Constraints, ChildLocal, ChildGlobal);
+			if (Relative)
+			{
+				ChildLocal = *Relative;
+			}
+		}
+		return;
+	}
+
 	// get all constraints for evaluation
 	const FConstraintsManagerController& Controller = FConstraintsManagerController::Get(InWorld);
 	static constexpr bool bSorted = true;
 	const TArray<ConstraintPtr> AllConstraints = Controller.GetAllConstraints(bSorted);
-
+	
 	ISequencer* Sequencer = InEvalParams.Sequencer;
 	UMovieScene* MovieScene = Sequencer->GetFocusedMovieSceneSequence()->GetMovieScene();
 	const FFrameRate TickResolution = MovieScene->GetTickResolution();
 	const EMovieScenePlayerStatus::Type PlaybackStatus = Sequencer->GetPlaybackStatus();
 
-	const TConstArrayView<FFrameNumber> Frames(InEvalParams.Frames);
 	const int32 NumFrames = Frames.Num();
 
 	ChildLocals.SetNum(NumFrames);
@@ -786,17 +813,17 @@ void FMovieSceneConstraintChannelHelper::HandleConstraintRemoved(
 	
 	InSection->Modify();
 	TGuardValue<bool> CompensateGuard(bDoNotCompensate, true);
-	if (const UTransformableHandle* ControlHandle =Constraint->ChildTRSHandle)
+	if (const UTransformableHandle* Handle = Constraint->ChildTRSHandle)
 	{
 		const TArrayView<const FFrameNumber> Times = InConstraintChannel->GetData().GetTimes();
-		if (Times.IsEmpty())
+		if (Times.IsEmpty() && !InConstraintChannel->IsInfinite())
 		{
 			return;
 		}
 
 		// get transform channels
-		const TArrayView<FMovieSceneFloatChannel*> FloatTransformChannels = ControlHandle->GetFloatChannels(InSection);
-		const TArrayView<FMovieSceneDoubleChannel*> DoubleTransformChannels = ControlHandle->GetDoubleChannels(InSection);
+		const TArrayView<FMovieSceneFloatChannel*> FloatTransformChannels = Handle->GetFloatChannels(InSection);
+		const TArrayView<FMovieSceneDoubleChannel*> DoubleTransformChannels = Handle->GetDoubleChannels(InSection);
 
 		// get frames after this time
 		TArray<FFrameNumber> FramesToCompensate;
@@ -808,19 +835,28 @@ void FMovieSceneConstraintChannelHelper::HandleConstraintRemoved(
 		{
 			GetFramesWithinActiveState(*InConstraintChannel, DoubleTransformChannels, FramesToCompensate);
 		}
+		
+		const FFrameRate TickResolution = InSequencer->GetFocusedTickResolution();		
+		
 		// do the compensation
 		UWorld* World = GCurrentLevelEditingViewportClient ? GCurrentLevelEditingViewportClient->GetWorld() : nullptr;
 
 		FCompensationEvaluator Evaluator(Constraint);
-		const FCompensationEvaluator::FEvalParameters EvalParams(InSequencer, FramesToCompensate);
+		FCompensationEvaluator::FEvalParameters EvalParams(InSequencer, FramesToCompensate);
+
+		// update channels' default values if there are no frames to compensate
+		if (FramesToCompensate.IsEmpty())
+		{
+			const FFrameTime CurrentTime = InSequencer->GetLocalTime().ConvertTo(TickResolution);
+			EvalParams.OptCurrentFrame = CurrentTime.GetFrame();
+		}
+		
 		Evaluator.ComputeLocalTransformsBeforeDeletion(World, EvalParams);
 		const TArray<FTransform>& ChildLocals = Evaluator.ChildLocals;
 
 		const EMovieSceneTransformChannel ChannelsToKey = Constraint->GetChannelsToKey();
-		const FFrameRate TickResolution = InSequencer->GetFocusedTickResolution();
 
-		ControlHandle->AddTransformKeys(FramesToCompensate,
-			ChildLocals, ChannelsToKey, TickResolution, InSection);
+		Handle->AddTransformKeys(FramesToCompensate, ChildLocals, ChannelsToKey, TickResolution, InSection);
 
 		// clean double keys
 		if (IMovieSceneConstrainedSection* Section = Cast<IMovieSceneConstrainedSection>(InSection))
diff --git a/Engine/Source/Runtime/Experimental/Animation/Constraints/Private/Transform/TransformableHandle.cpp b/Engine/Source/Runtime/Experimental/Animation/Constraints/Private/Transform/TransformableHandle.cpp
--- a/Engine/Source/Runtime/Experimental/Animation/Constraints/Private/Transform/TransformableHandle.cpp
+++ b/Engine/Source/Runtime/Experimental/Animation/Constraints/Private/Transform/TransformableHandle.cpp
@@ -464,7 +464,18 @@ bool UTransformableComponentHandle::AddTransformKeys(const TArray<FFrameNumber>&
 		return false;
 	}
 
-	if (Frames.IsEmpty() || Frames.Num() != InLocalTransforms.Num())
+	const int32 NumFrames = Frames.Num();
+	const int32 NumTransforms = InLocalTransforms.Num();
+
+	const bool bUpdateDefault = NumFrames == 0 && NumTransforms == 1;
+	if (NumFrames == 0)
+	{
+		if (!bUpdateDefault)
+		{
+			return false;
+		}
+	}
+	else if (NumTransforms != NumFrames)
 	{
 		return false;
 	}
@@ -529,61 +540,65 @@ bool UTransformableComponentHandle::AddTransformKeys(const TArray<FFrameNumber>&
 	{
 		if (Channels.IsValidIndex(ChannelIndex))
 		{
-			if (!Channels[ChannelIndex]->GetDefault().IsSet())
+			FMovieSceneDoubleChannel* Channel = Channels[ChannelIndex];
+			if (!Channel->GetDefault().IsSet() || (bUpdateDefault && Channel->GetNumKeys() == 0))
 			{
 				const double Value = GetValue(ChannelIndex, Location0, Rotation0, Scale3D0);
-				Channels[ChannelIndex]->SetDefault(Value);
+				Channel->SetDefault(Value);
 			}
 		}
 	}
 
-	// add keys
-	for (int32 Index = 0; Index < Frames.Num(); ++Index)
+	if (NumFrames > 0)
 	{
-		const FFrameNumber& Frame = Frames[Index];
-		const FTransform& LocalTransform = InLocalTransforms[Index];
+		// add keys
+		for (int32 Index = 0; Index < Frames.Num(); ++Index)
+		{
+			const FFrameNumber& Frame = Frames[Index];
+			const FTransform& LocalTransform = InLocalTransforms[Index];
 
-		const FVector Location = LocalTransform.GetLocation();
-		const FRotator Rotation = LocalTransform.GetRotation().Rotator();
-		const FVector Scale3D = LocalTransform.GetScale3D();
+			const FVector Location = LocalTransform.GetLocation();
+			const FRotator Rotation = LocalTransform.GetRotation().Rotator();
+			const FVector Scale3D = LocalTransform.GetScale3D();
 
-		for (const int32 ChannelIndex : ChannelsIndexToKey)
-		{
-			if (Channels.IsValidIndex(ChannelIndex))
+			for (const int32 ChannelIndex : ChannelsIndexToKey)
 			{
-				const double Value = GetValue(ChannelIndex, Location, Rotation, Scale3D);
-				TMovieSceneChannelData<FMovieSceneDoubleValue> ChannelData = Channels[ChannelIndex]->GetData();
+				if (Channels.IsValidIndex(ChannelIndex))
 				{
-					int32 ExistingIndex = ChannelData.FindKey(Frame);
-					if (ExistingIndex != INDEX_NONE)
+					const double Value = GetValue(ChannelIndex, Location, Rotation, Scale3D);
+					TMovieSceneChannelData<FMovieSceneDoubleValue> ChannelData = Channels[ChannelIndex]->GetData();
 					{
-						FMovieSceneDoubleValue& DoubleValue = ChannelData.GetValues()[ExistingIndex]; //-V758
-						DoubleValue.Value = Value;
-					}
-					else
-					{
-						FMovieSceneDoubleValue NewKey(Value);
-						ERichCurveTangentWeightMode WeightedMode = RCTWM_WeightedNone;
-						NewKey.InterpMode = ERichCurveInterpMode::RCIM_Cubic;
-						NewKey.TangentMode = ERichCurveTangentMode::RCTM_Auto;
-						NewKey.Tangent.ArriveTangent = 0.0f;
-						NewKey.Tangent.LeaveTangent = 0.0f;
-						NewKey.Tangent.TangentWeightMode = WeightedMode;
-						NewKey.Tangent.ArriveTangentWeight = 0.0f;
-						NewKey.Tangent.LeaveTangentWeight = 0.0f;
-						ChannelData.AddKey(Frame, NewKey);
+						int32 ExistingIndex = ChannelData.FindKey(Frame);
+						if (ExistingIndex != INDEX_NONE)
+						{
+							FMovieSceneDoubleValue& DoubleValue = ChannelData.GetValues()[ExistingIndex]; //-V758
+							DoubleValue.Value = Value;
+						}
+						else
+						{
+							FMovieSceneDoubleValue NewKey(Value);
+							ERichCurveTangentWeightMode WeightedMode = RCTWM_WeightedNone;
+							NewKey.InterpMode = ERichCurveInterpMode::RCIM_Cubic;
+							NewKey.TangentMode = ERichCurveTangentMode::RCTM_Auto;
+							NewKey.Tangent.ArriveTangent = 0.0f;
+							NewKey.Tangent.LeaveTangent = 0.0f;
+							NewKey.Tangent.TangentWeightMode = WeightedMode;
+							NewKey.Tangent.ArriveTangentWeight = 0.0f;
+							NewKey.Tangent.LeaveTangentWeight = 0.0f;
+							ChannelData.AddKey(Frame, NewKey);
+						}
 					}
 				}
 			}
 		}
-	}
 
-	//now we need to set auto tangents
-	for (const int32 ChannelIndex : ChannelsIndexToKey)
-	{
-		if (Channels.IsValidIndex(ChannelIndex))
+		//now we need to set auto tangents
+		for (const int32 ChannelIndex : ChannelsIndexToKey)
 		{
-			Channels[ChannelIndex]->AutoSetTangents();
+			if (Channels.IsValidIndex(ChannelIndex))
+			{
+				Channels[ChannelIndex]->AutoSetTangents();
+			}
 		}
 	}
 	return true;
