[TRS Gizmos] Implement `BeginTransform()` and `EndTransform()` for `FSkeletalAnimationTrackEditMode` that mirror `StartTracking()` and `EndTracking()`.

diff --git a/Engine/Source/Editor/MovieSceneTools/Private/EditModes/SkeletalAnimationTrackEditMode.cpp b/Engine/Source/Editor/MovieSceneTools/Private/EditModes/SkeletalAnimationTrackEditMode.cpp
--- a/Engine/Source/Editor/MovieSceneTools/Private/EditModes/SkeletalAnimationTrackEditMode.cpp
+++ b/Engine/Source/Editor/MovieSceneTools/Private/EditModes/SkeletalAnimationTrackEditMode.cpp
@@ -462,6 +462,38 @@ bool FSkeletalAnimationTrackEditMode::StartTracking(FEditorViewportClient* InVie
 	return false;
 }
 
+bool FSkeletalAnimationTrackEditMode::BeginTransform(const FGizmoState& InState)
+{
+	if (!bIsTransacting)
+	{
+		bIsTransacting = (IsSomethingSelected());
+
+		bManipulatorMadeChange = false;
+
+		return bIsTransacting;
+	}
+	return false;
+}
+
+bool FSkeletalAnimationTrackEditMode::EndTransform(const FGizmoState& InState)
+{
+	if (bIsTransacting)
+	{
+		if (bManipulatorMadeChange)
+		{
+			bManipulatorMadeChange = false;
+			GEditor->EndTransaction();
+
+		}
+		bIsTransacting = false;
+		return true;
+	}
+
+	bManipulatorMadeChange = false;
+
+	return false;
+}
+
 bool FSkeletalAnimationTrackEditMode::UsesTransformWidget() const
 {
 	return IsSomethingSelected();
diff --git a/Engine/Source/Editor/MovieSceneTools/Private/EditModes/SkeletalAnimationTrackEditMode.h b/Engine/Source/Editor/MovieSceneTools/Private/EditModes/SkeletalAnimationTrackEditMode.h
--- a/Engine/Source/Editor/MovieSceneTools/Private/EditModes/SkeletalAnimationTrackEditMode.h
+++ b/Engine/Source/Editor/MovieSceneTools/Private/EditModes/SkeletalAnimationTrackEditMode.h
@@ -40,6 +40,8 @@ public:
 	virtual bool InputKey(FEditorViewportClient* InViewportClient, FViewport* InViewport, FKey InKey, EInputEvent InEvent) override;
 	virtual bool EndTracking(FEditorViewportClient* InViewportClient, FViewport* InViewport) override;
 	virtual bool StartTracking(FEditorViewportClient* InViewportClient, FViewport* InViewport) override;
+	virtual bool BeginTransform(const FGizmoState& InState) override;
+	virtual bool EndTransform(const FGizmoState& InState) override; 
 	virtual bool HandleClick(FEditorViewportClient* InViewportClient, HHitProxy *HitProxy, const FViewportClick &Click) override;
 	virtual void SelectNone() override;
 	virtual bool InputDelta(FEditorViewportClient* InViewportClient, FViewport* InViewport, FVector& InDrag, FRotator& InRot, FVector& InScale) override;
