Anim In Engine: Add option to not set RetargetSourceAsset when saving anim sequences, keep state on linked anim sequences, make sure if no retarget asset is set to not set it, and finally add option to take recorder to set this with it ON, since originally this came from them(UE-140883)

diff --git a/Engine/Plugins/Animation/ControlRig/Source/ControlRigEditor/Private/ControlRigEditorModule.cpp b/Engine/Plugins/Animation/ControlRig/Source/ControlRigEditor/Private/ControlRigEditorModule.cpp
--- a/Engine/Plugins/Animation/ControlRig/Source/ControlRigEditor/Private/ControlRigEditorModule.cpp
+++ b/Engine/Plugins/Animation/ControlRig/Source/ControlRigEditor/Private/ControlRigEditorModule.cpp
@@ -1259,6 +1259,7 @@ void FControlRigEditorModule::BakeToControlRig(UClass* ControlRigClass, UAnimSeq
 							FLevelSequenceAnimSequenceLinkItem LevelAnimLinkItem;
 							LevelAnimLinkItem.SkelTrackGuid = ActorTrackGuid;
 							LevelAnimLinkItem.PathToAnimSequence = FSoftObjectPath(AnimSequence);
+							LevelAnimLinkItem.bSetRetargetSourceAsset = false;
 							LevelAnimLinkItem.bExportMorphTargets = true;
 							LevelAnimLinkItem.bExportAttributeCurves = true;
 							LevelAnimLinkItem.Interpolation = EAnimInterpolationType::Linear;
diff --git a/Engine/Plugins/MovieScene/SequencerScripting/Source/SequencerScriptingEditor/Private/SequencerTools.cpp b/Engine/Plugins/MovieScene/SequencerScripting/Source/SequencerScriptingEditor/Private/SequencerTools.cpp
--- a/Engine/Plugins/MovieScene/SequencerScripting/Source/SequencerScriptingEditor/Private/SequencerTools.cpp
+++ b/Engine/Plugins/MovieScene/SequencerScripting/Source/SequencerScriptingEditor/Private/SequencerTools.cpp
@@ -552,6 +552,7 @@ bool USequencerToolsFunctionLibrary::LinkAnimSequence(ULevelSequence*  Sequence,
 							}
 						}
 						LevelAnimLinkItem.PathToAnimSequence = FSoftObjectPath(AnimSequence);
+						LevelAnimLinkItem.bSetRetargetSourceAsset = ExportOptions->bSetRetargetSourceAsset;
 						LevelAnimLinkItem.bExportMorphTargets = ExportOptions->bExportMorphTargets;
 						LevelAnimLinkItem.bExportAttributeCurves = ExportOptions->bExportAttributeCurves;
 						LevelAnimLinkItem.bExportMaterialCurves = ExportOptions->bExportMaterialCurves;
@@ -586,6 +587,7 @@ bool USequencerToolsFunctionLibrary::LinkAnimSequence(ULevelSequence*  Sequence,
 				FLevelSequenceAnimSequenceLinkItem LevelAnimLinkItem;
 				LevelAnimLinkItem.SkelTrackGuid = Binding.BindingID;
 				LevelAnimLinkItem.PathToAnimSequence = FSoftObjectPath(AnimSequence);
+				LevelAnimLinkItem.bSetRetargetSourceAsset = ExportOptions->bSetRetargetSourceAsset;
 				LevelAnimLinkItem.bExportMorphTargets = ExportOptions->bExportMorphTargets;
 				LevelAnimLinkItem.bExportAttributeCurves = ExportOptions->bExportAttributeCurves;
 				LevelAnimLinkItem.Interpolation = ExportOptions->Interpolation;
diff --git a/Engine/Plugins/VirtualProduction/Takes/Source/TakeTrackRecorders/Private/MovieSceneAnimationTrackRecorder.cpp b/Engine/Plugins/VirtualProduction/Takes/Source/TakeTrackRecorders/Private/MovieSceneAnimationTrackRecorder.cpp
--- a/Engine/Plugins/VirtualProduction/Takes/Source/TakeTrackRecorders/Private/MovieSceneAnimationTrackRecorder.cpp
+++ b/Engine/Plugins/VirtualProduction/Takes/Source/TakeTrackRecorders/Private/MovieSceneAnimationTrackRecorder.cpp
@@ -305,6 +305,7 @@ void UMovieSceneAnimationTrackRecorder::RecordSampleImpl(const FQualifiedFrameTi
 		RecordingSettings.SampleFrameRate = SampleRate;
 		RecordingSettings.InterpMode = AnimSettings->InterpMode;
 		RecordingSettings.TangentMode = AnimSettings->TangentMode;
+		RecordingSettings.bSetRetargetSourceAsset = AnimSettings->bSetRetargetSourceAsset;
 		RecordingSettings.Length = 0;
 		RecordingSettings.bRecordInWorldSpace = bRecordInWorldSpace;
 		RecordingSettings.bRemoveRootAnimation = bRootWasRemoved;
diff --git a/Engine/Plugins/VirtualProduction/Takes/Source/TakeTrackRecorders/Private/MovieSceneAnimationTrackRecorderSettings.cpp b/Engine/Plugins/VirtualProduction/Takes/Source/TakeTrackRecorders/Private/MovieSceneAnimationTrackRecorderSettings.cpp
--- a/Engine/Plugins/VirtualProduction/Takes/Source/TakeTrackRecorders/Private/MovieSceneAnimationTrackRecorderSettings.cpp
+++ b/Engine/Plugins/VirtualProduction/Takes/Source/TakeTrackRecorders/Private/MovieSceneAnimationTrackRecorderSettings.cpp
@@ -16,6 +16,7 @@ UMovieSceneAnimationTrackRecorderEditorSettings::UMovieSceneAnimationTrackRecord
 	, InterpMode()
 	, TangentMode()
 	, bRemoveRootAnimation(true)
+	, bSetRetargetSourceAsset(true)
 {
 	TimecodeBoneMethod.BoneMode = ETimecodeBoneMode::Root;
 }
diff --git a/Engine/Plugins/VirtualProduction/Takes/Source/TakeTrackRecorders/Public/TrackRecorders/MovieSceneAnimationTrackRecorderSettings.h b/Engine/Plugins/VirtualProduction/Takes/Source/TakeTrackRecorders/Public/TrackRecorders/MovieSceneAnimationTrackRecorderSettings.h
--- a/Engine/Plugins/VirtualProduction/Takes/Source/TakeTrackRecorders/Public/TrackRecorders/MovieSceneAnimationTrackRecorderSettings.h
+++ b/Engine/Plugins/VirtualProduction/Takes/Source/TakeTrackRecorders/Public/TrackRecorders/MovieSceneAnimationTrackRecorderSettings.h
@@ -66,6 +66,10 @@ public:
 	UPROPERTY(config, EditAnywhere, BlueprintReadWrite, Category = "Animation Recorder Settings")
 	bool bRemoveRootAnimation;
 
+	/** If true we set the retarget source asset on the saved anim sequence, if not we do not*/
+	UPROPERTY(config, EditAnywhere, BlueprintReadWrite, Category = "Animation Recorder Settings")
+	bool bSetRetargetSourceAsset;
+
 	/** The method to record timecode values onto bones */
 	UPROPERTY(config, EditAnywhere, BlueprintReadWrite, Category = "Animation Recorder Settings", meta = (ShowOnlyInnerProperties))
 	FTimecodeBoneMethod TimecodeBoneMethod;
diff --git a/Engine/Source/Editor/MovieSceneTools/Private/MovieSceneToolHelpers.cpp b/Engine/Source/Editor/MovieSceneTools/Private/MovieSceneToolHelpers.cpp
--- a/Engine/Source/Editor/MovieSceneTools/Private/MovieSceneToolHelpers.cpp
+++ b/Engine/Source/Editor/MovieSceneTools/Private/MovieSceneToolHelpers.cpp
@@ -4455,6 +4455,7 @@ bool MovieSceneToolHelpers::ExportToAnimSequence(UAnimSequence* AnimSequence, UA
 			RecordingSettings.Length = 0;
 			RecordingSettings.bRemoveRootAnimation = false;
 			RecordingSettings.bCheckDeltaTimeAtBeginning = false;
+			RecordingSettings.bSetRetargetSourceAsset = ExportOptions->bSetRetargetSourceAsset;
 			RecordingSettings.bRecordTransforms = ExportOptions->bExportTransforms;
 			RecordingSettings.bRecordMorphTargets = ExportOptions->bExportMorphTargets;
 			RecordingSettings.bRecordAttributeCurves = ExportOptions->bExportAttributeCurves;
diff --git a/Engine/Source/Editor/MovieSceneTools/Private/TrackEditors/CommonAnimationTrackEditor.cpp b/Engine/Source/Editor/MovieSceneTools/Private/TrackEditors/CommonAnimationTrackEditor.cpp
--- a/Engine/Source/Editor/MovieSceneTools/Private/TrackEditors/CommonAnimationTrackEditor.cpp
+++ b/Engine/Source/Editor/MovieSceneTools/Private/TrackEditors/CommonAnimationTrackEditor.cpp
@@ -1541,6 +1541,7 @@ void FCommonAnimationTrackEditor::OnSequencerSaved(ISequencer& )
 						USkeletalMeshComponent* SkelMeshComp = UE::Sequencer::FCommonAnimationTrackEditor::AcquireSkeletalMeshFromObjectGuid(Item.SkelTrackGuid, GetSequencer());
 						if (AnimSequence && SkelMeshComp)
 						{
+							const bool bSavedSetRetargetSourceAsset = AnimSeqExportOption->bSetRetargetSourceAsset;
 							const bool bSavedExportMorphTargets = AnimSeqExportOption->bExportMorphTargets;
 							const bool bSavedExportAttributeCurves = AnimSeqExportOption->bExportAttributeCurves;
 							const bool bSavedExportMaterialCurves = AnimSeqExportOption->bExportMaterialCurves;
@@ -1560,6 +1561,22 @@ void FCommonAnimationTrackEditor::OnSequencerSaved(ISequencer& )
 							const bool bSavedUseCustomFrameRate = AnimSeqExportOption->bUseCustomFrameRate;
 							const FFrameRate SavedCustomFrameRate = AnimSeqExportOption->CustomFrameRate;
 
+							//If the asset cleared the retarget source asset we don't want to set it!
+							if (Item.bSetRetargetSourceAsset)
+							{
+								if (AnimSequence->GetRetargetSourceAsset().IsNull() == true)
+								{
+									AnimSeqExportOption->bSetRetargetSourceAsset = false;
+								}
+								else
+								{
+									AnimSeqExportOption->bSetRetargetSourceAsset = true;
+								}
+							}
+							else
+							{
+								AnimSeqExportOption->bSetRetargetSourceAsset = false;
+							}
 							AnimSeqExportOption->bExportMorphTargets = Item.bExportMorphTargets;
 							AnimSeqExportOption->bExportAttributeCurves = Item.bExportAttributeCurves;
 							AnimSeqExportOption->bExportMaterialCurves = Item.bExportMaterialCurves;
@@ -1591,6 +1608,7 @@ void FCommonAnimationTrackEditor::OnSequencerSaved(ISequencer& )
 							SequencerPtr->ForceEvaluate();
 							bool bResult = MovieSceneToolHelpers::ExportToAnimSequence(AnimSequence, AnimSeqExportOption, AESP, SkelMeshComp);
 
+							AnimSeqExportOption->bSetRetargetSourceAsset = bSavedSetRetargetSourceAsset;
 							AnimSeqExportOption->bExportMorphTargets = bSavedExportMorphTargets;
 							AnimSeqExportOption->bExportAttributeCurves = bSavedExportAttributeCurves;
 							AnimSeqExportOption->bExportMaterialCurves = bSavedExportMaterialCurves;
@@ -1673,7 +1691,7 @@ bool FCommonAnimationTrackEditor::CreateAnimationSequence(const TArray<UObject*>
 	if (NewAssets.Num() > 0)
 	{
 		UAnimSequence* AnimSequence = Cast<UAnimSequence>(NewAssets[0]);
-		if (AnimSequence)
+		if (AnimSequence && GetSequencer().IsValid())
 		{
 			UObject* NewAsset = NewAssets[0];
 			TSharedPtr<SWindow> ParentWindow;
@@ -1775,6 +1793,7 @@ bool FCommonAnimationTrackEditor::CreateAnimationSequence(const TArray<UObject*>
 									}
 								}
 								LevelAnimLinkItem.PathToAnimSequence = FSoftObjectPath(AnimSequence);
+								LevelAnimLinkItem.bSetRetargetSourceAsset = AnimSeqExportOption->bSetRetargetSourceAsset;
 								LevelAnimLinkItem.bExportMorphTargets = AnimSeqExportOption->bExportMorphTargets;
 								LevelAnimLinkItem.bExportAttributeCurves = AnimSeqExportOption->bExportAttributeCurves;
 								LevelAnimLinkItem.bExportMaterialCurves = AnimSeqExportOption->bExportMaterialCurves;
@@ -1807,6 +1826,7 @@ bool FCommonAnimationTrackEditor::CreateAnimationSequence(const TArray<UObject*>
 						FLevelSequenceAnimSequenceLinkItem LevelAnimLinkItem;
 						LevelAnimLinkItem.SkelTrackGuid = Binding;
 						LevelAnimLinkItem.PathToAnimSequence = FSoftObjectPath(AnimSequence);
+						LevelAnimLinkItem.bSetRetargetSourceAsset = AnimSeqExportOption->bSetRetargetSourceAsset;
 						LevelAnimLinkItem.bExportMorphTargets = AnimSeqExportOption->bExportMorphTargets;
 						LevelAnimLinkItem.bExportAttributeCurves = AnimSeqExportOption->bExportAttributeCurves;
 						LevelAnimLinkItem.bExportMaterialCurves = AnimSeqExportOption->bExportMaterialCurves;
diff --git a/Engine/Source/Editor/SequenceRecorder/Private/AnimationRecorder.cpp b/Engine/Source/Editor/SequenceRecorder/Private/AnimationRecorder.cpp
--- a/Engine/Source/Editor/SequenceRecorder/Private/AnimationRecorder.cpp
+++ b/Engine/Source/Editor/SequenceRecorder/Private/AnimationRecorder.cpp
@@ -298,11 +298,14 @@ void FAnimationRecorder::StartRecord(USkeletalMeshComponent* Component, UAnimSeq
 		}
 	}
 
-	AnimationObject->RetargetSource = Component->GetSkeletalMeshAsset() ? AnimSkeleton->GetRetargetSourceForMesh(Component->GetSkeletalMeshAsset()) : NAME_None;
-	if (AnimationObject->RetargetSource == NAME_None)
+	if (bSetRetargetSourceAsset)
 	{
-		AnimationObject->SetRetargetSourceAsset(Component->GetSkeletalMeshAsset());
-		AnimationObject->UpdateRetargetSourceAssetData();
+		AnimationObject->RetargetSource = Component->GetSkeletalMeshAsset() ? AnimSkeleton->GetRetargetSourceForMesh(Component->GetSkeletalMeshAsset()) : NAME_None;
+		if (AnimationObject->RetargetSource == NAME_None)
+		{
+			AnimationObject->SetRetargetSourceAsset(Component->GetSkeletalMeshAsset());
+			AnimationObject->UpdateRetargetSourceAssetData();
+		}
 	}
 
 	// record first-frame notifies
@@ -1198,6 +1201,7 @@ void FAnimRecorderInstance::InitInternal(USkeletalMeshComponent* InComponent, co
 	SkelComp = InComponent;
 	Recorder = MakeShareable(new FAnimationRecorder());
 	Recorder->SetSampleRateAndLength(Settings.SampleFrameRate, Settings.Length);
+	Recorder->bSetRetargetSourceAsset = Settings.bSetRetargetSourceAsset;
 	Recorder->bRecordLocalToWorld = Settings.bRecordInWorldSpace;
 	Recorder->Interpolation = Settings.Interpolation;
 	Recorder->InterpMode = Settings.InterpMode;
diff --git a/Engine/Source/Editor/SequenceRecorder/Public/AnimationRecorder.h b/Engine/Source/Editor/SequenceRecorder/Public/AnimationRecorder.h
--- a/Engine/Source/Editor/SequenceRecorder/Public/AnimationRecorder.h
+++ b/Engine/Source/Editor/SequenceRecorder/Public/AnimationRecorder.h
@@ -186,6 +186,8 @@ public:
 	ERichCurveTangentMode TangentMode;
 	/** Serializer, if set we also store data out incrementally while running*/
 	FAnimationSerializer* AnimationSerializer;
+	/** Whether or not to set Retarget Source Asset*/
+	uint8 bSetRetargetSourceAsset : 1;
 	/** Whether or not to record transforms*/
 	uint8 bRecordTransforms : 1;
 	/** Whether or not to record morph targets*/
diff --git a/Engine/Source/Editor/UnrealEd/Classes/Exporters/AnimSeqExportOption.h b/Engine/Source/Editor/UnrealEd/Classes/Exporters/AnimSeqExportOption.h
--- a/Engine/Source/Editor/UnrealEd/Classes/Exporters/AnimSeqExportOption.h
+++ b/Engine/Source/Editor/UnrealEd/Classes/Exporters/AnimSeqExportOption.h
@@ -42,6 +42,10 @@ public:
 	UPROPERTY(EditAnywhere, BlueprintReadWrite, AdvancedDisplay, Category = Export);	
 	bool bEvaluateAllSkeletalMeshComponents = true;
 
+	/** If enabled set the retarget source asset from current skel mesh component */
+	UPROPERTY(EditAnywhere, BlueprintReadWrite, AdvancedDisplay, Category = Export)
+	bool bSetRetargetSourceAsset = false;
+
 	/** This defines how values between keys are calculated for transforms*/
 	UPROPERTY(EditAnywhere, BlueprintReadWrite, AdvancedDisplay, Category = Export);
 	EAnimInterpolationType Interpolation = EAnimInterpolationType::Linear;
@@ -114,6 +118,7 @@ public:
 		bExportAttributeCurves = true;
 		bExportMaterialCurves = true;
 		bRecordInWorldSpace = false;
+		bSetRetargetSourceAsset = false;
 		Interpolation = EAnimInterpolationType::Linear;
 		CurveInterpolation = ERichCurveInterpMode::RCIM_Linear;
 		bEvaluateAllSkeletalMeshComponents = true;
diff --git a/Engine/Source/Runtime/Engine/Classes/Animation/AnimationRecordingSettings.h b/Engine/Source/Runtime/Engine/Classes/Animation/AnimationRecordingSettings.h
--- a/Engine/Source/Runtime/Engine/Classes/Animation/AnimationRecordingSettings.h
+++ b/Engine/Source/Runtime/Engine/Classes/Animation/AnimationRecordingSettings.h
@@ -34,6 +34,7 @@ struct FAnimationRecordingSettings
 		, InterpMode(ERichCurveInterpMode::RCIM_Linear)
 		, TangentMode(ERichCurveTangentMode::RCTM_Auto)
 		, bCheckDeltaTimeAtBeginning(true)
+		, bSetRetargetSourceAsset(false)
 		, bRecordTransforms(true)
 		, bRecordMorphTargets(true)
 		, bRecordAttributeCurves(true)
@@ -76,6 +77,11 @@ struct FAnimationRecordingSettings
 	/** Whether to check DeltaTime at recording for pauses, turned off for TakeRecorder*/
 	bool bCheckDeltaTimeAtBeginning;
 
+	/** Whether or not to set retarget source asset */
+	UPROPERTY(EditAnywhere, Category = "Settings")
+	bool bSetRetargetSourceAsset;
+
+
 	/** Whether or not to record transforms */
 	UPROPERTY(EditAnywhere, Category = "Settings")
 	bool bRecordTransforms;
diff --git a/Engine/Source/Runtime/LevelSequence/Public/LevelSequenceAnimSequenceLink.h b/Engine/Source/Runtime/LevelSequence/Public/LevelSequenceAnimSequenceLink.h
--- a/Engine/Source/Runtime/LevelSequence/Public/LevelSequenceAnimSequenceLink.h
+++ b/Engine/Source/Runtime/LevelSequence/Public/LevelSequenceAnimSequenceLink.h
@@ -37,6 +37,8 @@ public:
 	bool bExportAttributeCurves = true;
 	UPROPERTY(BlueprintReadWrite, Category = Property)
 	bool bExportMaterialCurves = true;
+	UPROPERTY(BlueprintReadWrite, Category = Property)
+	bool bSetRetargetSourceAsset = true;
 	UPROPERTY(BlueprintReadWrite, Category = Property);
 	EAnimInterpolationType Interpolation = EAnimInterpolationType::Linear;
 	UPROPERTY(BlueprintReadWrite, Category = Property);
