Constraints: faster sequential evaluation

with the new evaluation scheme, it's no longer necessary to force the parent evaluation when evaluating a constraint.
in particular, when several constraints are evaluated sequentially, it's possible to mark the parent for evaluation only for the first constraint and thus greatly improve performance.
this scheme is not active by default yet (`Constraints.UseFastEvaluation true` to activate it) but will be once the initial tests and feedback are positive.

#rnx

#Synced-CL 48654334

diff --git a/Engine/Source/Editor/MovieSceneTools/Private/Constraints/MovieSceneConstraintChannelHelper.cpp b/Engine/Source/Editor/MovieSceneTools/Private/Constraints/MovieSceneConstraintChannelHelper.cpp
--- a/Engine/Source/Editor/MovieSceneTools/Private/Constraints/MovieSceneConstraintChannelHelper.cpp
+++ b/Engine/Source/Editor/MovieSceneTools/Private/Constraints/MovieSceneConstraintChannelHelper.cpp
@@ -2,6 +2,7 @@
 
 #include "Constraints/MovieSceneConstraintChannelHelper.inl"
 
+#include "ConstraintsEvaluation.h"
 #include "Evaluation/MovieSceneEvaluationTemplateInstance.h"
 
 #include "LevelSequence.h"
@@ -12,7 +13,6 @@
 #include "Transform/TransformConstraintUtil.h"
 #include "Algo/Copy.h"
 
-#include "ScopedTransaction.h"
 #include "Channels/MovieSceneFloatChannel.h"
 #include "Channels/MovieSceneDoubleChannel.h"
 #include "LevelEditorViewport.h"
@@ -116,7 +116,7 @@ void FCompensationEvaluator::ComputeLocalTransforms(UWorld* InWorld, const FEval
 	
 	// get all constraints to evaluate
 	const FConstraintsManagerController& Controller = FConstraintsManagerController::Get(InWorld);
-	static constexpr bool bSorted = true, bTickHandles = true;
+	static constexpr bool bSorted = true;
 	const TArray<ConstraintPtr> AllConstraints = Controller.GetAllConstraints(bSorted);
 
 	FMovieSceneInverseSequenceTransform LocalToRootTransform = Sequencer->GetFocusedMovieSceneSequenceTransform().Inverse();
@@ -152,13 +152,7 @@ void FCompensationEvaluator::ComputeLocalTransforms(UWorld* InWorld, const FEval
 			}
 			
 			// evaluate constraints
-			for (const TWeakObjectPtr<UTickableConstraint>& InConstraint : AllConstraints)
-			{
-				if (InConstraint.IsValid())
-				{
-					InConstraint->Evaluate(bTickHandles);
-				}
-			}
+			UE::Anim::Constraints::EvaluateConstraints(AllConstraints);
 		
 			BakingWrapper.PostEvaluate(FrameNumber);
 		}
@@ -372,13 +366,7 @@ void FCompensationEvaluator::ComputeLocalTransformsForBaking(UWorld* InWorld, co
 			// evaluate constraints
 			{
 				const UE::Anim::FEvaluationForCachingScope CachingScope(DeltaTime);
-				for (const TWeakObjectPtr<UTickableConstraint>& InConstraint : AllConstraints)
-				{
-					if (InConstraint.IsValid())
-					{
-						InConstraint->Evaluate(true);
-					}
-				}
+				UE::Anim::Constraints::EvaluateConstraints(AllConstraints);
 			}
 
 			BakingWrapper.PostEvaluate(FrameNumber);
@@ -459,7 +447,7 @@ void FCompensationEvaluator::ComputeLocalTransformsBeforeDeletion(UWorld* InWorl
 		{
 			SpaceGlobal = LastConstraint->GetParentGlobalTransform();
 			TOptional<FTransform> Relative =
-				UE::TransformConstraintUtil::GetConstraintsRelativeTransform(Constraints, ChildLocal, ChildGlobal);
+				UE::TransformConstraintUtil::GetConstraintsRelativeTransform(ConstraintsMinusThis, ChildLocal, ChildGlobal);
 			if (Relative)
 			{
 				ChildLocal = *Relative;
@@ -511,13 +499,7 @@ void FCompensationEvaluator::ComputeLocalTransformsBeforeDeletion(UWorld* InWorl
 		Sequencer->GetEvaluationTemplate().EvaluateSynchronousBlocking(Context);
 
 		// evaluate constraints
-		for (const TWeakObjectPtr<UTickableConstraint>& InConstraint : AllConstraints)
-		{
-			if (InConstraint.IsValid())
-			{
-				InConstraint->Evaluate(true);
-			}
-		}
+		UE::Anim::Constraints::EvaluateConstraints(AllConstraints);
 		
 		BakingWrapper.PostEvaluate(FrameNumber);
 		
@@ -534,7 +516,7 @@ void FCompensationEvaluator::ComputeLocalTransformsBeforeDeletion(UWorld* InWorl
 		{
 			SpaceGlobal = LastConstraint->GetParentGlobalTransform();
 			TOptional<FTransform> Relative =
-				UE::TransformConstraintUtil::GetConstraintsRelativeTransform(Constraints, ChildLocal, ChildGlobal);
+				UE::TransformConstraintUtil::GetConstraintsRelativeTransform(ConstraintsMinusThis, ChildLocal, ChildGlobal);
 			if(Relative)
 			{
 				ChildLocal = *Relative;
@@ -600,13 +582,7 @@ void FCompensationEvaluator::ComputeCompensation(UWorld* InWorld, const TSharedP
 			
 		InSequencer->GetEvaluationTemplate().EvaluateSynchronousBlocking(Context0);
 
-		for (const TWeakObjectPtr<UTickableConstraint>& InConstraint : AllConstraints)
-		{
-			if (InConstraint.IsValid())
-			{
-				InConstraint->Evaluate(true);
-			}
-		}
+		UE::Anim::Constraints::EvaluateConstraints(AllConstraints);
 
 		BakingWrapper.PostEvaluate(InFrame);
 
@@ -701,13 +677,7 @@ void FCompensationEvaluator::CacheTransforms(UWorld* InWorld, const FEvalParamet
 		Sequencer->GetEvaluationTemplate().EvaluateSynchronousBlocking(Context);
 
 		// evaluate constraints
-		for (const TWeakObjectPtr<UTickableConstraint>& InConstraint : AllConstraints)
-		{
-			if (InConstraint.IsValid())
-			{
-				InConstraint->Evaluate(true);
-			}
-		}
+		UE::Anim::Constraints::EvaluateConstraints(AllConstraints);
 		
 		BakingWrapper.PostEvaluate(InFrame);
 	};
diff --git a/Engine/Source/Runtime/Experimental/Animation/Constraints/Private/ConstraintsEvaluation.cpp b/Engine/Source/Runtime/Experimental/Animation/Constraints/Private/ConstraintsEvaluation.cpp
new file mode 100644
--- /dev/null
+++ b/Engine/Source/Runtime/Experimental/Animation/Constraints/Private/ConstraintsEvaluation.cpp
@@ -0,0 +1,95 @@
+﻿// Copyright Epic Games, Inc. All Rights Reserved.
+
+#include "ConstraintsEvaluation.h"
+
+#include "ConstraintsManager.h"
+#include "Transform/TransformableHandleUtils.h"
+#include "Transform/TransformConstraint.h"
+
+namespace UE::Anim::Constraints
+{
+
+static bool	bUseFastEvaluation = false;
+static FAutoConsoleVariableRef CVarUseFastEvaluation(
+	TEXT("Constraints.UseFastEvaluation"),
+	bUseFastEvaluation,
+	TEXT("Whether to enable fast sequential constraints evaluation. (default is false)") );
+
+bool UseFastEvaluation()
+{
+	return bUseFastEvaluation && TransformableHandleUtils::SkipTicking();
+}
+
+void EvaluateConstraints(TConstArrayView<TWeakObjectPtr<UTickableConstraint>> InConstraints)
+{
+	if (UseFastEvaluation())
+	{
+		FSequentialConstraintsEvaluator Evaluator;
+		for (const TWeakObjectPtr<UTickableConstraint>& Constraint : InConstraints)
+		{
+			if (Constraint.IsValid())
+			{
+				Evaluator.DoEvaluate(Constraint.Get());
+			}
+		}
+	}
+	else
+	{
+		static constexpr bool bTickHandles = true;
+		for (const TWeakObjectPtr<UTickableConstraint>& InConstraint : InConstraints)
+		{
+			if (InConstraint.IsValid())
+			{
+				InConstraint->Evaluate(bTickHandles);
+			}
+		}
+	}
+}
+	
+}
+
+FSequentialConstraintsEvaluator::FSequentialConstraintsEvaluator(USceneComponent* InInitialComponent)
+	: LastMarkedForEvaluation(InInitialComponent)
+{}
+
+void FSequentialConstraintsEvaluator::DoEvaluate(UTickableConstraint* InConstraint, const FConstraintTickFunction* InConstraintTickFunction)
+{
+	if (!ensure(UE::Anim::Constraints::UseFastEvaluation()))
+	{
+		return;
+	}
+	
+	const bool bIsConstraintActive = InConstraint && InConstraint->IsFullyActive();
+	if (!bIsConstraintActive)
+	{
+		return;
+	}
+	
+	if (InConstraintTickFunction)
+	{
+		if (!InConstraintTickFunction->IsTickFunctionRegistered() || !InConstraintTickFunction->IsTickFunctionEnabled())
+		{
+			return;
+		}
+	}
+	
+	if (UTickableTransformConstraint* TransformConstraint = Cast<UTickableTransformConstraint>(InConstraint))
+	{
+		bool bAlsoTick = false;
+		if (USceneComponent* Component = TransformableHandleUtils::GetTarget<USceneComponent>(TransformConstraint->ParentTRSHandle))
+		{
+			if (Component != LastMarkedForEvaluation)
+			{
+				bAlsoTick = true;
+				LastMarkedForEvaluation = Component;
+			}
+		}
+		InConstraint->Evaluate(bAlsoTick);
+	}
+	else
+	{
+		static constexpr bool bTickHandles = true;
+		InConstraint->Evaluate(bTickHandles);	
+	}
+}
+
diff --git a/Engine/Source/Runtime/Experimental/Animation/Constraints/Private/ConstraintsEvaluationGraph.cpp b/Engine/Source/Runtime/Experimental/Animation/Constraints/Private/ConstraintsEvaluationGraph.cpp
--- a/Engine/Source/Runtime/Experimental/Animation/Constraints/Private/ConstraintsEvaluationGraph.cpp
+++ b/Engine/Source/Runtime/Experimental/Animation/Constraints/Private/ConstraintsEvaluationGraph.cpp
@@ -67,7 +67,8 @@ void FConstraintsEvaluationGraph::FlushPendingEvaluations()
 	}
 
 	State = Flushing;
-	
+
+	TGuardValue<FSequentialConstraintsEvaluator> SequentialEvaluatorGuard(SequentialEvaluator, nullptr);
 	for (FConstraintNode& Node: Nodes)
 	{
 		if (Node.bMarkedForEvaluation)
@@ -88,6 +89,7 @@ void FConstraintsEvaluationGraph::FlushPendingEvaluations()
 void FConstraintsEvaluationGraph::Rebuild()
 {
 	Nodes.Empty();
+	SequentialEvaluator = FSequentialConstraintsEvaluator(nullptr);
 	
 	if (!ensure(ConstraintsInWorld.World.Get()))
 	{
@@ -149,24 +151,6 @@ bool FConstraintsEvaluationGraph::IsPendingEvaluation() const
 	return State == PendingEvaluation;
 }
 
-void FConstraintsEvaluationGraph::Evaluate(const TWeakObjectPtr<UTickableConstraint>& InConstraint)
-{
-	if (State == InvalidData)
-	{
-		Rebuild();
-	}
-
-	if (Nodes.IsEmpty())
-	{
-		return;
-	}
-	
-	if (FConstraintNode* Node = FindNode(InConstraint))
-	{
-		Evaluate(Node);
-	}
-}
-
 void FConstraintsEvaluationGraph::Evaluate(FConstraintNode* InNode)
 {
 	const TArray<TWeakObjectPtr<UTickableConstraint>>& Constraints = ConstraintsInWorld.Constraints;
@@ -188,10 +172,16 @@ void FConstraintsEvaluationGraph::Evaluate(FConstraintNode* InNode)
 	
 	// evaluate
 	TGuardValue<bool> ReentrantGuard(InNode->bEvaluating, true);
-	
-	if (Constraint->IsFullyActive() && InNode->ConstraintTick->IsTickFunctionRegistered() && InNode->ConstraintTick->IsTickFunctionEnabled())
+	if (UE::Anim::Constraints::UseFastEvaluation())
+	{
+		SequentialEvaluator.DoEvaluate(Constraint.Get(), InNode->ConstraintTick);
+	}
+	else 
 	{
-		Constraint->Evaluate(TransformableHandleUtils::SkipTicking());
+		if (Constraint->IsFullyActive() && InNode->ConstraintTick->IsTickFunctionRegistered() && InNode->ConstraintTick->IsTickFunctionEnabled())
+		{
+			Constraint->Evaluate(TransformableHandleUtils::SkipTicking());
+		}
 	}
 	InNode->bMarkedForEvaluation = false;
 
@@ -209,6 +199,7 @@ void FConstraintsEvaluationGraph::InvalidateData()
 {
 	State = InvalidData;
 	Nodes.Empty();
+	SequentialEvaluator = FSequentialConstraintsEvaluator(nullptr);
 }
 
 void FConstraintsEvaluationGraph::MarkForEvaluation(const TWeakObjectPtr<UTickableConstraint>& InConstraint)
diff --git a/Engine/Source/Runtime/Experimental/Animation/Constraints/Private/ConstraintsManager.cpp b/Engine/Source/Runtime/Experimental/Animation/Constraints/Private/ConstraintsManager.cpp
--- a/Engine/Source/Runtime/Experimental/Animation/Constraints/Private/ConstraintsManager.cpp
+++ b/Engine/Source/Runtime/Experimental/Animation/Constraints/Private/ConstraintsManager.cpp
@@ -952,15 +952,14 @@ void FConstraintsManagerController::EvaluateAllConstraints() const
 {
 	using ConstraintPtr = TWeakObjectPtr<UTickableConstraint>;
 
-	static constexpr bool bSorted = true, bTickHandles = true;
+	static constexpr bool bSorted = true;
 	const TArray<ConstraintPtr> Constraints = GetAllConstraints(bSorted);
-	for (const ConstraintPtr& Constraint : Constraints)
+	if (Constraints.IsEmpty())
 	{
-		if (Constraint.IsValid())
-		{
-			Constraint->Evaluate(bTickHandles);
-		}
+		return;
 	}
+
+	UE::Anim::Constraints::EvaluateConstraints(Constraints);
 }
 
 FConstraintsManagerController::FOnSceneComponentConstrained& FConstraintsManagerController::OnSceneComponentConstrained()
diff --git a/Engine/Source/Runtime/Experimental/Animation/Constraints/Public/ConstraintsEvaluation.h b/Engine/Source/Runtime/Experimental/Animation/Constraints/Public/ConstraintsEvaluation.h
new file mode 100644
--- /dev/null
+++ b/Engine/Source/Runtime/Experimental/Animation/Constraints/Public/ConstraintsEvaluation.h
@@ -0,0 +1,37 @@
+﻿// Copyright Epic Games, Inc. All Rights Reserved.
+
+#pragma once
+
+#include "Containers/ContainersFwd.h"
+#include "UObject/WeakObjectPtrTemplates.h"
+
+class UTickableConstraint;
+struct FConstraintTickFunction;
+class USceneComponent;
+
+#define UE_API CONSTRAINTS_API
+
+/**
+ * FSequentialConstraintsEvaluator provides a sequential constraint evaluation scheme, ensuring that parents are marked for evaluation only if necessary.
+ * NOTE: this scheme is only effective if UseFastEvaluation() is true.
+ */
+
+struct FSequentialConstraintsEvaluator 
+{
+	FSequentialConstraintsEvaluator(USceneComponent* InInitialComponent = nullptr);
+	void DoEvaluate(UTickableConstraint* InConstraint, const FConstraintTickFunction* InConstraintTickFunction = nullptr);
+	
+private:
+	USceneComponent* LastMarkedForEvaluation = nullptr;
+};
+
+namespace UE::Anim::Constraints
+{
+	/* Whether the fast sequential evaluation is enabled. */
+	bool UseFastEvaluation();
+
+	/* Evaluate all the constraints sequentially. this will use a faster evaluation scheme UseFastEvaluation() is true. */
+	void UE_API EvaluateConstraints(TConstArrayView<TWeakObjectPtr<UTickableConstraint>> InConstraints);
+}
+
+#undef UE_API
\ No newline at end of file
diff --git a/Engine/Source/Runtime/Experimental/Animation/Constraints/Public/ConstraintsEvaluationGraph.h b/Engine/Source/Runtime/Experimental/Animation/Constraints/Public/ConstraintsEvaluationGraph.h
--- a/Engine/Source/Runtime/Experimental/Animation/Constraints/Public/ConstraintsEvaluationGraph.h
+++ b/Engine/Source/Runtime/Experimental/Animation/Constraints/Public/ConstraintsEvaluationGraph.h
@@ -2,6 +2,7 @@
 
 #pragma once
 
+#include "ConstraintsEvaluation.h"
 #include "UObject/Object.h"
 #include "Engine/EngineBaseTypes.h"
 
@@ -88,9 +89,6 @@ private:
 	/* Returns the node wrapping that constraint if any. */
 	FConstraintNode* FindNode(const TWeakObjectPtr<UTickableConstraint>& InConstraint);
 
-	/* Evaluates InConstraint ensuring there's a corresponding node. */
-	void Evaluate(const TWeakObjectPtr<UTickableConstraint>& InConstraint);
-
 	/* Evaluates InNode if its tick function is enabled + evaluates its children recursively. */
 	void Evaluate(FConstraintNode* InNode);
 
@@ -105,6 +103,9 @@ private:
 	
 	/* Current graph state. */
 	EGraphState	State = InvalidData;
+
+	/* Constraints sequential evaluator used if the fast evaluation scheme is enabled. */
+	FSequentialConstraintsEvaluator SequentialEvaluator;
 };
 
 namespace UE::Constraints::Graph
