[AvalancheTests] Automate 'Sequence - Changing values from Sequencer'

#rnx

#Synced-CL 45465245

diff --git a/Engine/Source/Editor/MovieSceneTools/Private/CurveKeyEditors/SNumericKeyEditor.h b/Engine/Source/Editor/MovieSceneTools/Private/CurveKeyEditors/SNumericKeyEditor.h
--- a/Engine/Source/Editor/MovieSceneTools/Private/CurveKeyEditors/SNumericKeyEditor.h
+++ b/Engine/Source/Editor/MovieSceneTools/Private/CurveKeyEditors/SNumericKeyEditor.h
@@ -158,7 +158,11 @@ private:
 
 	void OnValueCommitted(NumericType Value, ETextCommit::Type CommitInfo)
 	{
-		if (CommitInfo == ETextCommit::OnEnter || CommitInfo == ETextCommit::OnUserMovedFocus)
+		if (CommitInfo == ETextCommit::OnEnter || CommitInfo == ETextCommit::OnUserMovedFocus
+#if !UE_BUILD_SHIPPING
+			|| CommitInfo == ETextCommit::Default
+#endif // !UE_BUILD_SHIPPING
+			)
 		{
 			const FScopedTransaction Transaction( LOCTEXT("SetNumericKey", "Set Key Value") );
 			KeyEditor->SetValueWithNotify(Value, EMovieSceneDataChangeType::TrackValueChangedRefreshImmediately);
