Sequencer: Final unity build fix.

#Synced-CL 49029986

diff --git a/Engine/Source/Editor/MovieSceneTools/Private/IMovieSceneAnimMixerItemMenuProvider.cpp b/Engine/Source/Editor/MovieSceneTools/Private/IMovieSceneAnimMixerItemMenuProvider.cpp
--- a/Engine/Source/Editor/MovieSceneTools/Private/IMovieSceneAnimMixerItemMenuProvider.cpp
+++ b/Engine/Source/Editor/MovieSceneTools/Private/IMovieSceneAnimMixerItemMenuProvider.cpp
@@ -7,6 +7,8 @@ IMovieSceneAnimMixerItemMenuProvider::IMovieSceneAnimMixerItemMenuProvider()
 
 }
 
+IMovieSceneAnimMixerItemMenuProvider::~IMovieSceneAnimMixerItemMenuProvider(){}
+
 FName IMovieSceneAnimMixerItemMenuProvider::GetModularFeatureName()
 {
 	static FName FeatureName = FName(TEXT("MovieSceneAnimMixerItemProvider"));
diff --git a/Engine/Source/Editor/MovieSceneTools/Public/IMovieSceneAnimMixerItemMenuProvider.h b/Engine/Source/Editor/MovieSceneTools/Public/IMovieSceneAnimMixerItemMenuProvider.h
--- a/Engine/Source/Editor/MovieSceneTools/Public/IMovieSceneAnimMixerItemMenuProvider.h
+++ b/Engine/Source/Editor/MovieSceneTools/Public/IMovieSceneAnimMixerItemMenuProvider.h
@@ -25,7 +25,7 @@ public:
 
 	IMovieSceneAnimMixerItemMenuProvider();
 
-	virtual ~IMovieSceneAnimMixerItemMenuProvider() = default;
+	virtual ~IMovieSceneAnimMixerItemMenuProvider();
 
 	/**
 	 * Returns the UClass this provider handles.
