Sub Track Editor interfering with other editor transform widget placement

-Sub Track Editor Mode now returns false in UsesTransformWidget functions if an actor is selected, instead of calling its parent function. This prevents it from returning true when editing vector and transform properties through the viewport.
-Sequencer Tools Editor Mode no longer uses property widgets. UsesPropertyWidgets was originally overidden to address UE-136420, which was similar to this issue. Since then UsesTransformWidget was overidden to return false for other reasons, but resolves the issue more correctly. (UsesPropertyWidgets returning true actually caused transforms to be doubled through FLegacyEdModeWidgetHelper::InputDelta goign through the property widget path twice, and was causing the sign of the drag delta to flip, causing the gizmo to snap back and forth when dragging)

#Synced-CL 49032165

diff --git a/Engine/Plugins/Editor/SequencerAnimTools/Source/SequencerAnimTools/Public/SequencerToolsEditMode.h b/Engine/Plugins/Editor/SequencerAnimTools/Source/SequencerAnimTools/Public/SequencerToolsEditMode.h
--- a/Engine/Plugins/Editor/SequencerAnimTools/Source/SequencerAnimTools/Public/SequencerToolsEditMode.h
+++ b/Engine/Plugins/Editor/SequencerAnimTools/Source/SequencerAnimTools/Public/SequencerToolsEditMode.h
@@ -25,8 +25,7 @@ public:
 	
 	virtual bool UsesToolkits() const  override { return false; }
 	virtual bool InputKey(FEditorViewportClient* InViewportClient, FViewport* InViewport, FKey InKey, EInputEvent InEvent) override;
-	virtual bool UsesPropertyWidgets() const override { return true; }
-	virtual bool UsesTransformWidget() const { return false; }
+	virtual bool UsesTransformWidget() const override { return false; }
 	virtual bool IsCompatibleWith(FEditorModeID OtherModeID) const override;
 private:
 
diff --git a/Engine/Source/Editor/MovieSceneTools/Private/EditModes/SubTrackEditorMode.cpp b/Engine/Source/Editor/MovieSceneTools/Private/EditModes/SubTrackEditorMode.cpp
--- a/Engine/Source/Editor/MovieSceneTools/Private/EditModes/SubTrackEditorMode.cpp
+++ b/Engine/Source/Editor/MovieSceneTools/Private/EditModes/SubTrackEditorMode.cpp
@@ -153,7 +153,7 @@ bool FSubTrackEditorMode::UsesTransformWidget() const
 	{
 		return DoesSubSectionHaveTransformOverrides(*SubSection);
 	}
-	return FEdMode::UsesTransformWidget();
+	return false;
 }
 
 bool FSubTrackEditorMode::UsesTransformWidget(UE::Widget::EWidgetMode CheckMode) const
@@ -163,7 +163,7 @@ bool FSubTrackEditorMode::UsesTransformWidget(UE::Widget::EWidgetMode CheckMode)
 	{
 		return DoesSubSectionHaveTransformOverrides(*SubSection);
 	}
-	return FEdMode::UsesTransformWidget(CheckMode);
+	return false;
 }
 
 TOptional<FMovieSceneSequenceID> FSubTrackEditorMode::GetSequenceIDForSubSection(const UMovieSceneSubSection* InSubSection) const
