FAllSpawnableRestoreState: do not modify the UMovieSceneSpawnTrack's section if it is constant
also note that theUMovieSceneBindingLifetimeTrack should probably not be modified neither if its range encompasses the caching range

#rnx
#Synced-CL 49239030

diff --git a/Engine/Source/Editor/MovieSceneTools/Private/MovieSceneToolHelpers.cpp b/Engine/Source/Editor/MovieSceneTools/Private/MovieSceneToolHelpers.cpp
--- a/Engine/Source/Editor/MovieSceneTools/Private/MovieSceneToolHelpers.cpp
+++ b/Engine/Source/Editor/MovieSceneTools/Private/MovieSceneToolHelpers.cpp
@@ -4675,22 +4675,43 @@ void FAllSpawnableRestoreState::SaveStateForSpawnable(bool& bInOutWasChanged, TM
 
 	SpawnOwnershipSetter(ESpawnOwnership::RootSequence);
 
-	UMovieSceneSpawnTrack* SpawnTrack = WeakMovieScene->FindTrack<UMovieSceneSpawnTrack>(BindingGuid);
-
-	if (SpawnTrack && SpawnTrack->GetAllSections().Num() > 0)
+	const UMovieSceneSpawnTrack* SpawnTrack = WeakMovieScene->FindTrack<UMovieSceneSpawnTrack>(BindingGuid);
+	if (SpawnTrack && !SpawnTrack->IsEmpty())
 	{
-		// Start a transaction that will be undone later for the modifications to the spawn track
-		if (!bInOutWasChanged)
+		if (UMovieSceneSpawnSection* SpawnSection = Cast<UMovieSceneSpawnSection>(SpawnTrack->GetAllSections()[0]))
 		{
-			GEditor->BeginTransaction(NSLOCTEXT("MovieSceneToolHelpers", "SpwanableRestoreState", "SpawnableRestoreState"));
-		}
-
-		bInOutWasChanged = true;
+			// no need to modify the spawned state if it's constant
+			auto ShouldSectionBeModified = [SpawnSection]()
+			{
+				const FMovieSceneBoolChannel& SpawnedChannel = SpawnSection->GetChannel();
+				const TArrayView<const bool>& Values = SpawnedChannel.GetValues();
+				if (!Values.IsEmpty())
+				{
+					for (int32 Index = 1; Index < Values.Num(); Index++)
+					{
+						if (Values[Index] != Values[Index-1])
+						{
+							return true;
+						}
+					}
+				}
+				return false;
+			};
+		
+			if (ShouldSectionBeModified())
+			{
+				// Start a transaction that will be undone later for the modifications to the spawn track
+				if (!bInOutWasChanged)
+				{
+					GEditor->BeginTransaction(NSLOCTEXT("MovieSceneToolHelpers", "SpwanableRestoreState", "SpawnableRestoreState"));
+					bInOutWasChanged = true;
+				}
 
-		UMovieSceneSpawnSection* SpawnSection = Cast<UMovieSceneSpawnSection>(SpawnTrack->GetAllSections()[0]);
-		SpawnSection->Modify();
-		SpawnSection->GetChannel().Reset();
-		SpawnSection->GetChannel().SetDefault(true);
+				// spawn it by default 
+				SpawnSection->Modify();
+				SpawnSection->GetChannel().SetDefault(true);
+			}
+		}
 	}
 
 	UMovieSceneBindingLifetimeTrack* BindingLifetimeTrack = WeakMovieScene->FindTrack<UMovieSceneBindingLifetimeTrack>(BindingGuid);
@@ -4712,7 +4733,8 @@ void FAllSpawnableRestoreState::SaveStateForSpawnable(bool& bInOutWasChanged, TM
 		NewSection->SetRange(TRange<FFrameNumber>::All());
 		BindingLifetimeTrack->AddSection(*NewSection);
 	}
-};
+}
+
 FAllSpawnableRestoreState::FAllSpawnableRestoreState(TSharedPtr<UE::MovieScene::FSharedPlaybackState> InSharedPlaybackState, UMovieScene* InMovieScene)
 	: SharedPlaybackState(InSharedPlaybackState)
 {
