Sequencer: Add new protected member function to FSubTrackEditor, called AddSequenceToTrack, which makes it possible for subclasses to add sequences in the same way that the base class does without needing to duplicate any of the internal implementation.

#rnx
#Synced-CL 49044222

diff --git a/Engine/Source/Editor/MovieSceneTools/Private/TrackEditors/SubTrackEditor.cpp b/Engine/Source/Editor/MovieSceneTools/Private/TrackEditors/SubTrackEditor.cpp
--- a/Engine/Source/Editor/MovieSceneTools/Private/TrackEditors/SubTrackEditor.cpp
+++ b/Engine/Source/Editor/MovieSceneTools/Private/TrackEditors/SubTrackEditor.cpp
@@ -1663,6 +1663,12 @@ void FSubTrackEditor::HandleAddSubSequenceComboButtonMenuEntryEnterPressed(const
 	}
 }
 
+void FSubTrackEditor::AddSequenceToTrack(UMovieSceneSequence* InMovieSceneSequence, UMovieSceneTrack* InTrack)
+{
+	int32 RowIndex = INDEX_NONE;
+	AnimatablePropertyChanged(FOnKeyProperty::CreateRaw(this, &FSubTrackEditor::AddKeyInternal, InMovieSceneSequence, InTrack, RowIndex));
+}
+
 FKeyPropertyResult FSubTrackEditor::AddKeyInternal(FFrameNumber KeyTime, UMovieSceneSequence* InMovieSceneSequence, UMovieSceneTrack* InTrack, int32 RowIndex)
 {	
 	FKeyPropertyResult KeyPropertyResult;
diff --git a/Engine/Source/Editor/MovieSceneTools/Public/TrackEditors/SubTrackEditor.h b/Engine/Source/Editor/MovieSceneTools/Public/TrackEditors/SubTrackEditor.h
--- a/Engine/Source/Editor/MovieSceneTools/Public/TrackEditors/SubTrackEditor.h
+++ b/Engine/Source/Editor/MovieSceneTools/Public/TrackEditors/SubTrackEditor.h
@@ -169,6 +169,9 @@ protected:
 	/** Callback for generating the menu of the "Add Sequence" combo button. */
 	UE_API TSharedRef<SWidget> HandleAddSubSequenceComboButtonGetMenuContent(UE::Sequencer::TWeakViewModelPtr<UE::Sequencer::FViewModel> ViewModel);
 
+	/** Add the input UMovieSceneSequence to the input UMovieSceneTrack */
+	UE_API void AddSequenceToTrack(UMovieSceneSequence* InMovieSceneSequence, UMovieSceneTrack* InTrack);
+
 private:
 
 	/** Callback for executing a menu entry in the "Add Sequence" combo button. */
